Iris
* ReplicationSystem now has a reference to the UObjectReplicationBridge instead of the base UReplicationBrigde.h
* This will eventually lead to the removal or merging of the base and object bridge into one class

#rnx

diff --git a/Engine/Plugins/Runtime/ReplicationSystemTestPlugin/Source/Private/Tests/ReplicationSystem/Conditionals/TestConditionals.cpp b/Engine/Plugins/Runtime/ReplicationSystemTestPlugin/Source/Private/Tests/ReplicationSystem/Conditionals/TestConditionals.cpp
--- a/Engine/Plugins/Runtime/ReplicationSystemTestPlugin/Source/Private/Tests/ReplicationSystem/Conditionals/TestConditionals.cpp
+++ b/Engine/Plugins/Runtime/ReplicationSystemTestPlugin/Source/Private/Tests/ReplicationSystem/Conditionals/TestConditionals.cpp
@@ -891,7 +891,7 @@ UE_NET_TEST_FIXTURE(FTestConditionalsFixture, HierarchicalSubObjectIsNotReplicat
 	// Spawn object + subobjects on server
 	UTestReplicatedIrisObject* ServerObject = Server->CreateObject(0, 0);
 	UTestReplicatedIrisObject* ServerSubObject = Server->CreateSubObject(ServerObject->NetRefHandle, 0, 0);
-	UTestReplicatedIrisObject* ServerSubSubObject = Server->CreateSubObject<UTestReplicatedIrisObject>(ServerObject->NetRefHandle, ServerSubObject->NetRefHandle, UReplicationBridge::ESubObjectInsertionOrder::ReplicateWith);
+	UTestReplicatedIrisObject* ServerSubSubObject = Server->CreateSubObject<UTestReplicatedIrisObject>(ServerObject->NetRefHandle, ServerSubObject->NetRefHandle, ESubObjectInsertionOrder::ReplicateWith);
 
 	// Set condition to only replicate if root is the owning connection
 	Server->ReplicationBridge->SetSubObjectNetCondition(ServerSubObject->NetRefHandle, ELifetimeCondition::COND_SkipOwner);
diff --git a/Engine/Plugins/Runtime/ReplicationSystemTestPlugin/Source/Private/Tests/ReplicationSystem/ReplicatedTestObject.cpp b/Engine/Plugins/Runtime/ReplicationSystemTestPlugin/Source/Private/Tests/ReplicationSystem/ReplicatedTestObject.cpp
--- a/Engine/Plugins/Runtime/ReplicationSystemTestPlugin/Source/Private/Tests/ReplicationSystem/ReplicatedTestObject.cpp
+++ b/Engine/Plugins/Runtime/ReplicationSystemTestPlugin/Source/Private/Tests/ReplicationSystem/ReplicatedTestObject.cpp
@@ -8,7 +8,6 @@
 #include "Iris/ReplicationSystem/NetRefHandle.h"
 #include "Iris/ReplicationSystem/RepTag.h"
 #include "Iris/ReplicationSystem/ReplicationProtocol.h"
-#include "Iris/ReplicationSystem/ReplicationBridge.h"
 #include "Iris/ReplicationSystem/ReplicationSystem.h"
 #include "Iris/ReplicationSystem/ReplicationSystemInternal.h"
 #include "Iris/ReplicationState/InternalPropertyReplicationState.h"
@@ -85,7 +84,7 @@ UE::Net::FNetRefHandle UReplicatedTestObjectBridge::BeginReplication(UReplicated
 	return Handle;
 }
 
-UE::Net::FNetRefHandle UReplicatedTestObjectBridge::BeginReplication(FNetRefHandle OwnerHandle, UReplicatedTestObject* SubObjectInstance, FNetRefHandle InsertRelativeToSubObjectHandle, ESubObjectInsertionOrder InsertionOrder)
+UE::Net::FNetRefHandle UReplicatedTestObjectBridge::BeginReplication(FNetRefHandle OwnerHandle, UReplicatedTestObject* SubObjectInstance, FNetRefHandle InsertRelativeToSubObjectHandle, UE::Net::ESubObjectInsertionOrder InsertionOrder)
 {
 	check(OwnerHandle.IsValid());
 
diff --git a/Engine/Plugins/Runtime/ReplicationSystemTestPlugin/Source/Private/Tests/ReplicationSystem/ReplicatedTestObject.h b/Engine/Plugins/Runtime/ReplicationSystemTestPlugin/Source/Private/Tests/ReplicationSystem/ReplicatedTestObject.h
--- a/Engine/Plugins/Runtime/ReplicationSystemTestPlugin/Source/Private/Tests/ReplicationSystem/ReplicatedTestObject.h
+++ b/Engine/Plugins/Runtime/ReplicationSystemTestPlugin/Source/Private/Tests/ReplicationSystem/ReplicatedTestObject.h
@@ -623,7 +623,7 @@ public:
 	// In this example we have methods that directly uses UReplicatedTestObject;
 	FNetRefHandle BeginReplication(UReplicatedTestObject* Instance);
 	FNetRefHandle BeginReplication(UReplicatedTestObject* Instance, const UObjectReplicationBridge::FRootObjectReplicationParams& Params);
-	FNetRefHandle BeginReplication(FNetRefHandle OwnerHandle, UReplicatedTestObject* Instance, FNetRefHandle InsertRelativeToSubObjectHandle = FNetRefHandle::GetInvalid(), ESubObjectInsertionOrder InsertionOrder = UReplicationBridge::ESubObjectInsertionOrder::None);
+	FNetRefHandle BeginReplication(FNetRefHandle OwnerHandle, UReplicatedTestObject* Instance, FNetRefHandle InsertRelativeToSubObjectHandle = FNetRefHandle::GetInvalid(), UE::Net::ESubObjectInsertionOrder InsertionOrder = UE::Net::ESubObjectInsertionOrder::None);
 
 	void EndReplication(UReplicatedTestObject* Instance, EEndReplicationFlags Flags=EEndReplicationFlags::Destroy);
 
diff --git a/Engine/Plugins/Runtime/ReplicationSystemTestPlugin/Source/Private/Tests/ReplicationSystem/ReplicationSystemServerClientTestFixture.h b/Engine/Plugins/Runtime/ReplicationSystemTestPlugin/Source/Private/Tests/ReplicationSystem/ReplicationSystemServerClientTestFixture.h
--- a/Engine/Plugins/Runtime/ReplicationSystemTestPlugin/Source/Private/Tests/ReplicationSystem/ReplicationSystemServerClientTestFixture.h
+++ b/Engine/Plugins/Runtime/ReplicationSystemTestPlugin/Source/Private/Tests/ReplicationSystem/ReplicationSystemServerClientTestFixture.h
@@ -147,7 +147,7 @@ public:
 	}
 
 	template<typename T>
-	T* CreateSubObject(FNetRefHandle OwnerHandle, FNetRefHandle InsertRelativeToSubObjectHandle = FNetRefHandle::GetInvalid(), UReplicationBridge::ESubObjectInsertionOrder InsertionOrder = UReplicationBridge::ESubObjectInsertionOrder::None)
+	T* CreateSubObject(FNetRefHandle OwnerHandle, FNetRefHandle InsertRelativeToSubObjectHandle = FNetRefHandle::GetInvalid(), UE::Net::ESubObjectInsertionOrder InsertionOrder = UE::Net::ESubObjectInsertionOrder::None)
 	{
 		T* CreatedObject = NewObject<T>();
 		if (Cast<UReplicatedTestObject>(CreatedObject))
diff --git a/Engine/Plugins/Runtime/ReplicationSystemTestPlugin/Source/Private/Tests/ReplicationSystem/TestReplicationSystem.cpp b/Engine/Plugins/Runtime/ReplicationSystemTestPlugin/Source/Private/Tests/ReplicationSystem/TestReplicationSystem.cpp
--- a/Engine/Plugins/Runtime/ReplicationSystemTestPlugin/Source/Private/Tests/ReplicationSystem/TestReplicationSystem.cpp
+++ b/Engine/Plugins/Runtime/ReplicationSystemTestPlugin/Source/Private/Tests/ReplicationSystem/TestReplicationSystem.cpp
@@ -567,9 +567,9 @@ UE_NET_TEST_FIXTURE(FReplicationSystemServerClientTestFixture, TestSubObjectSpec
 	// Spawn some subobjects
 	UReplicatedSubObjectOrderObject* ServerSubObject0 = Server->CreateSubObject<UReplicatedSubObjectOrderObject>(ServerObject->NetRefHandle);
 	// Specify Subobject1 to replicate with SubObject0, which means that it will replicate before Subobjet0 is replicated
-	UReplicatedSubObjectOrderObject* ServerSubObject1 = Server->CreateSubObject<UReplicatedSubObjectOrderObject>(ServerObject->NetRefHandle, ServerSubObject0->NetRefHandle, UReplicationBridge::ESubObjectInsertionOrder::ReplicateWith);
+	UReplicatedSubObjectOrderObject* ServerSubObject1 = Server->CreateSubObject<UReplicatedSubObjectOrderObject>(ServerObject->NetRefHandle, ServerSubObject0->NetRefHandle, UE::Net::ESubObjectInsertionOrder::ReplicateWith);
 	// Specify SubObect 2 to replicate with no specific order (it will be added to the owner and thus replicate last)
-	UReplicatedSubObjectOrderObject* ServerSubObject2 = Server->CreateSubObject<UReplicatedSubObjectOrderObject>(ServerObject->NetRefHandle, ServerSubObject1->NetRefHandle, UReplicationBridge::ESubObjectInsertionOrder::None);
+	UReplicatedSubObjectOrderObject* ServerSubObject2 = Server->CreateSubObject<UReplicatedSubObjectOrderObject>(ServerObject->NetRefHandle, ServerSubObject1->NetRefHandle, UE::Net::ESubObjectInsertionOrder::None);
 
 	UReplicatedSubObjectOrderObject::RepOrderCounter = 0U;
 
@@ -606,10 +606,10 @@ UE_NET_TEST_FIXTURE(FReplicationSystemServerClientTestFixture, TestSubObjectInse
 	UReplicatedSubObjectOrderObject* ServerSubObject0 = Server->CreateSubObject<UReplicatedSubObjectOrderObject>(ServerObject->NetRefHandle);
 	
 	// Spawn a subobject and make it replicate before SubObject0
-	UReplicatedSubObjectOrderObject* ServerSubObject1 = Server->CreateSubObject<UReplicatedSubObjectOrderObject>(ServerObject->NetRefHandle, ServerSubObject0->NetRefHandle, UReplicationBridge::ESubObjectInsertionOrder::ReplicateWith);
+	UReplicatedSubObjectOrderObject* ServerSubObject1 = Server->CreateSubObject<UReplicatedSubObjectOrderObject>(ServerObject->NetRefHandle, ServerSubObject0->NetRefHandle, UE::Net::ESubObjectInsertionOrder::ReplicateWith);
 	
 	// Spawn a subobject and make it replicate first
-	UReplicatedSubObjectOrderObject* ServerSubObject2 = Server->CreateSubObject<UReplicatedSubObjectOrderObject>(ServerObject->NetRefHandle, FNetRefHandle::GetInvalid(), UReplicationBridge::ESubObjectInsertionOrder::InsertAtStart);
+	UReplicatedSubObjectOrderObject* ServerSubObject2 = Server->CreateSubObject<UReplicatedSubObjectOrderObject>(ServerObject->NetRefHandle, FNetRefHandle::GetInvalid(), UE::Net::ESubObjectInsertionOrder::InsertAtStart);
 
 	UReplicatedSubObjectOrderObject::RepOrderCounter = 0U;
 
diff --git a/Engine/Source/Runtime/Engine/Classes/Engine/NetDriver.h b/Engine/Source/Runtime/Engine/Classes/Engine/NetDriver.h
--- a/Engine/Source/Runtime/Engine/Classes/Engine/NetDriver.h
+++ b/Engine/Source/Runtime/Engine/Classes/Engine/NetDriver.h
@@ -352,7 +352,6 @@ enum class ECreateReplicationChangelistMgrFlags;
 enum class EEngineNetworkRuntimeFeatures : uint16;
 #if UE_WITH_IRIS
 class UReplicationSystem;
-class UReplicationBridge;
 namespace UE::Net
 {
 	class FNetObjectGroupHandle;
diff --git a/Engine/Source/Runtime/Engine/Private/Actor.cpp b/Engine/Source/Runtime/Engine/Private/Actor.cpp
--- a/Engine/Source/Runtime/Engine/Private/Actor.cpp
+++ b/Engine/Source/Runtime/Engine/Private/Actor.cpp
@@ -4626,7 +4626,7 @@ void AActor::SetAutonomousProxy(const bool bInAutonomousProxy, const bool bAllow
 			// Update autonomous role condition
 			if (UReplicationSystem* ReplicationSystem = UE::Net::FReplicationSystemUtil::GetReplicationSystem(this))
 			{
-				if (const UObjectReplicationBridge* Bridge = ReplicationSystem->GetReplicationBridgeAs<UObjectReplicationBridge>())
+				if (const UObjectReplicationBridge* Bridge = ReplicationSystem->GetReplicationBridge())
 				{
 					const UE::Net::FNetRefHandle RefHandle = Bridge->GetReplicatedRefHandle(this);
 					if (RefHandle.IsValid())
diff --git a/Engine/Source/Runtime/Engine/Private/ActorReplication.cpp b/Engine/Source/Runtime/Engine/Private/ActorReplication.cpp
--- a/Engine/Source/Runtime/Engine/Private/ActorReplication.cpp
+++ b/Engine/Source/Runtime/Engine/Private/ActorReplication.cpp
@@ -1083,7 +1083,7 @@ void AActor::UpdateOwningNetConnection()
 	using namespace UE::Net;
 
 	UReplicationSystem* ReplicationSystem = FReplicationSystemUtil::GetReplicationSystem(this);
-	UObjectReplicationBridge* ObjectReplicationBridge = (ReplicationSystem ? ReplicationSystem->GetReplicationBridgeAs<UObjectReplicationBridge>() : nullptr);
+	UObjectReplicationBridge* ObjectReplicationBridge = (ReplicationSystem ? ReplicationSystem->GetReplicationBridge() : nullptr);
 	if (ObjectReplicationBridge == nullptr)
 	{
 		return;
diff --git a/Engine/Source/Runtime/Engine/Private/Net/Iris/ReplicationSystem/EngineReplicationBridge.cpp b/Engine/Source/Runtime/Engine/Private/Net/Iris/ReplicationSystem/EngineReplicationBridge.cpp
--- a/Engine/Source/Runtime/Engine/Private/Net/Iris/ReplicationSystem/EngineReplicationBridge.cpp
+++ b/Engine/Source/Runtime/Engine/Private/Net/Iris/ReplicationSystem/EngineReplicationBridge.cpp
@@ -508,7 +508,7 @@ void UEngineReplicationBridge::StopReplicatingActor(AActor* Actor, EEndPlayReaso
 	// Store destruction info of static actors so their destruction can be replicated when a client loads their level in the future.
 	if (bIsServer && bIsObjectDestroyed && RefHandle.IsStatic())
 	{
-		const UObjectReplicationBridge::FDestructionParameters EndReplicationParameters
+		const UE::Net::FDestructionParameters EndReplicationParameters
 		{
 			.Location = Actor->GetActorLocation(),
 			.Level = Actor->GetLevel(),
diff --git a/Engine/Source/Runtime/Engine/Private/Net/Iris/ReplicationSystem/ReplicationSystemUtil.cpp b/Engine/Source/Runtime/Engine/Private/Net/Iris/ReplicationSystem/ReplicationSystemUtil.cpp
--- a/Engine/Source/Runtime/Engine/Private/Net/Iris/ReplicationSystem/ReplicationSystemUtil.cpp
+++ b/Engine/Source/Runtime/Engine/Private/Net/Iris/ReplicationSystem/ReplicationSystemUtil.cpp
@@ -280,7 +280,7 @@ void FReplicationSystemUtil::BeginReplicationForActorComponentSubObject(UActorCo
 							{ 
 								.RootObjectHandle = ActorRefHandle,
 								.InsertRelativeToSubObjectHandle = ActorComponentRefHandle,
-								.InsertionOrder = UReplicationBridge::ESubObjectInsertionOrder::ReplicateWith
+								.InsertionOrder = ESubObjectInsertionOrder::ReplicateWith
 							};
 							const FNetRefHandle SubObjectRefHandle = Bridge->StartReplicatingSubObject(SubObject, Params);
 							if (SubObjectRefHandle.IsValid() && NetCondition != ELifetimeCondition::COND_None)
@@ -338,7 +338,7 @@ void FReplicationSystemUtil::SetNetConditionForActorSubObject(const AActor* Acto
 		{
 			if (ReplicationSystem->IsServer())
 			{
-				if (UObjectReplicationBridge* Bridge = ReplicationSystem->GetReplicationBridgeAs<UObjectReplicationBridge>())
+				if (UObjectReplicationBridge* Bridge = ReplicationSystem->GetReplicationBridge())
 				{
 					const FNetRefHandle SubObjectRefHandle = Bridge->GetReplicatedRefHandle(SubObject);
 					if (SubObjectRefHandle.IsValid())
@@ -423,7 +423,7 @@ void FReplicationSystemUtil::RemoveDependentActor(const AActor* Parent, AActor*
 	{
 		if (ReplicationSystem->IsServer())
 		{
-			if (UObjectReplicationBridge* Bridge = ReplicationSystem->GetReplicationBridgeAs<UObjectReplicationBridge>())
+			if (UObjectReplicationBridge* Bridge = ReplicationSystem->GetReplicationBridge())
 			{
 				const FNetRefHandle ParentRefHandle = Bridge->GetReplicatedRefHandle(ParentHandle);
 				if (ParentRefHandle.IsValid())
@@ -456,7 +456,7 @@ void FReplicationSystemUtil::SetNetConditionForActorComponent(const UActorCompon
 		{
 			if (ReplicationSystem->IsServer())
 			{
-				if (UObjectReplicationBridge* Bridge = ReplicationSystem->GetReplicationBridgeAs<UObjectReplicationBridge>())
+				if (UObjectReplicationBridge* Bridge = ReplicationSystem->GetReplicationBridge())
 				{
 					const FNetRefHandle ActorComponentRefHandle = Bridge->GetReplicatedRefHandle(ActorComponent);
 					if (ActorComponentRefHandle.IsValid())
@@ -496,7 +496,7 @@ void FReplicationSystemUtil::NotifyActorDormancyChange(UReplicationSystem* Repli
 	const ENetDormancy Dormancy = Actor->NetDormancy;
 	const bool bIsPendingDormancy = (Dormancy > DORM_Awake);
 
-	if (UObjectReplicationBridge* Bridge = ReplicationSystem->GetReplicationBridgeAs<UObjectReplicationBridge>())
+	if (UObjectReplicationBridge* Bridge = ReplicationSystem->GetReplicationBridge())
 	{
 		const FNetRefHandle ActorRefHandle = Bridge->GetReplicatedRefHandle(Actor);
 		if (ActorRefHandle.IsValid())
@@ -541,7 +541,7 @@ void FReplicationSystemUtil::FlushNetDormancy(UReplicationSystem* ReplicationSys
 	const FNetHandle ActorHandle = GetNetHandle(Actor);
 	if (ActorHandle.IsValid())
 	{
-		if (UObjectReplicationBridge* Bridge = ReplicationSystem->GetReplicationBridgeAs<UObjectReplicationBridge>())
+		if (UObjectReplicationBridge* Bridge = ReplicationSystem->GetReplicationBridge())
 		{
 			const FNetRefHandle ActorRefHandle = Bridge->GetReplicatedRefHandle(ActorHandle);
 			if (ActorRefHandle.IsValid())
@@ -561,7 +561,7 @@ void FReplicationSystemUtil::UpdateSubObjectGroupMemberships(const UObject* SubO
 		{
 			if (ReplicationSystem->IsServer())
 			{
-				if (UObjectReplicationBridge* Bridge = ReplicationSystem->GetReplicationBridgeAs<UObjectReplicationBridge>())
+				if (UObjectReplicationBridge* Bridge = ReplicationSystem->GetReplicationBridge())
 				{
 					FNetRefHandle RefHandle = Bridge->GetReplicatedRefHandle(SubObject);
 					if (RefHandle.IsValid())
@@ -621,7 +621,7 @@ void FReplicationSystemUtil::SetReplicationCondition(FNetHandle NetHandle, ERepl
 	{
 		if (ReplicationSystem->IsServer())
 		{
-			if (UObjectReplicationBridge* Bridge = ReplicationSystem->GetReplicationBridgeAs<UObjectReplicationBridge>())
+			if (UObjectReplicationBridge* Bridge = ReplicationSystem->GetReplicationBridge())
 			{
 				FNetRefHandle RefHandle = Bridge->GetReplicatedRefHandle(NetHandle);
 				if (RefHandle.IsValid())
@@ -645,7 +645,7 @@ void FReplicationSystemUtil::SetStaticPriority(const AActor* Actor, float Priori
 	{
 		if (ReplicationSystem->IsServer())
 		{
-			if (UObjectReplicationBridge* Bridge = ReplicationSystem->GetReplicationBridgeAs<UObjectReplicationBridge>())
+			if (UObjectReplicationBridge* Bridge = ReplicationSystem->GetReplicationBridge())
 			{
 				const FNetRefHandle ActorRefHandle = Bridge->GetReplicatedRefHandle(ActorHandle);
 				ReplicationSystem->SetStaticPriority(ActorRefHandle, Priority);
@@ -666,7 +666,7 @@ void FReplicationSystemUtil::SetCullDistanceOverride(const AActor* Actor, float
 	{
 		if (ReplicationSystem->IsServer())
 		{
-			if (UObjectReplicationBridge* Bridge = ReplicationSystem->GetReplicationBridgeAs<UObjectReplicationBridge>())
+			if (UObjectReplicationBridge* Bridge = ReplicationSystem->GetReplicationBridge())
 			{
 				FNetRefHandle RefHandle = Bridge->GetReplicatedRefHandle(ActorHandle);
 				if (RefHandle.IsValid())
@@ -690,7 +690,7 @@ void FReplicationSystemUtil::ClearCullDistanceOverride(const AActor* Actor)
 	{
 		if (ReplicationSystem->IsServer())
 		{
-			if (UObjectReplicationBridge* Bridge = ReplicationSystem->GetReplicationBridgeAs<UObjectReplicationBridge>())
+			if (UObjectReplicationBridge* Bridge = ReplicationSystem->GetReplicationBridge())
 			{
 				FNetRefHandle RefHandle = Bridge->GetReplicatedRefHandle(ActorHandle);
 				if (RefHandle.IsValid())
@@ -714,7 +714,7 @@ void FReplicationSystemUtil::SetPollFrequency(const UObject* Object, float PollF
 	{
 		if (ReplicationSystem->IsServer())
 		{
-			if (UObjectReplicationBridge* Bridge = ReplicationSystem->GetReplicationBridgeAs<UObjectReplicationBridge>())
+			if (UObjectReplicationBridge* Bridge = ReplicationSystem->GetReplicationBridge())
 			{
 				FNetRefHandle RefHandle = Bridge->GetReplicatedRefHandle(NetHandle);
 				if (RefHandle.IsValid())
@@ -726,6 +726,6 @@ void FReplicationSystemUtil::SetPollFrequency(const UObject* Object, float PollF
 	});
 }
 
-}
+} // end namespace UE::Net
 
 #endif
diff --git a/Engine/Source/Runtime/Engine/Private/NetConnection.cpp b/Engine/Source/Runtime/Engine/Private/NetConnection.cpp
--- a/Engine/Source/Runtime/Engine/Private/NetConnection.cpp
+++ b/Engine/Source/Runtime/Engine/Private/NetConnection.cpp
@@ -1719,7 +1719,7 @@ FNetLevelVisibilityTransactionId UNetConnection::UpdateLevelStreamStatusChangedT
 		{
 			if (const ULevel* Level = LevelObject->GetLoadedLevel())
 			{
-				const UReplicationBridge* ReplicationBridge = ReplicationSystem->GetReplicationBridge();
+				const UObjectReplicationBridge* ReplicationBridge = ReplicationSystem->GetReplicationBridge();
 				if (UE::Net::FNetObjectGroupHandle GroupHandle = ReplicationBridge->GetLevelGroup(Level); GroupHandle.IsValid())
 				{
 					ReplicationSystem->SetGroupFilterStatus(GroupHandle, GetConnectionHandle().GetParentConnectionId(), UE::Net::ENetFilterStatus::Disallow);
@@ -1875,7 +1875,7 @@ void UNetConnection::UpdateLevelVisibilityInternal(const FUpdateLevelVisibilityL
 				const ULevel* Level = VisibilityResult.StreamingLevel ? VisibilityResult.StreamingLevel->GetLoadedLevel() : nullptr;
 				if (Level)
 				{
-					const UReplicationBridge* ReplicationBridge = ReplicationSystem->GetReplicationBridge();
+					const UObjectReplicationBridge* ReplicationBridge = ReplicationSystem->GetReplicationBridge();
 					if (UE::Net::FNetObjectGroupHandle GroupHandle = ReplicationBridge->GetLevelGroup(Level); GroupHandle.IsValid())
 					{
 						ReplicationSystem->SetGroupFilterStatus(GroupHandle, GetConnectionHandle().GetParentConnectionId(), UE::Net::ENetFilterStatus::Allow);
@@ -1973,7 +1973,7 @@ void UNetConnection::UpdateLevelVisibilityInternal(const FUpdateLevelVisibilityL
 					// $IRIS: $TODO: Implement forced scope-update for a group/connection to ensure that we treat replicated objects as destroyed if the filter status of a level is disabled/enabled on the same frame
 					// The reason for this is that the client currently destroys the instances rather then managing this through the replication system
 					// If we implement a way to re-instantiate instances on the client we might be able to persist the state
-					const UReplicationBridge* ReplicationBridge = ReplicationSystem->GetReplicationBridge();
+					const UObjectReplicationBridge* ReplicationBridge = ReplicationSystem->GetReplicationBridge();
 					if (UE::Net::FNetObjectGroupHandle GroupHandle = ReplicationBridge->GetLevelGroup(Level); GroupHandle.IsValid())
 					{
 						ReplicationSystem->SetGroupFilterStatus(GroupHandle, GetConnectionHandle().GetParentConnectionId(), UE::Net::ENetFilterStatus::Disallow);
diff --git a/Engine/Source/Runtime/Engine/Private/NetDriver.cpp b/Engine/Source/Runtime/Engine/Private/NetDriver.cpp
--- a/Engine/Source/Runtime/Engine/Private/NetDriver.cpp
+++ b/Engine/Source/Runtime/Engine/Private/NetDriver.cpp
@@ -992,7 +992,7 @@ FNetworkObjectInfo* UNetDriver::FindNetworkObjectInfo(const AActor* InActor)
 #if UE_WITH_IRIS
 	if (ReplicationSystem)
 	{
-		if (const UObjectReplicationBridge* Bridge = ReplicationSystem->GetReplicationBridgeAs<UObjectReplicationBridge>())
+		if (const UObjectReplicationBridge* Bridge = ReplicationSystem->GetReplicationBridge())
 		{
 			const UE::Net::FNetRefHandle ActorRefHandle = Bridge->GetReplicatedRefHandle(InActor);
 			
@@ -4310,7 +4310,7 @@ void UNetDriver::NotifyActorDestroyed(AActor* ThisActor, bool IsSeamlessTravel,
 					const UE::Net::FNetObjectReference Reference = Bridge->GetOrCreateObjectReference(ThisActor);
 					if (Reference.GetRefHandle().IsValid() && Reference.GetRefHandle().IsStatic())
 					{
-						const UReplicationBridge::FDestructionParameters Params
+						const UE::Net::FDestructionParameters Params
 						{ 
 							.Location = ThisActor->GetActorLocation(), 
 							.Level = ThisActor->GetLevel(), 
@@ -4675,7 +4675,7 @@ void UNetDriver::NotifyActorTearOff(AActor* Actor)
 #if UE_WITH_IRIS
 	if (ReplicationSystem)
 	{
-		if (UObjectReplicationBridge* Bridge = ReplicationSystem->GetReplicationBridgeAs<UObjectReplicationBridge>())
+		if (UObjectReplicationBridge* Bridge = ReplicationSystem->GetReplicationBridge())
 		{
 			// Set the actor to be torn-off during the next update of the replication system
 			const UE::Net::FNetRefHandle ActorRefHandle = Bridge->GetReplicatedRefHandle(Actor);
@@ -4709,7 +4709,7 @@ void UNetDriver::ForceNetUpdate(AActor* Actor)
 #if UE_WITH_IRIS
 	if (ReplicationSystem)
 	{
-		if (UObjectReplicationBridge* Bridge = ReplicationSystem->GetReplicationBridgeAs<UObjectReplicationBridge>())
+		if (UObjectReplicationBridge* Bridge = ReplicationSystem->GetReplicationBridge())
 		{
 			const UE::Net::FNetRefHandle ActorRefHandle = Bridge->GetReplicatedRefHandle(Actor);
 			ReplicationSystem->ForceNetUpdate(ActorRefHandle);
@@ -6570,7 +6570,7 @@ void UNetDriver::DrawNetDriverDebug()
 	const FBox BoxExpansion = FBox(&Extent, 1);
 
 #if UE_WITH_IRIS
-		const UObjectReplicationBridge* Bridge = ReplicationSystem ? ReplicationSystem->GetReplicationBridgeAs<UObjectReplicationBridge>() : nullptr;
+	const UObjectReplicationBridge* Bridge = ReplicationSystem ? ReplicationSystem->GetReplicationBridge() : nullptr;
 #endif
 
 	for (FActorIterator It(LocalWorld); It; ++It)
@@ -6833,7 +6833,7 @@ void UNetDriver::CreateReplicatedStaticActorDestructionInfo(ULevel* Level, const
 		if (UEngineReplicationBridge* Bridge = ReplicationSystem->GetReplicationBridgeAs<UEngineReplicationBridge>())
 		{
 			// Add explicit destruction info for this object
-			const UReplicationBridge::FDestructionParameters Params
+			const UE::Net::FDestructionParameters Params
 			{
 				.Location = Info.DestroyedPosition,
 				.Level = Level,				
@@ -7274,7 +7274,7 @@ bool UNetDriver::InitReplicationBridgeClass()
 
 	if (!ReplicationBridgeClassName.IsEmpty())
 	{
-		ReplicationBridgeClass = LoadClass<UReplicationBridge>(nullptr, *ReplicationBridgeClassName, nullptr, LOAD_None, nullptr);
+		ReplicationBridgeClass = LoadClass<UObjectReplicationBridge>(nullptr, *ReplicationBridgeClassName, nullptr, LOAD_None, nullptr);
 		if (ReplicationBridgeClass == nullptr)
 		{
 			UE_LOG(LogNet, Error, TEXT("Failed to load class '%s'"), *ReplicationBridgeClassName);
@@ -7326,7 +7326,7 @@ void UNetDriver::DisallowReplicationOfAllLevelGroupsForConnection(UE::Net::FConn
 		return;
 	}
 
-	const UReplicationBridge* Bridge = ReplicationSystem ? ReplicationSystem->GetReplicationBridge() : nullptr;
+	const UObjectReplicationBridge* Bridge = ReplicationSystem ? ReplicationSystem->GetReplicationBridge() : nullptr;
 	if (Bridge == nullptr)
 	{
 		return;
@@ -7347,7 +7347,7 @@ void UNetDriver::UpdateGroupFilterStatusForAllLevelsForConnection(UE::Net::FConn
 		return;
 	}
 
-	const UReplicationBridge* Bridge = ReplicationSystem ? ReplicationSystem->GetReplicationBridge() : nullptr;
+	const UObjectReplicationBridge* Bridge = ReplicationSystem ? ReplicationSystem->GetReplicationBridge() : nullptr;
 	if (Bridge == nullptr)
 	{
 		return;
@@ -7387,7 +7387,7 @@ void UNetDriver::UpdateGroupFilterStatusForAllLevelsForAllConnections()
 		return;
 	}
 
-	const UReplicationBridge* Bridge = ReplicationSystem ? ReplicationSystem->GetReplicationBridge() : nullptr;
+	const UObjectReplicationBridge* Bridge = ReplicationSystem ? ReplicationSystem->GetReplicationBridge() : nullptr;
 	if (Bridge == nullptr)
 	{
 		return;
@@ -9075,7 +9075,7 @@ UE::Net::FNetIDVariant UNetDriver::CreateNetIDForObject(UObject* Object)
 	{
 		if (ReplicationSystem)
 		{
-			UObjectReplicationBridge* ObjectBridge = ReplicationSystem->GetReplicationBridgeAs<UObjectReplicationBridge>();
+			UObjectReplicationBridge* ObjectBridge = ReplicationSystem->GetReplicationBridge();
 			if (ObjectBridge)
 			{
 				FNetRefHandle AllocatedRefHandle = ObjectBridge->PreRegisterNewObjectReferenceHandle(Object);
@@ -9111,7 +9111,7 @@ void UNetDriver::AssignNetIDToObject(UE::Net::FNetIDVariant NetID, const UObject
 		{
 			if (ReplicationSystem)
 			{
-				UObjectReplicationBridge* ObjectBridge = ReplicationSystem->GetReplicationBridgeAs<UObjectReplicationBridge>();
+				UObjectReplicationBridge* ObjectBridge = ReplicationSystem->GetReplicationBridge();
 				if (ObjectBridge)
 				{
 					ObjectBridge->PreRegisterObjectWithReferenceHandle(Object, NetID.GetVariant().Get<UE::Net::FNetRefHandle>());
diff --git a/Engine/Source/Runtime/Engine/Private/Tests/NetTestHelpers.cpp b/Engine/Source/Runtime/Engine/Private/Tests/NetTestHelpers.cpp
--- a/Engine/Source/Runtime/Engine/Private/Tests/NetTestHelpers.cpp
+++ b/Engine/Source/Runtime/Engine/Private/Tests/NetTestHelpers.cpp
@@ -260,7 +260,7 @@ FTestWorldInstance::FContext FTestWorldInstance::GetTestContext() const
 {
 #if UE_WITH_IRIS
 	UReplicationSystem* RepSystem = GetNetDriver() ? GetNetDriver()->GetReplicationSystem() : nullptr;
-	UObjectReplicationBridge* RepBridge  = RepSystem ? RepSystem->GetReplicationBridgeAs<UObjectReplicationBridge>() : nullptr;
+	UObjectReplicationBridge* RepBridge  = RepSystem ? RepSystem->GetReplicationBridge() : nullptr;
 #endif
 
 	return FContext
diff --git a/Engine/Source/Runtime/Net/Iris/Private/Iris/ReplicationSystem/ObjectReplicationBridge.cpp b/Engine/Source/Runtime/Net/Iris/Private/Iris/ReplicationSystem/ObjectReplicationBridge.cpp
--- a/Engine/Source/Runtime/Net/Iris/Private/Iris/ReplicationSystem/ObjectReplicationBridge.cpp
+++ b/Engine/Source/Runtime/Net/Iris/Private/Iris/ReplicationSystem/ObjectReplicationBridge.cpp
@@ -300,7 +300,7 @@ UE::Net::FNetObjectReference UObjectReplicationBridge::GetOrCreateObjectReferenc
 	return GetObjectReferenceCache()->GetOrCreateObjectReference(Path, Outer);
 }
 
-void UObjectReplicationBridge::AddStaticDestructionInfo(const FString& ObjectPath, const UObject* Outer, const FDestructionParameters& Parameters)
+void UObjectReplicationBridge::AddStaticDestructionInfo(const FString& ObjectPath, const UObject* Outer, const UE::Net::FDestructionParameters& Parameters)
 {
 	UE::Net::FNetObjectReference ObjectRef = GetOrCreateObjectReference(ObjectPath, Outer);
 	if (ObjectRef.IsValid())
diff --git a/Engine/Source/Runtime/Net/Iris/Private/Iris/ReplicationSystem/ReplicationBridge.cpp b/Engine/Source/Runtime/Net/Iris/Private/Iris/ReplicationSystem/ReplicationBridge.cpp
--- a/Engine/Source/Runtime/Net/Iris/Private/Iris/ReplicationSystem/ReplicationBridge.cpp
+++ b/Engine/Source/Runtime/Net/Iris/Private/Iris/ReplicationSystem/ReplicationBridge.cpp
@@ -580,8 +580,9 @@ void UReplicationBridge::DestroyLocalNetHandle(FNetRefHandle Handle, EEndReplica
 	InternalDestroyNetObject(Handle);
 }
 
-void UReplicationBridge::InternalAddSubObject(FNetRefHandle OwnerHandle, FNetRefHandle SubObjectHandle, FNetRefHandle InsertRelativeToSubObjectHandle, ESubObjectInsertionOrder InsertionOrder)
+void UReplicationBridge::InternalAddSubObject(FNetRefHandle OwnerHandle, FNetRefHandle SubObjectHandle, FNetRefHandle InsertRelativeToSubObjectHandle, UE::Net::ESubObjectInsertionOrder InsertionOrder)
 {
+	using namespace UE::Net;
 	using namespace UE::Net::Private;
 
 	EAddSubObjectFlags AddSubObjectFlags = EAddSubObjectFlags::Default;
@@ -1028,7 +1029,7 @@ void UReplicationBridge::InternalTearOff(FNetRefHandle Handle)
 	ObjectData.bShouldPropagateChangedStates = 0U;
 }
 
-UE::Net::FNetRefHandle UReplicationBridge::StoreDestructionInfo(FNetRefHandle Handle, const FDestructionParameters& Parameters)
+UE::Net::FNetRefHandle UReplicationBridge::StoreDestructionInfo(FNetRefHandle Handle, const UE::Net::FDestructionParameters& Parameters)
 {
 	using namespace UE::Net;
 	using namespace UE::Net::Private;
diff --git a/Engine/Source/Runtime/Net/Iris/Private/Iris/ReplicationSystem/ReplicationSystem.cpp b/Engine/Source/Runtime/Net/Iris/Private/Iris/ReplicationSystem/ReplicationSystem.cpp
--- a/Engine/Source/Runtime/Net/Iris/Private/Iris/ReplicationSystem/ReplicationSystem.cpp
+++ b/Engine/Source/Runtime/Net/Iris/Private/Iris/ReplicationSystem/ReplicationSystem.cpp
@@ -16,7 +16,7 @@
 #include "Iris/ReplicationState/ReplicationStateUtil.h"
 
 #include "Iris/ReplicationSystem/ChangeMaskCache.h"
-#include "Iris/ReplicationSystem/ReplicationBridge.h"
+#include "Iris/ReplicationSystem/ObjectReplicationBridge.h"
 #include "Iris/ReplicationSystem/ReplicationOperationsInternal.h"
 #include "Iris/ReplicationSystem/ReplicationReader.h"
 #include "Iris/ReplicationSystem/ReplicationSystemDelegates.h"
@@ -1400,7 +1400,7 @@ UObject* UReplicationSystem::GetConnectionUserData(uint32 ConnectionId) const
 	return nullptr;
 }
 
-UReplicationBridge* UReplicationSystem::GetReplicationBridge() const
+UObjectReplicationBridge* UReplicationSystem::GetReplicationBridge() const
 {
 	return Impl->ReplicationSystemInternal.GetReplicationBridge();
 }
diff --git a/Engine/Source/Runtime/Net/Iris/Private/Iris/ReplicationSystem/ReplicationSystemInternal.h b/Engine/Source/Runtime/Net/Iris/Private/Iris/ReplicationSystem/ReplicationSystemInternal.h
--- a/Engine/Source/Runtime/Net/Iris/Private/Iris/ReplicationSystem/ReplicationSystemInternal.h
+++ b/Engine/Source/Runtime/Net/Iris/Private/Iris/ReplicationSystem/ReplicationSystemInternal.h
@@ -14,7 +14,7 @@
 #include "Iris/ReplicationSystem/Filtering/ReplicationFiltering.h"
 #include "Iris/ReplicationSystem/Filtering/NetObjectGroups.h"
 #include "Iris/ReplicationSystem/ObjectReferenceCache.h"
-#include "Iris/ReplicationSystem/ReplicationBridge.h"
+#include "Iris/ReplicationSystem/ObjectReplicationBridge.h"
 #include "Iris/ReplicationSystem/ReplicationConnections.h"
 #include "Iris/ReplicationSystem/Prioritization/ReplicationPrioritization.h"
 #include "Iris/ReplicationSystem/ReplicationProtocolManager.h"
@@ -70,9 +70,9 @@ public:
 
 	FObjectReferenceCache& GetObjectReferenceCache() { return ObjectReferenceCache; }
 
-	void SetReplicationBridge(UReplicationBridge* InReplicationBridge) { ReplicationBridge = InReplicationBridge; }
-	UReplicationBridge* GetReplicationBridge() const { return ReplicationBridge; }
-	UReplicationBridge* GetReplicationBridge(FNetRefHandle Handle) const { return ReplicationBridge; }
+	void SetReplicationBridge(UObjectReplicationBridge* InReplicationBridge) { ReplicationBridge = InReplicationBridge; }
+	UObjectReplicationBridge* GetReplicationBridge() const { return ReplicationBridge; }
+	UObjectReplicationBridge* GetReplicationBridge(FNetRefHandle Handle) const { return ReplicationBridge; }
 
 	void SetIrisObjectReferencePackageMap(UIrisObjectReferencePackageMap* InIrisObjectReferencePackageMap) { IrisObjectReferencePackageMap = InIrisObjectReferencePackageMap; }
 	UIrisObjectReferencePackageMap* GetIrisObjectReferencePackageMap() { return IrisObjectReferencePackageMap; }
@@ -179,7 +179,7 @@ private:
 	FDirtyNetObjectTracker DirtyNetObjectTracker;
 	FReplicationStateStorage ReplicationStateStorage;
 	FReplicationStateDescriptorRegistry ReplicationStateDescriptorRegistry;
-	UReplicationBridge* ReplicationBridge;
+	UObjectReplicationBridge* ReplicationBridge;
 	UIrisObjectReferencePackageMap* IrisObjectReferencePackageMap;
 	FChangeMaskCache ChangeMaskCache;
 	FReplicationConnections Connections;
diff --git a/Engine/Source/Runtime/Net/Iris/Public/Iris/ReplicationSystem/ObjectReplicationBridge.h b/Engine/Source/Runtime/Net/Iris/Public/Iris/ReplicationSystem/ObjectReplicationBridge.h
--- a/Engine/Source/Runtime/Net/Iris/Public/Iris/ReplicationSystem/ObjectReplicationBridge.h
+++ b/Engine/Source/Runtime/Net/Iris/Public/Iris/ReplicationSystem/ObjectReplicationBridge.h
@@ -42,6 +42,7 @@ namespace UE::Net
 		typedef uint32 FInternalNetRefIndex;
 		class FObjectPollFrequencyLimiter;
 		class FObjectPoller;
+		class FReplicationSystemImpl;
 	}
 }
 
@@ -129,7 +130,7 @@ public:
 	IRISCORE_API FNetRefHandle GetRootObjectOfSubObject(FNetRefHandle SubObjectHandle) const;
 
 	/** Add static destruction info, this is used when stably named objects are destroyed prior to starting replication */
-	IRISCORE_API void AddStaticDestructionInfo(const FString& ObjectPath, const UObject* Outer, const FDestructionParameters& Parameters);
+	IRISCORE_API void AddStaticDestructionInfo(const FString& ObjectPath, const UObject* Outer, const UE::Net::FDestructionParameters& Parameters);
 	
 	/** PIE package name remapping support. */
 	virtual bool RemapPathForPIE(uint32 ConnectionId, FString& Path, bool bReading) const { return false; }
@@ -255,7 +256,7 @@ public:
 		FNetRefHandle InsertRelativeToSubObjectHandle; 
 
 		/** Optional enum to select where in the list this subobject will be inserted */
-		ESubObjectInsertionOrder InsertionOrder = ESubObjectInsertionOrder::None;
+		UE::Net::ESubObjectInsertionOrder InsertionOrder = UE::Net::ESubObjectInsertionOrder::None;
 	};
 
 public:
@@ -473,6 +474,8 @@ private:
 private:
 
 	friend UE::Net::Private::FObjectPoller;
+	friend UReplicationSystem;
+	friend UE::Net::Private::FReplicationSystemImpl;
 
 	TMap<const UClass*, FName> ConfigClassPathNameCache;
 
diff --git a/Engine/Source/Runtime/Net/Iris/Public/Iris/ReplicationSystem/ReplicationBridge.h b/Engine/Source/Runtime/Net/Iris/Public/Iris/ReplicationSystem/ReplicationBridge.h
--- a/Engine/Source/Runtime/Net/Iris/Public/Iris/ReplicationSystem/ReplicationBridge.h
+++ b/Engine/Source/Runtime/Net/Iris/Public/Iris/ReplicationSystem/ReplicationBridge.h
@@ -78,23 +78,10 @@ struct FReplicationBridgeSerializationContext
 	bool bIsDestructionInfo;
 };
 
+//------------------------------------------------------------------------
 
-
-UCLASS(Transient, MinimalAPI)
-class UReplicationBridge : public UObject
+namespace UE::Net
 {
-	GENERATED_BODY()
-
-protected:
-	using FNetHandle = UE::Net::FNetHandle;
-	using FNetRefHandle = UE::Net::FNetRefHandle;
-	using FNetDependencyInfoArray = UE::Net::FNetDependencyInfoArray;
-
-public:
-	IRISCORE_API UReplicationBridge();
-	IRISCORE_API virtual ~UReplicationBridge();
-
-
 	/** The destruction info needed to replicate the destruction event later. */
 	struct FDestructionParameters
 	{
@@ -117,6 +104,28 @@ public:
 		InsertAtStart,
 	};
 
+} // end namespace UE::Net
+
+//------------------------------------------------------------------------
+/**
+ * Base replication bridge
+ * Should only be used by internal Iris classes, public access should be restricted to the UObjectReplicationBridge.
+ */
+UCLASS(Transient, MinimalAPI)
+class UReplicationBridge : public UObject
+{
+	GENERATED_BODY()
+
+protected:
+	using FNetHandle = UE::Net::FNetHandle;
+	using FNetRefHandle = UE::Net::FNetRefHandle;
+	using FNetDependencyInfoArray = UE::Net::FNetDependencyInfoArray;
+
+public:
+	IRISCORE_API UReplicationBridge();
+	IRISCORE_API virtual ~UReplicationBridge();
+
+
 	/**
 	* Stop replicating the NetObject associated with the handle and mark the handle to be destroyed.
 	* If EEndReplication::TearOff is set the remote instance will be Torn-off rather than being destroyed on the receiving end, after the call, any state changes will not be replicated
@@ -127,7 +136,7 @@ public:
 	IRISCORE_API void StopReplicatingNetRefHandle(FNetRefHandle Handle, EEndReplicationFlags EndReplicationFlags);
 
 	/** Store destruction info for the referenced object. */
-	IRISCORE_API FNetRefHandle StoreDestructionInfo(FNetRefHandle Handle, const FDestructionParameters& Parameters);
+	IRISCORE_API FNetRefHandle StoreDestructionInfo(FNetRefHandle Handle, const UE::Net::FDestructionParameters& Parameters);
 
 	/** Returns true if the handle is replicated. */
 	IRISCORE_API bool IsReplicatedHandle(FNetRefHandle Handle) const;
@@ -240,7 +249,7 @@ protected:
 	IRISCORE_API void InternalDestroyNetObject(FNetRefHandle Handle);
 	
 	/** Add SubObjectHandle as SubObject to OwnerHandle. */
-	IRISCORE_API void InternalAddSubObject(FNetRefHandle OwnerHandle, FNetRefHandle SubObjectHandle, FNetRefHandle InsertRelativeToSubObjectHandle, ESubObjectInsertionOrder InsertionOrder);
+	IRISCORE_API void InternalAddSubObject(FNetRefHandle OwnerHandle, FNetRefHandle SubObjectHandle, FNetRefHandle InsertRelativeToSubObjectHandle, UE::Net::ESubObjectInsertionOrder InsertionOrder);
 
 	inline UE::Net::Private::FReplicationProtocolManager* GetReplicationProtocolManager() const { return ReplicationProtocolManager; }
 	inline UReplicationSystem* GetReplicationSystem() const { return ReplicationSystem; }
@@ -346,12 +355,12 @@ private:
 	void DestroyGlobalNetHandle(UE::Net::Private::FInternalNetRefIndex InternalReplicationIndex);
 	void ClearNetPushIds(UE::Net::Private::FInternalNetRefIndex InternalReplicationIndex);
 
-	friend UReplicationSystem;
 	friend UE::Net::Private::FReplicationSystemImpl;
 	friend UE::Net::Private::FReplicationSystemInternal;
 	friend UE::Net::Private::FReplicationWriter;
 	friend UE::Net::Private::FReplicationReader;
 	friend UObjectReplicationBridge;
+	friend UReplicationSystem;
 
 	UReplicationSystem* ReplicationSystem;
 	UE::Net::Private::FReplicationProtocolManager* ReplicationProtocolManager;
diff --git a/Engine/Source/Runtime/Net/Iris/Public/Iris/ReplicationSystem/ReplicationSystem.h b/Engine/Source/Runtime/Net/Iris/Public/Iris/ReplicationSystem/ReplicationSystem.h
--- a/Engine/Source/Runtime/Net/Iris/Public/Iris/ReplicationSystem/ReplicationSystem.h
+++ b/Engine/Source/Runtime/Net/Iris/Public/Iris/ReplicationSystem/ReplicationSystem.h
@@ -28,7 +28,7 @@ class UNetBlobHandler;
 class UNetDriver;
 class UNetObjectFilter;
 class UNetObjectPrioritizer;
-class UReplicationBridge;
+class UObjectReplicationBridge;
 class UReplicationSystem;
 namespace UE::Net
 {
@@ -79,7 +79,7 @@ public:
 	struct FReplicationSystemParams
 	{
 		/** The replication bridge that allows communication between the replication system and the game engine  */
-		UReplicationBridge* ReplicationBridge = nullptr;
+		UObjectReplicationBridge* ReplicationBridge = nullptr;
 
 		/** 
 		 * The maximum amount of netobjects that can be registered to the replication system 
@@ -376,8 +376,8 @@ public:
 	/** Resets all set RPCSendPolicy flags */
 	IRISCORE_API void ResetRPCSendPolicyFlags();
 
-	/** @return The UReplicationBridge that was passed with the system creation parameters. */
-	IRISCORE_API UReplicationBridge* GetReplicationBridge() const;
+	/** @return The UObjectReplicationBridge that was passed with the system creation parameters. */
+	IRISCORE_API UObjectReplicationBridge* GetReplicationBridge() const;
 
 	/** @return The UReplicationBridge that was passed with the system creation parameters. Will return nullptr if it cannot be cast to the desired type. */
 	template<typename T>
@@ -781,7 +781,7 @@ private:
 	FDelegateHandle PostGarbageCollectHandle;
 
 	UPROPERTY(transient)
-	TObjectPtr<UReplicationBridge> ReplicationBridge;
+	TObjectPtr<UObjectReplicationBridge> ReplicationBridge;
 
 	double ElapsedTime = 0;
 	uint32 Id;
