WorldPartitionBlueprintLibrary will now return actor descriptors that are child container instances.

diff --git a/Engine/Source/Runtime/Engine/Private/WorldPartitionBlueprintLibrary.cpp b/Engine/Source/Runtime/Engine/Private/WorldPartitionBlueprintLibrary.cpp
--- a/Engine/Source/Runtime/Engine/Private/WorldPartitionBlueprintLibrary.cpp
+++ b/Engine/Source/Runtime/Engine/Private/WorldPartitionBlueprintLibrary.cpp
@@ -100,6 +100,8 @@ bool UWorldPartitionBlueprintLibrary::GetActorDescs(const UActorDescContainerIns
 
 	for (UActorDescContainerInstance::TConstIterator<> Iterator(InContainerInstance); Iterator; ++Iterator)
 	{
+		OutActorDescs.Emplace(**Iterator);
+
 		if (Iterator->IsChildContainerInstance())
 		{
 			FWorldPartitionActorDesc::FContainerInstance ContainerInstance;
@@ -112,10 +114,6 @@ bool UWorldPartitionBlueprintLibrary::GetActorDescs(const UActorDescContainerIns
 				bResult = false;
 			}
 		}
-		else
-		{
-			OutActorDescs.Emplace(**Iterator);
-		}
 	}
 
 	return bResult;
@@ -125,6 +123,8 @@ bool UWorldPartitionBlueprintLibrary::HandleIntersectingActorDesc(const FWorldPa
 {
 	bool bResult = true;
 
+	OutActorDescs.Emplace(*ActorDescInstance);
+
 	if (ActorDescInstance->IsChildContainerInstance())
 	{
 		FWorldPartitionActorDesc::FContainerInstance ContainerInstance;
@@ -137,11 +137,7 @@ bool UWorldPartitionBlueprintLibrary::HandleIntersectingActorDesc(const FWorldPa
 			bResult = false;
 		}
 	}
-	else
-	{
-		OutActorDescs.Emplace(*ActorDescInstance);
-	}
-
+	
 	return bResult;
 }
 
