Gameplay framework code deprecation pass for 5.2 and earlier.
I skipped anything that would require more refactoring like BP-exposed functions.

diff --git a/Engine/Source/Runtime/Engine/Classes/GameFramework/Pawn.h b/Engine/Source/Runtime/Engine/Classes/GameFramework/Pawn.h
--- a/Engine/Source/Runtime/Engine/Classes/GameFramework/Pawn.h
+++ b/Engine/Source/Runtime/Engine/Classes/GameFramework/Pawn.h
@@ -222,10 +222,6 @@ public:
 	 */
 	ENGINE_API void SetRemoteViewPitch(double NewRemoteViewPitch);
 
-	/** Return Physics Volume for this Pawn */
-	UE_DEPRECATED(5.0, "GetPawnPhysicsVolume is deprecated. Please use GetPhysicsVolume instead.")
-	ENGINE_API virtual APhysicsVolume* GetPawnPhysicsVolume() const; 
-
 	/** Return Physics Volume for this Pawn */
 	ENGINE_API virtual APhysicsVolume* GetPhysicsVolume() const override;
 
diff --git a/Engine/Source/Runtime/Engine/Classes/Kismet/GameplayStatics.h b/Engine/Source/Runtime/Engine/Classes/Kismet/GameplayStatics.h
--- a/Engine/Source/Runtime/Engine/Classes/Kismet/GameplayStatics.h
+++ b/Engine/Source/Runtime/Engine/Classes/Kismet/GameplayStatics.h
@@ -1251,9 +1251,6 @@ public:
 	UFUNCTION(BlueprintPure, Category="Utilities|Time", meta=(WorldContext="WorldContextObject"))
 	static ENGINE_API double GetAudioTimeSeconds(const UObject* WorldContextObject);
 
-	UE_DEPRECATED(5.1, "This method has been deprecated and will be removed. Use the double version instead.")
-	static ENGINE_API void GetAccurateRealTime(int32& Seconds, float& PartialSeconds);
-
 	/** Returns time in seconds since the application was started. Unlike the other time functions this is accurate to the exact time this function is called instead of set once per frame. */
 	UFUNCTION(BlueprintPure, Category="Utilities|Time")
 	static ENGINE_API void GetAccurateRealTime(int32& Seconds, double& PartialSeconds);
diff --git a/Engine/Source/Runtime/Engine/Classes/Kismet/KismetMathLibrary.h b/Engine/Source/Runtime/Engine/Classes/Kismet/KismetMathLibrary.h
--- a/Engine/Source/Runtime/Engine/Classes/Kismet/KismetMathLibrary.h
+++ b/Engine/Source/Runtime/Engine/Classes/Kismet/KismetMathLibrary.h
@@ -212,12 +212,6 @@ class UKismetMathLibrary : public UBlueprintFunctionLibrary
 	UFUNCTION(BlueprintPure, Category = "Math|Random", meta=(Weight = "0.5", ScriptMethod = "RandomBoolWithWeight", ScriptMethodMutable))
 	static ENGINE_API bool RandomBoolWithWeightFromStream(const FRandomStream& RandomStream, float Weight);
 
-	UE_DEPRECATED(5.2, "Use RandomBoolWithWeightFromStream taking the FRandomStream as the first argument.")
-	static bool RandomBoolWithWeightFromStream(float Weight, const FRandomStream& RandomStream)
-	{
-		return RandomBoolWithWeightFromStream(RandomStream, Weight);
-	}
-
 	/** Returns the logical complement of the Boolean value (NOT A) */
 	UFUNCTION(BlueprintPure, meta=(DisplayName = "NOT Boolean", CompactNodeTitle = "NOT", Keywords = "! not negate"), Category="Math|Boolean")
 	static UE_INL_API bool Not_PreBool(bool A);
@@ -1311,9 +1305,6 @@ class UKismetMathLibrary : public UBlueprintFunctionLibrary
 	UFUNCTION(BlueprintPure, Category="Math|Vector", meta=(NativeBreakFunc))
 	static UE_INL_API void BreakVector(FVector InVec, double& X, double& Y, double& Z);
 
-	UE_DEPRECATED(5.3, "This method has been deprecated and will be removed. Use the double version instead.")
-	static UE_INL_API void BreakVector3f(FVector3f InVec, float& X, float& Y, float& Z);
-
 	/** Converts a vector to LinearColor */
 	UFUNCTION(BlueprintPure, meta=(DisplayName = "To LinearColor (Vector)", CompactNodeTitle = "->", ScriptMethod = "LinearColor", Keywords="cast convert", BlueprintAutocast), Category="Math|Conversions")
 	static UE_INL_API FLinearColor Conv_VectorToLinearColor(FVector InVec);
@@ -3085,10 +3076,6 @@ class UKismetMathLibrary : public UBlueprintFunctionLibrary
 	UFUNCTION(BlueprintCallable, meta = (ScriptMethod = "SetRandomHue"), Category = "Math|Color")
 	static UE_INL_API void LinearColor_SetRandomHue(UPARAM(ref) FLinearColor& InOutColor);
 
-	/** Converts a float into a LinearColor, where each element is a float */
-	UE_DEPRECATED(5.2, "This method has been deprecated and will be removed. Use the double version instead.")
-	static UE_INL_API FLinearColor Conv_FloatToLinearColor(float InFloat);
-
 	/** Converts a float into a LinearColor, where each RGB element is that float */
 	UFUNCTION(BlueprintPure, meta = (DisplayName = "To LinearColor (Float)", CompactNodeTitle = "->", Keywords = "cast convert", BlueprintAutocast), Category = "Math|Conversions")
 	static UE_INL_API FLinearColor Conv_DoubleToLinearColor(double InDouble);
@@ -3595,18 +3582,10 @@ class UKismetMathLibrary : public UBlueprintFunctionLibrary
 	
 	// -- Begin K2 utilities
 
-	/** Converts a byte to a float */
-	UE_DEPRECATED(5.2, "This method has been deprecated and will be removed. Use the double version instead.")
-	static UE_INL_API float Conv_ByteToFloat(uint8 InByte);
-
 	/** Converts a byte to a float */
 	UFUNCTION(BlueprintPure, meta = (DisplayName = "To Float (Byte)", CompactNodeTitle = "->", Keywords = "cast convert", BlueprintAutocast), Category = "Math|Conversions")
 	static UE_INL_API double Conv_ByteToDouble(uint8 InByte);
 
-	/** Converts an integer to a float */
-	UE_DEPRECATED(5.2, "This method has been deprecated and will be removed. Use the double version instead.")
-	static UE_INL_API float Conv_IntToFloat(int32 InInt);
-
 	/** Converts an integer to a float */
 	UFUNCTION(BlueprintPure, meta = (DisplayName = "To Float (Integer)", CompactNodeTitle = "->", Keywords = "cast convert", BlueprintAutocast), Category = "Math|Conversions")
 	static UE_INL_API double Conv_IntToDouble(int32 InInt);
@@ -3663,10 +3642,6 @@ class UKismetMathLibrary : public UBlueprintFunctionLibrary
 	UFUNCTION(BlueprintPure, meta=(DisplayName = "To Integer (Boolean)", CompactNodeTitle = "->", Keywords="cast convert", BlueprintAutocast), Category="Math|Conversions")
 	static UE_INL_API int32 Conv_BoolToInt(bool InBool);
 
-	/** Converts a bool to a float (0.0f or 1.0f) */
-	UE_DEPRECATED(5.2, "This method has been deprecated and will be removed. Use the double version instead.")
-	static UE_INL_API float Conv_BoolToFloat(bool InBool);
-
 	/** Converts a bool to a float (0.0 or 1.0) */
 	UFUNCTION(BlueprintPure, meta = (DisplayName = "To Float (Boolean)", CompactNodeTitle = "->", Keywords = "cast convert", BlueprintAutocast), Category = "Math|Conversions")
 	static UE_INL_API double Conv_BoolToDouble(bool InBool);
@@ -3695,10 +3670,6 @@ class UKismetMathLibrary : public UBlueprintFunctionLibrary
 	UFUNCTION(BlueprintPure, meta = (DisplayName = "To Vector2D (IntVector2)", CompactNodeTitle = "->", Keywords = "cast convert", BlueprintAutocast), Category = "Math|Conversions")
 	static UE_INL_API FVector2D Conv_IntVector2ToVector2D(const FIntVector2& InIntVector2);
 
-	/** Converts a float into a vector, where each element is that float */
-	UE_DEPRECATED(5.2, "This method has been deprecated and will be removed. Use the double version instead.")
-	static UE_INL_API FVector Conv_FloatToVector(float InFloat);
-
 	/** Converts a double into a vector, where each element is that float */
 	UFUNCTION(BlueprintPure, meta = (DisplayName = "To Vector (Float)", CompactNodeTitle = "->", Keywords = "cast convert", BlueprintAutocast), Category = "Math|Conversions")
 	static UE_INL_API FVector Conv_DoubleToVector(double InDouble);
@@ -4193,22 +4164,10 @@ class UKismetMathLibrary : public UBlueprintFunctionLibrary
 	UFUNCTION(BlueprintPure, Category="Math|Random", meta=(ScriptMethod="RandomInt", ScriptMethodMutable))
 	static ENGINE_API int32 RandomIntegerFromStream(const FRandomStream& Stream, int32 Max);
 
-	UE_DEPRECATED(5.2, "Use RandomIntegerFromStream taking the FRandomStream as the first argument.")
-	static int32 RandomIntegerFromStream(int32 Max, const FRandomStream& Stream)
-	{
-		return RandomIntegerFromStream(Stream, Max);
-	}
-
 	/** Return a random integer between Min and Max (>= Min and <= Max) */
 	UFUNCTION(BlueprintPure, Category="Math|Random", meta=(ScriptMethod="RandomIntInRange", ScriptMethodMutable))
 	static ENGINE_API int32 RandomIntegerInRangeFromStream(const FRandomStream& Stream, int32 Min, int32 Max);
 
-	UE_DEPRECATED(5.2, "Use RandomIntegerInRangeFromStream taking the FRandomStream as the first argument.")
-	static int32 RandomIntegerInRangeFromStream(int32 Min, int32 Max, const FRandomStream& Stream)
-	{
-		return RandomIntegerInRangeFromStream(Stream, Min, Max);
-	}
-
 	/** Returns a random bool */
 	UFUNCTION(BlueprintPure, Category="Math|Random", meta=(ScriptMethod="RandomBool", ScriptMethodMutable))
 	static ENGINE_API bool RandomBoolFromStream(const FRandomStream& Stream);
@@ -4221,12 +4180,6 @@ class UKismetMathLibrary : public UBlueprintFunctionLibrary
 	UFUNCTION(BlueprintPure, Category="Math|Random", meta=(ScriptMethod="RandomFloatInRange", ScriptMethodMutable))
 	static ENGINE_API float RandomFloatInRangeFromStream(const FRandomStream& Stream, float Min, float Max);
 
-	UE_DEPRECATED(5.2, "Use RandomFloatInRangeFromStream taking the FRandomStream as the first argument.")
-	static float RandomFloatInRangeFromStream(float Min, float Max, const FRandomStream& Stream)
-	{
-		return RandomFloatInRangeFromStream(Stream, Min, Max);
-	}
-
 	/** Returns a random vector with length of 1.0 */
 	UFUNCTION(BlueprintPure, Category="Math|Random", meta=(ScriptMethod="RandomUnitVector", ScriptMethodMutable))
 	static ENGINE_API FVector RandomUnitVectorFromStream(const FRandomStream& Stream);
@@ -4235,32 +4188,14 @@ class UKismetMathLibrary : public UBlueprintFunctionLibrary
 	UFUNCTION(BlueprintPure, Category="Math|Random", meta=(ScriptMethod="RandomPointInBoundedBox", ScriptMethodMutable))
 	static ENGINE_API FVector RandomPointInBoundingBoxFromStream(const FRandomStream& Stream, const FVector Center, const FVector HalfSize);
 
-	UE_DEPRECATED(5.2, "Use RandomPointInBoundingBoxFromStream taking the FRandomStream as the first argument.")
-	static FVector RandomPointInBoundingBoxFromStream(const FVector Center, const FVector HalfSize, const FRandomStream& Stream)
-	{
-		return RandomPointInBoundingBoxFromStream(Stream, Center, HalfSize);
-	}
-
 	/** Returns a random point within the specified bounding box. */
 	UFUNCTION(BlueprintPure, Category="Math|Random", meta=(DisplayName="Random Point In Bounding Box From Stream (Box)"), meta=(ScriptMethod="RandomPointInBox", ScriptMethodMutable))
 	static ENGINE_API FVector RandomPointInBoundingBoxFromStream_Box(const FRandomStream& Stream, const FBox Box);
 
-	UE_DEPRECATED(5.2, "Use RandomPointInBoundingBoxFromStream_Box taking the FRandomStream as the first argument.")
-	static FVector RandomPointInBoundingBoxFromStream_Box(const FBox Box, const FRandomStream& Stream)
-	{
-		return RandomPointInBoundingBoxFromStream_Box(Stream, Box);
-	}
-
 	/** Create a random rotation */
 	UFUNCTION(BlueprintPure, Category="Math|Random", meta=(ScriptMethod="RandomRotator", ScriptMethodMutable))
 	static ENGINE_API FRotator RandomRotatorFromStream(const FRandomStream& Stream, bool bRoll);
 
-	UE_DEPRECATED(5.2, "Use RandomRotatorFromStream taking the FRandomStream as the first argument.")
-	static FRotator RandomRotatorFromStream(bool bRoll, const FRandomStream& Stream)
-	{
-		return RandomRotatorFromStream(Stream, bRoll);
-	}
-
 	/** Reset a random stream */
 	UFUNCTION(BlueprintCallable, Category="Math|Random", meta=(ScriptMethod="Reset", ScriptMethodMutable))
 	static ENGINE_API void ResetRandomStream(const FRandomStream& Stream);
@@ -4282,12 +4217,6 @@ class UKismetMathLibrary : public UBlueprintFunctionLibrary
 	UFUNCTION(BlueprintPure, Category="Math|Random", meta = (Keywords = "RandomVector", ScriptMethod = "RandomUnitVectorInConeInRadians", ScriptMethodMutable))
 	static ENGINE_API FVector RandomUnitVectorInConeInRadiansFromStream(const FRandomStream& Stream, const FVector& ConeDir, float ConeHalfAngleInRadians);
 
-	UE_DEPRECATED(5.2, "Use RandomUnitVectorInConeInRadiansFromStream taking the FRandomStream as the first argument.")
-	static FVector RandomUnitVectorInConeInRadiansFromStream(const FVector& ConeDir, float ConeHalfAngleInRadians, const FRandomStream& Stream)
-	{
-		return RandomUnitVectorInConeInRadiansFromStream(Stream, ConeDir, ConeHalfAngleInRadians);
-	}
-
 	/** 
 	 * Returns a random vector with length of 1, within the specified cone, with uniform random distribution.
 	 * @param ConeDir					The base "center" direction of the cone.
@@ -4300,12 +4229,6 @@ class UKismetMathLibrary : public UBlueprintFunctionLibrary
 		return RandomUnitVectorInConeInRadiansFromStream(Stream, ConeDir, FMath::DegreesToRadians(ConeHalfAngleInDegrees));
 	}
 
-	UE_DEPRECATED(5.2, "Use RandomUnitVectorInConeInDegreesFromStream taking the FRandomStream as the first argument.")
-	static FVector RandomUnitVectorInConeInDegreesFromStream(const FVector& ConeDir, float ConeHalfAngleInDegrees, const FRandomStream& Stream)
-	{
-		return RandomUnitVectorInConeInDegreesFromStream(Stream, ConeDir, ConeHalfAngleInDegrees);
-	}
-
 	/**
 	* Returns a random vector with length of 1, within the specified cone, with uniform random distribution.
 	* The shape of the cone can be modified according to the yaw and pitch angles.
@@ -4317,12 +4240,6 @@ class UKismetMathLibrary : public UBlueprintFunctionLibrary
 	UFUNCTION(BlueprintPure, Category = "Math|Random", meta = (Keywords = "RandomVector", ScriptMethod = "RandomUnitVectorInEllipticalConeInRadians", ScriptMethodMutable))
 	static ENGINE_API FVector RandomUnitVectorInEllipticalConeInRadiansFromStream(const FRandomStream& Stream, const FVector& ConeDir, float MaxYawInRadians, float MaxPitchInRadians);
 
-	UE_DEPRECATED(5.2, "Use RandomUnitVectorInEllipticalConeInRadiansFromStream taking the FRandomStream as the first argument.")
-	static FVector RandomUnitVectorInEllipticalConeInRadiansFromStream(const FVector& ConeDir, float MaxYawInRadians, float MaxPitchInRadians, const FRandomStream& Stream)
-	{
-		return RandomUnitVectorInEllipticalConeInRadiansFromStream(Stream, ConeDir, MaxYawInRadians, MaxPitchInRadians);
-	}
-
 	/**
 	* Returns a random vector with length of 1, within the specified cone, with uniform random distribution.
 	* The shape of the cone can be modified according to the yaw and pitch angles.
@@ -4337,12 +4254,6 @@ class UKismetMathLibrary : public UBlueprintFunctionLibrary
 		return RandomUnitVectorInEllipticalConeInRadiansFromStream(Stream, ConeDir, FMath::DegreesToRadians(MaxYawInDegrees), FMath::DegreesToRadians(MaxPitchInDegrees));
 	}
 
-	UE_DEPRECATED(5.2, "Use RandomUnitVectorInEllipticalConeInDegreesFromStream taking the FRandomStream as the first argument.")
-	static FVector RandomUnitVectorInEllipticalConeInDegreesFromStream(const FVector& ConeDir, float MaxYawInDegrees, float MaxPitchInDegrees, const FRandomStream& Stream)
-	{
-		return RandomUnitVectorInEllipticalConeInDegreesFromStream(Stream, ConeDir, MaxYawInDegrees, MaxPitchInDegrees);
-	}
-
 	/**
 	 * Generates a 1D Perlin noise from the given value.  Returns a continuous random value between -1.0 and 1.0.
 	 *
diff --git a/Engine/Source/Runtime/Engine/Classes/Kismet/KismetMathLibrary.inl b/Engine/Source/Runtime/Engine/Classes/Kismet/KismetMathLibrary.inl
--- a/Engine/Source/Runtime/Engine/Classes/Kismet/KismetMathLibrary.inl
+++ b/Engine/Source/Runtime/Engine/Classes/Kismet/KismetMathLibrary.inl
@@ -1408,14 +1408,6 @@ void UKismetMathLibrary::BreakVector(FVector InVec, double& X, double& Y, double
 	Z = InVec.Z;
 }
 
-KISMET_MATH_INLINE
-void UKismetMathLibrary::BreakVector3f(FVector3f InVec, float& X, float& Y, float& Z)
-{
-	X = InVec.X;
-	Y = InVec.Y;
-	Z = InVec.Z;
-}
-
 KISMET_MATH_INLINE
 FRotator UKismetMathLibrary::Conv_VectorToRotator(FVector InVec)
 {
@@ -2627,24 +2619,12 @@ bool UKismetMathLibrary::TimespanFromString(FString TimespanString, FTimespan& R
 /* K2 Utilities
 *****************************************************************************/
 
-KISMET_MATH_INLINE
-float UKismetMathLibrary::Conv_ByteToFloat(uint8 InByte)
-{
-	return (float)InByte;
-}
-
 KISMET_MATH_INLINE
 double UKismetMathLibrary::Conv_ByteToDouble(uint8 InByte)
 {
 	return (double)InByte;
 }
 
-KISMET_MATH_INLINE
-float UKismetMathLibrary::Conv_IntToFloat(int32 InInt)
-{
-	return (float)InInt;
-}
-
 KISMET_MATH_INLINE
 double UKismetMathLibrary::Conv_IntToDouble(int32 InInt)
 {
@@ -2735,12 +2715,6 @@ uint8 UKismetMathLibrary::Conv_BoolToByte(bool InBool)
 	return InBool ? 1 : 0;
 }
 
-KISMET_MATH_INLINE
-float UKismetMathLibrary::Conv_BoolToFloat(bool InBool)
-{
-	return InBool ? 1.0f : 0.0f;
-}
-
 KISMET_MATH_INLINE
 double UKismetMathLibrary::Conv_BoolToDouble(bool InBool)
 {
@@ -2789,12 +2763,6 @@ FLinearColor UKismetMathLibrary::Conv_ColorToLinearColor(FColor InColor)
 	return FLinearColor(InColor);
 }
 
-KISMET_MATH_INLINE
-FVector UKismetMathLibrary::Conv_FloatToVector(float InFloat)
-{
-	return FVector(InFloat);
-}
-
 KISMET_MATH_INLINE
 FVector UKismetMathLibrary::Conv_DoubleToVector(double InDouble)
 {
@@ -2807,12 +2775,6 @@ FVector2D UKismetMathLibrary::Conv_DoubleToVector2D(double InDouble)
 	return FVector2D(InDouble);
 }
 
-KISMET_MATH_INLINE
-FLinearColor UKismetMathLibrary::Conv_FloatToLinearColor(float InFloat)
-{
-	return FLinearColor(InFloat, InFloat, InFloat);
-}
-
 KISMET_MATH_INLINE
 FLinearColor UKismetMathLibrary::Conv_DoubleToLinearColor(double InDouble)
 {
diff --git a/Engine/Source/Runtime/Engine/Classes/Kismet/KismetStringLibrary.h b/Engine/Source/Runtime/Engine/Classes/Kismet/KismetStringLibrary.h
--- a/Engine/Source/Runtime/Engine/Classes/Kismet/KismetStringLibrary.h
+++ b/Engine/Source/Runtime/Engine/Classes/Kismet/KismetStringLibrary.h
@@ -12,10 +12,6 @@ class UKismetStringLibrary : public UBlueprintFunctionLibrary
 {
 	GENERATED_UCLASS_BODY()
 
-	/** Converts a float value to a string */
-	UE_DEPRECATED(5.0, "This method has been deprecated and will be removed. Please use FString::SanitizeFloat() or LexToSanitzedString instead.")
-	static ENGINE_API FString Conv_FloatToString(float InFloat);
-
 	/** Converts a double value to a string */
 	UFUNCTION(BlueprintPure, meta = (DisplayName = "To String (Float)", CompactNodeTitle = "->", BlueprintAutocast), Category = "Utilities|String")
 	static ENGINE_API FString Conv_DoubleToString(double InDouble);
@@ -112,10 +108,6 @@ class UKismetStringLibrary : public UBlueprintFunctionLibrary
 	UFUNCTION(BlueprintPure, meta = (DisplayName = "String To Integer64", CompactNodeTitle = "->", BlueprintAutocast), Category = "Utilities|String")
 	static ENGINE_API int64 Conv_StringToInt64(const FString& InString);
 
-	/** Converts a string to a float value */
-	UE_DEPRECATED(5.0, "This method has been deprecated and will be removed. Please use LexFromString instead.")
-	static ENGINE_API float Conv_StringToFloat(const FString& InString);
-
 	/** Converts a string to a double value */
 	UFUNCTION(BlueprintPure, meta = (DisplayName = "String To Float", CompactNodeTitle = "->", BlueprintAutocast), Category = "Utilities|String")
 	static ENGINE_API double Conv_StringToDouble(const FString& InString);
@@ -140,17 +132,6 @@ class UKismetStringLibrary : public UBlueprintFunctionLibrary
 	UFUNCTION(BlueprintPure, meta=(DisplayName = "String To Color", CompactNodeTitle = "->"), Category="Utilities|String")
 	static ENGINE_API void Conv_StringToColor(const FString& InString, FLinearColor& OutConvertedColor, bool& OutIsValid);
 
-	/** 
-	 * Converts a float->string, create a new string in the form AppendTo+Prefix+InFloat+Suffix
-	 * @param AppendTo - An existing string to use as the start of the conversion string
-	 * @param Prefix - A string to use as a prefix, after the AppendTo string
-	 * @param InFloat - The float value to convert
-	 * @param Suffix - A suffix to append to the end of the conversion string
-	 * @return A new string built from the passed parameters
-	 */
-	UE_DEPRECATED(5.0, "This method has been deprecated and will be removed. Please use FString::Printf instead.")
-	static ENGINE_API FString BuildString_Float(const FString& AppendTo, const FString& Prefix, float InFloat, const FString& Suffix);
-
 	/**
 	 * Converts a double->string, create a new string in the form AppendTo+Prefix+InDouble+Suffix
 	 * @param AppendTo - An existing string to use as the start of the conversion string
diff --git a/Engine/Source/Runtime/Engine/Classes/Kismet/KismetSystemLibrary.h b/Engine/Source/Runtime/Engine/Classes/Kismet/KismetSystemLibrary.h
--- a/Engine/Source/Runtime/Engine/Classes/Kismet/KismetSystemLibrary.h
+++ b/Engine/Source/Runtime/Engine/Classes/Kismet/KismetSystemLibrary.h
@@ -980,10 +980,6 @@ class UKismetSystemLibrary : public UBlueprintFunctionLibrary
 	UFUNCTION(BlueprintCallable, meta=(BlueprintInternalUseOnly = "true"))
 	static ENGINE_API void SetBytePropertyByName(UObject* Object, FName PropertyName, uint8 Value);
 
-	/** Set a float property by name */
-	UE_DEPRECATED(5.0, "This method has been deprecated and will be removed. Use the double version instead.")
-	static ENGINE_API void SetFloatPropertyByName(UObject* Object, FName PropertyName, float Value);
-
 	/** Set a double property by name */
 	UFUNCTION(BlueprintCallable, meta = (BlueprintInternalUseOnly = "true"))
 	static ENGINE_API void SetDoublePropertyByName(UObject* Object, FName PropertyName, double Value);
diff --git a/Engine/Source/Runtime/Engine/Classes/Kismet/KismetTextLibrary.h b/Engine/Source/Runtime/Engine/Classes/Kismet/KismetTextLibrary.h
--- a/Engine/Source/Runtime/Engine/Classes/Kismet/KismetTextLibrary.h
+++ b/Engine/Source/Runtime/Engine/Classes/Kismet/KismetTextLibrary.h
@@ -247,11 +247,6 @@ class UKismetTextLibrary : public UBlueprintFunctionLibrary
 	UFUNCTION(BlueprintPure, meta = (DisplayName = "To Text (Float)", AdvancedDisplay = "1", BlueprintAutocast), Category = "Utilities|Text")
 	static ENGINE_API FText Conv_DoubleToText(double Value, TEnumAsByte<ERoundingMode> RoundingMode, bool bAlwaysSign = false, bool bUseGrouping = true, int32 MinimumIntegralDigits = 1, int32 MaximumIntegralDigits = 324, int32 MinimumFractionalDigits = 0, int32 MaximumFractionalDigits = 3);
 	
-	// Default values are duplicated from FNumberFormattingOptions and should be replicated in all functions and in the struct when changed!
-	/** Converts a passed in float to text based on formatting options */
-	UE_DEPRECATED(5.1, "This method has been deprecated and will be removed.")
-	static ENGINE_API FText Conv_FloatToText(float Value, TEnumAsByte<ERoundingMode> RoundingMode, bool bAlwaysSign = false, bool bUseGrouping = true, int32 MinimumIntegralDigits = 1, int32 MaximumIntegralDigits = 324, int32 MinimumFractionalDigits = 0, int32 MaximumFractionalDigits = 3);
-	
 	/**
 	 * Generate an FText that represents the passed number as currency in the current culture.
 	 * BaseVal is specified in the smallest fractional value of the currency and will be converted for formatting according to the selected culture.
diff --git a/Engine/Source/Runtime/Engine/Private/GameplayStatics.cpp b/Engine/Source/Runtime/Engine/Private/GameplayStatics.cpp
--- a/Engine/Source/Runtime/Engine/Private/GameplayStatics.cpp
+++ b/Engine/Source/Runtime/Engine/Private/GameplayStatics.cpp
@@ -2585,13 +2585,6 @@ double UGameplayStatics::GetAudioTimeSeconds(const UObject* WorldContextObject)
 	return World ? World->GetAudioTimeSeconds() : 0.0;
 }
 
-void UGameplayStatics::GetAccurateRealTime(int32& Seconds, float& PartialSeconds)
-{
-	double TimeSeconds = FPlatformTime::Seconds() - GStartTime;
-	Seconds = floor(TimeSeconds);
-	PartialSeconds = TimeSeconds - double(Seconds);
-}
-
 void UGameplayStatics::GetAccurateRealTime(int32& Seconds, double& PartialSeconds)
 {
 	double TimeSeconds = FPlatformTime::Seconds() - GStartTime;
diff --git a/Engine/Source/Runtime/Engine/Private/KismetStringLibrary.cpp b/Engine/Source/Runtime/Engine/Private/KismetStringLibrary.cpp
--- a/Engine/Source/Runtime/Engine/Private/KismetStringLibrary.cpp
+++ b/Engine/Source/Runtime/Engine/Private/KismetStringLibrary.cpp
@@ -61,11 +61,6 @@ bool UKismetStringLibrary::IsEmpty(const FString& InString)
 	return InString.IsEmpty();
 }
 
-FString UKismetStringLibrary::Conv_FloatToString(float InFloat)
-{
-	return FString::SanitizeFloat(InFloat);
-}
-
 FString UKismetStringLibrary::Conv_DoubleToString(double InDouble)
 {
 	return FString::SanitizeFloat(InDouble);
@@ -190,11 +185,6 @@ int64 UKismetStringLibrary::Conv_StringToInt64(const FString& InString)
 	return FCString::Atoi64(*InString);
 }
 
-float UKismetStringLibrary::Conv_StringToFloat(const FString& InString)
-{
-	return FCString::Atof(*InString);
-}
-
 double UKismetStringLibrary::Conv_StringToDouble(const FString& InString)
 {
 	return FCString::Atod(*InString);
@@ -225,11 +215,6 @@ void UKismetStringLibrary::Conv_StringToColor(const FString& InString, FLinearCo
 	OutIsValid = OutConvertedColor.InitFromString(InString);
 }
 
-FString UKismetStringLibrary::BuildString_Float(const FString& AppendTo, const FString& Prefix, float InFloat, const FString& Suffix)
-{
-	return BuildString_Double(AppendTo, Prefix, InFloat, Suffix);
-}
-
 FString UKismetStringLibrary::BuildString_Double(const FString& AppendTo, const FString& Prefix, double InDouble, const FString& Suffix)
 {
 	// despite the name, SanitizeFloat takes a double parameter
diff --git a/Engine/Source/Runtime/Engine/Private/KismetSystemLibrary.cpp b/Engine/Source/Runtime/Engine/Private/KismetSystemLibrary.cpp
--- a/Engine/Source/Runtime/Engine/Private/KismetSystemLibrary.cpp
+++ b/Engine/Source/Runtime/Engine/Private/KismetSystemLibrary.cpp
@@ -1153,18 +1153,6 @@ void UKismetSystemLibrary::SetBytePropertyByName(UObject* Object, FName Property
 	}
 }
 
-void UKismetSystemLibrary::SetFloatPropertyByName(UObject* Object, FName PropertyName, float Value)
-{
-	if(Object != NULL)
-	{
-		FFloatProperty* FloatProp = FindFProperty<FFloatProperty>(Object->GetClass(), PropertyName);
-		if(FloatProp != NULL)
-		{
-			FloatProp->SetPropertyValue_InContainer(Object, Value);
-		}		
-	}
-}
-
 void UKismetSystemLibrary::SetDoublePropertyByName(UObject* Object, FName PropertyName, double Value)
 {
 	if (Object != nullptr)
diff --git a/Engine/Source/Runtime/Engine/Private/KismetTextLibrary.cpp b/Engine/Source/Runtime/Engine/Private/KismetTextLibrary.cpp
--- a/Engine/Source/Runtime/Engine/Private/KismetTextLibrary.cpp
+++ b/Engine/Source/Runtime/Engine/Private/KismetTextLibrary.cpp
@@ -223,11 +223,6 @@ FText UKismetTextLibrary::Conv_DoubleToText(double Value, TEnumAsByte<ERoundingM
 	return FText::AsNumber(Value, &NumberFormatOptions);
 }
 
-FText UKismetTextLibrary::Conv_FloatToText(float Value, TEnumAsByte<ERoundingMode> RoundingMode, bool bAlwaysSign/* = false*/, bool bUseGrouping/* = true*/, int32 MinimumIntegralDigits/* = 1*/, int32 MaximumIntegralDigits/* = 324*/, int32 MinimumFractionalDigits/* = 0*/, int32 MaximumFractionalDigits/* = 3*/)
-{
-	return Conv_DoubleToText(Value, RoundingMode, bAlwaysSign, bUseGrouping, MinimumIntegralDigits, MaximumIntegralDigits, MinimumFractionalDigits, MaximumFractionalDigits);
-}
-
 FText UKismetTextLibrary::AsCurrencyBase(int32 BaseValue, const FString& CurrencyCode)
 {
 	return FText::AsCurrencyBase(BaseValue, CurrencyCode);
diff --git a/Engine/Source/Runtime/Engine/Private/Pawn.cpp b/Engine/Source/Runtime/Engine/Private/Pawn.cpp
--- a/Engine/Source/Runtime/Engine/Private/Pawn.cpp
+++ b/Engine/Source/Runtime/Engine/Private/Pawn.cpp
@@ -894,11 +894,6 @@ void APawn::Restart()
 	RecalculateBaseEyeHeight();
 }
 
-APhysicsVolume* APawn::GetPawnPhysicsVolume() const
-{
-	return GetPhysicsVolume();
-}
-
 APhysicsVolume* APawn::GetPhysicsVolume() const
 {
 	if (const UPawnMovementComponent* MovementComponent = GetMovementComponent())
diff --git a/Engine/Source/Runtime/Engine/Public/CanvasTypes.h b/Engine/Source/Runtime/Engine/Public/CanvasTypes.h
--- a/Engine/Source/Runtime/Engine/Public/CanvasTypes.h
+++ b/Engine/Source/Runtime/Engine/Public/CanvasTypes.h
@@ -240,12 +240,6 @@ public:
 
 	ENGINE_API static FCanvas* Create(FRDGBuilder& GraphBuilder, FRDGTextureRef InRenderTarget, FHitProxyConsumer* InHitProxyConsumer, const FGameTime& Time, ERHIFeatureLevel::Type InFeatureLevel, float InDPIScale = 1.0f);
 
-	UE_DEPRECATED(5.0, "Pass down a FGameTime instead.")
-	inline static FCanvas* Create(FRDGBuilder& GraphBuilder, FRDGTextureRef InRenderTarget, FHitProxyConsumer* InHitProxyConsumer, float InRealTime, float InWorldTime, float InWorldDeltaTime, ERHIFeatureLevel::Type InFeatureLevel, float InDPIScale = 1.0f)
-	{
-		return Create(GraphBuilder, InRenderTarget, InHitProxyConsumer, FGameTime::CreateDilated(InRealTime, InWorldDeltaTime, InWorldTime, InWorldDeltaTime), InFeatureLevel, InDPIScale);
-	}
-
 	/**
 	* Constructor.
 	*/
@@ -256,11 +250,6 @@ public:
 	*/
 	ENGINE_API FCanvas(FRenderTarget* InRenderTarget, FHitProxyConsumer* InHitProxyConsumer, const FGameTime& Time, ERHIFeatureLevel::Type InFeatureLevel, float InDPIScale = 1.0f);
 	
-	UE_DEPRECATED(5.0, "Pass down a FGameTime instead.")
-	inline FCanvas(FRenderTarget* InRenderTarget, FHitProxyConsumer* InHitProxyConsumer, float InRealTime, float InWorldTime, float InWorldDeltaTime, ERHIFeatureLevel::Type InFeatureLevel, float InDPIScale = 1.0f)
-		: FCanvas(InRenderTarget, InHitProxyConsumer, FGameTime::CreateDilated(InRealTime, InWorldDeltaTime, InWorldTime, InWorldDeltaTime), InFeatureLevel, InDPIScale)
-	{ }
-
 	/**
 	* Destructor.
 	*/
@@ -666,15 +655,6 @@ public:
 	{
 		return Time;
 	}
-
-	UE_DEPRECATED(5.0, "Use FCanvas::GetTime()")
-	float GetCurrentRealTime() const { return FloatCastChecked<float>(GetTime().GetRealTimeSeconds(), UE_DOUBLE_SMALL_NUMBER); }
-
-	UE_DEPRECATED(5.0, "Use FCanvas::GetTime()")
-	float GetCurrentWorldTime() const { return FloatCastChecked<float>(GetTime().GetWorldTimeSeconds(), UE_DOUBLE_SMALL_NUMBER); }
-
-	UE_DEPRECATED(5.0, "Use FCanvas::GetTime()")
-	float GetCurrentDeltaWorldTime() const { return GetTime().GetDeltaWorldTimeSeconds(); }
 	
 	/** 
 	 * Draw a CanvasItem
diff --git a/Engine/Source/Runtime/Engine/Public/GameDelegates.h b/Engine/Source/Runtime/Engine/Public/GameDelegates.h
--- a/Engine/Source/Runtime/Engine/Public/GameDelegates.h
+++ b/Engine/Source/Runtime/Engine/Public/GameDelegates.h
@@ -20,8 +20,6 @@ enum class EGameDelegates_SaveGame : short
 	Detail,	
 };
 
-// UE_DEPRECATED(5.0, "Use AssetManager or FModifyCookDelegate instead.")
-DECLARE_DELEGATE_OneParam(FCookModificationDelegate, TArray<FString>& /*ExtraPackagesToCook*/);
 /** Delegate to modify cooking behavior - return extra packages to cook, load up the asset registry, etc */
 DECLARE_MULTICAST_DELEGATE_ThreeParams(FModifyCookDelegate,
 	TConstArrayView<const ITargetPlatform*> /* InTargetPlatforms */, TArray<FName>& /* InOutPackagesToCook */,
diff --git a/Engine/Source/Runtime/GameplayTags/Classes/GameplayTagContainer.h b/Engine/Source/Runtime/GameplayTags/Classes/GameplayTagContainer.h
--- a/Engine/Source/Runtime/GameplayTags/Classes/GameplayTagContainer.h
+++ b/Engine/Source/Runtime/GameplayTags/Classes/GameplayTagContainer.h
@@ -23,15 +23,6 @@ DECLARE_STATS_GROUP(TEXT("Gameplay Tags"), STATGROUP_GameplayTags, STATCAT_Advan
 
 DECLARE_CYCLE_STAT_EXTERN(TEXT("FGameplayTagContainer::DoesTagContainerMatch"), STAT_FGameplayTagContainer_DoesTagContainerMatch, STATGROUP_GameplayTags, GAMEPLAYTAGS_API);
 
-namespace UE_DEPRECATED(5.0, "Deprecated in favor of HasExact and related functions") EGameplayTagMatchType
-{
-	enum Type
-	{
-		Explicit,			// This will check for a match against just this tag
-		IncludeParentTags,	// This will also check for matches against all parent tags
-	};
-}
-
 UENUM(BlueprintType)
 enum class EGameplayContainerMatchType : uint8
 {
@@ -615,69 +606,6 @@ struct FGameplayTagContainer
 	/** An empty Gameplay Tag Container */
 	static GAMEPLAYTAGS_API const FGameplayTagContainer EmptyContainer;
 
-PRAGMA_DISABLE_DEPRECATION_WARNINGS
-
-	UE_DEPRECATED(5.0, "Deprecated in favor of HasTag or HasTagExact")
-	inline bool HasTagFast(FGameplayTag const& TagToCheck, TEnumAsByte<EGameplayTagMatchType::Type> TagMatchType, TEnumAsByte<EGameplayTagMatchType::Type> TagToCheckMatchType) const
-	{
-		bool bResult;
-		if (TagToCheckMatchType == EGameplayTagMatchType::Explicit)
-		{
-			// Always check explicit
-			bResult = GameplayTags.Contains(TagToCheck);
-
-			if (!bResult && TagMatchType == EGameplayTagMatchType::IncludeParentTags)
-			{
-				// Check parent tags as well
-				bResult = ParentTags.Contains(TagToCheck);
-			}
-		}
-		else
-		{
-			bResult = ComplexHasTag(TagToCheck, TagMatchType, TagToCheckMatchType);
-		}
-		return bResult;
-	}
-
-	UE_DEPRECATED(5.0, "Deprecated in favor of HasTag or HasTagExact")
-	GAMEPLAYTAGS_API bool ComplexHasTag(FGameplayTag const& TagToCheck, TEnumAsByte<EGameplayTagMatchType::Type> TagMatchType, TEnumAsByte<EGameplayTagMatchType::Type> TagToCheckMatchType) const;
-
-	UE_DEPRECATED(5.0, "Deprecated in favor of HasAll and related functions")
-	inline bool DoesTagContainerMatch(const FGameplayTagContainer& OtherContainer, TEnumAsByte<EGameplayTagMatchType::Type> TagMatchType, TEnumAsByte<EGameplayTagMatchType::Type> OtherTagMatchType, EGameplayContainerMatchType ContainerMatchType) const
-	{
-		SCOPE_CYCLE_COUNTER(STAT_FGameplayTagContainer_DoesTagContainerMatch);
-		bool bResult;
-		if (OtherTagMatchType == EGameplayTagMatchType::Explicit)
-		{
-			// Start true for all, start false for any
-			bResult = (ContainerMatchType == EGameplayContainerMatchType::All);
-			for (const FGameplayTag& OtherTag : OtherContainer.GameplayTags)
-			{
-				if (HasTagFast(OtherTag, TagMatchType, OtherTagMatchType))
-				{
-					if (ContainerMatchType == EGameplayContainerMatchType::Any)
-					{
-						bResult = true;
-						break;
-					}
-				}
-				else if (ContainerMatchType == EGameplayContainerMatchType::All)
-				{
-					bResult = false;
-					break;
-				}
-			}			
-		}
-		else
-		{
-			FGameplayTagContainer OtherExpanded = OtherContainer.GetGameplayTagParents();
-			return DoesTagContainerMatch(OtherExpanded, TagMatchType, EGameplayTagMatchType::Explicit, ContainerMatchType);
-		}
-		return bResult;
-	}
-
-PRAGMA_ENABLE_DEPRECATION_WARNINGS
-
 protected:
 
 	/**
diff --git a/Engine/Source/Runtime/GameplayTags/Private/GameplayTagContainer.cpp b/Engine/Source/Runtime/GameplayTags/Private/GameplayTagContainer.cpp
--- a/Engine/Source/Runtime/GameplayTags/Private/GameplayTagContainer.cpp
+++ b/Engine/Source/Runtime/GameplayTags/Private/GameplayTagContainer.cpp
@@ -646,25 +646,6 @@ bool FGameplayTagContainer::operator!=(FGameplayTagContainer const& Other) const
 	return !operator==(Other);
 }
 
-PRAGMA_DISABLE_DEPRECATION_WARNINGS
-
-bool FGameplayTagContainer::ComplexHasTag(FGameplayTag const& TagToCheck, TEnumAsByte<EGameplayTagMatchType::Type> TagMatchType, TEnumAsByte<EGameplayTagMatchType::Type> TagToCheckMatchType) const
-{
-	check(TagMatchType != EGameplayTagMatchType::Explicit || TagToCheckMatchType != EGameplayTagMatchType::Explicit);
-
-	if (TagMatchType == EGameplayTagMatchType::IncludeParentTags)
-	{
-		FGameplayTagContainer ExpandedConatiner = GetGameplayTagParents();
-		return ExpandedConatiner.HasTagFast(TagToCheck, EGameplayTagMatchType::Explicit, TagToCheckMatchType);
-	}
-	else
-	{
-		return TagToCheck.GetSingleTagContainer().DoesTagContainerMatch(*this, EGameplayTagMatchType::IncludeParentTags, EGameplayTagMatchType::Explicit, EGameplayContainerMatchType::Any);
-	}
-}
-
-PRAGMA_ENABLE_DEPRECATION_WARNINGS
-
 bool FGameplayTagContainer::RemoveTagByExplicitName(const FName& TagName)
 {
 	SCOPE_CYCLE_COUNTER(STAT_FGameplayTagContainer_RemoveTagByExplicitName);
diff --git a/Engine/Source/Runtime/GameplayTasks/Classes/GameplayTasksComponent.h b/Engine/Source/Runtime/GameplayTasks/Classes/GameplayTasksComponent.h
--- a/Engine/Source/Runtime/GameplayTasks/Classes/GameplayTasksComponent.h
+++ b/Engine/Source/Runtime/GameplayTasks/Classes/GameplayTasksComponent.h
@@ -82,9 +82,6 @@ protected:
 		return SimulatedTasks;
 	}
 
-	UE_DEPRECATED(5.1, "This will be removed in future versions. Use AddSimulatedTask or RemoveSimulatedTask to modify the array")
-	GAMEPLAYTASKS_API TArray<TObjectPtr<UGameplayTask>>& GetSimulatedTasks_Mutable();
-
 	/** Remove all current tasks and register the one in the passed array. It's optimal to use Add/Remove of single tasks directly if possible.*/
 	GAMEPLAYTASKS_API void SetSimulatedTasks(const TArray<UGameplayTask*>& NewSimulatedTasks);
 
diff --git a/Engine/Source/Runtime/GameplayTasks/Private/GameplayTasksComponent.cpp b/Engine/Source/Runtime/GameplayTasks/Private/GameplayTasksComponent.cpp
--- a/Engine/Source/Runtime/GameplayTasks/Private/GameplayTasksComponent.cpp
+++ b/Engine/Source/Runtime/GameplayTasks/Private/GameplayTasksComponent.cpp
@@ -783,12 +783,6 @@ EGameplayTaskRunResult UGameplayTasksComponent::RunGameplayTask(IGameplayTaskOwn
 	return EGameplayTaskRunResult::Error;
 }
 
-TArray<TObjectPtr<UGameplayTask>>& UGameplayTasksComponent::GetSimulatedTasks_Mutable()
-{
-	MARK_PROPERTY_DIRTY_FROM_NAME(UGameplayTasksComponent, SimulatedTasks, this);
-	return SimulatedTasks;
-}
-
 bool UGameplayTasksComponent::AddSimulatedTask(UGameplayTask* NewTask)
 {
 	if (NewTask == nullptr)
