[Engine] Extended FCanvas public interface to allow preserving alpha when rendering certain canvas items.

Some batched canvas elements (e.g., runtime font glyphs and tiles) were not writing destination alpha. This was not an issue previously, since batched elements are normally rendered at the very end of the frame. However, with the new nDisplay feature that captures the entire UI layer into a separate transparent texture (for overlaying across multiple internal viewports), correct alpha output is required. Without this fix, text and tiles render visibly, but their alpha remains zero, breaking proper compositing.

#rnx

diff --git a/Engine/Source/Runtime/Engine/Private/UserInterface/Canvas.cpp b/Engine/Source/Runtime/Engine/Private/UserInterface/Canvas.cpp
--- a/Engine/Source/Runtime/Engine/Private/UserInterface/Canvas.cpp
+++ b/Engine/Source/Runtime/Engine/Private/UserInterface/Canvas.cpp
@@ -324,6 +324,7 @@ FCanvas::FCanvas(FRenderTarget* InRenderTarget, FHitProxyConsumer* InHitProxyCon
 ,	bRenderTargetDirty(false)
 ,	FeatureLevel(InFeatureLevel)
 ,	bUseInternalTexture(false)
+,	bWriteDestinationAlpha(false)
 ,	DrawMode(InDrawMode)
 ,	DPIScale(InDPIScale)
 {
@@ -346,6 +347,7 @@ FCanvas::FCanvas(FRenderTarget* InRenderTarget,FHitProxyConsumer* InHitProxyCons
 ,	Time(InTime)
 ,	FeatureLevel(InFeatureLevel)
 ,	bUseInternalTexture(false)
+,	bWriteDestinationAlpha(false)
 ,	DrawMode(CDM_DeferDrawing)
 ,	DPIScale(InDPIScale)
 {
@@ -1099,7 +1101,20 @@ void FCanvas::Clear(const FLinearColor& ClearColor)
 
 void FCanvas::DrawTile(double X, double Y, double SizeX, double SizeY, float U, float V, float SizeU, float SizeV, const FLinearColor& Color,	const FTexture* Texture, bool AlphaBlend)
 {
-	ESimpleElementBlendMode BlendMode = AlphaBlend ? (bUseInternalTexture ? SE_BLEND_TranslucentAlphaOnlyWriteAlpha : SE_BLEND_Translucent) : SE_BLEND_Opaque;
+	ESimpleElementBlendMode BlendMode = SE_BLEND_Opaque;
+	if (AlphaBlend)
+	{
+		if (bUseInternalTexture)
+		{
+			BlendMode = SE_BLEND_TranslucentAlphaOnlyWriteAlpha;
+		}
+		else
+		{
+			// If alpha is required, and no internal texture used, replace SE_BLEND_Translucent with SE_BLEND_AlphaBlend for all tiles.
+			BlendMode = (bWriteDestinationAlpha ? SE_BLEND_AlphaBlend : SE_BLEND_Translucent);
+		}
+	}
+
 	DrawTile(X, Y, SizeX, SizeY, U, V, SizeU, SizeV, Color, Texture, BlendMode);
 }
 
diff --git a/Engine/Source/Runtime/Engine/Private/UserInterface/CanvasItem.cpp b/Engine/Source/Runtime/Engine/Private/UserInterface/CanvasItem.cpp
--- a/Engine/Source/Runtime/Engine/Private/UserInterface/CanvasItem.cpp
+++ b/Engine/Source/Runtime/Engine/Private/UserInterface/CanvasItem.cpp
@@ -868,8 +868,10 @@ void FCanvasTextItemBase::Draw( class FCanvas* InCanvas )
 		// EnableShadow will set a default ShadowOffset value
 		EnableShadow( FLinearColor::Black );
 	}
+
+	const bool bWriteAlpha = InCanvas->IsWriteDestinationAlphaSet();
 	BlendMode = GetTextBlendMode( bHasShadow );
-	if (InCanvas->IsUsingInternalTexture())
+	if (InCanvas->IsUsingInternalTexture() || bWriteAlpha)
 	{
 		BlendMode = SE_BLEND_TranslucentAlphaOnlyWriteAlpha;
 	}
@@ -1423,7 +1425,7 @@ ESimpleElementBlendMode FCanvasShapedTextItem::GetTextBlendMode( const bool bHas
 	ESimpleElementBlendMode BlendModeToUse = BlendMode;
 
 	// The runtime font cache uses an alpha-only texture, so we have to force this blend mode so we use the correct shader
-	check(BlendModeToUse == SE_BLEND_Translucent || BlendModeToUse == SE_BLEND_TranslucentAlphaOnly);
+	check(BlendModeToUse == SE_BLEND_Translucent || BlendModeToUse == SE_BLEND_TranslucentAlphaOnly || BlendModeToUse == SE_BLEND_TranslucentAlphaOnlyWriteAlpha);
 	BlendModeToUse = SE_BLEND_TranslucentAlphaOnly;
 
 	return BlendModeToUse;
diff --git a/Engine/Source/Runtime/Engine/Public/CanvasTypes.h b/Engine/Source/Runtime/Engine/Public/CanvasTypes.h
--- a/Engine/Source/Runtime/Engine/Public/CanvasTypes.h
+++ b/Engine/Source/Runtime/Engine/Public/CanvasTypes.h
@@ -588,6 +588,13 @@ public:
 
 	inline bool IsUsingInternalTexture() const { return bUseInternalTexture; }
 
+	inline void SetWriteDestinationAlpha(const bool bInForceWriteAlpha)
+	{
+		bWriteDestinationAlpha = bInForceWriteAlpha;
+	}
+
+	inline bool IsWriteDestinationAlphaSet() const { return bWriteDestinationAlpha; }
+
 	inline void SetParentCanvasSize(FIntPoint InParentSize)
 	{
 		ParentSize = InParentSize;
@@ -635,6 +642,9 @@ private:
 	/** true, if Canvas is being rendered in its own texture */
 	bool bUseInternalTexture;
 
+	/** true, if the destination alpha channel must be written */
+	bool bWriteDestinationAlpha;
+
 	FIntPoint ParentSize;
 
 	ECanvasDrawMode DrawMode;
