Fixed SceneComponent not refreshing cache rotations after remtote object migrations

Effectively calling GetActorRotation() would return identity on freshly migrated actors. Now it preserves the correct transform.

diff --git a/Engine/Source/Runtime/Engine/Private/Components/SceneComponent.cpp b/Engine/Source/Runtime/Engine/Private/Components/SceneComponent.cpp
--- a/Engine/Source/Runtime/Engine/Private/Components/SceneComponent.cpp
+++ b/Engine/Source/Runtime/Engine/Private/Components/SceneComponent.cpp
@@ -3970,6 +3970,12 @@ void USceneComponent::Serialize(FArchive& Ar)
 	{
 		Ar << ComponentToWorld;
 		Ar << Bounds;
+		if (Ar.IsLoading())
+		{
+			// Freshen rotation cache (which would not have been serialized/copied), with our current transform
+			RelativeRotationCache.RotatorToQuat(RelativeRotation);
+			WorldRotationCache.QuatToRotator(ComponentToWorld.GetRotation());
+		}
 
 		AttachChildren.Append(RestoreChildren);
 		if (Ar.IsLoading() && !SceneComponentCVars::bAllowRemoteAttachedChildren)
