Remove stale entries in the UNetDriver::ServerReplicateActors ConsiderList when network actors are removed. This can happen when bReplicateTransactionally is true for the NetDriver.

diff --git a/Engine/Source/Runtime/Engine/Classes/Engine/NetDriver.h b/Engine/Source/Runtime/Engine/Classes/Engine/NetDriver.h
--- a/Engine/Source/Runtime/Engine/Classes/Engine/NetDriver.h
+++ b/Engine/Source/Runtime/Engine/Classes/Engine/NetDriver.h
@@ -2518,6 +2518,9 @@ private:
 	/** Stores the list of objects to replicate into the replay stream. This should be a TUniquePtr, but it appears the generated.cpp file needs the full definition of the pointed-to type. */
 	TSharedPtr<FNetworkObjectList> NetworkObjects;
 
+	/** Transient ConsiderList only used inside ServerReplicateActors. It's a member so that it can be easily updated between transactions when bReplicateTransactionally is true to remove stale entries. */
+	TArray<FNetworkObjectInfo*> CurrentConsiderList;
+
 	/** Set to "Lagging" on the server when all client connections are near timing out. We are lagging on the client when the server connection is near timed out. */
 	ENetworkLagState::Type LagState;
 
diff --git a/Engine/Source/Runtime/Engine/Private/NetDriver.cpp b/Engine/Source/Runtime/Engine/Private/NetDriver.cpp
--- a/Engine/Source/Runtime/Engine/Private/NetDriver.cpp
+++ b/Engine/Source/Runtime/Engine/Private/NetDriver.cpp
@@ -4434,6 +4434,14 @@ void UNetDriver::RemoveNetworkActor(AActor* Actor)
 		ReplicationDriver->RemoveNetworkActor(Actor);
 	}
 
+	// In case a network actor is removed while the ConsiderList is in use, remove the stale entries.
+	// This can happen when bReplicateTransactionally is true and ServerReplicateActors_ForConnection aborts.
+	CurrentConsiderList.RemoveAllSwap([Actor](const FNetworkObjectInfo* Info)
+		{
+			return Info->Actor == Actor;
+		},
+		EAllowShrinking::No);
+
 	GetNetworkObjectList().Remove(Actor);
 }
 
@@ -6192,8 +6200,14 @@ int32 UNetDriver::ServerReplicateActors(float DeltaSeconds)
 		bCPUSaturated	= DeltaSeconds > 1.2f * ServerTickTime;
 	}
 
-	TArray<FNetworkObjectInfo*> ConsiderList;
-	ConsiderList.Reserve( GetNetworkObjectList().GetActiveObjects().Num() );
+	ensureMsgf(CurrentConsiderList.IsEmpty(), TEXT("UNetDriver::ServerReplicateActors: CurrentConsiderList isn't empty before starting replication."));
+
+	CurrentConsiderList.Reserve( GetNetworkObjectList().GetActiveObjects().Num() );
+	ON_SCOPE_EXIT
+	{
+		// Deallocate the array while the NetDriver doesn't need it.
+		CurrentConsiderList.Empty();
+	};
 
 	// Build the consider list (actors that are ready to replicate)
 	if (bUseGranularTransactions)
@@ -6207,14 +6221,14 @@ int32 UNetDriver::ServerReplicateActors(float DeltaSeconds)
 		});
 
 		static const FName TransactionalWorkName("ServerReplicateActors_BuildConsiderList");
-		UE::RemoteExecutor::ExecuteTransactional(TransactionalWorkName, [this, &ConsiderList, ServerTickTime]()
+		UE::RemoteExecutor::ExecuteTransactional(TransactionalWorkName, [this, ServerTickTime]()
 			{
-				ServerReplicateActors_BuildConsiderList(ConsiderList, ServerTickTime);
+				ServerReplicateActors_BuildConsiderList(CurrentConsiderList, ServerTickTime);
 			});
 	}
 	else
 	{
-		ServerReplicateActors_BuildConsiderList(ConsiderList, ServerTickTime);
+		ServerReplicateActors_BuildConsiderList(CurrentConsiderList, ServerTickTime);
 	}
 
 	TSet<UNetConnection*> ConnectionsToClose;
@@ -6223,7 +6237,7 @@ int32 UNetDriver::ServerReplicateActors(float DeltaSeconds)
 
 	if (OnPreConsiderListUpdateOverride.IsBound())
 	{
-		OnPreConsiderListUpdateOverride.Execute({ DeltaSeconds, nullptr, bCPUSaturated }, Updated, ConsiderList);
+		OnPreConsiderListUpdateOverride.Execute({ DeltaSeconds, nullptr, bCPUSaturated }, Updated, CurrentConsiderList);
 	}
 
 	for ( int32 i=0; i < ClientConnections.Num(); i++ )
@@ -6254,20 +6268,20 @@ int32 UNetDriver::ServerReplicateActors(float DeltaSeconds)
 		{
 			//UE_LOG(LogNet, Log, TEXT("skipping update to %s"),*Connection->GetName());
 			// then mark each considered actor as bPendingNetUpdate so that they will be considered again the next frame when the connection is actually ticked
-			for (int32 ConsiderIdx = 0; ConsiderIdx < ConsiderList.Num(); ConsiderIdx++)
+			for (int32 ConsiderIdx = 0; ConsiderIdx < CurrentConsiderList.Num(); ConsiderIdx++)
 			{
-				AActor* Actor = ConsiderList[ConsiderIdx]->Actor;
+				AActor* Actor = CurrentConsiderList[ConsiderIdx]->Actor;
 				// if the actor hasn't already been flagged by another connection,
-				if (Actor != NULL && !ConsiderList[ConsiderIdx]->bPendingNetUpdate)
+				if (Actor != NULL && !CurrentConsiderList[ConsiderIdx]->bPendingNetUpdate)
 				{
 					// find the channel
-					UActorChannel* Channel = Connection->FindActorChannelRef(ConsiderList[ConsiderIdx]->WeakActor);
+					UActorChannel* Channel = Connection->FindActorChannelRef(CurrentConsiderList[ConsiderIdx]->WeakActor);
 					// and if the channel last update time doesn't match the last net update time for the actor
-					if (Channel != NULL && Channel->LastUpdateTime < ConsiderList[ConsiderIdx]->LastNetUpdateTimestamp)
+					if (Channel != NULL && Channel->LastUpdateTime < CurrentConsiderList[ConsiderIdx]->LastNetUpdateTimestamp)
 					{
 						//UE_LOG(LogNet, Log, TEXT("flagging %s for a future update"),*Actor->GetName());
 						// flag it for a pending update
-						ConsiderList[ConsiderIdx]->bPendingNetUpdate = true;
+						CurrentConsiderList[ConsiderIdx]->bPendingNetUpdate = true;
 					}
 				}
 			}
@@ -6281,7 +6295,7 @@ int32 UNetDriver::ServerReplicateActors(float DeltaSeconds)
 				.Connection = Connection,
 				.ConnectionViewers = WorldSettings->ReplicationViewers,
 				.DeltaSeconds = DeltaSeconds,
-				.ConsiderList = ConsiderList,
+				.ConsiderList = CurrentConsiderList,
 				.InOutUpdated = Updated,
 				.bCPUSaturated = bCPUSaturated
 			};
@@ -6317,7 +6331,7 @@ int32 UNetDriver::ServerReplicateActors(float DeltaSeconds)
 
 	if (OnPostConsiderListUpdateOverride.IsBound())
 	{
-		OnPostConsiderListUpdateOverride.ExecuteIfBound( { DeltaSeconds, nullptr, bCPUSaturated }, Updated, ConsiderList );
+		OnPostConsiderListUpdateOverride.ExecuteIfBound( { DeltaSeconds, nullptr, bCPUSaturated }, Updated, CurrentConsiderList);
 	}
 
 	// shuffle the list of connections if not all connections were ticked
