[Core] Fix remove during iteration

#tests ran client and performed undo in graphs (cause of pin crash)

#DisabledValidations nonUnityAndStaticAnalysis

diff --git a/Engine/Source/Runtime/Core/Private/FileCache/FileCache.cpp b/Engine/Source/Runtime/Core/Private/FileCache/FileCache.cpp
--- a/Engine/Source/Runtime/Core/Private/FileCache/FileCache.cpp
+++ b/Engine/Source/Runtime/Core/Private/FileCache/FileCache.cpp
@@ -264,11 +264,6 @@ public:
 	void Evict(CacheLineID Line);
 
 private:
-	void CheckForSizeRequestComplete();
-
-	CacheSlotID AcquireSlotAndReadLine(FFileCache& Cache, CacheLineID LineID, EAsyncIOPriorityAndFlags Priority);
-	void ReadLine(FFileCache& Cache, CacheSlotID SlotID, CacheLineID LineID, EAsyncIOPriorityAndFlags Priority, const FGraphEventRef& CompletionEvent);
-
 	struct FPendingRequest
 	{
 		FGraphEventRef Event;
@@ -278,6 +273,18 @@ private:
 		CacheSlotID SlotID;
 		FPendingRequest PendingRequest;
 	};
+	struct FValidLineInfo
+	{
+		FLineInfo& LineInfo;
+		bool bWasValid = false;
+	};
+
+
+	void CheckForSizeRequestComplete();
+
+	FValidLineInfo AcquireSlotAndReadLine(FFileCache& Cache, CacheLineID LineID, EAsyncIOPriorityAndFlags Priority);
+	void ReadLine(FFileCache& Cache, CacheSlotID SlotID, CacheLineID LineID, EAsyncIOPriorityAndFlags Priority, const FGraphEventRef& CompletionEvent);
+
 	TMap<int32, FLineInfo> LineInfos;
 
 	int64 BaseOffset;
@@ -686,14 +693,24 @@ void FFileCacheHandle::ReadLine(FFileCache& Cache, CacheSlotID SlotID, CacheLine
 	InnerHandle->ReadRequest(LineOffsetInFile, LineSizeInFile, Priority, &ReadCallbackFunction, CacheSlotMemory);
 }
 
-CacheSlotID FFileCacheHandle::AcquireSlotAndReadLine(FFileCache& Cache, CacheLineID LineID, EAsyncIOPriorityAndFlags Priority)
+FFileCacheHandle::FValidLineInfo FFileCacheHandle::AcquireSlotAndReadLine(FFileCache& Cache, CacheLineID LineID, EAsyncIOPriorityAndFlags Priority)
 {
 	SCOPED_LOADTIMER(FFileCacheHandle_AcquireSlotAndReadLine);
 
+	// Check if there's a valid line already and return that
+	{
+		FLineInfo& LineInfo = LineInfos.FindOrAdd(LineID.Get());
+		if (LineInfo.SlotID.IsValid())
+		{
+			return { LineInfo, true };
+		}
+	}
+
 	// no valid slot for this line, grab a new slot from cache and start a read request
 	CacheSlotID SlotID = Cache.AcquireAndLockSlot(this, LineID);
 
 	FLineInfo& LineInfo = LineInfos.FindOrAdd(LineID.Get());
+	LineInfo.SlotID = SlotID;
 	if (LineInfo.PendingRequest.Event)
 	{
 		// previous async request/event (if any) should be completed, if this is back in the free list
@@ -718,7 +735,7 @@ CacheSlotID FFileCacheHandle::AcquireSlotAndReadLine(FFileCache& Cache, CacheLin
 			TStatId(), SizeRequestEvent);
 	}
 
-	return SlotID;
+	return { LineInfo, false };
 }
 
 IMemoryReadStreamRef FFileCacheHandle::ReadData(FGraphEventArray& OutCompletionEvents, int64 InOffset, int64 BytesToRead, EAsyncIOPriorityAndFlags Priority)
@@ -760,14 +777,12 @@ IMemoryReadStreamRef FFileCacheHandle::ReadData(FGraphEventArray& OutCompletionE
 	bool bHasPendingSlot = false;
 	for (CacheLineID LineID = StartLine; LineID.Get() <= EndLine.Get(); ++LineID)
 	{
-		FLineInfo& LineInfo = LineInfos.FindOrAdd(LineID.Get());
-		if (!LineInfo.SlotID.IsValid())
-		{
-			// no valid slot for this line, grab a new slot from cache and start a read request
-			LineInfo.SlotID = AcquireSlotAndReadLine(Cache, LineID, Priority);
-		}
-		else
+		FValidLineInfo ValidLineInfo = AcquireSlotAndReadLine(Cache, LineID, Priority);
+
+		FLineInfo& LineInfo = ValidLineInfo.LineInfo;
+		if (ValidLineInfo.bWasValid)
 		{
+			// if it was already valid lock the slot for reading
 			Cache.LockSlot(LineInfo.SlotID);
 		}
 
@@ -845,12 +860,13 @@ FGraphEventRef FFileCacheHandle::PreloadData(const FFileCachePreloadEntry* Prelo
 		CurrentLine = CacheLineID(FMath::Max(CurrentLine.Get(), StartLine.Get()));
 		while (CurrentLine.Get() <= EndLine.Get() && Cache.NumFreeSlots > 0)
 		{
-			FLineInfo& LineInfo = LineInfos.FindOrAdd(CurrentLine.Get());
+			FValidLineInfo ValidLineInfo = AcquireSlotAndReadLine(Cache, CurrentLine, Priority);
+
+			FLineInfo& LineInfo = ValidLineInfo.LineInfo;
 			
-			if (!LineInfo.SlotID.IsValid())
+			// line info was not valid and just created, grab a new slot from cache and start a read request
+			if (!ValidLineInfo.bWasValid)
 			{
-				// no valid slot for this line, grab a new slot from cache and start a read request
-				LineInfo.SlotID = AcquireSlotAndReadLine(Cache, CurrentLine, Priority);
 				LockedSlots.Add(LineInfo.SlotID);
 			}
 
diff --git a/Engine/Source/Runtime/Engine/Private/EdGraph/EdGraphPin.cpp b/Engine/Source/Runtime/Engine/Private/EdGraph/EdGraphPin.cpp
--- a/Engine/Source/Runtime/Engine/Private/EdGraph/EdGraphPin.cpp
+++ b/Engine/Source/Runtime/Engine/Private/EdGraph/EdGraphPin.cpp
@@ -499,6 +499,7 @@ namespace PinHelpers
 	// Don't need bWasTrashed's name because it is transient
 #endif
 	static EPinResolveResult ImportText_PinReference(const TCHAR*& Buffer, UEdGraphPin*& PinRef, int32 ArrayIdx, UEdGraphPin* RequestingPin, EPinResolveType ResolveType);
+	static void ResolveReferencesToPin(UEdGraphPin* Pin, TArray<FUnresolvedPinData>& UnresolvedPinData, bool bStrictValidation = true);
 }
 
 /////////////////////////////////////////////////////
@@ -1642,7 +1643,7 @@ void UEdGraphPin::ResolveAllPinReferences()
 		UEdGraphPin** TargetPin = Node->Pins.FindByPredicate([&LocalGuid](const UEdGraphPin* Pin) { return Pin && Pin->PinId == LocalGuid; });
 		if (TargetPin)
 		{
-			ResolveReferencesToPin(*TargetPin);
+			PinHelpers::ResolveReferencesToPin(*TargetPin, Iter->Value);
 		}
 	}
 	PinHelpers::UnresolvedPins.Reset();
@@ -2002,56 +2003,60 @@ void UEdGraphPin::ResolveReferencesToPin(UEdGraphPin* Pin, bool bStrictValidatio
 	TArray<FUnresolvedPinData>* UnresolvedPinData = PinHelpers::UnresolvedPins.Find(ResolveId);
 	if (UnresolvedPinData)
 	{
-		for (const FUnresolvedPinData& PinData : *UnresolvedPinData)
+		PinHelpers::ResolveReferencesToPin(Pin, *UnresolvedPinData, bStrictValidation);
+		PinHelpers::UnresolvedPins.Remove(ResolveId);
+	}
+}
+
+void PinHelpers::ResolveReferencesToPin(UEdGraphPin* Pin, TArray<FUnresolvedPinData>& UnresolvedPinData, bool bStrictValidation)
+{
+	for (const FUnresolvedPinData& PinData : UnresolvedPinData)
+	{
+		UEdGraphPin* ReferencingPin = PinData.ReferencingPin;
+		switch (PinData.ResolveType)
 		{
-			UEdGraphPin* ReferencingPin = PinData.ReferencingPin;
-			switch (PinData.ResolveType)
+		case EPinResolveType::LinkedTo:
+			if (ensure(ReferencingPin->LinkedTo.Num() > PinData.ArrayIdx))
 			{
-			case EPinResolveType::LinkedTo:
-				if (ensure(ReferencingPin->LinkedTo.Num() > PinData.ArrayIdx))
-				{
-					ReferencingPin->LinkedTo[PinData.ArrayIdx] = Pin;
-				}
-				if (PinData.bResolveSymmetrically)
+				ReferencingPin->LinkedTo[PinData.ArrayIdx] = Pin;
+			}
+			if (PinData.bResolveSymmetrically)
+			{
+				if (ensure(!Pin->LinkedTo.Contains(ReferencingPin)))
 				{
-					if (ensure(!Pin->LinkedTo.Contains(ReferencingPin)))
-					{
-						Pin->LinkedTo.Add(ReferencingPin);
-					}
+					Pin->LinkedTo.Add(ReferencingPin);
 				}
-				if (bStrictValidation)
-				{
+			}
+			if (bStrictValidation)
+			{
 #if WITH_EDITOR
-					// When in the middle of a transaction the LinkedTo lists will be in an 
-					// indeterminate state. After PostEditUndo runs we could validate LinkedTo
-					// coherence.
-					ensureAlways(Pin->LinkedTo.Contains(ReferencingPin) || GEditor->Trans->IsObjectTransacting(Pin->GetOwningNode()));
+				// When in the middle of a transaction the LinkedTo lists will be in an 
+				// indeterminate state. After PostEditUndo runs we could validate LinkedTo
+				// coherence.
+				ensureAlways(Pin->LinkedTo.Contains(ReferencingPin) || GEditor->Trans->IsObjectTransacting(Pin->GetOwningNode()));
 #else
-					ensureAlways(Pin->LinkedTo.Contains(ReferencingPin));
+				ensureAlways(Pin->LinkedTo.Contains(ReferencingPin));
 #endif//WITH_EDITOR
-				}
-				break;
+			}
+			break;
 
-			case EPinResolveType::SubPins:
-				ReferencingPin->SubPins[PinData.ArrayIdx] = Pin;
-				break;
+		case EPinResolveType::SubPins:
+			ReferencingPin->SubPins[PinData.ArrayIdx] = Pin;
+			break;
 
-			case EPinResolveType::ParentPin:
-				ReferencingPin->ParentPin = Pin;
-				break;
+		case EPinResolveType::ParentPin:
+			ReferencingPin->ParentPin = Pin;
+			break;
 
-			case EPinResolveType::ReferencePassThroughConnection:
-				ReferencingPin->ReferencePassThroughConnection = Pin;
-				break;
+		case EPinResolveType::ReferencePassThroughConnection:
+			ReferencingPin->ReferencePassThroughConnection = Pin;
+			break;
 
-			case EPinResolveType::OwningNode:
-			default:
-				checkf(0, TEXT("Unhandled ResolveType %d"), (int32)PinData.ResolveType);
-				break;
-			}
+		case EPinResolveType::OwningNode:
+		default:
+			checkf(0, TEXT("Unhandled ResolveType %d"), (int32)PinData.ResolveType);
+			break;
 		}
-
-		PinHelpers::UnresolvedPins.Remove(ResolveId);
 	}
 }
 
