[Core] Fixup code using tset expecting order to remain unchanged after removal
* Fixup uses of compact stable to work with both compact and sparse sets

#tests Using the editor, running PIE and game client, making builds

diff --git a/Engine/Source/Developer/IoStoreUtilities/Private/PackageStoreOptimizer.cpp b/Engine/Source/Developer/IoStoreUtilities/Private/PackageStoreOptimizer.cpp
--- a/Engine/Source/Developer/IoStoreUtilities/Private/PackageStoreOptimizer.cpp
+++ b/Engine/Source/Developer/IoStoreUtilities/Private/PackageStoreOptimizer.cpp
@@ -922,7 +922,6 @@ TArray<FPackageStorePackage::FExportGraphNode*> FPackageStoreOptimizer::SortExpo
 			{
 				NodesWithNoIncomingEdges.HeapPush(ToNode, NodeSorter);
 			}
-			EdgeIt.RemoveCurrent();
 		}
 	}
 	check(LoadOrder.Num() == NodeCount);
diff --git a/Engine/Source/Editor/UnrealEd/Private/Commandlets/AssetRegUtilCommandlet.cpp b/Engine/Source/Editor/UnrealEd/Private/Commandlets/AssetRegUtilCommandlet.cpp
--- a/Engine/Source/Editor/UnrealEd/Private/Commandlets/AssetRegUtilCommandlet.cpp
+++ b/Engine/Source/Editor/UnrealEd/Private/Commandlets/AssetRegUtilCommandlet.cpp
@@ -540,14 +540,7 @@ bool UAssetRegUtilCommandlet::GeneratePartiallyUpdatedOrderFile(const FString& O
 	}
 
 	// Remove deleted files from old order file lines
-	for (auto OldOrderFileIter = OldOrderFileLineSet.CreateIterator(); OldOrderFileIter; ++OldOrderFileIter)
-	{
-		if (!NewOrderFileLineSet.Contains(*OldOrderFileIter))
-		{
-			OldOrderFileIter.RemoveCurrent();
-		}
-	}
-	OldOrderFileLineSet.CompactStable();
+	OldOrderFileLineSet = OldOrderFileLineSet.Difference(NewOrderFileLineSet);
 
 	// Add new files to old order file lines
 	int LastFoundIndexInOldFileLines = -1;
diff --git a/Engine/Source/Editor/VirtualizationEditor/Private/CommandletUtils.cpp b/Engine/Source/Editor/VirtualizationEditor/Private/CommandletUtils.cpp
--- a/Engine/Source/Editor/VirtualizationEditor/Private/CommandletUtils.cpp
+++ b/Engine/Source/Editor/VirtualizationEditor/Private/CommandletUtils.cpp
@@ -233,16 +233,16 @@ void FindVirtualizedPayloadsAndTrailers(const TArray<FString>& PackagePaths, TMa
 		{
 			const FString& PackageName = PackagePaths[Index];
 
-	UE::FPackageTrailer Trailer;
-	if (UE::FPackageTrailer::TryLoadFromFile(PackageName, Trailer))
-	{
-		const TArray<FIoHash> VirtualizedPayloads = Trailer.GetPayloads(UE::EPayloadStorageType::Virtualized);
-		if (!VirtualizedPayloads.IsEmpty())
-		{
-			Context.Packages.Emplace(PackageName, MoveTemp(Trailer));
-			Context.Payloads.Append(VirtualizedPayloads);
-		}
-	}
+			UE::FPackageTrailer Trailer;
+			if (UE::FPackageTrailer::TryLoadFromFile(PackageName, Trailer))
+			{
+				const TArray<FIoHash> VirtualizedPayloads = Trailer.GetPayloads(UE::EPayloadStorageType::Virtualized);
+				if (!VirtualizedPayloads.IsEmpty())
+				{
+					Context.Packages.Emplace(PackageName, MoveTemp(Trailer));
+					Context.Payloads.Append(VirtualizedPayloads);
+				}
+			}
 		});
 
 	for (const FTaskContext& Context : TaskContext)
@@ -250,8 +250,6 @@ void FindVirtualizedPayloadsAndTrailers(const TArray<FString>& PackagePaths, TMa
 		OutPackages.Append(Context.Packages);
 		OutPayloads.Append(Context.Payloads);
 	}
-
-	OutPayloads.CompactStable();
 }
 
 TArray<FPullRequest> ToRequestArray(TConstArrayView<FIoHash> IdentifierArray)
diff --git a/Engine/Source/Programs/Online/OSSTestsCore/Helpers/Identity/IdentityGetAllUserAccountsHelper.h b/Engine/Source/Programs/Online/OSSTestsCore/Helpers/Identity/IdentityGetAllUserAccountsHelper.h
--- a/Engine/Source/Programs/Online/OSSTestsCore/Helpers/Identity/IdentityGetAllUserAccountsHelper.h
+++ b/Engine/Source/Programs/Online/OSSTestsCore/Helpers/Identity/IdentityGetAllUserAccountsHelper.h
@@ -25,12 +25,17 @@ struct FIdentityGetAllUserAccountsStep : public FTestPipeline::FStep
 	{
 		IOnlineIdentityPtr OnlineIdentityPtr = OnlineSubsystem->GetIdentityInterface();
 
-		UserAccounts = OnlineIdentityPtr->GetAllUserAccounts();
+		TArray<TSharedPtr<FUserOnlineAccount>> UserAccounts = OnlineIdentityPtr->GetAllUserAccounts();
 		REQUIRE(!UserAccounts.IsEmpty());
 
-		for (size_t Idx = 0; Idx < UserAccounts.Num(); ++Idx)
+		for (const FUniqueNetIdPtr UniqueNetId : *UserUniqueNetIds)
 		{
-			CHECK(UserAccounts[Idx]->GetUserId()->ToString() == (*UserUniqueNetIds)[Idx]->ToString());
+			const FString UniqueName = UniqueNetId->ToString();
+			const bool bUserAccountExists = UserAccounts.ContainsByPredicate([&UniqueName](const TSharedPtr<FUserOnlineAccount>& UserAccount)
+				{
+					return UserAccount->GetUserId()->ToString() == UniqueName;
+				});
+			CHECK(bUserAccountExists);
 		}
 
 		StateSaver(UserAccounts);
@@ -39,7 +44,6 @@ struct FIdentityGetAllUserAccountsStep : public FTestPipeline::FStep
 	}
 
 protected:
-	TArray<TSharedPtr<FUserOnlineAccount>> UserAccounts;
 	TArray<FUniqueNetIdPtr>* UserUniqueNetIds = nullptr;
 	TFunction<void(TArray<TSharedPtr<FUserOnlineAccount>>)> StateSaver;
 };
\ No newline at end of file
diff --git a/Engine/Source/Runtime/Engine/Private/ActorConstruction.cpp b/Engine/Source/Runtime/Engine/Private/ActorConstruction.cpp
--- a/Engine/Source/Runtime/Engine/Private/ActorConstruction.cpp
+++ b/Engine/Source/Runtime/Engine/Private/ActorConstruction.cpp
@@ -221,11 +221,19 @@ void AActor::DestroyConstructedComponents()
 		GetPackage()->MarkPackageDirty();
 	}
 
-	// When a constructed component is destroyed, it is removed from this set. We compact the set to ensure any newly-constructed components
-	// will get added back into the set contiguously, so that the iteration order remains stable after re-running actor construction scripts.
+	// When a constructed component is destroyed, it is removed from this set. We rebuild it from the original set to compact
+	// the set and ensure that the iteration order remains stable after re-running actor construction scripts.
 	if (PreviouslyAttachedComponents.Num() > OwnedComponents.Num())
 	{
-		OwnedComponents.CompactStable();
+		TSet<TObjectPtr<UActorComponent>> OldComponents = MoveTemp(OwnedComponents);
+		OwnedComponents.Reserve(OldComponents.Num());
+		for (UActorComponent* Component : PreviouslyAttachedComponents)
+		{
+			if (OldComponents.Contains(Component))
+			{
+				OwnedComponents.Add(Component);
+			}
+		}
 	}
 }
 
diff --git a/Engine/Source/Runtime/Engine/Private/NetDriver.cpp b/Engine/Source/Runtime/Engine/Private/NetDriver.cpp
--- a/Engine/Source/Runtime/Engine/Private/NetDriver.cpp
+++ b/Engine/Source/Runtime/Engine/Private/NetDriver.cpp
@@ -1283,9 +1283,10 @@ void UNetDriver::UpdateUnmappedObjects()
 			TArray<FNetworkGUID> UnmappedGuids;
 			TArray<FNetworkGUID> NewlyMappedGuids;
 
-			for (auto It = ImportedNetGuidsRef.CreateIterator(); It; ++It)
+			auto ImportedNetGuidsIter = ImportedNetGuidsRef.CreateIterator();
+			for (; ImportedNetGuidsIter; ++ImportedNetGuidsIter)
 			{
-				const FNetworkGUID NetworkGuid = *It;
+				const FNetworkGUID& NetworkGuid = *ImportedNetGuidsIter;
 				bool bMappedOrBroken = false;
 
 				if (GetNetGuidCache()->GetObjectFromNetGUID(NetworkGuid, false) != nullptr)
@@ -1318,8 +1319,6 @@ void UNetDriver::UpdateUnmappedObjects()
 					bMappedOrBroken = true;
 				}
 
-				It.RemoveCurrent();
-
 				if (!bMappedOrBroken)
 				{
 					const FNetworkGUID OuterGUID = GetNetGuidCache()->GetOuterNetGUID(NetworkGuid);
@@ -1349,6 +1348,24 @@ void UNetDriver::UpdateUnmappedObjects()
 				}
 			}
 
+			// If there are any that were left unprocessed, then re-create the set with them at the beginning
+			if (ImportedNetGuidsIter)
+			{
+				TSet<FNetworkGUID> DeferredGuids;
+				DeferredGuids.Reserve(ImportedNetGuidsRef.Num());
+
+				for (; ImportedNetGuidsIter; ++ImportedNetGuidsIter)
+				{
+					DeferredGuids.Add(*ImportedNetGuidsIter);
+				}
+
+				Swap(ImportedNetGuidsRef, DeferredGuids);
+			}
+			else
+			{
+				ImportedNetGuidsRef.Empty();
+			}
+
 			// attempt to resolve dependent guids next tick (outer is now mapped)
 			for (const FNetworkGUID& NetGuid : NewlyMappedGuids)
 			{
@@ -1364,7 +1381,6 @@ void UNetDriver::UpdateUnmappedObjects()
 			// any tested guids that could not yet be mapped are added back to the list
 			if (UnmappedGuids.Num())
 			{
-				ImportedNetGuidsRef.CompactStable();
 				ImportedNetGuidsRef.Append(UnmappedGuids);
 			}
 
