Override GetDebugContextString in FDerivedAudioDataCompressor

diff --git a/Engine/Source/Runtime/Engine/Private/AudioDerivedData.cpp b/Engine/Source/Runtime/Engine/Private/AudioDerivedData.cpp
--- a/Engine/Source/Runtime/Engine/Private/AudioDerivedData.cpp
+++ b/Engine/Source/Runtime/Engine/Private/AudioDerivedData.cpp
@@ -2294,6 +2294,15 @@ bool FDerivedAudioDataCompressor::Build(TArray<uint8>& OutData)
 	return OutData.Num() > 0;
 }
 
+FString FDerivedAudioDataCompressor::GetDebugContextString() const
+{
+#if WITH_EDITORONLY_DATA
+	return CookInputs->SoundFullName;
+#else
+	return CookInputs->SoundName;
+#endif
+}
+
 /*---------------------------------------
 	USoundWave Derived Data functions
 ---------------------------------------*/
diff --git a/Engine/Source/Runtime/Engine/Public/AudioDerivedData.h b/Engine/Source/Runtime/Engine/Public/AudioDerivedData.h
--- a/Engine/Source/Runtime/Engine/Public/AudioDerivedData.h
+++ b/Engine/Source/Runtime/Engine/Public/AudioDerivedData.h
@@ -97,4 +97,6 @@ public:
 	virtual bool IsBuildThreadsafe() const override;
 
 	virtual bool Build(TArray<uint8>& OutData) override;
+
+	virtual FString GetDebugContextString() const override;
 };
