Switched category display name access to use FText::AsLocalizable_Advanced

This allows them to show more accurately in the transaction picker, even if there is no current translation available

#rnx

diff --git a/Engine/Source/Runtime/Engine/Private/ObjectEditorUtils.cpp b/Engine/Source/Runtime/Engine/Private/ObjectEditorUtils.cpp
--- a/Engine/Source/Runtime/Engine/Private/ObjectEditorUtils.cpp
+++ b/Engine/Source/Runtime/Engine/Private/ObjectEditorUtils.cpp
@@ -15,7 +15,6 @@ namespace FObjectEditorUtils
 	FText GetCategoryText( const FName InName )
 	{
 		static const FTextKey CategoryLocalizationNamespace = TEXT("UObjectCategory");
-		static const FName CategoryMetaDataKey = TEXT("Category");
 
 		FText DisplayName;
 
@@ -32,7 +31,7 @@ namespace FObjectEditorUtils
 		
 		if (DisplayName.IsEmpty())
 		{
-			DisplayName = FText::AsCultureInvariant(FName::NameToDisplayString(NativeCategory, false));
+			DisplayName = FText::AsLocalizable_Advanced(CategoryLocalizationNamespace, NativeCategory, FName::NameToDisplayString(NativeCategory, false));
 		}
 
 		return DisplayName;
@@ -49,12 +48,7 @@ namespace FObjectEditorUtils
 			const FString& NativeCategory = InField->GetMetaData(CategoryMetaDataKey);
 			if (!NativeCategory.IsEmpty())
 			{
-				FText LocalizedCategory;
-				if (!FText::FindTextInLiveTable_Advanced(CategoryLocalizationNamespace, NativeCategory, /*OUT*/LocalizedCategory, &NativeCategory))
-				{
-					LocalizedCategory = FText::AsCultureInvariant(NativeCategory);
-				}
-				return LocalizedCategory;
+				return FText::AsLocalizable_Advanced(CategoryLocalizationNamespace, NativeCategory, NativeCategory);
 			}
 		}
 
@@ -71,12 +65,7 @@ namespace FObjectEditorUtils
 			const FString& NativeCategory = InField->GetMetaData(CategoryMetaDataKey);
 			if (!NativeCategory.IsEmpty())
 			{
-				FText LocalizedCategory;
-				if (!FText::FindTextInLiveTable_Advanced(CategoryLocalizationNamespace, NativeCategory, /*OUT*/LocalizedCategory, &NativeCategory))
-				{
-					LocalizedCategory = FText::AsCultureInvariant(NativeCategory);
-				}
-				return LocalizedCategory;
+				return FText::AsLocalizable_Advanced(CategoryLocalizationNamespace, NativeCategory, NativeCategory);
 			}
 		}
 
