[Audio Insights] Remove AudioDevice argument from Audio::Trace::EventLog::SendActiveSoundEvent - it can be inferred from the active sound itself

diff --git a/Engine/Source/Runtime/Engine/Private/ActiveSound.cpp b/Engine/Source/Runtime/Engine/Private/ActiveSound.cpp
--- a/Engine/Source/Runtime/Engine/Private/ActiveSound.cpp
+++ b/Engine/Source/Runtime/Engine/Private/ActiveSound.cpp
@@ -694,7 +694,7 @@ void FActiveSound::Stop()
 	if (AudioDevice)
 	{
 #if UE_AUDIO_PROFILERTRACE_ENABLED
-		Audio::Trace::EventLog::SendActiveSoundEvent(*AudioDevice, *this, Audio::Trace::EventLog::ID::StopRequestActiveSound);
+		Audio::Trace::EventLog::SendActiveSoundEvent(*this, Audio::Trace::EventLog::ID::StopRequestActiveSound);
 #endif // UE_AUDIO_PROFILERTRACE_ENABLED
 
 		AudioDevice->AddSoundToStop(this);
diff --git a/Engine/Source/Runtime/Engine/Private/Audio/AudioTraceUtil.cpp b/Engine/Source/Runtime/Engine/Private/Audio/AudioTraceUtil.cpp
--- a/Engine/Source/Runtime/Engine/Private/Audio/AudioTraceUtil.cpp
+++ b/Engine/Source/Runtime/Engine/Private/Audio/AudioTraceUtil.cpp
@@ -100,11 +100,11 @@ namespace Audio::Trace
 
 	namespace EventLog
 	{
-		void SendActiveSoundEvent(const FAudioDevice& AudioDevice, const FActiveSound& InActiveSound, const FString& Event)
+		void SendActiveSoundEvent(const FActiveSound& InActiveSound, const FString& Event)
 		{
 			const bool bChannelEnabled = UE_TRACE_CHANNELEXPR_IS_ENABLED(AudioMixerChannel);
 
-			if (bChannelEnabled)
+			if (bChannelEnabled && InActiveSound.AudioDevice)
 			{
 				const TObjectPtr<UObject> SoundObjectPtr = Util::GetSoundObjectPointer(InActiveSound);
 				if (SoundObjectPtr == nullptr)
@@ -112,7 +112,7 @@ namespace Audio::Trace
 					return;
 				}
 
-				SendEvent(AudioDevice.DeviceID
+				SendEvent(InActiveSound.AudioDevice->DeviceID
 						, Event
 						, InActiveSound.GetPlayOrder()
 						, SoundObjectPtr->GetPathName()
diff --git a/Engine/Source/Runtime/Engine/Private/AudioDevice.cpp b/Engine/Source/Runtime/Engine/Private/AudioDevice.cpp
--- a/Engine/Source/Runtime/Engine/Private/AudioDevice.cpp
+++ b/Engine/Source/Runtime/Engine/Private/AudioDevice.cpp
@@ -5266,7 +5266,7 @@ void FAudioDevice::AddNewActiveSound(TSharedRef<FActiveSound> ActiveSound, TArra
 #if UE_AUDIO_PROFILERTRACE_ENABLED
 	if (EventLogID)
 	{
-		Audio::Trace::EventLog::SendActiveSoundEvent(*this, ActiveSound.Get(), *EventLogID);
+		Audio::Trace::EventLog::SendActiveSoundEvent(ActiveSound.Get(), *EventLogID);
 	}
 #endif // UE_AUDIO_PROFILERTRACE_ENABLED
 
@@ -5331,7 +5331,7 @@ void FAudioDevice::AddNewActiveSoundInternal(const FActiveSound& InNewActiveSoun
 		ReportSoundFailedToStart(InNewActiveSound.AudioComponentID, InVirtualLoopToRetrigger);
 
 #if UE_AUDIO_PROFILERTRACE_ENABLED
-		Audio::Trace::EventLog::SendActiveSoundEvent(*this, InNewActiveSound, Audio::Trace::EventLog::ID::PlayFailedNotPlayable);
+		Audio::Trace::EventLog::SendActiveSoundEvent(InNewActiveSound, Audio::Trace::EventLog::ID::PlayFailedNotPlayable);
 #endif // UE_AUDIO_PROFILERTRACE_ENABLED
 
 		return;
@@ -5348,7 +5348,7 @@ void FAudioDevice::AddNewActiveSoundInternal(const FActiveSound& InNewActiveSoun
 			ReportSoundFailedToStart(InNewActiveSound.AudioComponentID, InVirtualLoopToRetrigger);
 
 #if UE_AUDIO_PROFILERTRACE_ENABLED
-			Audio::Trace::EventLog::SendActiveSoundEvent(*this, InNewActiveSound, Audio::Trace::EventLog::ID::PlayFailedOutOfRange);
+			Audio::Trace::EventLog::SendActiveSoundEvent(InNewActiveSound, Audio::Trace::EventLog::ID::PlayFailedOutOfRange);
 #endif // UE_AUDIO_PROFILERTRACE_ENABLED
 
 			return;
@@ -5371,7 +5371,7 @@ void FAudioDevice::AddNewActiveSoundInternal(const FActiveSound& InNewActiveSoun
 				ReportSoundFailedToStart(InNewActiveSound.AudioComponentID, InVirtualLoopToRetrigger);
 
 #if UE_AUDIO_PROFILERTRACE_ENABLED
-				Audio::Trace::EventLog::SendActiveSoundEvent(*this, InNewActiveSound, Audio::Trace::EventLog::ID::PlayFailedDebugFiltered);
+				Audio::Trace::EventLog::SendActiveSoundEvent(InNewActiveSound, Audio::Trace::EventLog::ID::PlayFailedDebugFiltered);
 #endif // UE_AUDIO_PROFILERTRACE_ENABLED
 
 				return;
@@ -5424,7 +5424,7 @@ void FAudioDevice::AddNewActiveSoundInternal(const FActiveSound& InNewActiveSoun
 				ReportSoundFailedToStart(InNewActiveSound.GetAudioComponentID(), InVirtualLoopToRetrigger);
 
 #if UE_AUDIO_PROFILERTRACE_ENABLED
-				Audio::Trace::EventLog::SendActiveSoundEvent(*this, InNewActiveSound, Audio::Trace::EventLog::ID::PlayFailedConcurrency);
+				Audio::Trace::EventLog::SendActiveSoundEvent(InNewActiveSound, Audio::Trace::EventLog::ID::PlayFailedConcurrency);
 #endif // UE_AUDIO_PROFILERTRACE_ENABLED
 			}
 		}
@@ -5483,7 +5483,7 @@ void FAudioDevice::AddNewActiveSoundInternal(const FActiveSound& InNewActiveSoun
 		AddSoundToStop(&VirtualActiveSound);
 
 #if UE_AUDIO_PROFILERTRACE_ENABLED
-		Audio::Trace::EventLog::SendActiveSoundEvent(*this, VirtualActiveSound, Audio::Trace::EventLog::ID::SoundRealized);
+		Audio::Trace::EventLog::SendActiveSoundEvent(VirtualActiveSound, Audio::Trace::EventLog::ID::SoundRealized);
 #endif // UE_AUDIO_PROFILERTRACE_ENABLED
 
 		// Clear must be called after AddSoundToStop to ensure AudioComponent is properly removed from AudioComponentIDToActiveSoundMap
@@ -5661,7 +5661,7 @@ void FAudioDevice::AddVirtualLoop(const FAudioVirtualLoop& InVirtualLoop)
 	VirtualLoops.Add(&ActiveSound, MoveTemp(VirtualLoop));
 
 #if UE_AUDIO_PROFILERTRACE_ENABLED
-	Audio::Trace::EventLog::SendActiveSoundEvent(*this, ActiveSound, Audio::Trace::EventLog::ID::SoundVirtualized);
+	Audio::Trace::EventLog::SendActiveSoundEvent(ActiveSound, Audio::Trace::EventLog::ID::SoundVirtualized);
 #endif // UE_AUDIO_PROFILERTRACE_ENABLED
 
 	NotifySubsystemsVirtualizedSoundCreated(ActiveSound);
@@ -5916,7 +5916,7 @@ void FAudioDevice::StopActiveSound(const uint64 AudioComponentID, const FString*
 #if UE_AUDIO_PROFILERTRACE_ENABLED
 				if (EventLogID)
 				{
-					Audio::Trace::EventLog::SendActiveSoundEvent(*AudioDevice, ActiveSound, *EventLogID);
+					Audio::Trace::EventLog::SendActiveSoundEvent(ActiveSound, *EventLogID);
 				}
 #endif // UE_AUDIO_PROFILERTRACE_ENABLED
 
@@ -5933,7 +5933,7 @@ void FAudioDevice::StopActiveSound(FActiveSound* ActiveSound, const FString* Eve
 #if UE_AUDIO_PROFILERTRACE_ENABLED
 	if (ActiveSound && EventLogID)
 	{
-		Audio::Trace::EventLog::SendActiveSoundEvent(*this, *ActiveSound, *EventLogID);
+		Audio::Trace::EventLog::SendActiveSoundEvent(*ActiveSound, *EventLogID);
 	}
 #endif // UE_AUDIO_PROFILERTRACE_ENABLED
 
@@ -5960,7 +5960,7 @@ void FAudioDevice::PauseActiveSound(const uint64 AudioComponentID, const bool bI
 				FAudioDeviceHandle AudioDevice = AudioDeviceManager->GetAudioDevice(AudioDeviceID);
 				if (AudioDevice.IsValid())
 				{
-					Audio::Trace::EventLog::SendActiveSoundEvent(*AudioDevice, ActiveSound, bInIsPaused ? Audio::Trace::EventLog::ID::PauseSoundRequested : Audio::Trace::EventLog::ID::ResumeSoundRequested);
+					Audio::Trace::EventLog::SendActiveSoundEvent(ActiveSound, bInIsPaused ? Audio::Trace::EventLog::ID::PauseSoundRequested : Audio::Trace::EventLog::ID::ResumeSoundRequested);
 				}
 			}
 		}
@@ -7295,7 +7295,7 @@ void FAudioDevice::StopSoundsUsingResource(USoundWave* SoundWave, TArray<UAudioC
 				}
 
 #if UE_AUDIO_PROFILERTRACE_ENABLED
-				Audio::Trace::EventLog::SendActiveSoundEvent(*this, *ActiveSound, Audio::Trace::EventLog::ID::StopRequestSoundsUsingResource);
+				Audio::Trace::EventLog::SendActiveSoundEvent(*ActiveSound, Audio::Trace::EventLog::ID::StopRequestSoundsUsingResource);
 #endif // UE_AUDIO_PROFILERTRACE_ENABLED
 
 				AddSoundToStop(ActiveSound);
diff --git a/Engine/Source/Runtime/Engine/Private/SoundConcurrency.cpp b/Engine/Source/Runtime/Engine/Private/SoundConcurrency.cpp
--- a/Engine/Source/Runtime/Engine/Private/SoundConcurrency.cpp
+++ b/Engine/Source/Runtime/Engine/Private/SoundConcurrency.cpp
@@ -1076,7 +1076,7 @@ void FSoundConcurrencyManager::StopDueToVoiceStealing(FActiveSound& ActiveSound)
 	else
 	{
 #if UE_AUDIO_PROFILERTRACE_ENABLED
-		Audio::Trace::EventLog::SendActiveSoundEvent(*ActiveSound.AudioDevice, ActiveSound, Audio::Trace::EventLog::ID::StopRequestConcurrency);
+		Audio::Trace::EventLog::SendActiveSoundEvent(ActiveSound, Audio::Trace::EventLog::ID::StopRequestConcurrency);
 #endif // UE_AUDIO_PROFILERTRACE_ENABLED
 
 		ActiveSound.AudioDevice->AddSoundToStop(&ActiveSound);
diff --git a/Engine/Source/Runtime/Engine/Public/Audio/AudioTraceUtil.h b/Engine/Source/Runtime/Engine/Public/Audio/AudioTraceUtil.h
--- a/Engine/Source/Runtime/Engine/Public/Audio/AudioTraceUtil.h
+++ b/Engine/Source/Runtime/Engine/Public/Audio/AudioTraceUtil.h
@@ -23,7 +23,7 @@ namespace Audio::Trace
 	{
 		// Add an entry to the Audio Insights Event Log that related to an Active Sound
 		// Event may be customized to post custom event types to the Event Log
-		ENGINE_API void SendActiveSoundEvent(const FAudioDevice& AudioDevice, const FActiveSound& InActiveSound, const FString& Event);
+		ENGINE_API void SendActiveSoundEvent(const FActiveSound& InActiveSound, const FString& Event);
 
 		// Add an entry to the Audio Insights Event Log
 		ENGINE_API void SendEvent(const Audio::FDeviceId AudioDeviceID
