UPrimaryAssetLabel::UpdateAssetBundleData
- Only enumerate disk-only assets during cooking to ensure results are deterministic

#tests Tested full cooks of various projects

diff --git a/Engine/Source/Runtime/Engine/Private/PrimaryAssetLabel.cpp b/Engine/Source/Runtime/Engine/Private/PrimaryAssetLabel.cpp
--- a/Engine/Source/Runtime/Engine/Private/PrimaryAssetLabel.cpp
+++ b/Engine/Source/Runtime/Engine/Private/PrimaryAssetLabel.cpp
@@ -47,7 +47,11 @@ void UPrimaryAssetLabel::UpdateAssetBundleData()
 		FName PackagePath = FName(*FPackageName::GetLongPackagePath(GetOutermost()->GetName()));
 
 		TArray<FAssetData> DirectoryAssets;
-		AssetRegistry.GetAssetsByPath(PackagePath, DirectoryAssets, true);
+		FARFilter Filter;
+		Filter.bRecursivePaths = true;
+		Filter.PackagePaths.Add(PackagePath);
+		Filter.bIncludeOnlyOnDiskAssets = IsRunningCookCommandlet();	// Enumerate only on-disk packages when cooking to avoid indeterministic results
+		AssetRegistry.GetAssets(Filter, DirectoryAssets);
 
 		TArray<FTopLevelAssetPath> NewPaths;
 
