BugHawk: GetActorDescsForActors unconditionally returns false even when it successfully populates OutActorDescs in editor, which conflicts with the function's documented success semantics and with the established pattern in related methods.

diff --git a/Engine/Source/Runtime/Engine/Private/WorldPartitionBlueprintLibrary.cpp b/Engine/Source/Runtime/Engine/Private/WorldPartitionBlueprintLibrary.cpp
--- a/Engine/Source/Runtime/Engine/Private/WorldPartitionBlueprintLibrary.cpp
+++ b/Engine/Source/Runtime/Engine/Private/WorldPartitionBlueprintLibrary.cpp
@@ -266,6 +266,8 @@ bool UWorldPartitionBlueprintLibrary::GetIntersectingActorDescs(const FBox& InBo
 
 bool UWorldPartitionBlueprintLibrary::GetActorDescsForActors(const TArray<AActor*>& InActors, TArray<FActorDesc>& OutActorDescs)
 {
+	bool bResult = false;
+
 #if WITH_EDITOR
 	if (UWorldPartition* WorldPartition = GetWorldPartition())
 	{
@@ -274,11 +276,13 @@ bool UWorldPartitionBlueprintLibrary::GetActorDescsForActors(const TArray<AActor
 			if (FWorldPartitionActorDescInstance* ActorDescInstance = WorldPartition->GetActorDescInstance(Actor->GetActorGuid()))
 			{
 				OutActorDescs.Emplace(*ActorDescInstance);
+				bResult = true;
 			}
 		}
 	}
 #endif
-	return false;
+
+	return bResult;
 }
 
 UDataLayerManager* UWorldPartitionBlueprintLibrary::GetDataLayerManager(class UObject* WorldContextObject)
