Disable the selection of virtual textures on rect light components.
These have never been supported in game, but could appear to be supported when in the editor because the non VT editor fallback texture would be used.
To avoid confusion we now:
* Prevent selection of virtual textures in the UI.
* Log a warning when setting a virtual texture in Blueprint.
* Throw a map check if a virtual texture asset is set.

diff --git a/Engine/Source/Runtime/Engine/Classes/Components/RectLightComponent.h b/Engine/Source/Runtime/Engine/Classes/Components/RectLightComponent.h
--- a/Engine/Source/Runtime/Engine/Classes/Components/RectLightComponent.h
+++ b/Engine/Source/Runtime/Engine/Classes/Components/RectLightComponent.h
@@ -58,7 +58,7 @@ class URectLightComponent : public ULocalLightComponent
 	float LightFunctionConeAngle;
 
 	/** Texture mapped to the light source rectangle */
-	UPROPERTY(EditAnywhere, BlueprintReadOnly, Category = Light)
+	UPROPERTY(EditAnywhere, BlueprintReadOnly, Category = Light, meta = (GetAssetFilter = "ShouldFilterSourceTexture"))
 	TObjectPtr<class UTexture> SourceTexture;
 
 	/** Scales the source texture. Value in 0..1. (default=1) */
@@ -94,6 +94,7 @@ public:
 #if WITH_EDITOR
 	ENGINE_API virtual void SetLightBrightness(float InBrightness) override;
 	ENGINE_API virtual bool CanEditChange(const FProperty* InProperty) const override;
+	ENGINE_API virtual void CheckForErrors() override;
 #endif
 
 	//~ Begin ULightComponent Interface.
@@ -109,4 +110,14 @@ public:
 	//~ End UObject Interface
 
 	ENGINE_API virtual void Serialize(FArchive& Ar) override;
+	ENGINE_API virtual void PostLoad() override;
+
+protected:
+#if WITH_EDITOR
+	/** Validate source texture is of a valid type. */
+	void ValidateTexture() const;
+	
+	UFUNCTION()
+	bool ShouldFilterSourceTexture(const FAssetData& AssetData) const;
+#endif
 };
diff --git a/Engine/Source/Runtime/Engine/Private/Components/RectLightComponent.cpp b/Engine/Source/Runtime/Engine/Private/Components/RectLightComponent.cpp
--- a/Engine/Source/Runtime/Engine/Private/Components/RectLightComponent.cpp
+++ b/Engine/Source/Runtime/Engine/Private/Components/RectLightComponent.cpp
@@ -5,17 +5,21 @@
 =============================================================================*/
 
 #include "Components/RectLightComponent.h"
-#include "SceneInterface.h"
-#include "UObject/ConstructorHelpers.h"
-#include "UObject/FortniteMainBranchObjectVersion.h"
+
+#include "DataDrivenShaderPlatformInfo.h"
 #include "Engine/Texture2D.h"
+#include "Logging/MessageLog.h"
+#include "Misc/UObjectToken.h"
 #include "RectLightSceneProxy.h"
+#include "SceneInterface.h"
 #include "SceneView.h"
-
-#include "DataDrivenShaderPlatformInfo.h"
+#include "UObject/ConstructorHelpers.h"
+#include "UObject/FortniteMainBranchObjectVersion.h"
 
 #include UE_INLINE_GENERATED_CPP_BY_NAME(RectLightComponent)
 
+#define LOCTEXT_NAMESPACE "RectLightComponent"
+
 extern int32 GAllowPointLightCubemapShadows;
 
 float GetRectLightBarnDoorMaxAngle()
@@ -117,6 +121,14 @@ FLightSceneProxy* URectLightComponent::CreateSceneProxy() const
 
 void URectLightComponent::SetSourceTexture(UTexture* NewValue)
 {
+	if (NewValue && NewValue->VirtualTextureStreaming)
+	{
+		// Virtual textures aren't supported.
+		// We could add support in future, but we would need to force stream them before uploading to the RectLight Atlas so there would be little benefit for the added complexity.
+		UE_LOG(LogTemp, Warning, TEXT("RectLightComponent (%s) doesn't support Virtual Textures (%s)."), *GetName(), *NewValue->GetName())
+		return;
+	}
+
 	if (AreDynamicDataChangesAllowed()
 		&& SourceTexture != NewValue)
 	{
@@ -252,6 +264,33 @@ bool URectLightComponent::CanEditChange(const FProperty* InProperty) const
 
 	return Super::CanEditChange(InProperty);
 }
+
+void URectLightComponent::CheckForErrors()
+{
+	Super::CheckForErrors();
+
+	ValidateTexture();
+}
+
+void URectLightComponent::ValidateTexture() const
+{
+	if (SourceTexture && SourceTexture->VirtualTextureStreaming)
+	{
+		FFormatNamedArguments Arguments;
+		Arguments.Add(TEXT("SourceTexture"), FText::FromString(SourceTexture->GetName()));
+
+		FMessageLog("MapCheck").Warning()
+			->AddToken(FUObjectToken::Create(this))
+			->AddToken(FTextToken::Create(FText::Format(LOCTEXT("RectLight_VirtualTextureWarning", "Source Texture {SourceTexture} is a Virtual Texture.\n This is unsupported for Rect Lights and won't appear correctly in cooked builds."), Arguments)));
+	}
+}
+
+bool URectLightComponent::ShouldFilterSourceTexture(const FAssetData& InAssetData) const
+{
+	const FName VirtualTexturePropertyName(TEXT("VirtualTextureStreaming"));
+	const bool VirtualTextureStreaming = InAssetData.GetTagValueRef<bool>(VirtualTexturePropertyName);
+	return VirtualTextureStreaming != 0;
+}
 #endif // WITH_EDITOR
 
 /**
@@ -320,6 +359,16 @@ void URectLightComponent::Serialize(FArchive& Ar)
 	}
 }
 
+void URectLightComponent::PostLoad()
+{
+	Super::PostLoad();
+
+#if WITH_EDITOR
+	ValidateTexture();
+#endif
+}
+
+
 FRectLightSceneProxy::FRectLightSceneProxy(const URectLightComponent* Component)
 	: FLocalLightSceneProxy(Component)
 	, SourceWidth(Component->SourceWidth)
@@ -431,3 +480,5 @@ bool FRectLightSceneProxy::GetWholeSceneProjectedShadowInitializer(const FSceneV
 
 	return false;
 }
+
+#undef LOCTEXT_NAMESPACE
