UAssetManager::SearchAssetRegistryPaths
- enumerate disk-only assets when cooking

diff --git a/Engine/Source/Runtime/Engine/Private/AssetManager.cpp b/Engine/Source/Runtime/Engine/Private/AssetManager.cpp
--- a/Engine/Source/Runtime/Engine/Private/AssetManager.cpp
+++ b/Engine/Source/Runtime/Engine/Private/AssetManager.cpp
@@ -1086,7 +1086,7 @@ int32 UAssetManager::SearchAssetRegistryPaths(TArray<FAssetData>& OutAssetDataLi
 	}
 
 	ARFilter.bRecursivePaths = true;
-	ARFilter.bIncludeOnlyOnDiskAssets = !GIsEditor; // In editor check in memory, otherwise don't
+	ARFilter.bIncludeOnlyOnDiskAssets = !GIsEditor || IsRunningCookCommandlet(); // In editor check in memory, otherwise don't
 
 	auto FilterLambda = [this, &OutAssetDataList, &CompiledRules](const FAssetData& AssetData)
 	{
