Fix InitScalabilitySystem to set Scalability Group cvars even if they are the same as the defaults at Init.
They will be set as Scalability, no longer doing an overwrite of SetByConstructor.

diff --git a/Engine/Source/Runtime/Engine/Private/Scalability.cpp b/Engine/Source/Runtime/Engine/Private/Scalability.cpp
--- a/Engine/Source/Runtime/Engine/Private/Scalability.cpp
+++ b/Engine/Source/Runtime/Engine/Private/Scalability.cpp
@@ -199,6 +199,7 @@ namespace Scalability
 static FQualityLevels GScalabilityBackupQualityLevels(false);
 static FQualityLevels GScalabilityQualityLevelsOverride(false);
 static bool GScalabilityUsingTemporaryQualityLevels = false;
+static bool bIsInitialized = false;
 
 // Select a the correct quality level for the given benchmark value and thresholds
 int32 ComputeOptionFromPerfIndex(const FString& GroupName, float CPUPerfIndex, float GPUPerfIndex)
@@ -612,15 +613,9 @@ void OnChangeLandscapeQuality(IConsoleVariable* Var)
 void InitScalabilitySystem()
 {
 	// needed only once
+	if (bIsInitialized)
 	{
-		static bool bInit = false;
-
-		if(bInit)
-		{
-			return;
-		}
-
-		bInit = true;
+		return;
 	}
 
 	CVarResolutionQuality.AsVariable()->SetOnChangedCallback(FConsoleVariableDelegate::CreateStatic(&OnChangeResolutionQuality));
@@ -645,6 +640,7 @@ void InitScalabilitySystem()
 	GScalabilityBackupQualityLevels = FQualityLevels();
 	GScalabilityQualityLevelsOverride = FQualityLevels();
 	GScalabilityUsingTemporaryQualityLevels = false;
+	bIsInitialized = true;
 }
 
 /** Get the percentage scale for a given quality level */
@@ -904,7 +900,7 @@ template<typename T> void SetQualityLevelCVar(TAutoConsoleVariable<T>& CVarToEdi
 	{
 		T PreviousValue;
 		CVarToEdit.AsVariable()->GetValue(PreviousValue);
-		if (DesiredValue != PreviousValue)
+		if (DesiredValue != PreviousValue || !bIsInitialized)
 		{
 			CVarToEdit.AsVariable()->Set(DesiredValue, ECVF_SetByScalability);
 		}
@@ -936,7 +932,7 @@ void SetQualityLevels(const FQualityLevels& QualityLevels, bool bForce/* = false
 	}
 
 	//Skip the broadcast and assignations if the levels were the same.
-	if (GetQualityLevels() != ClampedLevels || bForce)
+	if (GetQualityLevels() != ClampedLevels || bForce || !bIsInitialized)
 	{
 		SetQualityLevelCVar(CVarResolutionQuality, ClampedLevels.ResolutionQuality, GScalabilityQualityLevelsOverride.ResolutionQuality, bForce);
 		SetQualityLevelCVar(CVarViewDistanceQuality, ClampedLevels.ViewDistanceQuality, GScalabilityQualityLevelsOverride.ViewDistanceQuality, bForce);
