[SimpleStreamableAssetManager] When using SimpleStreamableAssetManager, the computed MaxSize_VisibleOnly of the streaming render asset was always 0, which led to an incorrect VisibleWantedMips.
- Added an option for SimpleStreamableAssetManager to consider asset visibility (s.StreamableAssets.SimpleStreamableAssetManager.ConsiderVisibility).
- This is optional as it has a cost for maintaining the LastRenderTime.
- When disabled, fixed FRenderAssetInstanceAsyncView::UpdateBoundSizes by ignoring the logic based on "not seen recently" and generate a valid ScreenSizeOverDistance and VisibleWantedMips.

diff --git a/Engine/Source/Runtime/Engine/Private/Streaming/LevelTextureManager.cpp b/Engine/Source/Runtime/Engine/Private/Streaming/LevelTextureManager.cpp
--- a/Engine/Source/Runtime/Engine/Private/Streaming/LevelTextureManager.cpp
+++ b/Engine/Source/Runtime/Engine/Private/Streaming/LevelTextureManager.cpp
@@ -71,22 +71,7 @@ void FLevelRenderAssetManager::Remove(FRemovedRenderAssetArray* RemovedRenderAss
 
 float FLevelRenderAssetManager::GetWorldTime() const
 {
-	if (Level)
-	{
-		UWorld* World = Level->GetWorld();
-
-		// When paused, updating the world time sometimes break visibility logic.
-		if (World && !World->IsPaused())
-		{
-			// In the editor, we only return a time for the PIE world.
-			// TODO: figure out why there are more than one PIE world
-			if (!GIsEditor || (World->IsPlayInEditor() && World->Scene && World->Scene->GetFrameNumber()))
-			{
-				return World->GetTimeSeconds();
-			}
-		}
-	}
-	return 0;
+	return Level ? FRenderAssetStreamingManagerHelper::GetWorldTime(Level->GetWorld()) : 0.f;
 }
 
 void FLevelRenderAssetManager::SetAsStatic(FDynamicRenderAssetInstanceManager& DynamicComponentManager, const UPrimitiveComponent* Primitive, bool bIsConcurrent)
diff --git a/Engine/Source/Runtime/Engine/Private/Streaming/SimpleStreamableAssetManager.cpp b/Engine/Source/Runtime/Engine/Private/Streaming/SimpleStreamableAssetManager.cpp
--- a/Engine/Source/Runtime/Engine/Private/Streaming/SimpleStreamableAssetManager.cpp
+++ b/Engine/Source/Runtime/Engine/Private/Streaming/SimpleStreamableAssetManager.cpp
@@ -27,6 +27,7 @@ FSimpleStreamableAssetManager* FSimpleStreamableAssetManager::Instance = nullptr
 int32 FSimpleStreamableAssetManager::GUseSimpleStreamableAssetManager = 0;
 int32 FSimpleStreamableAssetManager::GSimpleStreamableAssetManagerSparseArrayGrowSize = 64;
 int32 FSimpleStreamableAssetManager::GSimpleStreamableAssetManagerEnsureAssetUniqueOnRegistration = 1;
+int32 FSimpleStreamableAssetManager::GSimpleStreamableAssetManagerConsiderVisibility = 0;
 
 FAutoConsoleVariableRef FSimpleStreamableAssetManager::CVarUseSimpleStreamableAssetManager(
 TEXT("s.StreamableAssets.UseSimpleStreamableAssetManager"),
@@ -43,6 +44,7 @@ FSimpleStreamableAssetManager::GSimpleStreamableAssetManagerSparseArrayGrowSize,
 TEXT("The growth size of SparseArray used for tracking objects pointing specific assets"),
 ECVF_SetByGameSetting
 );
+
 FAutoConsoleVariableRef FSimpleStreamableAssetManager::CVarUseSimpleStreamableAssetManagerEnsureAssetUniqueOnRegistration(
 TEXT("s.StreamableAssets.SimpleStreamableAssetManager.SortAssetsOnRegistration"),
 FSimpleStreamableAssetManager::GSimpleStreamableAssetManagerEnsureAssetUniqueOnRegistration,
@@ -52,6 +54,12 @@ TEXT("If true when object will be added, referenced assets will be sorted \n")
 ECVF_SetByGameSetting
 );
 
+FAutoConsoleVariableRef FSimpleStreamableAssetManager::CVarSimpleStreamableAssetManagerConsiderVisibility(
+TEXT("s.StreamableAssets.SimpleStreamableAssetManager.ConsiderVisibility"),
+FSimpleStreamableAssetManager::GSimpleStreamableAssetManagerConsiderVisibility,
+TEXT("Whether the SimpleStreamableAssetManager should consider the asset visibility when computing the biggest bounds normalized size (ScreenSize / Distance)"),
+ECVF_SetByGameSetting
+);
 
 void FSimpleStreamableAssetManager::Init()
 {
@@ -163,6 +171,12 @@ void FSimpleStreamableAssetManager::Update(FUpdate&& Record)
 	Instance->UpdateRecords.Push(MoveTemp(Record));
 }
 
+void FSimpleStreamableAssetManager::Update(FUpdateLastRenderTime&& Record)
+{
+	check(Instance != nullptr);
+	Instance->UpdateLastRenderTimeRecords.Push(MoveTemp(Record));
+}
+
 void FSimpleStreamableAssetManager::GetAssetReferenceBounds_Impl(const UStreamableRenderAsset* Asset, TArray<FBox>& AssetBoxes) const
 {
 	check(Asset->SimpleStreamableAssetManagerIndex.IsValid());
@@ -219,6 +233,7 @@ void FSimpleStreamableAssetManager::UpdateTask_Async()
 	
 
 	TArray<TLocklessGrowingStorage<FUpdate>::FStorageShard*> Pending_UpdateRecords{};
+	TArray<TLocklessGrowingStorage<FUpdateLastRenderTime>::FStorageShard*> Pending_UpdateLastRenderTimeRecords{};
 	TArray<TLocklessGrowingStorage<FRegister>::FStorageShard*> Pending_RegisterRecords{};
 	TArray<TLocklessGrowingStorage<FUnregister>::FStorageShard*> Pending_UnregisterRecords{};
 	TArray<TLocklessGrowingStorage<FRemovedAssetRecord>::FStorageShard*> Pending_RemovedAssetRecords{};
@@ -231,6 +246,7 @@ void FSimpleStreamableAssetManager::UpdateTask_Async()
 		 */
 		SCOPED_NAMED_EVENT(FSimpleStreamableAssetManager_MoveData, FColor::Silver);
 		Pending_UpdateRecords = UpdateRecords.ExtractShards();
+		Pending_UpdateLastRenderTimeRecords = UpdateLastRenderTimeRecords.ExtractShards();
 		Pending_RegisterRecords = RegisterRecords.ExtractShards();
 		Pending_UnregisterRecords = UnregisterRecords.ExtractShards();
 		Pending_RemovedAssetRecords = RemovedAssetsRecords.ExtractShards();
@@ -304,7 +320,6 @@ void FSimpleStreamableAssetManager::UpdateTask_Async()
 		}
 
 		{
-			
 			SCOPED_NAMED_EVENT(FSimpleStreamableAssetManager_Update, FColor::Silver);
 			for (TLocklessGrowingStorage<FUpdate>::FStorageShard* Shard : Pending_UpdateRecords)
 			{
@@ -316,6 +331,18 @@ void FSimpleStreamableAssetManager::UpdateTask_Async()
 			}
 		}
 
+		{
+			SCOPED_NAMED_EVENT(FSimpleStreamableAssetManager_UpdateLastRenderTime, FColor::Silver);
+			for (TLocklessGrowingStorage<FUpdateLastRenderTime>::FStorageShard* Shard : Pending_UpdateLastRenderTimeRecords)
+			{
+				TArrayView<FUpdateLastRenderTime> View = Shard->GetData();
+				for (FUpdateLastRenderTime& Entry : View)
+				{
+					UpdateRecord(Entry);
+				}
+			}
+		}
+
 		{
 			SCOPED_NAMED_EVENT(FSimpleStreamableAssetManager_RemoveAssets, FColor::Silver);
 			for (TLocklessGrowingStorage<FRemovedAssetRecord>::FStorageShard* Shard : Pending_RemovedAssetRecords)
@@ -359,7 +386,7 @@ void FSimpleStreamableAssetManager::UpdateTask_Async()
 			delete Shard;
 		}
 		TRACE_COUNTER_SET(UpdatedObjects, Updates);
-		
+
 #if COUNTERSTRACE_ENABLED
 		int32 Added = 0;
 #endif
@@ -388,6 +415,11 @@ void FSimpleStreamableAssetManager::UpdateTask_Async()
 		{
 			delete Shard;
 		}
+
+		for (TLocklessGrowingStorage<FUpdateLastRenderTime>::FStorageShard* Shard : Pending_UpdateLastRenderTimeRecords)
+		{
+			delete Shard;
+		}
 	}
 }
 
@@ -492,6 +524,15 @@ void FSimpleStreamableAssetManager::UpdateRecord(FUpdate& Record)
 	}
 }
 
+void FSimpleStreamableAssetManager::UpdateRecord(FUpdateLastRenderTime& Record)
+{
+	const int32 ObjectIndex = *Record.ObjectRegistrationIndex;
+	if (ObjectIndex != INDEX_NONE) // Filter out updates for not registered proxies
+	{
+		TrySetLastRenderTime(ObjectIndex, Record.LastRenderedTime);
+	}
+}
+
 void FSimpleStreamableAssetManager::UnregisterRecord(FUnregister& Record)
 {
 	const int32 ObjectIndex = *Record.ObjectRegistrationIndex;
@@ -502,6 +543,15 @@ void FSimpleStreamableAssetManager::UnregisterRecord(FUnregister& Record)
 	}
 }
 
+void FSimpleStreamableAssetManager::TrySetLastRenderTime(int32 BoundsIndex, float LastRenderedTime)
+{
+	const int32 ObjectBounds4Index = BoundsIndex / 4;
+	if (ObjectBounds4.IsValidIndex(ObjectBounds4Index))
+	{
+		ObjectBounds4[ObjectBounds4Index].UpdateLastRenderTime(BoundsIndex % 4, LastRenderedTime);
+	}
+}
+
 void FSimpleStreamableAssetManager::SetBounds(int32 BoundsIndex, const FBoxSphereBounds& Bounds, uint32 PackedRelativeBox, float StreamingScaleFactor, float LastRenderTime, const FVector4& RangeOrigin, float MinDistanceSq, float MinRangeSq, float MaxRangeSq)
 {
 	const int32 ObjectBounds4Index = BoundsIndex / 4;
diff --git a/Engine/Source/Runtime/Engine/Private/Streaming/StreamingManagerTexture.cpp b/Engine/Source/Runtime/Engine/Private/Streaming/StreamingManagerTexture.cpp
--- a/Engine/Source/Runtime/Engine/Private/Streaming/StreamingManagerTexture.cpp
+++ b/Engine/Source/Runtime/Engine/Private/Streaming/StreamingManagerTexture.cpp
@@ -29,6 +29,8 @@
 #include "HAL/LowLevelMemStats.h"
 #include "ProfilingDebugging/MetadataTrace.h"
 #include "Streaming/SimpleStreamableAssetManager.h"
+#include "SceneInterface.h"
+#include "Engine/Engine.h"
 
 #if !STATS
 #include "Async/ParallelFor.h"
@@ -1761,30 +1763,37 @@ void FRenderAssetStreamingManager::CheckUserSettings()
 
 void FRenderAssetStreamingManager::SetLastUpdateTime()
 {
-	if (FSimpleStreamableAssetManager::IsEnabled())
-	{
-		return;
-	}
-	
 	// Update the last update time.
 	float WorldTime = 0;
 
-	for (int32 LevelIndex = 0; LevelIndex < LevelRenderAssetManagers.Num(); ++LevelIndex)
+	if (FSimpleStreamableAssetManager::IsEnabled())
 	{
-		if (LevelRenderAssetManagers[LevelIndex] == nullptr)
+		if (FSimpleStreamableAssetManager::ShouldConsiderVisibility())
 		{
-			continue;
+			for (const FWorldContext& Context : GEngine->GetWorldContexts())
+			{
+				WorldTime = FRenderAssetStreamingManagerHelper::GetWorldTime(Context.World());
+				if (WorldTime > 0)
+				{
+					break;
+				}
+			}
 		}
-
-		// Update last update time only if there is a reasonable threshold to define visibility.
-		WorldTime = LevelRenderAssetManagers[LevelIndex]->GetWorldTime();
-		if (WorldTime > 0)
+	}
+	else
+	{
+		for (FLevelRenderAssetManager* LevelRenderAssetManager : LevelRenderAssetManagers)
 		{
-			break;
+			// Update last update time only if there is a reasonable threshold to define visibility.
+			WorldTime = LevelRenderAssetManager ? LevelRenderAssetManager->GetWorldTime() : 0.f;
+			if (WorldTime > 0)
+			{
+				break;
+			}
 		}
 	}
 
-	if (WorldTime> 0)
+	if (WorldTime > 0)
 	{
 		LastWorldUpdateTime = WorldTime - .5f;
 	}
@@ -3195,3 +3204,18 @@ bool FRenderAssetStreamingManager::Exec( UWorld* InWorld, const TCHAR* Cmd, FOut
 
 	return false;
 }
+
+float FRenderAssetStreamingManagerHelper::GetWorldTime(UWorld* InWorld)
+{
+	// When paused, updating the world time sometimes break visibility logic.
+	if (InWorld && !InWorld->IsPaused())
+	{
+		// In the editor, we only return a time for the PIE world.
+		if (!GIsEditor || (InWorld->IsPlayInEditor() && InWorld->Scene && InWorld->Scene->GetFrameNumber()))
+		{
+			return InWorld->GetTimeSeconds();
+		}
+	}
+
+	return 0.f;
+}
\ No newline at end of file
diff --git a/Engine/Source/Runtime/Engine/Private/Streaming/StreamingManagerTexture.h b/Engine/Source/Runtime/Engine/Private/Streaming/StreamingManagerTexture.h
--- a/Engine/Source/Runtime/Engine/Private/Streaming/StreamingManagerTexture.h
+++ b/Engine/Source/Runtime/Engine/Private/Streaming/StreamingManagerTexture.h
@@ -457,3 +457,8 @@ private:
 
 	friend bool TrackRenderAssetEvent( FStreamingRenderAsset* StreamingRenderAsset, UStreamableRenderAsset* RenderAsset, bool bForceMipLevelsToBeResident, const FRenderAssetStreamingManager* Manager);
 };
+
+struct FRenderAssetStreamingManagerHelper
+{
+	static float GetWorldTime(UWorld* InWorld);
+};
diff --git a/Engine/Source/Runtime/Engine/Private/Streaming/TextureInstanceView.cpp b/Engine/Source/Runtime/Engine/Private/Streaming/TextureInstanceView.cpp
--- a/Engine/Source/Runtime/Engine/Private/Streaming/TextureInstanceView.cpp
+++ b/Engine/Source/Runtime/Engine/Private/Streaming/TextureInstanceView.cpp
@@ -7,6 +7,7 @@
 #include "Streaming/TextureInstanceView.h"
 #include "Streaming/TextureInstanceView.inl"
 #include "Streaming/TextureStreamingHelpers.h"
+#include "Streaming/SimpleStreamableAssetManager.h"
 #include "Components/PrimitiveComponent.h"
 #include "ContentStreaming.h"
 
@@ -304,6 +305,7 @@ void FRenderAssetInstanceAsyncView::UpdateBoundSizes(
 	const int32 NumViews = ViewInfos.Num();
 	const int32 NumBounds4 = Bounds4.Num();
 
+	const bool bConsiderVisibility = !FSimpleStreamableAssetManager::IsEnabled() || FSimpleStreamableAssetManager::ShouldConsiderVisibility();
 	const VectorRegister LastUpdateTime4 = VectorSet(LastUpdateTime, LastUpdateTime, LastUpdateTime, LastUpdateTime);
 
 	OutBoundsViewInfos.Empty(NumBounds4 * 4);
@@ -420,7 +422,11 @@ void FRenderAssetInstanceAsyncView::UpdateBoundSizes(
 			// Now mask to zero if not in range, or not seen recently.
 			ScreenSizeOverDistance = VectorMultiply(ScreenSizeOverDistance, ExtraBoostForVisiblePrimitive);
 			ScreenSizeOverDistance = VectorSelect(InRangeMask, ScreenSizeOverDistance, VectorZero());
-			ScreenSizeOverDistance = VectorSelect(VectorCompareGT(LastRenderTime, LastUpdateTime4), ScreenSizeOverDistance, VectorZero());
+			// When using SimpleStreamableAssetManager, LastRenderTime is only valid when FSimpleStreamableAssetManager::ShouldConsiderVisibility() == true
+			if (bConsiderVisibility)
+			{
+				ScreenSizeOverDistance = VectorSelect(VectorCompareGT(LastRenderTime, LastUpdateTime4), ScreenSizeOverDistance, VectorZero());
+			}
 
 			MaxNormalizedSize_VisibleOnly = VectorMax(ScreenSizeOverDistance, MaxNormalizedSize_VisibleOnly);
 		}
diff --git a/Engine/Source/Runtime/Engine/Public/Streaming/SimpleStreamableAssetManager.h b/Engine/Source/Runtime/Engine/Public/Streaming/SimpleStreamableAssetManager.h
--- a/Engine/Source/Runtime/Engine/Public/Streaming/SimpleStreamableAssetManager.h
+++ b/Engine/Source/Runtime/Engine/Public/Streaming/SimpleStreamableAssetManager.h
@@ -178,6 +178,35 @@ public:
 		{}
 	};
 
+	struct FUpdateLastRenderTime
+	{
+		TSharedPtr<int32, ESPMode::ThreadSafe> ObjectRegistrationIndex{};
+#if SIMPLE_STREAMABLE_ASSET_MANAGER_ALLOW_DEBUG_POINTER
+		union
+		{
+			UPTRINT ObjectKey = 0;
+			const FPrimitiveSceneProxy* SceneProxy_ForDebug;
+		};
+#endif
+		float LastRenderedTime = -1000.0f;
+		FUpdateLastRenderTime() = default;
+		FUpdateLastRenderTime(const FUpdateLastRenderTime&) = default;
+		FUpdateLastRenderTime& operator=(const FUpdateLastRenderTime&) = default;
+		FUpdateLastRenderTime(FUpdateLastRenderTime&&) = default;
+		FUpdateLastRenderTime& operator=(FUpdateLastRenderTime&&) = default;
+
+		FUpdateLastRenderTime(
+			const void* InObject,
+			const TSharedPtr<int32, ESPMode::ThreadSafe>& InObjectRegistrationIndex,
+			const float InLastRenderedTime)
+				: ObjectRegistrationIndex(InObjectRegistrationIndex)
+#if SIMPLE_STREAMABLE_ASSET_MANAGER_ALLOW_DEBUG_POINTER
+				, ObjectKey(reinterpret_cast<UPTRINT>(InObject))
+#endif
+				, LastRenderedTime(InLastRenderedTime)
+		{}
+	};
+
 	struct FUpdate
 	{
 		TSharedPtr<int32, ESPMode::ThreadSafe> ObjectRegistrationIndex{};
@@ -332,6 +361,9 @@ private:
 
 	static FAutoConsoleVariableRef CVarUseSimpleStreamableAssetManagerEnsureAssetUniqueOnRegistration;
 	ENGINE_API static int32 GSimpleStreamableAssetManagerEnsureAssetUniqueOnRegistration;
+
+	static FAutoConsoleVariableRef CVarSimpleStreamableAssetManagerConsiderVisibility;
+	ENGINE_API static int32 GSimpleStreamableAssetManagerConsiderVisibility;
 	
 	static FSimpleStreamableAssetManager* Instance;
 	FCriticalSection CriticalSection;
@@ -340,6 +372,7 @@ private:
 	TLocklessGrowingStorage<FRegister> RegisterRecords;
 	TLocklessGrowingStorage<FUnregister> UnregisterRecords;
 	TLocklessGrowingStorage<FUpdate> UpdateRecords;
+	TLocklessGrowingStorage<FUpdateLastRenderTime> UpdateLastRenderTimeRecords;
 
 	// ** Variables to manage Object registration ** //
 	int32 RegisteredObjectCount = 0;
@@ -352,6 +385,7 @@ private:
 
 	void RegisterRecord(FRegister& Record);
 	void UpdateRecord(FUpdate& Record);
+	void UpdateRecord(FUpdateLastRenderTime& Record);
 	void UnregisterRecord(FUnregister& Record);
 
 	// ** Variables to manage Bounds registration ** //
@@ -366,6 +400,9 @@ private:
 	void AddRenderAssetElements(const TArrayView<FStreamingRenderAssetPrimitiveInfo>& RenderAssetInstanceInfos, int32 ObjectRegistrationIndex, bool bForceMipStreaming);
 	void RemoveRenderAssetElements(int32 ObjectRegistrationIndex);
 
+	// Used to update last render time
+	void TrySetLastRenderTime(int32 BoundsIndex, float LastRenderTime);
+
 	// ** Background task data** //
 	TArray<FBoundsViewInfo> BoundsViewInfos;
 	
@@ -399,8 +436,10 @@ public:
 	ENGINE_API static void Register(FRegister&& Record);
 	ENGINE_API static void Unregister(FUnregister&& Record);
 	ENGINE_API static void Update(FUpdate&& Record);
+	ENGINE_API static void Update(FUpdateLastRenderTime&& Record);
 	
 	static bool IsEnabled() { return GUseSimpleStreamableAssetManager != 0; }
+	static bool ShouldConsiderVisibility() { return GSimpleStreamableAssetManagerConsiderVisibility != 0; }
 	static void Init();
 	static void Shutdown();
 	static void Process();
diff --git a/Engine/Source/Runtime/Renderer/Private/PrimitiveSceneInfo.cpp b/Engine/Source/Runtime/Renderer/Private/PrimitiveSceneInfo.cpp
--- a/Engine/Source/Runtime/Renderer/Private/PrimitiveSceneInfo.cpp
+++ b/Engine/Source/Runtime/Renderer/Private/PrimitiveSceneInfo.cpp
@@ -43,6 +43,7 @@
 #include "VirtualTextureEnum.h"
 #include "MaterialCache/MaterialCacheMeshProcessor.h"
 #include "MaterialCache/MaterialCacheVirtualTextureRenderProxy.h"
+#include "Streaming/SimpleStreamableAssetManager.h"
 
 extern int32 GGPUSceneInstanceClearList;
 
@@ -2538,6 +2539,19 @@ void FPrimitiveSceneInfo::RemoveCachedReflectionCaptures()
 void FPrimitiveSceneInfo::UpdateComponentLastRenderTime(float CurrentWorldTime, bool bUpdateLastRenderTimeOnScreen)
 {
 	SceneData->SetLastRenderTime(CurrentWorldTime, bUpdateLastRenderTimeOnScreen);
+
+	if (bUpdateLastRenderTimeOnScreen && FSimpleStreamableAssetManager::IsEnabled() && FSimpleStreamableAssetManager::ShouldConsiderVisibility())
+	{
+		FSimpleStreamableAssetManager::Update(
+			FSimpleStreamableAssetManager::FUpdateLastRenderTime
+			{
+				Proxy,
+				Proxy->SimpleStreamableAssetManagerIndex,
+				CurrentWorldTime,
+			}
+		);
+	}
+
 #if UE_WITH_PSO_PRECACHING
 	Proxy->BoostPrecachedPSORequestsOnDraw();
 #endif
diff --git a/Engine/Source/Runtime/Renderer/Private/RendererScene.cpp b/Engine/Source/Runtime/Renderer/Private/RendererScene.cpp
--- a/Engine/Source/Runtime/Renderer/Private/RendererScene.cpp
+++ b/Engine/Source/Runtime/Renderer/Private/RendererScene.cpp
@@ -1681,7 +1681,7 @@ void FScene::UpdatePrimitiveTransformInternal(T* Primitive)
 										FSimpleStreamableAssetManager::GetStreamingScaleFactor(SceneProxy, UpdateParams.LocalToWorld),
 										SceneProxy->GetMinDrawDistance(),
 										SceneProxy->GetMaxDrawDistance(),
-										SceneProxy->PrimitiveSceneInfo->LastRenderTime,
+										SceneProxy->PrimitiveSceneInfo->GetSceneData()->LastRenderTimeOnScreen,
 										SceneProxy->IsForceMipStreaming()
 									}
 								);
@@ -1742,7 +1742,7 @@ void FScene::UpdatePrimitiveTransforms(IPrimitiveTransformUpdater* InUpdater)
 							FSimpleStreamableAssetManager::GetStreamingScaleFactor(SceneProxy, UpdateParams.LocalToWorld),
 							SceneProxy->GetMinDrawDistance(),
 							SceneProxy->GetMaxDrawDistance(),
-							SceneProxy->GetPrimitiveSceneInfo()->LastRenderTime,
+							SceneProxy->PrimitiveSceneInfo->GetSceneData()->LastRenderTimeOnScreen,
 							SceneProxy->IsForceMipStreaming()
 						}
 					);
