Changing more uses of RHICreateTexture to RHICmdList.CreateTexture.

diff --git a/Engine/Plugins/Experimental/GPULightmass/Source/GPULightmass/Private/Scene/Scene.cpp b/Engine/Plugins/Experimental/GPULightmass/Source/GPULightmass/Private/Scene/Scene.cpp
--- a/Engine/Plugins/Experimental/GPULightmass/Source/GPULightmass/Private/Scene/Scene.cpp
+++ b/Engine/Plugins/Experimental/GPULightmass/Source/GPULightmass/Private/Scene/Scene.cpp
@@ -1723,7 +1723,7 @@ void ReadbackVolumetricLightmapDataLayerFromGPU(FRHICommandListImmediate& RHICmd
 		FRHITextureCreateDesc::Create2D(TEXT("VolumetricLightmapDataLayerReadback"), Layer.Texture->GetSizeX(), Layer.Texture->GetSizeY(), Layer.Texture->GetFormat())
 		.SetFlags(ETextureCreateFlags::CPUReadback | ETextureCreateFlags::HideInVisualizeTexture);
 
-	FTextureRHIRef StagingTexture2DSlice = RHICreateTexture(StagingDesc);
+	FTextureRHIRef StagingTexture2DSlice = RHICmdList.CreateTexture(StagingDesc);
 	FGPUFenceRHIRef Fence = RHICreateGPUFence(TEXT("VolumetricLightmapDataLayerReadback"));
 
 	check(Dimensions.Z == Layer.Texture->GetSizeZ());
diff --git a/Engine/Plugins/Experimental/GPULightmass/Source/GPULightmass/Private/VolumetricLightmap.cpp b/Engine/Plugins/Experimental/GPULightmass/Source/GPULightmass/Private/VolumetricLightmap.cpp
--- a/Engine/Plugins/Experimental/GPULightmass/Source/GPULightmass/Private/VolumetricLightmap.cpp
+++ b/Engine/Plugins/Experimental/GPULightmass/Source/GPULightmass/Private/VolumetricLightmap.cpp
@@ -44,19 +44,19 @@ void InitializeBrickData(FRHICommandListBase& RHICmdList, FIntVector BrickDataDi
 		BrickData.SHCoefficients[i].Format = bForAccumulation ? PF_A32B32G32R32F : PF_R8G8B8A8;
 	}
 
-	BrickData.AmbientVector.CreateTargetTexture(BrickDataDimensions);
+	BrickData.AmbientVector.CreateTargetTexture(RHICmdList, BrickDataDimensions);
 	BrickData.AmbientVector.CreateUAV(RHICmdList);
 
 	for (int32 i = 0; i < UE_ARRAY_COUNT(BrickData.SHCoefficients); i++)
 	{
-		BrickData.SHCoefficients[i].CreateTargetTexture(BrickDataDimensions);
+		BrickData.SHCoefficients[i].CreateTargetTexture(RHICmdList, BrickDataDimensions);
 		BrickData.SHCoefficients[i].CreateUAV(RHICmdList);
 	}
 
-	BrickData.SkyBentNormal.CreateTargetTexture(BrickDataDimensions);
+	BrickData.SkyBentNormal.CreateTargetTexture(RHICmdList, BrickDataDimensions);
 	BrickData.SkyBentNormal.CreateUAV(RHICmdList);
 
-	BrickData.DirectionalLightShadowing.CreateTargetTexture(BrickDataDimensions);
+	BrickData.DirectionalLightShadowing.CreateTargetTexture(RHICmdList, BrickDataDimensions);
 	BrickData.DirectionalLightShadowing.CreateUAV(RHICmdList);
 
 	size_t TotalSize = 0;
diff --git a/Engine/Plugins/Experimental/VirtualHeightfieldMesh/Source/VirtualHeightfieldMesh/Private/VirtualHeightfieldMeshSceneProxy.cpp b/Engine/Plugins/Experimental/VirtualHeightfieldMesh/Source/VirtualHeightfieldMesh/Private/VirtualHeightfieldMeshSceneProxy.cpp
--- a/Engine/Plugins/Experimental/VirtualHeightfieldMesh/Source/VirtualHeightfieldMesh/Private/VirtualHeightfieldMeshSceneProxy.cpp
+++ b/Engine/Plugins/Experimental/VirtualHeightfieldMesh/Source/VirtualHeightfieldMesh/Private/VirtualHeightfieldMeshSceneProxy.cpp
@@ -846,15 +846,13 @@ namespace VirtualHeightfieldMesh
 		{
 			const FRHITextureCreateDesc Desc =
 				FRHITextureCreateDesc::Create2D(TEXT("VirtualHeightfieldMesh.MinMaxDefaultTexture"), 1, 1, PF_B8G8R8A8)
-				.SetFlags(ETextureCreateFlags::ShaderResource);
+				.SetFlags(ETextureCreateFlags::ShaderResource)
+				.DetermineInititialState()
+				.SetInitActionInitializer();
 
-			TextureRHI = RHICreateTexture(Desc);
-
-			// Write the contents of the texture.
-			uint32 DestStride;
-			FColor* DestBuffer = (FColor*)RHILockTexture2D(TextureRHI, 0, RLM_WriteOnly, DestStride, false);
-			*DestBuffer = FColor(0, 0, 255, 255);
-			RHIUnlockTexture2D(TextureRHI, 0, false);
+			FRHITextureInitializer Initializer = RHICmdList.CreateTextureInitializer(Desc);
+			Initializer.GetTexture2DSubresource(0).WriteColor(FColor(0, 0, 255, 255));
+			TextureRHI = Initializer.Finalize();
 
 			// Create the sampler state RHI resource.
 			FSamplerStateInitializerRHI SamplerStateInitializer(SF_Point, AM_Clamp, AM_Clamp, AM_Clamp);
@@ -1301,7 +1299,7 @@ void FVirtualHeightfieldMeshRendererExtension::SubmitWork(FRDGBuilder& GraphBuil
 						.SetNumMips(OcclusionResults->NumTextureMips)
 						.SetFlags(ETextureCreateFlags::ShaderResource);
 
-					OcclusionResults->OcclusionTexture = RHICreateTexture(Desc);
+					OcclusionResults->OcclusionTexture = GraphBuilder.RHICmdList.CreateTexture(Desc);
 
 					int32 SourceIndex = 0;
 
@@ -1309,7 +1307,7 @@ void FVirtualHeightfieldMeshRendererExtension::SubmitWork(FRDGBuilder& GraphBuil
 					for (int32 MipIndex = 0; MipIndex < OcclusionResults->NumTextureMips; ++MipIndex)
 					{
 						uint32 Stride;
-						uint8* Dst = (uint8*)RHILockTexture2D(OcclusionResults->OcclusionTexture, MipIndex, RLM_WriteOnly, Stride, false);
+						uint8* Dst = (uint8*)GraphBuilder.RHICmdList.LockTexture2D(OcclusionResults->OcclusionTexture, MipIndex, RLM_WriteOnly, Stride, false);
 
 						for (int Y = 0; Y < Size.Y; ++Y)
 						{
@@ -1319,7 +1317,7 @@ void FVirtualHeightfieldMeshRendererExtension::SubmitWork(FRDGBuilder& GraphBuil
 							}
 						}
 
-						RHIUnlockTexture2D(OcclusionResults->OcclusionTexture, MipIndex, false);
+						GraphBuilder.RHICmdList.UnlockTexture2D(OcclusionResults->OcclusionTexture, MipIndex, false);
 
 						Size.X = FMath::Max(Size.X / 2, 1);
 						Size.Y = FMath::Max(Size.Y / 2, 1);
diff --git a/Engine/Plugins/Experimental/VirtualHeightfieldMesh/Source/VirtualHeightfieldMeshEditor/Private/HeightfieldMinMaxTextureBuild.cpp b/Engine/Plugins/Experimental/VirtualHeightfieldMesh/Source/VirtualHeightfieldMeshEditor/Private/HeightfieldMinMaxTextureBuild.cpp
--- a/Engine/Plugins/Experimental/VirtualHeightfieldMesh/Source/VirtualHeightfieldMeshEditor/Private/HeightfieldMinMaxTextureBuild.cpp
+++ b/Engine/Plugins/Experimental/VirtualHeightfieldMesh/Source/VirtualHeightfieldMeshEditor/Private/HeightfieldMinMaxTextureBuild.cpp
@@ -37,10 +37,8 @@ namespace
 		}
 
 		//~ Begin FRenderResource Interface.
-		virtual void InitRHI(FRHICommandListBase&) override
+		virtual void InitRHI(FRHICommandListBase& RHICmdList) override
 		{
-			FRHICommandListImmediate& RHICmdList = FRHICommandListExecutor::GetImmediateCommandList();
-
 			FPooledRenderTargetDesc TileRenderTargetDesc = FPooledRenderTargetDesc::Create2DDesc(
 				FIntPoint(TileSize, TileSize),	PF_G16,	FClearValueBinding::None, TexCreate_None, TexCreate_ShaderResource, false);
 			GRenderTargetPool.FindFreeElement(RHICmdList, TileRenderTargetDesc, TileRenderTarget, TEXT("TileTarget"));
@@ -58,7 +56,7 @@ namespace
 					FRHITextureCreateDesc::Create2D(TEXT("FMinMaxTileRenderResources_StagingTextures"), SizeX, SizeY, PF_R8G8B8A8)
 					.SetFlags(ETextureCreateFlags::CPUReadback);
 
-				StagingTextures.Add(RHICreateTexture(Desc));
+				StagingTextures.Add(RHICmdList.CreateTexture(Desc));
 			}
 			
 			Fence = RHICreateGPUFence(TEXT("Runtime Virtual Texture Build"));
diff --git a/Engine/Plugins/Experimental/Water/Source/Runtime/Private/WaterQuadTreeGPU.cpp b/Engine/Plugins/Experimental/Water/Source/Runtime/Private/WaterQuadTreeGPU.cpp
--- a/Engine/Plugins/Experimental/Water/Source/Runtime/Private/WaterQuadTreeGPU.cpp
+++ b/Engine/Plugins/Experimental/Water/Source/Runtime/Private/WaterQuadTreeGPU.cpp
@@ -541,12 +541,12 @@ void FWaterQuadTreeGPU::Init(FRDGBuilder& GraphBuilder, const FInitParams& Param
 		FRHITextureCreateDesc QuadTreeTextureCreateDesc = FRHITextureCreateDesc::Create2D(TEXT("WaterQuadTree.QuadTree"), ResolutionPow2, PF_B8G8R8A8)
 			.SetNumMips(NumMipLevels)
 			.SetFlags(TexCreate_RenderTargetable | TexCreate_ShaderResource);
-		QuadTreeTexture = RHICreateTexture(QuadTreeTextureCreateDesc);
+		QuadTreeTexture = GraphBuilder.RHICmdList.CreateTexture(QuadTreeTextureCreateDesc);
 
 		FRHITextureCreateDesc ZBoundsTextureCreateDesc = FRHITextureCreateDesc::Create2D(TEXT("WaterQuadTree.ZBoundsTexture"), ResolutionPow2, PF_A2B10G10R10)
 			.SetNumMips(NumMipLevels)
 			.SetFlags(TexCreate_RenderTargetable | TexCreate_ShaderResource);
-		WaterZBoundsTexture = RHICreateTexture(ZBoundsTextureCreateDesc);
+		WaterZBoundsTexture = GraphBuilder.RHICmdList.CreateTexture(ZBoundsTextureCreateDesc);
 
 		WaterBodyRenderDataBuffer = GraphBuilder.ConvertToExternalBuffer(CreateStructuredBuffer<FWaterBodyRenderDataGPU>(GraphBuilder, TEXT("WaterQuadTree.WaterBodyRenderData"), Params.WaterBodyRenderData));
 
diff --git a/Engine/Plugins/FX/Niagara/Source/Niagara/Private/NiagaraDataInterfaceVolumeCache.cpp b/Engine/Plugins/FX/Niagara/Source/Niagara/Private/NiagaraDataInterfaceVolumeCache.cpp
--- a/Engine/Plugins/FX/Niagara/Source/Niagara/Private/NiagaraDataInterfaceVolumeCache.cpp
+++ b/Engine/Plugins/FX/Niagara/Source/Niagara/Private/NiagaraDataInterfaceVolumeCache.cpp
@@ -4,6 +4,7 @@
 #include "NiagaraSystemInstance.h"
 #include "NiagaraShaderParametersBuilder.h"
 #include "NiagaraSystem.h"
+#include "RenderGraphBuilder.h"
 #include "RHIStaticStates.h"
 #include "GlobalRenderResources.h"
 #include "VolumeCache.h"
@@ -182,7 +183,7 @@ void FNiagaraDataInterfaceVolumeCacheProxy::PreStage(const FNDIGpuComputePreStag
 				FRHITextureCreateDesc::Create3D(TEXT("NiagaraVolumeCache"), IntTextureSize.X, IntTextureSize.Y, IntTextureSize.Z, Format)
 				.SetFlags(ETextureCreateFlags::ShaderResource | ETextureCreateFlags::NoTiling);
 
-			InstanceData->ResolvedTextureRHI = RHICreateTexture(Desc);
+			InstanceData->ResolvedTextureRHI = Context.GetGraphBuilder().RHICmdList.CreateTexture(Desc);
 			InstanceData->SamplerStateRHI = TStaticSamplerState<SF_Bilinear, AM_Clamp, AM_Clamp, AM_Clamp>::GetRHI();
 		}
 		// must make sure readfile, curr frame, size are in the instance data - must sync in provide instance data for render thread
diff --git a/Engine/Plugins/FX/Niagara/Source/Niagara/Private/NiagaraEmptyUAVPool.cpp b/Engine/Plugins/FX/Niagara/Source/Niagara/Private/NiagaraEmptyUAVPool.cpp
--- a/Engine/Plugins/FX/Niagara/Source/Niagara/Private/NiagaraEmptyUAVPool.cpp
+++ b/Engine/Plugins/FX/Niagara/Source/Niagara/Private/NiagaraEmptyUAVPool.cpp
@@ -89,7 +89,7 @@ FRHIUnorderedAccessView* FNiagaraEmptyUAVPool::GetEmptyUAVFromPool(FRHICommandLi
 					FRHITextureCreateDesc::Create2D(ResourceName, 1, 1, Format)
 					.SetFlags(ETextureCreateFlags::ShaderResource | ETextureCreateFlags::UAV);
 
-				NewUAV.Texture = RHICreateTexture(Desc);
+				NewUAV.Texture = RHICmdList.CreateTexture(Desc);
 				NewUAV.UAV = RHICmdList.CreateUnorderedAccessView(NewUAV.Texture, FRHIViewDesc::CreateTextureUAV().SetDimensionFromTexture(NewUAV.Texture));
 				break;
 			}
@@ -100,7 +100,7 @@ FRHIUnorderedAccessView* FNiagaraEmptyUAVPool::GetEmptyUAVFromPool(FRHICommandLi
 					FRHITextureCreateDesc::Create2DArray(ResourceName, 1, 1, 1, Format)
 					.SetFlags(ETextureCreateFlags::ShaderResource | ETextureCreateFlags::UAV);
 
-				NewUAV.Texture = RHICreateTexture(Desc);
+				NewUAV.Texture = RHICmdList.CreateTexture(Desc);
 				NewUAV.UAV = RHICmdList.CreateUnorderedAccessView(NewUAV.Texture, FRHIViewDesc::CreateTextureUAV().SetDimensionFromTexture(NewUAV.Texture));
 				break;
 			}
@@ -111,7 +111,7 @@ FRHIUnorderedAccessView* FNiagaraEmptyUAVPool::GetEmptyUAVFromPool(FRHICommandLi
 					FRHITextureCreateDesc::Create3D(ResourceName, 1, 1, 1, Format)
 					.SetFlags(ETextureCreateFlags::ShaderResource | ETextureCreateFlags::UAV);
 
-				NewUAV.Texture = RHICreateTexture(Desc);
+				NewUAV.Texture = RHICmdList.CreateTexture(Desc);
 				NewUAV.UAV = RHICmdList.CreateUnorderedAccessView(NewUAV.Texture, FRHIViewDesc::CreateTextureUAV().SetDimensionFromTexture(NewUAV.Texture));
 				break;
 			}
@@ -122,7 +122,7 @@ FRHIUnorderedAccessView* FNiagaraEmptyUAVPool::GetEmptyUAVFromPool(FRHICommandLi
 					FRHITextureCreateDesc::CreateCube(ResourceName, 1, Format)
 					.SetFlags(ETextureCreateFlags::ShaderResource | ETextureCreateFlags::UAV);
 
-				NewUAV.Texture = RHICreateTexture(Desc);
+				NewUAV.Texture = RHICmdList.CreateTexture(Desc);
 				NewUAV.UAV = RHICmdList.CreateUnorderedAccessView(NewUAV.Texture, FRHIViewDesc::CreateTextureUAV().SetDimensionFromTexture(NewUAV.Texture));
 				break;
 			}
@@ -134,7 +134,7 @@ FRHIUnorderedAccessView* FNiagaraEmptyUAVPool::GetEmptyUAVFromPool(FRHICommandLi
 					FRHITextureCreateDesc::CreateCubeArray(ResourceName, 1, 1, Format)
 					.SetFlags(ETextureCreateFlags::ShaderResource | ETextureCreateFlags::UAV);
 
-				NewUAV.Texture = RHICreateTexture(Desc);
+				NewUAV.Texture = RHICmdList.CreateTexture(Desc);
 				NewUAV.UAV = RHICmdList.CreateUnorderedAccessView(NewUAV.Texture, FRHIViewDesc::CreateTextureUAV().SetDimensionFromTexture(NewUAV.Texture));
 				break;
 			}
diff --git a/Engine/Plugins/FX/Niagara/Source/Niagara/Private/NiagaraGpuComputeDebug.cpp b/Engine/Plugins/FX/Niagara/Source/Niagara/Private/NiagaraGpuComputeDebug.cpp
--- a/Engine/Plugins/FX/Niagara/Source/Niagara/Private/NiagaraGpuComputeDebug.cpp
+++ b/Engine/Plugins/FX/Niagara/Source/Niagara/Private/NiagaraGpuComputeDebug.cpp
@@ -175,7 +175,7 @@ void FNiagaraGpuComputeDebug::AddAttributeTexture(FRDGBuilder& GraphBuilder, FNi
 			.SetFlags(ETextureCreateFlags::ShaderResource)
 			.SetInitialState(ERHIAccess::SRVMask);
 
-		VisualizeEntry->Texture = CreateRenderTarget(RHICreateTexture(NewTextureDesc), TEXT("FNiagaraGpuComputeDebug"));
+		VisualizeEntry->Texture = CreateRenderTarget(GraphBuilder.RHICmdList.CreateTexture(NewTextureDesc), TEXT("FNiagaraGpuComputeDebug"));
 	}
 	check(VisualizeEntry->Texture.IsValid());
 
diff --git a/Engine/Plugins/FX/Niagara/Source/Niagara/Private/NiagaraRenderer.cpp b/Engine/Plugins/FX/Niagara/Source/Niagara/Private/NiagaraRenderer.cpp
--- a/Engine/Plugins/FX/Niagara/Source/Niagara/Private/NiagaraRenderer.cpp
+++ b/Engine/Plugins/FX/Niagara/Source/Niagara/Private/NiagaraRenderer.cpp
@@ -103,7 +103,7 @@ public:
 
 	virtual void InitRHI(FRHICommandListBase& RHICmdList) override
 	{
-		const FRHITextureCreateDesc Desc = FRHITextureCreateDesc(*DebugName, Dimension)
+		FRHITextureCreateDesc Desc = FRHITextureCreateDesc(*DebugName, Dimension)
 			.SetExtent(1, 1)
 			.SetDepth(1)
 			.SetArraySize(1)
@@ -112,36 +112,27 @@ public:
 
 		// Create a 1x1 texture.
 
-		uint32 Stride;
 		switch (Dimension)
 		{
 			case ETextureDimension::Texture2D:
-			{
-				Texture = RHICreateTexture(Desc);
-				void* Pixels = RHILockTexture2D(Texture, 0, RLM_WriteOnly, Stride, false);
-				FMemory::Memset(Pixels, 0, Stride);
-				RHIUnlockTexture2D(Texture, 0, 0, false);
-				break;
-			}
-
 			case ETextureDimension::Texture2DArray:
 			{
-				Texture = RHICreateTexture(Desc);
-				void* Pixels = RHILockTexture2DArray(Texture, 0, 0, RLM_WriteOnly, Stride, false);
-				FMemory::Memset(Pixels, 0, Stride);
-				RHIUnlockTexture2DArray(Texture, 0, 0, false);
+				FRHITextureInitializer Initializer = RHICmdList.CreateTextureInitializer(Desc.SetInitActionInitializer());
+				FRHITextureSubresourceInitializer Subresource = Initializer.GetSubresource(FRHITextureInitializer::FSubresourceIndex());
+				FMemory::Memzero(Subresource.Data, Subresource.Size);
+				Texture = Initializer.Finalize();
 				break;
 			}
 
 			case ETextureDimension::Texture3D:
 			{
-				FTextureRHIRef Texture3D = RHICreateTexture(Desc);
+				FTextureRHIRef Texture3D = RHICmdList.CreateTexture(Desc);
 				Texture = Texture3D;
 
 				const FPixelFormatInfo& Info = GPixelFormats[PixelFormat];
 				TArray<uint8, TInlineAllocator<16>> Data;
 				Data.AddZeroed(Info.BlockBytes);
-				RHIUpdateTexture3D(Texture3D, 0, FUpdateTextureRegion3D(0, 0, 0, 0, 0, 0, 1, 1, 1), Info.BlockBytes, Info.BlockBytes, Data.GetData());
+				RHICmdList.UpdateTexture3D(Texture3D, 0, FUpdateTextureRegion3D(0, 0, 0, 0, 0, 0, 1, 1, 1), Info.BlockBytes, Info.BlockBytes, Data.GetData());
 				break;
 			}
 
diff --git a/Engine/Plugins/Media/AndroidCamera/Source/AndroidCamera/Private/Player/AndroidCameraPlayer.cpp b/Engine/Plugins/Media/AndroidCamera/Source/AndroidCamera/Private/Player/AndroidCameraPlayer.cpp
--- a/Engine/Plugins/Media/AndroidCamera/Source/AndroidCamera/Private/Player/AndroidCameraPlayer.cpp
+++ b/Engine/Plugins/Media/AndroidCamera/Source/AndroidCamera/Private/Player/AndroidCameraPlayer.cpp
@@ -290,7 +290,7 @@ void FAndroidCameraPlayer::SetGuid(const FGuid& Guid)
 // ==============================================================
 // Used by TickFetch
 
-static void ConditionalInitializeVideoTexture_RenderThread(TWeakPtr<FJavaAndroidCameraPlayer, ESPMode::ThreadSafe> JavaCameraPlayerPtr, FGuid PlayerGuid)
+static void ConditionalInitializeVideoTexture_RenderThread(FRHICommandListBase& RHICmdList, TWeakPtr<FJavaAndroidCameraPlayer, ESPMode::ThreadSafe> JavaCameraPlayerPtr, FGuid PlayerGuid)
 {
 	auto PinnedJavaCameraPlayer = JavaCameraPlayerPtr.Pin();
 
@@ -307,12 +307,12 @@ static void ConditionalInitializeVideoTexture_RenderThread(TWeakPtr<FJavaAndroid
 			.SetFlags(ETextureCreateFlags::External)
 			.SetInitialState(ERHIAccess::SRVMask);
 
-		VideoTexture = RHICreateTexture(Desc);
+		VideoTexture = RHICmdList.CreateTexture(Desc);
 		PinnedJavaCameraPlayer->SetVideoTexture(VideoTexture);
 
 		if (VideoTexture == nullptr)
 		{
-			UE_LOG(LogAndroidCamera, Warning, TEXT("RHICreateTexture failed!"));
+			UE_LOG(LogAndroidCamera, Warning, TEXT("CreateTexture failed!"));
 			return;
 		}
 
@@ -373,15 +373,15 @@ struct FRHICommandUpdateExternalCameraSample final : public FRHICommand<FRHIComm
 	}
 };
 
-static void ConditionalInitializeVideoSample_RenderThread(TSharedRef<FAndroidCameraTextureSample, ESPMode::ThreadSafe> VideoSample)
+static void ConditionalInitializeVideoSample_RenderThread(FRHICommandListBase& RHICmdList, TSharedRef<FAndroidCameraTextureSample, ESPMode::ThreadSafe> VideoSample)
 {
 	if (VideoSample->GetTexture() == nullptr || VideoSample->GetTexture()->GetSizeXY() != VideoSample->GetOutputDim())
 	{
-		VideoSample->InitializeTexture(FTimespan::Zero());
+		VideoSample->InitializeTexture(RHICmdList, FTimespan::Zero());
 	}
 }
 
-static void DoUpdateTextureCameraSampleExecute(TWeakPtr<FJavaAndroidCameraPlayer, ESPMode::ThreadSafe> JavaCameraPlayerPtr, TWeakPtr<FMediaSamples, ESPMode::ThreadSafe> SamplesPtr,
+static void DoUpdateTextureCameraSampleExecute(FRHICommandListBase& RHICmdList, TWeakPtr<FJavaAndroidCameraPlayer, ESPMode::ThreadSafe> JavaCameraPlayerPtr, TWeakPtr<FMediaSamples, ESPMode::ThreadSafe> SamplesPtr,
 	TSharedRef<FAndroidCameraTextureSample, ESPMode::ThreadSafe> VideoSample, FTimespan SampleTime)
 {
 	auto PinnedJavaCameraPlayer = JavaCameraPlayerPtr.Pin();
@@ -393,7 +393,7 @@ static void DoUpdateTextureCameraSampleExecute(TWeakPtr<FJavaAndroidCameraPlayer
 	}
 
 	// write frame into texture
-	FRHITexture* Texture = VideoSample->InitializeTexture(SampleTime);
+	FRHITexture* Texture = VideoSample->InitializeTexture(RHICmdList, SampleTime);
 
 	if (Texture != nullptr)
 	{
@@ -426,10 +426,10 @@ struct FRHICommandUpdateTextureCameraSample final : public FRHICommand<FRHIComma
 		, SampleTime(InSampleTime)
 	{
 	}
-	void Execute(FRHICommandListBase& CmdList)
+	void Execute(FRHICommandListBase& RHICmdList)
 	{
 		QUICK_SCOPE_CYCLE_COUNTER(STAT_FRHICommandUpdateTextureCameraSample_Execute);
-		DoUpdateTextureCameraSampleExecute(JavaCameraPlayerPtr, SamplesPtr, VideoSample, SampleTime);
+		DoUpdateTextureCameraSampleExecute(RHICmdList, JavaCameraPlayerPtr, SamplesPtr, VideoSample, SampleTime);
 	}
 };
 
@@ -697,7 +697,7 @@ void FAndroidCameraPlayer::TickFetch(FTimespan DeltaTime, FTimespan Timecode)
 		ENQUEUE_RENDER_COMMAND(AndroidCameraPlayerWriteVideoSample)(
 			[Params](FRHICommandListImmediate& RHICmdList)
 			{
-				ConditionalInitializeVideoTexture_RenderThread(Params.JavaCameraPlayerPtr, Params.PlayerGuid);
+				ConditionalInitializeVideoTexture_RenderThread(RHICmdList, Params.JavaCameraPlayerPtr, Params.PlayerGuid);
 				if (IsRunningRHIInSeparateThread())
 				{
 					new (RHICmdList.AllocCommand<FRHICommandUpdateExternalCameraSample>()) FRHICommandUpdateExternalCameraSample(Params.JavaCameraPlayerPtr, Params.PlayerGuid);
@@ -787,7 +787,7 @@ void FAndroidCameraPlayer::TickFetch(FTimespan DeltaTime, FTimespan Timecode)
 		ENQUEUE_RENDER_COMMAND(AndroidCameraPlayerWriteVideoSample)(
 			[Params](FRHICommandListImmediate& RHICmdList)
 			{
-				ConditionalInitializeVideoSample_RenderThread(Params.VideoSample);
+				ConditionalInitializeVideoSample_RenderThread(RHICmdList, Params.VideoSample);
 				if (IsRunningRHIInSeparateThread())
 				{
 					new (RHICmdList.AllocCommand<FRHICommandUpdateTextureCameraSample>()) FRHICommandUpdateTextureCameraSample(Params.JavaCameraPlayerPtr, Params.SamplesPtr, Params.VideoSample, Params.SampleTime);
@@ -795,7 +795,7 @@ void FAndroidCameraPlayer::TickFetch(FTimespan DeltaTime, FTimespan Timecode)
 				}
 				else
 				{
-					DoUpdateTextureCameraSampleExecute(Params.JavaCameraPlayerPtr, Params.SamplesPtr, Params.VideoSample, Params.SampleTime);
+					DoUpdateTextureCameraSampleExecute(RHICmdList, Params.JavaCameraPlayerPtr, Params.SamplesPtr, Params.VideoSample, Params.SampleTime);
 				}
 			});
 	}
diff --git a/Engine/Plugins/Media/AndroidCamera/Source/AndroidCamera/Private/Player/AndroidCameraTextureSample.h b/Engine/Plugins/Media/AndroidCamera/Source/AndroidCamera/Private/Player/AndroidCameraTextureSample.h
--- a/Engine/Plugins/Media/AndroidCamera/Source/AndroidCamera/Private/Player/AndroidCameraTextureSample.h
+++ b/Engine/Plugins/Media/AndroidCamera/Source/AndroidCamera/Private/Player/AndroidCameraTextureSample.h
@@ -138,12 +138,13 @@ public:
 	/**
 	 * Initialize the sample with a texture resource.
 	 *
+	* @param RHICmdList the command list used to create the texture.
 	 * @param InTime The sample time (in the player's local clock).
 	 * @return The texture resource object that will hold the sample data.
 	 * @note This method must be called on the render thread.
 	 * @see InitializeBuffer
 	 */
-	FRHITexture* InitializeTexture(FTimespan InTime)
+	FRHITexture* InitializeTexture(FRHICommandListBase& RHICmdList, FTimespan InTime)
 	{
 		check(IsInRenderingThread() || IsInRHIThread());
 
@@ -161,7 +162,7 @@ public:
 			.SetFlags(ETextureCreateFlags::SRGB | ETextureCreateFlags::RenderTargetable | ETextureCreateFlags::ShaderResource)
 			.SetInitialState(ERHIAccess::SRVMask);
 
-		Texture = RHICreateTexture(Desc);
+		Texture = RHICmdList.CreateTexture(Desc);
 
 		return Texture;
 	}
diff --git a/Engine/Plugins/Media/AndroidMedia/Source/AndroidMedia/Private/Player/AndroidMediaPlayer.cpp b/Engine/Plugins/Media/AndroidMedia/Source/AndroidMedia/Private/Player/AndroidMediaPlayer.cpp
--- a/Engine/Plugins/Media/AndroidMedia/Source/AndroidMedia/Private/Player/AndroidMediaPlayer.cpp
+++ b/Engine/Plugins/Media/AndroidMedia/Source/AndroidMedia/Private/Player/AndroidMediaPlayer.cpp
@@ -463,7 +463,7 @@ void FAndroidMediaPlayer::SetGuid(const FGuid& Guid)
 // ==============================================================
 // Used by TickFetch
 
-static void ConditionalInitializeVideoTexture_RenderThread(TWeakPtr<FJavaAndroidMediaPlayer, ESPMode::ThreadSafe> JavaMediaPlayerPtr, FGuid PlayerGuid)
+static void ConditionalInitializeVideoTexture_RenderThread(FRHICommandListBase& RHICmdList, TWeakPtr<FJavaAndroidMediaPlayer, ESPMode::ThreadSafe> JavaMediaPlayerPtr, FGuid PlayerGuid)
 {
 	auto PinnedJavaMediaPlayer = JavaMediaPlayerPtr.Pin();
 
@@ -480,12 +480,12 @@ static void ConditionalInitializeVideoTexture_RenderThread(TWeakPtr<FJavaAndroid
 			.SetFlags(ETextureCreateFlags::External)
 			.SetInitialState(ERHIAccess::SRVMask);
 
-		VideoTexture = RHICreateTexture(Desc);
+		VideoTexture = RHICmdList.CreateTexture(Desc);
 		PinnedJavaMediaPlayer->SetVideoTexture(VideoTexture);
 
 		if (VideoTexture == nullptr)
 		{
-			UE_LOG(LogAndroidMedia, Warning, TEXT("RHICreateTexture failed!"));
+			UE_LOG(LogAndroidMedia, Warning, TEXT("CreateTexture failed!"));
 			return;
 		}
 
@@ -552,15 +552,15 @@ struct FRHICommandUpdateExternalMediaSample final : public FRHICommand<FRHIComma
 	}
 };
 
-static void ConditionalInitializeVideoSample_RenderThread(TSharedRef<FAndroidMediaTextureSample, ESPMode::ThreadSafe> VideoSample)
+static void ConditionalInitializeVideoSample_RenderThread(FRHICommandListBase& RHICmdList, TSharedRef<FAndroidMediaTextureSample, ESPMode::ThreadSafe> VideoSample)
 {
 	if (VideoSample->GetTexture() == nullptr || VideoSample->GetTexture()->GetSizeXY() != VideoSample->GetOutputDim())
 	{
-		VideoSample->InitializeTexture(FTimespan::Zero());
+		VideoSample->InitializeTexture(RHICmdList, FTimespan::Zero());
 	}
 }
 
-static void DoUpdateTextureMediaSampleExecute(TWeakPtr<FJavaAndroidMediaPlayer, ESPMode::ThreadSafe> JavaMediaPlayerPtr, TWeakPtr<FMediaSamples, ESPMode::ThreadSafe> SamplesPtr,
+static void DoUpdateTextureMediaSampleExecute(FRHICommandListBase& RHICmdList, TWeakPtr<FJavaAndroidMediaPlayer, ESPMode::ThreadSafe> JavaMediaPlayerPtr, TWeakPtr<FMediaSamples, ESPMode::ThreadSafe> SamplesPtr,
 	TSharedRef<FAndroidMediaTextureSample, ESPMode::ThreadSafe> VideoSample)
 {
 	auto PinnedJavaMediaPlayer = JavaMediaPlayerPtr.Pin();
@@ -575,7 +575,7 @@ static void DoUpdateTextureMediaSampleExecute(TWeakPtr<FJavaAndroidMediaPlayer,
 	const FTimespan Time = FTimespan::FromMilliseconds(CurrentFramePosition);
 
 	// write frame into texture
-	FRHITexture* Texture = VideoSample->InitializeTexture(Time);
+	FRHITexture* Texture = VideoSample->InitializeTexture(RHICmdList, Time);
 
 	if (Texture != nullptr)
 	{
@@ -602,10 +602,10 @@ struct FRHICommandUpdateTextureMediaSample final : public FRHICommand<FRHIComman
 		, VideoSample(InVideoSample)
 	{
 	}
-	void Execute(FRHICommandListBase& CmdList)
+	void Execute(FRHICommandListBase& RHICmdList)
 	{
 		QUICK_SCOPE_CYCLE_COUNTER(STAT_FRHICommandUpdateTextureMediaSample_Execute);
-		DoUpdateTextureMediaSampleExecute(JavaMediaPlayerPtr, SamplesPtr, VideoSample);
+		DoUpdateTextureMediaSampleExecute(RHICmdList, JavaMediaPlayerPtr, SamplesPtr, VideoSample);
 	}
 };
 
@@ -802,7 +802,7 @@ void FAndroidMediaPlayer::TickFetch(FTimespan DeltaTime, FTimespan /*Timecode*/)
 			ENQUEUE_RENDER_COMMAND(AndroidMediaPlayerWriteVideoSample)(
 				[Params](FRHICommandListImmediate& RHICmdList)
 				{
-					ConditionalInitializeVideoTexture_RenderThread(Params.JavaMediaPlayerPtr, Params.PlayerGuid);
+					ConditionalInitializeVideoTexture_RenderThread(RHICmdList, Params.JavaMediaPlayerPtr, Params.PlayerGuid);
 					if (IsRunningRHIInSeparateThread())
 					{
 						new (RHICmdList.AllocCommand<FRHICommandUpdateExternalMediaSample>()) FRHICommandUpdateExternalMediaSample(Params.JavaMediaPlayerPtr, Params.PlayerGuid);
@@ -888,14 +888,14 @@ void FAndroidMediaPlayer::TickFetch(FTimespan DeltaTime, FTimespan /*Timecode*/)
 		ENQUEUE_RENDER_COMMAND(AndroidMediaPlayerWriteVideoSample)(
 			[Params](FRHICommandListImmediate& RHICmdList)
 			{
-				ConditionalInitializeVideoSample_RenderThread(Params.VideoSample);
+				ConditionalInitializeVideoSample_RenderThread(RHICmdList, Params.VideoSample);
 				if (IsRunningRHIInSeparateThread())
 				{
 					new (RHICmdList.AllocCommand<FRHICommandUpdateTextureMediaSample>()) FRHICommandUpdateTextureMediaSample(Params.JavaMediaPlayerPtr, Params.SamplesPtr, Params.VideoSample);
 				}
 				else
 				{
-					DoUpdateTextureMediaSampleExecute(Params.JavaMediaPlayerPtr, Params.SamplesPtr, Params.VideoSample);
+					DoUpdateTextureMediaSampleExecute(RHICmdList, Params.JavaMediaPlayerPtr, Params.SamplesPtr, Params.VideoSample);
 				}
 			});
 	}
diff --git a/Engine/Plugins/Media/AndroidMedia/Source/AndroidMedia/Private/Player/AndroidMediaTextureSample.h b/Engine/Plugins/Media/AndroidMedia/Source/AndroidMedia/Private/Player/AndroidMediaTextureSample.h
--- a/Engine/Plugins/Media/AndroidMedia/Source/AndroidMedia/Private/Player/AndroidMediaTextureSample.h
+++ b/Engine/Plugins/Media/AndroidMedia/Source/AndroidMedia/Private/Player/AndroidMediaTextureSample.h
@@ -145,12 +145,13 @@ public:
 	/**
 	 * Initialize the sample with a texture resource.
 	 *
+	 * @param RHICmdList the command list used to create the texture.
 	 * @param InTime The sample time (in the player's local clock).
 	 * @return The texture resource object that will hold the sample data.
 	 * @note This method must be called on the render thread.
 	 * @see InitializeBuffer
 	 */
-	FRHITexture* InitializeTexture(FTimespan InTime)
+	FRHITexture* InitializeTexture(FRHICommandListBase& RHICmdList, FTimespan InTime)
 	{
 		check(IsInRenderingThread() || IsInRHIThread());
 
@@ -168,7 +169,7 @@ public:
 			.SetFlags(ETextureCreateFlags::SRGB | ETextureCreateFlags::RenderTargetable | ETextureCreateFlags::ShaderResource)
 			.SetInitialState(ERHIAccess::SRVMask);
 
-		Texture = RHICreateTexture(Desc);
+		Texture = RHICmdList.CreateTexture(Desc);
 
 		return Texture;
 	}
diff --git a/Engine/Plugins/Media/BinkMedia/Source/BinkMediaPlayer/Private/binktexturesRHI.cpp b/Engine/Plugins/Media/BinkMedia/Source/BinkMediaPlayer/Private/binktexturesRHI.cpp
--- a/Engine/Plugins/Media/BinkMedia/Source/BinkMediaPlayer/Private/binktexturesRHI.cpp
+++ b/Engine/Plugins/Media/BinkMedia/Source/BinkMediaPlayer/Private/binktexturesRHI.cpp
@@ -165,7 +165,7 @@ static BINKTEXTURES * Create_textures(BINKSHADERS * pshaders, HBINK bink, void *
 
 		if (bp_src->YPlane.Allocate)
 		{
-			textures->Ytexture[i] = RHICreateTexture(YADesc);
+			textures->Ytexture[i] = RHICmdList.CreateTexture(YADesc);
 #if PLATFORM_HAS_DIRECT_TEXTURE_MEMORY_ACCESS 
 			bp_src->YPlane.Buffer = RHICmdList.LockTexture2D(textures->Ytexture[i], 0, PLATFORM_DIRECT_TEXTURE_MEMORY_ACCESS_LOCK_MODE, bp_src->YPlane.BufferPitch, false, false);
 			bp_src->YPlane.BufferPitch = bp_src->YPlane.BufferPitch ? bp_src->YPlane.BufferPitch : ((bb->YABufferWidth+255)&-256);
@@ -174,7 +174,7 @@ static BINKTEXTURES * Create_textures(BINKSHADERS * pshaders, HBINK bink, void *
 
 		if (bp_src->cRPlane.Allocate)
 		{
-			textures->cRtexture[i] = RHICreateTexture(cRcBDesc);
+			textures->cRtexture[i] = RHICmdList.CreateTexture(cRcBDesc);
 #if PLATFORM_HAS_DIRECT_TEXTURE_MEMORY_ACCESS 
 			bp_src->cRPlane.Buffer = RHICmdList.LockTexture2D(textures->cRtexture[i], 0, PLATFORM_DIRECT_TEXTURE_MEMORY_ACCESS_LOCK_MODE, bp_src->cRPlane.BufferPitch, false, false);
 			bp_src->cRPlane.BufferPitch = bp_src->cRPlane.BufferPitch ? bp_src->cRPlane.BufferPitch : ((bb->cRcBBufferWidth+255)&-256);
@@ -183,7 +183,7 @@ static BINKTEXTURES * Create_textures(BINKSHADERS * pshaders, HBINK bink, void *
 
 		if (bp_src->cBPlane.Allocate)
 		{
-			textures->cBtexture[i] = RHICreateTexture(cRcBDesc);
+			textures->cBtexture[i] = RHICmdList.CreateTexture(cRcBDesc);
 #if PLATFORM_HAS_DIRECT_TEXTURE_MEMORY_ACCESS 
 			bp_src->cBPlane.Buffer = RHICmdList.LockTexture2D(textures->cBtexture[i], 0, PLATFORM_DIRECT_TEXTURE_MEMORY_ACCESS_LOCK_MODE, bp_src->cBPlane.BufferPitch, false, false);
 			bp_src->cBPlane.BufferPitch = bp_src->cBPlane.BufferPitch ? bp_src->cBPlane.BufferPitch : ((bb->cRcBBufferWidth+255)&-256);
@@ -192,7 +192,7 @@ static BINKTEXTURES * Create_textures(BINKSHADERS * pshaders, HBINK bink, void *
 
 		if (bp_src->APlane.Allocate)
 		{
-			textures->Atexture[i] = RHICreateTexture(YADesc);
+			textures->Atexture[i] = RHICmdList.CreateTexture(YADesc);
 #if PLATFORM_HAS_DIRECT_TEXTURE_MEMORY_ACCESS 
 			bp_src->APlane.Buffer = RHICmdList.LockTexture2D(textures->Atexture[i], 0, PLATFORM_DIRECT_TEXTURE_MEMORY_ACCESS_LOCK_MODE, bp_src->APlane.BufferPitch, false, false);
 			bp_src->APlane.BufferPitch = bp_src->APlane.BufferPitch ? bp_src->APlane.BufferPitch : ((bb->YABufferWidth+255)&-256);
@@ -201,7 +201,7 @@ static BINKTEXTURES * Create_textures(BINKSHADERS * pshaders, HBINK bink, void *
 
 		if (bp_src->HPlane.Allocate)
 		{
-			textures->Htexture[i] = RHICreateTexture(YADesc);
+			textures->Htexture[i] = RHICmdList.CreateTexture(YADesc);
 #if PLATFORM_HAS_DIRECT_TEXTURE_MEMORY_ACCESS 
 			bp_src->HPlane.Buffer = RHICmdList.LockTexture2D(textures->Htexture[i], 0, PLATFORM_DIRECT_TEXTURE_MEMORY_ACCESS_LOCK_MODE, bp_src->HPlane.BufferPitch, false, false);
 			bp_src->HPlane.BufferPitch = bp_src->HPlane.BufferPitch ? bp_src->HPlane.BufferPitch : ((bb->YABufferWidth+255)&-256);
diff --git a/Engine/Plugins/Media/ElectraUtil/Source/ElectraSamples/Private/Android/ElectraTextureSample_Android.cpp b/Engine/Plugins/Media/ElectraUtil/Source/ElectraSamples/Private/Android/ElectraTextureSample_Android.cpp
--- a/Engine/Plugins/Media/ElectraUtil/Source/ElectraSamples/Private/Android/ElectraTextureSample_Android.cpp
+++ b/Engine/Plugins/Media/ElectraUtil/Source/ElectraSamples/Private/Android/ElectraTextureSample_Android.cpp
@@ -543,7 +543,7 @@ bool FElectraTextureSample::Initialize()
 		ENQUEUE_RENDER_COMMAND(InitTextureSample)([WeakThis{ AsWeak() }](FRHICommandListImmediate& RHICmdList) {
 			if (TSharedPtr<FElectraTextureSample, ESPMode::ThreadSafe> This = WeakThis.Pin())
 			{
-				This->InitializeTexture(This->GetPixelFormat());
+				This->InitializeTexture(RHICmdList, This->GetPixelFormat());
 
 				if (This->Texture)
 				{
@@ -651,7 +651,7 @@ FTextureRHIRef FElectraTextureSample::InitializeTextureOES(AHardwareBuffer* Hard
 		.SetInitialState(ERHIAccess::SRVMask)
 		.SetFlags(ETextureCreateFlags::External);
 
-	return RHICreateTexture(Desc);
+	return RHICmdList.CreateTexture(Desc);
 
 #endif
 
@@ -668,7 +668,7 @@ FTextureRHIRef FElectraTextureSample::InitializeTextureVulkan(AHardwareBuffer* H
 }
 
 
-void FElectraTextureSample::InitializeTexture(EPixelFormat PixelFormat)
+void FElectraTextureSample::InitializeTexture(FRHICommandListBase& RHICmdList, EPixelFormat PixelFormat)
 {
 	check(IsInRenderingThread() || IsInRHIThread());
 
@@ -694,7 +694,7 @@ void FElectraTextureSample::InitializeTexture(EPixelFormat PixelFormat)
 		.SetInitialState(ERHIAccess::SRVMask)
 		.SetFlags(ETextureCreateFlags::RenderTargetable | ETextureCreateFlags::ShaderResource)
 		.SetFormat(PixelFormat);
-	Texture = RHICreateTexture(Desc);
+	Texture = RHICmdList.CreateTexture(Desc);
 
 	return;
 }
diff --git a/Engine/Plugins/Media/ElectraUtil/Source/ElectraSamples/Public/Android/ElectraTextureSample.h b/Engine/Plugins/Media/ElectraUtil/Source/ElectraSamples/Public/Android/ElectraTextureSample.h
--- a/Engine/Plugins/Media/ElectraUtil/Source/ElectraSamples/Public/Android/ElectraTextureSample.h
+++ b/Engine/Plugins/Media/ElectraUtil/Source/ElectraSamples/Public/Android/ElectraTextureSample.h
@@ -56,7 +56,7 @@ private:
 	UE_API FTextureRHIRef InitializeTextureOES(AHardwareBuffer* HardwareBuffer);
 	UE_API FTextureRHIRef InitializeTextureVulkan(AHardwareBuffer* HardwareBuffer);
 
-	UE_API void InitializeTexture(EPixelFormat PixelFormat);
+	UE_API void InitializeTexture(FRHICommandListBase& RHICmdList, EPixelFormat PixelFormat);
 	UE_API void SetupFromBuffer(const void* InBuffer, int32 InBufferSize);
 	UE_API void SetImageResources(jobject InImageResources);
 	UE_API void CleanupImageResources();
diff --git a/Engine/Plugins/Media/MediaIOFramework/Source/MediaIOCore/Private/Player/MediaIOCorePlayerBase.cpp b/Engine/Plugins/Media/MediaIOFramework/Source/MediaIOCore/Private/Player/MediaIOCorePlayerBase.cpp
--- a/Engine/Plugins/Media/MediaIOFramework/Source/MediaIOCore/Private/Player/MediaIOCorePlayerBase.cpp
+++ b/Engine/Plugins/Media/MediaIOFramework/Source/MediaIOCore/Private/Player/MediaIOCorePlayerBase.cpp
@@ -789,7 +789,7 @@ void FMediaIOCorePlayerBase::CreateAndRegisterTextures()
 					.SetFlags(CreateFlags | ETextureCreateFlags::RenderTargetable | ETextureCreateFlags::ShaderResource)
 					.SetInitialState(ERHIAccess::SRVMask);
 
-				RHITexture = RHICreateTexture(CreateDesc);
+				RHITexture = CommandList.CreateTexture(CreateDesc);
 
 				UE_LOG(LogMediaIOCore, Verbose, TEXT("Registering texture %u"), reinterpret_cast<uintptr_t>(RHITexture->GetNativeResource()));
 
diff --git a/Engine/Plugins/Media/NDIMedia/Source/NDIMedia/Private/Player/NDIMediaTextureSampleConverter.cpp b/Engine/Plugins/Media/NDIMedia/Source/NDIMedia/Private/Player/NDIMediaTextureSampleConverter.cpp
--- a/Engine/Plugins/Media/NDIMedia/Source/NDIMedia/Private/Player/NDIMediaTextureSampleConverter.cpp
+++ b/Engine/Plugins/Media/NDIMedia/Source/NDIMedia/Private/Player/NDIMediaTextureSampleConverter.cpp
@@ -146,7 +146,7 @@ bool FNDIMediaTextureSampleConverter::UpdateInputTextures(FRHICommandList& InRHI
 			.SetNumMips(1)
 			.SetFlags(ETextureCreateFlags::RenderTargetable | ETextureCreateFlags::Dynamic);
 
-		SourceYUVTexture = RHICreateTexture(CreateDesc);
+		SourceYUVTexture = InRHICmdList.CreateTexture(CreateDesc);
 	}
 
 	if (!SourceAlphaTexture.IsValid() || FrameSize != PreviousFrameSize)
@@ -157,20 +157,20 @@ bool FNDIMediaTextureSampleConverter::UpdateInputTextures(FRHICommandList& InRHI
 			.SetNumMips(1)
 			.SetFlags(ETextureCreateFlags::RenderTargetable | ETextureCreateFlags::Dynamic);
 
-		SourceAlphaTexture = RHICreateTexture(CreateAlphaDesc);
+		SourceAlphaTexture = InRHICmdList.CreateTexture(CreateAlphaDesc);
 	}
 
 	PreviousFrameSize = FrameSize;
 
 	// Update GPU textures from sample buffer
 	FUpdateTextureRegion2D YUVRegion(0, 0, 0, 0, FieldSize.X/2, FieldSize.Y);
-	RHIUpdateTexture2D(SourceYUVTexture, 0, YUVRegion, InSample->GetStride(), (uint8*&)Buffer);
+	InRHICmdList.UpdateTexture2D(SourceYUVTexture, 0, YUVRegion, InSample->GetStride(), (uint8*&)Buffer);
 
 	// Make sure resource is in SRV mode again
 	InRHICmdList.Transition(FRHITransitionInfo(SourceYUVTexture, ERHIAccess::Unknown, ERHIAccess::SRVMask));
 	
 	FUpdateTextureRegion2D AlphaRegion(0, 0, 0, 0, FieldSize.X, FieldSize.Y);
-	RHIUpdateTexture2D(SourceAlphaTexture, 0, AlphaRegion, FrameSize.X, ((uint8*&)Buffer) + FieldSize.Y * InSample->GetStride());
+	InRHICmdList.UpdateTexture2D(SourceAlphaTexture, 0, AlphaRegion, FrameSize.X, ((uint8*&)Buffer) + FieldSize.Y * InSample->GetStride());
 
 	// Make sure resource is in SRV mode again
 	InRHICmdList.Transition(FRHITransitionInfo(SourceAlphaTexture, ERHIAccess::Unknown, ERHIAccess::SRVMask));
diff --git a/Engine/Plugins/Media/PixelStreaming/Source/PixelStreaming/Private/FreezeFrame.cpp b/Engine/Plugins/Media/PixelStreaming/Source/PixelStreaming/Private/FreezeFrame.cpp
--- a/Engine/Plugins/Media/PixelStreaming/Source/PixelStreaming/Private/FreezeFrame.cpp
+++ b/Engine/Plugins/Media/PixelStreaming/Source/PixelStreaming/Private/FreezeFrame.cpp
@@ -138,7 +138,7 @@ namespace UE::PixelStreaming
 						.SetInitialState(ERHIAccess::Present)
 						.DetermineInititialState();
 
-				FTextureRHIRef DestTexture = RHICreateTexture(TextureDesc);
+				FTextureRHIRef DestTexture = RHICmdList.CreateTexture(TextureDesc);
 
 				// Copy freeze frame texture to empty texture
 				CopyTexture(RHICmdList, TextureRHI, DestTexture);
diff --git a/Engine/Plugins/Media/PixelStreaming2/Source/PixelStreaming2RTC/Private/FreezeFrame.cpp b/Engine/Plugins/Media/PixelStreaming2/Source/PixelStreaming2RTC/Private/FreezeFrame.cpp
--- a/Engine/Plugins/Media/PixelStreaming2/Source/PixelStreaming2RTC/Private/FreezeFrame.cpp
+++ b/Engine/Plugins/Media/PixelStreaming2/Source/PixelStreaming2RTC/Private/FreezeFrame.cpp
@@ -131,7 +131,7 @@ namespace UE::PixelStreaming2
 						.SetInitialState(ERHIAccess::Present)
 						.DetermineInititialState();
 
-				FTextureRHIRef DestTexture = RHICreateTexture(TextureDesc);
+				FTextureRHIRef DestTexture = RHICmdList.CreateTexture(TextureDesc);
 
 				// Copy freeze frame texture to empty texture
 				CopyTexture(RHICmdList, TextureRHI, DestTexture);
diff --git a/Engine/Plugins/Media/WebMMedia/Source/WebMMedia/Private/Player/WebMMediaTextureSample.cpp b/Engine/Plugins/Media/WebMMedia/Source/WebMMedia/Private/Player/WebMMediaTextureSample.cpp
--- a/Engine/Plugins/Media/WebMMedia/Source/WebMMedia/Private/Player/WebMMediaTextureSample.cpp
+++ b/Engine/Plugins/Media/WebMMedia/Source/WebMMedia/Private/Player/WebMMediaTextureSample.cpp
@@ -16,7 +16,7 @@ void FWebMMediaTextureSample::Initialize(FIntPoint InDisplaySize, FIntPoint InTo
 	Duration = InDuration;
 }
 
-void FWebMMediaTextureSample::CreateTexture()
+void FWebMMediaTextureSample::CreateTexture(FRHICommandListBase& RHICmdList)
 {
 	check(IsInRenderingThread());
 
@@ -27,7 +27,7 @@ void FWebMMediaTextureSample::CreateTexture()
 		.SetFlags(ETextureCreateFlags::SRGB | ETextureCreateFlags::RenderTargetable | ETextureCreateFlags::ShaderResource)
 		.SetInitialState(ERHIAccess::SRVMask);
 
-	Texture = RHICreateTexture(Desc);
+	Texture = RHICmdList.CreateTexture(Desc);
 }
 
 const void* FWebMMediaTextureSample::GetBuffer()
diff --git a/Engine/Plugins/Media/WebMMedia/Source/WebMMedia/Private/Player/WebMVideoDecoder.cpp b/Engine/Plugins/Media/WebMMedia/Source/WebMMedia/Private/Player/WebMVideoDecoder.cpp
--- a/Engine/Plugins/Media/WebMMedia/Source/WebMMedia/Private/Player/WebMVideoDecoder.cpp
+++ b/Engine/Plugins/Media/WebMMedia/Source/WebMMedia/Private/Player/WebMVideoDecoder.cpp
@@ -145,7 +145,7 @@ void FWebMVideoDecoder::DoDecodeVideoFrames(const TArray<TSharedPtr<FWebMFrame>>
 				ENQUEUE_RENDER_COMMAND(WebMMediaPlayerCreateTextures)(
 					[Self, Image](FRHICommandListImmediate& RHICmdList)
 					{
-						Self->CreateTextures(Image);
+						Self->CreateTextures(RHICmdList, Image);
 					});
 			}
 
@@ -159,13 +159,13 @@ void FWebMVideoDecoder::DoDecodeVideoFrames(const TArray<TSharedPtr<FWebMFrame>>
 			ENQUEUE_RENDER_COMMAND(WebMMediaPlayerConvertYUVToRGB)(
 				[Self, Params](FRHICommandListImmediate& RHICmdList)
 				{
-					Self->ConvertYUVToRGBAndSubmit(Params);
+					Self->ConvertYUVToRGBAndSubmit(RHICmdList, Params);
 				});
 		}
 	}
 }
 
-void FWebMVideoDecoder::CreateTextures(const vpx_image_t* Image)
+void FWebMVideoDecoder::CreateTextures(FRHICommandListBase& RHICmdList, const vpx_image_t* Image)
 {
 	const FRHITextureCreateDesc DecodedYDesc =
 		FRHITextureCreateDesc::Create2D(TEXT("FWebMVideoDecoder_DecodedY"), Image->stride[0], Image->d_h, PF_G8);
@@ -176,9 +176,9 @@ void FWebMVideoDecoder::CreateTextures(const vpx_image_t* Image)
 	const FRHITextureCreateDesc DecodedVDesc =
 		FRHITextureCreateDesc::Create2D(TEXT("FWebMVideoDecoder_DecodedV"), Image->stride[2], Image->d_h / 2, PF_G8);
 
-	DecodedY = RHICreateTexture(DecodedYDesc);
-	DecodedU = RHICreateTexture(DecodedUDesc);
-	DecodedV = RHICreateTexture(DecodedVDesc);
+	DecodedY = RHICmdList.CreateTexture(DecodedYDesc);
+	DecodedU = RHICmdList.CreateTexture(DecodedUDesc);
+	DecodedV = RHICmdList.CreateTexture(DecodedVDesc);
 }
 
 void FWebMVideoDecoder::Close()
@@ -206,16 +206,15 @@ void FWebMVideoDecoder::Close()
 	bTexturesCreated = false;
 }
 
-void FWebMVideoDecoder::ConvertYUVToRGBAndSubmit(const FConvertParams& Params)
+void FWebMVideoDecoder::ConvertYUVToRGBAndSubmit(FRHICommandListImmediate& CommandList, const FConvertParams& Params)
 {
 	TSharedPtr<FWebMMediaTextureSample, ESPMode::ThreadSafe> VideoSample = Params.VideoSample;
 	check(VideoSample.IsValid());
 	const vpx_image_t* Image = Params.Image;
 
-	VideoSample->CreateTexture();
+	VideoSample->CreateTexture(CommandList);
 
 	// render video frame into output texture
-	FRHICommandListImmediate& CommandList = GetImmediateCommandList_ForRenderCommand();
 
 	{
 		const auto CopyTextureMemory = [Image](
diff --git a/Engine/Plugins/Media/WebMMedia/Source/WebMMedia/Public/WebMMediaTextureSample.h b/Engine/Plugins/Media/WebMMedia/Source/WebMMedia/Public/WebMMediaTextureSample.h
--- a/Engine/Plugins/Media/WebMMedia/Source/WebMMedia/Public/WebMMediaTextureSample.h
+++ b/Engine/Plugins/Media/WebMMedia/Source/WebMMedia/Public/WebMMediaTextureSample.h
@@ -21,7 +21,7 @@ public:
 
 public:
 	UE_API void Initialize(FIntPoint InDisplaySize, FIntPoint InTotalSize, FTimespan InTime, FTimespan InDuration);
-	UE_API void CreateTexture();
+	UE_API void CreateTexture(FRHICommandListBase& RHICmdList);
 
 public:
 	//~ IMediaTextureSample interface
diff --git a/Engine/Plugins/Media/WebMMedia/Source/WebMMedia/Public/WebMVideoDecoder.h b/Engine/Plugins/Media/WebMMedia/Source/WebMMedia/Public/WebMVideoDecoder.h
--- a/Engine/Plugins/Media/WebMMedia/Source/WebMMedia/Public/WebMVideoDecoder.h
+++ b/Engine/Plugins/Media/WebMMedia/Source/WebMMedia/Public/WebMVideoDecoder.h
@@ -45,9 +45,9 @@ private:
 	bool bTexturesCreated;
 	bool bIsInitialized;
 
-	void ConvertYUVToRGBAndSubmit(const FConvertParams& Params);
+	void ConvertYUVToRGBAndSubmit(FRHICommandListImmediate& RHICmdList, const FConvertParams& Params);
 	void DoDecodeVideoFrames(const TArray<TSharedPtr<FWebMFrame>>& VideoFrames);
-	void CreateTextures(const vpx_image_t* Image);
+	void CreateTextures(FRHICommandListBase& RHICmdList, const vpx_image_t* Image);
 	void Close();
 };
 
diff --git a/Engine/Plugins/Media/WmfMedia/Source/WmfMedia/Private/Wmf/WmfMediaHardwareVideoDecodingRendering.cpp b/Engine/Plugins/Media/WmfMedia/Source/WmfMedia/Private/Wmf/WmfMediaHardwareVideoDecodingRendering.cpp
--- a/Engine/Plugins/Media/WmfMedia/Source/WmfMedia/Private/Wmf/WmfMediaHardwareVideoDecodingRendering.cpp
+++ b/Engine/Plugins/Media/WmfMedia/Source/WmfMedia/Private/Wmf/WmfMediaHardwareVideoDecodingRendering.cpp
@@ -139,7 +139,7 @@ bool FWmfMediaHardwareVideoDecodingParameters::ConvertTextureFormat_RenderThread
 		GraphicsPSOInit.BoundShaderState.VertexDeclarationRHI = GetVertexDeclarationFVector4();
 
 		// Copy buffer.
-		FTextureRHIRef SampleDestinationTexture = InSample->GetOrCreateDestinationTexture();
+		FTextureRHIRef SampleDestinationTexture = InSample->GetOrCreateDestinationTexture(RHICmdList);
 		if (InSample->IsBufferExternal() || !SampleTexture.IsValid())
 		{
 			const uint8* Data = (const uint8*)InSample->GetBuffer();
@@ -265,7 +265,7 @@ bool FWmfMediaHardwareVideoDecodingParameters::ConvertTextureFormat_RenderThread
 		else if (InSample->GetFormat() == EMediaTextureSampleFormat::YCoCg_DXT5_Alpha_BC4)
 		{
 			// Set up alpha texture.
-			FTextureRHIRef SampleDestinationAlphaTexture = InSample->GetOrCreateDestinationAlphaTexture();
+			FTextureRHIRef SampleDestinationAlphaTexture = InSample->GetOrCreateDestinationAlphaTexture(RHICmdList);
 			{
 				const uint8* Data = ((const uint8*)InSample->GetBuffer()) + InSample->GetDim().X * InSample->GetDim().Y;
 
diff --git a/Engine/Plugins/Media/WmfMedia/Source/WmfMedia/Public/WmfMediaHardwareVideoDecodingTextureSample.h b/Engine/Plugins/Media/WmfMedia/Source/WmfMedia/Public/WmfMediaHardwareVideoDecodingTextureSample.h
--- a/Engine/Plugins/Media/WmfMedia/Source/WmfMedia/Public/WmfMediaHardwareVideoDecodingTextureSample.h
+++ b/Engine/Plugins/Media/WmfMedia/Source/WmfMedia/Public/WmfMediaHardwareVideoDecodingTextureSample.h
@@ -113,9 +113,10 @@ public:
 	/**
 	 * Get Destination Texture of render thread device
 	 *
+	* @param RHICmdList the command list used to create the texture.
 	 * @return Destination texture 
 	 */
-	FTextureRHIRef GetOrCreateDestinationTexture()
+	FTextureRHIRef GetOrCreateDestinationTexture(FRHICommandListBase& RHICmdList)
 	{
 		if (DestinationTexture.IsValid() && DestinationTexture->GetSizeX() == Dim.X && DestinationTexture->GetSizeY() == Dim.Y)
 		{
@@ -136,7 +137,7 @@ public:
 			Desc.AddFlags(ETextureCreateFlags::SRGB);
 		}
 
-		DestinationTexture = RHICreateTexture(Desc);
+		DestinationTexture = RHICmdList.CreateTexture(Desc);
 
 		return DestinationTexture;
 	}
@@ -144,9 +145,10 @@ public:
 	/**
 	 * Get Destination Texture of render thread device
 	 *
+	* @param RHICmdList the command list used to create the texture.
 	 * @return Destination texture
 	 */
-	FTextureRHIRef GetOrCreateDestinationAlphaTexture()
+	FTextureRHIRef GetOrCreateDestinationAlphaTexture(FRHICommandListBase& RHICmdList)
 	{
 		if (DestinationAlphaTexture.IsValid() && DestinationAlphaTexture->GetSizeX() == Dim.X && DestinationAlphaTexture->GetSizeY() == Dim.Y)
 		{
@@ -157,7 +159,7 @@ public:
 			FRHITextureCreateDesc::Create2D(TEXT("FWmfMediaHardwareVideoDecodingTextureSample_DestinationAlphaTexture"), Dim, AlphaFormat)
 			.SetFlags(ETextureCreateFlags::DisableSRVCreation);
 		
-		DestinationAlphaTexture = RHICreateTexture(Desc);
+		DestinationAlphaTexture = RHICmdList.CreateTexture(Desc);
 
 		return DestinationAlphaTexture;
 	}
diff --git a/Engine/Plugins/PCG/Source/PCGCompute/Private/PCGTextureReadback.cpp b/Engine/Plugins/PCG/Source/PCGCompute/Private/PCGTextureReadback.cpp
--- a/Engine/Plugins/PCG/Source/PCGCompute/Private/PCGTextureReadback.cpp
+++ b/Engine/Plugins/PCG/Source/PCGCompute/Private/PCGTextureReadback.cpp
@@ -54,7 +54,7 @@ void FPCGTextureReadbackInterface::Dispatch_RenderThread(FRHICommandListImmediat
 			FRHITextureCreateDesc::Create2D(TEXT("PCGDummyTexture"), 1, 1, PF_G8)
 			.SetFlags(ETextureCreateFlags::ShaderResource);
 
-		FTextureRHIRef DummyTexture = RHICreateTexture(DummyTextureDesc);
+		FTextureRHIRef DummyTexture = RHICmdList.CreateTexture(DummyTextureDesc);
 
 		{
 			uint32 DestStride;
@@ -73,7 +73,7 @@ void FPCGTextureReadbackInterface::Dispatch_RenderThread(FRHICommandListImmediat
 			FRHITextureCreateDesc::Create2DArray(TEXT("PCGDummyTextureArray"), 1, 1, 1, PF_G8)
 			.SetFlags(ETextureCreateFlags::ShaderResource);
 
-		FTextureRHIRef DummyTexture = RHICreateTexture(DummyTextureDesc);
+		FTextureRHIRef DummyTexture = RHICmdList.CreateTexture(DummyTextureDesc);
 
 		{
 			uint32 DestStride;
@@ -99,7 +99,7 @@ void FPCGTextureReadbackInterface::Dispatch_RenderThread(FRHICommandListImmediat
 	check(TargetTextureDesc.IsValid());
 	
 	// Create temporary output texture
-	FTextureRHIRef OutputTexture = RHICreateTexture(TargetTextureDesc);
+	FTextureRHIRef OutputTexture = RHICmdList.CreateTexture(TargetTextureDesc);
 	PassParameters.OutputTexture = RHICmdList.CreateUnorderedAccessView(OutputTexture, FRHIViewDesc::CreateTextureUAV().SetDimensionFromTexture(OutputTexture));
 
 	TShaderMapRef<FPCGTextureReadbackCS> ComputeShader(GetGlobalShaderMap(GMaxRHIFeatureLevel));
diff --git a/Engine/Plugins/Runtime/HairStrands/Source/HairStrandsCore/Private/GroomResources.cpp b/Engine/Plugins/Runtime/HairStrands/Source/HairStrandsCore/Private/GroomResources.cpp
--- a/Engine/Plugins/Runtime/HairStrands/Source/HairStrandsCore/Private/GroomResources.cpp
+++ b/Engine/Plugins/Runtime/HairStrands/Source/HairStrandsCore/Private/GroomResources.cpp
@@ -838,7 +838,7 @@ OwnerName(InOwnerName)
 {
 }
 
-void FHairCommonResource::InitRHI(FRHICommandListBase&)
+void FHairCommonResource::InitRHI(FRHICommandListBase& RHICmdList)
 {
 	// 1. If the resource is Updated (ReleaseRHI, IniRHI), reset its loaded data counter
 	if (!bIsInitialized)
@@ -859,8 +859,7 @@ void FHairCommonResource::InitRHI(FRHICommandListBase&)
 	// 4. Allocate the resource, and update the data
 	if (bUseRenderGraph)
 	{
-		FRHICommandListImmediate& RHICmdList = FRHICommandListExecutor::GetImmediateCommandList();
-		FRDGBuilder GraphBuilder(RHICmdList);
+		FRDGBuilder GraphBuilder(RHICmdList.GetAsImmediate());
 		InternalAllocate(GraphBuilder, HAIR_MAX_NUM_CURVE_PER_GROUP, HAIR_MAX_NUM_POINT_PER_GROUP);
 		GraphBuilder.Execute();
 	}
diff --git a/Engine/Plugins/Runtime/HairStrands/Source/HairStrandsCore/Private/Niagara/NiagaraDataInterfaceHairStrands.cpp b/Engine/Plugins/Runtime/HairStrands/Source/HairStrandsCore/Private/Niagara/NiagaraDataInterfaceHairStrands.cpp
--- a/Engine/Plugins/Runtime/HairStrands/Source/HairStrandsCore/Private/Niagara/NiagaraDataInterfaceHairStrands.cpp
+++ b/Engine/Plugins/Runtime/HairStrands/Source/HairStrandsCore/Private/Niagara/NiagaraDataInterfaceHairStrands.cpp
@@ -430,13 +430,13 @@ void FNDIHairStrandsBuffer::Transfer(FRDGBuilder& GraphBuilder, const TStaticArr
 	}
 }
 
-void FNDIHairStrandsBuffer::InitRHI(FRHICommandListBase&)
+void FNDIHairStrandsBuffer::InitRHI(FRHICommandListBase& RHICmdList)
 {
 	//ReadbackBuffer = new FRHIGPUBufferReadback(TEXT("Hair.PositionOffsetBuffer"));
 	
 	if (bNeedResouces)
 	{
-		FRDGBuilder GraphBuilder(FRHICommandListImmediate::Get());
+		FRDGBuilder GraphBuilder(RHICmdList.GetAsImmediate());
 		{
 			static const uint32 ZeroData[] = {
 				UINT_MAX,UINT_MAX,UINT_MAX,0,0,0,
diff --git a/Engine/Plugins/Runtime/HairStrands/Source/HairStrandsCore/Private/Niagara/NiagaraDataInterfaceVelocityGrid.cpp b/Engine/Plugins/Runtime/HairStrands/Source/HairStrandsCore/Private/Niagara/NiagaraDataInterfaceVelocityGrid.cpp
--- a/Engine/Plugins/Runtime/HairStrands/Source/HairStrandsCore/Private/Niagara/NiagaraDataInterfaceVelocityGrid.cpp
+++ b/Engine/Plugins/Runtime/HairStrands/Source/HairStrandsCore/Private/Niagara/NiagaraDataInterfaceVelocityGrid.cpp
@@ -54,14 +54,14 @@ void FNDIVelocityGridBuffer::Initialize(const FIntVector InGridSize, const int32
 	NumAttributes = InNumAttributes;
 }
 
-void FNDIVelocityGridBuffer::InitRHI(FRHICommandListBase&)
+void FNDIVelocityGridBuffer::InitRHI(FRHICommandListBase& RHICmdList)
 {
 	if (GridSize.X != 0 && GridSize.Y != 0 && GridSize.Z != 0)
 	{
 		static const uint32 NumComponents = NumAttributes;
 
 		FMemMark MemMark(FMemStack::Get());
-		FRDGBuilder GraphBuilder(FRHICommandListImmediate::Get());
+		FRDGBuilder GraphBuilder(RHICmdList.GetAsImmediate());
 		GridDataBuffer.Initialize(GraphBuilder, TEXT("FNDIVelocityGridBuffer"), EPixelFormat::PF_R32_SINT, sizeof(int32), (GridSize.X + 1) * NumComponents * (GridSize.Y + 1) * (GridSize.Z + 1));
 		GridDataBuffer.EndGraphUsage();
 		GraphBuilder.Execute();
diff --git a/Engine/Plugins/Runtime/WarpUtils/Source/PFMExporter/Private/PFMExporterMesh.cpp b/Engine/Plugins/Runtime/WarpUtils/Source/PFMExporter/Private/PFMExporterMesh.cpp
--- a/Engine/Plugins/Runtime/WarpUtils/Source/PFMExporter/Private/PFMExporterMesh.cpp
+++ b/Engine/Plugins/Runtime/WarpUtils/Source/PFMExporter/Private/PFMExporterMesh.cpp
@@ -55,12 +55,12 @@ bool FPFMExporterMesh::BeginExport_RenderThread(FRHICommandListImmediate& RHICmd
 		.SetFormat(PFMTextureFormat)
 		.SetInitialState(ERHIAccess::CopySrc);
 
-	ShaderResourceTexture = RHICreateTexture(Desc);
+	ShaderResourceTexture = RHICmdList.CreateTexture(Desc);
 
 	Desc.SetFlags(ETextureCreateFlags::RenderTargetable | ETextureCreateFlags::ShaderResource)
 		.SetInitialState(ERHIAccess::RTV);
 
-	RenderTargetTexture = RHICreateTexture(Desc);
+	RenderTargetTexture = RHICmdList.CreateTexture(Desc);
 
 	return true;
 }
diff --git a/Engine/Plugins/Runtime/nDisplay/Source/DisplayCluster/Private/Render/Containers/DisplayClusterRender_TextureResource.cpp b/Engine/Plugins/Runtime/nDisplay/Source/DisplayCluster/Private/Render/Containers/DisplayClusterRender_TextureResource.cpp
--- a/Engine/Plugins/Runtime/nDisplay/Source/DisplayCluster/Private/Render/Containers/DisplayClusterRender_TextureResource.cpp
+++ b/Engine/Plugins/Runtime/nDisplay/Source/DisplayCluster/Private/Render/Containers/DisplayClusterRender_TextureResource.cpp
@@ -108,7 +108,7 @@ void FDisplayClusterRender_TextureResource::ReleaseTextureData()
 	}
 }
 
-void FDisplayClusterRender_TextureResource::InitRHI(FRHICommandListBase&)
+void FDisplayClusterRender_TextureResource::InitRHI(FRHICommandListBase& RHICmdList)
 {
 	check(IsInRenderingThread());
 
@@ -127,7 +127,7 @@ void FDisplayClusterRender_TextureResource::InitRHI(FRHICommandListBase&)
 		.SetFlags(ETextureCreateFlags::ShaderResource)
 		.SetInitActionBulkData(&BulkDataInterface);
 
-	TextureRHI = RHICreateTexture(Desc);
+	TextureRHI = RHICmdList.CreateTexture(Desc);
 
 	// CPU access not required, release from memory
 	if (bHasCPUAccess == false)
diff --git a/Engine/Plugins/Runtime/nDisplay/Source/DisplayCluster/Private/Render/Viewport/LightCard/DisplayClusterViewportLightCardResource.cpp b/Engine/Plugins/Runtime/nDisplay/Source/DisplayCluster/Private/Render/Viewport/LightCard/DisplayClusterViewportLightCardResource.cpp
--- a/Engine/Plugins/Runtime/nDisplay/Source/DisplayCluster/Private/Render/Viewport/LightCard/DisplayClusterViewportLightCardResource.cpp
+++ b/Engine/Plugins/Runtime/nDisplay/Source/DisplayCluster/Private/Render/Viewport/LightCard/DisplayClusterViewportLightCardResource.cpp
@@ -5,7 +5,7 @@
 ///////////////////////////////////////////////////////////////////////////////////////////////
 // FDisplayClusterViewportLightCardResource
 ///////////////////////////////////////////////////////////////////////////////////////////////
-void FDisplayClusterViewportLightCardResource::InitRHI(FRHICommandListBase&)
+void FDisplayClusterViewportLightCardResource::InitRHI(FRHICommandListBase& RHICmdList)
 {
 	ETextureCreateFlags CreateFlags = TexCreate_MultiGPUGraphIgnore;
 
@@ -16,5 +16,5 @@ void FDisplayClusterViewportLightCardResource::InitRHI(FRHICommandListBase&)
 		.SetInitialState(ERHIAccess::SRVMask)
 		.SetClearValue(FClearValueBinding::Transparent);
 
-	RenderTargetTextureRHI = TextureRHI = RHICreateTexture(Desc);
+	RenderTargetTextureRHI = TextureRHI = RHICmdList.CreateTexture(Desc);
 }
diff --git a/Engine/Plugins/Runtime/nDisplay/Source/DisplayCluster/Private/Render/Viewport/RenderTarget/DisplayClusterRenderTargetResource.cpp b/Engine/Plugins/Runtime/nDisplay/Source/DisplayCluster/Private/Render/Viewport/RenderTarget/DisplayClusterRenderTargetResource.cpp
--- a/Engine/Plugins/Runtime/nDisplay/Source/DisplayCluster/Private/Render/Viewport/RenderTarget/DisplayClusterRenderTargetResource.cpp
+++ b/Engine/Plugins/Runtime/nDisplay/Source/DisplayCluster/Private/Render/Viewport/RenderTarget/DisplayClusterRenderTargetResource.cpp
@@ -16,18 +16,18 @@ FDisplayClusterViewportTextureResource::FDisplayClusterViewportTextureResource(c
 	bGreyScaleFormat = false;
 }
 
-void FDisplayClusterViewportTextureResource::InitRHI(FRHICommandListBase&)
+void FDisplayClusterViewportTextureResource::InitRHI(FRHICommandListBase& RHICmdList)
 {
 	FTextureRHIRef NewTextureRHI;
 
 	if (GetResourceSettings().GetNumMips() > 1)
 	{
 		FTextureRHIRef DummyTextureRHI;
-		ImplInitDynamicRHI_RenderTargetResource2D(NewTextureRHI, DummyTextureRHI);
+		ImplInitDynamicRHI_RenderTargetResource2D(RHICmdList, NewTextureRHI, DummyTextureRHI);
 	}
 	else
 	{
-		ImplInitDynamicRHI_TextureResource2D(NewTextureRHI);
+		ImplInitDynamicRHI_TextureResource2D(RHICmdList, NewTextureRHI);
 	}
 
 	TextureRHI = (FTextureRHIRef&)NewTextureRHI;
@@ -36,11 +36,11 @@ void FDisplayClusterViewportTextureResource::InitRHI(FRHICommandListBase&)
 ////////////////////////////////////////////////////////////////////////
 // FDisplayClusterViewportRenderTargetResource
 ////////////////////////////////////////////////////////////////////////
-void FDisplayClusterViewportRenderTargetResource::InitRHI(FRHICommandListBase&)
+void FDisplayClusterViewportRenderTargetResource::InitRHI(FRHICommandListBase& RHICmdList)
 {
 	// Create RTT and shader resources
 	FTextureRHIRef NewTextureRHI;
-	ImplInitDynamicRHI_RenderTargetResource2D(RenderTargetTextureRHI, NewTextureRHI);
+	ImplInitDynamicRHI_RenderTargetResource2D(RHICmdList, RenderTargetTextureRHI, NewTextureRHI);
 	TextureRHI = (FTextureRHIRef&)NewTextureRHI;
 
 	// Create the sampler state RHI resource.
diff --git a/Engine/Plugins/Runtime/nDisplay/Source/DisplayCluster/Private/Render/Viewport/Resource/DisplayClusterViewportResource.cpp b/Engine/Plugins/Runtime/nDisplay/Source/DisplayCluster/Private/Render/Viewport/Resource/DisplayClusterViewportResource.cpp
--- a/Engine/Plugins/Runtime/nDisplay/Source/DisplayCluster/Private/Render/Viewport/Resource/DisplayClusterViewportResource.cpp
+++ b/Engine/Plugins/Runtime/nDisplay/Source/DisplayCluster/Private/Render/Viewport/Resource/DisplayClusterViewportResource.cpp
@@ -15,7 +15,7 @@ FDisplayClusterViewportResource::~FDisplayClusterViewportResource()
 	check(!EnumHasAnyFlags(ResourceState, EDisplayClusterViewportResourceState::Initialized));
 }
 
-void FDisplayClusterViewportResource::ImplInitDynamicRHI_RenderTargetResource2D(FTextureRHIRef& OutRenderTargetTextureRHI, FTextureRHIRef& OutTextureRHI)
+void FDisplayClusterViewportResource::ImplInitDynamicRHI_RenderTargetResource2D(FRHICommandListBase& RHICmdList, FTextureRHIRef& OutRenderTargetTextureRHI, FTextureRHIRef& OutTextureRHI)
 {
 	// -- we will be manually copying this cross GPU, tell render graph not to
 	ETextureCreateFlags CreateFlags = TexCreate_MultiGPUGraphIgnore;
@@ -43,10 +43,10 @@ void FDisplayClusterViewportResource::ImplInitDynamicRHI_RenderTargetResource2D(
 		.SetInitialState(ERHIAccess::SRVMask)
 		.SetClearValue(FClearValueBinding::Black);
 
-	OutRenderTargetTextureRHI = OutTextureRHI = RHICreateTexture(Desc);
+	OutRenderTargetTextureRHI = OutTextureRHI = RHICmdList.CreateTexture(Desc);
 }
 
-void FDisplayClusterViewportResource::ImplInitDynamicRHI_TextureResource2D(FTextureRHIRef& OutTextureRHI)
+void FDisplayClusterViewportResource::ImplInitDynamicRHI_TextureResource2D(FRHICommandListBase& RHICmdList, FTextureRHIRef& OutTextureRHI)
 {
 	FRHITextureCreateDesc Desc =
 		FRHITextureCreateDesc::Create2D(TEXT("DisplayClusterViewportTextureResource"), ResourceSettings.GetSizeXY().X, ResourceSettings.GetSizeXY().Y, ResourceSettings.GetFormat())
@@ -72,5 +72,5 @@ void FDisplayClusterViewportResource::ImplInitDynamicRHI_TextureResource2D(FText
 		Desc.AddFlags(ETextureCreateFlags::ResolveTargetable);
 	}
 
-	OutTextureRHI = RHICreateTexture(Desc);
+	OutTextureRHI = RHICmdList.CreateTexture(Desc);
 }
diff --git a/Engine/Plugins/Runtime/nDisplay/Source/DisplayCluster/Private/Render/Viewport/Resource/DisplayClusterViewportResource.h b/Engine/Plugins/Runtime/nDisplay/Source/DisplayCluster/Private/Render/Viewport/Resource/DisplayClusterViewportResource.h
--- a/Engine/Plugins/Runtime/nDisplay/Source/DisplayCluster/Private/Render/Viewport/Resource/DisplayClusterViewportResource.h
+++ b/Engine/Plugins/Runtime/nDisplay/Source/DisplayCluster/Private/Render/Viewport/Resource/DisplayClusterViewportResource.h
@@ -120,10 +120,10 @@ public:
 
 protected:
 	/** Helper function: create RTT texture. */
-	void ImplInitDynamicRHI_RenderTargetResource2D(FTextureRHIRef& OutRenderTargetTextureRHI, FTextureRHIRef& OutTextureRHI);
+	void ImplInitDynamicRHI_RenderTargetResource2D(FRHICommandListBase& RHICmdList, FTextureRHIRef& OutRenderTargetTextureRHI, FTextureRHIRef& OutTextureRHI);
 
 	/** Helper function: create 2D texture. */
-	void ImplInitDynamicRHI_TextureResource2D(FTextureRHIRef& OutTextureRHI);
+	void ImplInitDynamicRHI_TextureResource2D(FRHICommandListBase& RHICmdList, FTextureRHIRef& OutTextureRHI);
 
 protected:
 	// Settings used by this resource
diff --git a/Engine/Plugins/VirtualProduction/TextureShare/Source/TextureShare/Private/Resources/TextureShareResource.cpp b/Engine/Plugins/VirtualProduction/TextureShare/Source/TextureShare/Private/Resources/TextureShareResource.cpp
--- a/Engine/Plugins/VirtualProduction/TextureShare/Source/TextureShare/Private/Resources/TextureShareResource.cpp
+++ b/Engine/Plugins/VirtualProduction/TextureShare/Source/TextureShare/Private/Resources/TextureShareResource.cpp
@@ -166,24 +166,24 @@ bool FTextureShareResource::ReleaseTextureShareHandle_RenderThread()
 }
 
 //////////////////////////////////////////////////////////////////////////////////////////////
-void FTextureShareResource::InitRHI(FRHICommandListBase&)
+void FTextureShareResource::InitRHI(FRHICommandListBase& RHICmdList)
 {
 	FTextureRHIRef NewTextureRHI;
 	switch (CoreObject->GetObjectDesc_RenderThread().ProcessDesc.DeviceType)
 	{
 	case ETextureShareDeviceType::D3D12:
-		InitDynamicRHI_D3D12(NewTextureRHI);
+		InitDynamicRHI_D3D12(RHICmdList, NewTextureRHI);
 		break;
 
 	default:
-		InitDynamicRHI_Default(NewTextureRHI);
+		InitDynamicRHI_Default(RHICmdList, NewTextureRHI);
 		break;
 	};
 
 	TextureRHI = (FTextureRHIRef&)NewTextureRHI;
 }
 
-void FTextureShareResource::InitDynamicRHI_Default(FTextureRHIRef& OutTextureRHI)
+void FTextureShareResource::InitDynamicRHI_Default(FRHICommandListBase& RHICmdList, FTextureRHIRef& OutTextureRHI)
 {
 	FRHITextureCreateDesc Desc =
 		FRHITextureCreateDesc::Create2D(TEXT("TextureShareResource"), GetSizeX(), GetSizeY(), ResourceSettings.Format)
@@ -201,5 +201,5 @@ void FTextureShareResource::InitDynamicRHI_Default(FTextureRHIRef& OutTextureRHI
 		Desc.AddFlags(ETextureCreateFlags::External);
 	}
 
-	OutTextureRHI = RHICreateTexture(Desc);
+	OutTextureRHI = RHICmdList.CreateTexture(Desc);
 }
diff --git a/Engine/Plugins/VirtualProduction/TextureShare/Source/TextureShare/Private/Resources/TextureShareResource.h b/Engine/Plugins/VirtualProduction/TextureShare/Source/TextureShare/Private/Resources/TextureShareResource.h
--- a/Engine/Plugins/VirtualProduction/TextureShare/Source/TextureShare/Private/Resources/TextureShareResource.h
+++ b/Engine/Plugins/VirtualProduction/TextureShare/Source/TextureShare/Private/Resources/TextureShareResource.h
@@ -89,8 +89,8 @@ public:
 	}
 
 protected:
-	void InitDynamicRHI_Default(FTextureRHIRef& OutTextureRHI);
-	void InitDynamicRHI_D3D12(FTextureRHIRef& OutTextureRHI);
+	void InitDynamicRHI_Default(FRHICommandListBase& RHICmdList, FTextureRHIRef& OutTextureRHI);
+	void InitDynamicRHI_D3D12(FRHICommandListBase& RHICmdList, FTextureRHIRef& OutTextureRHI);
 
 	bool D3D11RegisterResourceHandle_RenderThread(FRHICommandListImmediate& RHICmdList, const FTextureShareCoreResourceRequest& InResourceRequest);
 	bool D3D11ReleaseTextureShareHandle_RenderThread();
diff --git a/Engine/Plugins/VirtualProduction/TextureShare/Source/TextureShare/Private/Resources/TextureShareResource_D3D12.cpp b/Engine/Plugins/VirtualProduction/TextureShare/Source/TextureShare/Private/Resources/TextureShareResource_D3D12.cpp
--- a/Engine/Plugins/VirtualProduction/TextureShare/Source/TextureShare/Private/Resources/TextureShareResource_D3D12.cpp
+++ b/Engine/Plugins/VirtualProduction/TextureShare/Source/TextureShare/Private/Resources/TextureShareResource_D3D12.cpp
@@ -81,7 +81,7 @@ using namespace UE::TextureShareCore;
 //////////////////////////////////////////////////////////////////////////////////////////////
 // FTextureShareResource
 //////////////////////////////////////////////////////////////////////////////////////////////
-void FTextureShareResource::InitDynamicRHI_D3D12(FTextureRHIRef& OutTextureRHI)
+void FTextureShareResource::InitDynamicRHI_D3D12(FRHICommandListBase& RHICmdList, FTextureRHIRef& OutTextureRHI)
 {
 	switch (ResourceDesc.ResourceType)
 	{
@@ -107,7 +107,7 @@ void FTextureShareResource::InitDynamicRHI_D3D12(FTextureRHIRef& OutTextureRHI)
 		break;
 
 	default:
-		InitDynamicRHI_Default(OutTextureRHI);
+		InitDynamicRHI_Default(RHICmdList, OutTextureRHI);
 		break;
 	}
 }
diff --git a/Engine/Source/Developer/MaterialBaking/Private/MaterialBakingModule.cpp b/Engine/Source/Developer/MaterialBaking/Private/MaterialBakingModule.cpp
--- a/Engine/Source/Developer/MaterialBaking/Private/MaterialBakingModule.cpp
+++ b/Engine/Source/Developer/MaterialBaking/Private/MaterialBakingModule.cpp
@@ -215,7 +215,7 @@ namespace
 				Desc.AddFlags(ETextureCreateFlags::SRGB);
 			}
 
-			return RHICreateTexture(Desc);
+			return RHICmdList.CreateTexture(Desc);
 		}
 
 		void ReleaseStagingBufferForUnmap_AnyThread(FTextureRHIRef& Texture2DRHIRef)
diff --git a/Engine/Source/Editor/VirtualTexturingEditor/Private/RuntimeVirtualTextureBuildStreamingMips.cpp b/Engine/Source/Editor/VirtualTexturingEditor/Private/RuntimeVirtualTextureBuildStreamingMips.cpp
--- a/Engine/Source/Editor/VirtualTexturingEditor/Private/RuntimeVirtualTextureBuildStreamingMips.cpp
+++ b/Engine/Source/Editor/VirtualTexturingEditor/Private/RuntimeVirtualTextureBuildStreamingMips.cpp
@@ -65,7 +65,7 @@ namespace
 
 				FRHITextureCreateDesc StagingTextureDesc = FRHITextureCreateDesc::Create2D(TEXT("FTileRenderResources"), TileSize, TileSize, LayerFormats[Layer]);
 				StagingTextureDesc.SetFlags(ETextureCreateFlags::CPUReadback);
-				StagingTextures[Layer] = RHICreateTexture(StagingTextureDesc);
+				StagingTextures[Layer] = RHICmdList.CreateTexture(StagingTextureDesc);
 			}
 
 			Fence = RHICreateGPUFence(TEXT("Runtime Virtual Texture Build"));
diff --git a/Engine/Source/Runtime/Engine/Private/Components/LightComponent.cpp b/Engine/Source/Runtime/Engine/Private/Components/LightComponent.cpp
--- a/Engine/Source/Runtime/Engine/Private/Components/LightComponent.cpp
+++ b/Engine/Source/Runtime/Engine/Private/Components/LightComponent.cpp
@@ -49,7 +49,7 @@
 static const TCHAR* GLightSpriteAssetName = TEXT("/Engine/EditorResources/LightIcons/S_LightError.S_LightError");
 #endif
 
-void FStaticShadowDepthMap::InitRHI(FRHICommandListBase&)
+void FStaticShadowDepthMap::InitRHI(FRHICommandListBase& RHICmdList)
 {
 	if (FApp::CanEverRender() && Data && Data->ShadowMapSizeX > 0 && Data->ShadowMapSizeY > 0 && GMaxRHIFeatureLevel >= ERHIFeatureLevel::SM5)
 	{
@@ -58,10 +58,10 @@ void FStaticShadowDepthMap::InitRHI(FRHICommandListBase&)
 			FRHITextureCreateDesc::Create2D(TEXT("FStaticShadowDepthMap"), Data->ShadowMapSizeX, Data->ShadowMapSizeY, PF_R16F)
 			.SetClassName(ClassName);
 
-		TextureRHI = RHICreateTexture(Desc);
+		TextureRHI = RHICmdList.CreateTexture(Desc);
 
 		uint32 DestStride = 0;
-		uint8* TextureData = (uint8*)RHILockTexture2D(TextureRHI, 0, RLM_WriteOnly, DestStride, false);
+		uint8* TextureData = (uint8*)RHICmdList.GetAsImmediate().LockTexture2D(TextureRHI, 0, RLM_WriteOnly, DestStride, false);
 		uint32 RowSize = Data->ShadowMapSizeX * GPixelFormats[PF_R16F].BlockBytes;
 
 		for (int32 Y = 0; Y < Data->ShadowMapSizeY; Y++)
@@ -69,7 +69,7 @@ void FStaticShadowDepthMap::InitRHI(FRHICommandListBase&)
 			FMemory::Memcpy(TextureData + DestStride * Y, ((uint8*)Data->DepthSamples.GetData()) + RowSize * Y, RowSize);
 		}
 
-		RHIUnlockTexture2D(TextureRHI, 0, false);
+		RHICmdList.GetAsImmediate().UnlockTexture2D(TextureRHI, 0, false);
 	}
 }
 
diff --git a/Engine/Source/Runtime/Engine/Private/Components/ReflectionCaptureComponent.cpp b/Engine/Source/Runtime/Engine/Private/Components/ReflectionCaptureComponent.cpp
--- a/Engine/Source/Runtime/Engine/Private/Components/ReflectionCaptureComponent.cpp
+++ b/Engine/Source/Runtime/Engine/Private/Components/ReflectionCaptureComponent.cpp
@@ -595,7 +595,7 @@ public:
 			.SetNumMips(NumMips)
 			.SetFlags(ETextureCreateFlags::ShaderResource);
 
-		TextureCubeRHI = RHICreateTexture(Desc);
+		TextureCubeRHI = RHICmdList.CreateTexture(Desc);
 		TextureRHI = TextureCubeRHI;
 
 		if (SourceData.Num())
@@ -613,7 +613,7 @@ public:
 				for (int32 CubeFace = 0; CubeFace < CubeFace_MAX; CubeFace++)
 				{
 					uint32 DestPitch = 0;
-					uint8* DestBuffer = (uint8*)RHILockTextureCubeFace(TextureCubeRHI, CubeFace, 0, MipIndex, RLM_WriteOnly, DestPitch, false);
+					uint8* DestBuffer = (uint8*)RHICmdList.GetAsImmediate().LockTextureCubeFace(TextureCubeRHI, CubeFace, 0, MipIndex, RLM_WriteOnly, DestPitch, false);
 					
 					const int32 SourceIndex = MipBaseIndex + CubeFace * CubeFaceBytes;
 					const uint8* SourceBuffer = &SourceData[SourceIndex];
@@ -636,7 +636,7 @@ public:
 						check((PTRINT(Src) - PTRINT(SourceBuffer)) == PTRINT(CubeFaceBytes));
 					}
 
-					RHIUnlockTextureCubeFace(TextureCubeRHI, CubeFace, 0, MipIndex, false);
+					RHICmdList.GetAsImmediate().UnlockTextureCubeFace(TextureCubeRHI, CubeFace, 0, MipIndex, false);
 				}
 
 				MipBaseIndex += CubeFaceBytes * CubeFace_MAX;
diff --git a/Engine/Source/Runtime/Engine/Private/Components/SkyLightComponent.cpp b/Engine/Source/Runtime/Engine/Private/Components/SkyLightComponent.cpp
--- a/Engine/Source/Runtime/Engine/Private/Components/SkyLightComponent.cpp
+++ b/Engine/Source/Runtime/Engine/Private/Components/SkyLightComponent.cpp
@@ -131,7 +131,7 @@ FAutoConsoleVariableRef CVarSkylightCubemapMaxResolution(
 
 constexpr EPixelFormat SKYLIGHT_CUBEMAP_FORMAT = PF_FloatRGBA;
 
-void FSkyTextureCubeResource::InitRHI(FRHICommandListBase&)
+void FSkyTextureCubeResource::InitRHI(FRHICommandListBase& RHICmdList)
 {
 	if (GetFeatureLevel() >= ERHIFeatureLevel::SM5 || GSupportsRenderTargetFormat_PF_FloatRGBA)
 	{
@@ -141,7 +141,7 @@ void FSkyTextureCubeResource::InitRHI(FRHICommandListBase&)
 			FRHITextureCreateDesc::CreateCube(TEXT("SkyTextureCube"), Size, Format)
 			.SetNumMips(NumMips);
 
-		TextureCubeRHI = RHICreateTexture(Desc);
+		TextureCubeRHI = RHICmdList.CreateTexture(Desc);
 		TextureRHI = TextureCubeRHI;
 
 		// Create the sampler state RHI resource.
diff --git a/Engine/Source/Runtime/Engine/Private/Particles/ParticleCurveTexture.cpp b/Engine/Source/Runtime/Engine/Private/Particles/ParticleCurveTexture.cpp
--- a/Engine/Source/Runtime/Engine/Private/Particles/ParticleCurveTexture.cpp
+++ b/Engine/Source/Runtime/Engine/Private/Particles/ParticleCurveTexture.cpp
@@ -508,7 +508,7 @@ FParticleCurveTexture::FParticleCurveTexture()
 /**
 * Initialize RHI resources for the curve texture.
 */
-void FParticleCurveTexture::InitRHI(FRHICommandListBase&)
+void FParticleCurveTexture::InitRHI(FRHICommandListBase& RHICmdList)
 {
 	const static FLazyName ClassName(TEXT("FParticleCurveTexture"));
 
@@ -522,7 +522,7 @@ void FParticleCurveTexture::InitRHI(FRHICommandListBase&)
 		.SetInitialState(ERHIAccess::SRVMask)
 		.SetClassName(ClassName);
 
-	CurveTextureRHI = RHICreateTexture(Desc);
+	CurveTextureRHI = RHICmdList.CreateTexture(Desc);
 }
 
 /**
diff --git a/Engine/Source/Runtime/Engine/Private/PrecomputedVolumetricLightmap.cpp b/Engine/Source/Runtime/Engine/Private/PrecomputedVolumetricLightmap.cpp
--- a/Engine/Source/Runtime/Engine/Private/PrecomputedVolumetricLightmap.cpp
+++ b/Engine/Source/Runtime/Engine/Private/PrecomputedVolumetricLightmap.cpp
@@ -22,7 +22,7 @@ DECLARE_MEMORY_STAT(TEXT("Volumetric Lightmap"),STAT_VolumetricLightmapBuildData
 
 const static FLazyName VolumetricLightmapDataLayerName(TEXT("VolumetricLightmapDataLayer"));
 
-void FVolumetricLightmapDataLayer::CreateTexture(FIntVector Dimensions)
+void FVolumetricLightmapDataLayer::CreateTexture(FRHICommandListBase& RHICmdList, FIntVector Dimensions)
 {
 	const FRHITextureCreateDesc Desc =
 		FRHITextureCreateDesc::Create3D(TEXT("VolumetricLightmap"), Dimensions, Format)
@@ -30,22 +30,27 @@ void FVolumetricLightmapDataLayer::CreateTexture(FIntVector Dimensions)
 		.SetInitActionBulkData(this)
 		.SetClassName(VolumetricLightmapDataLayerName);
 
-	Texture = RHICreateTexture(Desc);
+	Texture = RHICmdList.CreateTexture(Desc);
 }
 
-void FVolumetricLightmapDataLayer::CreateTargetTexture(FIntVector Dimensions)
+void FVolumetricLightmapDataLayer::CreateTargetTexture(FRHICommandListBase& RHICmdList, FIntVector Dimensions)
 {
 	const FRHITextureCreateDesc Desc =
 		FRHITextureCreateDesc::Create3D(TEXT("VolumetricLightmap"), Dimensions, Format)
 		.SetFlags(ETextureCreateFlags::ShaderResource | ETextureCreateFlags::UAV)
 		.SetClassName(VolumetricLightmapDataLayerName);
 
-	Texture = RHICreateTexture(Desc);
+	Texture = RHICmdList.CreateTexture(Desc);
 }
 
-void FVolumetricLightmapDataLayer::CreateUAV()
+void FVolumetricLightmapDataLayer::CreateTexture(FIntVector Dimensions)
 {
-	CreateUAV(FRHICommandListImmediate::Get());
+	CreateTexture(FRHICommandListImmediate::Get(), Dimensions);
+}
+
+void FVolumetricLightmapDataLayer::CreateTargetTexture(FIntVector Dimensions)
+{
+	CreateTargetTexture(FRHICommandListImmediate::Get(), Dimensions);
 }
 
 void FVolumetricLightmapDataLayer::CreateUAV(FRHICommandListBase& RHICmdList)
@@ -227,30 +232,30 @@ void FPrecomputedVolumetricLightmapData::FinalizeImport()
 	INC_DWORD_STAT_BY(STAT_VolumetricLightmapBuildData, VolumeBytes);
 }
 
-ENGINE_API void FPrecomputedVolumetricLightmapData::InitRHI(FRHICommandListBase&)
+ENGINE_API void FPrecomputedVolumetricLightmapData::InitRHI(FRHICommandListBase& RHICmdList)
 {
 	if (GetFeatureLevel() >= ERHIFeatureLevel::SM5)
 	{
 		if (IndirectionTextureDimensions.GetMax() > 0)
 		{
-			IndirectionTexture.CreateTexture(IndirectionTextureDimensions);
+			IndirectionTexture.CreateTexture(RHICmdList, IndirectionTextureDimensions);
 		}
 
 		if (BrickDataDimensions.GetMax() > 0)
 		{
-			BrickData.AmbientVector.CreateTexture(BrickDataDimensions);
+			BrickData.AmbientVector.CreateTexture(RHICmdList, BrickDataDimensions);
 
 			for (int32 i = 0; i < UE_ARRAY_COUNT(BrickData.SHCoefficients); i++)
 			{
-				BrickData.SHCoefficients[i].CreateTexture(BrickDataDimensions);
+				BrickData.SHCoefficients[i].CreateTexture(RHICmdList, BrickDataDimensions);
 			}
 
 			if (BrickData.SkyBentNormal.Data.Num() > 0)
 			{
-				BrickData.SkyBentNormal.CreateTexture(BrickDataDimensions);
+				BrickData.SkyBentNormal.CreateTexture(RHICmdList, BrickDataDimensions);
 			}
 
-			BrickData.DirectionalLightShadowing.CreateTexture(BrickDataDimensions);
+			BrickData.DirectionalLightShadowing.CreateTexture(RHICmdList, BrickDataDimensions);
 		}
 
 		GVolumetricLightmapBrickAtlas.Insert(INT_MAX, this);
@@ -260,12 +265,10 @@ ENGINE_API void FPrecomputedVolumetricLightmapData::InitRHI(FRHICommandListBase&
 	}
 }
 
-ENGINE_API void FPrecomputedVolumetricLightmapData::InitRHIForSubLevelResources()
+void FPrecomputedVolumetricLightmapData::InitRHIForSubLevelResources(FRHICommandListBase& RHICmdList)
 {
 	if (SubLevelBrickPositions.Num() > 0)
 	{
-		FRHICommandListBase& RHICmdList = FRHICommandListImmediate::Get();
-
 		SubLevelBrickPositions.SetAllowCPUAccess(true);
 		IndirectionTextureOriginalValues.SetAllowCPUAccess(true);
 
@@ -301,6 +304,11 @@ ENGINE_API void FPrecomputedVolumetricLightmapData::InitRHIForSubLevelResources(
 	}
 }
 
+void FPrecomputedVolumetricLightmapData::InitRHIForSubLevelResources()
+{
+	InitRHIForSubLevelResources(FRHICommandListImmediate::Get());
+}
+
 ENGINE_API void FPrecomputedVolumetricLightmapData::ReleaseRHI()
 {
 	if (GetFeatureLevel() >= ERHIFeatureLevel::SM5)
@@ -346,7 +354,7 @@ ENGINE_API void FPrecomputedVolumetricLightmapData::HandleDataMovementInAtlas(in
 				TShaderMapRef<FMoveWholeIndirectionTextureCS> ComputeShader(GlobalShaderMap);
 
 				FVolumetricLightmapDataLayer NewIndirectionTexture = SceneData->IndirectionTexture;
-				NewIndirectionTexture.CreateTargetTexture(IndirectionTextureDimensions);
+				NewIndirectionTexture.CreateTargetTexture(RHICmdList, IndirectionTextureDimensions);
 				NewIndirectionTexture.CreateUAV(RHICmdList);
 
 				FMoveWholeIndirectionTextureCS::FParameters Parameters;
@@ -372,7 +380,7 @@ ENGINE_API void FPrecomputedVolumetricLightmapData::HandleDataMovementInAtlas(in
 	}
 	else
 	{
-		InitRHIForSubLevelResources();
+		InitRHIForSubLevelResources(RHICmdList);
 
 		for (FPrecomputedVolumetricLightmapData* SceneData : SceneDataAdded)
 		{
@@ -458,7 +466,7 @@ ENGINE_API void FPrecomputedVolumetricLightmapData::AddToSceneData(FPrecomputedV
 			TShaderMapRef<FMoveWholeIndirectionTextureCS> ComputeShader(GlobalShaderMap);
 
 			FVolumetricLightmapDataLayer NewIndirectionTexture = SceneData->IndirectionTexture;
-			NewIndirectionTexture.CreateTargetTexture(IndirectionTextureDimensions);
+			NewIndirectionTexture.CreateTargetTexture(RHICmdList, IndirectionTextureDimensions);
 			NewIndirectionTexture.CreateUAV(RHICmdList);
 
 			RHICmdList.Transition(FRHITransitionInfo(NewIndirectionTexture.UAV, ERHIAccess::Unknown, ERHIAccess::UAVCompute));
@@ -509,7 +517,7 @@ ENGINE_API void FPrecomputedVolumetricLightmapData::AddToSceneData(FPrecomputedV
 			{
 				SceneDataAdded.Add(SceneData);
 
-				InitRHIForSubLevelResources();
+				InitRHIForSubLevelResources(RHICmdList);
 
 				FGlobalShaderMap* GlobalShaderMap = GetGlobalShaderMap(GetFeatureLevel());
 
@@ -601,7 +609,7 @@ ENGINE_API void FPrecomputedVolumetricLightmapData::RemoveFromSceneData(FPrecomp
 			{
 				SCOPED_DRAW_EVENT(RHICmdList, RemoveSubLevelBricksCS);
 
-				InitRHIForSubLevelResources();
+				InitRHIForSubLevelResources(RHICmdList);
 
 				FGlobalShaderMap* GlobalShaderMap = GetGlobalShaderMap(GetFeatureLevel());
 
@@ -994,7 +1002,7 @@ void FVolumetricLightmapBrickAtlas::Insert(int32 Index, FPrecomputedVolumetricLi
 		// TODO: release SkyBentNormal if no brick data in the atlas uses it
 		if (!TextureSet.SkyBentNormal.Texture.IsValid() && Data->BrickData.SkyBentNormal.Texture.IsValid())
 		{
-			TextureSet.SkyBentNormal.CreateTargetTexture(TextureSet.BrickDataDimensions);
+			TextureSet.SkyBentNormal.CreateTargetTexture(RHICmdList, TextureSet.BrickDataDimensions);
 			TextureSet.SkyBentNormal.CreateUAV(RHICmdList);
 			bReadAfterCreate = true;
 		}
diff --git a/Engine/Source/Runtime/Engine/Private/PrimitiveUniformShaderParameters.cpp b/Engine/Source/Runtime/Engine/Private/PrimitiveUniformShaderParameters.cpp
--- a/Engine/Source/Runtime/Engine/Private/PrimitiveUniformShaderParameters.cpp
+++ b/Engine/Source/Runtime/Engine/Private/PrimitiveUniformShaderParameters.cpp
@@ -55,7 +55,7 @@ void FSinglePrimitiveStructured::InitRHI(FRHICommandListBase& RHICmdList)
 			.SetFlags(ETextureCreateFlags::ShaderResource | ETextureCreateFlags::UAV)
 			.SetClassName(ClassName);
 
-		PrimitiveSceneDataTextureRHI = RHICreateTexture(Desc);
+		PrimitiveSceneDataTextureRHI = RHICmdList.CreateTexture(Desc);
 		PrimitiveSceneDataTextureSRV = RHICmdList.CreateShaderResourceView(PrimitiveSceneDataTextureRHI, FRHIViewDesc::CreateTextureSRV().SetDimensionFromTexture(PrimitiveSceneDataTextureRHI));
 	}
 
diff --git a/Engine/Source/Runtime/Engine/Private/Slate/SceneViewport.cpp b/Engine/Source/Runtime/Engine/Private/Slate/SceneViewport.cpp
--- a/Engine/Source/Runtime/Engine/Private/Slate/SceneViewport.cpp
+++ b/Engine/Source/Runtime/Engine/Private/Slate/SceneViewport.cpp
@@ -2146,7 +2146,7 @@ void FSceneViewport::InitRHI(FRHICommandListBase& RHICmdList)
 					.SetFlags(ETextureCreateFlags::RenderTargetable | ETextureCreateFlags::ShaderResource)
 					.SetInitialState(ERHIAccess::SRVMask);
 
-				RTRHI = SRVRHI = RHICreateTexture(Desc);
+				RTRHI = SRVRHI = RHICmdList.CreateTexture(Desc);
 			}
 			BufferedRenderTargetsRHI[i] = RTRHI;
 			BufferedShaderResourceTexturesRHI[i] = SRVRHI;
diff --git a/Engine/Source/Runtime/Engine/Private/Slate/SlateTextures.cpp b/Engine/Source/Runtime/Engine/Private/Slate/SlateTextures.cpp
--- a/Engine/Source/Runtime/Engine/Private/Slate/SlateTextures.cpp
+++ b/Engine/Source/Runtime/Engine/Private/Slate/SlateTextures.cpp
@@ -338,7 +338,7 @@ void FSlateTextureRenderTarget2DResource::ClampSize(int32 MaxSizeX,int32 MaxSize
 	}	
 }
 
-void FSlateTextureRenderTarget2DResource::InitRHI(FRHICommandListBase&)
+void FSlateTextureRenderTarget2DResource::InitRHI(FRHICommandListBase& RHICmdList)
 {
 	SCOPED_LOADTIMER(FSlateTextureRenderTarget2DResource_InitDynamicRHI);
 
@@ -358,7 +358,7 @@ void FSlateTextureRenderTarget2DResource::InitRHI(FRHICommandListBase&)
 			.SetInitialState(ERHIAccess::SRVMask)
 			.SetClassName(ClassName);
 
-		RenderTargetTextureRHI = TextureRHI = RHICreateTexture(Desc);
+		RenderTargetTextureRHI = TextureRHI = RHICmdList.CreateTexture(Desc);
 	}
 
 	// Create the sampler state RHI resource.
diff --git a/Engine/Source/Runtime/Engine/Private/Streaming/Texture2DMipAllocator_Virtual.cpp b/Engine/Source/Runtime/Engine/Private/Streaming/Texture2DMipAllocator_Virtual.cpp
--- a/Engine/Source/Runtime/Engine/Private/Streaming/Texture2DMipAllocator_Virtual.cpp
+++ b/Engine/Source/Runtime/Engine/Private/Streaming/Texture2DMipAllocator_Virtual.cpp
@@ -143,6 +143,8 @@ void FTexture2DMipAllocator_Virtual::DoConvertToVirtualWithNewMips(const FTextur
 		// If the texture is not virtual, then make it virtual immediately.
 		if (!Context.Resource->GetTexture2DRHI() || (Context.Resource->GetTexture2DRHI()->GetFlags() & ETextureCreateFlags::Virtual) != ETextureCreateFlags::Virtual)
 		{
+			FRHICommandListImmediate& RHICmdList = FRHICommandListImmediate::Get();
+
 			const FTexture2DMipMap& MipMap0 = *Context.MipsView[0];
 
 			FTexture2DResource* Resource2D = Context.Resource->GetTexture2DResource();
@@ -157,13 +159,13 @@ void FTexture2DMipAllocator_Virtual::DoConvertToVirtualWithNewMips(const FTextur
 				CreateDesc.SetInitActionBulkData(Resource2D->ResourceMem);
 			}
 
-			IntermediateTextureRHI = RHICreateTexture(CreateDesc);
+			IntermediateTextureRHI = RHICmdList.CreateTexture(CreateDesc);
 
-			RHIVirtualTextureSetFirstMipInMemory(IntermediateTextureRHI, CurrentFirstLODIdx);
-			RHIVirtualTextureSetFirstMipVisible(IntermediateTextureRHI, CurrentFirstLODIdx);
+			RHICmdList.VirtualTextureSetFirstMipInMemory(IntermediateTextureRHI, CurrentFirstLODIdx);
+			RHICmdList.VirtualTextureSetFirstMipVisible(IntermediateTextureRHI, CurrentFirstLODIdx);
 
 			UE::RHI::CopySharedMips_AssumeSRVMaskState(
-				FRHICommandListExecutor::GetImmediateCommandList(),
+				RHICmdList,
 				Context.Resource->GetTexture2DRHI(),
 				IntermediateTextureRHI);
 		}
diff --git a/Engine/Source/Runtime/Engine/Private/UnrealClient.cpp b/Engine/Source/Runtime/Engine/Private/UnrealClient.cpp
--- a/Engine/Source/Runtime/Engine/Private/UnrealClient.cpp
+++ b/Engine/Source/Runtime/Engine/Private/UnrealClient.cpp
@@ -2682,7 +2682,7 @@ FDummyViewport::~FDummyViewport()
 	}
 }
 
-void FDummyViewport::InitRHI(FRHICommandListBase&)
+void FDummyViewport::InitRHI(FRHICommandListBase& RHICmdList)
 {
 	EPixelFormat DummyViewportFormat = bSceneHDREnabled ? GRHIHDRDisplayOutputFormat : PF_A2B10G10R10;
 	const FRHITextureCreateDesc Desc =
@@ -2692,7 +2692,7 @@ void FDummyViewport::InitRHI(FRHICommandListBase&)
 		.SetFlags(ETextureCreateFlags::RenderTargetable | ETextureCreateFlags::ShaderResource)
 		.SetInitialState(ERHIAccess::SRVMask);
 
-	RenderTargetTextureRHI = RHICreateTexture(Desc);
+	RenderTargetTextureRHI = RHICmdList.CreateTexture(Desc);
 }
 
 EDisplayColorGamut FDummyViewport::GetDisplayColorGamut() const
diff --git a/Engine/Source/Runtime/Engine/Private/VT/VirtualTextureUploadCache.cpp b/Engine/Source/Runtime/Engine/Private/VT/VirtualTextureUploadCache.cpp
--- a/Engine/Source/Runtime/Engine/Private/VT/VirtualTextureUploadCache.cpp
+++ b/Engine/Source/Runtime/Engine/Private/VT/VirtualTextureUploadCache.cpp
@@ -343,7 +343,7 @@ void FVirtualTextureUploadCache::FinalizeWithLegacyCopyTexture(FRDGBuilder& Grap
 				Desc.AddFlags(ETextureCreateFlags::CPUWritable);
 			}
 
-			StagingTexture.RHITexture = RHICreateTexture(Desc);
+			StagingTexture.RHITexture = GraphBuilder.RHICmdList.CreateTexture(Desc);
 
 			StagingTexture.WidthInTiles = WidthInTiles;
 			StagingTexture.BatchCapacity = WidthInTiles * HeightInTiles;
diff --git a/Engine/Source/Runtime/Engine/Private/VectorField.cpp b/Engine/Source/Runtime/Engine/Private/VectorField.cpp
--- a/Engine/Source/Runtime/Engine/Private/VectorField.cpp
+++ b/Engine/Source/Runtime/Engine/Private/VectorField.cpp
@@ -176,7 +176,7 @@ public:
 				.SetFlags(ETextureCreateFlags::ShaderResource)
 				.SetInitActionBulkData(&BulkDataInterface);
 
-			VolumeTextureRHI = RHICreateTexture(Desc);
+			VolumeTextureRHI = RHICmdList.CreateTexture(Desc);
 
 			FMemory::Free(VolumeData);
 			VolumeData = NULL;
@@ -1039,7 +1039,7 @@ public:
 				FRHITextureCreateDesc::Create3D(TEXT("FVectorFieldAnimatedResource"), SizeX, SizeY, SizeZ, PF_FloatRGBA)
 				.SetFlags(TexCreateFlags);
 
-			VolumeTextureRHI = RHICreateTexture(Desc);
+			VolumeTextureRHI = RHICmdList.CreateTexture(Desc);
 
 			if (GetFeatureLevel() >= ERHIFeatureLevel::SM5)
 			{
diff --git a/Engine/Source/Runtime/Engine/Public/PlanarReflectionSceneProxy.h b/Engine/Source/Runtime/Engine/Public/PlanarReflectionSceneProxy.h
--- a/Engine/Source/Runtime/Engine/Public/PlanarReflectionSceneProxy.h
+++ b/Engine/Source/Runtime/Engine/Public/PlanarReflectionSceneProxy.h
@@ -26,7 +26,7 @@ public:
 		Size(InSize)
 	{}
 
-	virtual void InitRHI(FRHICommandListBase&)
+	virtual void InitRHI(FRHICommandListBase& RHICmdList)
 	{
 		// Create the sampler state RHI resource.
 		FSamplerStateInitializerRHI SamplerStateInitializer
@@ -46,7 +46,7 @@ public:
 			.SetFlags(ETextureCreateFlags::RenderTargetable | ETextureCreateFlags::ShaderResource)
 			.SetInitialState(ERHIAccess::SRVMask);
 
-		RenderTargetTextureRHI = TextureRHI = RHICreateTexture(Desc);
+		RenderTargetTextureRHI = TextureRHI = RHICmdList.CreateTexture(Desc);
 	}
 
 	virtual FIntPoint GetSizeXY() const { return Size; }
diff --git a/Engine/Source/Runtime/Engine/Public/PrecomputedVolumetricLightmap.h b/Engine/Source/Runtime/Engine/Public/PrecomputedVolumetricLightmap.h
--- a/Engine/Source/Runtime/Engine/Public/PrecomputedVolumetricLightmap.h
+++ b/Engine/Source/Runtime/Engine/Public/PrecomputedVolumetricLightmap.h
@@ -58,12 +58,14 @@ public:
 		DataSize = NewSize;
 	}
 
-	ENGINE_API void CreateTexture(FIntVector Dimensions);
-	ENGINE_API void CreateTargetTexture(FIntVector Dimensions);
+	ENGINE_API void CreateTexture(FRHICommandListBase& RHICmdList, FIntVector Dimensions);
+	ENGINE_API void CreateTargetTexture(FRHICommandListBase& RHICmdList, FIntVector Dimensions);
 	ENGINE_API void CreateUAV(FRHICommandListBase& RHICmdList);
 	
-	UE_DEPRECATED(5.3, "CreateUAV now requires a command list.")
-	ENGINE_API void CreateUAV();
+	UE_DEPRECATED(5.7, "CreateTexture now requires a command list.")
+	ENGINE_API void CreateTexture(FIntVector Dimensions);
+	UE_DEPRECATED(5.7, "CreateTargetTexture now requires a command list.")
+	ENGINE_API void CreateTargetTexture(FIntVector Dimensions);
 
 	TArray<uint8> Data;
 	// Stored redundantly for stats after Data has been discarded
@@ -158,6 +160,8 @@ public:
 	ENGINE_API virtual void InitRHI(FRHICommandListBase& RHICmdList) override;
 	ENGINE_API virtual void ReleaseRHI() override;
 
+	ENGINE_API void InitRHIForSubLevelResources(FRHICommandListBase& RHICmdList);
+	UE_DEPRECATED(5.7, "InitRHIForSubLevelResources now requires a command list.")
 	ENGINE_API void InitRHIForSubLevelResources();
 	ENGINE_API void ReleaseRHIForSubLevelResources();
 
@@ -514,22 +518,22 @@ struct FVolumetricLightmapBrickTextureSet : public FVolumetricLightmapBasicBrick
 			SHCoefficients[i].Format = BrickData.SHCoefficients[i].Format;
 		}
 
-		AmbientVector.CreateTargetTexture(BrickDataDimensions);
+		AmbientVector.CreateTargetTexture(RHICmdList, BrickDataDimensions);
 		AmbientVector.CreateUAV(RHICmdList);
 
 		for (int32 i = 0; i < UE_ARRAY_COUNT(SHCoefficients); i++)
 		{
-			SHCoefficients[i].CreateTargetTexture(BrickDataDimensions);
+			SHCoefficients[i].CreateTargetTexture(RHICmdList, BrickDataDimensions);
 			SHCoefficients[i].CreateUAV(RHICmdList);
 		}
 
 		if (BrickData.SkyBentNormal.Texture.IsValid())
 		{
-			SkyBentNormal.CreateTargetTexture(BrickDataDimensions);
+			SkyBentNormal.CreateTargetTexture(RHICmdList, BrickDataDimensions);
 			SkyBentNormal.CreateUAV(RHICmdList);
 		}
 
-		DirectionalLightShadowing.CreateTargetTexture(BrickDataDimensions);
+		DirectionalLightShadowing.CreateTargetTexture(RHICmdList, BrickDataDimensions);
 		DirectionalLightShadowing.CreateUAV(RHICmdList);
 	}
 
diff --git a/Engine/Source/Runtime/Landscape/Private/LandscapeEditReadback.cpp b/Engine/Source/Runtime/Landscape/Private/LandscapeEditReadback.cpp
--- a/Engine/Source/Runtime/Landscape/Private/LandscapeEditReadback.cpp
+++ b/Engine/Source/Runtime/Landscape/Private/LandscapeEditReadback.cpp
@@ -61,7 +61,7 @@ void InitTask_GameThread(FLandscapeEditReadbackTaskImpl& Task, UTexture2D const*
 }
 
 /** Initialize the read back task resources. */
-bool InitTask_RenderThread(FLandscapeEditReadbackTaskImpl& Task)
+bool InitTask_RenderThread(FRHICommandListImmediate& RHICmdList, FLandscapeEditReadbackTaskImpl& Task)
 {
 	using namespace UE::Landscape::Private;
 
@@ -78,7 +78,7 @@ bool InitTask_RenderThread(FLandscapeEditReadbackTaskImpl& Task)
 				FRHITextureCreateDesc::Create2D(TEXT("LandscapeEditReadbackTask"), MipWidth, MipHeight, Task.Format)
 				.SetFlags(ETextureCreateFlags::CPUReadback);
 
-			Task.StagingTextures[MipIndex] = RHICreateTexture(Desc);
+			Task.StagingTextures[MipIndex] = RHICmdList.CreateTexture(Desc);
 		}
 
 		Task.StagingTexturesAllocatedSize += CalcTextureSize(Task.Size.X, Task.Size.Y, Task.Format, Task.NumMips);
@@ -358,7 +358,7 @@ void FLandscapeEditLayerReadback::Enqueue(UTexture2D const* InSourceTexture, FRe
 
 		ENQUEUE_RENDER_COMMAND(FLandscapeEditLayerReadback_Queue)([TaskHandle](FRHICommandListImmediate& RHICmdList)
 		{
-			InitTask_RenderThread(GReadbackTaskPool.Pool[TaskHandle]);
+			InitTask_RenderThread(RHICmdList, GReadbackTaskPool.Pool[TaskHandle]);
 			KickTask_RenderThread(RHICmdList, GReadbackTaskPool.Pool[TaskHandle]);
 		});
 	}
diff --git a/Engine/Source/Runtime/Landscape/Private/LandscapeEditResources.cpp b/Engine/Source/Runtime/Landscape/Private/LandscapeEditResources.cpp
--- a/Engine/Source/Runtime/Landscape/Private/LandscapeEditResources.cpp
+++ b/Engine/Source/Runtime/Landscape/Private/LandscapeEditResources.cpp
@@ -40,7 +40,7 @@ void FLandscapeTexture2DResource::InitRHI(FRHICommandListBase& RHICmdList)
 		Desc.AddFlags(ETextureCreateFlags::ShaderResource);
 	}
 
-	TextureRHI = RHICreateTexture(Desc);
+	TextureRHI = RHICmdList.CreateTexture(Desc);
 
 	if (bCreateUAVs)
 	{
@@ -126,7 +126,7 @@ void FLandscapeTexture2DArrayResource::InitRHI(FRHICommandListBase& RHICmdList)
         Desc.AddFlags(ETextureCreateFlags::ShaderResource);
     }
     
-    TextureRHI = RHICreateTexture(Desc);
+    TextureRHI = RHICmdList.CreateTexture(Desc);
 
 	if (bCreateUAVs)
 	{
diff --git a/Engine/Source/Runtime/MovieSceneCapture/Private/FrameGrabber.cpp b/Engine/Source/Runtime/MovieSceneCapture/Private/FrameGrabber.cpp
--- a/Engine/Source/Runtime/MovieSceneCapture/Private/FrameGrabber.cpp
+++ b/Engine/Source/Runtime/MovieSceneCapture/Private/FrameGrabber.cpp
@@ -62,7 +62,7 @@ void FViewportSurfaceReader::Resize(uint32 Width, uint32 Height)
 				FRHITextureCreateDesc::Create2D(TEXT("FViewportSurfaceReader_ReadbackTexture"), Width, Height, This->PixelFormat)
 				.SetFlags(ETextureCreateFlags::CPUReadback);
 
-			This->ReadbackTexture = RHICreateTexture(Desc);
+			This->ReadbackTexture = RHICmdList.CreateTexture(Desc);
 		});
 }
 
diff --git a/Engine/Source/Runtime/RenderCore/Private/DumpGPU.cpp b/Engine/Source/Runtime/RenderCore/Private/DumpGPU.cpp
--- a/Engine/Source/Runtime/RenderCore/Private/DumpGPU.cpp
+++ b/Engine/Source/Runtime/RenderCore/Private/DumpGPU.cpp
@@ -1017,10 +1017,10 @@ public:
 	TArray<TUniquePtr<FStagingTexturePoolEntry>> StagingTexturePool;
 	TArray<TUniquePtr<FStagingBufferPoolEntry>> StagingBufferPool;
 
-	void CreateStagingTexture(const FRHITextureCreateDesc& Desc, FTextureRHIRef* Texture)
+	void CreateStagingTexture(FRHICommandListBase& RHICmdList, const FRHITextureCreateDesc& Desc, FTextureRHIRef* Texture)
 	{
 		check(IsInRenderingThread());
-		(*Texture) = RHICreateTexture(Desc);
+		(*Texture) = RHICmdList.CreateTexture(Desc);
 		(*Texture)->DisableLifetimeExtension();
 	}
 
@@ -1084,7 +1084,7 @@ public:
 			}
 		}
 
-		CreateStagingTexture(Desc, /* out */ Texture);
+		CreateStagingTexture(RHICmdList, Desc, /* out */ Texture);
 
 		FStagingTexturePoolEntry* NewEntry = new FStagingTexturePoolEntry;
 		NewEntry->Status = Access == ERHIAccess::CopyDest ? FStagingResourceStatus::Rendering : FStagingResourceStatus::RenderingIntermediaryOnly;
@@ -1178,7 +1178,7 @@ public:
 					}
 					else
 					{
-						CreateStagingTexture(Desc, /* out */ &StagingSrcTextureRef);
+						CreateStagingTexture(RHICmdList, Desc, /* out */ &StagingSrcTextureRef);
 
 						RHICmdList.Transition(FRHITransitionInfo(StagingSrcTextureRef, ERHIAccess::Unknown, ERHIAccess::UAVCompute));
 						StagingOutput = CreateUAVNoLifetimeExtension(RHICmdList, StagingSrcTextureRef, /* MipLevel = */ 0);
@@ -1218,7 +1218,7 @@ public:
 				}
 				else
 				{
-					CreateStagingTexture(Desc, /* out */ &StagingTexture);
+					CreateStagingTexture(RHICmdList, Desc, /* out */ &StagingTexture);
 					RHICmdList.Transition(FRHITransitionInfo(StagingTexture, ERHIAccess::Unknown, ERHIAccess::CopyDest));
 				}
 
diff --git a/Engine/Source/Runtime/RenderCore/Private/LongGPUTask.cpp b/Engine/Source/Runtime/RenderCore/Private/LongGPUTask.cpp
--- a/Engine/Source/Runtime/RenderCore/Private/LongGPUTask.cpp
+++ b/Engine/Source/Runtime/RenderCore/Private/LongGPUTask.cpp
@@ -23,7 +23,7 @@ void IssueScalableLongGPUTask(FRHICommandListImmediate& RHICmdList, int32 NumIte
 		FRHITextureCreateDesc::Create2D(TEXT("LongTaskRenderTarget"), 1920, 1080, PF_B8G8R8A8)
 		.SetFlags(ETextureCreateFlags::RenderTargetable);
 
-	FTextureRHIRef LongTaskRenderTarget = RHICreateTexture(Desc);
+	FTextureRHIRef LongTaskRenderTarget = RHICmdList.CreateTexture(Desc);
 
 	FRHIRenderPassInfo RPInfo(LongTaskRenderTarget, ERenderTargetActions::DontLoad_Store);
 	RHICmdList.Transition(FRHITransitionInfo(LongTaskRenderTarget, ERHIAccess::Unknown, ERHIAccess::RTV));
diff --git a/Engine/Source/Runtime/Renderer/Private/IndirectLightingCache.cpp b/Engine/Source/Runtime/Renderer/Private/IndirectLightingCache.cpp
--- a/Engine/Source/Runtime/Renderer/Private/IndirectLightingCache.cpp
+++ b/Engine/Source/Runtime/Renderer/Private/IndirectLightingCache.cpp
@@ -154,11 +154,10 @@ FIndirectLightingCache::FIndirectLightingCache(ERHIFeatureLevel::Type InFeatureL
 	CacheSize = GLightingCacheDimension;
 }
 
-void FIndirectLightingCache::InitRHI(FRHICommandListBase&)
+void FIndirectLightingCache::InitRHI(FRHICommandListBase& RHICmdList)
 {
 	if (CanIndirectLightingCacheUseVolumeTexture(GetFeatureLevel()))
 	{
-		FRHICommandListImmediate& RHICmdList = FRHICommandListExecutor::GetImmediateCommandList();
 		ETextureCreateFlags Flags = TexCreate_ShaderResource | TexCreate_NoTiling;
 
 		FPooledRenderTargetDesc Desc(FPooledRenderTargetDesc::CreateVolumeDesc(
diff --git a/Engine/Source/Runtime/Renderer/Private/MeshEdgesRendering.cpp b/Engine/Source/Runtime/Renderer/Private/MeshEdgesRendering.cpp
--- a/Engine/Source/Runtime/Renderer/Private/MeshEdgesRendering.cpp
+++ b/Engine/Source/Runtime/Renderer/Private/MeshEdgesRendering.cpp
@@ -99,7 +99,7 @@ public:
 		);
 		SamplerStateRHI = GetOrCreateSamplerState(SamplerStateInitializer);
 
-		RenderTargetTextureRHI = TextureRHI = RHICreateTexture(Desc);
+		RenderTargetTextureRHI = TextureRHI = RHICmdList.CreateTexture(Desc);
 	}
 
 	virtual uint32 GetSizeX() const
diff --git a/Engine/Source/Runtime/Renderer/Private/PathTracing.cpp b/Engine/Source/Runtime/Renderer/Private/PathTracing.cpp
--- a/Engine/Source/Runtime/Renderer/Private/PathTracing.cpp
+++ b/Engine/Source/Runtime/Renderer/Private/PathTracing.cpp
@@ -2014,7 +2014,7 @@ void PreparePathTracingCloudMaterial(FRDGBuilder& GraphBuilder, FScene* Scene, T
 				FRHITextureCreateDesc::Create2D(TEXT("PathTracer.CloudAccelerationMap"), Resolution, Resolution, CloudAccelerationMapFormat)
 			.SetFlags(ETextureCreateFlags::ShaderResource | ETextureCreateFlags::UAV);
 
-			PathTracingState->CloudAccelerationMap = CreateRenderTarget(RHICreateTexture(Desc), Desc.DebugName);
+			PathTracingState->CloudAccelerationMap = CreateRenderTarget(GraphBuilder.RHICmdList.CreateTexture(Desc), Desc.DebugName);
 		}
 
 		Params.CloudAccelerationMap = GraphBuilder.RegisterExternalTexture(PathTracingState->CloudAccelerationMap, TEXT("PathTracer.CloudAccelerationMap"));
diff --git a/Engine/Source/Runtime/Renderer/Private/PostProcess/ACESUtils.cpp b/Engine/Source/Runtime/Renderer/Private/PostProcess/ACESUtils.cpp
--- a/Engine/Source/Runtime/Renderer/Private/PostProcess/ACESUtils.cpp
+++ b/Engine/Source/Runtime/Renderer/Private/PostProcess/ACESUtils.cpp
@@ -707,7 +707,7 @@ public:
 			.SetFlags(ETextureCreateFlags::ShaderResource);
 
 		// Create the RHI texture
-		TextureRHI = RHICreateTexture(Desc);
+		TextureRHI = RHICmdList.CreateTexture(Desc);
 
 		// Create the sampler state RHI resource.
 		FSamplerStateInitializerRHI SamplerStateInitializer(SF_Point, AM_Clamp, AM_Clamp, AM_Clamp);
diff --git a/Engine/Source/Runtime/Renderer/Private/ReflectionEnvironment.cpp b/Engine/Source/Runtime/Renderer/Private/ReflectionEnvironment.cpp
--- a/Engine/Source/Runtime/Renderer/Private/ReflectionEnvironment.cpp
+++ b/Engine/Source/Runtime/Renderer/Private/ReflectionEnvironment.cpp
@@ -145,7 +145,7 @@ bool IsReflectionCaptureAvailable()
 	return IsStaticLightingAllowed();
 }
 
-void FReflectionEnvironmentCubemapArray::InitRHI(FRHICommandListBase&)
+void FReflectionEnvironmentCubemapArray::InitRHI(FRHICommandListBase& RHICmdList)
 {
 	if (SupportsTextureCubeArray(GetFeatureLevel()))
 	{
@@ -171,8 +171,6 @@ void FReflectionEnvironmentCubemapArray::InitRHI(FRHICommandListBase&)
 				)
 			);
 
-		FRHICommandListImmediate& RHICmdList = FRHICommandListExecutor::GetImmediateCommandList();
-
 		// Allocate TextureCubeArray for the scene's reflection captures
 		GRenderTargetPool.FindFreeElement(RHICmdList, Desc, ReflectionEnvs, TEXT("ReflectionEnvs"));
 	}
diff --git a/Engine/Source/Runtime/Renderer/Private/SystemTextures.cpp b/Engine/Source/Runtime/Renderer/Private/SystemTextures.cpp
--- a/Engine/Source/Runtime/Renderer/Private/SystemTextures.cpp
+++ b/Engine/Source/Runtime/Renderer/Private/SystemTextures.cpp
@@ -680,7 +680,7 @@ void FSystemTextures::InitializeFeatureLevelDependentTextures(FRHICommandListImm
 				.SetFlags(ETextureCreateFlags::ShaderResource | ETextureCreateFlags::NoTiling)
 				.SetClassName(SystemTexturesName);
 
-			FTextureRHIRef Texture3D = RHICreateTexture(Desc);
+			FTextureRHIRef Texture3D = RHICmdList.CreateTexture(Desc);
 
 		    // Write the contents of the texture.
 		    TArray<uint32> DestBuffer;
@@ -1616,7 +1616,7 @@ FRDGTextureRef GetInternalDefaultTexture(
 			FRHITextureCreateDesc::Create3D(TEXT("DefaultTexture3D"), 1, 1, 1, Format)
 			.SetFlags(ETextureCreateFlags::ShaderResource);
 
-		FTextureRHIRef Texture = RHICreateTexture(Desc);
+		FTextureRHIRef Texture = GraphBuilder.RHICmdList.CreateTexture(Desc);
 		SetDefaultTextureData3D(GraphBuilder.RHICmdList, Texture, Desc, Value);
 		Entry.Texture = CreateRenderTarget(Texture, Desc.DebugName);
 	}
diff --git a/Engine/Source/Runtime/SlateRHIRenderer/Private/SlateRHIFontTexture.cpp b/Engine/Source/Runtime/SlateRHIRenderer/Private/SlateRHIFontTexture.cpp
--- a/Engine/Source/Runtime/SlateRHIRenderer/Private/SlateRHIFontTexture.cpp
+++ b/Engine/Source/Runtime/SlateRHIRenderer/Private/SlateRHIFontTexture.cpp
@@ -12,7 +12,7 @@ FSlateFontTextureRHIResource::FSlateFontTextureRHIResource(uint32 InWidth, uint3
 {
 }
 
-void FSlateFontTextureRHIResource::InitRHI(FRHICommandListBase&)
+void FSlateFontTextureRHIResource::InitRHI(FRHICommandListBase& RHICmdList)
 {
 	check( IsInRenderingThread() );
 
@@ -42,7 +42,7 @@ void FSlateFontTextureRHIResource::InitRHI(FRHICommandListBase&)
 				break;
 		}
 
-		ShaderResource = RHICreateTexture(Desc);
+		ShaderResource = RHICmdList.CreateTexture(Desc);
 
 		check( IsValidRef( ShaderResource ) );
 
diff --git a/Engine/Source/Runtime/WebBrowser/Private/Android/AndroidWebBrowserWidget.cpp b/Engine/Source/Runtime/WebBrowser/Private/Android/AndroidWebBrowserWidget.cpp
--- a/Engine/Source/Runtime/WebBrowser/Private/Android/AndroidWebBrowserWidget.cpp
+++ b/Engine/Source/Runtime/WebBrowser/Private/Android/AndroidWebBrowserWidget.cpp
@@ -377,12 +377,12 @@ void SAndroidWebBrowserWidget::Tick(const FGeometry& AllottedGeometry, const dou
 							FRHITextureCreateDesc::Create2D(TEXT("VideoTexture"), LocalSize, PF_R8G8B8A8)
 							.SetFlags(ETextureCreateFlags::External);
 
-						VideoTexture = RHICreateTexture(Desc);
+						VideoTexture = RHICmdList.CreateTexture(Desc);
 						PinnedJavaWebBrowser->SetVideoTexture(VideoTexture);
 
 						if (VideoTexture == nullptr)
 						{
-							UE_LOG(LogAndroid, Warning, TEXT("RHICreateTexture failed!"));
+							UE_LOG(LogAndroid, Warning, TEXT("CreateTexture failed!"));
 							return;
 						}
 
@@ -446,7 +446,7 @@ void SAndroidWebBrowserWidget::Tick(const FGeometry& AllottedGeometry, const dou
 					}
 
 					// write frame into texture
-					FRHITexture* Texture = Params.NewTextureSamplePtr->InitializeTexture();
+					FRHITexture* Texture = Params.NewTextureSamplePtr->InitializeTexture(RHICmdList);
 
 					if (Texture != nullptr)
 					{
diff --git a/Engine/Source/Runtime/WebBrowser/Private/Apple/ApplePlatformWebBrowser.cpp b/Engine/Source/Runtime/WebBrowser/Private/Apple/ApplePlatformWebBrowser.cpp
--- a/Engine/Source/Runtime/WebBrowser/Private/Apple/ApplePlatformWebBrowser.cpp
+++ b/Engine/Source/Runtime/WebBrowser/Private/Apple/ApplePlatformWebBrowser.cpp
@@ -250,12 +250,12 @@ class SAppleWebBrowserWidget : public SLeafWidget
 									FRHITextureCreateDesc::Create2D(TEXT("SAppleWebBrowserWidget_VideoTexture"), Params.Size, PF_R8G8B8A8)
 									.SetFlags(ETextureCreateFlags::CPUWritable | ETextureCreateFlags::CPUReadback);
 
-								VideoTexture = RHICreateTexture(Desc);
+								VideoTexture = RHICmdList.CreateTexture(Desc);
 								[NativeWebBrowser SetVideoTexture : VideoTexture];
 
 								if (VideoTexture == nullptr)
 								{
-									UE_LOG(LogAppleWebBrowser, Warning, TEXT("RHICreateTexture failed!"));
+									UE_LOG(LogAppleWebBrowser, Warning, TEXT("CreateTexture failed!"));
 									return;
 								}
 
diff --git a/Engine/Source/Runtime/WebBrowserTexture/Public/WebBrowserTextureSample.h b/Engine/Source/Runtime/WebBrowserTexture/Public/WebBrowserTextureSample.h
--- a/Engine/Source/Runtime/WebBrowserTexture/Public/WebBrowserTextureSample.h
+++ b/Engine/Source/Runtime/WebBrowserTexture/Public/WebBrowserTextureSample.h
@@ -137,11 +137,12 @@ public:
 	/**
 	* Initialize the sample with a texture resource.
 	*
+	* @param RHICmdList the command list used to create the texture.
 	* @return The texture resource object that will hold the sample data.
 	* @note This method must be called on the render thread.
 	* @see InitializeBuffer
 	*/
-	FRHITexture* InitializeTexture()
+	FRHITexture* InitializeTexture(FRHICommandListBase& RHICmdList)
 	{
 		check(IsInRenderingThread());
 
@@ -157,7 +158,7 @@ public:
 			.SetFlags(ETextureCreateFlags::SRGB | ETextureCreateFlags::RenderTargetable | ETextureCreateFlags::ShaderResource)
 			.SetInitialState(ERHIAccess::SRVMask);
 
-		Texture = RHICreateTexture(Desc);
+		Texture = RHICmdList.CreateTexture(Desc);
 
 		return Texture;
 	}
