[Movement] Add Pawn name to CMC & PlayerController logs for easier debugging of ForcePositionUpdate

#rnx

diff --git a/Engine/Source/Runtime/Engine/Private/Components/CharacterMovementComponent.cpp b/Engine/Source/Runtime/Engine/Private/Components/CharacterMovementComponent.cpp
--- a/Engine/Source/Runtime/Engine/Private/Components/CharacterMovementComponent.cpp
+++ b/Engine/Source/Runtime/Engine/Private/Components/CharacterMovementComponent.cpp
@@ -8644,7 +8644,7 @@ bool UCharacterMovementComponent::ForcePositionUpdate(float DeltaTime)
 	const bool bServerMoveHasOccurred = (ServerData->ServerTimeStampLastServerMove != 0.f);
 	if (bServerMoveHasOccurred)
 	{
-		UE_LOG(LogNetPlayerMovement, Log, TEXT("ForcePositionUpdate: %s (DeltaTime %.2f -> ServerTimeStamp %.2f)"), *GetNameSafe(CharacterOwner), DeltaTime, ServerData->CurrentClientTimeStamp);
+		UE_LOG(LogNetPlayerMovement, Log, TEXT("ForcePositionUpdate: %s (DeltaTime %.2f -> ServerTimeStamp %.2f) Pawn: %s"), *GetNameSafe(CharacterOwner), DeltaTime, ServerData->CurrentClientTimeStamp, *GetNameSafe(CharacterOwner));
 	}
 
 	// Force movement update.
@@ -9290,7 +9290,7 @@ bool UCharacterMovementComponent::VerifyClientTimeStamp(float TimeStamp, FNetwor
 			// Let's sync up to this TimeStamp in order to resolve movement desyncs ASAP
 			// This will result in this move having a zero DeltaTime (TimeStamp - CurrentClientTimeStamp), so it will perform no movement but it will send a correction,
 			// and we should be able to process the next move that arrives, which will advance ServerTimeStamp and unlock the forced updates (see 'APlayerController::TickActor()').
-			UE_LOG(LogNetPlayerMovement, Log, TEXT("Received a new move after performing ForcedUpdates. Updating CurrentClientTimeStamp from %f to %f"), ServerData.CurrentClientTimeStamp, TimeStamp);
+			UE_LOG(LogNetPlayerMovement, Log, TEXT("Received a new move after performing ForcedUpdates. Updating CurrentClientTimeStamp from %f to %f Pawn: %s"), ServerData.CurrentClientTimeStamp, TimeStamp, *GetNameSafe(CharacterOwner));
 			ServerData.CurrentClientTimeStamp = TimeStamp;
 			if (World != nullptr)
 			{
@@ -9299,7 +9299,7 @@ bool UCharacterMovementComponent::VerifyClientTimeStamp(float TimeStamp, FNetwor
 		}
 		else if (!bTimeStampResetDetected)
 		{
-			UE_LOG(LogNetPlayerMovement, VeryVerbose, TEXT("TimeStamp %f Accepted! CurrentTimeStamp: %f"), TimeStamp, ServerData.CurrentClientTimeStamp);
+			UE_LOG(LogNetPlayerMovement, VeryVerbose, TEXT("TimeStamp %f Accepted! CurrentTimeStamp: %f Pawn: %s"), TimeStamp, ServerData.CurrentClientTimeStamp, *GetNameSafe(CharacterOwner));
 			ProcessClientTimeStampForTimeDiscrepancy(TimeStamp, ServerData);
 		}
 	}
@@ -9307,7 +9307,7 @@ bool UCharacterMovementComponent::VerifyClientTimeStamp(float TimeStamp, FNetwor
 	{
 		if (bTimeStampResetDetected)
 		{
-			UE_LOG(LogNetPlayerMovement, Log, TEXT("TimeStamp expired. Before TimeStamp Reset. CurrentTimeStamp: %f, TimeStamp: %f"), ServerData.CurrentClientTimeStamp, TimeStamp);
+			UE_LOG(LogNetPlayerMovement, Log, TEXT("TimeStamp expired. Before TimeStamp Reset. CurrentTimeStamp: %f, TimeStamp: %f Pawn: %s"), ServerData.CurrentClientTimeStamp, TimeStamp, *GetNameSafe(CharacterOwner));
 		}
 		else
 		{
diff --git a/Engine/Source/Runtime/Engine/Private/PlayerController.cpp b/Engine/Source/Runtime/Engine/Private/PlayerController.cpp
--- a/Engine/Source/Runtime/Engine/Private/PlayerController.cpp
+++ b/Engine/Source/Runtime/Engine/Private/PlayerController.cpp
@@ -5419,7 +5419,7 @@ void APlayerController::TickActor( float DeltaSeconds, ELevelTick TickType, FAct
 							if (ServerData->ServerTimeStamp > ServerData->ServerTimeLastForcedUpdate)
 							{
 								// An update came in that was not a forced update (ie a real move), since ServerTimeStamp advanced outside this code.
-								UE_LOG(LogNetPlayerMovement, Log, TEXT("Movement detected, resetting forced update state (ServerTimeStamp %.6f > ServerTimeLastForcedUpdate %.6f)"), ServerData->ServerTimeStamp, ServerData->ServerTimeLastForcedUpdate);
+								UE_LOG(LogNetPlayerMovement, Log, TEXT("Movement detected, resetting forced update state (ServerTimeStamp %.6f > ServerTimeLastForcedUpdate %.6f) Pawn: %s"), ServerData->ServerTimeStamp, ServerData->ServerTimeLastForcedUpdate, *GetNameSafe(GetPawn()));
 								ServerData->ResetForcedUpdateState();
 							}
 							else
@@ -5430,7 +5430,7 @@ void APlayerController::TickActor( float DeltaSeconds, ELevelTick TickType, FAct
 								if (PawnTimeSinceForcingUpdates > PawnTimeForcedUpdateMaxDuration)
 								{
 									// Waiting for ServerTimeStamp to advance from a client move.
-									UE_LOG(LogNetPlayerMovement, Log, TEXT("Setting bForcedUpdateDurationExceeded=true (PawnTimeSinceForcingUpdates %.6f > PawnTimeForcedUpdateMaxDuration %.6f) (bLastRequestNeedsForcedUpdates:%d)"), PawnTimeSinceForcingUpdates, PawnTimeForcedUpdateMaxDuration, (int32)ServerData->bLastRequestNeedsForcedUpdates);
+									UE_LOG(LogNetPlayerMovement, Log, TEXT("Setting bForcedUpdateDurationExceeded=true (PawnTimeSinceForcingUpdates %.6f > PawnTimeForcedUpdateMaxDuration %.6f) (bLastRequestNeedsForcedUpdates:%d) Pawn: %s"), PawnTimeSinceForcingUpdates, PawnTimeForcedUpdateMaxDuration, (int32)ServerData->bLastRequestNeedsForcedUpdates, *GetNameSafe(GetPawn()));
 									ServerData->bForcedUpdateDurationExceeded = true;
 								}
 							}
