[PLA] Fixed PackedLevelActors not rendering properly after commiting a change that had no impact of the generated BP.

diff --git a/Engine/Source/Runtime/Engine/Private/PackedLevelActor/PackedLevelActor.cpp b/Engine/Source/Runtime/Engine/Private/PackedLevelActor/PackedLevelActor.cpp
--- a/Engine/Source/Runtime/Engine/Private/PackedLevelActor/PackedLevelActor.cpp
+++ b/Engine/Source/Runtime/Engine/Private/PackedLevelActor/PackedLevelActor.cpp
@@ -226,6 +226,9 @@ void APackedLevelActor::OnCommitChild(bool bChanged)
 	{
 		UnloadLevelInstance();
 
+		// When child edit state changes we need to dirty render state so that actor is no longer hidden
+		MarkComponentsRenderStateDirty();
+
 		if (bChildChanged)
 		{
 			bChildChanged = false;
@@ -236,15 +239,15 @@ void APackedLevelActor::OnCommitChild(bool bChanged)
 			if (UBlueprint* BlueprintToPack = GetRootBlueprint())
 			{
 				Builder->UpdateBlueprint(BlueprintToPack);
-				return; // return here because Actor might have been reinstanced
+				// return here because Actor might have been reinstanced
+				// @todo_ow: We should only return if actor was reinstanced
+				return;
 			}
 			else
 			{
 				Builder->PackActor(this, GetWorldAsset());
 			}
 		}
-		// When child edit state changes we need to dirty render state so that actor is no longer hidden
-		MarkComponentsRenderStateDirty();
 	}
 }
 
@@ -258,18 +261,20 @@ void APackedLevelActor::OnCommit(bool bChanged)
 {
 	Super::OnCommit(bChanged);
 
+	// When edit state changes we need to dirty render state so that actor is no longer hidden
+	MarkComponentsRenderStateDirty();
+
 	if (bChanged)
 	{
 		if (UBlueprint* BlueprintToPack = GetRootBlueprint())
 		{
 			TSharedPtr<FPackedLevelActorBuilder> Builder = FPackedLevelActorBuilder::CreateDefaultBuilder();
 			Builder->UpdateBlueprint(BlueprintToPack);
-			return; // return here because Actor might have been reinstanced
+			// return here because Actor might have been reinstanced
+			// @todo_ow: We should only return if actor was reinstanced
+			return;
 		}
 	}
-	
-	// When edit state changes we need to dirty render state so that actor is no longer hidden
-	MarkComponentsRenderStateDirty();
 }
 
 bool APackedLevelActor::IsHiddenEd() const
