Fixed missing AnimationMinScreenSize in ISKM.

diff --git a/Engine/Source/Runtime/Engine/Private/InstancedSkinnedMeshSceneProxy.cpp b/Engine/Source/Runtime/Engine/Private/InstancedSkinnedMeshSceneProxy.cpp
--- a/Engine/Source/Runtime/Engine/Private/InstancedSkinnedMeshSceneProxy.cpp
+++ b/Engine/Source/Runtime/Engine/Private/InstancedSkinnedMeshSceneProxy.cpp
@@ -10,6 +10,7 @@ FInstancedSkinnedMeshData::FInstancedSkinnedMeshData(const FInstancedSkinnedMesh
 	InstanceStartCullDistance   = InMeshDesc.InstanceStartCullDistance;
 	InstanceEndCullDistance     = InMeshDesc.InstanceEndCullDistance;
 	InstanceDataSceneProxy      = InMeshDesc.InstanceDataSceneProxy;
+	AnimationMinScreenSize      = InMeshDesc.AnimationMinScreenSize;
 }
 
 bool FInstancedSkinnedMeshData::GetInstanceDrawDistanceMinMax(FVector2f& OutCullRange) const
