New Material Translator:  Support for material instance BlendMode overrides, "SelectionColor" editor feature which allows the output emissive color to be overridden, plus one expression (TextureSampleParameterSubUV).

#rnx

diff --git a/Engine/Source/Runtime/Engine/Private/Materials/HLSLMaterialTranslator.cpp b/Engine/Source/Runtime/Engine/Private/Materials/HLSLMaterialTranslator.cpp
--- a/Engine/Source/Runtime/Engine/Private/Materials/HLSLMaterialTranslator.cpp
+++ b/Engine/Source/Runtime/Engine/Private/Materials/HLSLMaterialTranslator.cpp
@@ -15990,26 +15990,7 @@ bool FHLSLMaterialTranslator::IsCurrentlyCompilingForPreviousFrame() const
 
 bool FHLSLMaterialTranslator::IsDevelopmentFeatureEnabled(const FName& FeatureName) const
 {
-	if (FeatureName == NAME_SelectionColor)
-	{
-		// This is an editor-only feature (see FDefaultMaterialInstance::GetVectorValue).
-
-		// Determine if we're sure the editor will never run using the target shader platform.
-		// The list below may not be comprehensive enough, but it definitely includes platforms which won't use selection color for sure.
-		const bool bEditorMayUseTargetShaderPlatform = IsPCPlatform(Platform);
-		static const auto CVar = IConsoleManager::Get().FindTConsoleVariableDataInt(TEXT("r.CompileShadersForDevelopment"));
-		const bool bCompileShadersForDevelopment = (CVar && CVar->GetValueOnAnyThread() != 0);
-
-		return
-			// Does the material explicitly forbid development features?
-			Material->GetAllowDevelopmentShaderCompile()
-			// Can the editor run using the current shader platform?
-			&& bEditorMayUseTargetShaderPlatform
-			// Are shader development features globally disabled?
-			&& bCompileShadersForDevelopment;
-	}
-
-	return true;
+	return UE::MaterialTranslatorUtils::IsDevelopmentFeatureEnabled(FeatureName, Platform, Material->GetMaterialInterface()->GetMaterial());
 }
 
 void FHLSLMaterialTranslator::PrepareEnvironmentDefines()
diff --git a/Engine/Source/Runtime/Engine/Private/Materials/MaterialExpressionsIR.cpp b/Engine/Source/Runtime/Engine/Private/Materials/MaterialExpressionsIR.cpp
--- a/Engine/Source/Runtime/Engine/Private/Materials/MaterialExpressionsIR.cpp
+++ b/Engine/Source/Runtime/Engine/Private/Materials/MaterialExpressionsIR.cpp
@@ -223,6 +223,7 @@
 #include "Materials/MaterialExpressionTextureSampleParameter.h"
 #include "Materials/MaterialExpressionTextureSampleParameter2DArray.h"
 #include "Materials/MaterialExpressionTextureSampleParameterCube.h"
+#include "Materials/MaterialExpressionTextureSampleParameterSubUV.h"
 #include "Materials/MaterialExpressionTextureSampleParameterVolume.h"
 #include "Materials/MaterialExpressionThinTranslucentMaterialOutput.h"
 #include "Materials/MaterialExpressionTime.h"
@@ -1339,6 +1340,38 @@
 	return TextureValue;
 }
 
+static FValueRef BuildTextureObjectParameter(MIR::FEmitter& Em, UMaterialExpressionTextureSampleParameter* Expr)
+{
+	FMaterialParameterMetadata Param;
+	if (!Expr->GetParameterValue(Param))
+	{
+		Em.Error(TEXT("Failed to get parameter value"));
+		return nullptr;
+	}
+
+	if (!Expr->Texture)
+	{
+		Em.Error(TEXT("Requires valid texture"));
+		return nullptr;
+	}
+
+	FString SamplerTypeError;
+	if (!UMaterialExpressionTextureBase::VerifySamplerType(Em.GetShaderPlatform(), Em.GetTargetPlatform(), Expr->Texture, Expr->SamplerType, SamplerTypeError))
+	{
+		Em.Errorf(TEXT("%s"), *SamplerTypeError);
+		return nullptr;
+	}
+
+	FValueRef ParameterValue = Em.Parameter(Expr->GetParameterName(), Param, Expr->SamplerType);
+	if (!ParameterValue->Type.IsTexture())
+	{
+		Em.Error(TEXT("Parameter is not a texture"));
+		return nullptr;
+	}
+
+	return ParameterValue;
+}
+
 static void BuildTextureSampleExpression(MIR::FEmitter& Em, UMaterialExpressionTextureSample* Expr, FValueRef Texture, EMaterialValueType TextureType)
 {
 	FValueRef TexCoords = Em.TryInput(&Expr->Coordinates);
@@ -1386,12 +1419,8 @@
 	BuildTextureSampleExpression(Em, this, TextureValue, MIR::Internal::GetTextureMaterialValueType(DefaultTexture));
 }
 
-// Inherits from UMaterialExpressionTextureSample, but uses different particle specific UVs, and optionally supports blending two different texture samples.
-void UMaterialExpressionParticleSubUV::Build(MIR::FEmitter& Em)
+static void EmitParticleSubUV(MIR::FEmitter& Em, UMaterialExpressionTextureSample* Expr, FValueRef TextureValue, bool bBlend, FValueRef DummyDependency)
 {
-	FValueRef TextureValue = BuildTextureValue(Em, this);
-	UE_MIR_CHECKPOINT(Em);
-
 	EMaterialValueType TextureType = MIR::Internal::GetTextureMaterialValueType(Em.GetTextureFromValue(TextureValue));
 
 	// Although the parent UMaterialExpressionTextureSample class includes an automatic view mip bias flag, it is specifically ignored by ParticleSubUV.
@@ -1401,13 +1430,13 @@
 	static FName NAME_ParticleSubUVCoords1("ParticleSubUVCoords1");
 	static FName NAME_ParticleSubUVLerp("ParticleSubUVLerp");
 
-	FValueRef TexCoords0 = EmitInlineHLSL(Em, NAME_ParticleSubUVCoords0);
-	FValueRef Sample0 = BuildTextureSample(Em, this, TextureValue, TextureType, TexCoords0, bAutomaticViewMipBias);
+	FValueRef TexCoords0 = DummyDependency ? EmitInlineHLSL(Em, NAME_ParticleSubUVCoords0, { DummyDependency }) : EmitInlineHLSL(Em, NAME_ParticleSubUVCoords0);
+	FValueRef Sample0 = BuildTextureSample(Em, Expr, TextureValue, TextureType, TexCoords0, bAutomaticViewMipBias);
 
 	if (bBlend)
 	{
 		FValueRef TexCoords1 = EmitInlineHLSL(Em, NAME_ParticleSubUVCoords1);
-		FValueRef Sample1 = BuildTextureSample(Em, this, TextureValue, TextureType, TexCoords1, bAutomaticViewMipBias);
+		FValueRef Sample1 = BuildTextureSample(Em, Expr, TextureValue, TextureType, TexCoords1, bAutomaticViewMipBias);
 
 		FValueRef SubImageLerp = EmitInlineHLSL(Em, NAME_ParticleSubUVLerp);
 
@@ -1423,6 +1452,29 @@
 	Em.Output(5, Sample0);
 }
 
+// Inherits from UMaterialExpressionTextureSample, but uses different particle specific UVs, and optionally supports blending two different texture samples.
+void UMaterialExpressionParticleSubUV::Build(MIR::FEmitter& Em)
+{
+	FValueRef TextureValue = BuildTextureValue(Em, this);
+	UE_MIR_CHECKPOINT(Em);
+
+	EmitParticleSubUV(Em, this, TextureValue, bBlend, {});
+}
+
+// Similar to above, but texture comes from a parameter, rather than a local or object texture reference.
+void UMaterialExpressionTextureSampleParameterSubUV::Build(MIR::FEmitter& Em)
+{
+	FValueRef ParameterValue = BuildTextureObjectParameter(Em, this);
+	UE_MIR_CHECKPOINT(Em);
+
+	// while this expression does provide a TextureCoordinate input pin, it is, and has always been, ignored.  And only
+	// supports using UV0.  Further, in order to support non-vertex fetch implementations we need to be sure to register
+	// the use of the first texture slot
+	FValueRef DummyDependency = Em.ExternalInput(MIR::TexCoordIndexToExternalInput(0));
+
+	EmitParticleSubUV(Em, this, ParameterValue, bBlend, DummyDependency);
+}
+
 // Inherits from UMaterialExpressionTextureSample, but does extra math on the sample afterwards.  Note that this was an HLSL utility function in the
 // original translator, but uses ops here.  The main advantage of the ops version is that it uses the standard texture sampling code path, rather than
 // sampling the texture in the utility function, meaning it supports all sampling features (the original would break if using non-standard sampling).
@@ -1456,38 +1508,6 @@
 	Em.Output(0, Em.Add(Em.Divide(Top, TexDDLength), ThresholdConst));
 }
 
-static FValueRef BuildTextureObjectParameter(MIR::FEmitter& Em, UMaterialExpressionTextureSampleParameter* Expr)
-{
-	FMaterialParameterMetadata Param;
-	if (!Expr->GetParameterValue(Param))
-	{
-		Em.Error(TEXT("Failed to get parameter value"));
-		return nullptr;
-	}
-
-	if (!Expr->Texture)
-	{
-		Em.Error(TEXT("Requires valid texture"));
-		return nullptr;
-	}
-
-	FString SamplerTypeError;
-	if (!UMaterialExpressionTextureBase::VerifySamplerType(Em.GetShaderPlatform(), Em.GetTargetPlatform(), Expr->Texture, Expr->SamplerType, SamplerTypeError))
-	{
-		Em.Errorf(TEXT("%s"), *SamplerTypeError);
-		return nullptr;
-	}
-
-	FValueRef ParameterValue = Em.Parameter(Expr->GetParameterName(), Param, Expr->SamplerType);
-	if (!ParameterValue->Type.IsTexture())
-	{
-		Em.Error(TEXT("Parameter is not a texture"));
-		return nullptr;
-	}
-
-	return ParameterValue;
-}
-
 static void BuildTextureSampleParameter(MIR::FEmitter& Em, UMaterialExpressionTextureSampleParameter* Expr)
 {
 	FString ErrorMessage;
diff --git a/Engine/Source/Runtime/Engine/Private/Materials/MaterialIRModuleBuilder.cpp b/Engine/Source/Runtime/Engine/Private/Materials/MaterialIRModuleBuilder.cpp
--- a/Engine/Source/Runtime/Engine/Private/Materials/MaterialIRModuleBuilder.cpp
+++ b/Engine/Source/Runtime/Engine/Private/Materials/MaterialIRModuleBuilder.cpp
@@ -27,6 +27,7 @@
 #include "MaterialShared.h"
 #include "Materials/MaterialExpression.h"
 #include "Materials/MaterialInsights.h"
+#include "Math/Color.h"
 #include "Engine/Texture.h"
 #include "Misc/FileHelper.h"
 #include "PostProcess/PostProcessMaterialInputs.h"
@@ -96,6 +97,7 @@ struct FMaterialIRModuleBuilderImpl
 		Module->TargetPlatform = Builder->TargetPlatform;
 		Module->FeatureLevel = Builder->FeatureLevel;
 		Module->QualityLevel = Builder->QualityLevel;
+		Module->BlendMode = Builder->BlendMode;
 		
 		// Declare a entry point to evaluate the vertex stage.
 		Module->AddEntryPoint(TEXTVIEW("VertexStage"), MIR::Stage_Vertex, 1);
@@ -454,6 +456,16 @@ struct FMaterialIRModuleBuilderImpl
 			// The value being set to this material output is now valid.
 			check(PropertyValue);
 
+			// Add support for lerp to selection color for PC development builds.
+			if (Property == MP_EmissiveColor &&
+				Builder->Material->MaterialDomain != MD_Volume &&
+				UE::MaterialTranslatorUtils::IsDevelopmentFeatureEnabled(NAME_SelectionColor, Module->GetShaderPlatform(), Builder->Material))
+			{
+				FMaterialParameterMetadata ParameterMetadata{ FMaterialParameterValue(FLinearColor::Transparent) };
+				MIR::FValue* SelectionColor = Emitter.Parameter(NAME_SelectionColor, ParameterMetadata);
+				PropertyValue = Emitter.Lerp(PropertyValue, Emitter.Swizzle(SelectionColor, MIR::FSwizzleMask::XYZ()), Emitter.Subscript(SelectionColor, 3));
+			}
+
 			// Set this property value onto the module
 			Module->SetPropertyValue(Property, PropertyValue);
 
@@ -789,7 +801,7 @@ struct FMaterialIRModuleBuilderImpl
 			ValueAnalyzer.Material,
 			CompilationOutput,
 			Module->GetCompiledShadingModels(),
-			ValueAnalyzer.Material->GetBlendMode(),					// @todo-jason.hoerner - we need to handle material instances overriding the blend mode
+			Module->GetBlendMode(),
 			Module->IsMaterialPropertyUsed(MP_FrontMaterial),
 			Module->GetShaderPlatform(),
 			ValidationErrors);
diff --git a/Engine/Source/Runtime/Engine/Private/Materials/MaterialIRValueAnalyzer.cpp b/Engine/Source/Runtime/Engine/Private/Materials/MaterialIRValueAnalyzer.cpp
--- a/Engine/Source/Runtime/Engine/Private/Materials/MaterialIRValueAnalyzer.cpp
+++ b/Engine/Source/Runtime/Engine/Private/Materials/MaterialIRValueAnalyzer.cpp
@@ -532,7 +532,7 @@ static void AnalyzeScreenTexture(FMaterialIRValueAnalyzer& Analyzer, MIR::FScree
 						Analyzer.Module->AddError(nullptr, TEXT("Can only read scene depth below water when material Domain is set to Surface or PostProcess."));
 					}
 
-					if (IsTranslucentBlendMode(Analyzer.Material->GetBlendMode()))
+					if (IsTranslucentBlendMode(Analyzer.Module->GetBlendMode()))
 					{
 						Analyzer.Module->AddError(nullptr, TEXT("Can only read scene depth below water when material Blend Mode isn't translucent."));
 					}
@@ -550,7 +550,7 @@ static void AnalyzeScreenTexture(FMaterialIRValueAnalyzer& Analyzer, MIR::FScree
 				Analyzer.CompilationOutput->SetIsDBufferTextureLookupUsed(true);
 
 				// Error checking
-				if (MaterialDomain != MD_Surface || IsTranslucentBlendMode(Analyzer.Material->GetBlendMode()))
+				if (MaterialDomain != MD_Surface || IsTranslucentBlendMode(Analyzer.Module->GetBlendMode()))
 				{
 					Analyzer.Module->AddError(nullptr, TEXT("DBuffer scene textures are only available on opaque or masked surfaces."));
 				}
diff --git a/Engine/Source/Runtime/Engine/Private/Materials/MaterialShared.cpp b/Engine/Source/Runtime/Engine/Private/Materials/MaterialShared.cpp
--- a/Engine/Source/Runtime/Engine/Private/Materials/MaterialShared.cpp
+++ b/Engine/Source/Runtime/Engine/Private/Materials/MaterialShared.cpp
@@ -3425,6 +3425,7 @@ bool FMaterial::Translate_New(const FMaterialShaderMapId& InShaderMapId,
 		.TargetPlatform = InTargetPlatform,
 		.FeatureLevel = InShaderMapId.FeatureLevel,
 		.QualityLevel = InShaderMapId.QualityLevel,
+		.BlendMode = GetBlendMode(),
 		.StaticParameters = InStaticParameters,
 		.TargetInsights = GetMaterialInterface()->MaterialInsight.Get(),
 		.PreviewExpression = GetMaterialGraphNodePreviewExpression()
@@ -6515,6 +6516,30 @@ FString CustomExpressionSceneTextureInputFixup(const UMaterialExpressionCustom*
 	return ModifiedCode;
 }
 
+bool IsDevelopmentFeatureEnabled(const FName& FeatureName, EShaderPlatform Platform, UMaterial* Material)
+{
+	if (FeatureName == NAME_SelectionColor)
+	{
+		// This is an editor-only feature (see FDefaultMaterialInstance::GetVectorValue).
+
+		// Determine if we're sure the editor will never run using the target shader platform.
+		// The list below may not be comprehensive enough, but it definitely includes platforms which won't use selection color for sure.
+		const bool bEditorMayUseTargetShaderPlatform = IsPCPlatform(Platform);
+		static const auto CVar = IConsoleManager::Get().FindTConsoleVariableDataInt(TEXT("r.CompileShadersForDevelopment"));
+		const bool bCompileShadersForDevelopment = (CVar && CVar->GetValueOnAnyThread() != 0);
+
+		return
+			// Does the material explicitly forbid development features?
+			Material->bAllowDevelopmentShaderCompile
+			// Can the editor run using the current shader platform?
+			&& bEditorMayUseTargetShaderPlatform
+			// Are shader development features globally disabled?
+			&& bCompileShadersForDevelopment;
+	}
+
+	return true;
+}
+
 }  // namespace UE::MaterialTranslatorUtils
 
 #endif // WITH_EDITOR
diff --git a/Engine/Source/Runtime/Engine/Private/Materials/MaterialSharedPrivate.h b/Engine/Source/Runtime/Engine/Private/Materials/MaterialSharedPrivate.h
--- a/Engine/Source/Runtime/Engine/Private/Materials/MaterialSharedPrivate.h
+++ b/Engine/Source/Runtime/Engine/Private/Materials/MaterialSharedPrivate.h
@@ -118,6 +118,9 @@ namespace UE::MaterialTranslatorUtils
 
 	// Fixup for SceneTexture and UserSceneTexture inputs to custom HLSL. See implementation for more info.
 	FString CustomExpressionSceneTextureInputFixup(const UMaterialExpressionCustom* Custom, const TCHAR* Code, TArray<int8>& OutSceneTextureInfo);
+	
+	// Check for whether development feature with the given name should be enabled (currently just used for editor Selection Color).
+	bool IsDevelopmentFeatureEnabled(const FName& FeatureName, EShaderPlatform Platform, UMaterial* Material);
 
 } // namespace UE::MaterialTranslatorUtils
 
diff --git a/Engine/Source/Runtime/Engine/Public/Materials/MaterialExpressionTextureSampleParameterSubUV.h b/Engine/Source/Runtime/Engine/Public/Materials/MaterialExpressionTextureSampleParameterSubUV.h
--- a/Engine/Source/Runtime/Engine/Public/Materials/MaterialExpressionTextureSampleParameterSubUV.h
+++ b/Engine/Source/Runtime/Engine/Public/Materials/MaterialExpressionTextureSampleParameterSubUV.h
@@ -26,6 +26,7 @@ public:
 	//~ End UMaterialExpressionTextureSampleParameter Interface
 
 	//~ Begin UMaterialExpression Interface
+	virtual void Build(MIR::FEmitter& Emitter) override;
 	virtual int32 Compile(class FMaterialCompiler* Compiler, int32 OutputIndex) override;
 	virtual void GetCaption(TArray<FString>& OutCaptions) const override;
 	//~ End UMaterialExpression Interface
diff --git a/Engine/Source/Runtime/Engine/Public/Materials/MaterialIRModule.h b/Engine/Source/Runtime/Engine/Public/Materials/MaterialIRModule.h
--- a/Engine/Source/Runtime/Engine/Public/Materials/MaterialIRModule.h
+++ b/Engine/Source/Runtime/Engine/Public/Materials/MaterialIRModule.h
@@ -84,6 +84,9 @@ public:
 	// Returns the quality level associated with this module.
 	EMaterialQualityLevel::Type GetQualityLevel() const { return QualityLevel; }
 
+	// Returns the blend mode associated with this module.
+	EBlendMode GetBlendMode() const { return BlendMode; }
+
 	// Retrieves the material compilation output.
 	const FMaterialCompilationOutput& GetCompilationOutput() const { return CompilationOutput; }
 
@@ -179,6 +182,9 @@ private:
 	// Target quality level.
 	EMaterialQualityLevel::Type QualityLevel;
 
+	// Target blend mode -- separate from UMaterial, as it may potentially overridden in a UMaterialInstance.
+	EBlendMode BlendMode;
+
 	// Compilation output data.
 	FMaterialCompilationOutput CompilationOutput;
 
diff --git a/Engine/Source/Runtime/Engine/Public/Materials/MaterialIRModuleBuilder.h b/Engine/Source/Runtime/Engine/Public/Materials/MaterialIRModuleBuilder.h
--- a/Engine/Source/Runtime/Engine/Public/Materials/MaterialIRModuleBuilder.h
+++ b/Engine/Source/Runtime/Engine/Public/Materials/MaterialIRModuleBuilder.h
@@ -1,6 +1,7 @@
 // Copyright Epic Games, Inc. All Rights Reserved.
 #pragma once
 
+#include "Engine/EngineTypes.h"
 #include "Materials/MaterialIRCommon.h"
 
 #if WITH_EDITOR
@@ -13,6 +14,7 @@ struct FMaterialIRModuleBuilder
 	const ITargetPlatform* TargetPlatform;
 	ERHIFeatureLevel::Type FeatureLevel;
 	EMaterialQualityLevel::Type QualityLevel;
+	EBlendMode BlendMode;
 	const FStaticParameterSet& StaticParameters;
 	FMaterialInsights* TargetInsights{};
 	UMaterialExpression* PreviewExpression;
