Fixed NumLODs not being assigned to the skeletal mesh proxy.

diff --git a/Engine/Source/Runtime/Engine/Private/SkeletalMeshSceneProxy.cpp b/Engine/Source/Runtime/Engine/Private/SkeletalMeshSceneProxy.cpp
--- a/Engine/Source/Runtime/Engine/Private/SkeletalMeshSceneProxy.cpp
+++ b/Engine/Source/Runtime/Engine/Private/SkeletalMeshSceneProxy.cpp
@@ -135,7 +135,8 @@ FSkeletalMeshSceneProxy::FSkeletalMeshSceneProxy(const FSkinnedMeshSceneProxyDes
 	// Skinning is supported by this proxy
 	bSkinnedMesh = true;
 	bInstancedSkinnedMesh = SceneExtensionProxy && SceneExtensionProxy->UseInstancing();
-	bDoesMeshBatchesUseSceneInstanceCount = true;
+	bDoesMeshBatchesUseSceneInstanceCount = bInstancedSkinnedMesh;
+	NumLODs = RenderData->LODRenderData.Num();
 
 #if RHI_RAYTRACING
 	bRayTraceStatic = bRenderStatic || bInstancedSkinnedMesh;
