Fixed incorrect deprecation string formatting.

diff --git a/Engine/Source/Runtime/Engine/Public/GPUSkinVertexFactory.h b/Engine/Source/Runtime/Engine/Public/GPUSkinVertexFactory.h
--- a/Engine/Source/Runtime/Engine/Public/GPUSkinVertexFactory.h
+++ b/Engine/Source/Runtime/Engine/Public/GPUSkinVertexFactory.h
@@ -812,7 +812,7 @@ public:
 		Color    = 1 << EVertexAttribute::VertexColor,
 	};
 
-	UE_DEPRECATED(5.7, TEXT("FGPUSkinPassthroughVertexFactory now requires NumVertices."))
+	UE_DEPRECATED(5.7, "FGPUSkinPassthroughVertexFactory now requires NumVertices.")
 	FGPUSkinPassthroughVertexFactory(ERHIFeatureLevel::Type InFeatureLevel, EVertexAttributeFlags InVertexAttributeMask)
 		: FGPUSkinPassthroughVertexFactory(InFeatureLevel, InVertexAttributeMask, 0)
 	{}
