Fixed another malformed deprecation string.

diff --git a/Engine/Source/Runtime/Engine/Public/GPUSkinCache.h b/Engine/Source/Runtime/Engine/Public/GPUSkinCache.h
--- a/Engine/Source/Runtime/Engine/Public/GPUSkinCache.h
+++ b/Engine/Source/Runtime/Engine/Public/GPUSkinCache.h
@@ -137,7 +137,7 @@ public:
 		IntermediateAccumBufferNumInts = 8,
 	};
 
-	UE_DEPRECATED(5.7, TEXT("SkinCache no longer requires a memory limit"))
+	UE_DEPRECATED(5.7, "SkinCache no longer requires a memory limit")
 	FGPUSkinCache(ERHIFeatureLevel::Type InFeatureLevel, bool bInRequiresMemoryLimit, UWorld* InWorld)
 		: FGPUSkinCache(InFeatureLevel, InWorld)
 	{}
