Fix IsThinSurface incorrectly visible when Substrate's FrontMaterial wasn't used/plugged.

diff --git a/Engine/Source/Runtime/Engine/Private/Materials/Material.cpp b/Engine/Source/Runtime/Engine/Private/Materials/Material.cpp
--- a/Engine/Source/Runtime/Engine/Private/Materials/Material.cpp
+++ b/Engine/Source/Runtime/Engine/Private/Materials/Material.cpp
@@ -7874,6 +7874,8 @@ static bool IsPropertyActive_Internal(EMaterialProperty InProperty,
 			Active = bUsesShadingModelFromMaterialExpression;
 			break;
 		case MP_SurfaceThickness:
+			Active = false; // Only shown when material is Shown as Substrate (e.g., with FrontMaterial connected)
+			break;
 		case MP_FrontMaterial:
 			{
 				Active = bSubstrateHiddenConversion;
