Fixed UMaterialExpressionParticleSubUV::Compile bug that caused the node not to use texture object connected to the Texture input pin.

diff --git a/Engine/Source/Runtime/Engine/Private/Materials/MaterialExpressions.cpp b/Engine/Source/Runtime/Engine/Private/Materials/MaterialExpressions.cpp
--- a/Engine/Source/Runtime/Engine/Private/Materials/MaterialExpressions.cpp
+++ b/Engine/Source/Runtime/Engine/Private/Materials/MaterialExpressions.cpp
@@ -10049,8 +10049,7 @@ int32 UMaterialExpressionParticleSubUV::Compile(class FMaterialCompiler* Compile
 	TEnumAsByte<enum EMaterialSamplerType> SamplerTypeToUse = SamplerType; 
 	if (TextureObject.GetTracedInput().Expression != nullptr)
 	{
-		TObjectPtr<UMaterialExpressionTextureObjectParameter> TextureObjectParameter = Cast<UMaterialExpressionTextureObjectParameter>(TextureObject.GetTracedInput().Expression);
-		if (TextureObjectParameter.Get())
+		if (UMaterialExpressionTextureObject* TextureObjectParameter = Cast<UMaterialExpressionTextureObject>(TextureObject.GetTracedInput().Expression))
 		{
 			TextureToCompile = TextureObjectParameter->Texture;
 			SamplerTypeToUse = TextureObjectParameter->SamplerType;
