PR #13715 - Fixed issue with FObjectReplicator's CanSkipUpdate with net.PushModelSkipUndirtiedReplication enabled.

diff --git a/Engine/Source/Runtime/Engine/Private/DataReplication.cpp b/Engine/Source/Runtime/Engine/Private/DataReplication.cpp
--- a/Engine/Source/Runtime/Engine/Private/DataReplication.cpp
+++ b/Engine/Source/Runtime/Engine/Private/DataReplication.cpp
@@ -898,7 +898,13 @@ void FObjectReplicator::ReceivedNak( int32 NakPacketId )
 	{
 		if (FSendingRepState* SendingRepState = RepState.IsValid() ? RepState->GetSendingRepState() : nullptr)
 		{
-			SendingRepState->CustomDeltaChangeIndex--;
+			// This is used to determine if we can skip the update,
+			// unconditionally decrementing leads to objects with no custom delta serializer
+			// to be unable to be skipped for replication since their value is never updated elsewhere
+			if (SendingRepState->CustomDeltaChangeIndex != 0)
+			{
+				SendingRepState->CustomDeltaChangeIndex--;
+			}
 			
 			// Go over properties tracked with histories, and mark them as needing to be resent.
 			for (int32 i = SendingRepState->HistoryStart; i < SendingRepState->HistoryEnd; ++i)
@@ -2027,6 +2033,11 @@ bool FObjectReplicator::ReplicateProperties_r(FOutBunch& Bunch, FReplicationFlag
 			if (bHasDirtyCustomDeltaProperties)
 			{
 				ChangelistState->CustomDeltaChangeIndex++;
+				// If we wrap around, skip 0 as this is used to detect that we never sent any custom delta properties
+				if (ChangelistState->CustomDeltaChangeIndex == 0)
+				{
+					++ChangelistState->CustomDeltaChangeIndex;
+				}
 			}
 
 			SendingRepState->CustomDeltaChangeIndex = ChangelistState->CustomDeltaChangeIndex;
@@ -2037,10 +2048,19 @@ bool FObjectReplicator::ReplicateProperties_r(FOutBunch& Bunch, FReplicationFlag
 			if (bHasDirtyCustomDeltaProperties && bSkippedPropertyCondition)
 			{
 				ChangelistState->CustomDeltaChangeIndex++;
-
-				// this skipping connection is up to date
-				SendingRepState->CustomDeltaChangeIndex = ChangelistState->CustomDeltaChangeIndex;
+				// If we wrap around, skip 0 as this is used to detect that we never sent any custom delta properties
+				if (ChangelistState->CustomDeltaChangeIndex == 0)
+				{
+					++ChangelistState->CustomDeltaChangeIndex;
+				}
 			}
+
+			// SendingRepState->CustomDeltaChangeIndex needs to be incremented in all cases just like above,
+			// otherwise we cannot skip CustomDeltaProperties for connections on which they do not write
+			// if they received a NAK on their last write
+			
+			// this skipping connection is up to date
+			SendingRepState->CustomDeltaChangeIndex = ChangelistState->CustomDeltaChangeIndex;
 		}
 #endif //WITH_PUSH_MODEL
 	}
