Fix incorrect material classification causing hair/eye tile/pixels to not being rendered correctly.

diff --git a/Engine/Source/Runtime/Core/Private/UObject/FortniteMainBranchObjectVersion.cpp b/Engine/Source/Runtime/Core/Private/UObject/FortniteMainBranchObjectVersion.cpp
--- a/Engine/Source/Runtime/Core/Private/UObject/FortniteMainBranchObjectVersion.cpp
+++ b/Engine/Source/Runtime/Core/Private/UObject/FortniteMainBranchObjectVersion.cpp
@@ -13,7 +13,7 @@ TMap<FGuid, FGuid> FFortniteMainBranchObjectVersion::GetSystemGuids()
 	SystemGuids.Add(DevGuids.MATERIALSHADERMAP_DERIVEDDATA_VER, FGuid("F051877974B44D409435F321BDB39DBB"));
 	SystemGuids.Add(DevGuids.NIAGARASHADERMAP_DERIVEDDATA_VER, FGuid("6360A977062842A29ED30E3A7ACB0E64"));
 	SystemGuids.Add(DevGuids.Niagara_LatestScriptCompileVersion, FGuid("783FF3700BB641C799B2D43552FBF91E"));
-	SystemGuids.Add(DevGuids.MaterialTranslationDDCVersion, FGuid("240CBCCB056D4E96BE11C168C07AA73C"));
+	SystemGuids.Add(DevGuids.MaterialTranslationDDCVersion, FGuid("0EE6461792BB4D3C885E06E5CE03B068"));
 
 	// These GUIDs are in AppendRestrictedFortniteMainStreamObjectVersionGuids as they should rarely be changed.
 	//SystemGuids.Add(DevGuids.NANITE_DERIVEDDATA_VER, FGuid("ACF98184978742A592537DA223A1E6BE"));
diff --git a/Engine/Source/Runtime/Engine/Private/Materials/HLSLMaterialTranslator.cpp b/Engine/Source/Runtime/Engine/Private/Materials/HLSLMaterialTranslator.cpp
--- a/Engine/Source/Runtime/Engine/Private/Materials/HLSLMaterialTranslator.cpp
+++ b/Engine/Source/Runtime/Engine/Private/Materials/HLSLMaterialTranslator.cpp
@@ -1505,18 +1505,19 @@ void FHLSLMaterialTranslator::TranslateMaterial()
 			}
 			else if (MaterialShadingModels.HasShadingModel(MSM_Eye))
 			{
-				MaterialCompilationOutput.SubstrateMaterialCompilationOutput.SubstrateMaterialType = SUBSTRATE_MATERIAL_TYPE_SINGLE;
+				MaterialCompilationOutput.SubstrateMaterialCompilationOutput.SubstrateMaterialType = SUBSTRATE_MATERIAL_TYPE_COMPLEX;
 				MaterialCompilationOutput.SubstrateMaterialCompilationOutput.SubstrateMaterialBsdfFeatures = ESubstrateBsdfFeature::Eye;
 			}
 			else if (MaterialShadingModels.HasShadingModel(MSM_Hair))
 			{
-				MaterialCompilationOutput.SubstrateMaterialCompilationOutput.SubstrateMaterialType = SUBSTRATE_MATERIAL_TYPE_SINGLE;
+				MaterialCompilationOutput.SubstrateMaterialCompilationOutput.SubstrateMaterialType = SUBSTRATE_MATERIAL_TYPE_COMPLEX;
 				MaterialCompilationOutput.SubstrateMaterialCompilationOutput.SubstrateMaterialBsdfFeatures = ESubstrateBsdfFeature::Hair;
 			}
 
 			if (MaterialShadingModels.HasAnyShadingModel({ MSM_Subsurface,MSM_SubsurfaceProfile,MSM_TwoSidedFoliage,MSM_Eye,MSM_PreintegratedSkin }))
 			{
-				MaterialCompilationOutput.SubstrateMaterialCompilationOutput.SubstrateMaterialType = SUBSTRATE_MATERIAL_TYPE_SINGLE;
+				// Override complexity only if this is Simple or Single
+				MaterialCompilationOutput.SubstrateMaterialCompilationOutput.SubstrateMaterialType = FMath::Max(uint8(SUBSTRATE_MATERIAL_TYPE_SINGLE), MaterialCompilationOutput.SubstrateMaterialCompilationOutput.SubstrateMaterialType);
 				MaterialCompilationOutput.SubstrateMaterialCompilationOutput.SubstrateMaterialBsdfFeatures |= ESubstrateBsdfFeature::SSS;
 			}
 		}
