FSoundWaveLoadingBehaviorUtil::CollectAllRelevantSoundClassAssetData now does disk-only asset reqiestry queries to ensure deterministic results during cook

diff --git a/Engine/Source/Runtime/Engine/Classes/Sound/SoundWaveLoadingBehavior.cpp b/Engine/Source/Runtime/Engine/Classes/Sound/SoundWaveLoadingBehavior.cpp
--- a/Engine/Source/Runtime/Engine/Classes/Sound/SoundWaveLoadingBehavior.cpp
+++ b/Engine/Source/Runtime/Engine/Classes/Sound/SoundWaveLoadingBehavior.cpp
@@ -157,6 +157,7 @@ private:
 		//Filter.ClassPaths.Add(USoundCue::StaticClass()->GetClassPathName());
 		//Filter.bRecursiveClasses = true;
 		Filter.PackageNames = SoundWaveReferencerNames;
+		Filter.bIncludeOnlyOnDiskAssets = IsRunningCookCommandlet();	// Enumerate only on-disk packages when cooking to avoid indeterministic results
 		TArray<FAssetData> ReferencingSoundCueAssetDataArray;
 		if (!AssetRegistry.GetAssets(Filter, ReferencingSoundCueAssetDataArray))
 		{
@@ -197,6 +198,7 @@ private:
 			FARFilter ClassFilter;
 			ClassFilter.ClassPaths.Add(USoundClass::StaticClass()->GetClassPathName());
 			ClassFilter.PackageNames = SoundCueReferences;
+			ClassFilter.bIncludeOnlyOnDiskAssets = IsRunningCookCommandlet();	// only enumerate disk-only assets during cooking
 			TArray<FAssetData> ReferencedSoundClasses;
 			if (!AssetRegistry.GetAssets(ClassFilter, ReferencedSoundClasses))
 			{
