Clean up deprecated code in skeleton in regards to mesh linkup caching

diff --git a/Engine/Source/Runtime/Engine/Classes/Animation/Skeleton.h b/Engine/Source/Runtime/Engine/Classes/Animation/Skeleton.h
--- a/Engine/Source/Runtime/Engine/Classes/Animation/Skeleton.h
+++ b/Engine/Source/Runtime/Engine/Classes/Animation/Skeleton.h
@@ -634,18 +634,7 @@ public:
 
 	typedef TArray<FBoneNode> FBoneTreeType;
 
-	
-	UE_DEPRECATED(5.4, "Public access to this member variable is deprecated.")
-	FCriticalSection LinkupCacheLock;
-
-	/** Non-serialised cache of linkups between different skeletal meshes and this Skeleton. */
-	UE_DEPRECATED(5.4, "Public access to this member variable is deprecated.")
-	TArray<struct FSkeletonToMeshLinkup> LinkupCache;
-
 private:
-	UE_DEPRECATED(5.4, "Please, use SkinnedAssetLinkupCache.")
-	TMap<TWeakObjectPtr<USkinnedAsset>, int32> SkinnedAsset2LinkupCache;
-	
 	//Use this Lock everytime you change or access SkinnedAssetLinkupCache member.
 	FTransactionallySafeRWLock SkinnedAssetLinkupCacheLock;
 
@@ -653,9 +642,6 @@ private:
 	TMap<TObjectKey<USkinnedAsset>, TUniquePtr<FSkeletonToMeshLinkup>> SkinnedAssetLinkupCache;
 
 public:
-	UE_DEPRECATED(5.1, "Public access to this member variable is deprecated.")
-	TMap<TWeakObjectPtr<USkeletalMesh>, int32> SkelMesh2LinkupCache;
-
 	/** A cached soft object pointer of this skeleton. This is done for performance reasons when searching for compatible skeletons when using IsCompatible(Skeleton). */
 	TSoftObjectPtr<USkeleton> CachedSoftObjectPtr;
 
@@ -796,9 +782,6 @@ public:
 	/** Clears all cache data **/
 	ENGINE_API void ClearCacheData();
 
-	UE_DEPRECATED(5.4, "Please use FindOrAddMeshLinkupData.")
-	ENGINE_API int32 GetMeshLinkupIndex(const USkinnedAsset* InSkinnedAsset);
-
 	/**
 	 * Find a mesh linkup table (mapping of skeleton bone tree indices to refpose indices) for a particular SkinnedAsset
 	 * If one does not already exist, create it now.
@@ -964,15 +947,6 @@ protected:
 	 */
 	bool DoesParentChainMatch(int32 StartBoneTreeIndex, const USkinnedAsset* InSkinnedAsset) const;
 
-	/**
-	 * Build Look up between SkinnedAsset to BoneTree
-	 *
-	 * @param	InSkinnedAsset	: SkinnedAsset to build look up for
-	 * @return	Index of LinkupCache that this SkelMesh is linked to
-	 */
-	UE_DEPRECATED(5.4, "Please use BuildLinkup with USkinnedAsset and FSkeletonToMeshLinkup parameters.")
-	int32 BuildLinkup(const USkinnedAsset* InSkinnedAsset);
-
 	/**
 	 * Build Look up between SkinnedAsset to BoneTree
 	 *
diff --git a/Engine/Source/Runtime/Engine/Private/Animation/Skeleton.cpp b/Engine/Source/Runtime/Engine/Private/Animation/Skeleton.cpp
--- a/Engine/Source/Runtime/Engine/Private/Animation/Skeleton.cpp
+++ b/Engine/Source/Runtime/Engine/Private/Animation/Skeleton.cpp
@@ -752,13 +752,6 @@ void USkeleton::ClearCacheData()
 		UE::TWriteScopeLock Lock(SkinnedAssetLinkupCacheLock);
 		SkinnedAssetLinkupCache.Empty();
 	}
-
-	PRAGMA_DISABLE_DEPRECATION_WARNINGS
-	FScopeLock ScopeLock(&LinkupCacheLock);
-	LinkupCache.Empty();
-	SkelMesh2LinkupCache.Empty();
-	SkinnedAsset2LinkupCache.Empty();
-	PRAGMA_ENABLE_DEPRECATION_WARNINGS
 }
 
 const FSkeletonToMeshLinkup& USkeleton::FindOrAddMeshLinkupData(const USkinnedAsset* InSkinnedAsset)
@@ -796,80 +789,12 @@ const FSkeletonToMeshLinkup& USkeleton::AddMeshLinkupData(const USkinnedAsset* I
 	}
 }
 
-int32 USkeleton::GetMeshLinkupIndex(const USkinnedAsset* InSkinnedAsset)
-{
-	int32* IndexPtr = nullptr;
-	{
-		PRAGMA_DISABLE_DEPRECATION_WARNINGS
-		FScopeLock ScopeLock(&LinkupCacheLock);
-		IndexPtr = SkinnedAsset2LinkupCache.Find(MakeWeakObjectPtr(const_cast<USkinnedAsset*>(InSkinnedAsset)));
-		PRAGMA_ENABLE_DEPRECATION_WARNINGS
-	}
-	int32 LinkupIndex = INDEX_NONE;
-
-	if ( IndexPtr == NULL )
-	{
-		PRAGMA_DISABLE_DEPRECATION_WARNINGS
-		LinkupIndex = BuildLinkup(InSkinnedAsset);
-		PRAGMA_ENABLE_DEPRECATION_WARNINGS
-	}
-	else
-	{
-		LinkupIndex = *IndexPtr;
-	}
-
-	{
-		PRAGMA_DISABLE_DEPRECATION_WARNINGS
-		FScopeLock ScopeLock(&LinkupCacheLock);
-		// make sure it's not out of range
-		check(LinkupIndex < LinkupCache.Num());
-		PRAGMA_ENABLE_DEPRECATION_WARNINGS
-	}
-
-	return LinkupIndex;
-}
-
 void USkeleton::RemoveLinkup(const USkinnedAsset* InSkinnedAsset)
 {
 	{
 		UE::TWriteScopeLock Lock(SkinnedAssetLinkupCacheLock);
 		SkinnedAssetLinkupCache.Remove(InSkinnedAsset);
 	}
-
-	PRAGMA_DISABLE_DEPRECATION_WARNINGS
-	FScopeLock ScopeLock(&LinkupCacheLock);
-	if (const USkeletalMesh* const InSkelMesh = Cast<const USkeletalMesh>(InSkinnedAsset))
-	{
-		SkelMesh2LinkupCache.Remove(MakeWeakObjectPtr(const_cast<USkeletalMesh*>(InSkelMesh)));
-	}
-	SkinnedAsset2LinkupCache.Remove(MakeWeakObjectPtr(const_cast<USkinnedAsset*>(InSkinnedAsset)));
-	PRAGMA_ENABLE_DEPRECATION_WARNINGS
-}
-
-int32 USkeleton::BuildLinkup(const USkinnedAsset* InSkinnedAsset)
-{
-	// @todoanim : need to refresh NULL SkeletalMeshes from Cache
-	// since now they're autoweak pointer, they will go away if not used
-	// so whenever map transition happens, this links will need to clear up
-	FSkeletonToMeshLinkup NewMeshLinkup;
-
-	BuildLinkupData(InSkinnedAsset, NewMeshLinkup);
-
-	int32 NewIndex = INDEX_NONE;
-	//LinkupCache lock scope
-	{
-		PRAGMA_DISABLE_DEPRECATION_WARNINGS
-		FScopeLock ScopeLock(&LinkupCacheLock);
-		NewIndex = LinkupCache.Add(NewMeshLinkup);
-		check(NewIndex != INDEX_NONE);
-			if (const USkeletalMesh* const InSkelMesh = Cast<const USkeletalMesh>(InSkinnedAsset))
-			{
-				SkelMesh2LinkupCache.Add(MakeWeakObjectPtr(const_cast<USkeletalMesh*>(InSkelMesh)), NewIndex);
-			}
-		SkinnedAsset2LinkupCache.Add(MakeWeakObjectPtr(const_cast<USkinnedAsset*>(InSkinnedAsset)), NewIndex);
-		PRAGMA_ENABLE_DEPRECATION_WARNINGS
-	}
-	return NewIndex;
 }
 
 void USkeleton::BuildLinkupData(const USkinnedAsset* InSkinnedAsset, FSkeletonToMeshLinkup& NewMeshLinkup)
