Catch situation where mesh asset in InitializeAnimScriptInstance is invalid.  Also deal with possible parallel anim eval tasks being in flight.
#DisabledValidations CheckUnacceptableWords

diff --git a/Engine/Source/Runtime/Engine/Private/Components/SkeletalMeshComponent.cpp b/Engine/Source/Runtime/Engine/Private/Components/SkeletalMeshComponent.cpp
--- a/Engine/Source/Runtime/Engine/Private/Components/SkeletalMeshComponent.cpp
+++ b/Engine/Source/Runtime/Engine/Private/Components/SkeletalMeshComponent.cpp
@@ -1124,10 +1124,15 @@ bool USkeletalMeshComponent::InitializeAnimScriptInstance(bool bForceReinit, boo
 	bool bInitializedMainInstance = false;
 	bool bInitializedPostInstance = false;
 
-	if (IsRegistered())
+	USkeletalMesh* SkelMesh = GetSkeletalMeshAsset();
+	if (SkelMesh && IsRegistered())
 	{
-		USkeletalMesh* SkelMesh = GetSkeletalMeshAsset();
-		check(SkelMesh);
+		// We may be doing parallel evaluation on the current anim instance
+		// Calling this here with true will block this init till that thread completes
+		// and it is safe to continue
+		const bool bBlockOnTask = true; // wait on evaluation task so it is safe to continue with Init
+		const bool bPerformPostAnimEvaluation = true; // That will swap buffer back to ComponentTransform, and finish evaluate. This is required - otherwise, we won't have a buffer.
+		HandleExistingParallelEvaluationTask(bBlockOnTask, bPerformPostAnimEvaluation);
 
 		if (NeedToSpawnAnimScriptInstance())
 		{
