Remove ValidateModel call in UAnimSequenceBase::IsDataModelValid - causing a check in ValidateModel defeats the purpose of being able to call ::IsDataModelValid

diff --git a/Engine/Source/Runtime/Engine/Classes/Animation/AnimSequenceBase.h b/Engine/Source/Runtime/Engine/Classes/Animation/AnimSequenceBase.h
--- a/Engine/Source/Runtime/Engine/Classes/Animation/AnimSequenceBase.h
+++ b/Engine/Source/Runtime/Engine/Classes/Animation/AnimSequenceBase.h
@@ -288,7 +288,6 @@ public:
 	{
 		if(ShouldDataModelBeValid())
 		{
-			ValidateModel();
 			return DataModelInterface != nullptr;
 		}
 
