Kill process tree

diff --git a/Engine/Source/Runtime/Engine/Private/ShaderCompiler/ShaderCompilerThreadRunnable.cpp b/Engine/Source/Runtime/Engine/Private/ShaderCompiler/ShaderCompilerThreadRunnable.cpp
--- a/Engine/Source/Runtime/Engine/Private/ShaderCompiler/ShaderCompilerThreadRunnable.cpp
+++ b/Engine/Source/Runtime/Engine/Private/ShaderCompiler/ShaderCompilerThreadRunnable.cpp
@@ -130,7 +130,7 @@ struct FShaderCompileWorkerInfo
 	{
 		if (WorkerProcess.IsValid())
 		{
-			FPlatformProcess::TerminateProc(WorkerProcess);
+			FPlatformProcess::TerminateProc(WorkerProcess, true);
 			if (!bAsynchronous)
 			{
 				while (FPlatformProcess::IsProcRunning(WorkerProcess))
