Fix missing subsurface profile with Substrate when using the hidden conversion.

diff --git a/Engine/Source/Runtime/Engine/Private/Materials/MaterialInterface.cpp b/Engine/Source/Runtime/Engine/Private/Materials/MaterialInterface.cpp
--- a/Engine/Source/Runtime/Engine/Private/Materials/MaterialInterface.cpp
+++ b/Engine/Source/Runtime/Engine/Private/Materials/MaterialInterface.cpp
@@ -1670,11 +1670,28 @@ void UMaterialInterface::UpdateMaterialRenderProxy(FMaterialRenderProxy& Proxy)
 	// no 0 pointer
 	check(&Proxy);
 
+	UMaterial* Material = GetMaterial();
+
+	// SSS profiles:
+	// * When using the legacy path, or using Substrate with the hidden conversion with no conversion, SSS profile is bound with the old parameter name. 
+	//   This requires SubsurfaceProfileRT() to be called. 
+	// * When using Substrate and the material has been converted, then SSS profiles are bound with the new parameter names 
+	bool bIsSubstrateMaterial = false;
+	if (Material)
+	{
+		const EShaderPlatform ShaderPlatform = GetMaxShaderPlatformChecked();
+		if (FMaterialResource* MaterialResource = Material->GetMaterialResource(ShaderPlatform))
+		{
+			bIsSubstrateMaterial = MaterialResource->IsSubstrateMaterial();
+		}
+	}
+	const bool bUseLegacySSSProfileSetup = !Substrate::IsSubstrateEnabled() || (Substrate::IsHiddenMaterialAssetConversionEnabled() && !bIsSubstrateMaterial);
+
 	FMaterialShadingModelField MaterialShadingModels = GetShadingModels();
 
 	// For better performance we only update SubsurfaceProfileRT if the feature is used and substrate is not enabled
 	// When Substrate is enabled, this is ONLY used as an override for Subsurface Profile on material instance (override all Subsurface Profiles at once for now)
-	if (UseSubsurfaceProfile(MaterialShadingModels) && !Substrate::IsSubstrateEnabled())
+	if (UseSubsurfaceProfile(MaterialShadingModels) && bUseLegacySSSProfileSetup)
 	{
 		USubsurfaceProfile* LocalSubsurfaceProfile = GetSubsurfaceProfile_Internal();
 		
@@ -1697,7 +1714,6 @@ void UMaterialInterface::UpdateMaterialRenderProxy(FMaterialRenderProxy& Proxy)
 		});
 	}
 
-	UMaterial* Material = GetMaterial();
 	if (Material && Material->IsPostProcessMaterial())
 	{
 		struct FEntry
@@ -1794,6 +1810,7 @@ void UMaterialInterface::UpdateMaterialRenderProxy(FMaterialRenderProxy& Proxy)
 		}
 
 		// Subsurface profiles
+		if (!bUseLegacySSSProfileSetup)
 		{
 			struct FEntry
 			{
