Don't broadcast number of bookmarks or bookmark class changed events from within world setting object initialization.

diff --git a/Engine/Source/Runtime/Engine/Private/WorldSettings.cpp b/Engine/Source/Runtime/Engine/Private/WorldSettings.cpp
--- a/Engine/Source/Runtime/Engine/Private/WorldSettings.cpp
+++ b/Engine/Source/Runtime/Engine/Private/WorldSettings.cpp
@@ -1141,7 +1141,10 @@ void AWorldSettings::UpdateNumberOfBookmarks()
 		AdjustNumberOfBookmarks();
 
 #if WITH_EDITOR
-		OnNumberOfBoomarksChanged.Broadcast(this);
+		if (!HasAnyFlags(RF_NeedInitialization))
+		{
+			OnNumberOfBoomarksChanged.Broadcast(this);
+		}
 #endif
 	}
 }
@@ -1184,7 +1187,10 @@ void AWorldSettings::UpdateBookmarkClass()
 	{
 
 #if WITH_EDITOR
-		OnBookmarkClassChanged.Broadcast(this);
+		if (!HasAnyFlags(RF_NeedInitialization))
+		{
+			OnBookmarkClassChanged.Broadcast(this);
+		}
 #endif
 
 		// Explicitly done after OnBookmarkClassChanged, in case there's any upgrade work
