[NetPhysics] NetworkPhysicsComponent improvement
- Rename the recently added (unreleased) FNetworkPhysicsBase to FNetworkPhysicsPayload
- Rename the recently added (unreleased) INetworkPhysicsInputStateImplementationBase to INetworkPhysicsInputState_Internal
- Add GT interface and pass through received inputs for validation on the server

diff --git a/Engine/Source/Runtime/Engine/Private/PhysicsEngine/NetworkPhysicsComponent.cpp b/Engine/Source/Runtime/Engine/Private/PhysicsEngine/NetworkPhysicsComponent.cpp
--- a/Engine/Source/Runtime/Engine/Private/PhysicsEngine/NetworkPhysicsComponent.cpp
+++ b/Engine/Source/Runtime/Engine/Private/PhysicsEngine/NetworkPhysicsComponent.cpp
@@ -576,7 +576,7 @@ void UNetworkPhysicsComponent::UninitializeComponent()
 		{
 			AsyncInput->ActorComponent = nullptr;
 			AsyncInput->PhysicsObject = nullptr;
-			AsyncInput->ImplementationInterface = nullptr;
+			AsyncInput->ImplementationInterface_Internal = nullptr;
 		}
 	}
 
@@ -1231,7 +1231,6 @@ void UNetworkPhysicsComponent::ServerReceiveInputData_Implementation(const FNetw
 			ClientInputs.History->ValidateDataInHistory(ActorComponent.Get());
 		}
 
-		AsyncInput->InputData->ResetFast();
 		AsyncInput->InputData->ResizeDataHistory(ClientInputs.History->GetHistorySize(), EAllowShrinking::No);
 		AsyncInput->InputData->SetRecordDataIncremental(true); // Only record data that is newer than already cached data
 		ClientInputs.History->CopyAllData(*AsyncInput->InputData.Get(), /*bIncludeUnimportant*/ true, /*bIncludeImportant*/ true);
@@ -1357,12 +1356,17 @@ void UNetworkPhysicsComponent::ServerReceiveInputCollection_Implementation(const
 			AsyncInput->InputData = InputHelper->CreateUniqueRewindHistory(ClientInputCollection.DataArray.Num());
 		}
 
-		AsyncInput->InputData->ResetFast();
 		AsyncInput->InputData->ResizeDataHistory(ClientInputCollection.DataArray.Num(), EAllowShrinking::No);
 		AsyncInput->InputData->SetRecordDataIncremental(true); // Only record data that is newer than already cached data
 
 		InputHelper->CopyData(ClientInputCollection, AsyncInput->InputData.Get());
 
+		// Validate Inputs
+		if (ImplementationInterface_External)
+		{
+			InputHelper->ValidateData(AsyncInput->InputData.Get(), *ImplementationInterface_External);
+		}
+
 		// Send received inputs to remote clients
 		ReplicatedRemoteInputCollection.DataArray.SetNum(InputsToNetwork_Simulated);
 		InputHelper->CopyIncrementalData(AsyncInput->InputData.Get(), ReplicatedRemoteInputCollection);
@@ -1602,9 +1606,9 @@ void UNetworkPhysicsComponent::UpdateAsyncComponent(const bool bFullUpdate)
 				{
 					AsyncInput->ActorComponent = ActorComponent;
 				}
-				if (ImplementationInterface)
+				if (ImplementationInterface_Internal)
 				{
-					AsyncInput->ImplementationInterface = ImplementationInterface;
+					AsyncInput->ImplementationInterface_Internal = ImplementationInterface_Internal;
 				}
 			}
 		}
@@ -1621,9 +1625,9 @@ void UNetworkPhysicsComponent::CreateAsyncDataHistory()
 			{
 				AsyncInput->ActorComponent = ActorComponent;
 			}
-			if (ImplementationInterface)
+			if (ImplementationInterface_Internal)
 			{
-				AsyncInput->ImplementationInterface = ImplementationInterface;
+				AsyncInput->ImplementationInterface_Internal = ImplementationInterface_Internal;
 			}
 
 			if (InputHelper)
@@ -1698,7 +1702,7 @@ FAsyncNetworkPhysicsComponent::FAsyncNetworkPhysicsComponent() : TSimCallbackObj
 	, bIsUsingLegacyData(false)
 	, SettingsComponent(nullptr)
 	, ActorComponent(nullptr)
-	, ImplementationInterface(nullptr)
+	, ImplementationInterface_Internal(nullptr)
 	, PhysicsObject(nullptr)
 	, bCompareStateToTriggerRewind(false)
 	, bCompareStateToTriggerRewindIncludeSimProxies(false)
@@ -1799,9 +1803,9 @@ void FAsyncNetworkPhysicsComponent::ConsumeAsyncInput(const int32 PhysicsStep)
 			{
 				ActorComponent = *AsyncInput->ActorComponent;
 			}
-			if (AsyncInput->ImplementationInterface.IsSet())
+			if (AsyncInput->ImplementationInterface_Internal.IsSet())
 			{
-				ImplementationInterface = *AsyncInput->ImplementationInterface;
+				ImplementationInterface_Internal = *AsyncInput->ImplementationInterface_Internal;
 			}
 			if (AsyncInput->PhysicsObject.IsSet())
 			{
@@ -1897,20 +1901,20 @@ void FAsyncNetworkPhysicsComponent::ConsumeAsyncInput(const int32 PhysicsStep)
 			if (AsyncInput->InputData && AsyncInput->InputData->HasDataInHistory())
 			{
 				// Validate data in the received inputs on the server
-				if (!ComponentSettings.GetValidateDataOnGameThread() && bIsServer)
+				if (bIsServer)
 				{
 					if (bIsUsingLegacyData)
 					{
-						if (ensure(ActorComponent.IsValid() && ActorComponent.Get()->IsBeingDestroyed() == false))
+						if (!ComponentSettings.GetValidateDataOnGameThread() && ensure(ActorComponent.IsValid() && ActorComponent.Get()->IsBeingDestroyed() == false))
 						{
 							AsyncInput->InputData->ValidateDataInHistory(ActorComponent.Get());
 						}
 					}
 					else
 					{
-						if (InputHelper && ensure(ImplementationInterface))
+						if (InputHelper && ensure(ImplementationInterface_Internal))
 						{
-							InputHelper->ValidateDataInHistory(AsyncInput->InputData.Get(), *ImplementationInterface);
+							InputHelper->ValidateData(AsyncInput->InputData.Get(), *ImplementationInterface_Internal);
 						}
 					}
 				}
@@ -1994,7 +1998,7 @@ void FAsyncNetworkPhysicsComponent::OnPreProcessInputs_Internal(const int32 Phys
 		const bool bApplySimProxyState = ComponentSettings.GetApplySimProxyStateAtRuntime() && !bIsServer && !IsLocallyControlled();
 		if ((bApplySimProxyState || bIsSolverResim) && StateHistory && StateData)
 		{
-			FNetworkPhysicsBase* PhysicsData = StateData.Get();
+			FNetworkPhysicsPayload* PhysicsData = StateData.Get();
 			PhysicsData->LocalFrame = PhysicsStep;
 			const bool bExactFrame = PhysicsReplicationCVars::ResimulationCVars::bAllowRewindToClosestState ? !bIsSolverReset : true;
 			if (StateHistory->ExtractData(PhysicsStep, bIsSolverReset, PhysicsData, (bExactFrame && bIsSolverResim)) && PhysicsData->bReceivedData)
@@ -2009,9 +2013,9 @@ void FAsyncNetworkPhysicsComponent::OnPreProcessInputs_Internal(const int32 Phys
 				}
 				else
 				{
-					if (ensure(ImplementationInterface))
+					if (ensure(ImplementationInterface_Internal))
 					{
-						ImplementationInterface->ApplyState(*PhysicsData);
+						ImplementationInterface_Internal->ApplyState(*PhysicsData);
 					}
 				}
 #if DEBUG_NETWORK_PHYSICS
@@ -2033,7 +2037,7 @@ void FAsyncNetworkPhysicsComponent::OnPreProcessInputs_Internal(const int32 Phys
 		const bool bApplyServerInput = bIsServer && !IsLocallyControlled();
 		if ((bApplyServerInput || bApplySimProxyInput || bIsSolverResim) && InputHistory && InputData)
 		{
-			FNetworkPhysicsBase* PhysicsData = InputData.Get();
+			FNetworkPhysicsPayload* PhysicsData = InputData.Get();
 			int32 NextExpectedLocalFrame = PhysicsData->LocalFrame + 1;
 
 			// There are important inputs earlier than upcoming input to apply
@@ -2055,9 +2059,9 @@ void FAsyncNetworkPhysicsComponent::OnPreProcessInputs_Internal(const int32 Phys
 /* // Importance is not implemented into the new network flow (yet?)
 					else
 					{
-						if (InputHelper && ensure(ImplementationInterface))
+						if (InputHelper && ensure(ImplementationInterface_Internal))
 						{
-							InputHelper->ApplyDataRange(InputHistory.Get(), NewImportantInputFrame, NextExpectedLocalFrame - 1, *ImplementationInterface);
+							InputHelper->ApplyDataRange(InputHistory.Get(), NewImportantInputFrame, NextExpectedLocalFrame - 1, *ImplementationInterface_Internal);
 						}
 					}
 */
@@ -2127,9 +2131,9 @@ void FAsyncNetworkPhysicsComponent::OnPreProcessInputs_Internal(const int32 Phys
 								}
 								else
 								{
-									if (ensure(ImplementationInterface))
+									if (ensure(ImplementationInterface_Internal))
 									{
-										ImplementationInterface->ApplyInput(*PhysicsData);
+										ImplementationInterface_Internal->ApplyInput(*PhysicsData);
 									}
 								}
 							}
@@ -2164,9 +2168,9 @@ void FAsyncNetworkPhysicsComponent::OnPreProcessInputs_Internal(const int32 Phys
 				}
 				else
 				{
-					if (ensure(ImplementationInterface))
+					if (ensure(ImplementationInterface_Internal))
 					{
-						ImplementationInterface->ApplyInput(*PhysicsData);
+						ImplementationInterface_Internal->ApplyInput(*PhysicsData);
 					}
 				}
 
@@ -2214,7 +2218,7 @@ void FAsyncNetworkPhysicsComponent::OnPostProcessInputs_Internal(const int32 Phy
 	if (bShouldCacheInputHistory && (InputData != nullptr))
 	{
 		// Prepare to gather input data
-		FNetworkPhysicsBase* PhysicsData = InputData.Get();
+		FNetworkPhysicsPayload* PhysicsData = InputData.Get();
 		PhysicsData->PrepareFrame(PhysicsStep, bIsServer, GetNetworkPhysicsTickOffset());
 
 		if (bIsUsingLegacyData)
@@ -2228,10 +2232,10 @@ void FAsyncNetworkPhysicsComponent::OnPostProcessInputs_Internal(const int32 Phy
 		}
 		else
 		{
-			if (ensure(ImplementationInterface))
+			if (ensure(ImplementationInterface_Internal))
 			{
 				// Gather input data from implementation
-				ImplementationInterface->BuildInput(*PhysicsData);
+				ImplementationInterface_Internal->BuildInput(*PhysicsData);
 			}
 		}
 
@@ -2253,7 +2257,7 @@ void FAsyncNetworkPhysicsComponent::OnPostProcessInputs_Internal(const int32 Phy
 	if (StateHistory && StateData && bShouldCacheStateHistory)
 	{
 		// Prepare to gather state data
-		FNetworkPhysicsBase* PhysicsData = StateData.Get();
+		FNetworkPhysicsPayload* PhysicsData = StateData.Get();
 		PhysicsData->PrepareFrame(PhysicsStep, bIsServer, GetNetworkPhysicsTickOffset());
 
 		if (bIsUsingLegacyData)
@@ -2267,10 +2271,10 @@ void FAsyncNetworkPhysicsComponent::OnPostProcessInputs_Internal(const int32 Phy
 		}
 		else
 		{
-			if (ensure(ImplementationInterface))
+			if (ensure(ImplementationInterface_Internal))
 			{
 				// Gather state data from implementation
-				ImplementationInterface->BuildState(*PhysicsData);
+				ImplementationInterface_Internal->BuildState(*PhysicsData);
 			}
 		}
 
@@ -2460,7 +2464,7 @@ void FAsyncNetworkPhysicsComponent::TriggerResimulation(int32 ResimFrame)
 	}
 }
 
-const float FAsyncNetworkPhysicsComponent::GetCurrentInputDecay(const FNetworkPhysicsBase* PhysicsData)
+const float FAsyncNetworkPhysicsComponent::GetCurrentInputDecay(const FNetworkPhysicsPayload* PhysicsData)
 {
 	if (!PhysicsData)
 	{
diff --git a/Engine/Source/Runtime/Engine/Public/Physics/NetworkPhysicsComponent.h b/Engine/Source/Runtime/Engine/Public/Physics/NetworkPhysicsComponent.h
--- a/Engine/Source/Runtime/Engine/Public/Physics/NetworkPhysicsComponent.h
+++ b/Engine/Source/Runtime/Engine/Public/Physics/NetworkPhysicsComponent.h
@@ -740,11 +740,11 @@ public:
 /** START Iris Compliant Data */
 
 USTRUCT()
-struct FNetworkPhysicsBase
+struct FNetworkPhysicsPayload
 {
 	GENERATED_BODY()
 
-	virtual ~FNetworkPhysicsBase() { }
+	virtual ~FNetworkPhysicsPayload() { }
 
 	UPROPERTY()
 	int32 ServerFrame = 0;
@@ -757,14 +757,14 @@ struct FNetworkPhysicsBase
 	* @param LerpAlpha is the 0.0 - 1.0 value of where *this* data is between MinData and MaxData calculated as so: float LerpAlpha = (LocalFrame - MinData.LocalFrame) / (MaxData.LocalFrame - MinData.LocalFrame)
  	* EXAMPLE: We have input data for frame 1 and 4 and we need to interpolate data for frame 2 and 3 based on frame 1 as MinData and frame 4 as MaxData, for frame 2 LerpAlpha will be 0.33 and for frame 3 LerpAlpha will be 0.66
 	*/
-	virtual void InterpolateData(const FNetworkPhysicsBase& MinData, const FNetworkPhysicsBase& MaxData, float LerpAlpha) {}
+	virtual void InterpolateData(const FNetworkPhysicsPayload& MinData, const FNetworkPhysicsPayload& MaxData, float LerpAlpha) {}
 
 	/** Define how to merge data together
 	* @param FromData is data from a previous frame that is getting merged into the current data.
 	* EXAMPLE: Simulated proxies might receive two inputs at the same time after having used the same input twice, to not miss any important inputs we need to take both inputs into account
 	* and to not get behind in simulation we need to apply them both at the same simulation tick meaning we merge the two new inputs to one input.
 	*/
-	virtual void MergeData(const FNetworkPhysicsBase& FromData) {}
+	virtual void MergeData(const FNetworkPhysicsPayload& FromData) {}
 
 	/** Use to decay desired data during resimulation if data is forward predicted.
 	* @param DecayAmount = Total amount of decay as a multiplier. 10% decay = 0.1.
@@ -778,7 +778,7 @@ struct FNetworkPhysicsBase
 	* NOTE: To use this function, CVars np2.Resim.CompareStateToTriggerRewind and/or np2.Resim.CompareInputToTriggerRewind needs to be set to true
 	* or the equivalent settings overridden on the actor via UNetworkPhysicsSettingsComponent.
 	*/
-	virtual bool CompareData(const FNetworkPhysicsBase& PredictedData) const { return true; }
+	virtual bool CompareData(const FNetworkPhysicsPayload& PredictedData) const { return true; }
 
 	/** Return string with custom debug data */
 	virtual const FString DebugData() { return FString(" - DebugData() not implemented - "); } // TODO, Deprecate non-const version
@@ -804,7 +804,7 @@ struct FNetworkPhysicsBase
 
 	/** Temporary virtual function for backwards compatibility with FNetworkPhysicsData.Use the InterpolateData function that passes the LerpAlpha instead */
 	UE_DEPRECATED(5.7, "Temporary virtual function for backwards compatibility with FNetworkPhysicsData. Use the InterpolateData function that passes the LerpAlpha instead")
-	virtual void DoInterpolateData(const FNetworkPhysicsBase& MinData, const FNetworkPhysicsBase& MaxData)
+	virtual void DoInterpolateData(const FNetworkPhysicsPayload& MinData, const FNetworkPhysicsPayload& MaxData)
 	{
 		const float LerpAlpha = static_cast<float>(LocalFrame - MinData.LocalFrame) / static_cast<float>(MaxData.LocalFrame - MinData.LocalFrame);
 		InterpolateData(MinData, MaxData, LerpAlpha);
@@ -817,14 +817,14 @@ struct FNetworkPhysicsDataCollection
 	GENERATED_BODY()
 
 	UPROPERTY()
-	TArray<TInstancedStruct<FNetworkPhysicsBase>> DataArray;
+	TArray<TInstancedStruct<FNetworkPhysicsPayload>> DataArray;
 
-	// TEMP, handle this via FNetworkPhysicsBaseNetSerializer (but without offset, just set LocalFrame = ServerFrame)
+	// TEMP, handle this via FNetworkPhysicsPayloadNetSerializer (but without offset, just set LocalFrame = ServerFrame)
 	void UpdateLocalFrameFromServerFrame(const int32 ClientFrameOffset) const
 	{
-		for (const TInstancedStruct<FNetworkPhysicsBase>& DataInstance : DataArray)
+		for (const TInstancedStruct<FNetworkPhysicsPayload>& DataInstance : DataArray)
 		{
-			const FNetworkPhysicsBase& Data = DataInstance.Get();
+			const FNetworkPhysicsPayload& Data = DataInstance.Get();
 			Data.LocalFrame = Data.ServerFrame - ClientFrameOffset;
 		}
 	}
@@ -838,8 +838,8 @@ struct FNetworkPhysicsDataCollection
 		{
 			for (int32 FrameIndex = 0; FrameIndex < DataArray.Num(); ++FrameIndex)
 			{
-				const TInstancedStruct<FNetworkPhysicsBase>& DataInstance = DataArray[FrameIndex];
-				const FNetworkPhysicsBase* Data = DataInstance.GetPtr<FNetworkPhysicsBase>();
+				const TInstancedStruct<FNetworkPhysicsPayload>& DataInstance = DataArray[FrameIndex];
+				const FNetworkPhysicsPayload* Data = DataInstance.GetPtr<FNetworkPhysicsPayload>();
 
 				if (Data)
 				{
@@ -862,28 +862,48 @@ struct FNetworkPhysicsDataCollection
 
 };
 
+// Game Thread Input and State Interface API
+class INetworkPhysicsInputState_External
+{
+public:
+	virtual void ValidateInput(FNetworkPhysicsPayload& Input) const = 0;
+};
+
+template<typename InputType, typename StateType>
+class TNetworkPhysicsInputState_External : public INetworkPhysicsInputState_External
+{
+	// Wrapper for non-templated base
+	void ValidateInput(FNetworkPhysicsPayload& Input) const override { ValidateInput_External(static_cast<InputType&>(Input)); }
+
+public:
+	/** Validate data received on the server from clients
+	* EXAMPLE: Validate incoming inputs from clients on the server and correct any invalid input commands.
+	* NOTE: Changes to the data in this callback will be sent from server to clients. */
+	virtual void ValidateInput_External(InputType& Input) const = 0;
+};
+
 // Physics Thread Input and State Interface API
-class INetworkPhysicsInputStateImplementationBase
+class INetworkPhysicsInputState_Internal
 {
 public:
-	virtual void BuildInput(FNetworkPhysicsBase& Input) const = 0;
-	virtual void ValidateInput(FNetworkPhysicsBase& Input) const = 0;
-	virtual void ApplyInput(const FNetworkPhysicsBase& Input) = 0;
+	virtual void BuildInput(FNetworkPhysicsPayload& Input) const = 0;
+	virtual void ValidateInput(FNetworkPhysicsPayload& Input) const = 0;
+	virtual void ApplyInput(const FNetworkPhysicsPayload& Input) = 0;
 
-	virtual void BuildState(FNetworkPhysicsBase& State) const = 0;
-	virtual void ApplyState(const FNetworkPhysicsBase& State) = 0;
+	virtual void BuildState(FNetworkPhysicsPayload& State) const = 0;
+	virtual void ApplyState(const FNetworkPhysicsPayload& State) = 0;
 };
 
 template<typename InputType, typename StateType>
-class TNetworkPhysicsInputStateImplementation : public INetworkPhysicsInputStateImplementationBase
+class TNetworkPhysicsInputState_Internal : public INetworkPhysicsInputState_Internal
 {
 	// Wrapper for non-templated base
-	void BuildInput(FNetworkPhysicsBase& Input) const override { BuildInput_Internal(static_cast<InputType&>(Input)); }
-	void ValidateInput(FNetworkPhysicsBase& Input) const override { ValidateInput_Internal(static_cast<InputType&>(Input)); }
-	void ApplyInput(const FNetworkPhysicsBase& Input) override { ApplyInput_Internal(static_cast<const InputType&>(Input)); }
+	void BuildInput(FNetworkPhysicsPayload& Input) const override { BuildInput_Internal(static_cast<InputType&>(Input)); }
+	void ValidateInput(FNetworkPhysicsPayload& Input) const override { ValidateInput_Internal(static_cast<InputType&>(Input)); }
+	void ApplyInput(const FNetworkPhysicsPayload& Input) override { ApplyInput_Internal(static_cast<const InputType&>(Input)); }
 
-	void BuildState(FNetworkPhysicsBase& State) const override { BuildState_Internal(static_cast<StateType&>(State)); }
-	void ApplyState(const FNetworkPhysicsBase& State) override { ApplyState_Internal(static_cast<const StateType&>(State)); }
+	void BuildState(FNetworkPhysicsPayload& State) const override { BuildState_Internal(static_cast<StateType&>(State)); }
+	void ApplyState(const FNetworkPhysicsPayload& State) override { ApplyState_Internal(static_cast<const StateType&>(State)); }
 
 public:
 	/** Populate the input struct with current input data */
@@ -911,7 +931,7 @@ public:
  * Base network physics data that will be used by physics
  */
 USTRUCT()
-struct FNetworkPhysicsData : public FNetworkPhysicsBase
+struct FNetworkPhysicsData : public FNetworkPhysicsPayload
 {
 	GENERATED_USTRUCT_BODY()
 
@@ -1069,10 +1089,10 @@ struct FNetworkPhysicsData : public FNetworkPhysicsBase
 
 	// Base API Wrapper
 	UE_DEPRECATED(5.7, "Temporary virtual function for backwards compatibility with FNetworkPhysicsData. Use the InterpolateData function that passes the LerpAlpha instead")
-	void DoInterpolateData(const FNetworkPhysicsBase& MinData, const FNetworkPhysicsBase& MaxData) override { this->InterpolateData(static_cast<const FNetworkPhysicsData&>(MinData), static_cast<const FNetworkPhysicsData&>(MaxData)); }
-	void MergeData(const FNetworkPhysicsBase& FromData) override { this->MergeData(static_cast<const FNetworkPhysicsData&>(FromData)); }
-	bool CompareData(const FNetworkPhysicsBase& PredictedData) const override { return const_cast<FNetworkPhysicsData*>(this)->CompareData(static_cast<const FNetworkPhysicsData&>(PredictedData)); }
-	using FNetworkPhysicsBase::InterpolateData; // Silence warning for FNetworkPhysicsData::InterpolateData not overriding a function in FNetworkPhysicsBase, which is intended
+	void DoInterpolateData(const FNetworkPhysicsPayload& MinData, const FNetworkPhysicsPayload& MaxData) override { this->InterpolateData(static_cast<const FNetworkPhysicsData&>(MinData), static_cast<const FNetworkPhysicsData&>(MaxData)); }
+	void MergeData(const FNetworkPhysicsPayload& FromData) override { this->MergeData(static_cast<const FNetworkPhysicsData&>(FromData)); }
+	bool CompareData(const FNetworkPhysicsPayload& PredictedData) const override { return const_cast<FNetworkPhysicsData*>(this)->CompareData(static_cast<const FNetworkPhysicsData&>(PredictedData)); }
+	using FNetworkPhysicsPayload::InterpolateData; // Silence warning for FNetworkPhysicsData::InterpolateData not overriding a function in FNetworkPhysicsPayload, which is intended
 };
 
 /** Base for helper, to create data and data history */
@@ -1081,7 +1101,7 @@ struct FNetworkPhysicsDataHelper
 	virtual ~FNetworkPhysicsDataHelper() = default;
 	
 	virtual TUniquePtr<FNetworkPhysicsDataHelper> Clone() const = 0;
-	virtual TUniquePtr<FNetworkPhysicsBase> CreateUniqueData() const = 0;
+	virtual TUniquePtr<FNetworkPhysicsPayload> CreateUniqueData() const = 0;
 	virtual TUniquePtr<Chaos::FBaseRewindHistory> CreateUniqueRewindHistory(const int32 Size) const = 0;
 	virtual bool IsUsingLegacyData() = 0;
 
@@ -1097,11 +1117,14 @@ struct FNetworkPhysicsDataHelper
 	/** Copy data from the rewind history to the networked data collection and only copy data that has been altered since it was created (i.e. the client created an input that the server later altered) */
 	virtual bool CopyAlteredData(const Chaos::FBaseRewindHistory* From, FNetworkPhysicsDataCollection& To) = 0;
 	
-	/** Call the ValidateData callback on the implementation interface for data in the provided rewind history */
-	virtual void ValidateDataInHistory(Chaos::FBaseRewindHistory* History, INetworkPhysicsInputStateImplementationBase& Interface) = 0;
+	/** Call the ValidateData callback on the external implementation interface for data in the provided rewind history */
+	virtual void ValidateData(Chaos::FBaseRewindHistory* History, INetworkPhysicsInputState_External& Interface) = 0;
+
+	/** Call the ValidateData callback on the internal implementation interface for data in the provided rewind history */
+	virtual void ValidateData(Chaos::FBaseRewindHistory* History, INetworkPhysicsInputState_Internal& Interface) = 0;
 
 /* // Importance is not implemented into the new network flow (yet?)
-	virtual void ApplyDataRange(const Chaos::FBaseRewindHistory* History, const int32 FromFrame, const int32 ToFrame, INetworkPhysicsInputStateImplementationBase& Interface) = 0;
+	virtual void ApplyDataRange(const Chaos::FBaseRewindHistory* History, const int32 FromFrame, const int32 ToFrame, INetworkPhysicsInputState_Internal& Interface) = 0;
 */
 };
 
@@ -1111,7 +1134,7 @@ struct TNetworkPhysicsDataHelper : FNetworkPhysicsDataHelper
 {
 	~TNetworkPhysicsDataHelper() = default;
 	TUniquePtr<FNetworkPhysicsDataHelper> Clone() const override { return MakeUnique<TNetworkPhysicsDataHelper>(*this); }
-	TUniquePtr<FNetworkPhysicsBase> CreateUniqueData() const override { return MakeUnique<DataType>(); }
+	TUniquePtr<FNetworkPhysicsPayload> CreateUniqueData() const override { return MakeUnique<DataType>(); }
 	TUniquePtr<Chaos::FBaseRewindHistory> CreateUniqueRewindHistory(const int32 Size) const override { return MakeUnique<TNetRewindHistory<DataType, bLegacyData>>(Size); }
 	bool IsUsingLegacyData() { return bLegacyData; }
 
@@ -1122,7 +1145,7 @@ struct TNetworkPhysicsDataHelper : FNetworkPhysicsDataHelper
 
 		// Record at same index as taken from
 		int32 Idx = 0;
-		for (const TInstancedStruct<FNetworkPhysicsBase>& DataInstance : From.DataArray)
+		for (const TInstancedStruct<FNetworkPhysicsPayload>& DataInstance : From.DataArray)
 		{
 			const DataType& Data = DataInstance.Get<DataType>();
 			NetHistory->RecordData(Idx, &Data);
@@ -1162,7 +1185,7 @@ struct TNetworkPhysicsDataHelper : FNetworkPhysicsDataHelper
 			const DataType& FromData = HistoryArray[FromIdx];
 
 			const int32 ToIdx = FromIdx % To.DataArray.Num();
-			const FNetworkPhysicsBase* ToData = To.DataArray[ToIdx].GetPtr<FNetworkPhysicsBase>();
+			const FNetworkPhysicsPayload* ToData = To.DataArray[ToIdx].GetPtr<FNetworkPhysicsPayload>();
 
 			// Only copy the data if it's newer than the already cached data
 			if (!ToData || ToData->LocalFrame < FromData.LocalFrame)
@@ -1200,21 +1223,42 @@ struct TNetworkPhysicsDataHelper : FNetworkPhysicsDataHelper
 		return bHasCopiedData;
 	}
 
-	/** Call the ValidateData callback on the implementation interface for data in the provided rewind history */
-	void ValidateDataInHistory(Chaos::FBaseRewindHistory* History, INetworkPhysicsInputStateImplementationBase& Interface) override
+	/** Call the ValidateData callback on the external implementation interface for data in the provided rewind history */
+	void ValidateData(Chaos::FBaseRewindHistory* History, INetworkPhysicsInputState_External& Interface) override
 	{
 		Chaos::TDataRewindHistory<DataType>* NetHistory = static_cast<Chaos::TDataRewindHistory<DataType>*>(History);
 		TArray<DataType>& HistoryArray = NetHistory->GetDataHistory();
 
 		for (int32 Idx = 0; Idx < NetHistory->GetHistorySize(); Idx++)
 		{
-			FNetworkPhysicsBase& Data = HistoryArray[Idx];
-			Interface.ValidateInput(Data);
+			FNetworkPhysicsPayload& Data = HistoryArray[Idx];
+			if (Data.ServerFrame > LatestValidatedInput_External)
+			{
+				Interface.ValidateInput(Data);
+			}
 		}
+		LatestValidatedInput_External = FMath::Max(LatestValidatedInput_External, NetHistory->GetLatestFrame());
+	}
+
+	/** Call the ValidateData callback on the internal implementation interface for data in the provided rewind history */
+	void ValidateData(Chaos::FBaseRewindHistory* History, INetworkPhysicsInputState_Internal& Interface) override
+	{
+		Chaos::TDataRewindHistory<DataType>* NetHistory = static_cast<Chaos::TDataRewindHistory<DataType>*>(History);
+		TArray<DataType>& HistoryArray = NetHistory->GetDataHistory();
+
+		for (int32 Idx = 0; Idx < NetHistory->GetHistorySize(); Idx++)
+		{
+			FNetworkPhysicsPayload& Data = HistoryArray[Idx];
+			if (Data.ServerFrame > LatestValidatedInput_Internal)
+			{
+				Interface.ValidateInput(Data);
+			}
+		}
+		LatestValidatedInput_Internal = FMath::Max(LatestValidatedInput_Internal, NetHistory->GetLatestFrame());
 	}
 
 /* // Important is not implemented into the new network flow (yet?)
-	void ApplyDataRange(const Chaos::FBaseRewindHistory* History, const int32 FromFrame, const int32 ToFrame, INetworkPhysicsInputStateImplementationBase& Interface) override
+	void ApplyDataRange(const Chaos::FBaseRewindHistory* History, const int32 FromFrame, const int32 ToFrame, INetworkPhysicsInputState_Internal& Interface) override
 	{
 		const Chaos::TDataRewindHistory<DataType>* NetHistory = static_cast<const Chaos::TDataRewindHistory<DataType>*>(History);
 		const TArray<DataType>& HistoryArray = NetHistory->GetDataHistoryConst();
@@ -1230,6 +1274,10 @@ struct TNetworkPhysicsDataHelper : FNetworkPhysicsDataHelper
 		}
 	}
 */
+
+private:
+	int32 LatestValidatedInput_External = 0;
+	int32 LatestValidatedInput_Internal = 0;
 };
 
 /**
@@ -1316,11 +1364,11 @@ public:
 	/** Check if this is controlled locally through relayed inputs from autonomous proxy. It's recommended to use IsLocallyControlled() when checking if this is locally controlled. */
 	const bool GetIsRelayingLocalInputs() const { return bIsRelayingLocalInputs; }
 
-	/** Override the initial value set by CVar np2.Resim.CompareStateToTriggerRewind on initialize -- When true, cache the clients FNetworkPhysicsBase state in rewind history for autonomous proxies and compare the predicted state with incoming server state to trigger resimulations if they differ, comparison done through FNetworkPhysicsData::CompareData. 
+	/** Override the initial value set by CVar np2.Resim.CompareStateToTriggerRewind on initialize -- When true, cache the clients FNetworkPhysicsPayload state in rewind history for autonomous proxies and compare the predicted state with incoming server state to trigger resimulations if they differ, comparison done through FNetworkPhysicsData::CompareData. 
 	* @param bInIncludeSimProxies overrides value set by CVar np2.Resim.CompareStateToTriggerRewind.IncludeSimProxies on initialize -- When true CompareStateToTriggerRewind is also done for simulated proxies. */
 	ENGINE_API void SetCompareStateToTriggerRewind(const bool bInCompareStateToTriggerRewind, const bool bInIncludeSimProxies = false);
 
-	/** Override the initial value set by CVar np2.Resim.CompareInputToTriggerRewind on initialize -- When true, compare autonomous proxies predicted FNetworkPhysicsBase inputs with incoming server inputs to trigger resimulations if they differ, comparison done through FNetworkPhysicsData::CompareData.*/
+	/** Override the initial value set by CVar np2.Resim.CompareInputToTriggerRewind on initialize -- When true, compare autonomous proxies predicted FNetworkPhysicsPayload inputs with incoming server inputs to trigger resimulations if they differ, comparison done through FNetworkPhysicsData::CompareData.*/
 	ENGINE_API void SetCompareInputToTriggerRewind(const bool bInCompareInputToTriggerRewind);
 
 	/** Get FAsyncNetworkPhysicsComponent on the Physics Thread */
@@ -1515,7 +1563,7 @@ public:
 
 	/** Register state and input to be networked and cached in history along with the interface implementation to interact with the input and state */
 	template<typename Input, typename State>
-	void CreateDataHistory(TNetworkPhysicsInputStateImplementation<Input, State>* Implementation);
+	void CreateDataHistory(TNetworkPhysicsInputState_Internal<Input, State>* Implementation_Internal, TNetworkPhysicsInputState_External<Input, State>* Implementation_External = nullptr);
 
 private:
 
@@ -1528,7 +1576,8 @@ private:
 	friend struct FNetworkPhysicsRewindDataDeltaSourceInputProxy;
 	friend struct FNetworkPhysicsRewindDataDeltaSourceStateProxy;
 
-	INetworkPhysicsInputStateImplementationBase* ImplementationInterface = nullptr;
+	INetworkPhysicsInputState_Internal* ImplementationInterface_Internal = nullptr;
+	INetworkPhysicsInputState_External* ImplementationInterface_External = nullptr;
 
 	bool bIsUsingLegacyData = false;
 
@@ -1592,7 +1641,7 @@ private:
 	// If we are currently relaying inputs and will stop after next network send.
 	bool bStopRelayingLocalInputsDeferred = false;
 
-	// Compare state / input to trigger rewind via FNetworkPhysicsBase::CompareData
+	// Compare state / input to trigger rewind via FNetworkPhysicsPayload::CompareData
 	bool bCompareStateToTriggerRewind = false;
 	bool bCompareStateToTriggerRewindIncludeSimProxies = false; // Include simulated proxies when bCompareStateToTriggerRewind is enabled
 	bool bCompareInputToTriggerRewind = false;
@@ -1613,11 +1662,12 @@ private:
 };
 
 template<typename Input, typename State>
-void UNetworkPhysicsComponent::CreateDataHistory(TNetworkPhysicsInputStateImplementation<Input, State>* Implementation)
+void UNetworkPhysicsComponent::CreateDataHistory(TNetworkPhysicsInputState_Internal<Input, State>* Implementation_Internal, TNetworkPhysicsInputState_External<Input, State>* Implementation_External)
 {
 	bIsUsingLegacyData = false;
 
-	ImplementationInterface = Implementation;
+	ImplementationInterface_Internal = Implementation_Internal;
+	ImplementationInterface_External = Implementation_External;
 
 	InputHelper = MakeUnique<TNetworkPhysicsDataHelper<Input, /*bLegacyData*/ false>>();
 	StateHelper = MakeUnique<TNetworkPhysicsDataHelper<State, /*bLegacyData*/ false>>();
@@ -1713,7 +1763,7 @@ struct FAsyncNetworkPhysicsComponentInput : public Chaos::FSimCallbackInput
 	TOptional<uint16> InputsToNetwork_Owner;
 	TOptional<EPhysicsReplicationMode> PhysicsReplicationMode;
 	TOptional<TWeakObjectPtr<UActorComponent>> ActorComponent;
-	TOptional<INetworkPhysicsInputStateImplementationBase*> ImplementationInterface;
+	TOptional<INetworkPhysicsInputState_Internal*> ImplementationInterface_Internal;
 	TOptional<Chaos::FConstPhysicsObjectHandle> PhysicsObject;
 	TOptional<FString> ActorName;
 	TOptional<TUniquePtr<FNetworkPhysicsDataHelper>> InputHelper;
@@ -1740,7 +1790,7 @@ struct FAsyncNetworkPhysicsComponentInput : public Chaos::FSimCallbackInput
 		InputsToNetwork_Owner.Reset();
 		PhysicsReplicationMode.Reset();
 		ActorComponent.Reset();
-		ImplementationInterface.Reset();
+		ImplementationInterface_Internal.Reset();
 		PhysicsObject.Reset();
 		ActorName.Reset();
 		InputHelper.Reset();
@@ -1868,7 +1918,7 @@ private:
 	void TriggerResimulation(int32 ResimFrame);
 
 	/** Returns the current amount of input decay during resimulation as a magnitude from 0.0 to 1.0. Returns 0 if not currently resimulating. */
-	const float GetCurrentInputDecay(const FNetworkPhysicsBase* PhysicsData);
+	const float GetCurrentInputDecay(const FNetworkPhysicsPayload* PhysicsData);
 
 	/** Populate input data in AsyncOutput to send over the network */
 	void SendInputData_Internal(FAsyncNetworkPhysicsComponentOutput& AsyncOutput, const int32 PhysicsStep);
@@ -1900,7 +1950,7 @@ private:
 	TWeakObjectPtr<UActorComponent> ActorComponent;
 
 	// Implementation of input / state interface
-	INetworkPhysicsInputStateImplementationBase* ImplementationInterface;
+	INetworkPhysicsInputState_Internal* ImplementationInterface_Internal;
 
 	// Root components physics object
 	Chaos::FConstPhysicsObjectHandle PhysicsObject;
@@ -1918,13 +1968,13 @@ private:
 	TSharedPtr<Chaos::FBaseRewindHistory> InputHistory;
 
 	// Local temporary inputs data used by pre/post process inputs functions
-	TUniquePtr<FNetworkPhysicsBase> InputData;
+	TUniquePtr<FNetworkPhysicsPayload> InputData;
 	
 	// Local temporary inputs data used by ConsumeAsyncInput
-	TUniquePtr<FNetworkPhysicsBase> LatestInputReceiveData;
+	TUniquePtr<FNetworkPhysicsPayload> LatestInputReceiveData;
 
 	// Local temporary states data used by pre/post process inputs functions
-	TUniquePtr<FNetworkPhysicsBase> StateData;
+	TUniquePtr<FNetworkPhysicsPayload> StateData;
 
 	// Send last N number of inputs each replication call
 	uint16 InputsToNetwork_OwnerDefault = 3; // Default value for owning client
