Don't reregister actor components when changing the actor label.
This is useful when setting a label right after spawning an actor.

diff --git a/Engine/Source/Runtime/Engine/Private/ActorEditor.cpp b/Engine/Source/Runtime/Engine/Private/ActorEditor.cpp
--- a/Engine/Source/Runtime/Engine/Private/ActorEditor.cpp
+++ b/Engine/Source/Runtime/Engine/Private/ActorEditor.cpp
@@ -151,10 +151,11 @@ bool AActor::CanEditChange(const FProperty* PropertyThatWillChange) const
 
 void AActor::PostEditChangeProperty(FPropertyChangedEvent& PropertyChangedEvent)
 {
+	static const FName Name_ActorLabel = GET_MEMBER_NAME_CHECKED(AActor, ActorLabel);
 	static const FName Name_RelativeLocation = USceneComponent::GetRelativeLocationPropertyName();
 	static const FName Name_RelativeRotation = USceneComponent::GetRelativeRotationPropertyName();
 	static const FName Name_RelativeScale3D = USceneComponent::GetRelativeScale3DPropertyName();
-	static const FName Name_MigratingAsset = GET_MEMBER_NAME_CHECKED(AActor, bMigratingAsset);
+	static const FName Name_MigratingAsset = GET_MEMBER_NAME_CHECKED(AActor, bMigratingAsset);	
 
 	FProperty* MemberPropertyThatChanged = PropertyChangedEvent.MemberProperty;
 	const FName MemberPropertyName = MemberPropertyThatChanged != NULL ? MemberPropertyThatChanged->GetFName() : NAME_None;
@@ -168,7 +169,9 @@ void AActor::PostEditChangeProperty(FPropertyChangedEvent& PropertyChangedEvent)
 
 	// During SIE, allow components to reregistered and reconstructed in PostEditChangeProperty.
 	// This is essential as construction is deferred during spawning / duplication when in SIE.
-	if ((GEditor && GEditor->bIsSimulatingInEditor && GetWorld() != nullptr) || ReregisterComponentsWhenModified())
+	const bool bReregisterComponents = (MemberPropertyName != Name_ActorLabel) && (((GEditor && GEditor->bIsSimulatingInEditor && GetWorld() != nullptr) || ReregisterComponentsWhenModified()));
+
+	if (bReregisterComponents)
 	{
 		// In the Undo case we have an annotation storing information about constructed components and we do not want
 		// to improperly apply out of date changes so we need to skip registration of all blueprint created components
