[NetPhysics] Fixup recent logic
- Destructors
- Use LocalFrame instead of ServerFrame (has the same value but is in line with the property that gets compared with)

diff --git a/Engine/Source/Runtime/Engine/Public/Physics/NetworkPhysicsComponent.h b/Engine/Source/Runtime/Engine/Public/Physics/NetworkPhysicsComponent.h
--- a/Engine/Source/Runtime/Engine/Public/Physics/NetworkPhysicsComponent.h
+++ b/Engine/Source/Runtime/Engine/Public/Physics/NetworkPhysicsComponent.h
@@ -865,6 +865,9 @@ struct FNetworkPhysicsDataCollection
 // Game Thread Input and State Interface API
 class INetworkPhysicsInputState_External
 {
+protected:
+	virtual ~INetworkPhysicsInputState_External() {}
+
 public:
 	virtual void ValidateInput(FNetworkPhysicsPayload& Input) const = 0;
 };
@@ -872,6 +875,10 @@ public:
 template<typename InputType, typename StateType>
 class TNetworkPhysicsInputState_External : public INetworkPhysicsInputState_External
 {
+protected:
+	virtual ~TNetworkPhysicsInputState_External() {}
+
+private:
 	// Wrapper for non-templated base
 	void ValidateInput(FNetworkPhysicsPayload& Input) const override { ValidateInput_External(static_cast<InputType&>(Input)); }
 
@@ -885,6 +892,9 @@ public:
 // Physics Thread Input and State Interface API
 class INetworkPhysicsInputState_Internal
 {
+protected:
+	virtual ~INetworkPhysicsInputState_Internal() {}
+
 public:
 	virtual void BuildInput(FNetworkPhysicsPayload& Input) const = 0;
 	virtual void ValidateInput(FNetworkPhysicsPayload& Input) const = 0;
@@ -897,6 +907,10 @@ public:
 template<typename InputType, typename StateType>
 class TNetworkPhysicsInputState_Internal : public INetworkPhysicsInputState_Internal
 {
+protected:
+	virtual ~TNetworkPhysicsInputState_Internal() {}
+
+private:
 	// Wrapper for non-templated base
 	void BuildInput(FNetworkPhysicsPayload& Input) const override { BuildInput_Internal(static_cast<InputType&>(Input)); }
 	void ValidateInput(FNetworkPhysicsPayload& Input) const override { ValidateInput_Internal(static_cast<InputType&>(Input)); }
@@ -1232,7 +1246,7 @@ struct TNetworkPhysicsDataHelper : FNetworkPhysicsDataHelper
 		for (int32 Idx = 0; Idx < NetHistory->GetHistorySize(); Idx++)
 		{
 			FNetworkPhysicsPayload& Data = HistoryArray[Idx];
-			if (Data.ServerFrame > LatestValidatedInput_External)
+			if (Data.LocalFrame > LatestValidatedInput_External)
 			{
 				Interface.ValidateInput(Data);
 			}
@@ -1249,7 +1263,7 @@ struct TNetworkPhysicsDataHelper : FNetworkPhysicsDataHelper
 		for (int32 Idx = 0; Idx < NetHistory->GetHistorySize(); Idx++)
 		{
 			FNetworkPhysicsPayload& Data = HistoryArray[Idx];
-			if (Data.ServerFrame > LatestValidatedInput_Internal)
+			if (Data.LocalFrame > LatestValidatedInput_Internal)
 			{
 				Interface.ValidateInput(Data);
 			}
