Iris Beta
- Renamed and split virtual function Actor::BeginReplication -> OnReplicationStartedForIris / FillStartReplicatingActorParameters
- Renamed virtual function Actor::EndReplication -> OnStopReplicationForIris
- Renamed AActorComponent::BeginReplication ->OnReplicationStartedForIris
- Renamed AActorComponent::EndReplication -> OnStopReplicationForIris

Dispatch of the callbacks are now routed through ReplicationSystemUtil

#rnx

diff --git a/Engine/Plugins/Runtime/InstancedActors/Source/InstancedActors/Private/InstancedActorsManager.cpp b/Engine/Plugins/Runtime/InstancedActors/Source/InstancedActors/Private/InstancedActorsManager.cpp
--- a/Engine/Plugins/Runtime/InstancedActors/Source/InstancedActors/Private/InstancedActorsManager.cpp
+++ b/Engine/Plugins/Runtime/InstancedActors/Source/InstancedActors/Private/InstancedActorsManager.cpp
@@ -222,9 +222,9 @@ AInstancedActorsManager::AInstancedActorsManager()
 }
 
 #if UE_WITH_IRIS
-void AInstancedActorsManager::BeginReplication()
+void AInstancedActorsManager::OnReplicationStartedForIris(const FOnReplicationStartedParameters& Params)
 {
-	Super::BeginReplication();
+	Super::OnReplicationStartedForIris(Params);
 
 	// This actor is configured to use spatial prioritization in Iris. If it's too far from
 	// a player's position it will replicate to their client less frequently. This becomes a 
diff --git a/Engine/Plugins/Runtime/InstancedActors/Source/InstancedActors/Public/InstancedActorsManager.h b/Engine/Plugins/Runtime/InstancedActors/Source/InstancedActors/Public/InstancedActorsManager.h
--- a/Engine/Plugins/Runtime/InstancedActors/Source/InstancedActors/Public/InstancedActorsManager.h
+++ b/Engine/Plugins/Runtime/InstancedActors/Source/InstancedActors/Public/InstancedActorsManager.h
@@ -263,7 +263,7 @@ protected:
 
 	//~ Begin AActor Overrides
 #if UE_WITH_IRIS
-	UE_API virtual void BeginReplication() override;
+	UE_API virtual void OnReplicationStartedForIris(const FOnReplicationStartedParameters&) override;
 #endif
 	UE_API virtual void BeginPlay() override;
 	UE_API virtual void EndPlay(const EEndPlayReason::Type EndPlayReason) override;
diff --git a/Engine/Source/Runtime/Engine/Classes/Components/ActorComponent.h b/Engine/Source/Runtime/Engine/Classes/Components/ActorComponent.h
--- a/Engine/Source/Runtime/Engine/Classes/Components/ActorComponent.h
+++ b/Engine/Source/Runtime/Engine/Classes/Components/ActorComponent.h
@@ -34,6 +34,12 @@ class UPrimitiveComponent;
 struct FSimpleMemberReference;
 struct FAsyncRegisterLevelContext;
 class IPrimitiveComponent;
+class UEngineReplicationBridge;
+
+namespace UE::Net
+{
+	struct FReplicationSystemUtil;
+}
 
 ENGINE_API extern int32 GEnableDeferredPhysicsCreation;
 
@@ -742,15 +748,35 @@ public:
 	/** Register all replication fragments */
 	ENGINE_API virtual void RegisterReplicationFragments(UE::Net::FFragmentRegistrationContext& Context, UE::Net::EFragmentRegistrationFlags RegistrationFlags) override;
 	
+	struct FOnReplicationStartedParameters
+	{
+	};
+
+	/** Called when we start replicating this component with iris, Do not be call explicitly. */
+	ENGINE_API virtual void OnReplicationStartedForIris(const FOnReplicationStartedParameters&);
+
+	struct FOnStopReplicationParameters
+	{
+	};
+	/** Called before we end replicating this component with ris, Do not be call explicitly. */
+	ENGINE_API virtual void OnStopReplicationForIris(const FOnStopReplicationParameters&);
+
 	/** Called when we want to start replicating this component, should not be called explicitly.*/
+	UE_DEPRECATED(5.7, "This is being deprecated as part of iris beta. Override OnReplicationStartedForIris instead.")
 	ENGINE_API virtual void BeginReplication();
 
 	/** Tell component to end replication, should not be called explicitly. */
+	UE_DEPRECATED(5.7, "This is being deprecated as part of iris beta. Override OnStopReplicationForIris instead.")
 	ENGINE_API virtual void EndReplication();
 #endif // UE_WITH_IRIS
 
 private:
 
+#if UE_WITH_IRIS
+	//ENGINE_API static void StartReplicatingForIris(UActorComponent* ActorComponent);
+	//ENGINE_API static void StopReplicatingForIris(UActorComponent* ActorComponent);
+#endif
+
 	/** If WorldPrivate isn't set this will determine the world from outers */
 	ENGINE_API UWorld* GetWorld_Uncached() const;
 
@@ -1411,6 +1437,7 @@ private:
 	friend class FComponentReregisterContextBase;
 	friend class FComponentRecreateRenderStateContext;
 	friend struct FActorComponentTickFunction;
+	friend struct UE::Net::FReplicationSystemUtil;
 #if WITH_EDITOR
 	friend struct SkeletalMeshUCSModifiedDataCleanupHelper;
 #endif
diff --git a/Engine/Source/Runtime/Engine/Classes/Engine/NetDriver.h b/Engine/Source/Runtime/Engine/Classes/Engine/NetDriver.h
--- a/Engine/Source/Runtime/Engine/Classes/Engine/NetDriver.h
+++ b/Engine/Source/Runtime/Engine/Classes/Engine/NetDriver.h
@@ -2506,7 +2506,7 @@ private:
 	UReplicationSystem* ReplicationSystem = nullptr;
 
 	/** When set this will skip registering all the network relevant actors when setting the World */
-	bool bSkipBeginReplicationForWorld = false;
+	bool bSkipStartReplicatingWorldForIris = false;
 
 	/** True when the NetDriver has been configured to run with the Iris replication system.*/
 	bool bIsUsingIris = false;
diff --git a/Engine/Source/Runtime/Engine/Classes/GameFramework/Actor.h b/Engine/Source/Runtime/Engine/Classes/GameFramework/Actor.h
--- a/Engine/Source/Runtime/Engine/Classes/GameFramework/Actor.h
+++ b/Engine/Source/Runtime/Engine/Classes/GameFramework/Actor.h
@@ -48,7 +48,9 @@ class UDataLayerInstance;
 class AWorldDataLayers;
 class IWorldPartitionCell;
 #if UE_WITH_IRIS
+struct FStartReplicatingActorParameters;
 struct FActorReplicationParams;
+struct FStopReplicatingActorParameters;
 #endif // UE_WITH_IRIS
 class UActorFolder;
 struct FActorDataLayer;
@@ -57,6 +59,7 @@ struct FHitResult;
 namespace UE::Net
 {
 	class FTearOffSetter;
+	struct FReplicationSystemUtil;
 }
 
 // By default, debug and development builds (even cooked) will keep actor labels. Manually define this if you want to make a local build
@@ -3460,27 +3463,60 @@ public:
 #if UE_WITH_IRIS
 	ENGINE_API virtual void RegisterReplicationFragments(UE::Net::FFragmentRegistrationContext& Context, UE::Net::EFragmentRegistrationFlags RegistrationFlags) override;
 
+	struct FFillStartReplicatingActorParametersContext
+	{
+	};
 	/**
-	 * Called for all Actors set to replicate during BeginPlay, It will also be called if SetReplicates(true) is called and the object is not already replicating
+	 * Called if using iris replication right before an actor start replicating, override to provide custom actor replication parameters 
 	 */
-	ENGINE_API virtual void BeginReplication();
+	ENGINE_API virtual void FillStartReplicatingActorParameters(const FFillStartReplicatingActorParametersContext& Context, FStartReplicatingActorParameters& OutParameters);
+
+protected:
 
+	struct FOnReplicationStartedParameters
+	{
+	};
 	/**
-	 * Called when we want to end replication for this actor, typically called from EndPlay() for actors that should be replicated during their lifespan
+	 * Called if using iris replication after an actor has started to replicate with iris replication, Normally this will occur right before BeginPlay, but might will also be called later if SetReplicates(true) is called and the object is not already replicating
 	 */
+	ENGINE_API virtual void OnReplicationStartedForIris(const FOnReplicationStartedParameters&);
+
+	struct FOnStopReplicationParameters
+	{
+		const FStopReplicatingActorParameters& StopReplicatingActorParameters;
+	};
+	/**
+	 * Called if using iris replication when we are about to stop replicating an actor, typically called from EndPlay() for actors that should be replicated during their lifespan
+	 */
+	ENGINE_API virtual void OnStopReplicationForIris(const FOnStopReplicationParameters&);
+
+	/**
+	* Called for all Actors set to replicate during BeginPlay, It will also be called if SetReplicates(true) is called and the object is not already replicating
+	*/
+	UE_DEPRECATED(5.7, "This is being deprecated as part of iris beta. Use OnReplicationStartedForIris or FillStartReplicatingActorParameters overrides instead.")
+	ENGINE_API virtual void BeginReplication();
+
+	/**
+	* Called when we want to end replication for this actor, typically called from EndPlay() for actors that should be replicated during their lifespan
+	*/
+	UE_DEPRECATED(5.7, "This is being deprecated as part of iris beta. Override OnStopReplicationForIris instead.")
 	ENGINE_API virtual void EndReplication(EEndPlayReason::Type EndPlayReason);
-protected:
+
+	/**
+	* Helper to BeginReplication passing on additional parameters to the ReplicationSystem, typically called from code overriding normal BeginReplication()
+	* @param Params Additional parameters we want to pass on
+	*/
+	UE_DEPRECATED(5.7, "This is being deprecated as part of iris beta. Do Not call this directly Override FillStartReplicatingActorParameters instead for custom parameters.")
+	ENGINE_API void BeginReplication(const FActorReplicationParams& Params);
+
 	/**
 	 * Pushes the owning NetConnection for the actor and all of its children to the replication system.
 	 * This information decides whether properties with owner conditionals are replicated or not.
 	 */
 	ENGINE_API void UpdateOwningNetConnection();
 
-	/**
-	 * Helper to BeginReplication passing on additional parameters to the ReplicationSystem, typically called from code overriding normal BeginReplication()
-	 * @param Params Additional parameters we want to pass on
-	 */
-	ENGINE_API void BeginReplication(const FActorReplicationParams& Params);
+	friend struct UE::Net::FReplicationSystemUtil;
+
 #endif // UE_WITH_IRIS
 
 	/**
diff --git a/Engine/Source/Runtime/Engine/Classes/GameFramework/PlayerController.h b/Engine/Source/Runtime/Engine/Classes/GameFramework/PlayerController.h
--- a/Engine/Source/Runtime/Engine/Classes/GameFramework/PlayerController.h
+++ b/Engine/Source/Runtime/Engine/Classes/GameFramework/PlayerController.h
@@ -377,7 +377,7 @@ public:
 
 protected:
 #if UE_WITH_IRIS
-	ENGINE_API virtual void BeginReplication() override;
+	ENGINE_API virtual void OnReplicationStartedForIris(const FOnReplicationStartedParameters&) override;
 #endif // UE_WITH_IRIS
 
 private:
@@ -1987,7 +1987,7 @@ public:
 
 	/** 
 	 * Get the server network address
-	 * @return the adress
+	 * @return the address
 	 */
 	ENGINE_API FString GetServerNetworkAddress();
 
diff --git a/Engine/Source/Runtime/Engine/Private/Actor.cpp b/Engine/Source/Runtime/Engine/Private/Actor.cpp
--- a/Engine/Source/Runtime/Engine/Private/Actor.cpp
+++ b/Engine/Source/Runtime/Engine/Private/Actor.cpp
@@ -47,7 +47,7 @@
 #include "WorldPartition/WorldPartitionRuntimeCellInterface.h"
 #include "AutoRTFM.h"
 #if UE_WITH_IRIS
-#include "Iris/ReplicationSystem/ObjectReplicationBridge.h"
+#include "Net/Iris/ReplicationSystem/EngineReplicationBridge.h"
 #include "Iris/ReplicationSystem/ReplicationSystem.h"
 #include "Iris/ReplicationSystem/Conditionals/ReplicationCondition.h"
 #include "Net/Iris/ReplicationSystem/ReplicationSystemUtil.h"
@@ -1394,7 +1394,7 @@ void AActor::PostMigrate(const FUObjectMigrationContext& MigrationContext)
 			ensureMsgf(!bActorIsBeingDestroyed, TEXT("We do not want to be assigned ownership of Actors that have already been destroyed. We do not yet know what to do with it."));
 
 			#if UE_WITH_IRIS
-			BeginReplication();
+			UE::Net::FReplicationSystemUtil::StartReplicatingActor(this);
 			#endif // UE_WITH_IRIS
 
 			// Each component will be in charge of registering itself in its own PostMigrate, but we must register
@@ -3220,7 +3220,7 @@ void AActor::RouteEndPlay(const EEndPlayReason::Type EndPlayReason)
 			{
 				World->RemoveNetworkActor(this);
 #if UE_WITH_IRIS
-				EndReplication(EndPlayReason);
+				UE::Net::FReplicationSystemUtil::StopReplicatingActor(this, FStopReplicatingActorParameters(EndPlayReason));
 #endif
 			}
 		}
@@ -3248,7 +3248,7 @@ void AActor::EndPlay(const EEndPlayReason::Type EndPlayReason)
 
 #if UE_WITH_IRIS
 		// This must be called otherwise the ReplicationSystem will keep a reference to the actor forever.
-		EndReplication(EndPlayReason);
+		UE::Net::FReplicationSystemUtil::StopReplicatingActor(this, FStopReplicatingActorParameters(EndPlayReason));
 #endif
 
 		// Dispatch the blueprint events
@@ -3876,14 +3876,14 @@ void AActor::UpdateReplicatedComponent(UActorComponent* Component)
 	{
 		if (NetCondition != COND_Never && HasActorBegunPlay())
 		{
-			// Begin replication and set NetCondition, if component already is replicated the NetCondition will be updated
-			Component->BeginReplication();
+			// Start replication and set NetCondition, if component already is replicated the NetCondition will be updated
+			UE::Net::FReplicationSystemUtil::StartReplicatingActorComponent(this, Component);
 		}
 		else if (bWasRemovedFromReplicatedComponents || NetCondition == COND_Never)
 		{
 			if (GbCallEndReplicationForActorComponentWhenUpdatingReplicatedComponent)
 			{
-				Component->EndReplication();
+				UE::Net::FReplicationSystemUtil::StopReplicatingActorComponent(Component);
 			}
 			else
 			{
@@ -3907,12 +3907,12 @@ void AActor::UpdateAllReplicatedComponents()
 #if UE_WITH_IRIS
 		if (HasAuthority())
 		{
-			// Need to end replication for all components that should no longer replicate
+			// Need to stop replication for all components that should no longer replicate
 			if (RepComponentInfo.Component && !RepComponentInfo.Component->GetIsReplicated())
 			{
 				if (GbCallEndReplicationForActorComponentWhenUpdatingReplicatedComponent)
 				{
-					RepComponentInfo.Component->EndReplication();
+					UE::Net::FReplicationSystemUtil::StopReplicatingActorComponent(RepComponentInfo.Component);
 				}
 				else
 				{
@@ -3946,14 +3946,14 @@ void AActor::UpdateAllReplicatedComponents()
 				{
 					if (NetCondition != COND_Never)
 					{
-						// Begin replication and set NetCondition, if component already is replicated the NetCondition will be updated
-						Component->BeginReplication();
+						// Start replication and set NetCondition, if component already is replicated the NetCondition will be updated
+						UE::Net::FReplicationSystemUtil::StartReplicatingActorComponent(this, Component);
 					}
 					else
 					{
 						if (GbCallEndReplicationForActorComponentWhenUpdatingReplicatedComponent)
 						{
-							Component->EndReplication();
+							UE::Net::FReplicationSystemUtil::StopReplicatingActorComponent(Component);
 						}
 						else
 						{
@@ -4568,7 +4568,7 @@ void AActor::SetReplicates(bool bInReplicates)
 					{
 						if (bNewlyReplicates)
 						{
-							BeginReplication();
+							UE::Net::FReplicationSystemUtil::StartReplicatingActor(this);
 						}
 						else
 						{
@@ -4583,7 +4583,7 @@ void AActor::SetReplicates(bool bInReplicates)
 #if UE_WITH_IRIS
 			else if (HasActorBegunPlay())
 			{
-				EndReplication(EEndPlayReason::RemovedFromWorld);
+				UE::Net::FReplicationSystemUtil::StopReplicatingActor(this, FStopReplicatingActorParameters(EEndPlayReason::RemovedFromWorld));
 			}
 #endif
 
@@ -4743,7 +4743,7 @@ void AActor::DispatchBeginPlay(bool bFromLevelStreaming)
 		ActorHasBegunPlay = EActorBeginPlayState::BeginningPlay;
 
 #if UE_WITH_IRIS
-		BeginReplication();
+		UE::Net::FReplicationSystemUtil::StartReplicatingActor(this);
 #endif // UE_WITH_IRIS
 
 #if UE_SUPPORT_FOR_ACTOR_TICK_DISABLE
diff --git a/Engine/Source/Runtime/Engine/Private/ActorReplication.cpp b/Engine/Source/Runtime/Engine/Private/ActorReplication.cpp
--- a/Engine/Source/Runtime/Engine/Private/ActorReplication.cpp
+++ b/Engine/Source/Runtime/Engine/Private/ActorReplication.cpp
@@ -672,7 +672,7 @@ void AActor::AddComponentForReplication(UActorComponent* Component)
 #if UE_WITH_IRIS
     if (NetCondition != COND_Never)
     {
-		Component->BeginReplication();
+		UE::Net::FReplicationSystemUtil::StartReplicatingActorComponent(this, Component);
     }
 #endif
 }
@@ -686,7 +686,7 @@ void AActor::RemoveReplicatedComponent(UActorComponent* Component)
 #if UE_WITH_IRIS
 		if (HasAuthority())
 		{
-			Component->EndReplication();
+			UE::Net::FReplicationSystemUtil::StopReplicatingActorComponent(Component);
 		}
 #endif
 	}
@@ -720,7 +720,7 @@ void AActor::AddReplicatedSubObject(UObject* SubObject, ELifetimeCondition NetCo
 #if UE_WITH_IRIS
 	if (GetIsReplicated() && (ActorHasBegunPlay != EActorBeginPlayState::HasNotBegunPlay))
 	{
-		UE::Net::FReplicationSystemUtil::BeginReplicationForActorSubObject(this, SubObject, NetCondition);
+		UE::Net::FReplicationSystemUtil::StartReplicatingActorSubObject(this, SubObject, NetCondition);
 	}
 #endif
 }
@@ -740,7 +740,7 @@ void AActor::RemoveReplicatedSubObject(UObject* SubObject)
 #if UE_WITH_IRIS
 	if (bWasRemoved && HasAuthority())
 	{
-		UE::Net::FReplicationSystemUtil::EndReplicationForActorSubObject(this, SubObject);
+		UE::Net::FReplicationSystemUtil::StopReplicatingActorSubObject(this, SubObject);
 	}
 #endif // UE_WITH_IRIS
 }
@@ -897,7 +897,7 @@ void AActor::AddActorComponentReplicatedSubObject(UActorComponent* OwnerComponen
 #if UE_WITH_IRIS
 		if (OwnerComponent->GetIsReplicated() && OwnerComponent->HasBegunPlay())
 		{
-			UE::Net::FReplicationSystemUtil::BeginReplicationForActorComponentSubObject(OwnerComponent, SubObject, NetCondition);
+			UE::Net::FReplicationSystemUtil::StartReplicatingActorComponentSubObject(OwnerComponent, SubObject, NetCondition);
 		}
 #endif
 	}
@@ -939,7 +939,7 @@ void AActor::RemoveActorComponentReplicatedSubObject(UActorComponent* OwnerCompo
 #if UE_WITH_IRIS
 	if (bWasRemoved && HasAuthority())
 	{
-		UE::Net::FReplicationSystemUtil::EndReplicationForActorComponentSubObject(OwnerComponent, SubObject);
+		UE::Net::FReplicationSystemUtil::StopReplicatingActorComponentSubObject(OwnerComponent, SubObject);
 	}
 #endif
 }
@@ -1060,22 +1060,30 @@ void AActor::RegisterReplicationFragments(UE::Net::FFragmentRegistrationContext&
 	}
 }
 
-void AActor::BeginReplication(const FActorReplicationParams& Params)
+void AActor::FillStartReplicatingActorParameters(const FFillStartReplicatingActorParametersContext& Context, FStartReplicatingActorParameters& OutParams)
+{
+	OutParams = FStartReplicatingActorParameters();
+}
+
+void AActor::OnReplicationStartedForIris(const FOnReplicationStartedParameters& Params)
 {
-	UE::Net::FReplicationSystemUtil::BeginReplication(this, Params);
 	UpdateOwningNetConnection();
 }
 
-void AActor::BeginReplication()
+void AActor::OnStopReplicationForIris(const FOnStopReplicationParameters&)
+{
+}
+
+void AActor::BeginReplication(const FActorReplicationParams& Params)
 {
-	const FActorReplicationParams BeginReplicationParams;
+}
 
-	BeginReplication(BeginReplicationParams);
+void AActor::BeginReplication()
+{
 }
 
 void AActor::EndReplication(EEndPlayReason::Type EndPlayReason)
 {
-	UE::Net::FReplicationSystemUtil::EndReplication(this, EndPlayReason);
 }
 
 void AActor::UpdateOwningNetConnection()
diff --git a/Engine/Source/Runtime/Engine/Private/Components/ActorComponent.cpp b/Engine/Source/Runtime/Engine/Private/Components/ActorComponent.cpp
--- a/Engine/Source/Runtime/Engine/Private/Components/ActorComponent.cpp
+++ b/Engine/Source/Runtime/Engine/Private/Components/ActorComponent.cpp
@@ -1633,7 +1633,7 @@ void UActorComponent::EndPlay(const EEndPlayReason::Type EndPlayReason)
 #if UE_WITH_IRIS
 	if (EndPlayReason != EEndPlayReason::EndPlayInEditor && EndPlayReason != EEndPlayReason::Quit)
 	{
-		EndReplication();
+		UE::Net::FReplicationSystemUtil::StopReplicatingActorComponent(this);
 	}
 #endif
 
@@ -3068,15 +3068,22 @@ void UActorComponent::RegisterReplicationFragments(UE::Net::FFragmentRegistratio
 	UE::Net::FReplicationFragmentUtil::CreateAndRegisterFragmentsForObject(this, Context, RegistrationFlags);
 }
 
+void UActorComponent::OnReplicationStartedForIris(const FOnReplicationStartedParameters&)
+{
+}
+
+void UActorComponent::OnStopReplicationForIris(const FOnStopReplicationParameters&)
+{
+}
+
 void UActorComponent::BeginReplication()
 {
-	UE::Net::FReplicationSystemUtil::BeginReplicationForActorComponent(GetOwner(), this);
 }
 
 void UActorComponent::EndReplication()
 {
-	UE::Net::FReplicationSystemUtil::EndReplicationForActorComponent(this);
 }
+
 #endif // UE_WITH_IRIS
 
 bool UActorComponent::IsEditableWhenInherited() const
diff --git a/Engine/Source/Runtime/Engine/Private/Net/Iris/ReplicationSystem/EngineReplicationBridge.cpp b/Engine/Source/Runtime/Engine/Private/Net/Iris/ReplicationSystem/EngineReplicationBridge.cpp
--- a/Engine/Source/Runtime/Engine/Private/Net/Iris/ReplicationSystem/EngineReplicationBridge.cpp
+++ b/Engine/Source/Runtime/Engine/Private/Net/Iris/ReplicationSystem/EngineReplicationBridge.cpp
@@ -54,6 +54,8 @@
 #include "Net/Core/Trace/NetDebugName.h"
 #include "Net/NetSubObjectRegistryGetter.h"
 
+#include "Net/Iris/ReplicationSystem/ReplicationSystemUtil.h"
+
 #include "ProfilingDebugging/AssetMetadataTrace.h"
 
 #include "Templates/Casts.h"
@@ -115,6 +117,90 @@ bool ShouldIncludeRootObjectInLevelGroups(const UObject* RootObject)
 
 } // end namespace UE::Net::Private
 
+namespace UE::Net
+{
+
+const TCHAR* LexToString(UE::Net::EStopReplicatingReason Reason)
+{
+	using namespace UE::Net;
+	switch (Reason)
+	{
+		case EStopReplicatingReason::Destroyed:
+		{
+			return TEXT("Destroyed");
+			break;
+		}
+		case EStopReplicatingReason::LevelTransition:
+		{
+			return TEXT("LevelTransition");
+			break;
+		}
+		case EStopReplicatingReason::EndPlayInEditor:
+		{
+			return TEXT("EndPlayInEditor");
+			break;
+		}
+		case EStopReplicatingReason::RemovedFromWorld:
+		{
+			return TEXT("RemovedFromWorld");
+			break;
+		}
+		case EStopReplicatingReason::Quit:
+		{
+			return TEXT("Quit");
+			break;
+		}
+		default:
+		{
+			return TEXT("Invalid");
+			break;
+		}
+	}
+}
+
+}
+
+FStopReplicatingActorParameters::FStopReplicatingActorParameters(EEndPlayReason::Type InEndPlayReason)
+{
+	using namespace UE::Net;
+
+	EndPlayReason = InEndPlayReason;
+
+	switch (EndPlayReason)
+	{
+		case EEndPlayReason::Type::Destroyed:
+		{
+			StopReplicatingReason = EStopReplicatingReason::Destroyed;
+			break;
+		}
+		case EEndPlayReason::Type::LevelTransition:
+		{
+			StopReplicatingReason = EStopReplicatingReason::LevelTransition;
+			break;
+		}
+		case EEndPlayReason::Type::EndPlayInEditor:
+		{
+			StopReplicatingReason = EStopReplicatingReason::EndPlayInEditor;
+			break;
+		}
+		case EEndPlayReason::Type::RemovedFromWorld:
+		{
+			StopReplicatingReason = EStopReplicatingReason::RemovedFromWorld;
+			break;
+		}
+		case EEndPlayReason::Type::Quit:
+		{
+			StopReplicatingReason = EStopReplicatingReason::Quit;
+			break;
+		}
+		default:
+		{
+			check(false);
+			break;
+		}
+	}
+}
+
 UEngineReplicationBridge::UEngineReplicationBridge()
 : UObjectReplicationBridge()
 , ActorFactoryId(UE::Net::InvalidNetObjectFactoryId)
@@ -205,7 +291,7 @@ void UEngineReplicationBridge::OnPreSeamlessTravelGarbageCollect()
 	ensureMsgf(UE::Net::Private::bAlwaysCreateLevelFilteringGroupsForPersistentLevels, TEXT("Seamless travel requires setting cvar net.Iris.AlwaysCreateLevelFilteringGroupsForPersistentLevels to true"));
 }
 
-UE::Net::FNetRefHandle UEngineReplicationBridge::StartReplicatingActor(AActor* Actor, const FActorReplicationParams& ActorParams )
+UE::Net::FNetRefHandle UEngineReplicationBridge::StartReplicatingActor(AActor* Actor)
 {
 	using namespace UE::Net;
 
@@ -290,6 +376,10 @@ UE::Net::FNetRefHandle UEngineReplicationBridge::StartReplicatingActor(AActor* A
 		UE_LOG_BRIDGEID(LogIris, Warning, TEXT("Actor %s does not replicate subobjects using the registered SubObjectsLists, SubObjects will not replicate properly"), ToCStr(GetFullNameSafe(Actor)));
 	}
 
+	FStartReplicatingActorParameters ActorParams;
+	AActor::FFillStartReplicatingActorParametersContext Context;
+	Actor->FillStartReplicatingActorParameters(Context, ActorParams);
+
 	// Create handles for the registered fragments
 	Super::FRootObjectReplicationParams RootObjectParams = 
 	{
@@ -301,22 +391,22 @@ UE::Net::FNetRefHandle UEngineReplicationBridge::StartReplicatingActor(AActor* A
 	};
 
 	// Assign a specific filter if requested
-	if (ActorParams.FilterType != FActorReplicationParams::ConfigFilter)
+	if (ActorParams.FilterType != FStartReplicatingActorParameters::ConfigFilter)
 	{
 		RootObjectParams.bUseExplicitDynamicFilter = true;
 
 		switch(ActorParams.FilterType)
 		{
-			case FActorReplicationParams::AlwaysRelevant:				
+			case FStartReplicatingActorParameters::AlwaysRelevant:				
 				RootObjectParams.ExplicitDynamicFilterName = NAME_None;
 				break;
 
-			case FActorReplicationParams::DefaultSpatial:
+			case FStartReplicatingActorParameters::DefaultSpatial:
 				RootObjectParams.ExplicitDynamicFilterName = GetDefaultSpatialFilterName();
 				break;
 
-			case FActorReplicationParams::ExplicitFilter:
-				ensureMsgf(!ActorParams.ExplicitDynamicFilterName.IsNone(), TEXT("FActorReplicationParams::ExplicitDynamicFilterName must be defined when FilterType is ExplicitFilter"));
+			case FStartReplicatingActorParameters::ExplicitFilter:
+				ensureMsgf(!ActorParams.ExplicitDynamicFilterName.IsNone(), TEXT("FStartReplicatingActorParameters::ExplicitDynamicFilterName must be defined when FilterType is ExplicitFilter"));
 				RootObjectParams.ExplicitDynamicFilterName = ActorParams.ExplicitDynamicFilterName;
 				break;
 		}
@@ -399,9 +489,7 @@ UE::Net::FNetRefHandle UEngineReplicationBridge::StartReplicatingActor(AActor* A
 		{
 			if (IsValid(RepComponentInfo.Component) && RepComponentInfo.NetCondition != ELifetimeCondition::COND_Never)
 			{
-				UActorComponent* ReplicatedComponent = RepComponentInfo.Component;
-				ReplicatedComponent->BeginReplication();
-				// NetCondition is set by replicated component
+				FReplicationSystemUtil::StartReplicatingActorComponent(Actor, RepComponentInfo.Component);
 			}
 		}
 	}
@@ -488,7 +576,7 @@ UE::Net::FNetRefHandle UEngineReplicationBridge::StartReplicatingSubObject(UObje
 	return Super::StartReplicatingSubObject(SubObject, Params, SubObjectFactoryId);
 }
 
-void UEngineReplicationBridge::StopReplicatingActor(AActor* Actor, EEndPlayReason::Type EndPlayReason)
+void UEngineReplicationBridge::StopReplicatingActor(AActor* Actor, const FStopReplicatingActorParameters& Parameters)
 {
 	using namespace UE::Net;
 
@@ -499,26 +587,26 @@ void UEngineReplicationBridge::StopReplicatingActor(AActor* Actor, EEndPlayReaso
 		return;
 	}
 
-	UE_LOG(LogIris, Verbose, TEXT("StopReplicatingActor %s. Reason %s "), *PrintObjectFromNetRefHandle(RefHandle), *UEnum::GetValueAsString(TEXT("Engine.EEndPlayReason"), EndPlayReason));
+	UE_LOG(LogIris, Verbose, TEXT("StopReplicatingActor %s. Reason %s "), *PrintObjectFromNetRefHandle(RefHandle), LexToString(Parameters.StopReplicatingReason));
 	ensureMsgf(IsValid(Actor), TEXT("StopReplicatingActor called on invalid actor tied to handle: %s."), *PrintObjectFromNetRefHandle(RefHandle));
 
 	const bool bIsServer = GetReplicationSystem()->IsServer();
-	const bool bIsObjectDestroyed = EndPlayReason == EEndPlayReason::Destroyed;
+	const bool bIsObjectDestroyed = Parameters.StopReplicatingReason == EStopReplicatingReason::Destroyed;
 
 	// Store destruction info of static actors so their destruction can be replicated when a client loads their level in the future.
 	if (bIsServer && bIsObjectDestroyed && RefHandle.IsStatic())
 	{
-		const UE::Net::FDestructionParameters EndReplicationParameters
+		const UE::Net::FDestructionParameters DestructionParameters
 		{
 			.Location = Actor->GetActorLocation(),
 			.Level = Actor->GetLevel(),
 			.bUseDistanceBasedPrioritization = (Actor->bAlwaysRelevant == false),
 			.NetFactoryId = ActorFactoryId,
 		};
-		StoreDestructionInfo(RefHandle, EndReplicationParameters);
+		StoreDestructionInfo(RefHandle, DestructionParameters);
 	}
 		
-	EEndReplicationFlags Flags = ConvertEndPlayIntoEndReplication(EndPlayReason);
+	EEndReplicationFlags Flags = ConvertToEndReplicationFlags(Parameters.StopReplicatingReason);
 	if (Actor->bNetTemporary)
 	{
 		Flags |= EEndReplicationFlags::SkipPendingEndReplicationValidation;
@@ -526,19 +614,21 @@ void UEngineReplicationBridge::StopReplicatingActor(AActor* Actor, EEndPlayReaso
 	StopReplicatingNetRefHandle(RefHandle, Flags);
 }
 
-EEndReplicationFlags UEngineReplicationBridge::ConvertEndPlayIntoEndReplication(EEndPlayReason::Type EndPlayReason) const
+EEndReplicationFlags UEngineReplicationBridge::ConvertToEndReplicationFlags(UE::Net::EStopReplicatingReason StopReplicatingReason) const
 {
+	using namespace UE::Net;
+
 	EEndReplicationFlags EndReplicationFlags = EEndReplicationFlags::None;
-	if (EndPlayReason == EEndPlayReason::Destroyed)
+	if (StopReplicatingReason == EStopReplicatingReason::Destroyed)
 	{
 		EndReplicationFlags |= EEndReplicationFlags::Destroy | EEndReplicationFlags::DestroyNetHandle | EEndReplicationFlags::ClearNetPushId;
 	}
 
 	// If we are shutting down we do not need to validate that we are not detaching remote instances by accident.
-	const bool bIsShuttingDown = (EndPlayReason == EEndPlayReason::EndPlayInEditor) || (EndPlayReason == EEndPlayReason::Quit) || NetDriver->GetPendingDestruction();
+	const bool bIsShuttingDown = (StopReplicatingReason == EStopReplicatingReason::EndPlayInEditor) || (StopReplicatingReason == EStopReplicatingReason::Quit) || NetDriver->GetPendingDestruction();
 
 	// Streaming out a level can happen prior to all actors in it being destroyed/ending replication through replication so we don't want to validate they're detached by accident.
-	const bool bIsStreamingOutLevel = (EndPlayReason == EEndPlayReason::RemovedFromWorld);
+	const bool bIsStreamingOutLevel = (StopReplicatingReason == EStopReplicatingReason::RemovedFromWorld);
 
 	if (bIsShuttingDown || bIsStreamingOutLevel)
 	{
diff --git a/Engine/Source/Runtime/Engine/Private/Net/Iris/ReplicationSystem/ReplicationSystemUtil.cpp b/Engine/Source/Runtime/Engine/Private/Net/Iris/ReplicationSystem/ReplicationSystemUtil.cpp
--- a/Engine/Source/Runtime/Engine/Private/Net/Iris/ReplicationSystem/ReplicationSystemUtil.cpp
+++ b/Engine/Source/Runtime/Engine/Private/Net/Iris/ReplicationSystem/ReplicationSystemUtil.cpp
@@ -114,82 +114,75 @@ FNetHandle FReplicationSystemUtil::GetNetHandle(const UObject* ReplicatedObject)
 	return FNetHandleManager::GetNetHandle(ReplicatedObject);
 }
 
-void FReplicationSystemUtil::BeginReplication(AActor* Actor, const FActorReplicationParams& Params)
+void FReplicationSystemUtil::StartReplicatingActor(AActor* Actor)
 {
+	bool bCallReplicationStarted = false;
 	if (const UWorld* World = Actor->GetWorld())
 	{
-		ForEachReplicationSystem(GEngine, World, [Actor, &Params](UReplicationSystem* ReplicationSystem)
+		ForEachReplicationSystem(GEngine, World, [Actor, &bCallReplicationStarted](UReplicationSystem* ReplicationSystem)
 		{
 			if (ReplicationSystem->IsServer())
 			{
 				if (UEngineReplicationBridge* Bridge = Cast<UEngineReplicationBridge>(ReplicationSystem->GetReplicationBridge()))
 				{
-					Bridge->StartReplicatingActor(Actor, Params);
+					const FNetRefHandle ReplicatedNetRefHandle = Bridge->StartReplicatingActor(Actor);
+					if (ReplicatedNetRefHandle.IsValid())
+					{
+						bCallReplicationStarted = true;
+					}
 				}
 			}
 		});
 	}
-}
 
-void FReplicationSystemUtil::BeginReplication(AActor* Actor)
-{
-	const FActorReplicationParams BeginReplicationParams;
-	return BeginReplication(Actor, BeginReplicationParams);
-}
-
-void FReplicationSystemUtil::EndReplication(AActor* Actor, EEndPlayReason::Type EndPlayReason)
-{
-	// If the call is coming from for example destroying an actor then a formerly associated NetRefHandle will no longer be valid.
-	// The bridge itself will verify that the actor is replicated by it so there's no reason to check that either.
-	ForEachReplicationSystem([Actor, EndPlayReason](UReplicationSystem* ReplicationSystem)
+	if (bCallReplicationStarted)
 	{
-		if (UEngineReplicationBridge* Bridge = Cast<UEngineReplicationBridge>(ReplicationSystem->GetReplicationBridge()))
-		{
-			Bridge->StopReplicatingActor(Actor, EndPlayReason);
-		}
-	});
-}
+		AActor::FOnReplicationStartedParameters Params;
 
+		Actor->OnReplicationStartedForIris(Params);	
+
+PRAGMA_DISABLE_DEPRECATION_WARNINGS
+		// TODO: Remove as soon as everything is submitted.
+		// Call old callback
+		Actor->BeginReplication();
+PRAGMA_ENABLE_DEPRECATION_WARNINGS
 
-void FReplicationSystemUtil::BeginReplicationForActorComponent(FNetHandle ActorHandle, UActorComponent* ActorComp)
-{
-	if (!ActorHandle.IsValid())
-	{
-		return;
 	}
+}
 
-	if (!ensure(ActorComp != nullptr))
+void FReplicationSystemUtil::StopReplicatingActor(AActor* Actor, const FStopReplicatingActorParameters& StopReplicatingActorParameters)
+{
+	if (!Actor)
 	{
 		return;
 	}
 
-	AActor* Actor = Cast<AActor>(ActorComp->GetOwner());
-	ensureMsgf(FNetHandleManager::GetNetHandle(Actor) == ActorHandle, TEXT("BeginReplicationForActorComponent received invalid owner handle %s for actual owner %s"), *ActorHandle.ToString(), *GetNameSafe(Actor));
-	if (Actor)
+	// Invoke OnStopReplication before we destroy NetRefHandles etc.
+	AActor::FOnStopReplicationParameters OnStopReplicationParams = { .StopReplicatingActorParameters = StopReplicatingActorParameters };
+	Actor->OnStopReplicationForIris(OnStopReplicationParams);
+
+PRAGMA_DISABLE_DEPRECATION_WARNINGS
+	// TODO: Remove as soon as everything is submitted.
+	// Call old callback
+	Actor->EndReplication(StopReplicatingActorParameters.EndPlayReason);
+PRAGMA_ENABLE_DEPRECATION_WARNINGS
+
+	// If the call is coming from for example destroying an actor then a formerly associated NetRefHandle will no longer be valid.
+	// The bridge itself will verify that the actor is replicated by it so there's no reason to check that either.
+	ForEachReplicationSystem([Actor, StopReplicatingActorParameters](UReplicationSystem* ReplicationSystem)
 	{
-		if (const UWorld* World = Actor->GetWorld())
+		if (UEngineReplicationBridge* Bridge = Cast<UEngineReplicationBridge>(ReplicationSystem->GetReplicationBridge()))
 		{
-			ForEachReplicationSystem(GEngine, World, [ActorHandle, ActorComp](UReplicationSystem* ReplicationSystem)
-			{
-				if (ReplicationSystem->IsServer())
-				{
-					if (UEngineReplicationBridge* Bridge = Cast<UEngineReplicationBridge>(ReplicationSystem->GetReplicationBridge()))
-					{
-						const FNetRefHandle OwnerRefHandle = Bridge->GetReplicatedRefHandle(ActorHandle);
-						if (OwnerRefHandle.IsValid())
-						{
-							Bridge->StartReplicatingComponent(OwnerRefHandle, ActorComp);
-						}
-					}
-				}
-			});
+			Bridge->StopReplicatingActor(Actor, StopReplicatingActorParameters);
 		}
-	}
+	});
 }
 
-void FReplicationSystemUtil::BeginReplicationForActorComponent(const AActor* Actor, UActorComponent* ActorComp)
+void FReplicationSystemUtil::StartReplicatingActorComponent(const AActor* Actor, UActorComponent* ActorComp)
 {
 	const FNetHandle ActorHandle = GetNetHandle(Actor);
+	bool bCallReplicationStarted = false;
+
 	// If the actor doesn't have a valid handle we assume it's not replicated by any ReplicationSystem
 	if (!ActorHandle.IsValid())
 	{
@@ -198,7 +191,7 @@ void FReplicationSystemUtil::BeginReplicationForActorComponent(const AActor* Act
 
 	if (const UWorld* World = Actor->GetWorld())
 	{
-		ForEachReplicationSystem(GEngine, World, [ActorHandle, ActorComp](UReplicationSystem* ReplicationSystem)
+		ForEachReplicationSystem(GEngine, World, [ActorHandle, ActorComp, &bCallReplicationStarted](UReplicationSystem* ReplicationSystem)
 		{
 			if (ReplicationSystem->IsServer())
 			{
@@ -207,15 +200,31 @@ void FReplicationSystemUtil::BeginReplicationForActorComponent(const AActor* Act
 					const FNetRefHandle ActorRefHandle = Bridge->GetReplicatedRefHandle(ActorHandle);
 					if (ActorRefHandle.IsValid())
 					{
-						Bridge->StartReplicatingComponent(ActorRefHandle, ActorComp);
+						const FNetRefHandle ComponentHandle = Bridge->StartReplicatingComponent(ActorRefHandle, ActorComp);
+						if (ComponentHandle.IsValid())
+						{
+							bCallReplicationStarted = true;
+						}
 					}
 				}
 			}
 		});
 	}
+
+	if (bCallReplicationStarted)
+	{
+		UActorComponent::FOnReplicationStartedParameters Params;
+		ActorComp->OnReplicationStartedForIris(Params);
+
+PRAGMA_DISABLE_DEPRECATION_WARNINGS
+		// TODO: Remove as soon as everything is submitted.
+		// Call old callback
+		ActorComp->BeginReplication();
+PRAGMA_ENABLE_DEPRECATION_WARNINGS
+	}
 }
 
-void FReplicationSystemUtil::BeginReplicationForActorSubObject(const AActor* Actor, UObject* ActorSubObject, ELifetimeCondition NetCondition)
+void FReplicationSystemUtil::StartReplicatingActorSubObject(const AActor* Actor, UObject* ActorSubObject, ELifetimeCondition NetCondition)
 {
 	if (NetCondition == ELifetimeCondition::COND_Never)
 	{
@@ -253,7 +262,7 @@ void FReplicationSystemUtil::BeginReplicationForActorSubObject(const AActor* Act
 	}
 }
 
-void FReplicationSystemUtil::BeginReplicationForActorComponentSubObject(UActorComponent* ActorComponent, UObject* SubObject, ELifetimeCondition NetCondition)
+void FReplicationSystemUtil::StartReplicatingActorComponentSubObject(UActorComponent* ActorComponent, UObject* SubObject, ELifetimeCondition NetCondition)
 {
 	const AActor* Actor = ActorComponent->GetOwner();
 	if (Actor && NetCondition != ELifetimeCondition::COND_Never)
@@ -295,8 +304,23 @@ void FReplicationSystemUtil::BeginReplicationForActorComponentSubObject(UActorCo
 	}
 }
 
-void FReplicationSystemUtil::EndReplicationForActorComponent(UActorComponent* ActorComponent)
+void FReplicationSystemUtil::StopReplicatingActorComponent(UActorComponent* ActorComponent)
 {
+	if (!ensure(ActorComponent))
+	{
+		return;
+	}
+
+	// Invoke OnStopReplication before we destroy NetRefHandles etc.
+	UActorComponent::FOnStopReplicationParameters OnStopReplicationParams;
+	ActorComponent->OnStopReplicationForIris(OnStopReplicationParams);
+
+PRAGMA_DISABLE_DEPRECATION_WARNINGS
+	// TODO: Remove as soon as everything is submitted.
+	// Call old callback
+	ActorComponent->EndReplication();
+PRAGMA_ENABLE_DEPRECATION_WARNINGS
+
 	ForEachReplicationSystem([ActorComponent](UReplicationSystem* ReplicationSystem)
 	{
 		if (UEngineReplicationBridge* Bridge = Cast<UEngineReplicationBridge>(ReplicationSystem->GetReplicationBridge()))
@@ -307,7 +331,7 @@ void FReplicationSystemUtil::EndReplicationForActorComponent(UActorComponent* Ac
 	});
 }
 
-void FReplicationSystemUtil::EndReplicationForActorSubObject(const AActor* Actor, UObject* SubObject)
+void FReplicationSystemUtil::StopReplicatingActorSubObject(const AActor* Actor, UObject* SubObject)
 {
 	ForEachReplicationSystem([SubObject](UReplicationSystem* ReplicationSystem)
 	{
@@ -351,9 +375,9 @@ void FReplicationSystemUtil::SetNetConditionForActorSubObject(const AActor* Acto
 	}
 }
 
-void FReplicationSystemUtil::EndReplicationForActorComponentSubObject(UActorComponent* ActorComponent, UObject* SubObject)
+void FReplicationSystemUtil::StopReplicatingActorComponentSubObject(UActorComponent* ActorComponent, UObject* SubObject)
 {
-	EndReplicationForActorSubObject(ActorComponent->GetOwner(), SubObject);
+	StopReplicatingActorSubObject(ActorComponent->GetOwner(), SubObject);
 }
 
 void FReplicationSystemUtil::AddDependentActor(const AActor* Parent, AActor* Child, EDependentObjectSchedulingHint SchedulingHint)
@@ -385,8 +409,8 @@ void FReplicationSystemUtil::AddDependentActor(const AActor* Parent, AActor* Chi
 					FNetRefHandle ChildRefHandle = Bridge->GetReplicatedRefHandle(Child);
 					if (!ChildRefHandle.IsValid())
 					{
-						const FActorReplicationParams BeginReplicationParams;
-						ChildRefHandle = Bridge->StartReplicatingActor(Child, BeginReplicationParams);
+						// Note: direct call to bridge, might be that it should call ReplicationUtil::StartReplicatingActor instead 
+						ChildRefHandle = Bridge->StartReplicatingActor(Child);
 					}
 					if (ensureMsgf(ChildRefHandle.IsValid(), TEXT("FReplicationSystemUtil::AddDependentActor Child %s is not replicated"), *GetPathNameSafe(Child)))
 					{
@@ -469,7 +493,7 @@ void FReplicationSystemUtil::SetNetConditionForActorComponent(const UActorCompon
 	}
 }
 
-void FReplicationSystemUtil::BeginReplicationForActorsInWorldForNetDriver(UWorld* World, UNetDriver* NetDriver)
+void FReplicationSystemUtil::StartReplicatingActorsInWorldForNetDriver(UWorld* World, UNetDriver* NetDriver)
 {
 	// We only do this if the world already is initialized
 	// Normally we rely on Begin/EndPlay to control if an actor is replicated or not.
@@ -480,7 +504,7 @@ void FReplicationSystemUtil::BeginReplicationForActorsInWorldForNetDriver(UWorld
 			AActor* Actor = *Iter;
 			if (IsValid(Actor) && Actor->HasActorBegunPlay() && NetDriver->ShouldReplicateActor(Actor) && ULevel::IsNetActor(Actor))
 			{
-				Actor->BeginReplication();
+				StartReplicatingActor(Actor);
 			}
 		}
 	}
@@ -528,14 +552,14 @@ void FReplicationSystemUtil::FlushNetDormancy(UReplicationSystem* ReplicationSys
 	// Handle DormInitial actors
 	if (bWasDormInitial && (Actor->HasActorBegunPlay() || Actor->IsActorBeginningPlay()))
 	{
-		// Call BeginReplication for DORM_Initial actors the first time their dormancy is flushed
+		// Start Replication for DORM_Initial actors the first time their dormancy is flushed
 		// (since it's not called when they BeginPlay).
 		// 
-		// We still don't want to call BeginReplication before BeginPlay though:
+		// We still don't want to call Start Replication before BeginPlay though:
 		// -The actor and its components/subobjects may not be completely set up for replication yet
-		// -If the actor is DormInitial, and is flushed before BeginPlay, its dormancy state will change to DormantAll and it will BeginReplication normally
+		// -If the actor is DormInitial, and is flushed before BeginPlay, its dormancy state will change to DormantAll and it will Start Replication normally
 
-		Actor->BeginReplication();
+		StartReplicatingActor(Actor);
 	}
 
 	const FNetHandle ActorHandle = GetNetHandle(Actor);
diff --git a/Engine/Source/Runtime/Engine/Private/NetDriver.cpp b/Engine/Source/Runtime/Engine/Private/NetDriver.cpp
--- a/Engine/Source/Runtime/Engine/Private/NetDriver.cpp
+++ b/Engine/Source/Runtime/Engine/Private/NetDriver.cpp
@@ -7065,15 +7065,15 @@ void UNetDriver::SetWorld(class UWorld* InWorld)
 #if UE_WITH_IRIS
 		if (IsUsingIrisReplication())
 		{
-			if (bSkipBeginReplicationForWorld)
+			if (bSkipStartReplicatingWorldForIris)
 			{
 				// We skipped, reset the flag in case there is another call.
-				bSkipBeginReplicationForWorld = false;
+				bSkipStartReplicatingWorldForIris = false;
 			}
 			else
 			{
 				UpdateGroupFilterStatusForAllLevelsForAllConnections();				
-				UE::Net::FReplicationSystemUtil::BeginReplicationForActorsInWorldForNetDriver(InWorld, this);
+				UE::Net::FReplicationSystemUtil::StartReplicatingActorsInWorldForNetDriver(InWorld, this);
 			}
 		}
 		else
@@ -7447,7 +7447,7 @@ void UNetDriver::SetReplicationSystem(UReplicationSystem* InReplicationSystem)
 		// When we run using Iris, we use ReplicationSystemId as our unique identifier
 		NetTraceId = ReplicationSystem->GetId();
 
-		UE::Net::FReplicationSystemUtil::BeginReplicationForActorsInWorldForNetDriver(GetWorld(), this);
+		UE::Net::FReplicationSystemUtil::StartReplicatingActorsInWorldForNetDriver(GetWorld(), this);
 	}
 }
 
diff --git a/Engine/Source/Runtime/Engine/Private/PlayerController.cpp b/Engine/Source/Runtime/Engine/Private/PlayerController.cpp
--- a/Engine/Source/Runtime/Engine/Private/PlayerController.cpp
+++ b/Engine/Source/Runtime/Engine/Private/PlayerController.cpp
@@ -6473,11 +6473,11 @@ void APlayerController::RemoveFromNetConditionGroup(FName NetGroup)
 
 
 #if UE_WITH_IRIS
-void APlayerController::BeginReplication()
+void APlayerController::OnReplicationStartedForIris(const FOnReplicationStartedParameters& Params)
 {
 	using namespace UE::Net;
 
-	Super::BeginReplication();
+	Super::OnReplicationStartedForIris(Params);
 
 	// Bump prio of playercontroller in order to make sure it replicates really early
 	static constexpr float PlayerControllerStaticPriority = 100.f;
diff --git a/Engine/Source/Runtime/Engine/Public/Net/Iris/ReplicationSystem/EngineReplicationBridge.h b/Engine/Source/Runtime/Engine/Public/Net/Iris/ReplicationSystem/EngineReplicationBridge.h
--- a/Engine/Source/Runtime/Engine/Public/Net/Iris/ReplicationSystem/EngineReplicationBridge.h
+++ b/Engine/Source/Runtime/Engine/Public/Net/Iris/ReplicationSystem/EngineReplicationBridge.h
@@ -27,9 +27,26 @@ namespace UE::Net
 // If actor should be replicated using IRIS or old replication system
 ENGINE_API bool ShouldUseIrisReplication(const UObject* Actor);
 
+// Currently just a direct mapping of EEEndPlayReason::Type but we might want to add more specific reasons later on
+enum class EStopReplicatingReason : uint32
+{
+	/** When the Actor or Component is explicitly destroyed. */
+	Destroyed = 0,
+	/** When the world is being unloaded for a level transition. */
+	LevelTransition,
+	/** When the world is being unloaded because PIE is ending. */
+	EndPlayInEditor,
+	/** When the level it is a member of is streamed out. */
+	RemovedFromWorld,
+	/** When the application is being exited. */
+	Quit,
+};
+
+ENGINE_API const TCHAR* LexToString(EStopReplicatingReason Reason);
+
 }
 
-struct FActorReplicationParams
+struct FStartReplicatingActorParameters
 {
 	enum EFilterType
 	{
@@ -54,6 +71,21 @@ struct FActorReplicationParams
 	FName ExplicitDynamicFilterName;
 };
 
+// TODO: Remove when we can remove BeginRepliction
+struct FActorReplicationParams : public FStartReplicatingActorParameters
+{
+};
+
+struct FStopReplicatingActorParameters
+{
+	explicit ENGINE_API FStopReplicatingActorParameters(EEndPlayReason::Type);
+
+	UE::Net::EStopReplicatingReason StopReplicatingReason = UE::Net::EStopReplicatingReason::Destroyed;
+
+	//$TODO: Remove when we can remove EndReplication
+	EEndPlayReason::Type EndPlayReason = EEndPlayReason::Type::Destroyed;
+};
+
 #endif // UE_WITH_IRIS
 
 UCLASS(Transient, MinimalAPI, config=Engine)
@@ -76,13 +108,13 @@ public:
 	inline UNetDriver* GetNetDriver() const { return NetDriver; }
 
 	/** Begin replication of an actor and its registered ActorComponents and SubObjects. */
-	ENGINE_API FNetRefHandle StartReplicatingActor(AActor* Instance, const FActorReplicationParams& Params);
+	ENGINE_API FNetRefHandle StartReplicatingActor(AActor* Instance);
 
-	/** Stop replicating an actor. This will destroy the handle of the actor and those of his subobjects. */
-	ENGINE_API void StopReplicatingActor(AActor* Actor, EEndPlayReason::Type EndPlayReason);
+	/** Stop replicating an actor. This will destroy the handle of the actor and those of it's compponents and subobjects. */
+	ENGINE_API void StopReplicatingActor(AActor* Actor, const FStopReplicatingActorParameters& Parameters);
 
-	/** Convert EndPlayReason types into the proper EndReplicationFlags */
-	ENGINE_API EEndReplicationFlags ConvertEndPlayIntoEndReplication(EEndPlayReason::Type EndPlayReason) const;
+	/** Convert StopReplication reason to EndReplicationFlags */
+	ENGINE_API EEndReplicationFlags ConvertToEndReplicationFlags(UE::Net::EStopReplicatingReason StopReplicatingReason) const;
 		
 	/**
 	 * Begin replication of an ActorComponent and its registered SubObjects, 
diff --git a/Engine/Source/Runtime/Engine/Public/Net/Iris/ReplicationSystem/ReplicationSystemUtil.h b/Engine/Source/Runtime/Engine/Public/Net/Iris/ReplicationSystem/ReplicationSystemUtil.h
--- a/Engine/Source/Runtime/Engine/Public/Net/Iris/ReplicationSystem/ReplicationSystemUtil.h
+++ b/Engine/Source/Runtime/Engine/Public/Net/Iris/ReplicationSystem/ReplicationSystemUtil.h
@@ -21,7 +21,8 @@ class UNetConnection;
 class UNetDriver;
 class FRepChangedPropertyTracker;
 
-struct FActorReplicationParams;
+struct FStartReplicatingActorParameters;
+struct FStopReplicatingActorParameters;
 
 namespace UE::Net
 {
@@ -61,35 +62,29 @@ struct FReplicationSystemUtil
 		/** Returns the NetHandle for any replicated object. The returned handle may be invalid. */
 	ENGINE_API static FNetHandle GetNetHandle(const UObject* ReplicatedObject);
 
-	/** Begins replication of an actor and all of its registered subobjects. If any ReplicationSystem wants to replicate the actor a NetHandle will be created. */
-	ENGINE_API static void BeginReplication(AActor* Actor, const FActorReplicationParams& Params);
-
-	/** Begins replication of an actor and all of its registered subobjects. If any ReplicationSystem wants to replicate the actor a NetHandle will be created. */
-	ENGINE_API static void BeginReplication(AActor* Actor);
+	/** Start replicating actor and all of its registered components and subobjects. If any ReplicationSystem wants to replicate the actor a NetHandle will be created. */
+	ENGINE_API static void StartReplicatingActor(AActor* Actor);
 	
 	/** Stop replicating an actor. Will destroy handle for actor and registered subobjects. */
-	ENGINE_API static void EndReplication(AActor* Actor, EEndPlayReason::Type EndPlayReason);
-
-	/** Create NetHandle for actor component and add it as a subobject to the actor handle. */
-	ENGINE_API static void BeginReplicationForActorComponent(FNetHandle ActorHandle, UActorComponent* ActorComponent);
+	ENGINE_API static void StopReplicatingActor(AActor* Actor, const FStopReplicatingActorParameters& StopReplicatingActorParameters);
 
 	/** Create NetHandle for actor component and add it as a subobject to the actor. */
-	ENGINE_API static void BeginReplicationForActorComponent(const AActor* Actor, UActorComponent* ActorComponent);
+	ENGINE_API static void StartReplicatingActorComponent(const AActor* Actor, UActorComponent* ActorComponent);
 
 	/** Stop replicating an ActorComponent and its associated SubObjects. */
-	ENGINE_API static void EndReplicationForActorComponent(UActorComponent* SubObject);
+	ENGINE_API static void StopReplicatingActorComponent(UActorComponent* SubObject);
 
 	/** Create NetHandle for SubObject and add it as a subobject to the actor. */
-	ENGINE_API static void BeginReplicationForActorSubObject(const AActor* Actor, UObject* SubObject, ELifetimeCondition NetCondition);
+	ENGINE_API static void StartReplicatingActorSubObject(const AActor* Actor, UObject* SubObject, ELifetimeCondition NetCondition);
 
-	/** Stop replicating a SubObject ActorComponent and its associated SubObjects. */
-	ENGINE_API static void EndReplicationForActorSubObject(const AActor* Actor, UObject* SubObject);
+	/** Stop replicating a SubObject and its associated SubObjects. */
+	ENGINE_API static void StopReplicatingActorSubObject(const AActor* Actor, UObject* SubObject);
 
 	/** Create NetHandle for SubObject and add it as a subobject to the actor component, but only replicated if the ActorComponent replicates. */
-	ENGINE_API static void BeginReplicationForActorComponentSubObject(UActorComponent* ActorComponent, UObject* SubObject, ELifetimeCondition Condition);
+	ENGINE_API static void StartReplicatingActorComponentSubObject(UActorComponent* ActorComponent, UObject* SubObject, ELifetimeCondition Condition);
 
-	/** Stop replicating an ActorComponent and its associated SubObjects. */
-	ENGINE_API static void EndReplicationForActorComponentSubObject(UActorComponent* ActorComponent, UObject* SubObject);
+	/** Stop replicating an SubObject and its associated SubObjects of an ActorComponent. */
+	ENGINE_API static void StopReplicatingActorComponentSubObject(UActorComponent* ActorComponent, UObject* SubObject);
 
 	/** 
 	 * Set subobject NetCondition for a subobject, the conditions is used to determine if the SubObject should be replicated or not.
@@ -126,8 +121,8 @@ struct FReplicationSystemUtil
 	/** Remove dependent actor from parent. The dependent actor will function as a standard standalone replicated actor. */
 	ENGINE_API static void RemoveDependentActor(const AActor* Parent, AActor* Child);
 
-	/** Begin replication for all networked actors that belong to a specific driver in the world. */
-	ENGINE_API static void BeginReplicationForActorsInWorldForNetDriver(UWorld* World, UNetDriver* NetDriver);
+	/** Start replication for all networked actors that belong to a specific driver in the world. */
+	ENGINE_API static void StartReplicatingActorsInWorldForNetDriver(UWorld* World, UNetDriver* NetDriver);
 
 	/** Notify the ReplicationSystem of a dormancy change. */
 	ENGINE_API static void NotifyActorDormancyChange(UReplicationSystem* ReplicationSystem, AActor* Actor, ENetDormancy OldDormancyState);
diff --git a/Engine/Source/Runtime/GameplayDebugger/Private/GameplayDebuggerCategoryReplicator.cpp b/Engine/Source/Runtime/GameplayDebugger/Private/GameplayDebuggerCategoryReplicator.cpp
--- a/Engine/Source/Runtime/GameplayDebugger/Private/GameplayDebuggerCategoryReplicator.cpp
+++ b/Engine/Source/Runtime/GameplayDebugger/Private/GameplayDebuggerCategoryReplicator.cpp
@@ -575,7 +575,7 @@ void AGameplayDebuggerCategoryReplicator::BeginPlay()
 }
 
 #if UE_WITH_IRIS
-void AGameplayDebuggerCategoryReplicator::BeginReplication()
+void AGameplayDebuggerCategoryReplicator::OnReplicationStartedForIris(const FOnReplicationStartedParameters& Params)
 {
 	if (UE::Net::FReplicationSystemUtil::GetReplicationSystem(this))
 	{
@@ -585,7 +585,7 @@ void AGameplayDebuggerCategoryReplicator::BeginReplication()
 		bOnlyRelevantToOwner = true;
 	}
 
-	Super::BeginReplication();
+	Super::OnReplicationStartedForIris(Params);
 }
 
 const AActor* AGameplayDebuggerCategoryReplicator::GetNetOwner() const 
diff --git a/Engine/Source/Runtime/GameplayDebugger/Public/GameplayDebuggerCategoryReplicator.h b/Engine/Source/Runtime/GameplayDebugger/Public/GameplayDebuggerCategoryReplicator.h
--- a/Engine/Source/Runtime/GameplayDebugger/Public/GameplayDebuggerCategoryReplicator.h
+++ b/Engine/Source/Runtime/GameplayDebugger/Public/GameplayDebuggerCategoryReplicator.h
@@ -128,7 +128,7 @@ protected:
 	GAMEPLAYDEBUGGER_API virtual void BeginPlay() override;
 
 #if UE_WITH_IRIS
-	GAMEPLAYDEBUGGER_API virtual void BeginReplication() override;
+	GAMEPLAYDEBUGGER_API virtual void OnReplicationStartedForIris(const FOnReplicationStartedParameters&) override;
 
 	GAMEPLAYDEBUGGER_API virtual const AActor* GetNetOwner() const override;
 #endif
