Cloth - Renamed the USkeletalMeshComponent::GetClothingSimulationInteractor()  ClothingDataName argument into ClothingAssetName to match the existing clothing interactor function.

#rnx

diff --git a/Engine/Source/Runtime/Engine/Classes/Components/SkeletalMeshComponent.h b/Engine/Source/Runtime/Engine/Classes/Components/SkeletalMeshComponent.h
--- a/Engine/Source/Runtime/Engine/Classes/Components/SkeletalMeshComponent.h
+++ b/Engine/Source/Runtime/Engine/Classes/Components/SkeletalMeshComponent.h
@@ -1576,7 +1576,7 @@ public:
 	 * @note Using NAME_None when there are clothing data of different types on the Skeletal Mesh asset will cause a warning.
 	 */
 	UFUNCTION(BlueprintCallable, Category = Clothing)
-	ENGINE_API UClothingSimulationInteractor* GetClothingSimulationInteractor(const FName ClothingDataName = NAME_None) const;
+	ENGINE_API UClothingSimulationInteractor* GetClothingSimulationInteractor(const FName ClothingAssetName = NAME_None) const;
 
 	/** Callback when the parallel clothing task finishes, copies needed data back to component for gamethread */
 	ENGINE_API void CompleteParallelClothSimulation();
