PR #13622: [Input] DeviceSpecificData should use Device Id as one user could hav?

[RN] When calling GetDeviceSpecificData, use the given input device id. if that is invalid, use the platform user's most recent device.

diff --git a/Engine/Source/Runtime/Engine/Classes/GameFramework/InputDeviceProperties.h b/Engine/Source/Runtime/Engine/Classes/GameFramework/InputDeviceProperties.h
--- a/Engine/Source/Runtime/Engine/Classes/GameFramework/InputDeviceProperties.h
+++ b/Engine/Source/Runtime/Engine/Classes/GameFramework/InputDeviceProperties.h
@@ -125,7 +125,7 @@ protected:
 	template<class TDataLayout>
 	TDataLayout* GetDeviceSpecificDataMutable(const FPlatformUserId UserId, const FInputDeviceId DeviceId, TMap<FName, TDataLayout>& InMap)
 	{
-		return GetDeviceSpecificData<TDataLayout>(UserId, InMap);
+		return GetDeviceSpecificData<TDataLayout>(UserId, DeviceId, InMap);
 	}
 
 	/**
@@ -144,7 +144,9 @@ const TDataLayout* UInputDeviceProperty::GetDeviceSpecificData(const FPlatformUs
 {
 	if (const UInputDeviceSubsystem* SubSystem = UInputDeviceSubsystem::Get())
 	{
-		const FHardwareDeviceIdentifier Hardware = SubSystem->GetMostRecentlyUsedHardwareDevice(UserId);
+		// If we were given a valid device, use that. Otherwise, default to the player's most recently used device
+		const FHardwareDeviceIdentifier Hardware = DeviceId.IsValid() ? SubSystem->GetInputDeviceHardwareIdentifier(DeviceId) : SubSystem->GetMostRecentlyUsedHardwareDevice(UserId);
+		
 		// Check if there are any per-input device overrides available
 		if (const TDataLayout* DeviceDetails = InDeviceData.Find(Hardware.HardwareDeviceIdentifier))
 		{
