World Partition - Data Layers: Fixed issue with AActor::SupportsDataLayerType() returning false for actors in streaming cells
* This in turn was causing AActor::FixupDataLayers() to clear the DataLayerAssets array.
* Only an issue in PIE and some utilities that load actors through streaming levels (UWorldPartitionLevelStreamingDynamic::LoadInEditor())

diff --git a/Engine/Source/Runtime/Engine/Private/ActorEditor.cpp b/Engine/Source/Runtime/Engine/Private/ActorEditor.cpp
--- a/Engine/Source/Runtime/Engine/Private/ActorEditor.cpp
+++ b/Engine/Source/Runtime/Engine/Private/ActorEditor.cpp
@@ -1978,11 +1978,11 @@ static bool HasComponentForceActorNoDataLayers(const AActor* InActor)
 bool AActor::SupportsDataLayerType(TSubclassOf<UDataLayerInstance> InDataLayerType) const
 {
 	ULevel* Level = GetLevel();
-	const bool bIsLevelNotPartitioned = Level ? !Level->bIsPartitioned : false;
+	const bool bIsLevelIsPartitioned = Level ? Level->bIsPartitioned || Level->IsWorldPartitionRuntimeCell() : true;
 	const bool bHasComponentForceActorNoDataLayers = HasComponentForceActorNoDataLayers(this);
 	const bool bActorTypeSupportsDataLayerType = InDataLayerType->IsChildOf<UExternalDataLayerInstance>() ? ActorTypeSupportsExternalDataLayer() : ActorTypeSupportsDataLayer();
 	
-	return (!bIsLevelNotPartitioned &&
+	return (bIsLevelIsPartitioned &&
 		!bHasComponentForceActorNoDataLayers &&
 		bActorTypeSupportsDataLayerType &&
 		IsPackageExternal() &&
