Enable virtual texture sampling on volume material domains.
For volumetric cloud where sampling may happen in a loop we enable the reservoir sampling version of virtual texture feedback collection.

diff --git a/Engine/Shaders/Private/VolumetricCloud.usf b/Engine/Shaders/Private/VolumetricCloud.usf
--- a/Engine/Shaders/Private/VolumetricCloud.usf
+++ b/Engine/Shaders/Private/VolumetricCloud.usf
@@ -8,6 +8,9 @@
  // Change this to force recompilation of all volumetric cloud material shaders
 #pragma message("UESHADERMETADATA_VERSION 26C4467C-50E5-414D-86B5-3F3C27A85891")
 
+// Use reservoir sampling for any virtual texture feedback because the number of samples in the tracing loop is not known upfront.
+#define VT_FEEDBACK_RESERVOIR_SAMPLE 1
+
 #include "Common.ush"
 #include "ColorMap.ush"
 #include "ColorSpace.ush"
@@ -1695,6 +1698,9 @@ void MainCommon(in FMaterialPixelParameters MaterialParameters, in float4 SvPosi
 	}
 #endif // CLOUD_MIN_AND_MAX_DEPTH
 
+#if MATERIAL_VIRTUALTEXTURE_FEEDBACK
+	FinalizeVirtualTextureFeedback(MaterialParameters.VirtualTextureFeedback, SvPosition, View.VTFeedbackBuffer);
+#endif
 
 	//
 	// Output result
diff --git a/Engine/Shaders/Private/VolumetricFogVoxelization.usf b/Engine/Shaders/Private/VolumetricFogVoxelization.usf
--- a/Engine/Shaders/Private/VolumetricFogVoxelization.usf
+++ b/Engine/Shaders/Private/VolumetricFogVoxelization.usf
@@ -336,6 +336,10 @@ void VoxelizePS(
 	// Would be faster to branch around the whole material evaluation, but that would cause divergent flow control for gradient operations
 	Scale *= ComputeVolumeShapeMasking(MaterialParameters.WorldPosition_CamRelative, MaterialParameters.PrimitiveId, Interpolants.FactoryInterpolants);
 
+#if MATERIAL_VIRTUALTEXTURE_FEEDBACK
+	FinalizeVirtualTextureFeedback(MaterialParameters.VirtualTextureFeedback, SvPosition, View.VTFeedbackBuffer);
+#endif
+
 	OutVBufferA = float4(Scattering * Scale, Extinction * Scale);
 	OutVBufferB = float4(EmissiveColor * Scale, 0);
 }
diff --git a/Engine/Source/Runtime/Engine/Private/Materials/HLSLMaterialTranslator.cpp b/Engine/Source/Runtime/Engine/Private/Materials/HLSLMaterialTranslator.cpp
--- a/Engine/Source/Runtime/Engine/Private/Materials/HLSLMaterialTranslator.cpp
+++ b/Engine/Source/Runtime/Engine/Private/Materials/HLSLMaterialTranslator.cpp
@@ -7343,21 +7343,6 @@ int32 FHLSLMaterialTranslator::TextureSample(
 	const bool bMeshPaint = TextureType == MCT_TextureMeshPaint;
 	const bool bMaterialCache = TextureType == MCT_TextureMaterialCache;
 	const bool bVirtualTexture = bMeshPaint || bMaterialCache || TextureType == MCT_TextureVirtual;
-	if (bVirtualTexture)
-	{
-		const EMaterialDomain Domain = Material->GetMaterialDomain();
-		switch (Domain)
-		{
-			case MD_Surface:
-			case MD_DeferredDecal:
-			case MD_PostProcess:
-			case MD_UI:
-			case MD_LightFunction:
-				break;
-			default:
-				return Errorf(TEXT("Sampling a virtual texture is not supported for the %s material domain."), *MaterialDomainString(Domain));
-		}
-	}
 
 	if (MipValueMode == TMVM_Derivative)
 	{
