- Fixed the double-NaniteSetting section in the StaticMesh editoir

diff --git a/Engine/Source/Runtime/Engine/Classes/Engine/StaticMesh.h b/Engine/Source/Runtime/Engine/Classes/Engine/StaticMesh.h
--- a/Engine/Source/Runtime/Engine/Classes/Engine/StaticMesh.h
+++ b/Engine/Source/Runtime/Engine/Classes/Engine/StaticMesh.h
@@ -849,7 +849,7 @@ public:
 
 	static FName GetNaniteSettingsMemberName()
 	{
-		return GET_MEMBER_NAME_CHECKED(UStaticMesh, ThumbnailInfo);
+		return GET_MEMBER_NAME_CHECKED(UStaticMesh, NaniteSettings);
 	}
 
 	void NotifyNaniteSettingsChanged()
