Fix LoadGameIfExistsAsync on platforms that do custom async queuing to call the normal load by default to avoid deadlocks

diff --git a/Engine/Source/Runtime/Engine/Private/SaveGameSystem.cpp b/Engine/Source/Runtime/Engine/Private/SaveGameSystem.cpp
--- a/Engine/Source/Runtime/Engine/Private/SaveGameSystem.cpp
+++ b/Engine/Source/Runtime/Engine/Private/SaveGameSystem.cpp
@@ -295,6 +295,13 @@ void FBaseAsyncSaveGameSystem::LoadGameAsync(bool bAttemptToUseUI, const TCHAR*
 	InternalLoadGameAsync(bAttemptToUseUI, Name, PlatformUserId, Data, Callback);
 }
 
+void FBaseAsyncSaveGameSystem::LoadGameIfExistsAsync(bool bAttemptToUseUI, const TCHAR* Name, FPlatformUserId PlatformUserId, FSaveGameAsyncLoadCompleteCallback Callback)
+{
+	// If a platform has a custom async implementation it is not safe to use the generic task pipe. So, fall back to calling the normal async load
+	// Platforms can override this if InternalDoesSaveGameExistAsync always needs to be called before InternalLoadGameAsync
+	LoadGameAsync(bAttemptToUseUI, Name, PlatformUserId, Callback);
+}
+
 void FBaseAsyncSaveGameSystem::DeleteGameAsync(bool bAttemptToUseUI, const TCHAR* Name, FPlatformUserId PlatformUserId, FSaveGameAsyncOpCompleteCallback Callback)
 {
 	InternalDeleteGameAsync(bAttemptToUseUI, Name, PlatformUserId, Callback);
diff --git a/Engine/Source/Runtime/Engine/Public/SaveGameSystem.h b/Engine/Source/Runtime/Engine/Public/SaveGameSystem.h
--- a/Engine/Source/Runtime/Engine/Public/SaveGameSystem.h
+++ b/Engine/Source/Runtime/Engine/Public/SaveGameSystem.h
@@ -195,6 +195,7 @@ public:
 	ENGINE_API virtual void DoesSaveGameExistAsync(const TCHAR* Name, FPlatformUserId PlatformUserId, FSaveGameAsyncExistsCallback Callback) override;
 	ENGINE_API virtual void SaveGameAsync(bool bAttemptToUseUI, const TCHAR* Name, FPlatformUserId PlatformUserId, TSharedRef<const TArray<uint8>> Data, FSaveGameAsyncOpCompleteCallback Callback) override;
 	ENGINE_API virtual void LoadGameAsync(bool bAttemptToUseUI, const TCHAR* Name, FPlatformUserId PlatformUserId, FSaveGameAsyncLoadCompleteCallback Callback) override;
+	ENGINE_API virtual void LoadGameIfExistsAsync(bool bAttemptToUseUI, const TCHAR* Name, FPlatformUserId PlatformUserId, FSaveGameAsyncLoadCompleteCallback Callback) override;
 	ENGINE_API virtual void DeleteGameAsync(bool bAttemptToUseUI, const TCHAR* Name, FPlatformUserId PlatformUserId, FSaveGameAsyncOpCompleteCallback Callback) override;
 	ENGINE_API virtual void GetSaveGameNamesAsync(FPlatformUserId PlatformUserId, FSaveGameAsyncGetNamesCallback Callback) override;
 
