[Editor] Extend Schema to get the color using the UEdGraphPin directly and not the PinType.

* By default the pin type will still be used
* Replaced in SGraphPin.

#rnx

diff --git a/Engine/Source/Editor/GraphEditor/Private/SGraphPin.cpp b/Engine/Source/Editor/GraphEditor/Private/SGraphPin.cpp
--- a/Engine/Source/Editor/GraphEditor/Private/SGraphPin.cpp
+++ b/Engine/Source/Editor/GraphEditor/Private/SGraphPin.cpp
@@ -1202,10 +1202,10 @@ FSlateColor SGraphPin::GetPinColor() const
 		{
 			if (!GetPinObj()->GetOwningNode()->IsNodeEnabled() || GetPinObj()->GetOwningNode()->IsDisplayAsDisabledForced() || !IsEditingEnabled() || GetPinObj()->GetOwningNode()->IsNodeUnrelated())
 			{
-				return Schema->GetPinTypeColor(GraphPin->PinType) * FLinearColor(1.0f, 1.0f, 1.0f, 0.5f);
+				return Schema->GetPinColor(GraphPin) * FLinearColor(1.0f, 1.0f, 1.0f, 0.5f);
 			}
 
-			return Schema->GetPinTypeColor(GraphPin->PinType) * PinColorModifier;
+			return Schema->GetPinColor(GraphPin) * PinColorModifier;
 		}
 	}
 
diff --git a/Engine/Source/Runtime/Engine/Classes/EdGraph/EdGraphSchema.h b/Engine/Source/Runtime/Engine/Classes/EdGraph/EdGraphSchema.h
--- a/Engine/Source/Runtime/Engine/Classes/EdGraph/EdGraphSchema.h
+++ b/Engine/Source/Runtime/Engine/Classes/EdGraph/EdGraphSchema.h
@@ -963,6 +963,16 @@ class UEdGraphSchema : public UObject
 	 */
 	virtual FLinearColor GetPinTypeColor(const FEdGraphPinType& PinType) const { return FLinearColor::Black; }
 
+	/**
+	 * Gets the draw color of a pin based on the UEdGraphPin.
+	 * Will return GetPinTypeColor by default
+	 *
+	 * @param	InPin The pin to convert into a color.
+	 *
+	 * @return	The color representing the passed in pin
+	 */
+	virtual FLinearColor GetPinColor(const UEdGraphPin* InPin) const { return InPin ? GetPinTypeColor(InPin->PinType) : FLinearColor::Black; }
+
 	virtual FLinearColor GetSecondaryPinTypeColor(const FEdGraphPinType& PinType) const { return FLinearColor::White; };
 
 #if WITH_EDITORONLY_DATA
