[Engine] Added PDI overload allowing separate top radius for BuildCylinderVerts

diff --git a/Engine/Source/Runtime/Engine/Private/PrimitiveDrawingUtils.cpp b/Engine/Source/Runtime/Engine/Private/PrimitiveDrawingUtils.cpp
--- a/Engine/Source/Runtime/Engine/Private/PrimitiveDrawingUtils.cpp
+++ b/Engine/Source/Runtime/Engine/Private/PrimitiveDrawingUtils.cpp
@@ -571,9 +571,15 @@ void GetDiscMesh(const FMatrix& ToWorld, const FVector& Center, const FVector& X
 }
 
 void BuildCylinderVerts(const FVector& Base, const FVector& XAxis, const FVector& YAxis, const FVector& ZAxis, double Radius, double HalfHeight, uint32 Sides, TArray<FDynamicMeshVertex>& OutVerts, TArray<uint32>& OutIndices)
+{
+	// Call with equal top and bottom radius values
+	BuildTaperedCylinderVerts(Base, XAxis, YAxis, ZAxis, Radius, Radius, HalfHeight, Sides, OutVerts, OutIndices);
+}
+
+void BuildTaperedCylinderVerts(const FVector& Base, const FVector& XAxis, const FVector& YAxis, const FVector& ZAxis, double BaseRadius, double TopRadius, double HalfHeight, uint32 Sides, TArray<FDynamicMeshVertex>& OutVerts, TArray<uint32>& OutIndices)
 {
 	const float	AngleDelta = 2.0f * UE_PI / Sides;
-	FVector	LastVertex = Base + XAxis * Radius;
+	FVector	LastVertex = Base + XAxis * BaseRadius;
 
 	FVector2D TC = FVector2D(0.0f, 0.0f);
 	float TCStep = 1.0f / Sides;
@@ -585,7 +591,7 @@ void BuildCylinderVerts(const FVector& Base, const FVector& XAxis, const FVector
 	//Compute vertices for base circle.
 	for (uint32 SideIndex = 0; SideIndex < Sides; SideIndex++)
 	{
-		const FVector Vertex = Base + (XAxis * FMath::Cos(AngleDelta * (SideIndex + 1)) + YAxis * FMath::Sin(AngleDelta * (SideIndex + 1))) * Radius;
+		const FVector Vertex = Base + (XAxis * FMath::Cos(AngleDelta * (SideIndex + 1)) + YAxis * FMath::Sin(AngleDelta * (SideIndex + 1))) * BaseRadius;
 		FVector Normal = Vertex - Base;
 		Normal.Normalize();
 
@@ -606,13 +612,13 @@ void BuildCylinderVerts(const FVector& Base, const FVector& XAxis, const FVector
 		TC.X += TCStep;
 	}
 
-	LastVertex = Base + XAxis * Radius;
+	LastVertex = Base + XAxis * BaseRadius;
 	TC = FVector2D(0.0f, 1.0f);
 
 	//Compute vertices for the top circle
 	for (uint32 SideIndex = 0; SideIndex < Sides; SideIndex++)
 	{
-		const FVector Vertex = Base + (XAxis * FMath::Cos(AngleDelta * (SideIndex + 1)) + YAxis * FMath::Sin(AngleDelta * (SideIndex + 1))) * Radius;
+		const FVector Vertex = Base + (XAxis * FMath::Cos(AngleDelta * (SideIndex + 1)) + YAxis * FMath::Sin(AngleDelta * (SideIndex + 1))) * TopRadius;
 		FVector Normal = Vertex - Base;
 		Normal.Normalize();
 
@@ -670,7 +676,6 @@ void BuildCylinderVerts(const FVector& Base, const FVector& XAxis, const FVector
 		OutIndices.Add(V3);
 		OutIndices.Add(V1);
 	}
-
 }
 
 void GetCylinderMesh(const FVector& Start, const FVector& End, double Radius, int32 Sides, const FMaterialRenderProxy* MaterialInstance, uint8 DepthPriority, int32 ViewIndex, FMeshElementCollector& Collector, HHitProxy* HitProxy)
diff --git a/Engine/Source/Runtime/Engine/Public/SceneManagement.h b/Engine/Source/Runtime/Engine/Public/SceneManagement.h
--- a/Engine/Source/Runtime/Engine/Public/SceneManagement.h
+++ b/Engine/Source/Runtime/Engine/Public/SceneManagement.h
@@ -1708,6 +1708,9 @@ extern ENGINE_API void BuildConeVerts(float Angle1, float Angle2, float Scale, f
 
 extern ENGINE_API void BuildCylinderVerts(const FVector& Base, const FVector& XAxis, const FVector& YAxis, const FVector& ZAxis, double Radius, double HalfHeight, uint32 Sides, TArray<FDynamicMeshVertex>& OutVerts, TArray<uint32>& OutIndices);
 
+/** Build tapered cylinder vertices with a separate base and top radius. */
+extern ENGINE_API void BuildTaperedCylinderVerts(const FVector& Base, const FVector& XAxis, const FVector& YAxis, const FVector& ZAxis, double BaseRadius, double TopRadius, double HalfHeight, uint32 Sides, TArray<FDynamicMeshVertex>& OutVerts, TArray<uint32>& OutIndices);
+
 
 /**
  * Given a base color and a selection state, returns a color which accounts for the selection state.
