[ALT/Misc]
- Defer NotifyToolsOfObjectReplacement notification so it execute on the game thread.

diff --git a/Engine/Source/Runtime/Engine/Private/UnrealEngine.cpp b/Engine/Source/Runtime/Engine/Private/UnrealEngine.cpp
--- a/Engine/Source/Runtime/Engine/Private/UnrealEngine.cpp
+++ b/Engine/Source/Runtime/Engine/Private/UnrealEngine.cpp
@@ -20,6 +20,7 @@ UnrealEngine.cpp: Implements the UEngine class and helpers.
 #include "Engine/PendingNetGame.h"
 #include "Misc/FileHelper.h"
 #include "EngineGlobals.h"
+#include "Containers/Ticker.h"
 #include "Misc/FrameValue.h"
 #include "FramePro/FramePro.h"
 #include "Misc/BufferedOutputDevice.h"
@@ -18439,8 +18440,20 @@ void UEngine::CopyPropertiesForUnrelatedObjects(UObject* OldObject, UObject* New
 	if (Params.bNotifyObjectReplacement && GEngine != nullptr &&
 		!OldObject->HasAnyInternalFlags(EInternalObjectFlags::Async | EInternalObjectFlags_AsyncLoading))
 	{
-		check(IsInGameThread());
-		GEngine->NotifyToolsOfObjectReplacement(ReferenceReplacementMap);
+		if (IsInGameThread())
+		{
+			GEngine->NotifyToolsOfObjectReplacement(ReferenceReplacementMap);
+		}
+		else
+		{
+			ExecuteOnGameThread(
+				TEXT("NotifyToolsOfObjectReplacement"),
+				[ReferenceReplacementMap]()
+				{
+					GEngine->NotifyToolsOfObjectReplacement(ReferenceReplacementMap);
+				}
+			);
+		}
 	}
 }
 
