Fix incorrect SSS  profile fetching when using Substrate's hidden conversion.

diff --git a/Engine/Source/Runtime/Engine/Private/Materials/Material.cpp b/Engine/Source/Runtime/Engine/Private/Materials/Material.cpp
--- a/Engine/Source/Runtime/Engine/Private/Materials/Material.cpp
+++ b/Engine/Source/Runtime/Engine/Private/Materials/Material.cpp
@@ -461,9 +461,23 @@ public:
 		{
 			if (Type == EMaterialParameterType::Scalar && ParameterInfo.Name == SubsurfaceProfile::GetSubsurfaceProfileParameterName())
 			{
-				// Legacy single SubsurfaceProfile (Substrate do not add this one in the material shader)
-				OutValue = SubsurfaceProfile::GetSubsurfaceProfileId(GetSubsurfaceProfileRT());
-				return true;
+				// When using Substrate with hidden conversion, this path can still be used. Since GetSubsurfaceProfileRT() won't be correct 
+				// (it will retrieve the SSS profile override instead of the actual SSS Profile), we retrieve the SSS profile through the 
+				// array version at index 0
+				if (Substrate::IsSubstrateEnabled())
+				{
+					const USubsurfaceProfile* SSProfileOverrideRT = GetSubsurfaceProfileRT();
+					const USubsurfaceProfile* SSProfileRT = NumSubsurfaceProfileRT() > 0 ? GetSubsurfaceProfileRT(0) : nullptr;
+					// Set the root material Profile, or the profile overriden by any instances.
+					OutValue = SubsurfaceProfile::GetSubsurfaceProfileId(SSProfileOverrideRT ? SSProfileOverrideRT : SSProfileRT);
+					return true;
+				}
+				else
+				{
+					// Legacy single SubsurfaceProfile (Substrate do not add this one in the material shader)
+					OutValue = SubsurfaceProfile::GetSubsurfaceProfileId(GetSubsurfaceProfileRT());
+					return true;
+				}
 			}
 			else if (Type == EMaterialParameterType::Scalar && NumSubsurfaceProfileRT() > 0)
 			{
diff --git a/Engine/Source/Runtime/Engine/Private/Materials/MaterialInstance.cpp b/Engine/Source/Runtime/Engine/Private/Materials/MaterialInstance.cpp
--- a/Engine/Source/Runtime/Engine/Private/Materials/MaterialInstance.cpp
+++ b/Engine/Source/Runtime/Engine/Private/Materials/MaterialInstance.cpp
@@ -300,9 +300,23 @@ bool FMaterialInstanceResource::GetParameterValue(EMaterialParameterType Type, c
 	{
 		check(ParameterInfo.Association == EMaterialParameterAssociation::GlobalParameter);
 
-		const USubsurfaceProfile* MySubsurfaceProfileRT = GetSubsurfaceProfileRT();
-		OutValue = SubsurfaceProfile::GetSubsurfaceProfileId(MySubsurfaceProfileRT);
-		bResult = true;
+		// When using Substrate with hidden conversion, this path can still be used. Since GetSubsurfaceProfileRT() won't be correct 
+		// (it will retrieve the SSS profile override instead of the actual SSS Profile), we retrieve the SSS profile through the 
+		// array version at index 0
+		if (Substrate::IsSubstrateEnabled())
+		{
+			const USubsurfaceProfile* SSProfileOverrideRT = GetSubsurfaceProfileRT();
+			const USubsurfaceProfile* SSProfileRT = NumSubsurfaceProfileRT() > 0 ? GetSubsurfaceProfileRT(0) : nullptr;
+			// Set the root material Profile, or the profile overriden by any instances.
+			OutValue = SubsurfaceProfile::GetSubsurfaceProfileId(SSProfileOverrideRT ? SSProfileOverrideRT : SSProfileRT);
+			bResult = true;
+		}
+		else
+		{
+			const USubsurfaceProfile* MySubsurfaceProfileRT = GetSubsurfaceProfileRT();
+			OutValue = SubsurfaceProfile::GetSubsurfaceProfileId(MySubsurfaceProfileRT);
+			bResult = true;
+		}
 	}
 	else if (Type == EMaterialParameterType::Scalar && NumSubsurfaceProfileRT() > 0)
 	{
