Iris Beta

- Renamed FStartReplicatingActorParameters -> FActorReplicationParams
- Renamed FOnReplicationStartedParameters -> FOnReplicationStartedParams
- Renamed FOnStopReplicationParameters -> FOnStopReplicationParams
- Renamed FFillStartReplicatingActorParametersContext -> FFillReplicationParamsContext
- Renamed FillStartReplicatingActorParameters() -> FillReplicationParams()

diff --git a/Engine/Plugins/Runtime/InstancedActors/Source/InstancedActors/Private/InstancedActorsManager.cpp b/Engine/Plugins/Runtime/InstancedActors/Source/InstancedActors/Private/InstancedActorsManager.cpp
--- a/Engine/Plugins/Runtime/InstancedActors/Source/InstancedActors/Private/InstancedActorsManager.cpp
+++ b/Engine/Plugins/Runtime/InstancedActors/Source/InstancedActors/Private/InstancedActorsManager.cpp
@@ -222,7 +222,7 @@ AInstancedActorsManager::AInstancedActorsManager()
 }
 
 #if UE_WITH_IRIS
-void AInstancedActorsManager::OnReplicationStartedForIris(const FOnReplicationStartedParameters& Params)
+void AInstancedActorsManager::OnReplicationStartedForIris(const FOnReplicationStartedParams& Params)
 {
 	Super::OnReplicationStartedForIris(Params);
 
diff --git a/Engine/Plugins/Runtime/InstancedActors/Source/InstancedActors/Public/InstancedActorsManager.h b/Engine/Plugins/Runtime/InstancedActors/Source/InstancedActors/Public/InstancedActorsManager.h
--- a/Engine/Plugins/Runtime/InstancedActors/Source/InstancedActors/Public/InstancedActorsManager.h
+++ b/Engine/Plugins/Runtime/InstancedActors/Source/InstancedActors/Public/InstancedActorsManager.h
@@ -263,7 +263,7 @@ protected:
 
 	//~ Begin AActor Overrides
 #if UE_WITH_IRIS
-	UE_API virtual void OnReplicationStartedForIris(const FOnReplicationStartedParameters&) override;
+	UE_API virtual void OnReplicationStartedForIris(const FOnReplicationStartedParams&) override;
 #endif
 	UE_API virtual void BeginPlay() override;
 	UE_API virtual void EndPlay(const EEndPlayReason::Type EndPlayReason) override;
diff --git a/Engine/Source/Runtime/Engine/Classes/Components/ActorComponent.h b/Engine/Source/Runtime/Engine/Classes/Components/ActorComponent.h
--- a/Engine/Source/Runtime/Engine/Classes/Components/ActorComponent.h
+++ b/Engine/Source/Runtime/Engine/Classes/Components/ActorComponent.h
@@ -748,18 +748,18 @@ public:
 	/** Register all replication fragments */
 	ENGINE_API virtual void RegisterReplicationFragments(UE::Net::FFragmentRegistrationContext& Context, UE::Net::EFragmentRegistrationFlags RegistrationFlags) override;
 	
-	struct FOnReplicationStartedParameters
+	struct FOnReplicationStartedParams
 	{
 	};
 
 	/** Called when we start replicating this component with iris, Do not be call explicitly. */
-	ENGINE_API virtual void OnReplicationStartedForIris(const FOnReplicationStartedParameters&);
+	ENGINE_API virtual void OnReplicationStartedForIris(const FOnReplicationStartedParams&);
 
-	struct FOnStopReplicationParameters
+	struct FOnStopReplicationParams
 	{
 	};
 	/** Called before we end replicating this component with ris, Do not be call explicitly. */
-	ENGINE_API virtual void OnStopReplicationForIris(const FOnStopReplicationParameters&);
+	ENGINE_API virtual void OnStopReplicationForIris(const FOnStopReplicationParams&);
 
 	/** Called when we want to start replicating this component, should not be called explicitly.*/
 	UE_DEPRECATED(5.7, "This is being deprecated as part of iris beta. Override OnReplicationStartedForIris instead.")
diff --git a/Engine/Source/Runtime/Engine/Classes/GameFramework/Actor.h b/Engine/Source/Runtime/Engine/Classes/GameFramework/Actor.h
--- a/Engine/Source/Runtime/Engine/Classes/GameFramework/Actor.h
+++ b/Engine/Source/Runtime/Engine/Classes/GameFramework/Actor.h
@@ -48,9 +48,8 @@ class UDataLayerInstance;
 class AWorldDataLayers;
 class IWorldPartitionCell;
 #if UE_WITH_IRIS
-struct FStartReplicatingActorParameters;
 struct FActorReplicationParams;
-struct FStopReplicatingActorParameters;
+struct FStopReplicatingActorParams;
 #endif // UE_WITH_IRIS
 class UActorFolder;
 struct FActorDataLayer;
@@ -3463,37 +3462,37 @@ public:
 #if UE_WITH_IRIS
 	ENGINE_API virtual void RegisterReplicationFragments(UE::Net::FFragmentRegistrationContext& Context, UE::Net::EFragmentRegistrationFlags RegistrationFlags) override;
 
-	struct FFillStartReplicatingActorParametersContext
+	struct FFillReplicationParamsContext
 	{
 	};
 	/**
 	 * Called if using iris replication right before an actor start replicating, override to provide custom actor replication parameters 
 	 */
-	ENGINE_API virtual void FillStartReplicatingActorParameters(const FFillStartReplicatingActorParametersContext& Context, FStartReplicatingActorParameters& OutParameters);
+	ENGINE_API virtual void FillReplicationParams(const FFillReplicationParamsContext& Context, FActorReplicationParams& OutParams);
 
 protected:
 
-	struct FOnReplicationStartedParameters
+	struct FOnReplicationStartedParams
 	{
 	};
 	/**
 	 * Called if using iris replication after an actor has started to replicate with iris replication, Normally this will occur right before BeginPlay, but might will also be called later if SetReplicates(true) is called and the object is not already replicating
 	 */
-	ENGINE_API virtual void OnReplicationStartedForIris(const FOnReplicationStartedParameters&);
+	ENGINE_API virtual void OnReplicationStartedForIris(const FOnReplicationStartedParams&);
 
-	struct FOnStopReplicationParameters
+	struct FOnStopReplicationParams
 	{
-		const FStopReplicatingActorParameters& StopReplicatingActorParameters;
+		const FStopReplicatingActorParams& StopReplicatingParams;
 	};
 	/**
 	 * Called if using iris replication when we are about to stop replicating an actor, typically called from EndPlay() for actors that should be replicated during their lifespan
 	 */
-	ENGINE_API virtual void OnStopReplicationForIris(const FOnStopReplicationParameters&);
+	ENGINE_API virtual void OnStopReplicationForIris(const FOnStopReplicationParams&);
 
 	/**
 	* Called for all Actors set to replicate during BeginPlay, It will also be called if SetReplicates(true) is called and the object is not already replicating
 	*/
-	UE_DEPRECATED(5.7, "This is being deprecated as part of iris beta. Use OnReplicationStartedForIris or FillStartReplicatingActorParameters overrides instead.")
+	UE_DEPRECATED(5.7, "This is being deprecated as part of iris beta. Use OnReplicationStartedForIris or FillReplicationParams overrides instead.")
 	ENGINE_API virtual void BeginReplication();
 
 	/**
@@ -3506,7 +3505,7 @@ protected:
 	* Helper to BeginReplication passing on additional parameters to the ReplicationSystem, typically called from code overriding normal BeginReplication()
 	* @param Params Additional parameters we want to pass on
 	*/
-	UE_DEPRECATED(5.7, "This is being deprecated as part of iris beta. Do Not call this directly Override FillStartReplicatingActorParameters instead for custom parameters.")
+	UE_DEPRECATED(5.7, "This is being deprecated as part of iris beta. Do not call directly override FillReplicationParams instead to provide custom parameters.")
 	ENGINE_API void BeginReplication(const FActorReplicationParams& Params);
 
 	/**
diff --git a/Engine/Source/Runtime/Engine/Classes/GameFramework/PlayerController.h b/Engine/Source/Runtime/Engine/Classes/GameFramework/PlayerController.h
--- a/Engine/Source/Runtime/Engine/Classes/GameFramework/PlayerController.h
+++ b/Engine/Source/Runtime/Engine/Classes/GameFramework/PlayerController.h
@@ -377,7 +377,7 @@ public:
 
 protected:
 #if UE_WITH_IRIS
-	ENGINE_API virtual void OnReplicationStartedForIris(const FOnReplicationStartedParameters&) override;
+	ENGINE_API virtual void OnReplicationStartedForIris(const FOnReplicationStartedParams&) override;
 #endif // UE_WITH_IRIS
 
 private:
diff --git a/Engine/Source/Runtime/Engine/Private/Actor.cpp b/Engine/Source/Runtime/Engine/Private/Actor.cpp
--- a/Engine/Source/Runtime/Engine/Private/Actor.cpp
+++ b/Engine/Source/Runtime/Engine/Private/Actor.cpp
@@ -3220,7 +3220,7 @@ void AActor::RouteEndPlay(const EEndPlayReason::Type EndPlayReason)
 			{
 				World->RemoveNetworkActor(this);
 #if UE_WITH_IRIS
-				UE::Net::FReplicationSystemUtil::StopReplicatingActor(this, FStopReplicatingActorParameters(EndPlayReason));
+				UE::Net::FReplicationSystemUtil::StopReplicatingActor(this, FStopReplicatingActorParams(EndPlayReason));
 #endif
 			}
 		}
@@ -3248,7 +3248,7 @@ void AActor::EndPlay(const EEndPlayReason::Type EndPlayReason)
 
 #if UE_WITH_IRIS
 		// This must be called otherwise the ReplicationSystem will keep a reference to the actor forever.
-		UE::Net::FReplicationSystemUtil::StopReplicatingActor(this, FStopReplicatingActorParameters(EndPlayReason));
+		UE::Net::FReplicationSystemUtil::StopReplicatingActor(this, FStopReplicatingActorParams(EndPlayReason));
 #endif
 
 		// Dispatch the blueprint events
@@ -4583,7 +4583,7 @@ void AActor::SetReplicates(bool bInReplicates)
 #if UE_WITH_IRIS
 			else if (HasActorBegunPlay())
 			{
-				UE::Net::FReplicationSystemUtil::StopReplicatingActor(this, FStopReplicatingActorParameters(EEndPlayReason::RemovedFromWorld));
+				UE::Net::FReplicationSystemUtil::StopReplicatingActor(this, FStopReplicatingActorParams(EEndPlayReason::RemovedFromWorld));
 			}
 #endif
 
diff --git a/Engine/Source/Runtime/Engine/Private/ActorReplication.cpp b/Engine/Source/Runtime/Engine/Private/ActorReplication.cpp
--- a/Engine/Source/Runtime/Engine/Private/ActorReplication.cpp
+++ b/Engine/Source/Runtime/Engine/Private/ActorReplication.cpp
@@ -1060,17 +1060,17 @@ void AActor::RegisterReplicationFragments(UE::Net::FFragmentRegistrationContext&
 	}
 }
 
-void AActor::FillStartReplicatingActorParameters(const FFillStartReplicatingActorParametersContext& Context, FStartReplicatingActorParameters& OutParams)
+void AActor::FillReplicationParams(const FFillReplicationParamsContext& Context, FActorReplicationParams& OutParams)
 {
-	OutParams = FStartReplicatingActorParameters();
+	OutParams = FActorReplicationParams();
 }
 
-void AActor::OnReplicationStartedForIris(const FOnReplicationStartedParameters& Params)
+void AActor::OnReplicationStartedForIris(const FOnReplicationStartedParams& Params)
 {
 	UpdateOwningNetConnection();
 }
 
-void AActor::OnStopReplicationForIris(const FOnStopReplicationParameters&)
+void AActor::OnStopReplicationForIris(const FOnStopReplicationParams&)
 {
 }
 
diff --git a/Engine/Source/Runtime/Engine/Private/Components/ActorComponent.cpp b/Engine/Source/Runtime/Engine/Private/Components/ActorComponent.cpp
--- a/Engine/Source/Runtime/Engine/Private/Components/ActorComponent.cpp
+++ b/Engine/Source/Runtime/Engine/Private/Components/ActorComponent.cpp
@@ -3068,11 +3068,11 @@ void UActorComponent::RegisterReplicationFragments(UE::Net::FFragmentRegistratio
 	UE::Net::FReplicationFragmentUtil::CreateAndRegisterFragmentsForObject(this, Context, RegistrationFlags);
 }
 
-void UActorComponent::OnReplicationStartedForIris(const FOnReplicationStartedParameters&)
+void UActorComponent::OnReplicationStartedForIris(const FOnReplicationStartedParams&)
 {
 }
 
-void UActorComponent::OnStopReplicationForIris(const FOnStopReplicationParameters&)
+void UActorComponent::OnStopReplicationForIris(const FOnStopReplicationParams&)
 {
 }
 
diff --git a/Engine/Source/Runtime/Engine/Private/Net/Iris/ReplicationSystem/EngineReplicationBridge.cpp b/Engine/Source/Runtime/Engine/Private/Net/Iris/ReplicationSystem/EngineReplicationBridge.cpp
--- a/Engine/Source/Runtime/Engine/Private/Net/Iris/ReplicationSystem/EngineReplicationBridge.cpp
+++ b/Engine/Source/Runtime/Engine/Private/Net/Iris/ReplicationSystem/EngineReplicationBridge.cpp
@@ -160,7 +160,7 @@ const TCHAR* LexToString(UE::Net::EStopReplicatingReason Reason)
 
 }
 
-FStopReplicatingActorParameters::FStopReplicatingActorParameters(EEndPlayReason::Type InEndPlayReason)
+FStopReplicatingActorParams::FStopReplicatingActorParams(EEndPlayReason::Type InEndPlayReason)
 {
 	using namespace UE::Net;
 
@@ -376,9 +376,9 @@ UE::Net::FNetRefHandle UEngineReplicationBridge::StartReplicatingActor(AActor* A
 		UE_LOG_BRIDGEID(LogIris, Warning, TEXT("Actor %s does not replicate subobjects using the registered SubObjectsLists, SubObjects will not replicate properly"), ToCStr(GetFullNameSafe(Actor)));
 	}
 
-	FStartReplicatingActorParameters ActorParams;
-	AActor::FFillStartReplicatingActorParametersContext Context;
-	Actor->FillStartReplicatingActorParameters(Context, ActorParams);
+	FActorReplicationParams ActorParams;
+	AActor::FFillReplicationParamsContext Context;
+	Actor->FillReplicationParams(Context, ActorParams);
 
 	// Create handles for the registered fragments
 	Super::FRootObjectReplicationParams RootObjectParams = 
@@ -391,22 +391,22 @@ UE::Net::FNetRefHandle UEngineReplicationBridge::StartReplicatingActor(AActor* A
 	};
 
 	// Assign a specific filter if requested
-	if (ActorParams.FilterType != FStartReplicatingActorParameters::ConfigFilter)
+	if (ActorParams.FilterType != FActorReplicationParams::ConfigFilter)
 	{
 		RootObjectParams.bUseExplicitDynamicFilter = true;
 
 		switch(ActorParams.FilterType)
 		{
-			case FStartReplicatingActorParameters::AlwaysRelevant:				
+			case FActorReplicationParams::AlwaysRelevant:				
 				RootObjectParams.ExplicitDynamicFilterName = NAME_None;
 				break;
 
-			case FStartReplicatingActorParameters::DefaultSpatial:
+			case FActorReplicationParams::DefaultSpatial:
 				RootObjectParams.ExplicitDynamicFilterName = GetDefaultSpatialFilterName();
 				break;
 
-			case FStartReplicatingActorParameters::ExplicitFilter:
-				ensureMsgf(!ActorParams.ExplicitDynamicFilterName.IsNone(), TEXT("FStartReplicatingActorParameters::ExplicitDynamicFilterName must be defined when FilterType is ExplicitFilter"));
+			case FActorReplicationParams::ExplicitFilter:
+				ensureMsgf(!ActorParams.ExplicitDynamicFilterName.IsNone(), TEXT("FActorReplicationParams::ExplicitDynamicFilterName must be defined when FilterType is ExplicitFilter"));
 				RootObjectParams.ExplicitDynamicFilterName = ActorParams.ExplicitDynamicFilterName;
 				break;
 		}
@@ -576,7 +576,7 @@ UE::Net::FNetRefHandle UEngineReplicationBridge::StartReplicatingSubObject(UObje
 	return Super::StartReplicatingSubObject(SubObject, Params, SubObjectFactoryId);
 }
 
-void UEngineReplicationBridge::StopReplicatingActor(AActor* Actor, const FStopReplicatingActorParameters& Parameters)
+void UEngineReplicationBridge::StopReplicatingActor(AActor* Actor, const FStopReplicatingActorParams& Params)
 {
 	using namespace UE::Net;
 
@@ -587,11 +587,11 @@ void UEngineReplicationBridge::StopReplicatingActor(AActor* Actor, const FStopRe
 		return;
 	}
 
-	UE_LOG(LogIris, Verbose, TEXT("StopReplicatingActor %s. Reason %s "), *PrintObjectFromNetRefHandle(RefHandle), LexToString(Parameters.StopReplicatingReason));
+	UE_LOG(LogIris, Verbose, TEXT("StopReplicatingActor %s. Reason %s "), *PrintObjectFromNetRefHandle(RefHandle), LexToString(Params.StopReplicatingReason));
 	ensureMsgf(IsValid(Actor), TEXT("StopReplicatingActor called on invalid actor tied to handle: %s."), *PrintObjectFromNetRefHandle(RefHandle));
 
 	const bool bIsServer = GetReplicationSystem()->IsServer();
-	const bool bIsObjectDestroyed = Parameters.StopReplicatingReason == EStopReplicatingReason::Destroyed;
+	const bool bIsObjectDestroyed = Params.StopReplicatingReason == EStopReplicatingReason::Destroyed;
 
 	// Store destruction info of static actors so their destruction can be replicated when a client loads their level in the future.
 	if (bIsServer && bIsObjectDestroyed && RefHandle.IsStatic())
@@ -606,7 +606,7 @@ void UEngineReplicationBridge::StopReplicatingActor(AActor* Actor, const FStopRe
 		StoreDestructionInfo(RefHandle, DestructionParameters);
 	}
 		
-	EEndReplicationFlags Flags = ConvertToEndReplicationFlags(Parameters.StopReplicatingReason);
+	EEndReplicationFlags Flags = ConvertToEndReplicationFlags(Params.StopReplicatingReason);
 	if (Actor->bNetTemporary)
 	{
 		Flags |= EEndReplicationFlags::SkipPendingEndReplicationValidation;
diff --git a/Engine/Source/Runtime/Engine/Private/Net/Iris/ReplicationSystem/ReplicationSystemUtil.cpp b/Engine/Source/Runtime/Engine/Private/Net/Iris/ReplicationSystem/ReplicationSystemUtil.cpp
--- a/Engine/Source/Runtime/Engine/Private/Net/Iris/ReplicationSystem/ReplicationSystemUtil.cpp
+++ b/Engine/Source/Runtime/Engine/Private/Net/Iris/ReplicationSystem/ReplicationSystemUtil.cpp
@@ -137,7 +137,7 @@ void FReplicationSystemUtil::StartReplicatingActor(AActor* Actor)
 
 	if (bCallReplicationStarted)
 	{
-		AActor::FOnReplicationStartedParameters Params;
+		AActor::FOnReplicationStartedParams Params;
 
 		Actor->OnReplicationStartedForIris(Params);	
 
@@ -150,7 +150,7 @@ PRAGMA_ENABLE_DEPRECATION_WARNINGS
 	}
 }
 
-void FReplicationSystemUtil::StopReplicatingActor(AActor* Actor, const FStopReplicatingActorParameters& StopReplicatingActorParameters)
+void FReplicationSystemUtil::StopReplicatingActor(AActor* Actor, const FStopReplicatingActorParams& StopReplicatingParams)
 {
 	if (!Actor)
 	{
@@ -158,22 +158,22 @@ void FReplicationSystemUtil::StopReplicatingActor(AActor* Actor, const FStopRepl
 	}
 
 	// Invoke OnStopReplication before we destroy NetRefHandles etc.
-	AActor::FOnStopReplicationParameters OnStopReplicationParams = { .StopReplicatingActorParameters = StopReplicatingActorParameters };
+	AActor::FOnStopReplicationParams OnStopReplicationParams = { .StopReplicatingParams = StopReplicatingParams };
 	Actor->OnStopReplicationForIris(OnStopReplicationParams);
 
 PRAGMA_DISABLE_DEPRECATION_WARNINGS
 	// TODO: Remove as soon as everything is submitted.
 	// Call old callback
-	Actor->EndReplication(StopReplicatingActorParameters.EndPlayReason);
+	Actor->EndReplication(StopReplicatingParams.EndPlayReason);
 PRAGMA_ENABLE_DEPRECATION_WARNINGS
 
 	// If the call is coming from for example destroying an actor then a formerly associated NetRefHandle will no longer be valid.
 	// The bridge itself will verify that the actor is replicated by it so there's no reason to check that either.
-	ForEachReplicationSystem([Actor, StopReplicatingActorParameters](UReplicationSystem* ReplicationSystem)
+	ForEachReplicationSystem([Actor, StopReplicatingParams](UReplicationSystem* ReplicationSystem)
 	{
 		if (UEngineReplicationBridge* Bridge = Cast<UEngineReplicationBridge>(ReplicationSystem->GetReplicationBridge()))
 		{
-			Bridge->StopReplicatingActor(Actor, StopReplicatingActorParameters);
+			Bridge->StopReplicatingActor(Actor, StopReplicatingParams);
 		}
 	});
 }
@@ -213,7 +213,7 @@ void FReplicationSystemUtil::StartReplicatingActorComponent(const AActor* Actor,
 
 	if (bCallReplicationStarted)
 	{
-		UActorComponent::FOnReplicationStartedParameters Params;
+		UActorComponent::FOnReplicationStartedParams Params;
 		ActorComp->OnReplicationStartedForIris(Params);
 
 PRAGMA_DISABLE_DEPRECATION_WARNINGS
@@ -312,7 +312,7 @@ void FReplicationSystemUtil::StopReplicatingActorComponent(UActorComponent* Acto
 	}
 
 	// Invoke OnStopReplication before we destroy NetRefHandles etc.
-	UActorComponent::FOnStopReplicationParameters OnStopReplicationParams;
+	UActorComponent::FOnStopReplicationParams OnStopReplicationParams;
 	ActorComponent->OnStopReplicationForIris(OnStopReplicationParams);
 
 PRAGMA_DISABLE_DEPRECATION_WARNINGS
diff --git a/Engine/Source/Runtime/Engine/Private/PlayerController.cpp b/Engine/Source/Runtime/Engine/Private/PlayerController.cpp
--- a/Engine/Source/Runtime/Engine/Private/PlayerController.cpp
+++ b/Engine/Source/Runtime/Engine/Private/PlayerController.cpp
@@ -6473,7 +6473,7 @@ void APlayerController::RemoveFromNetConditionGroup(FName NetGroup)
 
 
 #if UE_WITH_IRIS
-void APlayerController::OnReplicationStartedForIris(const FOnReplicationStartedParameters& Params)
+void APlayerController::OnReplicationStartedForIris(const FOnReplicationStartedParams& Params)
 {
 	using namespace UE::Net;
 
diff --git a/Engine/Source/Runtime/Engine/Public/Net/Iris/ReplicationSystem/EngineReplicationBridge.h b/Engine/Source/Runtime/Engine/Public/Net/Iris/ReplicationSystem/EngineReplicationBridge.h
--- a/Engine/Source/Runtime/Engine/Public/Net/Iris/ReplicationSystem/EngineReplicationBridge.h
+++ b/Engine/Source/Runtime/Engine/Public/Net/Iris/ReplicationSystem/EngineReplicationBridge.h
@@ -46,7 +46,7 @@ ENGINE_API const TCHAR* LexToString(EStopReplicatingReason Reason);
 
 }
 
-struct FStartReplicatingActorParameters
+struct FActorReplicationParams
 {
 	enum EFilterType
 	{
@@ -71,14 +71,9 @@ struct FStartReplicatingActorParameters
 	FName ExplicitDynamicFilterName;
 };
 
-// TODO: Remove when we can remove BeginRepliction
-struct FActorReplicationParams : public FStartReplicatingActorParameters
+struct FStopReplicatingActorParams
 {
-};
-
-struct FStopReplicatingActorParameters
-{
-	explicit ENGINE_API FStopReplicatingActorParameters(EEndPlayReason::Type);
+	explicit ENGINE_API FStopReplicatingActorParams(EEndPlayReason::Type);
 
 	UE::Net::EStopReplicatingReason StopReplicatingReason = UE::Net::EStopReplicatingReason::Destroyed;
 
@@ -111,7 +106,7 @@ public:
 	ENGINE_API FNetRefHandle StartReplicatingActor(AActor* Instance);
 
 	/** Stop replicating an actor. This will destroy the handle of the actor and those of it's compponents and subobjects. */
-	ENGINE_API void StopReplicatingActor(AActor* Actor, const FStopReplicatingActorParameters& Parameters);
+	ENGINE_API void StopReplicatingActor(AActor* Actor, const FStopReplicatingActorParams& Params);
 
 	/** Convert StopReplication reason to EndReplicationFlags */
 	ENGINE_API EEndReplicationFlags ConvertToEndReplicationFlags(UE::Net::EStopReplicatingReason StopReplicatingReason) const;
@@ -182,7 +177,6 @@ private:
 	
 	void OnMaxTickRateChanged(UNetDriver* InNetDriver, int32 NewMaxTickRate, int32 OldMaxTickRate);
 
-
 private:
 
 	UE::Net::FNetObjectFactoryId ActorFactoryId;
@@ -202,6 +196,4 @@ private:
 	UPROPERTY(Config)
 	/** Name of the NetObjectFactory to use for default replicated subobjects */
 	FName SubObjectFactoryName;
-
-
 };
diff --git a/Engine/Source/Runtime/Engine/Public/Net/Iris/ReplicationSystem/ReplicationSystemUtil.h b/Engine/Source/Runtime/Engine/Public/Net/Iris/ReplicationSystem/ReplicationSystemUtil.h
--- a/Engine/Source/Runtime/Engine/Public/Net/Iris/ReplicationSystem/ReplicationSystemUtil.h
+++ b/Engine/Source/Runtime/Engine/Public/Net/Iris/ReplicationSystem/ReplicationSystemUtil.h
@@ -21,8 +21,8 @@ class UNetConnection;
 class UNetDriver;
 class FRepChangedPropertyTracker;
 
-struct FStartReplicatingActorParameters;
-struct FStopReplicatingActorParameters;
+struct FActorReplicationParams;
+struct FStopReplicatingActorParams;
 
 namespace UE::Net
 {
@@ -66,7 +66,7 @@ struct FReplicationSystemUtil
 	ENGINE_API static void StartReplicatingActor(AActor* Actor);
 	
 	/** Stop replicating an actor. Will destroy handle for actor and registered subobjects. */
-	ENGINE_API static void StopReplicatingActor(AActor* Actor, const FStopReplicatingActorParameters& StopReplicatingActorParameters);
+	ENGINE_API static void StopReplicatingActor(AActor* Actor, const FStopReplicatingActorParams& StopReplicatingActorParams);
 
 	/** Create NetHandle for actor component and add it as a subobject to the actor. */
 	ENGINE_API static void StartReplicatingActorComponent(const AActor* Actor, UActorComponent* ActorComponent);
diff --git a/Engine/Source/Runtime/GameplayDebugger/Private/GameplayDebuggerCategoryReplicator.cpp b/Engine/Source/Runtime/GameplayDebugger/Private/GameplayDebuggerCategoryReplicator.cpp
--- a/Engine/Source/Runtime/GameplayDebugger/Private/GameplayDebuggerCategoryReplicator.cpp
+++ b/Engine/Source/Runtime/GameplayDebugger/Private/GameplayDebuggerCategoryReplicator.cpp
@@ -575,7 +575,7 @@ void AGameplayDebuggerCategoryReplicator::BeginPlay()
 }
 
 #if UE_WITH_IRIS
-void AGameplayDebuggerCategoryReplicator::OnReplicationStartedForIris(const FOnReplicationStartedParameters& Params)
+void AGameplayDebuggerCategoryReplicator::OnReplicationStartedForIris(const FOnReplicationStartedParams& Params)
 {
 	if (UE::Net::FReplicationSystemUtil::GetReplicationSystem(this))
 	{
diff --git a/Engine/Source/Runtime/GameplayDebugger/Public/GameplayDebuggerCategoryReplicator.h b/Engine/Source/Runtime/GameplayDebugger/Public/GameplayDebuggerCategoryReplicator.h
--- a/Engine/Source/Runtime/GameplayDebugger/Public/GameplayDebuggerCategoryReplicator.h
+++ b/Engine/Source/Runtime/GameplayDebugger/Public/GameplayDebuggerCategoryReplicator.h
@@ -128,7 +128,7 @@ protected:
 	GAMEPLAYDEBUGGER_API virtual void BeginPlay() override;
 
 #if UE_WITH_IRIS
-	GAMEPLAYDEBUGGER_API virtual void OnReplicationStartedForIris(const FOnReplicationStartedParameters&) override;
+	GAMEPLAYDEBUGGER_API virtual void OnReplicationStartedForIris(const FOnReplicationStartedParams&) override;
 
 	GAMEPLAYDEBUGGER_API virtual const AActor* GetNetOwner() const override;
 #endif
