Clean up deprecated skeleton compatibility code in the skeleton class

diff --git a/Engine/Source/Runtime/Engine/Classes/Animation/Skeleton.h b/Engine/Source/Runtime/Engine/Classes/Animation/Skeleton.h
--- a/Engine/Source/Runtime/Engine/Classes/Animation/Skeleton.h
+++ b/Engine/Source/Runtime/Engine/Classes/Animation/Skeleton.h
@@ -726,15 +726,6 @@ public:
 	ENGINE_API void GetCompatibleAssets(UClass* AssetClass, const TCHAR* InTag, TArray<FAssetData>& OutAssets) const;
 #endif
 
-	UE_DEPRECATED(5.2, "Compatibility is now an editor-only concern. Please use IsCompatibleForEditor.")
-	ENGINE_API bool IsCompatible(const USkeleton* InSkeleton) const;
-
-	UE_DEPRECATED(5.2, "Compatibility is now an editor-only concern. Please use IsCompatibleForEditor.")
-	ENGINE_API bool IsCompatibleSkeletonByAssetData(const FAssetData& AssetData, const TCHAR* InTag = TEXT("Skeleton")) const;
-
-	UE_DEPRECATED(5.2, "Compatibility is now an editor-only concern. Please use IsCompatibleForEditor.")
-	ENGINE_API bool IsCompatibleSkeletonByAssetString(const FString& SkeletonAssetString) const;
-
 	DECLARE_EVENT(USkeleton, FSmartNamesChangedEvent);
 
 #if WITH_EDITORONLY_DATA
diff --git a/Engine/Source/Runtime/Engine/Private/Animation/Skeleton.cpp b/Engine/Source/Runtime/Engine/Private/Animation/Skeleton.cpp
--- a/Engine/Source/Runtime/Engine/Private/Animation/Skeleton.cpp
+++ b/Engine/Source/Runtime/Engine/Private/Animation/Skeleton.cpp
@@ -273,29 +273,6 @@ void USkeleton::GetCompatibleAssets(UClass* AssetClass, const TCHAR* InTag, TArr
 
 #endif
 
-bool USkeleton::IsCompatible(const USkeleton* InSkeleton) const
-{
-#if WITH_EDITORONLY_DATA
-	return IsCompatibleForEditor(InSkeleton);
-#else
-	if (InSkeleton == nullptr)
-	{
-		return false;
-	}
-
-	return true;
-#endif
-}
-
-bool USkeleton::IsCompatibleSkeletonByAssetString(const FString& SkeletonAssetString) const
-{
-#if WITH_EDITORONLY_DATA
-	return IsCompatibleForEditor(SkeletonAssetString);
-#else
-	return false;
-#endif
-}
-
 void USkeleton::AddCompatibleSkeleton(const USkeleton* SourceSkeleton)
 {
 	CompatibleSkeletons.AddUnique(const_cast<USkeleton*>(SourceSkeleton));
@@ -316,15 +293,6 @@ void USkeleton::RemoveCompatibleSkeleton(const TSoftObjectPtr<USkeleton>& Source
 	CompatibleSkeletons.Remove(SourceSkeleton);
 }
 
-bool USkeleton::IsCompatibleSkeletonByAssetData(const FAssetData& AssetData, const TCHAR* InTag) const
-{
-#if WITH_EDITORONLY_DATA
-	return IsCompatibleForEditor(AssetData, InTag);
-#else
-	return false;
-#endif
-}
-
 USkeleton::USkeleton(const FObjectInitializer& ObjectInitializer)
 	: Super(ObjectInitializer)
 {
