AssetManager - Improved ensure message when more than one PrimaryAssetType rules are defined with settings that do not match each other

diff --git a/Engine/Source/Runtime/Engine/Private/AssetManager.cpp b/Engine/Source/Runtime/Engine/Private/AssetManager.cpp
--- a/Engine/Source/Runtime/Engine/Private/AssetManager.cpp
+++ b/Engine/Source/Runtime/Engine/Private/AssetManager.cpp
@@ -1292,6 +1292,11 @@ int32 UAssetManager::ScanPathsForPrimaryAssets(FPrimaryAssetType PrimaryAssetTyp
 
 	if (!FoundType)
 	{
+		UE_LOG(LogAssetManager, VeryVerbose, TEXT("UAssetManager::ScanPathsForPrimaryAssets adding TypeData to AssetTypeMap: '%s', (Paths=%s%s)"),
+			*PrimaryAssetType.ToString(),
+			*FString::Join(TArrayView<const FString>(Paths).Left(FMath::Min(Paths.Num(), 3)), TEXT(", ")),
+			(Paths.Num() > 3) ? TEXT(", ...") : TEXT(""));
+
 		TSharedPtr<FPrimaryAssetTypeData> NewAsset = MakeShareable(new FPrimaryAssetTypeData(PrimaryAssetType, BaseClass, bHasBlueprintClasses, bIsEditorOnly));
 
 		FoundType = &AssetTypeMap.Add(PrimaryAssetType, NewAsset.ToSharedRef());
@@ -1303,8 +1308,40 @@ int32 UAssetManager::ScanPathsForPrimaryAssets(FPrimaryAssetType PrimaryAssetTyp
 	FPrimaryAssetTypeData& TypeData = FoundType->Get();
 
 	// Make sure types match
-	if (!ensureMsgf(TypeData.Info.AssetBaseClassLoaded == BaseClass && TypeData.Info.bHasBlueprintClasses == bHasBlueprintClasses && TypeData.Info.bIsEditorOnly == bIsEditorOnly, TEXT("UAssetManager::ScanPathsForPrimaryAssets TypeData parameters did not match for type '%s'"), *TypeData.Info.PrimaryAssetType.ToString()))
+	if (TypeData.Info.AssetBaseClassLoaded != BaseClass || TypeData.Info.bHasBlueprintClasses != bHasBlueprintClasses || TypeData.Info.bIsEditorOnly != bIsEditorOnly)
 	{
+		auto CalculateEnsureMessage = [&TypeData, &Paths, BaseClass, bHasBlueprintClasses, bIsEditorOnly]()
+		{
+			FString Message;
+			if (TypeData.Info.AssetBaseClassLoaded != BaseClass)
+			{
+				Message += FString::Printf(TEXT(", (BaseClass: '%s' != expected: '%s')"), *GetNameSafe(BaseClass), *GetNameSafe(TypeData.Info.AssetBaseClassLoaded.Get()));
+			}
+
+			if (TypeData.Info.bHasBlueprintClasses != bHasBlueprintClasses)
+			{
+				Message += FString::Printf(TEXT(", (bHasBlueprintClasses: '%d' != expected: '%d')"), bHasBlueprintClasses, TypeData.Info.bHasBlueprintClasses);
+			}
+
+			if (TypeData.Info.bIsEditorOnly != bIsEditorOnly)
+			{
+				Message += FString::Printf(TEXT(", (bIsEditorOnly: '%d' != expected: '%d')"), bIsEditorOnly, TypeData.Info.bIsEditorOnly);
+			}
+
+			if (Paths.Num() > 0)
+			{
+				// Context about where first PrimaryAssetType was set
+				Message += FString::Printf(TEXT(", (Paths=%s%s)"),
+					*FString::Join(TArrayView<const FString>(Paths).Left(FMath::Min(Paths.Num(), 3)), TEXT(", ")),
+					(Paths.Num() > 3) ? TEXT(", ...") : TEXT(""));
+			}
+
+			Message += TEXT(". Another plugin, ini, or system has different parameters for this PrimaryAssetType name. Run with -LogCmds=\"LogAssetManager VeryVerbose\" to log more information.");
+
+			return Message;
+		};
+
+		ensureMsgf(false, TEXT("UAssetManager::ScanPathsForPrimaryAssets TypeData parameters did not match PrimaryAssetType '%s'%s"), *TypeData.Info.PrimaryAssetType.ToString(), *CalculateEnsureMessage());
 		return 0;
 	}
 
