Iris Beta
- Renamed commends etc. referring to BeginReplication/EndReplication

diff --git a/Engine/Source/Runtime/Engine/Private/Net/Iris/ReplicationSystem/EngineReplicationBridge.cpp b/Engine/Source/Runtime/Engine/Private/Net/Iris/ReplicationSystem/EngineReplicationBridge.cpp
--- a/Engine/Source/Runtime/Engine/Private/Net/Iris/ReplicationSystem/EngineReplicationBridge.cpp
+++ b/Engine/Source/Runtime/Engine/Private/Net/Iris/ReplicationSystem/EngineReplicationBridge.cpp
@@ -337,9 +337,9 @@ UE::Net::FNetRefHandle UEngineReplicationBridge::StartReplicatingActor(AActor* A
 		return FNetRefHandle::GetInvalid();
 	}
 
-	if (!ensureMsgf(Actor->HasActorBegunPlay() || Actor->IsActorBeginningPlay(), TEXT("Actor %s hasn't begun or isn't beginning play in BeginReplication."), ToCStr(GetFullNameSafe(Actor))))
+	if (!ensureMsgf(Actor->HasActorBegunPlay() || Actor->IsActorBeginningPlay(), TEXT("Actor %s hasn't begun or isn't beginning play in OnReplicationStartedForIris."), ToCStr(GetFullNameSafe(Actor))))
 	{
-		UE_LOG_BRIDGEID(LogIris, Error, TEXT("Actor %s hasn't begun or isn't beginning play in BeginReplication."), ToCStr(GetFullNameSafe(Actor)));
+		UE_LOG_BRIDGEID(LogIris, Error, TEXT("Actor %s hasn't begun or isn't beginning play in OnReplicationStartedForIris."), ToCStr(GetFullNameSafe(Actor)));
 		return FNetRefHandle::GetInvalid();
 	}
 
