FExternalDataLayerHelper::ForEachExternalDataLayerLevelPackagePath - Fix indeterministic asset registry queries in world partition cooking

diff --git a/Engine/Source/Runtime/Engine/Private/WorldPartition/DataLayer/ExternalDataLayerHelper.cpp b/Engine/Source/Runtime/Engine/Private/WorldPartition/DataLayer/ExternalDataLayerHelper.cpp
--- a/Engine/Source/Runtime/Engine/Private/WorldPartition/DataLayer/ExternalDataLayerHelper.cpp
+++ b/Engine/Source/Runtime/Engine/Private/WorldPartition/DataLayer/ExternalDataLayerHelper.cpp
@@ -382,6 +382,7 @@ void FExternalDataLayerHelper::ForEachExternalDataLayerLevelPackagePath(const FS
 		Filter.bIncludeOnlyOnDiskAssets = false;
 		Filter.ClassPaths = { GameFeatureDataClass->GetClassPathName() };
 		Filter.bRecursivePaths = true;
+		Filter.bIncludeOnlyOnDiskAssets = IsRunningCookCommandlet();	// When cooking, only enumerate on-disk assets- avoids indeterministic results
 		TArray<FAssetData> AssetsData;
 		FExternalPackageHelper::GetSortedAssets(Filter, AssetsData);
 
