Virtual texture material translation fix.
Handle edge case where the multiple duplicated virtual texture samples happen.
We relied on the shader code chunk deduplication to handle this.
But that deduplication didn't prevent virtual stack and layer allocation.
That was potentially wasteful for allocating SRV slots.
But also it didn't handle the case where the duplicated samples were processed after a virtual texture stack containing one had been "filled".
This case would prevent the code chunk deduplication completely.
Only one of the non dedplicated layer settings would persist to the UniformExpressionSet which led to sampling bugs.
Also tweaked some local variable naming and comments.
#rnx

diff --git a/Engine/Source/Runtime/Core/Private/UObject/FortniteMainBranchObjectVersion.cpp b/Engine/Source/Runtime/Core/Private/UObject/FortniteMainBranchObjectVersion.cpp
--- a/Engine/Source/Runtime/Core/Private/UObject/FortniteMainBranchObjectVersion.cpp
+++ b/Engine/Source/Runtime/Core/Private/UObject/FortniteMainBranchObjectVersion.cpp
@@ -13,7 +13,7 @@ TMap<FGuid, FGuid> FFortniteMainBranchObjectVersion::GetSystemGuids()
 	SystemGuids.Add(DevGuids.MATERIALSHADERMAP_DERIVEDDATA_VER, FGuid("F051877974B44D409435F321BDB39DBB"));
 	SystemGuids.Add(DevGuids.NIAGARASHADERMAP_DERIVEDDATA_VER, FGuid("6360A977062842A29ED30E3A7ACB0E64"));
 	SystemGuids.Add(DevGuids.Niagara_LatestScriptCompileVersion, FGuid("783FF3700BB641C799B2D43552FBF91E"));
-	SystemGuids.Add(DevGuids.MaterialTranslationDDCVersion, FGuid("92421BCD0BFF479F90F4DB0B76C18C55"));
+	SystemGuids.Add(DevGuids.MaterialTranslationDDCVersion, FGuid("448CDD72B0A34496BDED50908658DCBE"));
 
 	// These GUIDs are in AppendRestrictedFortniteMainStreamObjectVersionGuids as they should rarely be changed.
 	//SystemGuids.Add(DevGuids.NANITE_DERIVEDDATA_VER, FGuid("ACF98184978742A592537DA223A1E6BE"));
diff --git a/Engine/Source/Runtime/Engine/Private/Materials/HLSLMaterialTranslator.cpp b/Engine/Source/Runtime/Engine/Private/Materials/HLSLMaterialTranslator.cpp
--- a/Engine/Source/Runtime/Engine/Private/Materials/HLSLMaterialTranslator.cpp
+++ b/Engine/Source/Runtime/Engine/Private/Materials/HLSLMaterialTranslator.cpp
@@ -7017,6 +7017,7 @@ void FHLSLMaterialTranslator::SetPotentiallyManipulateTexCoords()
 }
 
 uint32 FHLSLMaterialTranslator::AcquireVTStackIndex(
+	int32 VirtualTexureIndex,
 	ETextureMipValueMode MipValueMode, 
 	TextureAddress AddressU, TextureAddress AddressV, 
 	float AspectRatio, 
@@ -7062,7 +7063,10 @@ uint32 FHLSLMaterialTranslator::AcquireVTStackIndex(
 	{
 		const FMaterialVirtualTextureStack& Stack = MaterialCompilationOutput.UniformExpressionSet.VTStacks[Index];
 		const FMaterialVTStackEntry& Entry = VTStacks[Index];
-		if (!Stack.AreLayersFull() &&
+		
+		const bool bStackIsFull = Stack.AreLayersFull() && Stack.FindLayer(VirtualTexureIndex) == INDEX_NONE;
+		
+		if (!bStackIsFull &&
 			Entry.ScopeID == CurrentScopeID &&
 			Entry.CoordinateHash == CoordinateHash &&
 			Entry.MipValue0Hash == MipValue0Hash &&
@@ -7810,9 +7814,11 @@ int32 FHLSLMaterialTranslator::TextureSample(
 	int32 SamplingCodeIndex = INDEX_NONE;
 	if (bVirtualTexture)
 	{
-		int32 VTStackIndex = INDEX_NONE;
-		int32 VTLayerIndex = INDEX_NONE;
-		int32 VTPageTableIndex = INDEX_NONE;
+		// Index of the page table to sample.
+		int32 VTStackIndex = INDEX_NONE; 
+		// Index of the channel in the page table to sample.
+		int32 VTPageTableLayerIndex = INDEX_NONE; 
+		// Name for binding of the page table uniform.
 		FString VTPackedUniformName;
 
 		// Only support GPU feedback from pixel shader
@@ -7821,14 +7827,14 @@ int32 FHLSLMaterialTranslator::TextureSample(
 
 		if (bMeshPaint)
 		{
-			VTStackIndex = AcquireVTStackIndex(MipValueMode, StaticAddressX, StaticAddressY, 1.0f, CoordinateIndex, MipValue0Index, MipValue1Index, INDEX_NONE, UV_Value, UV_Ddx, UV_Ddy, PageTableOverride, PageTableUniformOverride, AdaptiveVirtualTexture, bGenerateFeedback, bMeshPaint, bMaterialCache);
-			VTPageTableIndex = 0;
+			VTStackIndex = AcquireVTStackIndex(VirtualTextureIndex, MipValueMode, StaticAddressX, StaticAddressY, 1.0f, CoordinateIndex, MipValue0Index, MipValue1Index, INDEX_NONE, UV_Value, UV_Ddx, UV_Ddy, PageTableOverride, PageTableUniformOverride, AdaptiveVirtualTexture, bGenerateFeedback, bMeshPaint, bMaterialCache);
+			VTPageTableLayerIndex = 0;
 			VTPackedUniformName = TEXT("Scene.MeshPaint.PackedUniform");
 		}
 		else if (bMaterialCache)
 		{
-			VTStackIndex = AcquireVTStackIndex(MipValueMode, StaticAddressX, StaticAddressY, 1.0f, CoordinateIndex, MipValue0Index, MipValue1Index, INDEX_NONE, UV_Value, UV_Ddx, UV_Ddy, PageTableOverride, PageTableUniformOverride, AdaptiveVirtualTexture, bGenerateFeedback, bMeshPaint, bMaterialCache);
-			VTPageTableIndex = 0;
+			VTStackIndex = AcquireVTStackIndex(VirtualTextureIndex, MipValueMode, StaticAddressX, StaticAddressY, 1.0f, CoordinateIndex, MipValue0Index, MipValue1Index, INDEX_NONE, UV_Value, UV_Ddx, UV_Ddy, PageTableOverride, PageTableUniformOverride, AdaptiveVirtualTexture, bGenerateFeedback, bMeshPaint, bMaterialCache);
+			VTPageTableLayerIndex = 0;
 
 			FMaterialUniformExpression* UniformExpression = GetParameterUniformExpression(TextureIndex); 
 			if (!UniformExpression || !UniformExpression->GetTextureUniformExpression())
@@ -7844,14 +7850,15 @@ int32 FHLSLMaterialTranslator::TextureSample(
 			check(VirtualTextureIndex >= 0);
 			VTPackedUniformName = FString::Printf(TEXT("Material.VTPackedUniform[%d]"), VirtualTextureIndex);
 
-			VTLayerIndex = UniformTextureExpressions[(uint32)EMaterialTextureParameterType::Virtual][VirtualTextureIndex]->GetTextureLayerIndex();
-
-			if (VTLayerIndex != INDEX_NONE)
+			const int32 VTTextureLayerIndex = UniformTextureExpressions[(uint32)EMaterialTextureParameterType::Virtual][VirtualTextureIndex]->GetTextureLayerIndex();
+			if (VTTextureLayerIndex != INDEX_NONE)
 			{
 				// The layer index in the virtual texture stack is already known
 				// Create a page table sample for each new combination of virtual texture and sample parameters
-				VTStackIndex = AcquireVTStackIndex(MipValueMode, StaticAddressX, StaticAddressY, 1.0f, NonLWCCoordinateIndex, MipValue0Index, MipValue1Index, TextureReferenceIndex, UV_Value, UV_Ddx, UV_Ddy, PageTableOverride, PageTableUniformOverride, AdaptiveVirtualTexture, bGenerateFeedback, bMeshPaint, bMaterialCache);
-				VTPageTableIndex = UniformTextureExpressions[(uint32)EMaterialTextureParameterType::Virtual][VirtualTextureIndex]->GetPageTableLayerIndex();
+				VTStackIndex = AcquireVTStackIndex(VirtualTextureIndex, MipValueMode, StaticAddressX, StaticAddressY, 1.0f, NonLWCCoordinateIndex, MipValue0Index, MipValue1Index, TextureReferenceIndex, UV_Value, UV_Ddx, UV_Ddy, PageTableOverride, PageTableUniformOverride, AdaptiveVirtualTexture, bGenerateFeedback, bMeshPaint, bMaterialCache);
+				VTPageTableLayerIndex = UniformTextureExpressions[(uint32)EMaterialTextureParameterType::Virtual][VirtualTextureIndex]->GetPageTableLayerIndex();
+		
+				MaterialCompilationOutput.UniformExpressionSet.VTStacks[VTStackIndex].SetLayer(VTTextureLayerIndex, VirtualTextureIndex);
 			}
 			else
 			{
@@ -7866,13 +7873,16 @@ int32 FHLSLMaterialTranslator::TextureSample(
 				const float TextureAspectRatio = (float)Texture->Source.GetSizeX() / (float)Texture->Source.GetSizeY();
 
 				// Create a page table sample for each new set of sample parameters
-				VTStackIndex = AcquireVTStackIndex(MipValueMode, StaticAddressX, StaticAddressY, TextureAspectRatio, NonLWCCoordinateIndex, MipValue0Index, MipValue1Index, INDEX_NONE, UV_Value, UV_Ddx, UV_Ddy, PageTableOverride, PageTableUniformOverride, AdaptiveVirtualTexture, bGenerateFeedback, bMeshPaint, bMaterialCache);
-				// Allocate a layer in the virtual texture stack for this physical sample
-				VTLayerIndex = MaterialCompilationOutput.UniformExpressionSet.VTStacks[VTStackIndex].AddLayer();
-				VTPageTableIndex = VTLayerIndex;
+				VTStackIndex = AcquireVTStackIndex(VirtualTextureIndex, MipValueMode, StaticAddressX, StaticAddressY, TextureAspectRatio, NonLWCCoordinateIndex, MipValue0Index, MipValue1Index, INDEX_NONE, UV_Value, UV_Ddx, UV_Ddy, PageTableOverride, PageTableUniformOverride, AdaptiveVirtualTexture, bGenerateFeedback, bMeshPaint, bMaterialCache);
+
+				// Find or allocate the layer in the virtual texture stack for this physical sample
+				VTPageTableLayerIndex = MaterialCompilationOutput.UniformExpressionSet.VTStacks[VTStackIndex].FindLayer(VirtualTextureIndex);
+				if (VTPageTableLayerIndex == INDEX_NONE)
+				{
+					VTPageTableLayerIndex = MaterialCompilationOutput.UniformExpressionSet.VTStacks[VTStackIndex].AddLayer();
+					MaterialCompilationOutput.UniformExpressionSet.VTStacks[VTStackIndex].SetLayer(VTPageTableLayerIndex, VirtualTextureIndex);
+				}
 			}
-		
-			MaterialCompilationOutput.UniformExpressionSet.VTStacks[VTStackIndex].SetLayer(VTLayerIndex, VirtualTextureIndex);
 		}
 
 		// VT MipValueMode logic (most of work for VT case is in page table lookup)
@@ -7885,7 +7895,7 @@ int32 FHLSLMaterialTranslator::TextureSample(
 		const FString VTPageTableResult_Finite = GetParameterCode(VTStackEntry.CodeIndex);
 
 		// 'Texture name/sampler', 'PageTableResult', 'LayerIndex', 'PackedUniform'
-		FString SampleCodeFinite = FString::Printf(TEXT("%s(%s, %s, %d, VTUniform_Unpack(%s))"), *TextureTypeName, *TextureName, *VTPageTableResult_Finite, VTPageTableIndex, *VTPackedUniformName);
+		FString SampleCodeFinite = FString::Printf(TEXT("%s(%s, %s, %d, VTUniform_Unpack(%s))"), *TextureTypeName, *TextureName, *VTPageTableResult_Finite, VTPageTableLayerIndex, *VTPackedUniformName);
 		SampleCodeFinite = ApplySamplerType(SampleCodeFinite, SamplerType);
 
 		if (IsAnalyticDerivEnabled() && IsDerivativeValid(UvDerivativeStatus))
@@ -7896,11 +7906,11 @@ int32 FHLSLMaterialTranslator::TextureSample(
 
 			if (SamplerDebugSupported(TextureType, bVirtualTexture, bDecal))
 			{
-				SampleCodeAnalytic = FString::Printf(TEXT("Debug%s(%s, %s, %d, VTUniform_Unpack(%s), %s)"), *TextureTypeName, *TextureName, *VTPageTableResult_Analytic, VTPageTableIndex, *VTPackedUniformName, *UV_Scale);
+				SampleCodeAnalytic = FString::Printf(TEXT("Debug%s(%s, %s, %d, VTUniform_Unpack(%s), %s)"), *TextureTypeName, *TextureName, *VTPageTableResult_Analytic, VTPageTableLayerIndex, *VTPackedUniformName, *UV_Scale);
 			}
 			else
 			{
-				SampleCodeAnalytic = FString::Printf(TEXT("%s(%s, %s, %d, VTUniform_Unpack(%s))"), *TextureTypeName, *TextureName, *VTPageTableResult_Analytic, VTPageTableIndex, *VTPackedUniformName);
+				SampleCodeAnalytic = FString::Printf(TEXT("%s(%s, %s, %d, VTUniform_Unpack(%s))"), *TextureTypeName, *TextureName, *VTPageTableResult_Analytic, VTPageTableLayerIndex, *VTPackedUniformName);
 			}
 
 			SampleCodeAnalytic = ApplySamplerType(SampleCodeAnalytic, SamplerType);
@@ -7911,7 +7921,7 @@ int32 FHLSLMaterialTranslator::TextureSample(
 		}
 		else
 		{
-			SamplingCodeIndex = AddCodeChunk(MCT_Float4, *SampleCodeFinite, *TextureName, *VTPageTableResult_Finite, VTPageTableIndex, VirtualTextureIndex);
+			SamplingCodeIndex = AddCodeChunk(MCT_Float4, *SampleCodeFinite, *TextureName, *VTPageTableResult_Finite, VTPageTableLayerIndex, VirtualTextureIndex);
 		}
 	}
 	else
diff --git a/Engine/Source/Runtime/Engine/Private/Materials/HLSLMaterialTranslator.h b/Engine/Source/Runtime/Engine/Private/Materials/HLSLMaterialTranslator.h
--- a/Engine/Source/Runtime/Engine/Private/Materials/HLSLMaterialTranslator.h
+++ b/Engine/Source/Runtime/Engine/Private/Materials/HLSLMaterialTranslator.h
@@ -1056,7 +1056,8 @@ protected:
 	virtual void SetPotentiallyManipulateTexCoords() override;
 
 	uint32 AcquireVTStackIndex(
-		ETextureMipValueMode MipValueMode, 
+		int32 VirtualTexureIndex,
+		ETextureMipValueMode MipValueMode,
 		TextureAddress AddressU, TextureAddress AddressV, 
 		float AspectRatio, 
 		int32 CoordinateIndex, 
