Added VirtualTextureCollection's, a separate texture collection object with virtual texture backing
- Internally hosts a tile-wise atlas of the texture set. Single physical texture, all respective texture entries must share the same format and srgb mode. Alternatively, runtime format conversion is supported.
- Texture entries are naturally (mip wise) block aligned. Due to block alignment, we can avoid expensive UV re-mapping and manual sampler addressing by modifying the page table uniforms.
- Format conversion supports both virtual and regular textures, and block (re)compression should it be needed.
- Due to the extensive amount of pipelining needed for another collection parameter type, it's easier to extend the current type. For in-type bindings each parameter has a prefix index.
- Pending sharing of page table between collections

(Reviewed in 43854629, changed workspace)

diff --git a/Engine/Shaders/Private/MaterialTemplate.ush b/Engine/Shaders/Private/MaterialTemplate.ush
--- a/Engine/Shaders/Private/MaterialTemplate.ush
+++ b/Engine/Shaders/Private/MaterialTemplate.ush
@@ -116,6 +116,7 @@
 
 #if NUM_VIRTUALTEXTURE_SAMPLES || LIGHTMAP_VT_ENABLED
 	#include "/Engine/Private/VirtualTextureCommon.ush"
+	#include "/Engine/Public/IndirectVirtualTextureCommon.ush"
 #endif
 
 #include "/Engine/Private/MaterialTexture.ush"
diff --git a/Engine/Shaders/Public/IndirectVirtualTexture.ush b/Engine/Shaders/Public/IndirectVirtualTexture.ush
new file mode 100644
--- /dev/null
+++ b/Engine/Shaders/Public/IndirectVirtualTexture.ush
@@ -0,0 +1,7 @@
+﻿// Copyright Epic Games, Inc. All Rights Reserved.
+
+#pragma once
+
+#include "/Engine/Shared/IndirectVirtualTextureDefinitions.h"
+
+typedef Buffer<uint2> FIndirectVirtualTexture;
diff --git a/Engine/Shaders/Public/IndirectVirtualTextureCommon.ush b/Engine/Shaders/Public/IndirectVirtualTextureCommon.ush
new file mode 100644
--- /dev/null
+++ b/Engine/Shaders/Public/IndirectVirtualTextureCommon.ush
@@ -0,0 +1,76 @@
+﻿// Copyright Epic Games, Inc. All Rights Reserved.
+
+#pragma once
+
+#include "/Engine/Public/IndirectVirtualTexture.ush"
+#include "/Engine/Private/VirtualTextureCommon.ush"
+
+struct FIndirectVirtualTextureInfo
+{
+	uint BlockX;
+	uint BlockY;
+	uint BlockWidth;
+	uint BlockHeight;
+};
+
+FIndirectVirtualTextureUniform GetIndirectVirtualTextureUniform(uint Index)
+{
+	// Uniform layouts with struct's doesn't appear supported, so, load
+	// it manually for now. When it's supported, this can be removed.
+	uint Offset = IndirectVirtualTextureUniformDWord4Count * Index;
+	
+	FIndirectVirtualTextureUniform Out = (FIndirectVirtualTextureUniform)0;
+#if TEXTURE_COLLECTION_PACKED_UNIFORMS
+	Out.UniformCountSub1          = Material.TextureCollectionPackedUniforms[Offset + 0].w;
+	Out.PackedPageTableUniform[0] = Material.TextureCollectionPackedUniforms[Offset + 1];
+	Out.PackedPageTableUniform[1] = Material.TextureCollectionPackedUniforms[Offset + 2];
+	Out.PackedUniform             = Material.TextureCollectionPackedUniforms[Offset + 3];
+#endif // TEXTURE_COLLECTION_PACKED_UNIFORMS
+	return Out;
+}
+
+uint GetIndirectVirtualTextureCount(in FIndirectVirtualTextureUniform Uniform)
+{
+	return Uniform.UniformCountSub1 + 1;
+}
+
+FIndirectVirtualTextureEntry LoadIndirectVirtualTexture(in FIndirectVirtualTexture Collection, in FIndirectVirtualTextureUniform Uniform, uint CollectionIndex)
+{
+	CollectionIndex = min(Uniform.UniformCountSub1, CollectionIndex);
+
+	FIndirectVirtualTextureEntry Out = (FIndirectVirtualTextureEntry)0;
+	Out.PackedCoordinateAndSize = Collection[CollectionIndex];
+	return Out;
+}
+
+FIndirectVirtualTextureInfo UnpackIndirectVirtualTexture(in FIndirectVirtualTextureEntry Entry)
+{
+	FIndirectVirtualTextureInfo Out = (FIndirectVirtualTextureInfo)0;
+	Out.BlockX      = Entry.PackedCoordinateAndSize.x  & 0xFFFF;
+	Out.BlockY      = Entry.PackedCoordinateAndSize.x >> 16;
+	Out.BlockWidth  = Entry.PackedCoordinateAndSize.y  & 0xFFFF;
+	Out.BlockHeight = Entry.PackedCoordinateAndSize.y >> 16;
+	return Out;
+}
+
+VTPageTableUniform UnpackIndirectVirtualTexturePageTableUniform(in FIndirectVirtualTextureEntry Entry, uint4 A, uint4 B)
+{
+	VTPageTableUniform Out = VTPageTableUniform_Unpack(A, B);
+	
+	FIndirectVirtualTextureInfo Info = UnpackIndirectVirtualTexture(Entry);
+
+	// Allow packing textures of different sizes and in turn number of mip levels
+	// So, to avoid lower physical mip levels "leaking" in, limit the level
+	// to that of this entry only.
+	// TODO[MP]: We could probably pack this as part of PackedCoordinateAndSize.y
+	Out.MaxLevel = log2(max(Info.BlockWidth, Info.BlockHeight));
+	
+	// This avoids both:
+	//  1. Re-mapping the UV coordinates manually from 0-1 to the physical region
+	//  2. Manual sampler addressing handling on the texture region borders
+	Out.SizeInPages     = uint2(Info.BlockWidth, Info.BlockHeight);
+	Out.XOffsetInPages += Info.BlockX;
+	Out.YOffsetInPages += Info.BlockY;
+	
+	return Out;
+}
diff --git a/Engine/Shaders/Public/Platform.ush b/Engine/Shaders/Public/Platform.ush
--- a/Engine/Shaders/Public/Platform.ush
+++ b/Engine/Shaders/Public/Platform.ush
@@ -67,6 +67,7 @@
 
 #include "/Engine/Public/BindlessResources.ush"
 #include "/Engine/Shared/MaterialCacheDefinitions.h"
+#include "/Engine/Public/IndirectVirtualTexture.ush"
 #include "/Engine/Public/OverloadMacros.ush"
 
 
diff --git a/Engine/Shaders/Shared/IndirectVirtualTextureDefinitions.h b/Engine/Shaders/Shared/IndirectVirtualTextureDefinitions.h
new file mode 100644
--- /dev/null
+++ b/Engine/Shaders/Shared/IndirectVirtualTextureDefinitions.h
@@ -0,0 +1,29 @@
+// Copyright Epic Games, Inc. All Rights Reserved.
+
+#pragma once
+
+#ifdef __cplusplus
+#include "HLSLTypeAliases.h"
+namespace UE::HLSL
+{
+#endif  // __cplusplus
+
+struct FIndirectVirtualTextureUniform
+{
+	uint3 Pad;
+	uint  UniformCountSub1;
+	uint4 PackedPageTableUniform[2];
+	uint4 PackedUniform;
+};
+
+struct FIndirectVirtualTextureEntry
+{
+	// Could be packed in one uint if we accept 4k textures as max
+	uint2 PackedCoordinateAndSize;
+};
+
+static const uint IndirectVirtualTextureUniformDWord4Count = 4u;
+
+#ifdef __cplusplus
+} // namespace
+#endif // __cplusplus
diff --git a/Engine/Source/Runtime/Core/Private/UObject/FortniteMainBranchObjectVersion.cpp b/Engine/Source/Runtime/Core/Private/UObject/FortniteMainBranchObjectVersion.cpp
--- a/Engine/Source/Runtime/Core/Private/UObject/FortniteMainBranchObjectVersion.cpp
+++ b/Engine/Source/Runtime/Core/Private/UObject/FortniteMainBranchObjectVersion.cpp
@@ -13,7 +13,7 @@ TMap<FGuid, FGuid> FFortniteMainBranchObjectVersion::GetSystemGuids()
 	SystemGuids.Add(DevGuids.MATERIALSHADERMAP_DERIVEDDATA_VER, FGuid("F051877974B44D409435F321BDB39DBB"));
 	SystemGuids.Add(DevGuids.NIAGARASHADERMAP_DERIVEDDATA_VER, FGuid("6360A977062842A29ED30E3A7ACB0E64"));
 	SystemGuids.Add(DevGuids.Niagara_LatestScriptCompileVersion, FGuid("783FF3700BB641C799B2D43552FBF91E"));
-	SystemGuids.Add(DevGuids.MaterialTranslationDDCVersion, FGuid("448CDD72B0A34496BDED50908658DCBE"));
+	SystemGuids.Add(DevGuids.MaterialTranslationDDCVersion, FGuid("F9BB134100414BE4AA156DEF0C01C4DF"));
 
 	// These GUIDs are in AppendRestrictedFortniteMainStreamObjectVersionGuids as they should rarely be changed.
 	//SystemGuids.Add(DevGuids.NANITE_DERIVEDDATA_VER, FGuid("ACF98184978742A592537DA223A1E6BE"));
diff --git a/Engine/Source/Runtime/Engine/Classes/Engine/Texture.h b/Engine/Source/Runtime/Engine/Classes/Engine/Texture.h
--- a/Engine/Source/Runtime/Engine/Classes/Engine/Texture.h
+++ b/Engine/Source/Runtime/Engine/Classes/Engine/Texture.h
@@ -1150,6 +1150,8 @@ struct FTextureFormatSettings
 		, SRGB(false)
 	{}
 
+	bool operator==(const FTextureFormatSettings& other) const = default;
+
 	UPROPERTY()
 	TEnumAsByte<enum TextureCompressionSettings> CompressionSettings;
 
diff --git a/Engine/Source/Runtime/Engine/Classes/Engine/TextureCollection.h b/Engine/Source/Runtime/Engine/Classes/Engine/TextureCollection.h
--- a/Engine/Source/Runtime/Engine/Classes/Engine/TextureCollection.h
+++ b/Engine/Source/Runtime/Engine/Classes/Engine/TextureCollection.h
@@ -14,17 +14,32 @@ struct FTextureCollectionResource : public FRenderResource
 {
 	FTextureCollectionResource(UTextureCollection* InParent);
 
+	/** Is this a bindless collection? */
+	bool bIsBindless = false;
+	
+	/** Name of the parent collection */
+	FName CollectionName;
+};
+
+struct FBindlessTextureCollectionResource : public FTextureCollectionResource
+{
+	FBindlessTextureCollectionResource(UTextureCollection* InParent);
+	virtual ~FBindlessTextureCollectionResource() override = default;
+
 	// FRenderResource Interface
 	virtual void InitRHI(FRHICommandListBase& RHICmdList) final;
 	virtual void ReleaseRHI() final;
 	// ~FRenderResource
 
-	FRHIResourceCollection* GetRHI() const { return ResourceCollectionRHI.GetReference(); }
-
+	FRHIResourceCollection* GetResourceCollectionRHI() const
+	{
+		check(ResourceCollectionRHI);
+		return ResourceCollectionRHI.GetReference();
+	}
+	
 private:
-	TArray<FTextureResource*>       InputTextureResources;
-	TArray<FTextureReferenceRHIRef> InputTextures;
-	FRHIResourceCollectionRef       ResourceCollectionRHI;
+	TArray<FTextureResource*> InputTextureResources;
+	FRHIResourceCollectionRef ResourceCollectionRHI;
 };
 
 UCLASS(MinimalAPI)
@@ -47,10 +62,12 @@ public:
 	ENGINE_API FTextureCollectionResource* GetResource() const;
 	ENGINE_API FTextureCollectionResource* GetResource();
 
-	ENGINE_API FTextureCollectionResource* CreateResource();
+	virtual ENGINE_API FTextureCollectionResource* CreateResource();
 	ENGINE_API void ReleaseResource();
 	ENGINE_API void UpdateResource();
 
+	virtual bool IsVirtualCollection() const { return false; }
+	
 protected:
 #if WITH_EDITOR
 	void NotifyMaterials();
diff --git a/Engine/Source/Runtime/Engine/Classes/Engine/VirtualTextureCollection.h b/Engine/Source/Runtime/Engine/Classes/Engine/VirtualTextureCollection.h
new file mode 100644
--- /dev/null
+++ b/Engine/Source/Runtime/Engine/Classes/Engine/VirtualTextureCollection.h
@@ -0,0 +1,158 @@
+// Copyright Epic Games, Inc. All Rights Reserved.
+
+#pragma once
+
+#include "IndirectVirtualTextureDefinitions.h"
+#include "Engine/DataAsset.h"
+#include "RenderResource.h"
+#include "IndirectVirtualTextureDefinitions.h"
+#include "Containers/DynamicRHIResourceArray.h"
+#include "Texture2D.h"
+#include "TextureCollection.h"
+#include "VirtualTexturing.h"
+#include "VT/VirtualTextureBuildSettings.h"
+#include "VirtualTextureCollection.generated.h"
+
+class UVirtualTextureAdapter;
+class FTextureResource;
+class UTexture;
+class UTextureCollection;
+struct FTextureCollectionProducerData;
+
+struct FVirtualTextureCollectionResource : public FTextureCollectionResource
+{
+	FVirtualTextureCollectionResource(UVirtualTextureCollection* InParent);
+	virtual ~FVirtualTextureCollectionResource() override = default;
+
+public: /** FRenderResource */
+	virtual void InitRHI(FRHICommandListBase& RHICmdList) final;
+	virtual void ReleaseRHI() override;
+
+	struct FTextureEntry
+	{
+		// Optional, the virtual producer
+		FVirtualTextureProducerHandle VirtualProducerHandle;
+
+		// Optional, physical texture
+		// Keep the GT object around since RT resources may be released
+		UTexture* PhysicalTexture = nullptr;
+
+		// If true, this entry will be converted internally to the physical format
+		// This also includes compression
+		bool bRequiresAdapter = false;
+
+		// The format of this texture, before adaptation
+		EPixelFormat Format = PF_Unknown;
+
+		// Number of blocks in this entry
+		FUintVector2 BlockCount;
+	};
+	
+	UE::HLSL::FIndirectVirtualTextureUniform GetVirtualPackedUniform() const;
+	
+	FRHIShaderResourceView* GetVirtualCollectionRHI() const
+	{
+		check(VirtualCollectionRHISRV);
+		return VirtualCollectionRHISRV;
+	}
+
+	FRHITexture* GetVirtualPageTable() const
+	{
+		return PageTable;
+	}
+	
+	FRHIShaderResourceView* GetVirtualPhysicalTextureSRV() const
+	{
+		return PhysicalTextureSRV;
+	}
+
+	FVirtualTextureProducerHandle GetProducerHandle() const
+	{
+		return ProducerHandle;
+	}
+	
+private:
+	/** Find the first applicable format */
+	void FindFirstFormat();
+	
+	/** Find a conservative format for the collection */
+	void FindConservativeFormat();
+
+	/** Compute the texture block layout */
+	void ComputeLayout(FTextureCollectionProducerData& Data);
+	
+	/** Create the host side index table */
+	void CreateIndexTable(FTextureCollectionProducerData& Data);
+
+	/** Helper for error formatting */
+	void FormatCollectionError(const TCHAR* Reason, uint32 TextureIndex);
+
+private:
+	TArray<UTexture*> Textures;
+
+	/** Virtual build settings */
+	FVirtualTextureBuildSettings BuildSettings;
+
+	/** Software (compared to the bindless collection) collection buffer */
+	FBufferRHIRef             VirtualCollectionRHI;
+	FRHIShaderResourceViewRef VirtualCollectionRHISRV;
+
+	/** Virtual resources are currently limited to a single physical texture and associated page table */
+	FRHITexture*            PageTable          = nullptr;
+	FRHIShaderResourceView* PhysicalTextureSRV = nullptr;
+
+	/** Actual virtual table */
+	IAllocatedVirtualTexture*     AllocatedVT = nullptr;
+	FVirtualTextureProducerHandle ProducerHandle;
+	
+	TResourceArray<UE::HLSL::FIndirectVirtualTextureEntry> VirtualUniforms;
+
+	EPixelFormat Format         = PF_Unknown;
+	bool bIsSRGB                = false;
+	bool bAllowFormatConversion = false;
+};
+
+UCLASS(MinimalAPI)
+class UVirtualTextureCollection : public UTextureCollection
+{
+	GENERATED_BODY()
+
+public:	
+	/**
+	 * Allow format conversions, including differing compression schemes.
+	 * This has a potentially large runtime overhead.
+	 **/
+	UPROPERTY(EditAnywhere, Category=VirtualTextureCollection)
+	bool bAllowFormatConversion = true;
+
+	/**
+	 * If this texture collection if in SRGB, requires format conversion.
+	 * Textures not matching this will be converted at runtime.
+	 */
+	UPROPERTY(EditAnywhere, Category=VirtualTextureCollection, meta = (editcondition = "bAllowFormatConversion"))
+	bool bIsSRGB = true;
+
+	/**
+	 * The pixel format chosen at runtime, should PixelFormat not be specified.
+	 **/
+	UPROPERTY(VisibleAnywhere, Category=VirtualTextureCollection, AdvancedDisplay)
+	TEnumAsByte<EPixelFormat> RuntimePixelFormat = PF_Unknown;
+	
+public: /** UObject */
+#if WITH_EDITOR
+	virtual void PostEditChangeProperty(FPropertyChangedEvent& PropertyChangedEvent) override;
+#endif // WITH_EDITOR
+	virtual void BeginDestroy() override;
+
+	/** UTextureCollection */
+	virtual FTextureCollectionResource* CreateResource() override;
+	virtual bool IsVirtualCollection() const override { return true; }
+
+protected:
+	friend struct FVirtualTextureCollectionResource;
+	
+#if WITH_EDITOR
+	void ValidateVirtualCollection();
+	void FormatCollectionError(const TCHAR* Reason, uint32 TextureIndex);
+#endif //  WITH_EDITOR
+};
diff --git a/Engine/Source/Runtime/Engine/Classes/VT/CopyCompressShader.h b/Engine/Source/Runtime/Engine/Classes/VT/CopyCompressShader.h
new file mode 100644
--- /dev/null
+++ b/Engine/Source/Runtime/Engine/Classes/VT/CopyCompressShader.h
@@ -0,0 +1,43 @@
+﻿// Copyright Epic Games, Inc. All Rights Reserved.
+
+#pragma once
+
+#include "GlobalShader.h"
+#include "SceneView.h"
+#include "ShaderParameterMacros.h"
+#include "ShaderParameterStruct.h"
+
+/** Copy with compression shader. */
+class FCopyCompressCS : public FGlobalShader
+{
+public:
+	DECLARE_GLOBAL_SHADER(FCopyCompressCS);
+	SHADER_USE_PARAMETER_STRUCT(FCopyCompressCS, FGlobalShader);
+
+	BEGIN_SHADER_PARAMETER_STRUCT(FParameters, )
+		SHADER_PARAMETER_SRV(Texture2D, SourceTextureA)
+		SHADER_PARAMETER_RDG_TEXTURE_SRV(Texture2D, SourceTextureB)
+		SHADER_PARAMETER_SAMPLER(SamplerState, TextureSampler)
+		SHADER_PARAMETER_RDG_TEXTURE_UAV(RWTexture2D, DestTexture)
+		SHADER_PARAMETER_RDG_TEXTURE_UAV(RWTexture2D, DestCompressTexture_64bit)
+		SHADER_PARAMETER_RDG_TEXTURE_UAV(RWTexture2D, DestCompressTexture_128bit)
+		SHADER_PARAMETER(FVector2f, SourceUV)
+		SHADER_PARAMETER(FVector2f, TexelSize)
+		SHADER_PARAMETER(FVector2f, TexelOffsets)
+		SHADER_PARAMETER(FIntVector4, DestRect)
+	END_SHADER_PARAMETER_STRUCT()
+
+	static constexpr int32 GroupSize = 8;
+
+	class FSourceTextureSelector : SHADER_PERMUTATION_BOOL("SOURCE_TEXTURE_A");
+	class FDestSrgb : SHADER_PERMUTATION_BOOL("TEXTURE_SRGB");
+	class FCompressionFormatDim : SHADER_PERMUTATION_INT("COMPRESSION_FORMAT", 7);
+	using FPermutationDomain = TShaderPermutationDomain<FSourceTextureSelector, FDestSrgb, FCompressionFormatDim>;
+
+	/** Get index to use for FCompressionFormatDim. */
+	static int32 GetCompressionPermutation(EPixelFormat InFormat)
+	{
+		const EPixelFormat CompressedFormats[] = { PF_DXT1, PF_DXT5, PF_BC4, PF_BC5, PF_BC6H, PF_BC7 };
+		return MakeArrayView(CompressedFormats).Find(InFormat) + 1;
+	}
+};
diff --git a/Engine/Source/Runtime/Engine/Private/Materials/HLSLMaterialTranslator.cpp b/Engine/Source/Runtime/Engine/Private/Materials/HLSLMaterialTranslator.cpp
--- a/Engine/Source/Runtime/Engine/Private/Materials/HLSLMaterialTranslator.cpp
+++ b/Engine/Source/Runtime/Engine/Private/Materials/HLSLMaterialTranslator.cpp
@@ -212,11 +212,12 @@ struct FHLSLMaterialTranslator::FEnvironmentDefines
 {
 	enum EVirtualPageType
 	{
-		TABLE_MESH_PAINT = 1,
-		TABLE_MATERIAL_CACHE = 2,
-		LOCAL_TABLE0 = 4,
-		LOCAL_TABLE1 = 8,
-		LOCAL_ADAPTIVE_INDIRECTION = 16
+		TABLE_MESH_PAINT           = (1 << 0),
+		TABLE_MATERIAL_CACHE       = (1 << 1),
+		TABLE_COLLECTION           = (1 << 2),
+		LOCAL_TABLE0               = (1 << 3),
+		LOCAL_TABLE1               = (1 << 4),
+		LOCAL_ADAPTIVE_INDIRECTION = (1 << 5)
 	};
 
 	bool bNeedsParticlePosition;
@@ -2880,6 +2881,10 @@ void FHLSLMaterialTranslator::GetMaterialEnvironment(EShaderPlatform InPlatform,
 		{
 			// Driven manually
 		}
+		else if (EnvironmentDefines->VirtualPageTypes[i] & FEnvironmentDefines::TABLE_COLLECTION)
+		{
+			/** Page table fetched during sampling */
+		}
 		else
 		{
 			// Setup page table defines to map each VT stack to either 1 or 2 page table textures, depending on how many layers it uses
@@ -3088,6 +3093,15 @@ void FHLSLMaterialTranslator::GetMaterialEnvironment(EShaderPlatform InPlatform,
 	}
 
 	OutEnvironment.SetDefine(TEXT("TEXTURE_SAMPLE_DEBUG"), EnvironmentDefines->bTextureSampleDebug);
+
+	for (const TRefCountPtr<FMaterialUniformExpressionTextureCollection>& TextureCollection : UniformTextureCollectionExpressions)
+	{
+		if (TextureCollection->IsVirtualCollection())
+		{
+			OutEnvironment.SetDefine(TEXT("TEXTURE_COLLECTION_PACKED_UNIFORMS"), 1u);
+			break;
+		}
+	}
 }
 
 // Assign custom interpolators to slots, packing them as much as possible in unused slots.
@@ -3333,7 +3347,7 @@ FString FHLSLMaterialTranslator::GetParameterCodeDeriv(int32 Index, ECompiledPar
 	}
 	else
 	{
-		if (CodeChunk.UniformExpression)
+		if (CodeChunk.UniformExpression && !CodeChunk.bIntermediate)
 		{
 			// If the code chunk has a uniform expression, create a new code chunk to access it
 			const int32 AccessedIndex = AccessUniformExpression(Index);
@@ -3391,8 +3405,7 @@ FString FHLSLMaterialTranslator::GetParameterCodeRaw(int32 Index, const TCHAR* D
 	}
 	else
 	{
-		// TODO - don't access uniform expressions for intermediate values
-		if (CodeChunk.UniformExpression)
+		if (CodeChunk.UniformExpression && !CodeChunk.bIntermediate)
 		{
 			// If the code chunk has a uniform expression, create a new code chunk to access it
 			const int32 AccessedIndex = AccessUniformExpression(Index);
@@ -3468,7 +3481,7 @@ FString FHLSLMaterialTranslator::GetDefinitions(const TArray<FShaderCodeChunk>&
 	{
 		const FShaderCodeChunk& CodeChunk = CodeChunks[ChunkIndex];
 		// Uniform expressions (both constant and variable) and inline expressions don't have definitions.
-		if (!CodeChunk.UniformExpression &&
+		if ((!CodeChunk.UniformExpression || CodeChunk.bIntermediate) &&
 			(!CodeChunk.bInline || CodeChunk.Type == MCT_VoidStatement))
 		{
 			Definitions += CodeChunk.AtDefinition(Variation);
@@ -3850,7 +3863,7 @@ int32 FHLSLMaterialTranslator::AddCodeChunkInner(uint64 Hash, const TCHAR* Forma
 		new(*CurrentScopeChunks) FShaderCodeChunk(Hash, FormattedCode, FormattedCode, TEXT(""), Type, DerivativeStatus, true);
 	}
 	// Can only create temporaries for certain types
-	else if ((Type & (MCT_Float | MCT_LWCType | MCT_VTPageTableResult | MCT_UInt)) || Type == MCT_ShadingModel || Type == MCT_MaterialAttributes || Type == MCT_Substrate || Type == MCT_UInt || bFromTextureCollection || Type == MCT_MaterialCacheABuffer)
+	else if ((Type & (MCT_Float | MCT_LWCType | MCT_VTPageTableResult | MCT_Unexposed | MCT_UInt)) || Type == MCT_ShadingModel || Type == MCT_MaterialAttributes || Type == MCT_Substrate || Type == MCT_UInt || bFromTextureCollection || Type == MCT_MaterialCacheABuffer)
 	{
 		// Check for existing
 		for (int32 i = 0; i < CurrentScopeChunks->Num(); ++i)
@@ -3911,6 +3924,40 @@ static inline uint32 GetTCharStringBytes(const TCHAR* String)
 	return Length * sizeof(TCHAR);
 }
 
+int32 FHLSLMaterialTranslator::AddInternalCodeChunk(EMaterialValueType Type, const TCHAR* InternalType, const TCHAR* FormattedCode, const TRefCountPtr<FMaterialUniformExpression>& UniformExpression)
+{
+	const uint64 Hash = CityHash64(reinterpret_cast<const char*>(FormattedCode), GetTCharStringBytes(FormattedCode));
+	
+	int32 CodeIndex = INDEX_NONE;
+	for (int32 i = 0; i < CurrentScopeChunks->Num(); ++i)
+	{
+		if ((*CurrentScopeChunks)[i].Hash == Hash)
+		{
+			CodeIndex = i;
+			break;
+		}
+	}
+	
+	if (CodeIndex == INDEX_NONE)
+	{
+		CodeIndex = CurrentScopeChunks->Num();
+		const FString SymbolName = CreateSymbolName(TEXT("Local"));
+		const FString LocalVariableDefinitionFinite = FString("	") + InternalType + TEXT(" ") + SymbolName + TEXT(" = ") + FormattedCode + TEXT(";") + HLSL_LINE_TERMINATOR;
+		const FString LocalVariableDefinitionAnalytic = FString("	") + InternalType + TEXT(" ") + SymbolName + TEXT(" = ") + FormattedCode + TEXT(";") + HLSL_LINE_TERMINATOR;
+
+		FShaderCodeChunk* Chunk = new(*CurrentScopeChunks) FShaderCodeChunk(
+			Hash, *LocalVariableDefinitionFinite, *LocalVariableDefinitionAnalytic,
+			SymbolName, Type, EDerivativeStatus::NotAware, false
+		);
+
+		Chunk->bIntermediate     = true;
+		Chunk->UniformExpression = UniformExpression;
+	}
+	
+	AddCodeChunkToCurrentScope(CodeIndex);
+	return CodeIndex;
+}
+
 /** Adds an already formatted inline or referenced code chunk, and notes the derivative status. */
 int32 FHLSLMaterialTranslator::AddCodeChunkInnerDeriv(const TCHAR* FormattedCodeFinite, const TCHAR* FormattedCodeAnalytic, EMaterialValueType Type, bool bInlined, EDerivativeStatus DerivativeStatus)
 {
@@ -4400,6 +4447,18 @@ int32 FHLSLMaterialTranslator::AccessUniformExpression(int32 Index)
 	}
 	else if (CodeChunk.Type == MCT_TextureCollection)
 	{
+		// If unallocated, assign a type prefix
+		if (TextureCollectionUniformExpression->GetTextureCollectionTypePrefixIndex() == UINT32_MAX)
+		{
+			uint32 TypePrefix = 0;
+			for (int32 i = 0; i < UniformTextureCollectionExpressions.Num(); i++)
+			{
+				TypePrefix += UniformTextureCollectionExpressions[i]->IsVirtualCollection() == TextureCollectionUniformExpression->IsVirtualCollection();
+			}
+
+			TextureCollectionUniformExpression->SetTextureCollectionTypePrefixIndex(TypePrefix);
+		}
+		
 		int32 TextureCollectionInputIndex = UniformTextureCollectionExpressions.AddUnique(TextureCollectionUniformExpression);
 		FormattedCode.Appendf(TEXT("Material.TextureCollection_%u"), TextureCollectionInputIndex);
 	}
@@ -7017,45 +7076,35 @@ void FHLSLMaterialTranslator::SetPotentiallyManipulateTexCoords()
 }
 
 uint32 FHLSLMaterialTranslator::AcquireVTStackIndex(
-	int32 VirtualTexureIndex,
-	ETextureMipValueMode MipValueMode, 
-	TextureAddress AddressU, TextureAddress AddressV, 
-	float AspectRatio, 
-	int32 CoordinateIndex, 
-	int32 MipValue0Index, int32 MipValue1Index, 
-	int32 PreallocatedStackTextureIndex, 
+	const FMaterialVTSampleInfo& SampleInfo,
 	const FString& UV_Value,
 	const FString& UV_Ddx,
-	const FString& UV_Ddy,
-	const FString& PageTableOverride,
-	const FString& PageTableUniformOverride,
-	bool bAdaptive, bool bGenerateFeedback, bool bMeshPaint, bool bMaterialCache)
+	const FString& UV_Ddy)
 {
 	using namespace UE::MaterialTranslatorUtils;
 
-	const uint64 CoordinateHash = GetParameterHash(CoordinateIndex);
-	const uint64 MipValue0Hash = GetParameterHash(MipValue0Index);
-	const uint64 MipValue1Hash = GetParameterHash(MipValue1Index);
-
-	const uint64 PageTableHash        = CityHash64(reinterpret_cast<const char*>(*PageTableOverride),        PageTableOverride.Len()        * sizeof(TCHAR));
-	const uint64 PageTableUniformHash = CityHash64(reinterpret_cast<const char*>(*PageTableUniformOverride), PageTableUniformOverride.Len() * sizeof(TCHAR));
+	const uint64 CoordinateHash = GetParameterHash(SampleInfo.CoordinateIndex);
+	const uint64 MipValue0Hash = GetParameterHash(SampleInfo.MipValue0Index);
+	const uint64 MipValue1Hash = GetParameterHash(SampleInfo.MipValue1Index);
+	const uint64 PageTableHash = GetParameterHash(SampleInfo.PageTableIndex);
+	const uint64 PageTableUniformHash = GetParameterHash(SampleInfo.PageTableUniformIndex);
 	
 	uint32 Flags = 0;
-	Flags |= bAdaptive ? (1 << 0) : 0;
-	Flags |= bGenerateFeedback ? (1 << 1) : 0;
-	Flags |= bMeshPaint ? (1 << 2) : 0;
-	Flags |= bMaterialCache ? (1 << 3) : 0;
+	Flags |= SampleInfo.bAdaptive ? (1 << 0) : 0;
+	Flags |= SampleInfo.bGenerateFeedback ? (1 << 1) : 0;
+	Flags |= SampleInfo.bMeshPaint ? (1 << 2) : 0;
+	Flags |= SampleInfo.bMaterialCache ? (1 << 3) : 0;
 
 	uint64 Hash = CityHash128to64({ CurrentScopeID, CoordinateHash });
 	Hash = CityHash128to64({ Hash, MipValue0Hash });
 	Hash = CityHash128to64({ Hash, MipValue1Hash });
 	Hash = CityHash128to64({ Hash, PageTableHash });
 	Hash = CityHash128to64({ Hash, PageTableUniformHash });
-	Hash = CityHash128to64({ Hash, (uint64)MipValueMode });
-	Hash = CityHash128to64({ Hash, (uint64)AddressU });
-	Hash = CityHash128to64({ Hash, (uint64)AddressV });
-	Hash = CityHash128to64({ Hash, (uint64)(AspectRatio * 1000.0f) });
-	Hash = CityHash128to64({ Hash, (uint64)PreallocatedStackTextureIndex });
+	Hash = CityHash128to64({ Hash, (uint64)SampleInfo.MipValueMode });
+	Hash = CityHash128to64({ Hash, (uint64)SampleInfo.AddressU });
+	Hash = CityHash128to64({ Hash, (uint64)SampleInfo.AddressV });
+	Hash = CityHash128to64({ Hash, (uint64)(SampleInfo.AspectRatio * 1000.0f) });
+	Hash = CityHash128to64({ Hash, (uint64)SampleInfo.PreallocatedStackTextureIndex });
 	Hash = CityHash128to64({ Hash, (uint64)Flags });
 
 	// First check to see if we have an existing VTStack that matches this key, that can still fit another layer
@@ -7064,7 +7113,7 @@ uint32 FHLSLMaterialTranslator::AcquireVTStackIndex(
 		const FMaterialVirtualTextureStack& Stack = MaterialCompilationOutput.UniformExpressionSet.VTStacks[Index];
 		const FMaterialVTStackEntry& Entry = VTStacks[Index];
 		
-		const bool bStackIsFull = Stack.AreLayersFull() && Stack.FindLayer(VirtualTexureIndex) == INDEX_NONE;
+		const bool bStackIsFull = Stack.AreLayersFull() && Stack.FindLayer(SampleInfo.VirtualTextureIndex) == INDEX_NONE;
 		
 		if (!bStackIsFull &&
 			Entry.ScopeID == CurrentScopeID &&
@@ -7073,15 +7122,7 @@ uint32 FHLSLMaterialTranslator::AcquireVTStackIndex(
 			Entry.MipValue1Hash == MipValue1Hash &&
 			Entry.PageTableHash == PageTableHash &&
 			Entry.PageTableUniformHash == PageTableUniformHash &&
-			Entry.MipValueMode == MipValueMode &&
-			Entry.AddressU == AddressU &&
-			Entry.AddressV == AddressV &&
-			Entry.AspectRatio == AspectRatio &&
-			Entry.PreallocatedStackTextureIndex == PreallocatedStackTextureIndex &&
-			Entry.bAdaptive == bAdaptive &&
-			Entry.bGenerateFeedback == bGenerateFeedback &&
-			Entry.bMeshPaint == bMeshPaint &&
-			Entry.bMaterialCache == bMaterialCache)
+			Entry.SampleInfo == SampleInfo)
 		{
 			return Index;
 		}
@@ -7091,72 +7132,59 @@ uint32 FHLSLMaterialTranslator::AcquireVTStackIndex(
 	const int32 StackIndex = VTStacks.AddDefaulted();
 	VTStackHash.Add(Hash, StackIndex);
 	FMaterialVTStackEntry& Entry = VTStacks[StackIndex];
+	Entry.SampleInfo = SampleInfo;
 	Entry.ScopeID = CurrentScopeID;
 	Entry.CoordinateHash = CoordinateHash;
 	Entry.MipValue0Hash = MipValue0Hash;
 	Entry.MipValue1Hash = MipValue1Hash;
 	Entry.PageTableHash = PageTableHash;
 	Entry.PageTableUniformHash = PageTableUniformHash;
-	Entry.MipValueMode = MipValueMode;
-	Entry.AddressU = AddressU;
-	Entry.AddressV = AddressV;
-	Entry.AspectRatio = AspectRatio;
-	Entry.DebugCoordinateIndex = CoordinateIndex;
-	Entry.DebugMipValue0Index = MipValue0Index;
-	Entry.DebugMipValue1Index = MipValue1Index;
-	Entry.PreallocatedStackTextureIndex = PreallocatedStackTextureIndex;
-	Entry.bAdaptive = bAdaptive;
-	Entry.bGenerateFeedback = bGenerateFeedback;
-	Entry.bMeshPaint = bMeshPaint;
-	Entry.bMaterialCache = bMaterialCache;
-
-	MaterialCompilationOutput.UniformExpressionSet.VTStacks.Add(FMaterialVirtualTextureStack(PreallocatedStackTextureIndex));
+
+	MaterialCompilationOutput.UniformExpressionSet.VTStacks.Add(FMaterialVirtualTextureStack(SampleInfo.PreallocatedStackTextureIndex));
 
 	// These two arrays need to stay in sync
 	check(VTStacks.Num() == MaterialCompilationOutput.UniformExpressionSet.VTStacks.Num());
 
 	// Select LoadVirtualPageTable function name for this context
-	FString BaseFunctionName = bAdaptive ? TEXT("TextureLoadVirtualPageTableAdaptive") : TEXT("TextureLoadVirtualPageTable");
+	FString BaseFunctionName = SampleInfo.bAdaptive ? TEXT("TextureLoadVirtualPageTableAdaptive") : TEXT("TextureLoadVirtualPageTable");
 
 	// Optionally sample without virtual texture feedback but only for miplevel mode
-	check(bGenerateFeedback || MipValueMode == TMVM_MipLevel)
-	FString FeedbackParameter = bGenerateFeedback ? FString::Printf(TEXT(", Parameters.VirtualTextureFeedback")) : TEXT("");
+	check(SampleInfo.bGenerateFeedback || SampleInfo.MipValueMode == TMVM_MipLevel)
+	FString FeedbackParameter = SampleInfo.bGenerateFeedback ? FString::Printf(TEXT(", Parameters.VirtualTextureFeedback")) : TEXT("");
 
-	EDerivativeStatus UVDerivativeStatus = GetDerivativeStatus(CoordinateIndex);
+	EDerivativeStatus UVDerivativeStatus = GetDerivativeStatus(SampleInfo.CoordinateIndex);
 	const bool bHasValidDerivative = IsAnalyticDerivEnabled() && IsDerivativeValid(UVDerivativeStatus);
 
-	FString PageTable        = PageTableOverride.IsEmpty()        ? FString::Printf(TEXT("VIRTUALTEXTURE_PAGETABLE_%d"),         StackIndex) : PageTableOverride;
-	FString PageTableUniform = PageTableUniformOverride.IsEmpty() ? FString::Printf(TEXT("VIRTUALTEXTURE_PAGETABLE_UNIFORM_%d"), StackIndex) : PageTableUniformOverride;
+	FString PageTable        = SampleInfo.PageTableIndex        != INDEX_NONE ? GetParameterCode(SampleInfo.PageTableIndex)        : FString::Printf(TEXT("VIRTUALTEXTURE_PAGETABLE_%d"), StackIndex);
+	FString PageTableUniform = SampleInfo.PageTableUniformIndex != INDEX_NONE ? GetParameterCode(SampleInfo.PageTableUniformIndex) : FString::Printf(TEXT("VTPageTableUniform_Unpack(VIRTUALTEXTURE_PAGETABLE_UNIFORM_%d)"), StackIndex);
 
 	// Code to load the VT page table...this will execute the first time a given VT stack is accessed
 	// Additional stack layers will simply reuse these results
-	switch (MipValueMode)
+	switch (SampleInfo.MipValueMode)
 	{
 	case TMVM_None:
 	{
 		FString SampleCodeFinite = FString::Printf(TEXT(
 			"%s("
-			"%s, "
-			"VTPageTableUniform_Unpack(%s), "
+			"%s, %s, "
 			"%s, %s, %s, "
 			"0, Parameters.SvPosition.xy, "
 			"Parameters.VirtualTextureFeedback)"),
 			*BaseFunctionName,
 			*PageTable, *PageTableUniform,
-			*CoerceParameter(CoordinateIndex, MCT_Float2), GetVTAddressMode(AddressU), GetVTAddressMode(AddressV));
+			*CoerceParameter(SampleInfo.CoordinateIndex, MCT_Float2), GetVTAddressMode(SampleInfo.AddressU), GetVTAddressMode(SampleInfo.AddressV));
 
 		if (bHasValidDerivative)
 		{
 			FString SampleCodeAnalytic = FString::Printf(TEXT(
 				"%sGrad("
-				"%s, "
-				"VTPageTableUniform_Unpack(%s), "
+				"%s, %s, "
 				"%s, %s, %s, "
 				"%s, %s, Parameters.SvPosition.xy, "
 				"Parameters.VirtualTextureFeedback)"),
 				*BaseFunctionName,
 				*PageTable, *PageTableUniform,
-				*UV_Value, GetVTAddressMode(AddressU), GetVTAddressMode(AddressV), 
+				*UV_Value, GetVTAddressMode(SampleInfo.AddressU), GetVTAddressMode(SampleInfo.AddressV), 
 				*UV_Ddx, *UV_Ddy);
 			Entry.CodeIndex = AddCodeChunkInnerDeriv(*SampleCodeFinite, *SampleCodeAnalytic, MCT_VTPageTableResult, false, EDerivativeStatus::NotValid);
 		}
@@ -7171,27 +7199,25 @@ uint32 FHLSLMaterialTranslator::AcquireVTStackIndex(
 	{
 		FString SampleCodeFinite = FString::Printf(TEXT(
 			"%s("
-			"%s, "
-			"VTPageTableUniform_Unpack(%s), "
+			"%s, %s, "
 			"%s, %s, %s, "
 			"%s, Parameters.SvPosition.xy, "
 			"Parameters.VirtualTextureFeedback)"),
 			*BaseFunctionName,
 			*PageTable, *PageTableUniform,
-			*CoerceParameter(CoordinateIndex, MCT_Float2), GetVTAddressMode(AddressU), GetVTAddressMode(AddressV), 
-			*CoerceParameter(MipValue0Index, MCT_Float1));
+			*CoerceParameter(SampleInfo.CoordinateIndex, MCT_Float2), GetVTAddressMode(SampleInfo.AddressU), GetVTAddressMode(SampleInfo.AddressV), 
+			*CoerceParameter(SampleInfo.MipValue0Index, MCT_Float1));
 		if (bHasValidDerivative)
 		{
 			FString SampleCodeAnalytic = FString::Printf(TEXT(
 				"%sGrad("
-				"%s, "
-				"VTPageTableUniform_Unpack(%s), "
+				"%s, %s, "
 				"%s, %s, %s, "
 				"%s, %s, Parameters.SvPosition.xy, "
 				"Parameters.VirtualTextureFeedback)"),
 				*BaseFunctionName,
 				*PageTable, *PageTableUniform,
-				*UV_Value, GetVTAddressMode(AddressU), GetVTAddressMode(AddressV),
+				*UV_Value, GetVTAddressMode(SampleInfo.AddressU), GetVTAddressMode(SampleInfo.AddressV),
 				*UV_Ddx, *UV_Ddy);
 			Entry.CodeIndex = AddCodeChunkInnerDeriv(*SampleCodeFinite, *SampleCodeAnalytic, MCT_VTPageTableResult, false, EDerivativeStatus::NotValid);
 		}
@@ -7204,29 +7230,27 @@ uint32 FHLSLMaterialTranslator::AcquireVTStackIndex(
 	case TMVM_MipLevel:
 		Entry.CodeIndex = AddCodeChunk(MCT_VTPageTableResult, TEXT(
 			"%sLevel("
-			"%s, "
-			"VTPageTableUniform_Unpack(%s), "
+			"%s, %s, "
 			"%s, %s, %s, "
 			"%s"
 			"%s)"),
 			*BaseFunctionName,
 			*PageTable, *PageTableUniform,
-			*CoerceParameter(CoordinateIndex, MCT_Float2), GetVTAddressMode(AddressU), GetVTAddressMode(AddressV), 
-			*CoerceParameter(MipValue0Index, MCT_Float1),
+			*CoerceParameter(SampleInfo.CoordinateIndex, MCT_Float2), GetVTAddressMode(SampleInfo.AddressU), GetVTAddressMode(SampleInfo.AddressV), 
+			*CoerceParameter(SampleInfo.MipValue0Index, MCT_Float1),
 			*FeedbackParameter);
 		break;
 	case TMVM_Derivative:
 		Entry.CodeIndex = AddCodeChunk(MCT_VTPageTableResult, TEXT(
 			"%sGrad("
-			"%s, "
-			"VTPageTableUniform_Unpack(%s), "
+			"%s, %s, "
 			"%s, %s, %s, "
 			"%s, %s, Parameters.SvPosition.xy, "
 			"Parameters.VirtualTextureFeedback)"),
 			*BaseFunctionName,
 			*PageTable, *PageTableUniform,
-			*CoerceParameter(CoordinateIndex, MCT_Float2), GetVTAddressMode(AddressU), GetVTAddressMode(AddressV),
-			*CoerceParameter(MipValue0Index, MCT_Float2), *CoerceParameter(MipValue1Index, MCT_Float2));
+			*CoerceParameter(SampleInfo.CoordinateIndex, MCT_Float2), GetVTAddressMode(SampleInfo.AddressU), GetVTAddressMode(SampleInfo.AddressV),
+			*CoerceParameter(SampleInfo.MipValue0Index, MCT_Float2), *CoerceParameter(SampleInfo.MipValue1Index, MCT_Float2));
 		break;
 	default:
 		checkNoEntry();
@@ -7343,9 +7367,11 @@ int32 FHLSLMaterialTranslator::TextureSample(
 		return INDEX_NONE;
 	}
 
+	FMaterialVTSampleInfo SampleInfo;
+	
 	const bool bMeshPaint = TextureType == MCT_TextureMeshPaint;
 	const bool bMaterialCache = TextureType == MCT_TextureMaterialCache;
-	const bool bVirtualTexture = bMeshPaint || bMaterialCache || TextureType == MCT_TextureVirtual;
+	const bool bVirtualTexture = bMeshPaint || bMaterialCache || (TextureType & MCT_TextureVirtual) != 0;
 
 	if (MipValueMode == TMVM_Derivative)
 	{
@@ -7421,9 +7447,6 @@ int32 FHLSLMaterialTranslator::TextureSample(
 		AutomaticViewMipBias = false;
 	}
 
-	FString PageTableOverride;
-	FString PageTableUniformOverride;
-
 	FString TextureName;
 	int32 VirtualTextureIndex = INDEX_NONE;
 	if (TextureType == MCT_TextureCube)
@@ -7469,8 +7492,16 @@ int32 FHLSLMaterialTranslator::TextureSample(
 		// Binding index is tied to discovery order
 		uint32 MaterialCacheTagIndex = TextureUniform->GetMaterialCacheTagIndex();
 
-		PageTableOverride        = FString::Printf(TEXT("Material.MaterialCacheTagPageTable0_%d"), MaterialCacheTagIndex);
-		PageTableUniformOverride = FString::Printf(TEXT("%s.PackedUniform"), *GetParameterCode(TextureUniform->GetTextureIndex()));
+		SampleInfo.PageTableIndex = AddInternalCodeChunk(
+			MCT_Unexposed, TEXT("Texture2D<uint4>"),
+			*FString::Printf(TEXT("Material.MaterialCacheTagPageTable0_%d"), MaterialCacheTagIndex)
+		);
+		
+		SampleInfo.PageTableUniformIndex = AddCodeChunkInner(
+			MCT_Unexposed, TEXT("VTPageTableUniform"),
+			EDerivativeStatus::NotAware, false,
+			TEXT("VTPageTableUniform_Unpack(%s.PackedUniform)"), *GetParameterCode(TextureUniform->GetTextureIndex())
+		);
 
 		// Physical texture indexing assumes TagIndex * Stride + LayerIndex
 		TextureName = FString::Printf(
@@ -7480,6 +7511,32 @@ int32 FHLSLMaterialTranslator::TextureSample(
 
 		NumVtSamples++;
 	}
+	else if (bFromCollection && bVirtualTexture)
+	{
+		CoerceParameter(TextureIndex, TextureType);
+		
+		FMaterialUniformExpression* UniformExpression = GetParameterUniformExpression(TextureIndex); 
+		if (!UniformExpression || !UniformExpression->GetTextureCollectionUniformExpression())
+		{
+			return Errorf(TEXT("The provided uniform expression is not from a texture collection"));
+		}
+		
+		FMaterialUniformExpressionTextureCollection* CollectionExpression = UniformExpression->GetTextureCollectionUniformExpression();
+		VirtualTextureIndex = CollectionExpression->GetTextureCollectionTypePrefixIndex();
+
+		if (SamplerSource != SSM_FromTextureAsset)
+		{
+			const bool   bUseAnisoSampler  = VirtualTextureScalability::IsAnisotropicFilteringEnabled() && MipValueMode != TMVM_MipLevel;
+			const TCHAR* SharedSamplerName = bUseAnisoSampler ? TEXT("View.SharedBilinearAnisoClampedSampler") : TEXT("View.SharedBilinearClampedSampler");
+			TextureName += FString::Printf(TEXT("Material.TextureCollectionPhysical_%d, GetMaterialSharedSampler(Material.TextureCollectionSampler_%d, %s)"), VirtualTextureIndex, VirtualTextureIndex, SharedSamplerName);
+		}
+		else
+		{
+			TextureName += FString::Printf(TEXT("Material.TextureCollectionPhysical_%d, Material.TextureCollectionSampler_%d"), VirtualTextureIndex, VirtualTextureIndex);
+		}
+		
+		NumVtSamples++;
+	}
 	else if (bVirtualTexture)
 	{
 		// Note, this does not really do anything (by design) other than adding it to the UniformExpressionSet
@@ -7793,7 +7850,7 @@ int32 FHLSLMaterialTranslator::TextureSample(
 		AddCodeChunk(MCT_Float, TEXT("MaterialStoreTexCoordScale(Parameters, %s, %d)"), *CoerceParameter(NonLWCCoordinateIndex, UVsType), (int)TextureReferenceIndex);
 	}
 
-	const FString UV_Value = CoerceParameter(NonLWCCoordinateIndex, UVsType);
+	FString UV_Value = CoerceParameter(NonLWCCoordinateIndex, UVsType);
 	FString UV_Ddx = TEXT("0.0f");
 	FString UV_Ddy = TEXT("0.0f");
 	FString UV_Scale = TEXT("1.0f");
@@ -7825,15 +7882,29 @@ int32 FHLSLMaterialTranslator::TextureSample(
 		//todo[vt]: Support feedback from other shader types
 		const bool bGenerateFeedback = EnableFeedback && ShaderFrequency == SF_Pixel;
 
+		SampleInfo.MipValueMode = MipValueMode;
+		SampleInfo.AddressU = StaticAddressX;
+		SampleInfo.AddressV = StaticAddressY;
+		SampleInfo.AspectRatio = 1.0f;
+		SampleInfo.CoordinateIndex = CoordinateIndex;
+		SampleInfo.MipValue0Index = MipValue0Index;
+		SampleInfo.MipValue1Index = MipValue1Index;
+		SampleInfo.bAdaptive = AdaptiveVirtualTexture;
+		SampleInfo.bGenerateFeedback = bGenerateFeedback;
+		SampleInfo.bMeshPaint = bMeshPaint;
+		SampleInfo.bMaterialCache = bMaterialCache;
+		SampleInfo.bCollection = bFromCollection;
+        SampleInfo.VirtualTextureIndex = VirtualTextureIndex;
+
 		if (bMeshPaint)
 		{
-			VTStackIndex = AcquireVTStackIndex(VirtualTextureIndex, MipValueMode, StaticAddressX, StaticAddressY, 1.0f, CoordinateIndex, MipValue0Index, MipValue1Index, INDEX_NONE, UV_Value, UV_Ddx, UV_Ddy, PageTableOverride, PageTableUniformOverride, AdaptiveVirtualTexture, bGenerateFeedback, bMeshPaint, bMaterialCache);
+			VTStackIndex = AcquireVTStackIndex(SampleInfo, UV_Value, UV_Ddx, UV_Ddy);
 			VTPageTableLayerIndex = 0;
-			VTPackedUniformName = TEXT("Scene.MeshPaint.PackedUniform");
+			VTPackedUniformName = TEXT("VTUniform_Unpack(Scene.MeshPaint.PackedUniform)");
 		}
 		else if (bMaterialCache)
 		{
-			VTStackIndex = AcquireVTStackIndex(VirtualTextureIndex, MipValueMode, StaticAddressX, StaticAddressY, 1.0f, CoordinateIndex, MipValue0Index, MipValue1Index, INDEX_NONE, UV_Value, UV_Ddx, UV_Ddy, PageTableOverride, PageTableUniformOverride, AdaptiveVirtualTexture, bGenerateFeedback, bMeshPaint, bMaterialCache);
+			VTStackIndex = AcquireVTStackIndex(SampleInfo, UV_Value, UV_Ddx, UV_Ddy);
 			VTPageTableLayerIndex = 0;
 
 			FMaterialUniformExpression* UniformExpression = GetParameterUniformExpression(TextureIndex); 
@@ -7843,21 +7914,53 @@ int32 FHLSLMaterialTranslator::TextureSample(
 			}
 
 			FMaterialUniformExpressionTexture* TextureUniform = UniformExpression->GetTextureUniformExpression();
-			VTPackedUniformName = FString::Printf(TEXT("Material.MaterialCacheTagPageTableUniform[%d]"), TextureUniform->GetMaterialCacheTagIndex());
+			VTPackedUniformName = FString::Printf(TEXT("VTUniform_Unpack(Material.MaterialCacheTagPageTableUniform[%d])"), TextureUniform->GetMaterialCacheTagIndex());
+		}
+		else if (bFromCollection)
+		{
+			FString TextureCode = GetParameterCode(TextureIndex);
+			
+			SampleInfo.PageTableIndex = AddInternalCodeChunk(
+				MCT_Unexposed, TEXT("Texture2D<uint4>"),
+				*FString::Printf(TEXT("Material.TextureCollectionPageTable_%i"), VirtualTextureIndex)
+			);
+
+			FString Uniforms = GetParameterCode(AddInternalCodeChunk(
+				MCT_Unexposed, TEXT("FIndirectVirtualTextureUniform"),
+				*FString::Printf(TEXT("GetIndirectVirtualTextureUniform(%i)"), VirtualTextureIndex)
+			));
+
+			// Manually unpack the page table to redirect the page boundaries
+			SampleInfo.PageTableUniformIndex = AddInternalCodeChunk(
+				MCT_Unexposed, TEXT("VTPageTableUniform"),
+				*FString::Printf(
+					TEXT("UnpackIndirectVirtualTexturePageTableUniform(%s, %s.PackedPageTableUniform[0], %s.PackedPageTableUniform[1])"),
+					*TextureCode, *Uniforms, *Uniforms
+				)
+			);
+
+			VTPackedUniformName = FString::Printf(TEXT("VTUniform_Unpack(%s.PackedUniform)"), *Uniforms);
+			
+			VTStackIndex = AcquireVTStackIndex(SampleInfo, UV_Value, UV_Ddx, UV_Ddy);
+			VTPageTableLayerIndex = 0;
 		}
 		else 
 		{
 			check(VirtualTextureIndex >= 0);
-			VTPackedUniformName = FString::Printf(TEXT("Material.VTPackedUniform[%d]"), VirtualTextureIndex);
+			VTPackedUniformName = FString::Printf(TEXT("VTUniform_Unpack(Material.VTPackedUniform[%d])"), VirtualTextureIndex);
+
+			SampleInfo.CoordinateIndex = NonLWCCoordinateIndex;
 
 			const int32 VTTextureLayerIndex = UniformTextureExpressions[(uint32)EMaterialTextureParameterType::Virtual][VirtualTextureIndex]->GetTextureLayerIndex();
 			if (VTTextureLayerIndex != INDEX_NONE)
 			{
+				SampleInfo.PreallocatedStackTextureIndex = TextureReferenceIndex;
+				
 				// The layer index in the virtual texture stack is already known
 				// Create a page table sample for each new combination of virtual texture and sample parameters
-				VTStackIndex = AcquireVTStackIndex(VirtualTextureIndex, MipValueMode, StaticAddressX, StaticAddressY, 1.0f, NonLWCCoordinateIndex, MipValue0Index, MipValue1Index, TextureReferenceIndex, UV_Value, UV_Ddx, UV_Ddy, PageTableOverride, PageTableUniformOverride, AdaptiveVirtualTexture, bGenerateFeedback, bMeshPaint, bMaterialCache);
+				VTStackIndex = AcquireVTStackIndex(SampleInfo, UV_Value, UV_Ddx, UV_Ddy);
 				VTPageTableLayerIndex = UniformTextureExpressions[(uint32)EMaterialTextureParameterType::Virtual][VirtualTextureIndex]->GetPageTableLayerIndex();
-		
+
 				MaterialCompilationOutput.UniformExpressionSet.VTStacks[VTStackIndex].SetLayer(VTTextureLayerIndex, VirtualTextureIndex);
 			}
 			else
@@ -7872,8 +7975,10 @@ int32 FHLSLMaterialTranslator::TextureSample(
 				// This is because we only need to overlay VT pages from within a given block
 				const float TextureAspectRatio = (float)Texture->Source.GetSizeX() / (float)Texture->Source.GetSizeY();
 
+				SampleInfo.AspectRatio = TextureAspectRatio;
+				
 				// Create a page table sample for each new set of sample parameters
-				VTStackIndex = AcquireVTStackIndex(VirtualTextureIndex, MipValueMode, StaticAddressX, StaticAddressY, TextureAspectRatio, NonLWCCoordinateIndex, MipValue0Index, MipValue1Index, INDEX_NONE, UV_Value, UV_Ddx, UV_Ddy, PageTableOverride, PageTableUniformOverride, AdaptiveVirtualTexture, bGenerateFeedback, bMeshPaint, bMaterialCache);
+				VTStackIndex = AcquireVTStackIndex(SampleInfo, UV_Value, UV_Ddx, UV_Ddy);
 
 				// Find or allocate the layer in the virtual texture stack for this physical sample
 				VTPageTableLayerIndex = MaterialCompilationOutput.UniformExpressionSet.VTStacks[VTStackIndex].FindLayer(VirtualTextureIndex);
@@ -7895,7 +8000,7 @@ int32 FHLSLMaterialTranslator::TextureSample(
 		const FString VTPageTableResult_Finite = GetParameterCode(VTStackEntry.CodeIndex);
 
 		// 'Texture name/sampler', 'PageTableResult', 'LayerIndex', 'PackedUniform'
-		FString SampleCodeFinite = FString::Printf(TEXT("%s(%s, %s, %d, VTUniform_Unpack(%s))"), *TextureTypeName, *TextureName, *VTPageTableResult_Finite, VTPageTableLayerIndex, *VTPackedUniformName);
+		FString SampleCodeFinite = FString::Printf(TEXT("%s(%s, %s, %d, %s)"), *TextureTypeName, *TextureName, *VTPageTableResult_Finite, VTPageTableLayerIndex, *VTPackedUniformName);
 		SampleCodeFinite = ApplySamplerType(SampleCodeFinite, SamplerType);
 
 		if (IsAnalyticDerivEnabled() && IsDerivativeValid(UvDerivativeStatus))
@@ -7906,11 +8011,11 @@ int32 FHLSLMaterialTranslator::TextureSample(
 
 			if (SamplerDebugSupported(TextureType, bVirtualTexture, bDecal))
 			{
-				SampleCodeAnalytic = FString::Printf(TEXT("Debug%s(%s, %s, %d, VTUniform_Unpack(%s), %s)"), *TextureTypeName, *TextureName, *VTPageTableResult_Analytic, VTPageTableLayerIndex, *VTPackedUniformName, *UV_Scale);
+				SampleCodeAnalytic = FString::Printf(TEXT("Debug%s(%s, %s, %d, %s, %s)"), *TextureTypeName, *TextureName, *VTPageTableResult_Analytic, VTPageTableLayerIndex, *VTPackedUniformName, *UV_Scale);
 			}
 			else
 			{
-				SampleCodeAnalytic = FString::Printf(TEXT("%s(%s, %s, %d, VTUniform_Unpack(%s))"), *TextureTypeName, *TextureName, *VTPageTableResult_Analytic, VTPageTableLayerIndex, *VTPackedUniformName);
+				SampleCodeAnalytic = FString::Printf(TEXT("%s(%s, %s, %d, %s)"), *TextureTypeName, *TextureName, *VTPageTableResult_Analytic, VTPageTableLayerIndex, *VTPackedUniformName);
 			}
 
 			SampleCodeAnalytic = ApplySamplerType(SampleCodeAnalytic, SamplerType);
@@ -8063,6 +8168,30 @@ int32 FHLSLMaterialTranslator::TextureFromCollection(int32 TextureCollectionCode
 		return Errorf(TEXT("Expected a texture return type"));
 	}
 
+	if (TextureCollectionExpression->IsVirtualCollection())
+	{
+		if (ResultTextureType != MCT_Texture2D)
+		{
+			return Errorf(TEXT("Texture collection virtual sampling only supports 2d textures"));
+		}
+		
+		FString Uniforms = GetParameterCode(AddInternalCodeChunk(
+			MCT_Unexposed, TEXT("FIndirectVirtualTextureUniform"),
+			*FString::Printf(TEXT("GetIndirectVirtualTextureUniform(%i)"), TextureCollectionExpression->GetTextureCollectionTypePrefixIndex())
+		));
+		
+		return AddInternalCodeChunk(
+			static_cast<EMaterialValueType>(ResultTextureType | MCT_TextureVirtual | MCT_TextureCollection | MCT_Unexposed), TEXT("FIndirectVirtualTextureEntry"),
+			*FString::Printf(
+				TEXT("LoadIndirectVirtualTexture(%s, %s, %s)"),
+				*GetParameterCode(TextureCollectionCodeIndex),
+				*Uniforms,
+				*GetParameterCode(IndexIntoCollectionCodeIndex)
+			),
+			TextureCollectionExpression
+		);
+	}
+	
 	return AddCodeChunk(
 		EMaterialValueType(ResultTextureType | MCT_TextureCollection),
 		TEXT("TextureFromCollection_%s(%s, %s)"),
@@ -8815,8 +8944,10 @@ int32 FHLSLMaterialTranslator::TextureCollection(UTextureCollection* InTextureCo
 		return INDEX_NONE;
 	}
 
+	const bool bIsVirtualCollection = InTextureCollection->IsVirtualCollection();
+
 	FAddUniformExpressionScope Scope(this);
-	return AddUniformExpression(Scope, new FMaterialUniformExpressionTextureCollection(TextureCollectionReferenceIndex), MCT_TextureCollection, TEXT(""));
+	return AddUniformExpression(Scope, new FMaterialUniformExpressionTextureCollection(TextureCollectionReferenceIndex, UINT32_MAX, bIsVirtualCollection), MCT_TextureCollection, TEXT(""));
 }
 
 int32 FHLSLMaterialTranslator::TextureCollectionParameter(FName ParameterName, UTextureCollection* InDefaultValue, int32& TextureCollectionReferenceIndex)
@@ -8839,11 +8970,13 @@ int32 FHLSLMaterialTranslator::TextureCollectionParameter(FName ParameterName, U
 		return INDEX_NONE;
 	}
 
+	const bool bIsVirtualCollection = InDefaultValue->IsVirtualCollection();
+
 	FMaterialParameterInfo ParameterInfo = GetParameterAssociationInfo();
 	ParameterInfo.Name = ParameterName;
 
 	FAddUniformExpressionScope Scope(this);
-	return AddUniformExpression(Scope, new FMaterialUniformExpressionTextureCollectionParameter(ParameterInfo, TextureCollectionReferenceIndex), MCT_TextureCollection, TEXT(""));
+	return AddUniformExpression(Scope, new FMaterialUniformExpressionTextureCollectionParameter(ParameterInfo, TextureCollectionReferenceIndex, UINT32_MAX, bIsVirtualCollection), MCT_TextureCollection, TEXT(""));
 }
 
 int32 FHLSLMaterialTranslator::TextureCollectionCount(int32 InTextureCollectionCodeIndex)
@@ -8864,7 +8997,19 @@ int32 FHLSLMaterialTranslator::TextureCollectionCount(int32 InTextureCollectionC
 	{
 		return Errorf(TEXT("Expected a texture collection expression"));
 	}
+	
+	if (TextureCollectionExpression->IsVirtualCollection())
+	{
+		// Make sure the uniform is present
+		AccessUniformExpression(InTextureCollectionCodeIndex);
 
+		return AddCodeChunk(
+			EMaterialValueType(MCT_UInt1),
+			TEXT("GetIndirectVirtualTextureCount(GetIndirectVirtualTextureUniform(%u))"),
+			TextureCollectionExpression->GetTextureCollectionTypePrefixIndex()
+		);
+	}
+	
 	return AddCodeChunk(
 		EMaterialValueType(MCT_UInt1),
 		TEXT("GetCountFromResourceCollection(%s)"),
@@ -16053,7 +16198,7 @@ void FHLSLMaterialTranslator::PrepareEnvironmentDefines()
 	uint32 NumVirtualTextureFeedbackRequests = 0;
 	for (int i = 0; i < VTStacks.Num(); ++i)
 	{
-		if (VTStacks[i].bGenerateFeedback)
+		if (VTStacks[i].SampleInfo.bGenerateFeedback)
 		{
 			++NumVirtualTextureFeedbackRequests;
 		}
@@ -16071,18 +16216,22 @@ void FHLSLMaterialTranslator::PrepareEnvironmentDefines()
 		{
 			EnvironmentDefines->VirtualPageTypes[i] |= FEnvironmentDefines::LOCAL_TABLE1;
 		}
-		if (VTStacks[i].bAdaptive)
+		if (VTStacks[i].SampleInfo.bAdaptive)
 		{
 			EnvironmentDefines->VirtualPageTypes[i] |= FEnvironmentDefines::LOCAL_ADAPTIVE_INDIRECTION;
 		}
-		if (VTStacks[i].bMeshPaint)
+		if (VTStacks[i].SampleInfo.bMeshPaint)
 		{
 			EnvironmentDefines->VirtualPageTypes[i] |= FEnvironmentDefines::TABLE_MESH_PAINT;
 		}
-		if (VTStacks[i].bMaterialCache)
+		if (VTStacks[i].SampleInfo.bMaterialCache)
 		{
 			EnvironmentDefines->VirtualPageTypes[i] |= FEnvironmentDefines::TABLE_MATERIAL_CACHE;
 		}
+		if (VTStacks[i].SampleInfo.bCollection)
+		{
+			EnvironmentDefines->VirtualPageTypes[i] |= FEnvironmentDefines::TABLE_COLLECTION;
+		}
 	}
 
 	// Set all the shading models for this material here 
diff --git a/Engine/Source/Runtime/Engine/Private/Materials/HLSLMaterialTranslator.h b/Engine/Source/Runtime/Engine/Private/Materials/HLSLMaterialTranslator.h
--- a/Engine/Source/Runtime/Engine/Private/Materials/HLSLMaterialTranslator.h
+++ b/Engine/Source/Runtime/Engine/Private/Materials/HLSLMaterialTranslator.h
@@ -123,6 +123,7 @@ struct FShaderCodeChunk
 
 	/** Whether the code chunk should be inlined or not.  If true, SymbolName is empty and Definition contains the code to inline. */
 	bool bInline;
+	bool bIntermediate;
 	/** The status of partial derivatives at this expression. **/
 	EDerivativeStatus DerivativeStatus;
 
@@ -151,6 +152,7 @@ struct FShaderCodeChunk
 		UsedScopeIndex(INDEX_NONE),
 		ScopeLevel(0),
 		bInline(bInInline),
+		bIntermediate(false),
 		DerivativeStatus(InDerivativeStatus)
 	{}
 
@@ -166,6 +168,7 @@ struct FShaderCodeChunk
 		UsedScopeIndex(INDEX_NONE),
 		ScopeLevel(0),
 		bInline(false),
+		bIntermediate(false),
 		DerivativeStatus(InDerivativeStatus)
 	{}
 };
@@ -235,28 +238,39 @@ class FHLSLMaterialTranslator : public FMaterialCompiler
 {
 	friend class FMaterialDerivativeAutogen;
 protected:
+	struct FMaterialVTSampleInfo
+	{
+		bool operator==(const FMaterialVTSampleInfo&) const = default;
+		
+		int32 VirtualTextureIndex = INDEX_NONE;
+		ETextureMipValueMode MipValueMode = TMVM_None;
+		TextureAddress AddressU = TA_Clamp;
+		TextureAddress AddressV = TA_Clamp;
+		int32 CoordinateIndex = INDEX_NONE;
+		int32 MipValue0Index = INDEX_NONE;
+		int32 MipValue1Index = INDEX_NONE;
+		int32 PreallocatedStackTextureIndex = INDEX_NONE;
+		int32 PageTableIndex = INDEX_NONE;
+		int32 PageTableUniformIndex = INDEX_NONE;
+		bool bAdaptive = false;
+		bool bGenerateFeedback = false;
+		bool bMeshPaint = false;
+		bool bMaterialCache = false;
+		bool bCollection = false;
+		float AspectRatio = 0.0f;
+		int32 CodeIndex = INDEX_NONE;
+	};
+
 	/** Translator specific structure during generation of VT stacks. This is later converted to FMaterialVirtualTextureStack entries in FUniformExpressionSet. */
 	struct FMaterialVTStackEntry
 	{
+		FMaterialVTSampleInfo SampleInfo;
 		uint64 ScopeID;
 		uint64 CoordinateHash;
 		uint64 MipValue0Hash;
 		uint64 MipValue1Hash;
 		uint64 PageTableHash;
 		uint64 PageTableUniformHash;
-		ETextureMipValueMode MipValueMode;
-		TextureAddress AddressU;
-		TextureAddress AddressV;
-		int32 DebugCoordinateIndex;
-		int32 DebugMipValue0Index;
-		int32 DebugMipValue1Index;
-		int32 PreallocatedStackTextureIndex;
-		bool bAdaptive;
-		bool bGenerateFeedback;
-		bool bMeshPaint;
-		bool bMaterialCache;
-		float AspectRatio;
-
 		int32 CodeIndex;
 	};
 
@@ -751,6 +765,9 @@ protected:
 	int32 AddCodeChunkInner(uint64 Hash, const TCHAR* FormattedType, const TCHAR* FormattedCode, EMaterialValueType Type, EDerivativeStatus DerivativeStatus, bool bInlined);
 	int32 AddCodeChunkInner(uint64 Hash, const TCHAR* FormattedCode, EMaterialValueType Type, EDerivativeStatus DerivativeStatus, bool bInlined);
 
+	/** Adds an already formatted inline or referenced code chunk */
+	int32 AddInternalCodeChunk(EMaterialValueType Type, const TCHAR* InternalType, const TCHAR* FormattedCode, const TRefCountPtr<FMaterialUniformExpression>& UniformExpression = nullptr);
+
 public:
 	/** Adds an already formatted inline or referenced code chunk, and notes the derivative status. */
 	int32 AddCodeChunkInnerDeriv(const TCHAR* FormattedCodeFinite, const TCHAR* FormattedCodeAnalytic, EMaterialValueType Type, bool bInlined, EDerivativeStatus DerivativeStatus);
@@ -1056,19 +1073,11 @@ protected:
 	virtual void SetPotentiallyManipulateTexCoords() override;
 
 	uint32 AcquireVTStackIndex(
-		int32 VirtualTexureIndex,
-		ETextureMipValueMode MipValueMode,
-		TextureAddress AddressU, TextureAddress AddressV, 
-		float AspectRatio, 
-		int32 CoordinateIndex, 
-		int32 MipValue0Index, int32 MipValue1Index, 
-		int32 PreallocatedStackTextureIndex, 
+		const FMaterialVTSampleInfo& SampleInfo,
 		const FString& UV_Value,
 		const FString& UV_Ddx,
-		const FString& UV_Ddy,
-		const FString& PageTableOverride,
-		const FString& PageTableUniformOverride,
-		bool bAdaptive, bool bGenerateFeedback, bool bMeshPaint, bool bMaterialCache);
+		const FString& UV_Ddy
+	);
 
 	virtual int32 TextureSample(
 		int32 TextureIndex,
diff --git a/Engine/Source/Runtime/Engine/Private/Materials/MaterialRenderProxy.cpp b/Engine/Source/Runtime/Engine/Private/Materials/MaterialRenderProxy.cpp
--- a/Engine/Source/Runtime/Engine/Private/Materials/MaterialRenderProxy.cpp
+++ b/Engine/Source/Runtime/Engine/Private/Materials/MaterialRenderProxy.cpp
@@ -14,6 +14,8 @@
 #include "TextureResource.h"
 #include "RenderGraphBuilder.h"
 #include "MaterialCache/IMaterialCacheTagProvider.h"
+#include "Engine/TextureCollection.h"
+#include "Engine/VirtualTextureCollection.h"
 #include "VT/RuntimeVirtualTexture.h"
 
 #if WITH_EDITOR
@@ -487,6 +489,35 @@ void FMaterialRenderProxy::EvaluateUniformExpressions(FRHICommandListBase& RHICm
 		bHasMaterialCacheCallbacks = true;
 	}
 
+	for (int32 VTCIndex = 0; VTCIndex < UniformExpressionSet.UniformTextureCollectionParameters.Num(); ++VTCIndex)
+	{
+		const FMaterialTextureCollectionParameterInfo& Parameter = UniformExpressionSet.UniformTextureCollectionParameters[VTCIndex];
+		if (!Parameter.bIsVirtualCollection)
+		{
+			continue;
+		}
+
+		const UTextureCollection* Texture = nullptr;
+		UniformExpressionSet.GetTextureCollectionValue(VTCIndex, Context, Context.Material, Texture);
+
+		// May not be bound, but must be a VTC
+		if (!Texture || !ensure(Texture->IsVirtualCollection()))
+		{
+			continue;
+		}
+
+		// May not have a valid resource
+		FVirtualTextureCollectionResource* Resource = static_cast<FVirtualTextureCollectionResource*>(Texture->GetResource());
+		if (!Resource || !ensure(!Resource->bIsBindless))
+		{
+			continue;
+		}
+
+		// Subscribe to VT destruction
+		GetRendererModule().AddVirtualTextureProducerDestroyedCallback(Resource->GetProducerHandle(), &OnVirtualTextureDestroyedCB, const_cast<FMaterialRenderProxy*>(this));
+		bHasVirtualTextureCallbacks = true;
+	}
+
 	const FRHIUniformBufferLayout* UniformBufferLayout = ShaderMap->GetUniformBufferLayout();
 
 	if (IsValidRef(OutUniformExpressionCache.UniformBuffer))
diff --git a/Engine/Source/Runtime/Engine/Private/Materials/MaterialShared.cpp b/Engine/Source/Runtime/Engine/Private/Materials/MaterialShared.cpp
--- a/Engine/Source/Runtime/Engine/Private/Materials/MaterialShared.cpp
+++ b/Engine/Source/Runtime/Engine/Private/Materials/MaterialShared.cpp
@@ -116,6 +116,19 @@ bool IsUsingNewMaterialTranslatorPrototype()
 
 IMPLEMENT_TYPE_LAYOUT(FHashedMaterialParameterInfo);
 IMPLEMENT_TYPE_LAYOUT(FUniformExpressionSet);
+
+void FUniformExpressionSet::CountTextureCollections(uint32& BindlessCollections, uint32& VirtualCollections) const
+{
+	VirtualCollections = 0;
+	
+	for (const FMaterialTextureCollectionParameterInfo& UniformTextureCollectionParameter : UniformTextureCollectionParameters)
+	{
+		VirtualCollections += UniformTextureCollectionParameter.bIsVirtualCollection;
+	}
+
+	BindlessCollections = UniformTextureCollectionParameters.Num() - VirtualCollections;
+}
+
 IMPLEMENT_TYPE_LAYOUT(FMaterialCompilationOutput);
 IMPLEMENT_TYPE_LAYOUT(FMeshMaterialShaderMap);
 IMPLEMENT_TYPE_LAYOUT(FMaterialProcessedSource);
diff --git a/Engine/Source/Runtime/Engine/Private/Materials/MaterialUniformExpressions.cpp b/Engine/Source/Runtime/Engine/Private/Materials/MaterialUniformExpressions.cpp
--- a/Engine/Source/Runtime/Engine/Private/Materials/MaterialUniformExpressions.cpp
+++ b/Engine/Source/Runtime/Engine/Private/Materials/MaterialUniformExpressions.cpp
@@ -25,6 +25,7 @@
 #include "MaterialCache/MaterialCacheAttribute.h"
 #include "MaterialCache/MaterialCacheTagProvider.h"
 #include "MaterialCache/MaterialCacheTagSceneData.h"
+#include "Engine/VirtualTextureCollection.h"
 
 void WriteMaterialUniformAccess(UE::Shader::EValueComponentType ComponentType, uint32 NumComponents, uint32 UniformOffset, FStringBuilderBase& OutResult)
 {
@@ -394,7 +395,22 @@ FShaderParametersMetadata* FUniformExpressionSet::CreateBufferStruct()
 	// Make sure FUniformExpressionSet::CreateDebugLayout() is in sync
 	TArray<FShaderParametersMetadata::FMember> Members;
 	uint32 NextMemberOffset = 0;
+	
+	if (!UniformTextureCollectionParameters.IsEmpty())
+	{
+		uint32 BindlessCollectionCount = 0;
+		uint32 VirtualCollectionCount  = 0;
+		CountTextureCollections(BindlessCollectionCount, VirtualCollectionCount);
 
+		// Bindless collections do not make use of uniforms
+		if (VirtualCollectionCount)
+		{
+			static_assert(sizeof(FUintVector4) * UE::HLSL::IndirectVirtualTextureUniformDWord4Count == sizeof(UE::HLSL::FIndirectVirtualTextureUniform), "Unexpected size");
+			new(Members) FShaderParametersMetadata::FMember(TEXT("TextureCollectionPackedUniforms"), TEXT("FIndirectVirtualTextureUniform"), __LINE__, NextMemberOffset, UBMT_UINT32, EShaderPrecisionModifier::Float, 1, 4, UE::HLSL::IndirectVirtualTextureUniformDWord4Count * VirtualCollectionCount, NULL);
+			NextMemberOffset += sizeof(UE::HLSL::FIndirectVirtualTextureUniform) * UniformTextureCollectionParameters.Num();
+		}
+	}
+	
 	if (VTStacks.Num())
 	{
 		// 2x uint4 per VTStack
@@ -446,6 +462,9 @@ FShaderParametersMetadata* FUniformExpressionSet::CreateBufferStruct()
 	static FString VolumeTextureSamplerNames[MaxUniformExpressions];
 	static FString ExternalTextureNames[MaxUniformExpressions];
 	static FString TextureCollectionNames[MaxUniformExpressions];
+	static FString TextureCollectionPhysicalNames[MaxUniformExpressions];
+	static FString TextureCollectionSamplerNames[MaxUniformExpressions];
+	static FString TextureCollectionPageTableNames[MaxUniformExpressions];
 	static FString MediaTextureSamplerNames[MaxUniformExpressions];
 	static FString VirtualTexturePageTableNames0[MaxUniformExpressions];
 	static FString VirtualTexturePageTableNames1[MaxUniformExpressions];
@@ -480,6 +499,9 @@ FShaderParametersMetadata* FUniformExpressionSet::CreateBufferStruct()
 			VolumeTextureSamplerNames[i] = FString::Printf(TEXT("VolumeTexture_%dSampler"), i);
 			ExternalTextureNames[i] = FString::Printf(TEXT("ExternalTexture_%d"), i);
 			TextureCollectionNames[i] = FString::Printf(TEXT("TextureCollection_%d"), i);
+			TextureCollectionPhysicalNames[i] = FString::Printf(TEXT("TextureCollectionPhysical_%d"), i);
+			TextureCollectionSamplerNames[i] = FString::Printf(TEXT("TextureCollectionSampler_%d"), i);
+			TextureCollectionPageTableNames[i] = FString::Printf(TEXT("TextureCollectionPageTable_%d"), i);
 			MediaTextureSamplerNames[i] = FString::Printf(TEXT("ExternalTexture_%dSampler"), i);
 			VirtualTexturePageTableNames0[i] = FString::Printf(TEXT("VirtualTexturePageTable0_%d"), i);
 			VirtualTexturePageTableNames1[i] = FString::Printf(TEXT("VirtualTexturePageTable1_%d"), i);
@@ -589,8 +611,27 @@ FShaderParametersMetadata* FUniformExpressionSet::CreateBufferStruct()
 	for (int32 i = 0; i < UniformTextureCollectionParameters.Num(); ++i)
 	{
 		check((NextMemberOffset % SHADER_PARAMETER_POINTER_ALIGNMENT) == 0);
-		new(Members) FShaderParametersMetadata::FMember(*TextureCollectionNames[i], TEXT("FResourceCollection"), __LINE__, NextMemberOffset, UBMT_RESOURCE_COLLECTION, EShaderPrecisionModifier::Float, 1, 1, 0, NULL);
-		NextMemberOffset += SHADER_PARAMETER_POINTER_ALIGNMENT;
+
+		// Resource collections are RHI constructs, texture collections are software virtual tables
+		if (!UniformTextureCollectionParameters[i].bIsVirtualCollection)
+		{
+			new(Members) FShaderParametersMetadata::FMember(*TextureCollectionNames[i], TEXT("FResourceCollection"), __LINE__, NextMemberOffset, UBMT_RESOURCE_COLLECTION, EShaderPrecisionModifier::Float, 1, 1, 0, NULL);
+			NextMemberOffset += SHADER_PARAMETER_POINTER_ALIGNMENT;
+		}
+		else
+		{
+			new(Members) FShaderParametersMetadata::FMember(*TextureCollectionNames[i], TEXT("FIndirectVirtualTexture"), __LINE__, NextMemberOffset, UBMT_SRV, EShaderPrecisionModifier::Float, 1, 1, 0, NULL);
+			NextMemberOffset += SHADER_PARAMETER_POINTER_ALIGNMENT;
+			
+			new(Members) FShaderParametersMetadata::FMember(*TextureCollectionPageTableNames[i], TEXT("Texture2D<uint4>"), __LINE__, NextMemberOffset, UBMT_TEXTURE, EShaderPrecisionModifier::Float, 1, 1, 0, NULL);
+			NextMemberOffset += SHADER_PARAMETER_POINTER_ALIGNMENT;
+		
+			new(Members) FShaderParametersMetadata::FMember(*TextureCollectionPhysicalNames[i], TEXT("Texture2D"), __LINE__, NextMemberOffset, UBMT_SRV, EShaderPrecisionModifier::Float, 1, 1, 0, NULL);
+			NextMemberOffset += SHADER_PARAMETER_POINTER_ALIGNMENT;
+	
+			new(Members) FShaderParametersMetadata::FMember(*TextureCollectionSamplerNames[i], TEXT("SamplerState"), __LINE__, NextMemberOffset, UBMT_SAMPLER, EShaderPrecisionModifier::Float, 1, 1, 0, NULL);
+			NextMemberOffset += SHADER_PARAMETER_POINTER_ALIGNMENT;
+		}
 	}
 
 	for (int32 i = 0; i < VTStacks.Num(); ++i)
@@ -879,6 +920,20 @@ void FUniformExpressionSet::GetTextureValue(int32 Index, const FMaterialRenderCo
 	}
 }
 
+void FUniformExpressionSet::GetTextureCollectionValue(int32 Index, const FMaterialRenderContext& Context, const FMaterial& Material, const UTextureCollection*& OutValue) const
+{
+	check(IsInParallelRenderingThread());
+	if (ensure(Index < GetNumTextureCollections()))
+	{
+		const FMaterialTextureCollectionParameterInfo& Parameter = GetTextureCollectionParameter(Index);
+		Context.GetTextureCollectionParameterValue(Parameter.ParameterInfo, Parameter.TextureCollectionIndex, OutValue);
+	}
+	else
+	{
+		OutValue = nullptr;
+	}
+}
+
 void FUniformExpressionSet::GetGameThreadTextureCollectionValue(int32 Index, const UMaterialInterface* MaterialInterface, const FMaterial& Material, UTextureCollection*& OutValue) const
 {
 	check(IsInGameThread() || IsInParallelGameThread());
@@ -984,7 +1039,34 @@ void FUniformExpressionSet::FillUniformBuffer(const FMaterialRenderContext& Mate
 
 		void* BufferCursor = TempBuffer;
 		check(BufferCursor <= TempBuffer + TempBufferSize);
+		
+		if (!UniformTextureCollectionParameters.IsEmpty())
+		{
+			// Bindless collections do not make use of uniforms
+			for (const FMaterialTextureCollectionParameterInfo& ParameterInfo : UniformTextureCollectionParameters)
+			{
+				if (!ParameterInfo.bIsVirtualCollection)
+				{
+					continue;
+				}
+				
+				const UTextureCollection* TextureCollection = nullptr;
+				ParameterInfo.GetTextureCollection(MaterialRenderContext, TextureCollection);
+		
+				auto* VTPackedPageTableUniform = static_cast<UE::HLSL::FIndirectVirtualTextureUniform*>(BufferCursor);
+				BufferCursor = VTPackedPageTableUniform + 1;
 
+				if (FVirtualTextureCollectionResource* Resource = static_cast<FVirtualTextureCollectionResource*>(TextureCollection->GetResource()))
+				{
+					*VTPackedPageTableUniform = Resource->GetVirtualPackedUniform();
+				}
+				else
+				{
+					*VTPackedPageTableUniform = UE::HLSL::FIndirectVirtualTextureUniform{};
+				}
+			}
+		}
+		
 		// Dump virtual texture per page table uniform data
 		check(AllocatedVTs.Num() == VTStacks.Num());
 		for ( int32 VTStackIndex = 0; VTStackIndex < VTStacks.Num(); ++VTStackIndex)
@@ -1198,6 +1280,10 @@ void FUniformExpressionSet::FillUniformBuffer(const FMaterialRenderContext& Mate
 				NumPageTableIndirectionTextures++;
 			}
 	
+			uint32 BindlessCollectionCount = 0;
+			uint32 VirtualCollectionCount  = 0;
+			CountTextureCollections(BindlessCollectionCount, VirtualCollectionCount);
+		
 			check(UniformBufferLayout->Resources.Num() == 
 				UniformTextureParameters[(uint32)EMaterialTextureParameterType::Standard2D].Num() * 2
 				+ UniformTextureParameters[(uint32)EMaterialTextureParameterType::Cube].Num() * 2
@@ -1206,7 +1292,8 @@ void FUniformExpressionSet::FillUniformBuffer(const FMaterialRenderContext& Mate
 				+ UniformTextureParameters[(uint32)EMaterialTextureParameterType::Volume].Num() * 2
 				+ UniformTextureParameters[(uint32)EMaterialTextureParameterType::SparseVolume].Num() * 4
 				+ UniformExternalTextureParameters.Num() * 2
-				+ UniformTextureCollectionParameters.Num()
+				+ BindlessCollectionCount * 1
+				+ VirtualCollectionCount * 4
 				+ UniformTextureParameters[(uint32)EMaterialTextureParameterType::Virtual].Num() * 2
 				+ MaterialCacheTagStacks.Num() * (MaterialCacheMaxRuntimeLayers + 2)
 				+ NumPageTableTextures
@@ -1543,25 +1630,56 @@ void FUniformExpressionSet::FillUniformBuffer(const FMaterialRenderContext& Mate
 			}
 		}
 
-		for (const FMaterialTextureCollectionParameterInfo& ParameterInfo : UniformTextureCollectionParameters)
+		if (!UniformTextureCollectionParameters.IsEmpty())
 		{
-			void** ResourceTableBufferPtr = reinterpret_cast<void**>(BufferCursor);
-			BufferCursor = reinterpret_cast<uint8*>(BufferCursor) + SHADER_PARAMETER_POINTER_ALIGNMENT;
-			check(BufferCursor <= TempBuffer + TempBufferSize);
+			for (const FMaterialTextureCollectionParameterInfo& ParameterInfo : UniformTextureCollectionParameters)
+			{
+				const UTextureCollection* TextureCollection = nullptr;
+				ParameterInfo.GetTextureCollection(MaterialRenderContext, TextureCollection);
 
-			const UTextureCollection* TextureCollection = nullptr;
-			ParameterInfo.GetTextureCollection(MaterialRenderContext, TextureCollection);
+				FTextureCollectionResource* Resource = TextureCollection ? TextureCollection->GetResource() : nullptr;
 
-			if (TextureCollection != nullptr)
-			{
-				FTextureCollectionResource* Resource = TextureCollection->GetResource();
-				check(Resource);
-				*ResourceTableBufferPtr = Resource->GetRHI();
-			}
-			else
-			{
-				// TODO: christopher.waters - implement default texture collections
-				unimplemented();
+				if (!ParameterInfo.bIsVirtualCollection)
+				{
+					void** ResourceTableBufferPtr = (void**)((uint8*)BufferCursor + 0 * SHADER_PARAMETER_POINTER_ALIGNMENT);
+					
+					if (FBindlessTextureCollectionResource* BindlessResource = static_cast<FBindlessTextureCollectionResource*>(Resource))
+					{
+						*ResourceTableBufferPtr = BindlessResource->GetResourceCollectionRHI();
+					}
+					else
+					{
+						*ResourceTableBufferPtr = GEmptyResourceCollection.ResourceCollection;
+					}
+		
+					BufferCursor = ((uint8*)BufferCursor) + (SHADER_PARAMETER_POINTER_ALIGNMENT);
+					check(BufferCursor <= TempBuffer + TempBufferSize);
+				}
+				else
+				{
+					void** ResourceTableBufferPtr          = (void**)((uint8*)BufferCursor + 0 * SHADER_PARAMETER_POINTER_ALIGNMENT);
+					void** ResourceTablePageTablePtr       = (void**)((uint8*)BufferCursor + 1 * SHADER_PARAMETER_POINTER_ALIGNMENT);
+					void** ResourceTablePhysicalTexturePtr = (void**)((uint8*)BufferCursor + 2 * SHADER_PARAMETER_POINTER_ALIGNMENT);
+					void** ResourceTablePhysicalSamplerPtr = (void**)((uint8*)BufferCursor + 3 * SHADER_PARAMETER_POINTER_ALIGNMENT);
+					
+					if (FVirtualTextureCollectionResource* VirtualResource = static_cast<FVirtualTextureCollectionResource*>(Resource))
+					{
+						*ResourceTableBufferPtr          = VirtualResource->GetVirtualCollectionRHI();
+						*ResourceTablePageTablePtr       = VirtualResource->GetVirtualPageTable();
+						*ResourceTablePhysicalTexturePtr = VirtualResource->GetVirtualPhysicalTextureSRV();
+					}
+					else
+					{
+						*ResourceTableBufferPtr          = GEmptyVertexBufferUInt4WithUAV->ShaderResourceViewRHI;
+						*ResourceTablePageTablePtr       = GBlackUintTexture->TextureRHI.GetReference();
+						*ResourceTablePhysicalTexturePtr = GBlackTexture->TextureRHI.GetReference();
+					}
+
+					*ResourceTablePhysicalSamplerPtr = TStaticSamplerState<SF_AnisotropicPoint, AM_Clamp, AM_Clamp, AM_Clamp, 0, 8>::GetRHI();
+
+					BufferCursor = ((uint8*)BufferCursor) + (SHADER_PARAMETER_POINTER_ALIGNMENT * 4);
+					check(BufferCursor <= TempBuffer + TempBufferSize);
+				}
 			}
 		}
 
@@ -1834,8 +1952,10 @@ bool FMaterialUniformExpressionTexture::IsIdentical(const FMaterialUniformExpres
 
 // FMaterialUniformExpressionTextureCollection
 
-FMaterialUniformExpressionTextureCollection::FMaterialUniformExpressionTextureCollection(int32 InTextureCollectionIndex)
+FMaterialUniformExpressionTextureCollection::FMaterialUniformExpressionTextureCollection(int32 InTextureCollectionIndex, int32 InTextureCollectionTypePrefixIndex, bool bIsVirtualCollection)
 	: TextureCollectionIndex(InTextureCollectionIndex)
+	, TextureCollectionTypePrefixIndex(InTextureCollectionTypePrefixIndex)
+    , bIsVirtualCollection(bIsVirtualCollection)
 {
 }
 
@@ -1858,12 +1978,13 @@ FMaterialUniformExpressionTextureCollectionParameter* FMaterialUniformExpression
 void FMaterialUniformExpressionTextureCollection::GetTextureCollectionParameterInfo(FMaterialTextureCollectionParameterInfo& OutParameter) const
 {
 	OutParameter.TextureCollectionIndex = this->TextureCollectionIndex;
+	OutParameter.bIsVirtualCollection = this->bIsVirtualCollection;
 }
 
 // FMaterialUniformExpressionTextureCollectionParameter
 
-FMaterialUniformExpressionTextureCollectionParameter::FMaterialUniformExpressionTextureCollectionParameter(const FMaterialParameterInfo& InParameterInfo, int32 InTextureCollectionIndex)
-	: Super(InTextureCollectionIndex)
+FMaterialUniformExpressionTextureCollectionParameter::FMaterialUniformExpressionTextureCollectionParameter(const FMaterialParameterInfo& InParameterInfo, int32 InTextureCollectionIndex, int32 InTextureCollectionTypePrefixIndex, bool bIsVirtualCollection)
+	: Super(InTextureCollectionIndex, InTextureCollectionTypePrefixIndex, bIsVirtualCollection)
 	, ParameterInfo(InParameterInfo)
 {
 }
diff --git a/Engine/Source/Runtime/Engine/Private/Materials/MaterialUniformExpressions.h b/Engine/Source/Runtime/Engine/Private/Materials/MaterialUniformExpressions.h
--- a/Engine/Source/Runtime/Engine/Private/Materials/MaterialUniformExpressions.h
+++ b/Engine/Source/Runtime/Engine/Private/Materials/MaterialUniformExpressions.h
@@ -140,7 +140,7 @@ class FMaterialUniformExpressionTextureCollection : public FMaterialUniformExpre
 	DECLARE_MATERIALUNIFORMEXPRESSION_TYPE(FMaterialUniformExpressionTextureCollection);
 
 public:
-	FMaterialUniformExpressionTextureCollection(int32 InTextureCollectionIndex);
+	FMaterialUniformExpressionTextureCollection(int32 InTextureCollectionIndex, int32 InTextureCollectionTypePrefixIndex, bool bIsVirtualCollection);
 
 	//~ Begin FMaterialUniformExpression Interface.
 	virtual class FMaterialUniformExpressionTextureCollection* GetTextureCollectionUniformExpression() final { return this; }
@@ -155,8 +155,25 @@ public:
 		return TextureCollectionIndex;
 	}
 
+	void SetTextureCollectionTypePrefixIndex(uint32 InPrefix)
+	{
+		TextureCollectionTypePrefixIndex = InPrefix;
+	}
+
+	int32 GetTextureCollectionTypePrefixIndex() const
+	{
+		return TextureCollectionTypePrefixIndex;
+	}
+
+	bool IsVirtualCollection() const
+	{
+		return bIsVirtualCollection;
+	}
+
 protected:
 	int32 TextureCollectionIndex;
+	int32 TextureCollectionTypePrefixIndex;
+	bool bIsVirtualCollection;
 };
 
 class FMaterialUniformExpressionTextureCollectionParameter : public FMaterialUniformExpressionTextureCollection
@@ -165,7 +182,7 @@ class FMaterialUniformExpressionTextureCollectionParameter : public FMaterialUni
 	DECLARE_MATERIALUNIFORMEXPRESSION_TYPE(FMaterialUniformExpressionTextureCollectionParameter);
 
 public:
-	FMaterialUniformExpressionTextureCollectionParameter(const FMaterialParameterInfo& InParameterInfo, int32 InTextureCollectionIndex);
+	FMaterialUniformExpressionTextureCollectionParameter(const FMaterialParameterInfo& InParameterInfo, int32 InTextureCollectionIndex, int32 InTextureCollectionTypePrefixIndex, bool bIsVirtualCollection);
 
 	//~ Begin FMaterialUniformExpression Interface.
 	virtual bool IsIdentical(const FMaterialUniformExpression* OtherExpression) const override;
@@ -189,7 +206,6 @@ public:
 
 private:
 	FHashedMaterialParameterInfo ParameterInfo;
-	int32 ParameterIndex;
 };
 
 class FMaterialUniformExpressionExternalTextureBase : public FMaterialUniformExpression
diff --git a/Engine/Source/Runtime/Engine/Private/TextureCollection.cpp b/Engine/Source/Runtime/Engine/Private/TextureCollection.cpp
--- a/Engine/Source/Runtime/Engine/Private/TextureCollection.cpp
+++ b/Engine/Source/Runtime/Engine/Private/TextureCollection.cpp
@@ -18,6 +18,11 @@
 #define LOCTEXT_NAMESPACE "UTextureCollection"
 
 FTextureCollectionResource::FTextureCollectionResource(UTextureCollection* InParent)
+{
+	CollectionName = InParent->GetFName();
+}
+
+FBindlessTextureCollectionResource::FBindlessTextureCollectionResource(UTextureCollection* InParent) : FTextureCollectionResource(InParent)
 {
 	TConstArrayView<TObjectPtr<UTexture>> ParentTextures = InParent->Textures;
 	InputTextureResources.Reserve(ParentTextures.Num());
@@ -30,25 +35,23 @@ FTextureCollectionResource::FTextureCollectionResource(UTextureCollection* InPar
 	}
 }
 
-void FTextureCollectionResource::InitRHI(FRHICommandListBase& RHICmdList)
+void FBindlessTextureCollectionResource::InitRHI(FRHICommandListBase& RHICmdList)
 {
-	InputTextures.Empty(InputTextureResources.Num());
-
 	TArray<FRHIResourceCollectionMember> CollectionMembers;
 	CollectionMembers.Reserve(InputTextureResources.Num());
 
 	for (FTextureResource* TextureResource : InputTextureResources)
 	{
 		FRHITextureReference* TextureReference = TextureResource ? TextureResource->GetTextureReference() : nullptr;
-		InputTextures.Emplace(TextureReference);
 		CollectionMembers.Emplace(TextureReference);
 	}
 
 	ResourceCollectionRHI = RHICmdList.CreateResourceCollection(CollectionMembers);
 }
 
-void FTextureCollectionResource::ReleaseRHI()
+void FBindlessTextureCollectionResource::ReleaseRHI()
 {
+	
 }
 
 #if WITH_EDITOR
@@ -80,8 +83,6 @@ void UTextureCollection::NotifyMaterials()
 void UTextureCollection::PostEditChangeProperty(FPropertyChangedEvent& PropertyChangedEvent)
 {
 	UpdateResource();
-
-	NotifyMaterials();
 }
 #endif // WITH_EDITOR
 
@@ -149,7 +150,14 @@ FTextureCollectionResource* UTextureCollection::GetResource()
 
 FTextureCollectionResource* UTextureCollection::CreateResource()
 {
-	return new FTextureCollectionResource(this);
+	// TODO: christopher.waters - We should add a way to query the RHI to see if bindless is enabled.
+	ERHIBindlessConfiguration BindlessResourcesConfig = RHIGetRuntimeBindlessConfiguration(GMaxRHIShaderPlatform);
+	if (BindlessResourcesConfig == ERHIBindlessConfiguration::Disabled || !GRHIGlobals.bSupportsBindless)
+	{
+		return nullptr;
+	}
+
+	return new FBindlessTextureCollectionResource(this);
 }
 
 void UTextureCollection::ReleaseResource()
@@ -180,13 +188,6 @@ void UTextureCollection::UpdateResource()
 		return;
 	}
 
-	// TODO: christopher.waters - We should add a way to query the RHI to see if bindless is enabled.
-	ERHIBindlessConfiguration BindlessConfig = RHIGetRuntimeBindlessConfiguration(GMaxRHIShaderPlatform);
-	if (BindlessConfig == ERHIBindlessConfiguration::Disabled || !GRHIGlobals.bSupportsBindless)
-	{
-		return;
-	}
-
 	if (FTextureCollectionResource* NewResource = CreateResource())
 	{
 		LLM_SCOPE(ELLMTag::Textures);
@@ -206,6 +207,10 @@ void UTextureCollection::UpdateResource()
 	{
 		SetResource(nullptr);
 	}
+
+#if WITH_EDITOR
+	NotifyMaterials();
+#endif // WITH_EDITOR
 }
 
 #undef LOCTEXT_NAMESPACE
diff --git a/Engine/Source/Runtime/Engine/Private/VT/VirtualTextureAdapter.cpp b/Engine/Source/Runtime/Engine/Private/VT/VirtualTextureAdapter.cpp
--- a/Engine/Source/Runtime/Engine/Private/VT/VirtualTextureAdapter.cpp
+++ b/Engine/Source/Runtime/Engine/Private/VT/VirtualTextureAdapter.cpp
@@ -14,48 +14,15 @@
 #include "TextureResource.h"
 #include "VirtualTextureEnum.h"
 #include "VirtualTexturing.h"
+#include "VT/CopyCompressShader.h"
 #include "VT/VirtualTextureBuildSettings.h"
 
 #include UE_INLINE_GENERATED_CPP_BY_NAME(VirtualTextureAdapter)
 
+IMPLEMENT_GLOBAL_SHADER(FCopyCompressCS, "/Engine/Private/VirtualTextureAdapter.usf", "CopyCompressCS", SF_Compute);
+
 namespace VirtualTextureAdapter
 {
-	/** Copy with compression shader. */
-	class FCopyCompressCS : public FGlobalShader
-	{
-	public:
-		DECLARE_GLOBAL_SHADER(FCopyCompressCS);
-		SHADER_USE_PARAMETER_STRUCT(FCopyCompressCS, FGlobalShader);
-
-		BEGIN_SHADER_PARAMETER_STRUCT(FParameters, )
-			SHADER_PARAMETER_SRV(Texture2D, SourceTextureA)
-			SHADER_PARAMETER_RDG_TEXTURE_SRV(Texture2D, SourceTextureB)
-			SHADER_PARAMETER_SAMPLER(SamplerState, TextureSampler)
-			SHADER_PARAMETER_RDG_TEXTURE_UAV(RWTexture2D, DestTexture)
-			SHADER_PARAMETER_RDG_TEXTURE_UAV(RWTexture2D, DestCompressTexture_64bit)
-			SHADER_PARAMETER_RDG_TEXTURE_UAV(RWTexture2D, DestCompressTexture_128bit)
-			SHADER_PARAMETER(FVector2f, SourceUV)
-			SHADER_PARAMETER(FVector2f, TexelSize)
-			SHADER_PARAMETER(FVector2f, TexelOffsets)
-			SHADER_PARAMETER(FIntVector4, DestRect)
-		END_SHADER_PARAMETER_STRUCT()
-
-		static const int32 GroupSize = 8;
-
-		class FSourceTextureSelector : SHADER_PERMUTATION_BOOL("SOURCE_TEXTURE_A");
-		class FDestSrgb : SHADER_PERMUTATION_BOOL("TEXTURE_SRGB");
-		class FCompressionFormatDim : SHADER_PERMUTATION_INT("COMPRESSION_FORMAT", 7);
-		using FPermutationDomain = TShaderPermutationDomain<FSourceTextureSelector, FDestSrgb, FCompressionFormatDim>;
-
-		/** Get index to use for FCompressionFormatDim. */
-		static int32 GetCompressionPermutation(EPixelFormat InFormat)
-		{
-			const EPixelFormat CompressedFormats[] = { PF_DXT1, PF_DXT5, PF_BC4, PF_BC5, PF_BC6H, PF_BC7 };
-			return MakeArrayView(CompressedFormats).Find(InFormat) + 1;
-		}
-	};
-	IMPLEMENT_GLOBAL_SHADER(FCopyCompressCS, "/Engine/Private/VirtualTextureAdapter.usf", "CopyCompressCS", SF_Compute);
-
 	/** Final copy to output RDG parameters. */
 	BEGIN_SHADER_PARAMETER_STRUCT(FCopyToOutputParameters, )
 		RDG_TEXTURE_ACCESS(Input, ERHIAccess::CopySrc)
diff --git a/Engine/Source/Runtime/Engine/Private/VT/VirtualTextureRecreate.cpp b/Engine/Source/Runtime/Engine/Private/VT/VirtualTextureRecreate.cpp
--- a/Engine/Source/Runtime/Engine/Private/VT/VirtualTextureRecreate.cpp
+++ b/Engine/Source/Runtime/Engine/Private/VT/VirtualTextureRecreate.cpp
@@ -9,6 +9,7 @@
 #include "TextureResource.h"
 #include "UObject/UObjectIterator.h"
 #include "VirtualTexturing.h"
+#include "Engine/VirtualTextureCollection.h"
 #include "VT/RuntimeVirtualTexture.h"
 
 namespace VirtualTexture
@@ -45,6 +46,11 @@ namespace VirtualTexture
 				Owner->MarkRenderStateDirty();
 			}
 		}
+		
+		for (TObjectIterator<UVirtualTextureCollection> It(RF_ClassDefaultObject, true, EInternalObjectFlags::Garbage); It; ++It)
+		{
+			It->UpdateResource();
+		}
 	}
 
 	void Recreate(TConstArrayView< TEnumAsByte<EPixelFormat> > InFormat)
@@ -128,5 +134,13 @@ namespace VirtualTexture
 				}
 			}
 		}
+		
+		for (TObjectIterator<UVirtualTextureCollection> It(RF_ClassDefaultObject, true, EInternalObjectFlags::Garbage); It; ++It)
+		{
+			if (InFormat.Num() == 1 && It->RuntimePixelFormat == InFormat[0])
+			{
+				It->UpdateResource();
+			}
+		}
 	}
 }
diff --git a/Engine/Source/Runtime/Engine/Private/VirtualTextureCollection.cpp b/Engine/Source/Runtime/Engine/Private/VirtualTextureCollection.cpp
new file mode 100644
--- /dev/null
+++ b/Engine/Source/Runtime/Engine/Private/VirtualTextureCollection.cpp
@@ -0,0 +1,1106 @@
+// Copyright Epic Games, Inc. All Rights Reserved.
+
+#include "Engine/VirtualTextureCollection.h"
+#include "EngineModule.h"
+#include "Engine/Texture.h"
+#include "RenderingThread.h"
+#include "TextureResource.h"
+#include "GlobalRenderResources.h"
+#include "RenderGraphBuilder.h"
+#include "RenderGraphUtils.h"
+#include "RenderTargetPool.h"
+#include "TextureLayout.h"
+#include "VirtualTextureEnum.h"
+#include "RHIStaticStates.h"
+#include "Rendering/Texture2DResource.h"
+#include "VT/CopyCompressShader.h"
+
+#if WITH_EDITOR
+#include "Materials/Material.h"
+#include "Materials/MaterialInterface.h"
+#include "MaterialShared.h"
+#include "TextureCompiler.h"
+#include "VT/CopyCompressShader.h"
+#endif // WITH_EDITOR
+
+#define LOCTEXT_NAMESPACE "UVirtualTextureCollection"
+
+DEFINE_LOG_CATEGORY_STATIC(LogVirtualTextureCollection, Log, All);
+
+struct FTextureCollectionProducerData
+{
+	/** All textures being produced (virtual or physical) */
+	TArray<FVirtualTextureCollectionResource::FTextureEntry> Textures;
+
+	/** The host-side page table for block to texture indices */
+	TArray<int32> BlockVirtualPageTable;
+
+	/** All starting block coordinates of the textures */
+	TArray<FUintVector2> BlockCoordinates;
+
+	/** Total block count of the shared VT */
+	FUintVector2 BlockCount;
+};
+
+struct FTextureCollectionPendingAdapterTile
+{
+	/** Texture entry to apply adapter to */
+	FVirtualTextureCollectionResource::FTextureEntry Entry;
+	
+	/** Physical destination for the adapter */
+	FVTProduceTargetLayer Target;
+
+	/** Optional, finalizer for virtual textures, to be invoked before the adapter */
+	IVirtualTextureFinalizer* VirtualFinalizer = nullptr;
+
+	/** The intermediate physical target for the virtual texture, not physical */
+	TRefCountPtr<IPooledRenderTarget> PooledVirtualRenderTarget;
+	
+	/** Intra-texture block address */
+	uint64 Address = 0;
+	
+	/** Mip level */
+	uint8 Level = 0;
+};
+
+static FUintVector2 DecodeBlockAddress(uint32 vAddress)
+{
+	return FUintVector2(FMath::ReverseMortonCode2_64(vAddress), FMath::ReverseMortonCode2_64(vAddress >> 1));
+}
+
+static uint32 EncodeBlockAddress(const FUintVector2& Addr)
+{
+	return FMath::MortonCode2_64(Addr.X) | (FMath::MortonCode2_64(Addr.Y) << 1);
+}
+
+static FTexture2DResource* GetTextureObjectResource2D(UTexture* Texture)
+{
+	FTextureResource* Resource = Texture ? Texture->GetResource() : nullptr;
+	return Resource ? Resource->GetTexture2DResource() : nullptr;
+}
+
+/** Adapter finalizer, handles (optional) format conversions and block compression */
+class FTextureCollectionVirtualAdapterFinalizer : public IVirtualTextureFinalizer
+{
+public:
+	FTextureCollectionVirtualAdapterFinalizer(const FVTProducerDescription& ProducerDesc) : ProducerDesc(ProducerDesc)
+	{
+		
+	}
+
+	virtual ~FTextureCollectionVirtualAdapterFinalizer() = default;
+
+	virtual void RenderFinalize(FRDGBuilder& GraphBuilder, ISceneRenderer* SceneRenderingContext) override
+	{
+		for (IVirtualTextureFinalizer* PendingFinalizer : PendingFinalizers)
+		{
+			PendingFinalizer->RenderFinalize(GraphBuilder, SceneRenderingContext);
+		}
+	}
+
+	virtual void Finalize(FRDGBuilder& GraphBuilder) override
+	{
+		// First, finalize all intermediate data onto the temporary targets
+		for (IVirtualTextureFinalizer* PendingFinalizer : PendingFinalizers)
+		{
+			PendingFinalizer->Finalize(GraphBuilder);
+		}
+
+		// Second, run the adapter on each pending tile
+		for (const FTextureCollectionPendingAdapterTile& PendingTile : PendingTiles)
+		{
+			const FPooledRenderTargetDesc TargetDesc = PendingTile.Target.PooledRenderTarget->GetDesc();
+			const bool bTargetIsCompressed = IsBlockCompressedFormat(TargetDesc.Format);
+			
+			// Physical tile location
+			const int32     TileSize = ProducerDesc.TileSize + 2 * ProducerDesc.TileBorderSize;
+			const FIntPoint DestinationPos(PendingTile.Target.pPageLocation.X * TileSize, PendingTile.Target.pPageLocation.Y * TileSize);
+			const FIntRect  DestRect(DestinationPos, DestinationPos + FIntPoint(TileSize, TileSize));
+
+			// Select the UAV aliasing format, 32 or 64 wide
+			const bool         bAliasTo64bit = ProducerDesc.LayerFormat[0] == PF_DXT1 || ProducerDesc.LayerFormat[0] == PF_BC4;
+			const EPixelFormat AliasFormat = bAliasTo64bit ? PF_R32G32_UINT : PF_R32G32B32A32_UINT;
+
+			// If not compressed or the target supports aliasing, just write directly into the physical resource
+			const bool bWriteToPhysical = !bTargetIsCompressed || GRHISupportsUAVFormatAliasing;
+			
+			FRDGTextureRef CurrentOutput;
+			if (bWriteToPhysical)
+			{
+				CurrentOutput = GraphBuilder.RegisterExternalTexture(PendingTile.Target.PooledRenderTarget, ERDGTextureFlags::None);
+			}
+			else
+			{
+				CurrentOutput = GraphBuilder.CreateTexture(FRDGTextureDesc::Create2D(
+					FIntPoint(TileSize >> 2, TileSize >> 2),
+					AliasFormat, 
+					FClearValueBinding::None,
+					ETextureCreateFlags::ShaderResource | ETextureCreateFlags::UAV
+				), TEXT("IntermediateTarget"));
+			}
+
+			// Select the compression permutation
+			int32 CompressionDim = 0;
+			if (IsBlockCompressedFormat(CurrentOutput->Desc.Format))
+			{
+				CompressionDim = FCopyCompressCS::GetCompressionPermutation(CurrentOutput->Desc.Format);
+			}
+
+			FCopyCompressCS::FPermutationDomain PermutationVector;
+			PermutationVector.Set<FCopyCompressCS::FSourceTextureSelector>(true);
+			PermutationVector.Set<FCopyCompressCS::FDestSrgb>(ProducerDesc.bIsLayerSRGB[0]);
+			PermutationVector.Set<FCopyCompressCS::FCompressionFormatDim>(CompressionDim);
+			TShaderMapRef<FCopyCompressCS> Shader(GetGlobalShaderMap(GMaxRHIFeatureLevel), PermutationVector);
+
+			FCopyCompressCS::FParameters* Parameters = GraphBuilder.AllocParameters<FCopyCompressCS::FParameters>();
+			Parameters->TextureSampler = TStaticSamplerState<SF_Bilinear>::GetRHI();
+			Parameters->DestTexture = GraphBuilder.CreateUAV(CurrentOutput);
+			Parameters->DestCompressTexture_64bit = bAliasTo64bit ? GraphBuilder.CreateUAV(FRDGTextureUAVDesc(CurrentOutput, 0, PF_R32G32_UINT)) : nullptr;
+			Parameters->DestCompressTexture_128bit = bAliasTo64bit ? nullptr : GraphBuilder.CreateUAV(FRDGTextureUAVDesc(CurrentOutput, 0, PF_R32G32B32A32_UINT));
+			Parameters->TexelOffsets = FVector2f(1.0f, 0.5f);
+			Parameters->DestRect = FIntVector4(DestRect.Min.X, DestRect.Min.Y, DestRect.Max.X, DestRect.Max.Y);
+
+			// Coming from a virtual texture?
+			if (PendingTile.PooledVirtualRenderTarget)
+			{
+				// No mipped views, always a single tile
+				Parameters->SourceTextureA = GraphBuilder.RHICmdList.CreateShaderResourceView(
+					PendingTile.PooledVirtualRenderTarget->GetRHI(),
+					FRHIViewDesc::CreateTextureSRV()
+						.SetDimensionFromTexture(PendingTile.PooledVirtualRenderTarget->GetRHI())
+				);
+
+				// Set UV ranges
+				Parameters->SourceUV  = FVector2f(0, 0);
+				Parameters->TexelSize = FVector2f(1.0f / PendingTile.PooledVirtualRenderTarget->GetDesc().Extent.X, 1.0f / PendingTile.PooledVirtualRenderTarget->GetDesc().Extent.Y);
+			}
+			else
+			{
+				// Resource may have been released
+				FTexture2DResource* Resource = GetTextureObjectResource2D(PendingTile.Entry.PhysicalTexture);
+				if (!Resource)
+				{
+					continue;
+				}
+				
+				FTextureRHIRef TextureRHI = Resource->GetTextureRHI();
+				if (!ensure(TextureRHI.IsValid()))
+				{
+					continue;
+				}
+
+				// Create view into the repsective mip
+				Parameters->SourceTextureA = GraphBuilder.RHICmdList.CreateShaderResourceView(
+					TextureRHI,
+					FRHIViewDesc::CreateTextureSRV()
+						.SetDimensionFromTexture(TextureRHI)
+						.SetMipRange(FMath::Min(PendingTile.Level - Resource->GetCurrentFirstMip(), Resource->GetCurrentMipCount() - 1), 1)
+				);
+
+				// Physical coordinates on the texture entry
+				const float X        = static_cast<float>(FMath::ReverseMortonCode2_64(PendingTile.Address));
+				const float Y        = static_cast<float>(FMath::ReverseMortonCode2_64(PendingTile.Address >> 1));
+				const float DivisorX = static_cast<float>(PendingTile.Entry.BlockCount.X) / static_cast<float>(1 << PendingTile.Level);
+				const float DivisorY = static_cast<float>(PendingTile.Entry.BlockCount.Y) / static_cast<float>(1 << PendingTile.Level);
+
+				// To UV coordinates
+				const FVector2f UV(X / DivisorX, Y / DivisorY);
+				const FVector2f UVSize(1.0f / DivisorX, 1.0f / DivisorY);
+				const FVector2f UVBorder = UVSize * (static_cast<float>(ProducerDesc.TileBorderSize) / static_cast<float>(ProducerDesc.TileSize));
+				const FBox2f    UVRect(UV - UVBorder, UV + UVSize + UVBorder);
+
+				// Set UV ranges
+				Parameters->SourceUV  = UVRect.Min;
+				Parameters->TexelSize = (UVRect.Max - UVRect.Min) / FVector2f(DestRect.Width(), DestRect.Height());
+			}
+
+			FIntPoint ThreadCount = FIntPoint(DestRect.Width(), DestRect.Height());
+
+			// If compressed, the kernel's dealing with blocks, not texels
+			if (bTargetIsCompressed)
+			{
+				ThreadCount /= 4;
+				Parameters->DestRect /= 4;
+				Parameters->TexelOffsets = FVector2f(4.0f, 0.5f);
+			}
+
+			ClearUnusedGraphResources(Shader, Parameters);
+			
+			GraphBuilder.AddPass(
+				RDG_EVENT_NAME("VirtualTextureCollectionAdapter"),
+				Parameters,
+				ERDGPassFlags::Compute,
+				[Shader, Parameters, ThreadCount](FRDGAsyncTask, FRHIComputeCommandList& RHICmdList)
+				{
+					FComputeShaderUtils::Dispatch(RHICmdList, Shader, *Parameters, FComputeShaderUtils::GetGroupCount(ThreadCount, 8));
+				}
+			);
+
+			if (!bWriteToPhysical)
+			{
+				FRHICopyTextureInfo CopyInfo;
+				CopyInfo.DestPosition = FIntVector(DestRect.Min.X, DestRect.Min.Y, 0);
+				CopyInfo.Size = FIntVector(DestRect.Width(), DestRect.Height(), 0);
+
+				if (bTargetIsCompressed)
+				{
+					CopyInfo.Size /= 4;
+				}
+				
+				GraphBuilder.AddPass(
+					RDG_EVENT_NAME("VirtualTextureCollectionPhysicalBlit"),
+					Parameters,
+					ERDGPassFlags::Copy | ERDGPassFlags::NeverCull,
+					[
+						CurrentOutput,
+						DestTexture = PendingTile.Target.PooledRenderTarget->GetRHI(),
+						CopyInfo
+					](FRDGAsyncTask, FRHICommandList& RHICmdList)
+				{
+					RHICmdList.CopyTexture(CurrentOutput->GetRHI(), DestTexture, CopyInfo);
+				});
+			}
+		}
+
+		PendingFinalizers.Empty();
+		PendingTiles.Empty();
+	}
+
+	void Add(const FTextureCollectionPendingAdapterTile& Tile)
+	{
+		if (Tile.VirtualFinalizer)
+		{
+			PendingFinalizers.AddUnique(Tile.VirtualFinalizer);
+		}
+		
+		PendingTiles.Add(Tile);
+	}
+	
+private:
+	TArray<IVirtualTextureFinalizer*> PendingFinalizers;
+	
+	TArray<FTextureCollectionPendingAdapterTile> PendingTiles;
+
+	FVTProducerDescription ProducerDesc;
+};
+
+class FTextureCollectionVirtualRedirector : public IVirtualTexture
+{
+public:
+	FTextureCollectionVirtualRedirector(FTextureCollectionProducerData&& Data, const FVTProducerDescription& ProducerDesc) : Data(MoveTemp(Data)), ProducerDesc(ProducerDesc), Finalizer(ProducerDesc)
+	{
+		
+	}
+	
+	virtual ~FTextureCollectionVirtualRedirector() override = default;
+
+	virtual bool IsPageStreamed(uint8 vLevel, uint32 vAddress) const override
+	{
+		const FUintVector2 Address = DecodeBlockAddress(vAddress);
+
+		// Get the owning texture, index table is always base mip
+		const int32 TextureIndex = Data.BlockVirtualPageTable[(Address.Y << vLevel) * Data.BlockCount.X + (Address.X << vLevel)];
+		if (TextureIndex == INDEX_NONE)
+		{
+			return false;
+		}
+
+		// Non-adapter entries are cheap to produce
+		const FVirtualTextureCollectionResource::FTextureEntry& Texture = Data.Textures[TextureIndex];
+		return !Texture.bRequiresAdapter;
+	}
+
+	virtual FVTRequestPageResult RequestPageData(
+		FRHICommandListBase& RHICmdList,
+		const FVirtualTextureProducerHandle& ProducerHandle,
+		uint8 LayerMask,
+		uint8 vLevel,
+		uint64 vAddress,
+		EVTRequestPagePriority Priority) override
+	{
+		const FUintVector2 Address = DecodeBlockAddress(vAddress);
+
+		// Get the owning texture, index table is always base mip
+		const int32 Index = Data.BlockVirtualPageTable[(Address.Y << vLevel) * Data.BlockCount.X + (Address.X << vLevel)];
+		if (Index == INDEX_NONE)
+		{
+			return EVTRequestPageStatus::Invalid;
+		}
+
+		// Shift the coordinates to the mip
+		FUintVector2 BlockAddress = Data.BlockCoordinates[Index];
+		BlockAddress.X >>= vLevel;
+		BlockAddress.Y >>= vLevel;
+
+		FVTRequestPageResult Result = EVTRequestPageStatus::Saturated;
+
+		const FVirtualTextureCollectionResource::FTextureEntry& Texture = Data.Textures[Index];
+		if (Texture.VirtualProducerHandle.IsValid())
+		{
+			// Virtual, query the owning texture
+			if (IVirtualTexture* VirtualTexture = GetRendererModule().FindProducer(Texture.VirtualProducerHandle))
+			{
+				Result = VirtualTexture->RequestPageData(
+					RHICmdList,
+					GetLocalProducerHandle(ProducerHandle, Index),
+					LayerMask,
+					vLevel,
+					EncodeBlockAddress(Address - BlockAddress),
+					Priority
+				);
+			}
+		}
+		else if (FTexture2DResource* Resource = GetTextureObjectResource2D(Texture.PhysicalTexture))
+		{
+			// Check if the mip is streamed in
+			FStreamableRenderResourceState State = Resource->GetState();
+			if (State.IsValid())
+			{
+				if (vLevel >= State.MaxNumLODs)
+				{
+					// Will never be produced
+					// Due to VT-streaming logic, just assume it's available to mark the upper mips as available overall
+					Result.Status = EVTRequestPageStatus::Available;
+				}
+				else
+				{
+					// Check if the mip is streamed in yet
+					bool bIsStreamedIn = vLevel >= State.ResidentFirstLODIdx();
+					Result.Status = bIsStreamedIn ? EVTRequestPageStatus::Available : EVTRequestPageStatus::Saturated;
+				}
+			}
+			else
+			{
+				// The resource isn't streamed, always available
+				Result.Status = EVTRequestPageStatus::Available;
+			}
+		}
+		else
+		{
+			// Invalid resource
+			Result.Status = EVTRequestPageStatus::Invalid;
+		}
+
+		return Result;
+	}
+
+	virtual IVirtualTextureFinalizer* ProducePageData(
+		FRHICommandListBase& RHICmdList,
+		ERHIFeatureLevel::Type FeatureLevel,
+		EVTProducePageFlags Flags,
+		const FVirtualTextureProducerHandle& ProducerHandle,
+		uint8 LayerMask,
+		uint8 vLevel,
+		uint64 vAddress,
+		uint64 RequestHandle,
+		const FVTProduceTargetLayer* TargetLayers) override
+	{
+		const FUintVector2 Address = DecodeBlockAddress(vAddress);
+
+		// Get the owning texture, index table is always base mip
+		const int32 Index = Data.BlockVirtualPageTable[(Address.Y << vLevel) * Data.BlockCount.X + (Address.X << vLevel)];
+		if (Index == INDEX_NONE)
+		{
+			return nullptr;
+		}
+
+		// Shift the coordinates to the mip
+		FUintVector2 BlockAddress = Data.BlockCoordinates[Index];
+		BlockAddress.X >>= vLevel;
+		BlockAddress.Y >>= vLevel;
+
+		const FVirtualTextureCollectionResource::FTextureEntry& Texture = Data.Textures[Index];
+
+		// Optional, adapter tile for format conversion
+		FTextureCollectionPendingAdapterTile AdapterTile;
+		AdapterTile.Entry = Texture;
+		AdapterTile.Target = *TargetLayers;
+		AdapterTile.Level = vLevel;
+		AdapterTile.Address = EncodeBlockAddress(Address - BlockAddress);
+
+		// Optional, target tile redirection for format conversion
+		FVTProduceTargetLayer RedirectedTarget = AdapterTile.Target;
+
+		if (Texture.VirtualProducerHandle.IsValid())
+		{
+			// Virtual, pass to the owning texture
+			if (IVirtualTexture* VirtualTexture = GetRendererModule().FindProducer(Texture.VirtualProducerHandle))
+			{
+				// If this producer requires an adapter, we render to an intermediate target instead of the current tile
+				if (Texture.bRequiresAdapter)
+				{
+					const int32 TileExtent = ProducerDesc.TileSize + ProducerDesc.TileBorderSize * 2;
+
+					// Find a free target
+					GRenderTargetPool.FindFreeElement(RHICmdList, FPooledRenderTargetDesc::Create2DDesc(
+						FIntPoint(TileExtent, TileExtent),
+						Texture.Format,
+						FClearValueBinding(FLinearColor::Black),
+						TexCreate_None, TexCreate_ShaderResource,
+						false
+					), AdapterTile.PooledVirtualRenderTarget, TEXT("VirtualRenderTarget"));
+
+					// External access in FVirtualTexturePhysicalSpace::FinalizeTextures assumes SRVMask
+					RHICmdList.TransitionInternal(FRHITransitionInfo(AdapterTile.PooledVirtualRenderTarget->GetRHI(), ERHIAccess::Unknown, ERHIAccess::SRVMask));
+					
+					RedirectedTarget.PooledRenderTarget = AdapterTile.PooledVirtualRenderTarget;
+					RedirectedTarget.pPageLocation       = FIntVector::ZeroValue;
+				}
+				
+				AdapterTile.VirtualFinalizer = VirtualTexture->ProducePageData(
+					RHICmdList,
+					FeatureLevel,
+					Flags,
+					GetLocalProducerHandle(ProducerHandle, Index),
+					LayerMask,
+					vLevel,
+					AdapterTile.Address,
+					RequestHandle,
+					&RedirectedTarget
+				);
+			}
+		}
+
+		// If this doesn't require an adapter, just use the producers finalizer
+		if (!Texture.bRequiresAdapter)
+		{
+			return AdapterTile.VirtualFinalizer;
+		}
+
+		// Otherwise, use the adapters finalizer
+		// This also invokes the above finalizer
+		Finalizer.Add(AdapterTile);
+		return &Finalizer;
+	}
+
+private:
+	static FVirtualTextureProducerHandle GetLocalProducerHandle(const FVirtualTextureProducerHandle& RedirectorHandle, uint32 Index)
+	{
+		// Due to the virtual transcoder cache, we need to offset the local handle
+		// to avoid hashing on physically different, but virtually/logically same requests
+		FVirtualTextureProducerHandle LocalHandle = RedirectorHandle;
+		LocalHandle.Magic += Index;
+		return LocalHandle;
+	}
+
+private:
+	FTextureCollectionProducerData Data;
+	FVTProducerDescription         ProducerDesc;
+
+	/** Shared finalizer for this redirector*/
+	FTextureCollectionVirtualAdapterFinalizer Finalizer;
+};
+
+static EPixelFormat GetResourcePixelFormat(FTextureResource* Resource)
+{
+	if (FVirtualTexture2DResource* VirtualResource = Resource->GetVirtualTexture2DResource())
+	{
+		return VirtualResource->GetFormat(0);
+	}
+	
+	if (FTexture2DResource* PhysicalResource = Resource->GetTexture2DResource())
+	{
+		return PhysicalResource->GetPixelFormat();
+	}
+
+	return PF_Unknown;
+}
+
+FVirtualTextureCollectionResource::FVirtualTextureCollectionResource(UVirtualTextureCollection* InParent) : FTextureCollectionResource(InParent)
+{
+	bIsBindless            = false;
+	Textures               = InParent->Textures;
+	bIsSRGB                = InParent->bIsSRGB;
+	bAllowFormatConversion = InParent->bAllowFormatConversion;
+
+	// Share the same build settings as general VT resources
+	// This avoids costly handling of mismatched tile sizes
+	BuildSettings.Init();
+
+	// If implicit, find the best matching format for the collection
+	// TODO[MP]: We could likely handle this at edit-time for most collections, then serialize?
+	if (Format == PF_Unknown)
+	{
+		if (bAllowFormatConversion)
+		{
+			FindConservativeFormat();
+		}
+		else
+		{
+			FindFirstFormat();
+		}
+	}
+
+	// TODO[MP]: We currently do not have runtime BC7 SRGB compression
+	if (Format == PF_BC7)
+	{
+		bIsSRGB = false;
+	}
+
+	// Report back the chosen format
+	InParent->RuntimePixelFormat = Format;
+}
+
+void FVirtualTextureCollectionResource::ComputeLayout(FTextureCollectionProducerData& Data)
+{
+	// Note: The texture layout only supports 16-bits on each axis
+	FTextureLayout Layout(0, 0, UINT16_MAX, UINT16_MAX);
+
+	for (UTexture* Texture : Textures)
+	{
+		FTextureResource* Resource = Texture ? Texture->GetResource() : nullptr;
+		if (!Resource)
+		{
+			continue;
+		}
+		
+		// To blocks
+		uint32 WidthInBlocks  = FMath::DivideAndRoundUp(Resource->GetSizeX(), static_cast<uint32>(BuildSettings.TileSize));
+		uint32 HeightInBlocks = FMath::DivideAndRoundUp(Resource->GetSizeY(), static_cast<uint32>(BuildSettings.TileSize));
+
+		// Natural mip reduction, if non-power of two, the lower mip regions will overlap
+		// While this does waste some page table memory, it should be fairly minimal compared to physical padding
+		WidthInBlocks  = FMath::RoundUpToPowerOfTwo(WidthInBlocks);
+		HeightInBlocks = FMath::RoundUpToPowerOfTwo(HeightInBlocks);
+
+		// Base addresses do not need to be a power of two
+		uint32 BaseX, BaseY;
+		Layout.AddElement(BaseX, BaseY, WidthInBlocks, HeightInBlocks);
+
+		Data.BlockCoordinates.Add(FUintVector2(BaseX, BaseY));
+	}
+
+	Data.BlockCount.X = FMath::Max(1u, Layout.GetSizeX());
+	Data.BlockCount.Y = FMath::Max(1u, Layout.GetSizeY());
+}
+
+static void OnTextureCollectionVirtualTextureDestroyed(const FVirtualTextureProducerHandle& InHandle, void* Baton)
+{
+	FTextureCollectionResource* Self = static_cast<FTextureCollectionResource*>(Baton);
+
+	// Reinitialize contents from the game thread objects
+	ENQUEUE_RENDER_COMMAND(Update)([Self](FRHICommandListImmediate& RHICmdList)
+	{
+		Self->InitRHI(RHICmdList);
+	});
+}
+
+void FVirtualTextureCollectionResource::CreateIndexTable(FTextureCollectionProducerData& Data)
+{
+	Data.BlockVirtualPageTable.Init(INDEX_NONE, Data.BlockCount.Y * Data.BlockCount.X);
+
+	// Create a block to texture index table
+	for (int32 TextureIndex = 0; TextureIndex < Textures.Num(); TextureIndex++)
+	{
+		FTextureResource* Resource = Textures[TextureIndex] ? Textures[TextureIndex]->GetResource() : nullptr;
+		if (!Resource)
+		{
+			continue;
+		}
+		
+		// To blocks
+		const uint32       WidthInBlocks  = FMath::DivideAndRoundUp(Resource->GetSizeX(), static_cast<uint32>(BuildSettings.TileSize));
+		const uint32       HeightInBlocks = FMath::DivideAndRoundUp(Resource->GetSizeY(), static_cast<uint32>(BuildSettings.TileSize));
+		const FUintVector2 Offset         = Data.BlockCoordinates[TextureIndex];
+
+		// Write all block indices
+		for (uint32 Y = 0; Y < HeightInBlocks; Y++)
+		{
+			const uint32 RowOffset = (Offset.Y + Y) * Data.BlockCount.X + Offset.X;
+			for (uint32 X = 0; X < WidthInBlocks; X++)
+			{
+				checkf(Data.BlockVirtualPageTable[RowOffset + X] == INDEX_NONE, TEXT("Overlapping host index page table"));
+				Data.BlockVirtualPageTable[RowOffset + X] = TextureIndex;
+			}
+		}
+	}
+}
+
+void FVirtualTextureCollectionResource::InitRHI(FRHICommandListBase& RHICmdList)
+{
+	PageTable          = nullptr;
+	PhysicalTextureSRV = nullptr;
+
+	FTextureCollectionProducerData ProducerData;
+	ProducerData.Textures.SetNumZeroed(Textures.Num());
+
+	// Compute the general texture atlas layout
+	ComputeLayout(ProducerData);
+
+	// Create the host side index table
+	CreateIndexTable(ProducerData);
+
+	// The number of mips required
+	uint32 MaxVirtualMipCount = 0;
+
+	VirtualUniforms.SetNumZeroed(Textures.Num());
+	
+	for (int32 TextureIndex = 0; TextureIndex < Textures.Num(); TextureIndex++)
+	{
+		FTextureResource* Resource = Textures[TextureIndex] ? Textures[TextureIndex]->GetResource() : nullptr;
+		if (!Resource)
+		{
+			continue;
+		}
+		
+		FTextureEntry TextureEntry;
+		if (FVirtualTexture2DResource* VirtualResource = Resource->GetVirtualTexture2DResource())
+		{
+			if (VirtualResource->GetNumLayers() != 1)
+			{
+				FormatCollectionError(TEXT("Multiple physical layers are not supported"), TextureIndex);
+				continue;
+			}
+
+			TextureEntry.Format = VirtualResource->GetFormat(0);
+			if (TextureEntry.Format != Format)
+			{
+				if (!bAllowFormatConversion)
+				{
+					FormatCollectionError(TEXT("Mismatched format to collection"), TextureIndex);
+					continue;
+				}
+				
+				TextureEntry.bRequiresAdapter = true;
+			}
+
+			// Keep the producer handle for later queries
+			TextureEntry.VirtualProducerHandle = VirtualResource->GetProducerHandle();
+
+			// Actual maximum mip count required by this collection
+			MaxVirtualMipCount = FMath::Max(MaxVirtualMipCount, VirtualResource->GetNumMips());
+			
+			// To block count
+			TextureEntry.BlockCount.X = FMath::DivideAndRoundUp(VirtualResource->GetSizeX(), static_cast<uint32>(BuildSettings.TileSize)); 
+			TextureEntry.BlockCount.Y = FMath::DivideAndRoundUp(VirtualResource->GetSizeY(), static_cast<uint32>(BuildSettings.TileSize));
+
+			// Register destruction events, recreates from the new resources
+			GetRendererModule().AddVirtualTextureProducerDestroyedCallback(
+				VirtualResource->GetProducerHandle(),
+				&OnTextureCollectionVirtualTextureDestroyed,
+				this
+			);
+		}
+		else if (FTexture2DResource* PhysicalResource = Resource->GetTexture2DResource())
+		{
+			if (!bAllowFormatConversion)
+			{
+				FormatCollectionError(TEXT("Non-virtual entries requires format conversion"), TextureIndex);
+				continue;
+			}
+
+			// Physical textures always require adapters
+			TextureEntry.bRequiresAdapter = true;
+			TextureEntry.PhysicalTexture  = Textures[TextureIndex];
+			TextureEntry.Format           = PhysicalResource->GetPixelFormat();
+			
+			// Actual maximum mip count required by this collection
+			const uint32 EntryMipLimit = FMath::CeilLogTwo(FMath::Max(PhysicalResource->GetSizeX(), PhysicalResource->GetSizeY()));
+			const uint32 EntryMipCount = FMath::Min<uint32>(EntryMipLimit, PhysicalResource->GetState().MaxNumLODs);
+			MaxVirtualMipCount = FMath::Max<uint32>(MaxVirtualMipCount, EntryMipCount);
+			
+			// To block count
+			TextureEntry.BlockCount.X = FMath::DivideAndRoundUp(PhysicalResource->GetSizeX(), static_cast<uint32>(BuildSettings.TileSize));
+			TextureEntry.BlockCount.Y = FMath::DivideAndRoundUp(PhysicalResource->GetSizeY(), static_cast<uint32>(BuildSettings.TileSize));
+		}
+		else
+		{
+			FormatCollectionError(TEXT("Invalid resource"), TextureIndex);
+			continue;
+		}
+
+		// Keep only the producer handle around
+		ProducerData.Textures[TextureIndex] = TextureEntry;
+
+		// Initialize uniforms
+		// Note that we could fit this in one uint, if we accepted 4k limits
+		UE::HLSL::FIndirectVirtualTextureEntry& VirtualUniform = VirtualUniforms[TextureIndex];
+		VirtualUniform.PackedCoordinateAndSize.X = ProducerData.BlockCoordinates[TextureIndex].X;
+		VirtualUniform.PackedCoordinateAndSize.X |= ProducerData.BlockCoordinates[TextureIndex].Y << 16;
+		VirtualUniform.PackedCoordinateAndSize.Y = TextureEntry.BlockCount.X; 
+		VirtualUniform.PackedCoordinateAndSize.Y |= TextureEntry.BlockCount.Y << 16;
+	}
+
+	// Producer description, standard virtual texture with atlassed blocks
+	FVTProducerDescription ProducerDesc;
+	ProducerDesc.Name               = TEXT("TextureCollectionVirtualRedirector");
+	ProducerDesc.FullNameHash       = GetTypeHash(ProducerDesc.Name);
+	ProducerDesc.bContinuousUpdate  = false;
+	ProducerDesc.Dimensions         = 2;
+	ProducerDesc.TileSize           = BuildSettings.TileSize;
+	ProducerDesc.TileBorderSize     = BuildSettings.TileBorderSize;
+	ProducerDesc.BlockWidthInTiles  = ProducerData.BlockCount.X;
+	ProducerDesc.BlockHeightInTiles = ProducerData.BlockCount.Y;
+	ProducerDesc.DepthInTiles       = 1u;
+	ProducerDesc.MaxLevel           = MaxVirtualMipCount;
+	ProducerDesc.NumTextureLayers   = 1;
+	ProducerDesc.NumPhysicalGroups  = 1;
+	ProducerDesc.Priority           = EVTProducerPriority::Normal;
+	ProducerDesc.LayerFormat[0]     = Format;
+	ProducerDesc.bIsLayerSRGB[0]    = bIsSRGB;
+
+	// Register producer on page feedback
+	FTextureCollectionVirtualRedirector* Producer = new FTextureCollectionVirtualRedirector(MoveTemp(ProducerData), ProducerDesc);
+	ProducerHandle = GetRendererModule().RegisterVirtualTextureProducer(RHICmdList, ProducerDesc, Producer);
+
+	// Underlying VT, standard
+	FAllocatedVTDescription VTDesc;
+	VTDesc.Dimensions            = 2;
+	VTDesc.NumTextureLayers      = 1;
+	VTDesc.TileSize              = BuildSettings.TileSize;
+	VTDesc.TileBorderSize        = BuildSettings.TileBorderSize;
+	VTDesc.bShareDuplicateLayers = false;
+	VTDesc.ProducerHandle[0]     = ProducerHandle;
+	VTDesc.ProducerLayerIndex[0] = 0;
+	
+	AllocatedVT        = GetRendererModule().AllocateVirtualTexture(VTDesc);
+	PhysicalTextureSRV = AllocatedVT->GetPhysicalTextureSRV(0, bIsSRGB);
+	PageTable          = AllocatedVT->GetPageTableTexture(0);
+
+	// Register destruction events
+	GetRendererModule().AddVirtualTextureProducerDestroyedCallback(
+		ProducerHandle,
+		&OnTextureCollectionVirtualTextureDestroyed,
+		this
+	);
+
+	// Safe fallback
+	if (!PageTable || !PhysicalTextureSRV)
+	{
+		PageTable          = GBlackUintTexture->TextureRHI;
+		PhysicalTextureSRV = GBlackTextureWithSRV->ShaderResourceViewRHI;
+	}
+
+	// Allow empty collections, zero'd page table uniforms will early out
+	if (VirtualUniforms.IsEmpty())
+	{
+		VirtualUniforms.AddZeroed();
+	}
+
+	// Create as texel buffer
+	VirtualCollectionRHI = RHICmdList.CreateBuffer(
+		FRHIBufferCreateDesc::Create(
+			TEXT("TextureCollectionVirtualUniforms"),
+			VirtualUniforms.NumBytes(),
+			sizeof(FUintVector2),
+			EBufferUsageFlags::Static | EBufferUsageFlags::ShaderResource
+		)
+		.SetInitialState(ERHIAccess::SRVMask)
+		.SetInitActionResourceArray(&VirtualUniforms)
+	);
+
+	// Buffer SRV
+	VirtualCollectionRHISRV = RHICmdList.CreateShaderResourceView(VirtualCollectionRHI, FRHIViewDesc::CreateBufferSRV()
+		.SetType(FRHIViewDesc::EBufferType::Typed)
+		.SetFormat(PF_R32G32_UINT)
+	);
+}
+
+void FVirtualTextureCollectionResource::ReleaseRHI()
+{
+	if (AllocatedVT)
+	{
+		GetRendererModule().RemoveAllVirtualTextureProducerDestroyedCallbacks(this);
+		GetRendererModule().DestroyVirtualTexture(AllocatedVT);
+		GetRendererModule().ReleaseVirtualTextureProducer(ProducerHandle);
+		AllocatedVT = nullptr;
+	}
+
+	FTextureCollectionResource::ReleaseRHI();
+}
+
+static int GetBlockCompressionFormatPriority(EPixelFormat Format)
+{
+	// Ordered first by component counts, second to source bit-width
+	static EPixelFormat Ascending[] = {
+		PF_BC4,  // R8
+		PF_BC5,  // R8G8
+		PF_BC6H, // R16G16B16
+		PF_DXT1, // R5G6B5A1
+		PF_DXT3, // R5G6B5A4
+		PF_DXT5, // R5G6B5A8
+		PF_BC7   // R8G8B8A8 (* variable)
+	};
+
+	for (int32 i = 0; i < UE_ARRAY_COUNT(Ascending); i++)
+	{
+		if (Ascending[i] == Format)
+		{
+			return i;
+		}
+	}
+	
+	checkf(false, TEXT("Unexpected block compressed format"));
+	return 0;
+}
+
+static EPixelFormat GetBlockCompressionSourceFormat(EPixelFormat Format)
+{
+	// TODO[MP]: This isn't entirely accurate, as some formats are variable bit-width
+	switch (Format) {
+		default:
+			checkf(false, TEXT("Unexpected block compressed format"));
+			return PF_R8G8B8A8;
+		case PF_BC4:
+			return PF_R8;
+		case PF_BC5:
+			return PF_R8G8;
+		case PF_BC6H:
+			return PF_FloatRGB;
+		case PF_DXT1:
+			return PF_B5G5R5A1_UNORM;
+		case PF_DXT3:
+			return PF_R8G8B8A8;
+		case PF_DXT5:
+			return PF_R8G8B8A8;
+		case PF_BC7: 
+			return PF_R8G8B8A8;
+	};
+}
+
+static EPixelFormat FindConservativeSourceFormat(int32 SourceComponentCount, int32 SourceByteWidth)
+{
+	EPixelFormat     CandidateFormat = PF_A32B32G32R32F;
+	FPixelFormatInfo CandidateInfo   = GPixelFormats[CandidateFormat];
+
+	for (const FPixelFormatInfo& PixelFormat : GPixelFormats)
+	{
+		if (!PixelFormat.Supported || IsBlockCompressedFormat(PixelFormat.UnrealFormat))
+		{
+			continue;
+		}
+
+		// Check if this format can accomodate at all
+		if (PixelFormat.NumComponents < SourceComponentCount || PixelFormat.BlockBytes < SourceByteWidth)
+		{
+			continue;
+		}
+		
+		// Check if this format has either a redued component count or byte-width
+		if (PixelFormat.NumComponents < CandidateInfo.NumComponents || PixelFormat.BlockBytes < CandidateInfo.BlockBytes)
+		{
+			CandidateFormat = PixelFormat.UnrealFormat;
+			CandidateInfo   = GPixelFormats[CandidateFormat];
+		}
+	}
+
+	return CandidateFormat;
+}
+
+void FVirtualTextureCollectionResource::FindFirstFormat()
+{
+	// Find the first valid resource
+	for (UTexture* Texture : Textures)
+	{
+		FTextureResource* Resource = Texture ? Texture->GetResource() : nullptr;
+		if (!Resource)
+		{
+			continue;
+		}
+
+		EPixelFormat TextureFormat = GetResourcePixelFormat(Resource);
+		if (TextureFormat != PF_Unknown)
+		{
+			Format  = TextureFormat;
+			bIsSRGB = Resource->bSRGB;
+			return;
+		}
+	}
+
+	// No valid resource, dummy format
+	Format  = PF_R8G8B8A8;
+	bIsSRGB = true;
+}
+
+void FVirtualTextureCollectionResource::FindConservativeFormat()
+{
+	EPixelFormat CandidateFormat = PF_Unknown;
+	bool         bAnyFormatSRGB  = false;
+
+	int32 SourceComponentCount  = 0;
+	int32 SourceByteWidth       = 0;
+
+	for (UTexture* Texture : Textures)
+	{
+		FTextureResource* Resource = Texture ? Texture->GetResource() : nullptr;
+		if (!Resource)
+		{
+			continue;
+		}
+
+		// If any format is SRGB, so is this
+		bAnyFormatSRGB |= Resource->bSRGB;
+
+		EPixelFormat TextureFormat = GetResourcePixelFormat(Resource);
+		if (TextureFormat == PF_Unknown || TextureFormat == CandidateFormat)
+		{
+			continue;
+		}
+
+		// If first format, just accept it as is
+		if (CandidateFormat == PF_Unknown)
+		{
+			CandidateFormat = TextureFormat;
+			continue;
+		}
+
+		const bool bIsCandidateBC = IsBlockCompressedFormat(CandidateFormat);
+		const bool bIsTextureBC   = IsBlockCompressedFormat(TextureFormat);
+
+		// We keep the bit-width as is, so opt for decompressing the formats
+		if (!bIsCandidateBC || !bIsTextureBC)
+		{
+			// "Decompress" if needed
+			if (bIsCandidateBC)
+			{
+				CandidateFormat = GetBlockCompressionSourceFormat(CandidateFormat);
+			}
+			
+			// "Decompress" if needed
+			if (bIsTextureBC)
+			{
+				TextureFormat = GetBlockCompressionSourceFormat(TextureFormat);
+			}
+
+			const FPixelFormatInfo& CandidateInfo = GPixelFormats[TextureFormat];
+			const FPixelFormatInfo& TextureInfo   = GPixelFormats[TextureFormat];
+
+			// Search later on the channel count and byte-width
+			SourceComponentCount  = FMath::Max3(SourceComponentCount, TextureInfo.NumComponents, CandidateInfo.NumComponents);
+			SourceByteWidth       = FMath::Max3(SourceByteWidth,      TextureInfo.BlockBytes,    CandidateInfo.BlockBytes);
+		}
+		else
+		{
+			// Both are block compressed, find the best fit
+			int Priority = GetBlockCompressionFormatPriority(TextureFormat);
+			if (Priority > GetBlockCompressionFormatPriority(CandidateFormat))
+			{
+				CandidateFormat = TextureFormat;
+			}
+		}
+	}
+
+	// If this is a source format, find a conservative pixel format
+	if (SourceComponentCount)
+	{
+		Format  = FindConservativeSourceFormat(SourceComponentCount, SourceByteWidth);
+		bIsSRGB = bAnyFormatSRGB;
+		return;
+	}
+
+	// No relevant resources, assign a dummy format
+	if (CandidateFormat == PF_Unknown)
+	{
+		Format  = PF_R8G8B8A8;
+		bIsSRGB = true;
+		return;
+	}
+
+	Format  = CandidateFormat;
+	bIsSRGB = bAnyFormatSRGB;
+}
+
+UE::HLSL::FIndirectVirtualTextureUniform FVirtualTextureCollectionResource::GetVirtualPackedUniform() const
+{
+	UE::HLSL::FIndirectVirtualTextureUniform Out{};
+		
+	if (!Textures.IsEmpty())
+	{
+		Out.UniformCountSub1 = Textures.Num() - 1;
+	}
+
+	AllocatedVT->GetPackedPageTableUniform(Out.PackedPageTableUniform);
+	AllocatedVT->GetPackedUniform(&Out.PackedUniform, 0);
+	
+	return Out;
+}
+
+void FVirtualTextureCollectionResource::FormatCollectionError(const TCHAR* Reason, uint32 TextureIndex)
+{
+	UE_LOG(
+		LogVirtualTextureCollection, Error,
+		TEXT("Texture collection '%s' received texture [%i] '%s' - %s"),
+		*CollectionName.ToString(), TextureIndex, *Textures[TextureIndex]->GetFName().ToString(), Reason
+	);
+}
+
+#if WITH_EDITOR
+void UVirtualTextureCollection::ValidateVirtualCollection()
+{
+	ETextureSourceFormat   Format         = TSF_Invalid;
+	FTextureFormatSettings FormatSettings = {};
+	
+	for (int32 TextureIndex = 0; TextureIndex < Textures.Num(); TextureIndex++)
+	{
+		UTexture* Texture = Textures[TextureIndex];
+		if (!Texture)
+		{
+			// Null/default textures are allowed
+			continue;
+		}
+
+		UTexture2D* Texture2D = Cast<UTexture2D>(Texture);
+		if (!Texture2D)
+		{
+			FormatCollectionError(TEXT("Virtual collections only support 2d textures"), TextureIndex);
+			continue;
+		}
+
+		// Format conversion can handle any kind of format differences
+		if (bAllowFormatConversion)
+		{
+			continue;
+		}
+
+		// Actual formats are only known after building the textures
+		// In case it's not built, let's just validate the source format settings
+		FTextureFormatSettings TextureFormat;
+		Texture->GetLayerFormatSettings(0, TextureFormat);
+			
+		if (Format == TSF_Invalid)
+		{
+			Format         = Texture->Source.GetFormat(0);
+			FormatSettings = TextureFormat;
+		}
+		else if (Format != Texture->Source.GetFormat(0) || FormatSettings != TextureFormat)
+		{
+			FormatCollectionError(TEXT("Mismatched format settings, all textures in a (virtual) collection must share the same format settings"), TextureIndex);
+		}
+	}
+}
+
+void UVirtualTextureCollection::FormatCollectionError(const TCHAR* Reason, uint32 TextureIndex)
+{
+	UE_LOG(
+		LogVirtualTextureCollection, Error,
+		TEXT("Texture collection '%s' received texture [%i] '%s' - %s"),
+		*GetFName().ToString(), TextureIndex, *Textures[TextureIndex]->GetFName().ToString(), Reason
+	);
+}
+#endif //  WITH_EDITOR
+
+#if WITH_EDITOR
+void UVirtualTextureCollection::PostEditChangeProperty(FPropertyChangedEvent& PropertyChangedEvent)
+{
+	UTextureCollection::PostEditChangeProperty(PropertyChangedEvent);
+	ValidateVirtualCollection();
+}
+#endif // WITH_EDITOR
+
+void UVirtualTextureCollection::BeginDestroy()
+{
+	ReleaseResource();
+	Super::BeginDestroy();
+}
+
+FTextureCollectionResource* UVirtualTextureCollection::CreateResource()
+{
+	// Wait for the virtual resources to finish compilation
+#if WITH_EDITOR
+	FTextureCompilingManager::Get().FinishCompilation(Textures);
+#endif // WITH_EDITOR
+
+	return new FVirtualTextureCollectionResource(this);
+}
+
+#undef LOCTEXT_NAMESPACE
diff --git a/Engine/Source/Runtime/Engine/Public/MaterialShared.h b/Engine/Source/Runtime/Engine/Public/MaterialShared.h
--- a/Engine/Source/Runtime/Engine/Public/MaterialShared.h
+++ b/Engine/Source/Runtime/Engine/Public/MaterialShared.h
@@ -569,7 +569,7 @@ class FMaterialTextureCollectionParameterInfo
 public:
 	friend inline bool operator==(const FMaterialTextureCollectionParameterInfo& Lhs, const FMaterialTextureCollectionParameterInfo& Rhs)
 	{
-		return Lhs.ParameterInfo == Rhs.ParameterInfo && Lhs.TextureCollectionIndex == Rhs.TextureCollectionIndex;
+		return Lhs.ParameterInfo == Rhs.ParameterInfo && Lhs.TextureCollectionIndex == Rhs.TextureCollectionIndex && Lhs.bIsVirtualCollection == Rhs.bIsVirtualCollection;
 	}
 	friend inline bool operator!=(const FMaterialTextureCollectionParameterInfo& Lhs, const FMaterialTextureCollectionParameterInfo& Rhs)
 	{
@@ -581,6 +581,7 @@ public:
 
 	LAYOUT_FIELD(FHashedMaterialParameterInfo, ParameterInfo);
 	LAYOUT_FIELD(int32, TextureCollectionIndex);
+	LAYOUT_FIELD(bool, bIsVirtualCollection);
 };
 
 class FMaterialExternalTextureParameterInfo
@@ -713,6 +714,7 @@ public:
 
 	int32 GetNumTextureCollections() const { return UniformTextureCollectionParameters.Num(); }
 	ENGINE_API void GetGameThreadTextureCollectionValue(int32 Index, const UMaterialInterface* MaterialInterface, const FMaterial& Material, UTextureCollection*& OutValue) const;
+	ENGINE_API void GetTextureCollectionValue(int32 Index, const FMaterialRenderContext& Context, const FMaterial& Material, const UTextureCollection*& OutValue) const;
 
 	int32 FindOrAddTextureParameter(EMaterialTextureParameterType Type, const FMaterialTextureParameterInfo& Info);
 	int32 FindOrAddTextureCollectionParameter(const FMaterialTextureCollectionParameterInfo& Info);
@@ -746,6 +748,8 @@ public:
 	/// @param EntryGenerator Callback that writes the preshader opcodes.
 	void WriteUniformPreshaderEntry(uint32 BufferOffset, UE::Shader::EValueType EntryType, const TFunction<void(UE::Shader::FPreshaderData&)>& EntryGenerator);
 
+	void CountTextureCollections(uint32& BindlessCollections, uint32& VirtualCollections) const;
+
 protected:
 	union FVTPackedStackAndLayerIndex
 	{
diff --git a/Engine/Source/Runtime/Engine/Public/MaterialValueType.h b/Engine/Source/Runtime/Engine/Public/MaterialValueType.h
--- a/Engine/Source/Runtime/Engine/Public/MaterialValueType.h
+++ b/Engine/Source/Runtime/Engine/Public/MaterialValueType.h
@@ -76,6 +76,9 @@ enum EMaterialValueType : uint64
 	/** Material cache */
 	MCT_MaterialCacheABuffer  = 1ull << 35,
 
+	/** An internal type */
+	MCT_Unexposed             = 1ull << 36,
+
 	/** Reserved range for licensee modifications, any bits in the (inclusive) range is free for modifications */
 	MCT_LicenseeReservedBegin     = 1ull << 48,
 	MCT_LicenseeReservedEnd       = 1ull << 63,
diff --git a/Engine/Source/Runtime/RenderCore/Private/GlobalRenderResources.cpp b/Engine/Source/Runtime/RenderCore/Private/GlobalRenderResources.cpp
--- a/Engine/Source/Runtime/RenderCore/Private/GlobalRenderResources.cpp
+++ b/Engine/Source/Runtime/RenderCore/Private/GlobalRenderResources.cpp
@@ -104,6 +104,33 @@ public:
 	}
 };
 
+class FEmptyUInt4VertexBuffer : public FVertexBufferWithSRV
+{
+public:
+	virtual void InitRHI(FRHICommandListBase& RHICmdList) override
+	{
+		const FRHIBufferCreateDesc CreateDesc =
+			FRHIBufferCreateDesc::CreateVertex(TEXT("EmptyUInt4VertexBuffer"), sizeof(FUintVector4))
+			.AddUsage(BUF_Static | BUF_ShaderResource | BUF_UnorderedAccess)
+			.DetermineInitialState();
+		
+		const FUintVector4 InitialData(ForceInitToZero);
+		VertexBufferRHI = UE::RHIResourceUtils::CreateBufferWithValue(RHICmdList, CreateDesc, InitialData);
+
+		ShaderResourceViewRHI = RHICmdList.CreateShaderResourceView(
+			VertexBufferRHI,
+			FRHIViewDesc::CreateBufferSRV()
+				.SetType(FRHIViewDesc::EBufferType::Typed)
+				.SetFormat(PF_R32G32B32A32_UINT));
+
+		UnorderedAccessViewRHI = RHICmdList.CreateUnorderedAccessView(
+			VertexBufferRHI,
+			FRHIViewDesc::CreateBufferUAV()
+				.SetType(FRHIViewDesc::EBufferType::Typed)
+				.SetFormat(PF_R32G32B32A32_UINT));
+	}
+};
+
 class FEmptyStructuredBuffer : public FVertexBufferWithSRV
 {
 public:
@@ -227,6 +254,7 @@ FTexture* GBlackTexture = GBlackTextureWithSRV;
 FTexture* GTransparentBlackTexture = GTransparentBlackTextureWithSRV;
 
 FVertexBufferWithSRV* GEmptyVertexBufferWithUAV = new TGlobalResource<FEmptyVertexBuffer, FRenderResource::EInitPhase::Pre>;
+FVertexBufferWithSRV* GEmptyVertexBufferUInt4WithUAV = new TGlobalResource<FEmptyUInt4VertexBuffer, FRenderResource::EInitPhase::Pre>;
 FVertexBufferWithSRV* GEmptyStructuredBufferWithUAV = new TGlobalResource<FEmptyStructuredBuffer, FRenderResource::EInitPhase::Pre>;
 FVertexBufferWithSRV* GBlackFloat4StructuredBufferWithSRV = new TGlobalResource<FBlackFloat4StructuredBufferWithSRV, FRenderResource::EInitPhase::Pre>;
 FVertexBufferWithSRV* GBlackFloat4VertexBufferWithSRV = new TGlobalResource<FBlackFloat4VertexBuffer, FRenderResource::EInitPhase::Pre>;
@@ -739,6 +767,18 @@ void FTwoTrianglesIndexBuffer::InitRHI(FRHICommandListBase& RHICmdList)
 
 TGlobalResource<FTwoTrianglesIndexBuffer, FRenderResource::EInitPhase::Pre> GTwoTrianglesIndexBuffer;
 
+void FEmptyResourceCollection::InitRHI(FRHICommandListBase& RHICmdList)
+{
+	ResourceCollection = RHICmdList.CreateResourceCollection({});
+}
+
+void FEmptyResourceCollection::ReleaseRHI()
+{
+	ResourceCollection.SafeRelease();
+}
+
+TGlobalResource<FEmptyResourceCollection, FRenderResource::EInitPhase::Pre> GEmptyResourceCollection;
+
 /*------------------------------------------------------------------------------
 	FGlobalDynamicVertexBuffer implementation.
 ------------------------------------------------------------------------------*/
diff --git a/Engine/Source/Runtime/RenderCore/Private/ShaderParameterMetadata.cpp b/Engine/Source/Runtime/RenderCore/Private/ShaderParameterMetadata.cpp
--- a/Engine/Source/Runtime/RenderCore/Private/ShaderParameterMetadata.cpp
+++ b/Engine/Source/Runtime/RenderCore/Private/ShaderParameterMetadata.cpp
@@ -383,6 +383,10 @@ EShaderCodeResourceBindingType ParseShaderResourceBindingType(const TCHAR* Shade
 	{
 		BindingType = EShaderCodeResourceBindingType::ResourceCollection;
 	}
+	else if (!bIsRWResource && ShaderTypeLength == 23 && FCString::Strncmp(ComparedShaderType, TEXT("FIndirectVirtualTexture"), ShaderTypeLength) == 0)
+	{
+		BindingType = EShaderCodeResourceBindingType::Buffer;
+	}
 	else if (bIsRasterizerOrderedResource)
 	{
 		BindingType = EShaderCodeResourceBindingType::RasterizerOrderedTexture2D;
@@ -777,7 +781,7 @@ void FShaderParametersMetadata::InitializeLayout(FRHIUniformBufferLayoutInitiali
 				}
 				else if (BindingType == EShaderCodeResourceBindingType::ResourceCollection)
 				{
-					bIsValidBindingType = (BaseType == UBMT_RESOURCE_COLLECTION);
+					bIsValidBindingType = BaseType == UBMT_RESOURCE_COLLECTION;
 				}
 				else
 				{
diff --git a/Engine/Source/Runtime/RenderCore/Public/GlobalRenderResources.h b/Engine/Source/Runtime/RenderCore/Public/GlobalRenderResources.h
--- a/Engine/Source/Runtime/RenderCore/Public/GlobalRenderResources.h
+++ b/Engine/Source/Runtime/RenderCore/Public/GlobalRenderResources.h
@@ -28,6 +28,9 @@ extern RENDERCORE_API FTexture* GDummyTransitionTexture;
 // An empty vertex buffer with a UAV
 extern RENDERCORE_API FVertexBufferWithSRV* GEmptyVertexBufferWithUAV;
 
+// An empty uint4 vertex buffer with a UAV
+extern RENDERCORE_API FVertexBufferWithSRV* GEmptyVertexBufferUInt4WithUAV;
+
 // An empty structured buffer with a UAV
 extern RENDERCORE_API FVertexBufferWithSRV* GEmptyStructuredBufferWithUAV;
 
@@ -164,6 +167,19 @@ public:
 };
 extern RENDERCORE_API TGlobalResource<FTwoTrianglesIndexBuffer, FRenderResource::EInitPhase::Pre> GTwoTrianglesIndexBuffer;
 
+class FEmptyResourceCollection : public FRenderResource
+{
+public:
+	FEmptyResourceCollection() = default;
+	~FEmptyResourceCollection() = default;
+
+	RENDERCORE_API virtual void InitRHI(FRHICommandListBase& RHICmdList) override;
+	RENDERCORE_API virtual void ReleaseRHI() override;
+
+	FRHIResourceCollectionRef ResourceCollection;
+};
+
+extern RENDERCORE_API TGlobalResource<FEmptyResourceCollection, FRenderResource::EInitPhase::Pre> GEmptyResourceCollection;
 
 //////////////////////////////////////////////////////////////////////////////////////////////////////////////////////
 // FGlobalDynamicVertexBuffer
diff --git a/Engine/Source/Runtime/RenderCore/Public/RendererInterface.h b/Engine/Source/Runtime/RenderCore/Public/RendererInterface.h
--- a/Engine/Source/Runtime/RenderCore/Public/RendererInterface.h
+++ b/Engine/Source/Runtime/RenderCore/Public/RendererInterface.h
@@ -820,6 +820,8 @@ public:
 	virtual uint32 RemoveAllVirtualTextureProducerDestroyedCallbacks(const void* Baton) = 0;
 	virtual void ReleaseVirtualTexturePendingResources() = 0;
 
+	virtual IVirtualTexture* FindProducer(const FVirtualTextureProducerHandle& Handle) = 0;
+
 	virtual void RequestVirtualTextureTiles(const FVector2D& InScreenSpaceSize, int32 InMipLevel) = 0;
 	virtual void RequestVirtualTextureTiles(const FMaterialRenderProxy* InMaterialRenderProxy, const FVector2D& InScreenSpaceSize, ERHIFeatureLevel::Type InFeatureLevel) = 0;
 
diff --git a/Engine/Source/Runtime/Renderer/Private/RendererModule.h b/Engine/Source/Runtime/Renderer/Private/RendererModule.h
--- a/Engine/Source/Runtime/Renderer/Private/RendererModule.h
+++ b/Engine/Source/Runtime/Renderer/Private/RendererModule.h
@@ -112,6 +112,7 @@ public:
 	virtual void RequestVirtualTextureTiles(const FVector2D& InScreenSpaceSize, int32 InMipLevel) override;
 	virtual void RequestVirtualTextureTiles(const FMaterialRenderProxy* InMaterialRenderProxy, const FVector2D& InScreenSpaceSize, ERHIFeatureLevel::Type InFeatureLevel) override;
 	virtual void RequestVirtualTextureTiles(IAllocatedVirtualTexture* AllocatedVT, const FVector2D& InScreenSpaceSize, const FVector2D& InViewportPosition, const FVector2D& InViewportSize, const FVector2D& InUV0, const FVector2D& InUV1, int32 InMipLevel) override;
+	virtual IVirtualTexture* FindProducer(const FVirtualTextureProducerHandle& Handle) override;
 
 	UE_DEPRECATED(5.4, "Use RequestVirtualTextureTiles() overloads that takes similar parameters. Make sure not to negate the InViewportPosition.")
 	virtual void RequestVirtualTextureTilesForRegion(IAllocatedVirtualTexture* AllocatedVT, const FVector2D& InScreenSpaceSize, const FVector2D& InViewportPosition, const FVector2D& InViewportSize, const FVector2D& InUV0, const FVector2D& InUV1, int32 InMipLevel) override;
diff --git a/Engine/Source/Runtime/Renderer/Private/SceneRendering.cpp b/Engine/Source/Runtime/Renderer/Private/SceneRendering.cpp
--- a/Engine/Source/Runtime/Renderer/Private/SceneRendering.cpp
+++ b/Engine/Source/Runtime/Renderer/Private/SceneRendering.cpp
@@ -5521,6 +5521,16 @@ void FRendererModule::RequestVirtualTextureTiles(IAllocatedVirtualTexture* Alloc
 	FVirtualTextureSystem::Get().RequestTiles(AllocatedVT, InScreenSpaceSize, InViewportPosition, InViewportSize, InUV0, InUV1, InMipLevel);
 }
 
+IVirtualTexture* FRendererModule::FindProducer(const FVirtualTextureProducerHandle& Handle)
+{
+	if (FVirtualTextureProducer* Producer = FVirtualTextureSystem::Get().FindProducer(Handle))
+	{
+		return Producer->GetVirtualTexture();
+	}
+
+	return nullptr;
+}
+
 void FRendererModule::RequestVirtualTextureTilesForRegion(IAllocatedVirtualTexture* AllocatedVT, const FVector2D& InScreenSpaceSize, const FVector2D& InViewportPosition, const FVector2D& InViewportSize, const FVector2D& InUV0, const FVector2D& InUV1, int32 InMipLevel)
 {
 PRAGMA_DISABLE_DEPRECATION_WARNINGS
diff --git a/Engine/Source/Runtime/Renderer/Private/VT/VirtualTextureSystem.h b/Engine/Source/Runtime/Renderer/Private/VT/VirtualTextureSystem.h
--- a/Engine/Source/Runtime/Renderer/Private/VT/VirtualTextureSystem.h
+++ b/Engine/Source/Runtime/Renderer/Private/VT/VirtualTextureSystem.h
@@ -182,6 +182,7 @@ private:
 	friend class FVirtualTextureProducerCollection;
 	friend class FTexturePageMap;
 	friend class FTexturePagePool;
+	friend class FRendererModule;
 
 	FVirtualTextureSystem();
 	~FVirtualTextureSystem();
