Iris
* Removed UE_WITH_IRIS from EngineBridge and NetFactories

#rnx

diff --git a/Engine/Source/Runtime/Engine/Private/Net/Iris/ReplicationSystem/EngineReplicationBridge.cpp b/Engine/Source/Runtime/Engine/Private/Net/Iris/ReplicationSystem/EngineReplicationBridge.cpp
--- a/Engine/Source/Runtime/Engine/Private/Net/Iris/ReplicationSystem/EngineReplicationBridge.cpp
+++ b/Engine/Source/Runtime/Engine/Private/Net/Iris/ReplicationSystem/EngineReplicationBridge.cpp
@@ -3,8 +3,6 @@
 
 #include UE_INLINE_GENERATED_CPP_BY_NAME(EngineReplicationBridge)
 
-#if UE_WITH_IRIS
-
 #include "Net/Iris/ReplicationSystem/NetActorFactory.h"
 #include "Net/Iris/ReplicationSystem/NetSubObjectFactory.h"
 
@@ -1080,10 +1078,3 @@ void UEngineReplicationBridge::ConsumeNetMetrics(TArray<FAnalyticsEventAttribute
 
 	GetReplicationSystem()->ResetNetMetrics();
 }
-
-#else //!UE_WITH_IRIS
-
-UEngineReplicationBridge::UEngineReplicationBridge() = default;
-UEngineReplicationBridge::~UEngineReplicationBridge() = default;
-
-#endif
diff --git a/Engine/Source/Runtime/Engine/Private/Net/Iris/ReplicationSystem/NetActorFactory.cpp b/Engine/Source/Runtime/Engine/Private/Net/Iris/ReplicationSystem/NetActorFactory.cpp
--- a/Engine/Source/Runtime/Engine/Private/Net/Iris/ReplicationSystem/NetActorFactory.cpp
+++ b/Engine/Source/Runtime/Engine/Private/Net/Iris/ReplicationSystem/NetActorFactory.cpp
@@ -2,8 +2,6 @@
 
 #include "Net/Iris/ReplicationSystem/NetActorFactory.h"
 
-#if UE_WITH_IRIS
-
 #include "Engine/Level.h"
 #include "Engine/NetDriver.h"
 #include "Engine/World.h"
@@ -29,13 +27,10 @@
 
 #include "UObject/Package.h"
 
-#endif // UE_WITH_IRIS
-
 #include UE_INLINE_GENERATED_CPP_BY_NAME(NetActorFactory)
 
 namespace UE::Net::Private
 {
-#if UE_WITH_IRIS
 	static int32 SerializeNewActorMaxBunchSize = 1024;
 	static FAutoConsoleVariableRef CVarSerializeNewActorMaxBunchSize(
 	TEXT("net.Iris.SerializeNewActorMaxBunchSize"),
@@ -85,11 +80,8 @@ namespace UE::Net::Private
 	}
 
 	
-#endif // UE_WITH_IRIS
 } // end namespace UE::Net::Private
 
-#if UE_WITH_IRIS
-
 //------------------------------------------------------------------------
 // UNetActorFactory
 //------------------------------------------------------------------------
@@ -747,5 +739,4 @@ FString FPreRegisteredActorNetCreationHeader::ToString() const
 
 } // end namespace UE::Net
 
-#endif // UE_WITH_IRIS
 
diff --git a/Engine/Source/Runtime/Engine/Private/Net/Iris/ReplicationSystem/NetSubObjectFactory.cpp b/Engine/Source/Runtime/Engine/Private/Net/Iris/ReplicationSystem/NetSubObjectFactory.cpp
--- a/Engine/Source/Runtime/Engine/Private/Net/Iris/ReplicationSystem/NetSubObjectFactory.cpp
+++ b/Engine/Source/Runtime/Engine/Private/Net/Iris/ReplicationSystem/NetSubObjectFactory.cpp
@@ -2,8 +2,6 @@
 
 #include "Net/Iris/ReplicationSystem/NetSubObjectFactory.h"
 
-#if UE_WITH_IRIS
-
 #include "GameFramework/Actor.h"
 
 #include "HAL/UnrealMemory.h"
@@ -22,12 +20,8 @@
 
 #include "Net/DataBunch.h"
 
-#endif // UE_WITH_IRIS
-
 #include UE_INLINE_GENERATED_CPP_BY_NAME(NetSubObjectFactory)
 
-#if UE_WITH_IRIS
-
 namespace UE::Net::Private
 {
 	static bool bWarnWhenOuterIsReplaced = false;
@@ -388,5 +382,3 @@ bool FNetDynamicSubObjectCreationHeader::Deserialize(const FCreationHeaderContex
 }
 
 } // end namespace UE::Net
-
-#endif //UE_WITH_IRIS
diff --git a/Engine/Source/Runtime/Engine/Public/Net/Iris/ReplicationSystem/EngineReplicationBridge.h b/Engine/Source/Runtime/Engine/Public/Net/Iris/ReplicationSystem/EngineReplicationBridge.h
--- a/Engine/Source/Runtime/Engine/Public/Net/Iris/ReplicationSystem/EngineReplicationBridge.h
+++ b/Engine/Source/Runtime/Engine/Public/Net/Iris/ReplicationSystem/EngineReplicationBridge.h
@@ -11,8 +11,6 @@ class UNetDriver;
 class UIrisObjectReferencePackageMap;
 struct FAnalyticsEventAttribute;
 
-#if UE_WITH_IRIS
-
 class UActorComponent;
 class UWorld;
 
@@ -81,8 +79,6 @@ struct FStopReplicatingActorParams
 	EEndPlayReason::Type EndPlayReason = EEndPlayReason::Type::Destroyed;
 };
 
-#endif // UE_WITH_IRIS
-
 UCLASS(Transient, MinimalAPI, config=Engine)
 class UEngineReplicationBridge final : public UObjectReplicationBridge
 {
@@ -92,8 +88,6 @@ public:
 	ENGINE_API UEngineReplicationBridge();
 	virtual ENGINE_API ~UEngineReplicationBridge() override;
 
-#if UE_WITH_IRIS
-
 	ENGINE_API static UEngineReplicationBridge* Create(UNetDriver* NetDriver);
 
 	/** Sets the net driver for the bridge. */
@@ -182,8 +176,6 @@ private:
 	UE::Net::FNetObjectFactoryId ActorFactoryId;
 	UE::Net::FNetObjectFactoryId SubObjectFactoryId;
 
-#endif // UE_WITH_IRIS
-
 	UNetDriver* NetDriver = nullptr;
 
 	UPROPERTY(Transient)
diff --git a/Engine/Source/Runtime/Engine/Public/Net/Iris/ReplicationSystem/NetActorFactory.h b/Engine/Source/Runtime/Engine/Public/Net/Iris/ReplicationSystem/NetActorFactory.h
--- a/Engine/Source/Runtime/Engine/Public/Net/Iris/ReplicationSystem/NetActorFactory.h
+++ b/Engine/Source/Runtime/Engine/Public/Net/Iris/ReplicationSystem/NetActorFactory.h
@@ -4,9 +4,7 @@
 
 #include "Iris/ReplicationSystem/NetObjectFactory.h"
 
-#if UE_WITH_IRIS
 #include "Iris/Core/NetObjectReference.h"
-#endif
 
 #include "NetActorFactory.generated.h"
 
@@ -28,8 +26,6 @@ class UNetActorFactory;
 namespace UE::Net
 {
 
-#if UE_WITH_IRIS
-
 /**
  * Header information to be able to tell if its a dynamic or static header
  */
@@ -166,8 +162,6 @@ public:
 	ENGINE_API virtual FString ToString() const override;
 };
 
-#endif // UE_WITH_IRIS
-
 } // end namespace UE::Net
 
 /**
@@ -178,8 +172,6 @@ class UNetActorFactory : public UNetObjectFactory
 {
 	GENERATED_BODY()
 
-#if UE_WITH_IRIS
-
 public:
 
 	ENGINE_API static FName GetFactoryName();
@@ -226,7 +218,6 @@ protected:
 	/** Cache of the spawn info flags created from the net.QuantizeActor cvars */
 	UE::Net::EActorNetSpawnInfoFlags SpawnInfoFlags;
 
-#endif // UE_WITH_IRIS
 };
 
 
diff --git a/Engine/Source/Runtime/Engine/Public/Net/Iris/ReplicationSystem/NetSubObjectFactory.h b/Engine/Source/Runtime/Engine/Public/Net/Iris/ReplicationSystem/NetSubObjectFactory.h
--- a/Engine/Source/Runtime/Engine/Public/Net/Iris/ReplicationSystem/NetSubObjectFactory.h
+++ b/Engine/Source/Runtime/Engine/Public/Net/Iris/ReplicationSystem/NetSubObjectFactory.h
@@ -4,16 +4,13 @@
 
 #include "Iris/ReplicationSystem/NetObjectFactory.h"
 
-#if UE_WITH_IRIS
 #include "Iris/Core/NetObjectReference.h"
-#endif
 
 #include "NetSubObjectFactory.generated.h"
 
 
 namespace UE::Net
 {
-#if UE_WITH_IRIS
 
 /**
 * Header information to to tell if it is a dynamic or static header
@@ -74,8 +71,6 @@ public:
 
 };
 
-#endif // UE_WITH_IRIS
-
 } // end namespace UE::Net
 
 /**
@@ -86,8 +81,6 @@ class UNetSubObjectFactory : public UNetObjectFactory
 {
 	GENERATED_BODY()
 
-#if UE_WITH_IRIS
-
 public:
 
 	ENGINE_API static FName GetFactoryName();
@@ -111,7 +104,5 @@ protected:
 
 	/** Write the needed header information for a dynamic subobject */
 	ENGINE_API bool FillDynamicHeader(UE::Net::FNetDynamicSubObjectCreationHeader* DynamicHeader, const UObject* SubObject, UE::Net::FNetRefHandle Handle);
-
-#endif // UE_WITH_IRIS
 };
 
