Add target platform name to logging for when a UDIM texture needs to discard block data during texture build.
#rnx

diff --git a/Engine/Source/Runtime/Engine/Private/TextureDerivedData.cpp b/Engine/Source/Runtime/Engine/Private/TextureDerivedData.cpp
--- a/Engine/Source/Runtime/Engine/Private/TextureDerivedData.cpp
+++ b/Engine/Source/Runtime/Engine/Private/TextureDerivedData.cpp
@@ -1299,9 +1299,9 @@ static void GetTextureBuildSettings(
 	}
 	if ( Texture.Source.GetNumBlocks() > 1 && !bVirtualTextureStreaming )
 	{
-		UE_LOG(LogTexture, Warning, TEXT("Texture %s has UDIM Blocks, but bVirtualTextureStreaming is off; will build just the first block."), 
-			*Texture.GetPathName()
-			);
+		UE_LOG(LogTexture, Warning, TEXT("Texture '%s' has UDIM Blocks, but Virtual Texturing is off for platform '%s'. Only the first UDIM block will be built for the platform."),
+			*Texture.GetPathName(),
+			*TargetPlatform.IniPlatformName());
 	}
 
 	const FIntPoint SourceSize = Texture.Source.GetLogicalSize();
