expose depth testing option to bp debug draw methods

diff --git a/Engine/Source/Runtime/Engine/Classes/Kismet/KismetSystemLibrary.h b/Engine/Source/Runtime/Engine/Classes/Kismet/KismetSystemLibrary.h
--- a/Engine/Source/Runtime/Engine/Classes/Kismet/KismetSystemLibrary.h
+++ b/Engine/Source/Runtime/Engine/Classes/Kismet/KismetSystemLibrary.h
@@ -44,6 +44,16 @@ namespace EDrawDebugTrace
 	};
 }
 
+/** Enum used to indicate the scene depth priority -- whether to draw with depth testing vs the world, or to draw in the foreground */
+UENUM(BlueprintType)
+enum class EDrawDebugSceneDepthPriorityGroup : uint8
+{
+	// Draw in world with depth testing, allowing occlusion by other surfaces
+	World,
+	// Draw in foreground, so drawn shapes will always be visible
+	Foreground
+};
+
 /** Enum used to indicate desired behavior for MoveComponentTo latent function. */
 UENUM()
 namespace EMoveComponentAction
@@ -1647,35 +1657,35 @@ class UKismetSystemLibrary : public UBlueprintFunctionLibrary
 
 	/** Draw a debug line */
 	UFUNCTION(BlueprintCallable, Category="Rendering|Debug", meta=(WorldContext="WorldContextObject", DevelopmentOnly))
-	static ENGINE_API void DrawDebugLine(const UObject* WorldContextObject, const FVector LineStart, const FVector LineEnd, FLinearColor LineColor, float Duration=0.f, float Thickness = 0.f);
+	static ENGINE_API void DrawDebugLine(const UObject* WorldContextObject, const FVector LineStart, const FVector LineEnd, FLinearColor LineColor, float Duration=0.f, float Thickness = 0.f, EDrawDebugSceneDepthPriorityGroup DepthPriority = EDrawDebugSceneDepthPriorityGroup::World);
 
 	/** Draw a debug circle! */
 	UFUNCTION(BlueprintCallable, Category="Rendering|Debug", meta=(WorldContext="WorldContextObject", DevelopmentOnly))
-	static ENGINE_API void DrawDebugCircle(const UObject* WorldContextObject, FVector Center, float Radius, int32 NumSegments=12, FLinearColor LineColor = FLinearColor::White, float Duration=0.f, float Thickness=0.f, FVector YAxis=FVector(0.f,1.f,0.f),FVector ZAxis=FVector(0.f,0.f,1.f), bool bDrawAxis=false);
+	static ENGINE_API void DrawDebugCircle(const UObject* WorldContextObject, FVector Center, float Radius, int32 NumSegments=12, FLinearColor LineColor = FLinearColor::White, float Duration=0.f, float Thickness=0.f, FVector YAxis=FVector(0.f,1.f,0.f),FVector ZAxis=FVector(0.f,0.f,1.f), bool bDrawAxis=false, EDrawDebugSceneDepthPriorityGroup DepthPriority = EDrawDebugSceneDepthPriorityGroup::World);
 	
 	/** Draw a debug point */
 	UFUNCTION(BlueprintCallable, Category="Rendering|Debug", meta=(WorldContext="WorldContextObject", DevelopmentOnly))
-	static ENGINE_API void DrawDebugPoint(const UObject* WorldContextObject, const FVector Position, float Size, FLinearColor PointColor, float Duration=0.f);
+	static ENGINE_API void DrawDebugPoint(const UObject* WorldContextObject, const FVector Position, float Size, FLinearColor PointColor, float Duration=0.f, EDrawDebugSceneDepthPriorityGroup DepthPriority = EDrawDebugSceneDepthPriorityGroup::World);
 
 	/** Draw directional arrow, pointing from LineStart to LineEnd. */
 	UFUNCTION(BlueprintCallable, Category="Rendering|Debug", meta=(WorldContext="WorldContextObject", DevelopmentOnly))
-	static ENGINE_API void DrawDebugArrow(const UObject* WorldContextObject, const FVector LineStart, const FVector LineEnd, float ArrowSize, FLinearColor LineColor, float Duration=0.f, float Thickness = 0.f);
+	static ENGINE_API void DrawDebugArrow(const UObject* WorldContextObject, const FVector LineStart, const FVector LineEnd, float ArrowSize, FLinearColor LineColor, float Duration=0.f, float Thickness = 0.f, EDrawDebugSceneDepthPriorityGroup DepthPriority = EDrawDebugSceneDepthPriorityGroup::World);
 
 	/** Draw a debug box */
 	UFUNCTION(BlueprintCallable, Category="Rendering|Debug", meta=(WorldContext="WorldContextObject", DevelopmentOnly))
-	static ENGINE_API void DrawDebugBox(const UObject* WorldContextObject, const FVector Center, FVector Extent, FLinearColor LineColor, const FRotator Rotation=FRotator::ZeroRotator, float Duration=0.f, float Thickness = 0.f);
+	static ENGINE_API void DrawDebugBox(const UObject* WorldContextObject, const FVector Center, FVector Extent, FLinearColor LineColor, const FRotator Rotation=FRotator::ZeroRotator, float Duration=0.f, float Thickness = 0.f, EDrawDebugSceneDepthPriorityGroup DepthPriority = EDrawDebugSceneDepthPriorityGroup::World);
 
 	/** Draw a debug coordinate system. */
 	UFUNCTION(BlueprintCallable, Category="Rendering|Debug", meta=(WorldContext="WorldContextObject", DevelopmentOnly))
-	static ENGINE_API void DrawDebugCoordinateSystem(const UObject* WorldContextObject, const FVector AxisLoc, const FRotator AxisRot, float Scale=1.f, float Duration=0.f, float Thickness = 0.f);
+	static ENGINE_API void DrawDebugCoordinateSystem(const UObject* WorldContextObject, const FVector AxisLoc, const FRotator AxisRot, float Scale=1.f, float Duration=0.f, float Thickness = 0.f, EDrawDebugSceneDepthPriorityGroup DepthPriority = EDrawDebugSceneDepthPriorityGroup::World);
 
 	/** Draw a debug sphere */
 	UFUNCTION(BlueprintCallable, Category="Rendering|Debug", meta=(WorldContext="WorldContextObject", DevelopmentOnly))
-	static ENGINE_API void DrawDebugSphere(const UObject* WorldContextObject, const FVector Center, float Radius=100.f, int32 Segments=12, FLinearColor LineColor = FLinearColor::White, float Duration=0.f, float Thickness = 0.f);
+	static ENGINE_API void DrawDebugSphere(const UObject* WorldContextObject, const FVector Center, float Radius=100.f, int32 Segments=12, FLinearColor LineColor = FLinearColor::White, float Duration=0.f, float Thickness = 0.f, EDrawDebugSceneDepthPriorityGroup DepthPriority = EDrawDebugSceneDepthPriorityGroup::World);
 
 	/** Draw a debug cylinder */
 	UFUNCTION(BlueprintCallable, Category="Rendering|Debug", meta=(WorldContext="WorldContextObject", DevelopmentOnly))
-	static ENGINE_API void DrawDebugCylinder(const UObject* WorldContextObject, const FVector Start, const FVector End, float Radius=100.f, int32 Segments=12, FLinearColor LineColor = FLinearColor::White, float Duration=0.f, float Thickness = 0.f);
+	static ENGINE_API void DrawDebugCylinder(const UObject* WorldContextObject, const FVector Start, const FVector End, float Radius=100.f, int32 Segments=12, FLinearColor LineColor = FLinearColor::White, float Duration=0.f, float Thickness = 0.f, EDrawDebugSceneDepthPriorityGroup DepthPriority = EDrawDebugSceneDepthPriorityGroup::World);
 	
 	/** Draw a debug cone */
 	UFUNCTION(BlueprintCallable, Category="Rendering|Debug", meta=(WorldContext="WorldContextObject", DeprecatedFunction, DeprecationMessage="DrawDebugCone has been changed to use degrees for angles instead of radians. Place a new DrawDebugCone node and pass your angles as degrees.", DevelopmentOnly))
@@ -1686,11 +1696,11 @@ class UKismetSystemLibrary : public UBlueprintFunctionLibrary
 	 * Angles are specified in degrees
 	 */
 	UFUNCTION(BlueprintCallable, Category="Rendering|Debug", meta=(WorldContext="WorldContextObject", DisplayName="Draw Debug Cone", DevelopmentOnly))
-	static ENGINE_API void DrawDebugConeInDegrees(const UObject* WorldContextObject, const FVector Origin, const FVector Direction, float Length=100.f, float AngleWidth=45.f, float AngleHeight=45.f, int32 NumSides = 12, FLinearColor LineColor = FLinearColor::White, float Duration=0.f, float Thickness = 0.f);
+	static ENGINE_API void DrawDebugConeInDegrees(const UObject* WorldContextObject, const FVector Origin, const FVector Direction, float Length=100.f, float AngleWidth=45.f, float AngleHeight=45.f, int32 NumSides = 12, FLinearColor LineColor = FLinearColor::White, float Duration=0.f, float Thickness = 0.f, EDrawDebugSceneDepthPriorityGroup DepthPriority = EDrawDebugSceneDepthPriorityGroup::World);
 
 	/** Draw a debug capsule */
 	UFUNCTION(BlueprintCallable, Category="Rendering|Debug", meta=(WorldContext="WorldContextObject", DevelopmentOnly))
-	static ENGINE_API void DrawDebugCapsule(const UObject* WorldContextObject, const FVector Center, float HalfHeight, float Radius, const FRotator Rotation, FLinearColor LineColor = FLinearColor::White, float Duration=0.f, float Thickness = 0.f);
+	static ENGINE_API void DrawDebugCapsule(const UObject* WorldContextObject, const FVector Center, float HalfHeight, float Radius, const FRotator Rotation, FLinearColor LineColor = FLinearColor::White, float Duration=0.f, float Thickness = 0.f, EDrawDebugSceneDepthPriorityGroup DepthPriority = EDrawDebugSceneDepthPriorityGroup::World);
 
 	/** Draw a debug string at a 3d world location. */
 	UFUNCTION(BlueprintCallable, Category="Rendering|Debug", meta=(WorldContext="WorldContextObject", DevelopmentOnly))
@@ -1705,7 +1715,7 @@ class UKismetSystemLibrary : public UBlueprintFunctionLibrary
 
 	/** Draws a debug plane. */
 	UFUNCTION(BlueprintCallable, Category="Rendering|Debug", meta=(WorldContext="WorldContextObject", DevelopmentOnly))
-	static ENGINE_API void DrawDebugPlane(const UObject* WorldContextObject, const FPlane& PlaneCoordinates, const FVector Location, float Size, FLinearColor PlaneColor = FLinearColor::White, float Duration=0.f);
+	static ENGINE_API void DrawDebugPlane(const UObject* WorldContextObject, const FPlane& PlaneCoordinates, const FVector Location, float Size, FLinearColor PlaneColor = FLinearColor::White, float Duration=0.f, EDrawDebugSceneDepthPriorityGroup DepthPriority = EDrawDebugSceneDepthPriorityGroup::World);
 
 	/** 
 	 * Flush all persistent debug lines and shapes.
diff --git a/Engine/Source/Runtime/Engine/Private/KismetSystemLibrary.cpp b/Engine/Source/Runtime/Engine/Private/KismetSystemLibrary.cpp
--- a/Engine/Source/Runtime/Engine/Private/KismetSystemLibrary.cpp
+++ b/Engine/Source/Runtime/Engine/Private/KismetSystemLibrary.cpp
@@ -98,6 +98,13 @@ namespace UE::Blueprint::Private
 			}
 		}
 	}
+
+	// Convert the BP-exposed Draw Debug version of Scene Depth Priority Group to the corresponding engine type enum value
+	ESceneDepthPriorityGroup ConvertSceneDepthPriorityToEngineType(EDrawDebugSceneDepthPriorityGroup DepthPriority)
+	{
+		return (DepthPriority == EDrawDebugSceneDepthPriorityGroup::World)
+			? ESceneDepthPriorityGroup::SDPG_World : ESceneDepthPriorityGroup::SDPG_Foreground;
+	}
 }
 
 UKismetSystemLibrary::UKismetSystemLibrary(const FObjectInitializer& ObjectInitializer)
@@ -2324,97 +2331,97 @@ TArray<FName> UKismetSystemLibrary::GetCollisionProfileNames()
 #endif
 
 /** Draw a debug line */
-void UKismetSystemLibrary::DrawDebugLine(const UObject* WorldContextObject, FVector const LineStart, FVector const LineEnd, FLinearColor Color, float LifeTime, float Thickness)
+void UKismetSystemLibrary::DrawDebugLine(const UObject* WorldContextObject, FVector const LineStart, FVector const LineEnd, FLinearColor Color, float LifeTime, float Thickness, EDrawDebugSceneDepthPriorityGroup DepthPriority)
 {
 #if ENABLE_DRAW_DEBUG
 	if (UWorld* World = GEngine->GetWorldFromContextObject(WorldContextObject, EGetWorldErrorMode::LogAndReturnNull))
 	{
-		::DrawDebugLine(World, LineStart, LineEnd, Color.ToFColor(true), false, LifeTime, SDPG_World, Thickness);
+		::DrawDebugLine(World, LineStart, LineEnd, Color.ToFColor(true), false, LifeTime, UE::Blueprint::Private::ConvertSceneDepthPriorityToEngineType(DepthPriority), Thickness);
 	}
 #endif
 }
 
 /** Draw a debug circle */
-void UKismetSystemLibrary::DrawDebugCircle(const UObject* WorldContextObject,FVector Center, float Radius, int32 NumSegments, FLinearColor LineColor, float LifeTime, float Thickness, FVector YAxis, FVector ZAxis, bool bDrawAxis)
+void UKismetSystemLibrary::DrawDebugCircle(const UObject* WorldContextObject,FVector Center, float Radius, int32 NumSegments, FLinearColor LineColor, float LifeTime, float Thickness, FVector YAxis, FVector ZAxis, bool bDrawAxis, EDrawDebugSceneDepthPriorityGroup DepthPriority)
 { 
 #if ENABLE_DRAW_DEBUG
 	if (UWorld* World = GEngine->GetWorldFromContextObject(WorldContextObject, EGetWorldErrorMode::LogAndReturnNull))
 	{
-		::DrawDebugCircle(World, Center, Radius, NumSegments, LineColor.ToFColor(true), false, LifeTime, SDPG_World, Thickness, YAxis, ZAxis, bDrawAxis);
+		::DrawDebugCircle(World, Center, Radius, NumSegments, LineColor.ToFColor(true), false, LifeTime, UE::Blueprint::Private::ConvertSceneDepthPriorityToEngineType(DepthPriority), Thickness, YAxis, ZAxis, bDrawAxis);
 	}
 #endif
 }
 
 /** Draw a debug point */
-void UKismetSystemLibrary::DrawDebugPoint(const UObject* WorldContextObject, FVector const Position, float Size, FLinearColor PointColor, float LifeTime)
+void UKismetSystemLibrary::DrawDebugPoint(const UObject* WorldContextObject, FVector const Position, float Size, FLinearColor PointColor, float LifeTime, EDrawDebugSceneDepthPriorityGroup DepthPriority)
 {
 #if ENABLE_DRAW_DEBUG
 	if (UWorld* World = GEngine->GetWorldFromContextObject(WorldContextObject, EGetWorldErrorMode::LogAndReturnNull))
 	{
-		::DrawDebugPoint(World, Position, Size, PointColor.ToFColor(true), false, LifeTime, SDPG_World);
+		::DrawDebugPoint(World, Position, Size, PointColor.ToFColor(true), false, LifeTime, UE::Blueprint::Private::ConvertSceneDepthPriorityToEngineType(DepthPriority));
 	}
 #endif
 }
 
 /** Draw directional arrow, pointing from LineStart to LineEnd. */
-void UKismetSystemLibrary::DrawDebugArrow(const UObject* WorldContextObject, FVector const LineStart, FVector const LineEnd, float ArrowSize, FLinearColor Color, float LifeTime, float Thickness)
+void UKismetSystemLibrary::DrawDebugArrow(const UObject* WorldContextObject, FVector const LineStart, FVector const LineEnd, float ArrowSize, FLinearColor Color, float LifeTime, float Thickness, EDrawDebugSceneDepthPriorityGroup DepthPriority)
 {
 #if ENABLE_DRAW_DEBUG
 	UWorld* World = GEngine->GetWorldFromContextObject(WorldContextObject, EGetWorldErrorMode::LogAndReturnNull);
 	if (World != nullptr)
 	{
-		::DrawDebugDirectionalArrow(World, LineStart, LineEnd, ArrowSize, Color.ToFColor(true), false, LifeTime, SDPG_World, Thickness);
+		::DrawDebugDirectionalArrow(World, LineStart, LineEnd, ArrowSize, Color.ToFColor(true), false, LifeTime, UE::Blueprint::Private::ConvertSceneDepthPriorityToEngineType(DepthPriority), Thickness);
 	}
 #endif
 }
 
 /** Draw a debug box */
-void UKismetSystemLibrary::DrawDebugBox(const UObject* WorldContextObject, FVector const Center, FVector Extent, FLinearColor Color, const FRotator Rotation, float LifeTime, float Thickness)
+void UKismetSystemLibrary::DrawDebugBox(const UObject* WorldContextObject, FVector const Center, FVector Extent, FLinearColor Color, const FRotator Rotation, float LifeTime, float Thickness, EDrawDebugSceneDepthPriorityGroup DepthPriority)
 {
 #if ENABLE_DRAW_DEBUG
 	if (UWorld* World = GEngine->GetWorldFromContextObject(WorldContextObject, EGetWorldErrorMode::LogAndReturnNull))
 	{
 		if (Rotation == FRotator::ZeroRotator)
 		{
-			::DrawDebugBox(World, Center, Extent, Color.ToFColor(true), false, LifeTime, SDPG_World, Thickness);
+			::DrawDebugBox(World, Center, Extent, Color.ToFColor(true), false, LifeTime, UE::Blueprint::Private::ConvertSceneDepthPriorityToEngineType(DepthPriority), Thickness);
 		}
 		else
 		{
-			::DrawDebugBox(World, Center, Extent, Rotation.Quaternion(), Color.ToFColor(true), false, LifeTime, SDPG_World, Thickness);
+			::DrawDebugBox(World, Center, Extent, Rotation.Quaternion(), Color.ToFColor(true), false, LifeTime, UE::Blueprint::Private::ConvertSceneDepthPriorityToEngineType(DepthPriority), Thickness);
 		}
 	}
 #endif
 }
 
 /** Draw a debug coordinate system. */
-void UKismetSystemLibrary::DrawDebugCoordinateSystem(const UObject* WorldContextObject, FVector const AxisLoc, FRotator const AxisRot, float Scale, float LifeTime, float Thickness)
+void UKismetSystemLibrary::DrawDebugCoordinateSystem(const UObject* WorldContextObject, FVector const AxisLoc, FRotator const AxisRot, float Scale, float LifeTime, float Thickness, EDrawDebugSceneDepthPriorityGroup DepthPriority)
 {
 #if ENABLE_DRAW_DEBUG
 	if (UWorld* World = GEngine->GetWorldFromContextObject(WorldContextObject, EGetWorldErrorMode::LogAndReturnNull))
 	{
-		::DrawDebugCoordinateSystem(World, AxisLoc, AxisRot, Scale, false, LifeTime, SDPG_World, Thickness);
+		::DrawDebugCoordinateSystem(World, AxisLoc, AxisRot, Scale, false, LifeTime, UE::Blueprint::Private::ConvertSceneDepthPriorityToEngineType(DepthPriority), Thickness);
 	}
 #endif
 }
 
 /** Draw a debug sphere */
-void UKismetSystemLibrary::DrawDebugSphere(const UObject* WorldContextObject, FVector const Center, float Radius, int32 Segments, FLinearColor Color, float LifeTime, float Thickness)
+void UKismetSystemLibrary::DrawDebugSphere(const UObject* WorldContextObject, FVector const Center, float Radius, int32 Segments, FLinearColor Color, float LifeTime, float Thickness, EDrawDebugSceneDepthPriorityGroup DepthPriority)
 {
 #if ENABLE_DRAW_DEBUG
 	if (UWorld* World = GEngine->GetWorldFromContextObject(WorldContextObject, EGetWorldErrorMode::LogAndReturnNull))
 	{
-		::DrawDebugSphere(World, Center, Radius, Segments, Color.ToFColor(true), false, LifeTime, SDPG_World, Thickness);
+		::DrawDebugSphere(World, Center, Radius, Segments, Color.ToFColor(true), false, LifeTime, UE::Blueprint::Private::ConvertSceneDepthPriorityToEngineType(DepthPriority), Thickness);
 	}
 #endif
 }
 
 /** Draw a debug cylinder */
-void UKismetSystemLibrary::DrawDebugCylinder(const UObject* WorldContextObject, FVector const Start, FVector const End, float Radius, int32 Segments, FLinearColor Color, float LifeTime, float Thickness)
+void UKismetSystemLibrary::DrawDebugCylinder(const UObject* WorldContextObject, FVector const Start, FVector const End, float Radius, int32 Segments, FLinearColor Color, float LifeTime, float Thickness, EDrawDebugSceneDepthPriorityGroup DepthPriority)
 {
 #if ENABLE_DRAW_DEBUG
 	if (UWorld* World = GEngine->GetWorldFromContextObject(WorldContextObject, EGetWorldErrorMode::LogAndReturnNull))
 	{
-		::DrawDebugCylinder(World, Start, End, Radius, Segments, Color.ToFColor(true), false, LifeTime, SDPG_World, Thickness);
+		::DrawDebugCylinder(World, Start, End, Radius, Segments, Color.ToFColor(true), false, LifeTime, UE::Blueprint::Private::ConvertSceneDepthPriorityToEngineType(DepthPriority), Thickness);
 	}
 #endif
 }
@@ -2430,23 +2437,23 @@ void UKismetSystemLibrary::DrawDebugCone(const UObject* WorldContextObject, FVec
 #endif
 }
 
-void UKismetSystemLibrary::DrawDebugConeInDegrees(const UObject* WorldContextObject, FVector const Origin, FVector const Direction, float Length, float AngleWidth, float AngleHeight, int32 NumSides, FLinearColor Color, float LifeTime, float Thickness)
+void UKismetSystemLibrary::DrawDebugConeInDegrees(const UObject* WorldContextObject, FVector const Origin, FVector const Direction, float Length, float AngleWidth, float AngleHeight, int32 NumSides, FLinearColor Color, float LifeTime, float Thickness, EDrawDebugSceneDepthPriorityGroup DepthPriority)
 {
 #if ENABLE_DRAW_DEBUG
 	if (UWorld* World = GEngine->GetWorldFromContextObject(WorldContextObject, EGetWorldErrorMode::LogAndReturnNull))
 	{
-		::DrawDebugCone(World, Origin, Direction, Length, FMath::DegreesToRadians(AngleWidth), FMath::DegreesToRadians(AngleHeight), NumSides, Color.ToFColor(true), false, LifeTime, SDPG_World, Thickness);
+		::DrawDebugCone(World, Origin, Direction, Length, FMath::DegreesToRadians(AngleWidth), FMath::DegreesToRadians(AngleHeight), NumSides, Color.ToFColor(true), false, LifeTime, UE::Blueprint::Private::ConvertSceneDepthPriorityToEngineType(DepthPriority), Thickness);
 	}
 #endif
 }
 
 /** Draw a debug capsule */
-void UKismetSystemLibrary::DrawDebugCapsule(const UObject* WorldContextObject, FVector const Center, float HalfHeight, float Radius, const FRotator Rotation, FLinearColor Color, float LifeTime, float Thickness)
+void UKismetSystemLibrary::DrawDebugCapsule(const UObject* WorldContextObject, FVector const Center, float HalfHeight, float Radius, const FRotator Rotation, FLinearColor Color, float LifeTime, float Thickness, EDrawDebugSceneDepthPriorityGroup DepthPriority)
 {
 #if ENABLE_DRAW_DEBUG
 	if (UWorld* World = GEngine->GetWorldFromContextObject(WorldContextObject, EGetWorldErrorMode::LogAndReturnNull))
 	{
-		::DrawDebugCapsule(World, Center, HalfHeight, Radius, Rotation.Quaternion(), Color.ToFColor(true), false, LifeTime, SDPG_World, Thickness);
+		::DrawDebugCapsule(World, Center, HalfHeight, Radius, Rotation.Quaternion(), Color.ToFColor(true), false, LifeTime, UE::Blueprint::Private::ConvertSceneDepthPriorityToEngineType(DepthPriority), Thickness);
 	}
 #endif
 }
@@ -2474,12 +2481,12 @@ void UKismetSystemLibrary::FlushDebugStrings(const UObject* WorldContextObject )
 }
 
 /** Draws a debug plane. */
-void UKismetSystemLibrary::DrawDebugPlane(const UObject* WorldContextObject, FPlane const& P, FVector const Loc, float Size, FLinearColor Color, float LifeTime)
+void UKismetSystemLibrary::DrawDebugPlane(const UObject* WorldContextObject, FPlane const& P, FVector const Loc, float Size, FLinearColor Color, float LifeTime, EDrawDebugSceneDepthPriorityGroup DepthPriority)
 {
 #if ENABLE_DRAW_DEBUG
 	if (UWorld* World = GEngine->GetWorldFromContextObject(WorldContextObject, EGetWorldErrorMode::LogAndReturnNull))
 	{
-		::DrawDebugSolidPlane(World, P, Loc, Size, Color.ToFColor(true), false, LifeTime, SDPG_World);
+		::DrawDebugSolidPlane(World, P, Loc, Size, Color.ToFColor(true), false, LifeTime, UE::Blueprint::Private::ConvertSceneDepthPriorityToEngineType(DepthPriority));
 	}
 #endif
 }
