Iris
* Added ability for Actors to use different net factories via FillReplicatedParams

#rnx

diff --git a/Engine/Source/Runtime/Engine/Private/Net/Iris/ReplicationSystem/EngineReplicationBridge.cpp b/Engine/Source/Runtime/Engine/Private/Net/Iris/ReplicationSystem/EngineReplicationBridge.cpp
--- a/Engine/Source/Runtime/Engine/Private/Net/Iris/ReplicationSystem/EngineReplicationBridge.cpp
+++ b/Engine/Source/Runtime/Engine/Private/Net/Iris/ReplicationSystem/EngineReplicationBridge.cpp
@@ -409,14 +409,22 @@ UE::Net::FNetRefHandle UEngineReplicationBridge::StartReplicatingActor(AActor* A
 				break;
 		}
 	}
-	
-
 
 #if !UE_BUILD_SHIPPING
 	ensureMsgf(!(Actor->bAlwaysRelevant || Actor->bOnlyRelevantToOwner) || RootObjectParams.StaticPriority >= 1.0f, TEXT("Very low NetPriority %.02f for always relevant or owner relevant Actor %s. Set it to 1.0f or higher."), Actor->NetPriority, ToCStr(Actor->GetName()));
 #endif
 
-	FNetRefHandle ActorRefHandle = StartReplicatingRootObject(Actor, RootObjectParams, ActorFactoryId);
+	// Find the factory ID if this actor does not use the default one.
+	FNetObjectFactoryId FactoryId = ActorFactoryId;
+	if (ActorParams.NetFactoryName != NAME_None)
+	{
+		const FNetObjectFactoryId CustomFactoryId = FNetObjectFactoryRegistry::GetFactoryIdFromName(ActorParams.NetFactoryName);
+		ensureMsgf(CustomFactoryId != InvalidNetObjectFactoryId, TEXT("Could not find Factory: %s for Actor: %s"), *ActorParams.NetFactoryName.ToString(), *GetNameSafe(Actor));
+
+		FactoryId = CustomFactoryId != InvalidNetObjectFactoryId ? CustomFactoryId : ActorFactoryId;
+	}
+
+	FNetRefHandle ActorRefHandle = StartReplicatingRootObject(Actor, RootObjectParams, FactoryId);
 
 	if (!ActorRefHandle.IsValid())
 	{
diff --git a/Engine/Source/Runtime/Engine/Public/Net/Iris/ReplicationSystem/EngineReplicationBridge.h b/Engine/Source/Runtime/Engine/Public/Net/Iris/ReplicationSystem/EngineReplicationBridge.h
--- a/Engine/Source/Runtime/Engine/Public/Net/Iris/ReplicationSystem/EngineReplicationBridge.h
+++ b/Engine/Source/Runtime/Engine/Public/Net/Iris/ReplicationSystem/EngineReplicationBridge.h
@@ -67,6 +67,9 @@ struct FActorReplicationParams
 
 	/** Only used when ExplicitFilter is the type used. The dynamic filter to assign to this actor. */
 	FName ExplicitDynamicFilterName;
+
+	/** Optional  factory name if the actor is not using the default NetActorFactory */
+	FName NetFactoryName;
 };
 
 struct FStopReplicatingActorParams
