Iris
* Remove more UE_WITH_IRIS defines

#rnx

diff --git a/Engine/Source/Runtime/Engine/Private/Tests/NetTestHelpers.cpp b/Engine/Source/Runtime/Engine/Private/Tests/NetTestHelpers.cpp
--- a/Engine/Source/Runtime/Engine/Private/Tests/NetTestHelpers.cpp
+++ b/Engine/Source/Runtime/Engine/Private/Tests/NetTestHelpers.cpp
@@ -22,10 +22,8 @@
 #include "Settings/LevelEditorPlaySettings.h"
 #endif
 
-#if UE_WITH_IRIS
 #include "Iris/ReplicationSystem/ReplicationSystem.h"
 #include "Iris/ReplicationSystem/ObjectReplicationBridge.h"
-#endif
 
 #include "UObject/UObjectGlobals.h"
 #include "UObject/Package.h"
@@ -258,19 +256,15 @@ int32 FTestWorldInstance::FindUnusedPIEInstance()
 
 FTestWorldInstance::FContext FTestWorldInstance::GetTestContext() const
 {
-#if UE_WITH_IRIS
 	UReplicationSystem* RepSystem = GetNetDriver() ? GetNetDriver()->GetReplicationSystem() : nullptr;
 	UObjectReplicationBridge* RepBridge  = RepSystem ? RepSystem->GetReplicationBridge() : nullptr;
-#endif
 
 	return FContext
 	{ 
 		.World = GetWorld(),
 		.NetDriver = GetNetDriver(),
-#if UE_WITH_IRIS
 		.IrisRepSystem = RepSystem,
 		.IrisRepBridge = RepBridge,
-#endif
 	};
 }
 
@@ -563,7 +557,6 @@ UObject* FTestWorlds::FindReplicatedObjectOnClient(UObject* ServerObject, uint32
 
 	if (ServerContext.NetDriver->IsUsingIrisReplication())
 	{
-#if UE_WITH_IRIS
 		const FNetRefHandle NetHandle = ServerContext.IrisRepBridge->GetReplicatedRefHandle(ServerObject);
 
 		if (!NetHandle.IsValid())
@@ -573,10 +566,6 @@ UObject* FTestWorlds::FindReplicatedObjectOnClient(UObject* ServerObject, uint32
 		}
 
 		return ClientContext.IrisRepBridge->GetReplicatedObject(NetHandle);
-
-#else
-		return nullptr;
-#endif
 	}
 	else
 	{
@@ -611,12 +600,8 @@ bool FTestWorlds::IsServerObjectReplicated(const UObject* ServerObject) const
 	
 	if (ServerContext.NetDriver->IsUsingIrisReplication())
 	{
-#if UE_WITH_IRIS
 		const FNetRefHandle NetHandle = ServerContext.IrisRepBridge->GetReplicatedRefHandle(ServerObject);
 		return NetHandle.IsValid();
-#else
-		return false;
-#endif
 	}
 	else
 	{
diff --git a/Engine/Source/Runtime/Engine/Public/Net/Iris/ReplicationSystem/NetEngineFactories.h b/Engine/Source/Runtime/Engine/Public/Net/Iris/ReplicationSystem/NetEngineFactories.h
--- a/Engine/Source/Runtime/Engine/Public/Net/Iris/ReplicationSystem/NetEngineFactories.h
+++ b/Engine/Source/Runtime/Engine/Public/Net/Iris/ReplicationSystem/NetEngineFactories.h
@@ -2,12 +2,8 @@
 
 #pragma once
 
-#if UE_WITH_IRIS
-
 namespace UE::Net
 {
 	void InitEngineNetObjectFactories();
 	void ShutdownEngineNetObjectFactories();
-}
-
-#endif //UE_WITH_IRIS
\ No newline at end of file
+}
\ No newline at end of file
diff --git a/Engine/Source/Runtime/Engine/Public/Tests/NetTestHelpers.h b/Engine/Source/Runtime/Engine/Public/Tests/NetTestHelpers.h
--- a/Engine/Source/Runtime/Engine/Public/Tests/NetTestHelpers.h
+++ b/Engine/Source/Runtime/Engine/Public/Tests/NetTestHelpers.h
@@ -153,18 +153,12 @@ public:
 	{
 		UWorld* World = nullptr;
 		UNetDriver* NetDriver = nullptr;
-#if UE_WITH_IRIS
 		UReplicationSystem* IrisRepSystem = nullptr;
 		UObjectReplicationBridge* IrisRepBridge = nullptr;
-#endif
 
 		bool IsValid() const
 		{
-			return World && NetDriver
-#if UE_WITH_IRIS
-				&& IrisRepSystem && IrisRepBridge
-#endif
-				;
+			return World && NetDriver && IrisRepSystem && IrisRepBridge;
 		}
 	};
 
diff --git a/Engine/Source/Runtime/Net/Iris/Private/Iris/ReplicationSystem/NetObjectFactoryRegistry.cpp b/Engine/Source/Runtime/Net/Iris/Private/Iris/ReplicationSystem/NetObjectFactoryRegistry.cpp
--- a/Engine/Source/Runtime/Net/Iris/Private/Iris/ReplicationSystem/NetObjectFactoryRegistry.cpp
+++ b/Engine/Source/Runtime/Net/Iris/Private/Iris/ReplicationSystem/NetObjectFactoryRegistry.cpp
@@ -24,8 +24,6 @@ void FNetObjectFactoryRegistry::RegisterFactory(UClass* FactoryClass, FName Fact
 {
 	using namespace UE::Net::Private;
 
-#if UE_WITH_IRIS
-
 	// Factories cannot be modified while Iris NetDrivers exist
 	if (!bIsFactoryRegistrationAllowed)
 	{
@@ -93,14 +91,12 @@ void FNetObjectFactoryRegistry::RegisterFactory(UClass* FactoryClass, FName Fact
 	}
 
 	UE_LOG(LogIris, Verbose, TEXT("FNetObjectFactoryRegistry::RegisterFactory has registered factory: %s name: %s id: %u"), *GetNameSafe(NetFactories[Index].NetFactoryClass.Get()), *NetFactories[Index].Name.ToString(), NetFactories[Index].Id);
-#endif
 }
 
 void FNetObjectFactoryRegistry::UnregisterFactory(FName FactoryName)
 {
 	using namespace UE::Net::Private;
-#if UE_WITH_IRIS
-	
+
 	// Factories cannot be modified while Iris NetDrivers exist
 	if (!bIsFactoryRegistrationAllowed)
 	{
@@ -124,12 +120,10 @@ void FNetObjectFactoryRegistry::UnregisterFactory(FName FactoryName)
 
 	UE_LOG(LogIris, Error, TEXT("FNetObjectFactoryRegistry::UnregisterFactory could not find any factories using name: %s"), *FactoryName.ToString());
 	ensure(false);
-#endif
 }
 	  
 FNetObjectFactoryId FNetObjectFactoryRegistry::GetFactoryIdFromName(FName FactoryName)
 {
-#if UE_WITH_IRIS
 	for (const FFactoryData& Data : NetFactories)
 	{
 		if (Data.Name == FactoryName)
@@ -137,7 +131,6 @@ FNetObjectFactoryId FNetObjectFactoryRegistry::GetFactoryIdFromName(FName Factor
 			return Data.Id;
 		}
 	}
-#endif
 
 	return InvalidNetObjectFactoryId;
 }
