[Backout] - CL45447759
#rnx
Original CL Desc
-----------------------------------------------------------------
Added VirtualTextureCollection's, a separate texture collection object with virtual texture backing
- Internally hosts a tile-wise atlas of the texture set. Single physical texture, all respective texture entries must share the same format and srgb mode. Alternatively, runtime format conversion is supported.
- Texture entries are naturally (mip wise) block aligned. Due to block alignment, we can avoid expensive UV re-mapping and manual sampler addressing by modifying the page table uniforms.
- Format conversion supports both virtual and regular textures, and block (re)compression should it be needed.
- Due to the extensive amount of pipelining needed for another collection parameter type, it's easier to extend the current type. For in-type bindings each parameter has a prefix index.
- Pending sharing of page table between collections

(Reviewed in 43854629, changed workspace)

diff --git a/Engine/Shaders/Private/MaterialTemplate.ush b/Engine/Shaders/Private/MaterialTemplate.ush
--- a/Engine/Shaders/Private/MaterialTemplate.ush
+++ b/Engine/Shaders/Private/MaterialTemplate.ush
@@ -116,7 +116,6 @@
 
 #if NUM_VIRTUALTEXTURE_SAMPLES || LIGHTMAP_VT_ENABLED
 	#include "/Engine/Private/VirtualTextureCommon.ush"
-	#include "/Engine/Public/IndirectVirtualTextureCommon.ush"
 #endif
 
 #include "/Engine/Private/MaterialTexture.ush"
diff --git a/Engine/Shaders/Public/IndirectVirtualTexture.ush b/Engine/Shaders/Public/IndirectVirtualTexture.ush
deleted file mode 100644
diff --git a/Engine/Shaders/Public/IndirectVirtualTextureCommon.ush b/Engine/Shaders/Public/IndirectVirtualTextureCommon.ush
deleted file mode 100644
diff --git a/Engine/Shaders/Public/Platform.ush b/Engine/Shaders/Public/Platform.ush
--- a/Engine/Shaders/Public/Platform.ush
+++ b/Engine/Shaders/Public/Platform.ush
@@ -67,7 +67,6 @@
 
 #include "/Engine/Public/BindlessResources.ush"
 #include "/Engine/Shared/MaterialCacheDefinitions.h"
-#include "/Engine/Public/IndirectVirtualTexture.ush"
 #include "/Engine/Public/OverloadMacros.ush"
 
 
diff --git a/Engine/Shaders/Shared/IndirectVirtualTextureDefinitions.h b/Engine/Shaders/Shared/IndirectVirtualTextureDefinitions.h
deleted file mode 100644
diff --git a/Engine/Source/Runtime/Core/Private/UObject/FortniteMainBranchObjectVersion.cpp b/Engine/Source/Runtime/Core/Private/UObject/FortniteMainBranchObjectVersion.cpp
--- a/Engine/Source/Runtime/Core/Private/UObject/FortniteMainBranchObjectVersion.cpp
+++ b/Engine/Source/Runtime/Core/Private/UObject/FortniteMainBranchObjectVersion.cpp
@@ -13,7 +13,7 @@ TMap<FGuid, FGuid> FFortniteMainBranchObjectVersion::GetSystemGuids()
 	SystemGuids.Add(DevGuids.MATERIALSHADERMAP_DERIVEDDATA_VER, FGuid("F051877974B44D409435F321BDB39DBB"));
 	SystemGuids.Add(DevGuids.NIAGARASHADERMAP_DERIVEDDATA_VER, FGuid("6360A977062842A29ED30E3A7ACB0E64"));
 	SystemGuids.Add(DevGuids.Niagara_LatestScriptCompileVersion, FGuid("783FF3700BB641C799B2D43552FBF91E"));
-	SystemGuids.Add(DevGuids.MaterialTranslationDDCVersion, FGuid("F9BB134100414BE4AA156DEF0C01C4DF"));
+	SystemGuids.Add(DevGuids.MaterialTranslationDDCVersion, FGuid("448CDD72B0A34496BDED50908658DCBE"));
 
 	// These GUIDs are in AppendRestrictedFortniteMainStreamObjectVersionGuids as they should rarely be changed.
 	//SystemGuids.Add(DevGuids.NANITE_DERIVEDDATA_VER, FGuid("ACF98184978742A592537DA223A1E6BE"));
diff --git a/Engine/Source/Runtime/Engine/Classes/Engine/Texture.h b/Engine/Source/Runtime/Engine/Classes/Engine/Texture.h
--- a/Engine/Source/Runtime/Engine/Classes/Engine/Texture.h
+++ b/Engine/Source/Runtime/Engine/Classes/Engine/Texture.h
@@ -1150,8 +1150,6 @@ struct FTextureFormatSettings
 		, SRGB(false)
 	{}
 
-	bool operator==(const FTextureFormatSettings& other) const = default;
-
 	UPROPERTY()
 	TEnumAsByte<enum TextureCompressionSettings> CompressionSettings;
 
diff --git a/Engine/Source/Runtime/Engine/Classes/Engine/TextureCollection.h b/Engine/Source/Runtime/Engine/Classes/Engine/TextureCollection.h
--- a/Engine/Source/Runtime/Engine/Classes/Engine/TextureCollection.h
+++ b/Engine/Source/Runtime/Engine/Classes/Engine/TextureCollection.h
@@ -14,32 +14,17 @@ struct FTextureCollectionResource : public FRenderResource
 {
 	FTextureCollectionResource(UTextureCollection* InParent);
 
-	/** Is this a bindless collection? */
-	bool bIsBindless = false;
-	
-	/** Name of the parent collection */
-	FName CollectionName;
-};
-
-struct FBindlessTextureCollectionResource : public FTextureCollectionResource
-{
-	FBindlessTextureCollectionResource(UTextureCollection* InParent);
-	virtual ~FBindlessTextureCollectionResource() override = default;
-
 	// FRenderResource Interface
 	virtual void InitRHI(FRHICommandListBase& RHICmdList) final;
 	virtual void ReleaseRHI() final;
 	// ~FRenderResource
 
-	FRHIResourceCollection* GetResourceCollectionRHI() const
-	{
-		check(ResourceCollectionRHI);
-		return ResourceCollectionRHI.GetReference();
-	}
-	
+	FRHIResourceCollection* GetRHI() const { return ResourceCollectionRHI.GetReference(); }
+
 private:
-	TArray<FTextureResource*> InputTextureResources;
-	FRHIResourceCollectionRef ResourceCollectionRHI;
+	TArray<FTextureResource*>       InputTextureResources;
+	TArray<FTextureReferenceRHIRef> InputTextures;
+	FRHIResourceCollectionRef       ResourceCollectionRHI;
 };
 
 UCLASS(MinimalAPI)
@@ -62,12 +47,10 @@ public:
 	ENGINE_API FTextureCollectionResource* GetResource() const;
 	ENGINE_API FTextureCollectionResource* GetResource();
 
-	virtual ENGINE_API FTextureCollectionResource* CreateResource();
+	ENGINE_API FTextureCollectionResource* CreateResource();
 	ENGINE_API void ReleaseResource();
 	ENGINE_API void UpdateResource();
 
-	virtual bool IsVirtualCollection() const { return false; }
-	
 protected:
 #if WITH_EDITOR
 	void NotifyMaterials();
diff --git a/Engine/Source/Runtime/Engine/Classes/Engine/VirtualTextureCollection.h b/Engine/Source/Runtime/Engine/Classes/Engine/VirtualTextureCollection.h
deleted file mode 100644
diff --git a/Engine/Source/Runtime/Engine/Classes/VT/CopyCompressShader.h b/Engine/Source/Runtime/Engine/Classes/VT/CopyCompressShader.h
deleted file mode 100644
diff --git a/Engine/Source/Runtime/Engine/Private/Materials/HLSLMaterialTranslator.cpp b/Engine/Source/Runtime/Engine/Private/Materials/HLSLMaterialTranslator.cpp
--- a/Engine/Source/Runtime/Engine/Private/Materials/HLSLMaterialTranslator.cpp
+++ b/Engine/Source/Runtime/Engine/Private/Materials/HLSLMaterialTranslator.cpp
@@ -212,12 +212,11 @@ struct FHLSLMaterialTranslator::FEnvironmentDefines
 {
 	enum EVirtualPageType
 	{
-		TABLE_MESH_PAINT           = (1 << 0),
-		TABLE_MATERIAL_CACHE       = (1 << 1),
-		TABLE_COLLECTION           = (1 << 2),
-		LOCAL_TABLE0               = (1 << 3),
-		LOCAL_TABLE1               = (1 << 4),
-		LOCAL_ADAPTIVE_INDIRECTION = (1 << 5)
+		TABLE_MESH_PAINT = 1,
+		TABLE_MATERIAL_CACHE = 2,
+		LOCAL_TABLE0 = 4,
+		LOCAL_TABLE1 = 8,
+		LOCAL_ADAPTIVE_INDIRECTION = 16
 	};
 
 	bool bNeedsParticlePosition;
@@ -2881,10 +2880,6 @@ void FHLSLMaterialTranslator::GetMaterialEnvironment(EShaderPlatform InPlatform,
 		{
 			// Driven manually
 		}
-		else if (EnvironmentDefines->VirtualPageTypes[i] & FEnvironmentDefines::TABLE_COLLECTION)
-		{
-			/** Page table fetched during sampling */
-		}
 		else
 		{
 			// Setup page table defines to map each VT stack to either 1 or 2 page table textures, depending on how many layers it uses
@@ -3093,15 +3088,6 @@ void FHLSLMaterialTranslator::GetMaterialEnvironment(EShaderPlatform InPlatform,
 	}
 
 	OutEnvironment.SetDefine(TEXT("TEXTURE_SAMPLE_DEBUG"), EnvironmentDefines->bTextureSampleDebug);
-
-	for (const TRefCountPtr<FMaterialUniformExpressionTextureCollection>& TextureCollection : UniformTextureCollectionExpressions)
-	{
-		if (TextureCollection->IsVirtualCollection())
-		{
-			OutEnvironment.SetDefine(TEXT("TEXTURE_COLLECTION_PACKED_UNIFORMS"), 1u);
-			break;
-		}
-	}
 }
 
 // Assign custom interpolators to slots, packing them as much as possible in unused slots.
@@ -3347,7 +3333,7 @@ FString FHLSLMaterialTranslator::GetParameterCodeDeriv(int32 Index, ECompiledPar
 	}
 	else
 	{
-		if (CodeChunk.UniformExpression && !CodeChunk.bIntermediate)
+		if (CodeChunk.UniformExpression)
 		{
 			// If the code chunk has a uniform expression, create a new code chunk to access it
 			const int32 AccessedIndex = AccessUniformExpression(Index);
@@ -3405,7 +3391,8 @@ FString FHLSLMaterialTranslator::GetParameterCodeRaw(int32 Index, const TCHAR* D
 	}
 	else
 	{
-		if (CodeChunk.UniformExpression && !CodeChunk.bIntermediate)
+		// TODO - don't access uniform expressions for intermediate values
+		if (CodeChunk.UniformExpression)
 		{
 			// If the code chunk has a uniform expression, create a new code chunk to access it
 			const int32 AccessedIndex = AccessUniformExpression(Index);
@@ -3481,7 +3468,7 @@ FString FHLSLMaterialTranslator::GetDefinitions(const TArray<FShaderCodeChunk>&
 	{
 		const FShaderCodeChunk& CodeChunk = CodeChunks[ChunkIndex];
 		// Uniform expressions (both constant and variable) and inline expressions don't have definitions.
-		if ((!CodeChunk.UniformExpression || CodeChunk.bIntermediate) &&
+		if (!CodeChunk.UniformExpression &&
 			(!CodeChunk.bInline || CodeChunk.Type == MCT_VoidStatement))
 		{
 			Definitions += CodeChunk.AtDefinition(Variation);
@@ -3863,7 +3850,7 @@ int32 FHLSLMaterialTranslator::AddCodeChunkInner(uint64 Hash, const TCHAR* Forma
 		new(*CurrentScopeChunks) FShaderCodeChunk(Hash, FormattedCode, FormattedCode, TEXT(""), Type, DerivativeStatus, true);
 	}
 	// Can only create temporaries for certain types
-	else if ((Type & (MCT_Float | MCT_LWCType | MCT_VTPageTableResult | MCT_Unexposed | MCT_UInt)) || Type == MCT_ShadingModel || Type == MCT_MaterialAttributes || Type == MCT_Substrate || Type == MCT_UInt || bFromTextureCollection || Type == MCT_MaterialCacheABuffer)
+	else if ((Type & (MCT_Float | MCT_LWCType | MCT_VTPageTableResult | MCT_UInt)) || Type == MCT_ShadingModel || Type == MCT_MaterialAttributes || Type == MCT_Substrate || Type == MCT_UInt || bFromTextureCollection || Type == MCT_MaterialCacheABuffer)
 	{
 		// Check for existing
 		for (int32 i = 0; i < CurrentScopeChunks->Num(); ++i)
@@ -3924,40 +3911,6 @@ static inline uint32 GetTCharStringBytes(const TCHAR* String)
 	return Length * sizeof(TCHAR);
 }
 
-int32 FHLSLMaterialTranslator::AddInternalCodeChunk(EMaterialValueType Type, const TCHAR* InternalType, const TCHAR* FormattedCode, const TRefCountPtr<FMaterialUniformExpression>& UniformExpression)
-{
-	const uint64 Hash = CityHash64(reinterpret_cast<const char*>(FormattedCode), GetTCharStringBytes(FormattedCode));
-	
-	int32 CodeIndex = INDEX_NONE;
-	for (int32 i = 0; i < CurrentScopeChunks->Num(); ++i)
-	{
-		if ((*CurrentScopeChunks)[i].Hash == Hash)
-		{
-			CodeIndex = i;
-			break;
-		}
-	}
-	
-	if (CodeIndex == INDEX_NONE)
-	{
-		CodeIndex = CurrentScopeChunks->Num();
-		const FString SymbolName = CreateSymbolName(TEXT("Local"));
-		const FString LocalVariableDefinitionFinite = FString("	") + InternalType + TEXT(" ") + SymbolName + TEXT(" = ") + FormattedCode + TEXT(";") + HLSL_LINE_TERMINATOR;
-		const FString LocalVariableDefinitionAnalytic = FString("	") + InternalType + TEXT(" ") + SymbolName + TEXT(" = ") + FormattedCode + TEXT(";") + HLSL_LINE_TERMINATOR;
-
-		FShaderCodeChunk* Chunk = new(*CurrentScopeChunks) FShaderCodeChunk(
-			Hash, *LocalVariableDefinitionFinite, *LocalVariableDefinitionAnalytic,
-			SymbolName, Type, EDerivativeStatus::NotAware, false
-		);
-
-		Chunk->bIntermediate     = true;
-		Chunk->UniformExpression = UniformExpression;
-	}
-	
-	AddCodeChunkToCurrentScope(CodeIndex);
-	return CodeIndex;
-}
-
 /** Adds an already formatted inline or referenced code chunk, and notes the derivative status. */
 int32 FHLSLMaterialTranslator::AddCodeChunkInnerDeriv(const TCHAR* FormattedCodeFinite, const TCHAR* FormattedCodeAnalytic, EMaterialValueType Type, bool bInlined, EDerivativeStatus DerivativeStatus)
 {
@@ -4447,18 +4400,6 @@ int32 FHLSLMaterialTranslator::AccessUniformExpression(int32 Index)
 	}
 	else if (CodeChunk.Type == MCT_TextureCollection)
 	{
-		// If unallocated, assign a type prefix
-		if (TextureCollectionUniformExpression->GetTextureCollectionTypePrefixIndex() == UINT32_MAX)
-		{
-			uint32 TypePrefix = 0;
-			for (int32 i = 0; i < UniformTextureCollectionExpressions.Num(); i++)
-			{
-				TypePrefix += UniformTextureCollectionExpressions[i]->IsVirtualCollection() == TextureCollectionUniformExpression->IsVirtualCollection();
-			}
-
-			TextureCollectionUniformExpression->SetTextureCollectionTypePrefixIndex(TypePrefix);
-		}
-		
 		int32 TextureCollectionInputIndex = UniformTextureCollectionExpressions.AddUnique(TextureCollectionUniformExpression);
 		FormattedCode.Appendf(TEXT("Material.TextureCollection_%u"), TextureCollectionInputIndex);
 	}
@@ -7076,35 +7017,45 @@ void FHLSLMaterialTranslator::SetPotentiallyManipulateTexCoords()
 }
 
 uint32 FHLSLMaterialTranslator::AcquireVTStackIndex(
-	const FMaterialVTSampleInfo& SampleInfo,
+	int32 VirtualTexureIndex,
+	ETextureMipValueMode MipValueMode, 
+	TextureAddress AddressU, TextureAddress AddressV, 
+	float AspectRatio, 
+	int32 CoordinateIndex, 
+	int32 MipValue0Index, int32 MipValue1Index, 
+	int32 PreallocatedStackTextureIndex, 
 	const FString& UV_Value,
 	const FString& UV_Ddx,
-	const FString& UV_Ddy)
+	const FString& UV_Ddy,
+	const FString& PageTableOverride,
+	const FString& PageTableUniformOverride,
+	bool bAdaptive, bool bGenerateFeedback, bool bMeshPaint, bool bMaterialCache)
 {
 	using namespace UE::MaterialTranslatorUtils;
 
-	const uint64 CoordinateHash = GetParameterHash(SampleInfo.CoordinateIndex);
-	const uint64 MipValue0Hash = GetParameterHash(SampleInfo.MipValue0Index);
-	const uint64 MipValue1Hash = GetParameterHash(SampleInfo.MipValue1Index);
-	const uint64 PageTableHash = GetParameterHash(SampleInfo.PageTableIndex);
-	const uint64 PageTableUniformHash = GetParameterHash(SampleInfo.PageTableUniformIndex);
+	const uint64 CoordinateHash = GetParameterHash(CoordinateIndex);
+	const uint64 MipValue0Hash = GetParameterHash(MipValue0Index);
+	const uint64 MipValue1Hash = GetParameterHash(MipValue1Index);
+
+	const uint64 PageTableHash        = CityHash64(reinterpret_cast<const char*>(*PageTableOverride),        PageTableOverride.Len()        * sizeof(TCHAR));
+	const uint64 PageTableUniformHash = CityHash64(reinterpret_cast<const char*>(*PageTableUniformOverride), PageTableUniformOverride.Len() * sizeof(TCHAR));
 	
 	uint32 Flags = 0;
-	Flags |= SampleInfo.bAdaptive ? (1 << 0) : 0;
-	Flags |= SampleInfo.bGenerateFeedback ? (1 << 1) : 0;
-	Flags |= SampleInfo.bMeshPaint ? (1 << 2) : 0;
-	Flags |= SampleInfo.bMaterialCache ? (1 << 3) : 0;
+	Flags |= bAdaptive ? (1 << 0) : 0;
+	Flags |= bGenerateFeedback ? (1 << 1) : 0;
+	Flags |= bMeshPaint ? (1 << 2) : 0;
+	Flags |= bMaterialCache ? (1 << 3) : 0;
 
 	uint64 Hash = CityHash128to64({ CurrentScopeID, CoordinateHash });
 	Hash = CityHash128to64({ Hash, MipValue0Hash });
 	Hash = CityHash128to64({ Hash, MipValue1Hash });
 	Hash = CityHash128to64({ Hash, PageTableHash });
 	Hash = CityHash128to64({ Hash, PageTableUniformHash });
-	Hash = CityHash128to64({ Hash, (uint64)SampleInfo.MipValueMode });
-	Hash = CityHash128to64({ Hash, (uint64)SampleInfo.AddressU });
-	Hash = CityHash128to64({ Hash, (uint64)SampleInfo.AddressV });
-	Hash = CityHash128to64({ Hash, (uint64)(SampleInfo.AspectRatio * 1000.0f) });
-	Hash = CityHash128to64({ Hash, (uint64)SampleInfo.PreallocatedStackTextureIndex });
+	Hash = CityHash128to64({ Hash, (uint64)MipValueMode });
+	Hash = CityHash128to64({ Hash, (uint64)AddressU });
+	Hash = CityHash128to64({ Hash, (uint64)AddressV });
+	Hash = CityHash128to64({ Hash, (uint64)(AspectRatio * 1000.0f) });
+	Hash = CityHash128to64({ Hash, (uint64)PreallocatedStackTextureIndex });
 	Hash = CityHash128to64({ Hash, (uint64)Flags });
 
 	// First check to see if we have an existing VTStack that matches this key, that can still fit another layer
@@ -7113,7 +7064,7 @@ uint32 FHLSLMaterialTranslator::AcquireVTStackIndex(
 		const FMaterialVirtualTextureStack& Stack = MaterialCompilationOutput.UniformExpressionSet.VTStacks[Index];
 		const FMaterialVTStackEntry& Entry = VTStacks[Index];
 		
-		const bool bStackIsFull = Stack.AreLayersFull() && Stack.FindLayer(SampleInfo.VirtualTextureIndex) == INDEX_NONE;
+		const bool bStackIsFull = Stack.AreLayersFull() && Stack.FindLayer(VirtualTexureIndex) == INDEX_NONE;
 		
 		if (!bStackIsFull &&
 			Entry.ScopeID == CurrentScopeID &&
@@ -7122,7 +7073,15 @@ uint32 FHLSLMaterialTranslator::AcquireVTStackIndex(
 			Entry.MipValue1Hash == MipValue1Hash &&
 			Entry.PageTableHash == PageTableHash &&
 			Entry.PageTableUniformHash == PageTableUniformHash &&
-			Entry.SampleInfo == SampleInfo)
+			Entry.MipValueMode == MipValueMode &&
+			Entry.AddressU == AddressU &&
+			Entry.AddressV == AddressV &&
+			Entry.AspectRatio == AspectRatio &&
+			Entry.PreallocatedStackTextureIndex == PreallocatedStackTextureIndex &&
+			Entry.bAdaptive == bAdaptive &&
+			Entry.bGenerateFeedback == bGenerateFeedback &&
+			Entry.bMeshPaint == bMeshPaint &&
+			Entry.bMaterialCache == bMaterialCache)
 		{
 			return Index;
 		}
@@ -7132,59 +7091,72 @@ uint32 FHLSLMaterialTranslator::AcquireVTStackIndex(
 	const int32 StackIndex = VTStacks.AddDefaulted();
 	VTStackHash.Add(Hash, StackIndex);
 	FMaterialVTStackEntry& Entry = VTStacks[StackIndex];
-	Entry.SampleInfo = SampleInfo;
 	Entry.ScopeID = CurrentScopeID;
 	Entry.CoordinateHash = CoordinateHash;
 	Entry.MipValue0Hash = MipValue0Hash;
 	Entry.MipValue1Hash = MipValue1Hash;
 	Entry.PageTableHash = PageTableHash;
 	Entry.PageTableUniformHash = PageTableUniformHash;
-
-	MaterialCompilationOutput.UniformExpressionSet.VTStacks.Add(FMaterialVirtualTextureStack(SampleInfo.PreallocatedStackTextureIndex));
+	Entry.MipValueMode = MipValueMode;
+	Entry.AddressU = AddressU;
+	Entry.AddressV = AddressV;
+	Entry.AspectRatio = AspectRatio;
+	Entry.DebugCoordinateIndex = CoordinateIndex;
+	Entry.DebugMipValue0Index = MipValue0Index;
+	Entry.DebugMipValue1Index = MipValue1Index;
+	Entry.PreallocatedStackTextureIndex = PreallocatedStackTextureIndex;
+	Entry.bAdaptive = bAdaptive;
+	Entry.bGenerateFeedback = bGenerateFeedback;
+	Entry.bMeshPaint = bMeshPaint;
+	Entry.bMaterialCache = bMaterialCache;
+
+	MaterialCompilationOutput.UniformExpressionSet.VTStacks.Add(FMaterialVirtualTextureStack(PreallocatedStackTextureIndex));
 
 	// These two arrays need to stay in sync
 	check(VTStacks.Num() == MaterialCompilationOutput.UniformExpressionSet.VTStacks.Num());
 
 	// Select LoadVirtualPageTable function name for this context
-	FString BaseFunctionName = SampleInfo.bAdaptive ? TEXT("TextureLoadVirtualPageTableAdaptive") : TEXT("TextureLoadVirtualPageTable");
+	FString BaseFunctionName = bAdaptive ? TEXT("TextureLoadVirtualPageTableAdaptive") : TEXT("TextureLoadVirtualPageTable");
 
 	// Optionally sample without virtual texture feedback but only for miplevel mode
-	check(SampleInfo.bGenerateFeedback || SampleInfo.MipValueMode == TMVM_MipLevel)
-	FString FeedbackParameter = SampleInfo.bGenerateFeedback ? FString::Printf(TEXT(", Parameters.VirtualTextureFeedback")) : TEXT("");
+	check(bGenerateFeedback || MipValueMode == TMVM_MipLevel)
+	FString FeedbackParameter = bGenerateFeedback ? FString::Printf(TEXT(", Parameters.VirtualTextureFeedback")) : TEXT("");
 
-	EDerivativeStatus UVDerivativeStatus = GetDerivativeStatus(SampleInfo.CoordinateIndex);
+	EDerivativeStatus UVDerivativeStatus = GetDerivativeStatus(CoordinateIndex);
 	const bool bHasValidDerivative = IsAnalyticDerivEnabled() && IsDerivativeValid(UVDerivativeStatus);
 
-	FString PageTable        = SampleInfo.PageTableIndex        != INDEX_NONE ? GetParameterCode(SampleInfo.PageTableIndex)        : FString::Printf(TEXT("VIRTUALTEXTURE_PAGETABLE_%d"), StackIndex);
-	FString PageTableUniform = SampleInfo.PageTableUniformIndex != INDEX_NONE ? GetParameterCode(SampleInfo.PageTableUniformIndex) : FString::Printf(TEXT("VTPageTableUniform_Unpack(VIRTUALTEXTURE_PAGETABLE_UNIFORM_%d)"), StackIndex);
+	FString PageTable        = PageTableOverride.IsEmpty()        ? FString::Printf(TEXT("VIRTUALTEXTURE_PAGETABLE_%d"),         StackIndex) : PageTableOverride;
+	FString PageTableUniform = PageTableUniformOverride.IsEmpty() ? FString::Printf(TEXT("VIRTUALTEXTURE_PAGETABLE_UNIFORM_%d"), StackIndex) : PageTableUniformOverride;
 
 	// Code to load the VT page table...this will execute the first time a given VT stack is accessed
 	// Additional stack layers will simply reuse these results
-	switch (SampleInfo.MipValueMode)
+	switch (MipValueMode)
 	{
 	case TMVM_None:
 	{
 		FString SampleCodeFinite = FString::Printf(TEXT(
 			"%s("
-			"%s, %s, "
+			"%s, "
+			"VTPageTableUniform_Unpack(%s), "
 			"%s, %s, %s, "
 			"0, Parameters.SvPosition.xy, "
 			"Parameters.VirtualTextureFeedback)"),
 			*BaseFunctionName,
 			*PageTable, *PageTableUniform,
-			*CoerceParameter(SampleInfo.CoordinateIndex, MCT_Float2), GetVTAddressMode(SampleInfo.AddressU), GetVTAddressMode(SampleInfo.AddressV));
+			*CoerceParameter(CoordinateIndex, MCT_Float2), GetVTAddressMode(AddressU), GetVTAddressMode(AddressV));
 
 		if (bHasValidDerivative)
 		{
 			FString SampleCodeAnalytic = FString::Printf(TEXT(
 				"%sGrad("
-				"%s, %s, "
+				"%s, "
+				"VTPageTableUniform_Unpack(%s), "
 				"%s, %s, %s, "
 				"%s, %s, Parameters.SvPosition.xy, "
 				"Parameters.VirtualTextureFeedback)"),
 				*BaseFunctionName,
 				*PageTable, *PageTableUniform,
-				*UV_Value, GetVTAddressMode(SampleInfo.AddressU), GetVTAddressMode(SampleInfo.AddressV), 
+				*UV_Value, GetVTAddressMode(AddressU), GetVTAddressMode(AddressV), 
 				*UV_Ddx, *UV_Ddy);
 			Entry.CodeIndex = AddCodeChunkInnerDeriv(*SampleCodeFinite, *SampleCodeAnalytic, MCT_VTPageTableResult, false, EDerivativeStatus::NotValid);
 		}
@@ -7199,25 +7171,27 @@ uint32 FHLSLMaterialTranslator::AcquireVTStackIndex(
 	{
 		FString SampleCodeFinite = FString::Printf(TEXT(
 			"%s("
-			"%s, %s, "
+			"%s, "
+			"VTPageTableUniform_Unpack(%s), "
 			"%s, %s, %s, "
 			"%s, Parameters.SvPosition.xy, "
 			"Parameters.VirtualTextureFeedback)"),
 			*BaseFunctionName,
 			*PageTable, *PageTableUniform,
-			*CoerceParameter(SampleInfo.CoordinateIndex, MCT_Float2), GetVTAddressMode(SampleInfo.AddressU), GetVTAddressMode(SampleInfo.AddressV), 
-			*CoerceParameter(SampleInfo.MipValue0Index, MCT_Float1));
+			*CoerceParameter(CoordinateIndex, MCT_Float2), GetVTAddressMode(AddressU), GetVTAddressMode(AddressV), 
+			*CoerceParameter(MipValue0Index, MCT_Float1));
 		if (bHasValidDerivative)
 		{
 			FString SampleCodeAnalytic = FString::Printf(TEXT(
 				"%sGrad("
-				"%s, %s, "
+				"%s, "
+				"VTPageTableUniform_Unpack(%s), "
 				"%s, %s, %s, "
 				"%s, %s, Parameters.SvPosition.xy, "
 				"Parameters.VirtualTextureFeedback)"),
 				*BaseFunctionName,
 				*PageTable, *PageTableUniform,
-				*UV_Value, GetVTAddressMode(SampleInfo.AddressU), GetVTAddressMode(SampleInfo.AddressV),
+				*UV_Value, GetVTAddressMode(AddressU), GetVTAddressMode(AddressV),
 				*UV_Ddx, *UV_Ddy);
 			Entry.CodeIndex = AddCodeChunkInnerDeriv(*SampleCodeFinite, *SampleCodeAnalytic, MCT_VTPageTableResult, false, EDerivativeStatus::NotValid);
 		}
@@ -7230,27 +7204,29 @@ uint32 FHLSLMaterialTranslator::AcquireVTStackIndex(
 	case TMVM_MipLevel:
 		Entry.CodeIndex = AddCodeChunk(MCT_VTPageTableResult, TEXT(
 			"%sLevel("
-			"%s, %s, "
+			"%s, "
+			"VTPageTableUniform_Unpack(%s), "
 			"%s, %s, %s, "
 			"%s"
 			"%s)"),
 			*BaseFunctionName,
 			*PageTable, *PageTableUniform,
-			*CoerceParameter(SampleInfo.CoordinateIndex, MCT_Float2), GetVTAddressMode(SampleInfo.AddressU), GetVTAddressMode(SampleInfo.AddressV), 
-			*CoerceParameter(SampleInfo.MipValue0Index, MCT_Float1),
+			*CoerceParameter(CoordinateIndex, MCT_Float2), GetVTAddressMode(AddressU), GetVTAddressMode(AddressV), 
+			*CoerceParameter(MipValue0Index, MCT_Float1),
 			*FeedbackParameter);
 		break;
 	case TMVM_Derivative:
 		Entry.CodeIndex = AddCodeChunk(MCT_VTPageTableResult, TEXT(
 			"%sGrad("
-			"%s, %s, "
+			"%s, "
+			"VTPageTableUniform_Unpack(%s), "
 			"%s, %s, %s, "
 			"%s, %s, Parameters.SvPosition.xy, "
 			"Parameters.VirtualTextureFeedback)"),
 			*BaseFunctionName,
 			*PageTable, *PageTableUniform,
-			*CoerceParameter(SampleInfo.CoordinateIndex, MCT_Float2), GetVTAddressMode(SampleInfo.AddressU), GetVTAddressMode(SampleInfo.AddressV),
-			*CoerceParameter(SampleInfo.MipValue0Index, MCT_Float2), *CoerceParameter(SampleInfo.MipValue1Index, MCT_Float2));
+			*CoerceParameter(CoordinateIndex, MCT_Float2), GetVTAddressMode(AddressU), GetVTAddressMode(AddressV),
+			*CoerceParameter(MipValue0Index, MCT_Float2), *CoerceParameter(MipValue1Index, MCT_Float2));
 		break;
 	default:
 		checkNoEntry();
@@ -7367,11 +7343,9 @@ int32 FHLSLMaterialTranslator::TextureSample(
 		return INDEX_NONE;
 	}
 
-	FMaterialVTSampleInfo SampleInfo;
-	
 	const bool bMeshPaint = TextureType == MCT_TextureMeshPaint;
 	const bool bMaterialCache = TextureType == MCT_TextureMaterialCache;
-	const bool bVirtualTexture = bMeshPaint || bMaterialCache || (TextureType & MCT_TextureVirtual) != 0;
+	const bool bVirtualTexture = bMeshPaint || bMaterialCache || TextureType == MCT_TextureVirtual;
 
 	if (MipValueMode == TMVM_Derivative)
 	{
@@ -7447,6 +7421,9 @@ int32 FHLSLMaterialTranslator::TextureSample(
 		AutomaticViewMipBias = false;
 	}
 
+	FString PageTableOverride;
+	FString PageTableUniformOverride;
+
 	FString TextureName;
 	int32 VirtualTextureIndex = INDEX_NONE;
 	if (TextureType == MCT_TextureCube)
@@ -7492,16 +7469,8 @@ int32 FHLSLMaterialTranslator::TextureSample(
 		// Binding index is tied to discovery order
 		uint32 MaterialCacheTagIndex = TextureUniform->GetMaterialCacheTagIndex();
 
-		SampleInfo.PageTableIndex = AddInternalCodeChunk(
-			MCT_Unexposed, TEXT("Texture2D<uint4>"),
-			*FString::Printf(TEXT("Material.MaterialCacheTagPageTable0_%d"), MaterialCacheTagIndex)
-		);
-		
-		SampleInfo.PageTableUniformIndex = AddCodeChunkInner(
-			MCT_Unexposed, TEXT("VTPageTableUniform"),
-			EDerivativeStatus::NotAware, false,
-			TEXT("VTPageTableUniform_Unpack(%s.PackedUniform)"), *GetParameterCode(TextureUniform->GetTextureIndex())
-		);
+		PageTableOverride        = FString::Printf(TEXT("Material.MaterialCacheTagPageTable0_%d"), MaterialCacheTagIndex);
+		PageTableUniformOverride = FString::Printf(TEXT("%s.PackedUniform"), *GetParameterCode(TextureUniform->GetTextureIndex()));
 
 		// Physical texture indexing assumes TagIndex * Stride + LayerIndex
 		TextureName = FString::Printf(
@@ -7511,32 +7480,6 @@ int32 FHLSLMaterialTranslator::TextureSample(
 
 		NumVtSamples++;
 	}
-	else if (bFromCollection && bVirtualTexture)
-	{
-		CoerceParameter(TextureIndex, TextureType);
-		
-		FMaterialUniformExpression* UniformExpression = GetParameterUniformExpression(TextureIndex); 
-		if (!UniformExpression || !UniformExpression->GetTextureCollectionUniformExpression())
-		{
-			return Errorf(TEXT("The provided uniform expression is not from a texture collection"));
-		}
-		
-		FMaterialUniformExpressionTextureCollection* CollectionExpression = UniformExpression->GetTextureCollectionUniformExpression();
-		VirtualTextureIndex = CollectionExpression->GetTextureCollectionTypePrefixIndex();
-
-		if (SamplerSource != SSM_FromTextureAsset)
-		{
-			const bool   bUseAnisoSampler  = VirtualTextureScalability::IsAnisotropicFilteringEnabled() && MipValueMode != TMVM_MipLevel;
-			const TCHAR* SharedSamplerName = bUseAnisoSampler ? TEXT("View.SharedBilinearAnisoClampedSampler") : TEXT("View.SharedBilinearClampedSampler");
-			TextureName += FString::Printf(TEXT("Material.TextureCollectionPhysical_%d, GetMaterialSharedSampler(Material.TextureCollectionSampler_%d, %s)"), VirtualTextureIndex, VirtualTextureIndex, SharedSamplerName);
-		}
-		else
-		{
-			TextureName += FString::Printf(TEXT("Material.TextureCollectionPhysical_%d, Material.TextureCollectionSampler_%d"), VirtualTextureIndex, VirtualTextureIndex);
-		}
-		
-		NumVtSamples++;
-	}
 	else if (bVirtualTexture)
 	{
 		// Note, this does not really do anything (by design) other than adding it to the UniformExpressionSet
@@ -7850,7 +7793,7 @@ int32 FHLSLMaterialTranslator::TextureSample(
 		AddCodeChunk(MCT_Float, TEXT("MaterialStoreTexCoordScale(Parameters, %s, %d)"), *CoerceParameter(NonLWCCoordinateIndex, UVsType), (int)TextureReferenceIndex);
 	}
 
-	FString UV_Value = CoerceParameter(NonLWCCoordinateIndex, UVsType);
+	const FString UV_Value = CoerceParameter(NonLWCCoordinateIndex, UVsType);
 	FString UV_Ddx = TEXT("0.0f");
 	FString UV_Ddy = TEXT("0.0f");
 	FString UV_Scale = TEXT("1.0f");
@@ -7882,29 +7825,15 @@ int32 FHLSLMaterialTranslator::TextureSample(
 		//todo[vt]: Support feedback from other shader types
 		const bool bGenerateFeedback = EnableFeedback && ShaderFrequency == SF_Pixel;
 
-		SampleInfo.MipValueMode = MipValueMode;
-		SampleInfo.AddressU = StaticAddressX;
-		SampleInfo.AddressV = StaticAddressY;
-		SampleInfo.AspectRatio = 1.0f;
-		SampleInfo.CoordinateIndex = CoordinateIndex;
-		SampleInfo.MipValue0Index = MipValue0Index;
-		SampleInfo.MipValue1Index = MipValue1Index;
-		SampleInfo.bAdaptive = AdaptiveVirtualTexture;
-		SampleInfo.bGenerateFeedback = bGenerateFeedback;
-		SampleInfo.bMeshPaint = bMeshPaint;
-		SampleInfo.bMaterialCache = bMaterialCache;
-		SampleInfo.bCollection = bFromCollection;
-        SampleInfo.VirtualTextureIndex = VirtualTextureIndex;
-
 		if (bMeshPaint)
 		{
-			VTStackIndex = AcquireVTStackIndex(SampleInfo, UV_Value, UV_Ddx, UV_Ddy);
+			VTStackIndex = AcquireVTStackIndex(VirtualTextureIndex, MipValueMode, StaticAddressX, StaticAddressY, 1.0f, CoordinateIndex, MipValue0Index, MipValue1Index, INDEX_NONE, UV_Value, UV_Ddx, UV_Ddy, PageTableOverride, PageTableUniformOverride, AdaptiveVirtualTexture, bGenerateFeedback, bMeshPaint, bMaterialCache);
 			VTPageTableLayerIndex = 0;
-			VTPackedUniformName = TEXT("VTUniform_Unpack(Scene.MeshPaint.PackedUniform)");
+			VTPackedUniformName = TEXT("Scene.MeshPaint.PackedUniform");
 		}
 		else if (bMaterialCache)
 		{
-			VTStackIndex = AcquireVTStackIndex(SampleInfo, UV_Value, UV_Ddx, UV_Ddy);
+			VTStackIndex = AcquireVTStackIndex(VirtualTextureIndex, MipValueMode, StaticAddressX, StaticAddressY, 1.0f, CoordinateIndex, MipValue0Index, MipValue1Index, INDEX_NONE, UV_Value, UV_Ddx, UV_Ddy, PageTableOverride, PageTableUniformOverride, AdaptiveVirtualTexture, bGenerateFeedback, bMeshPaint, bMaterialCache);
 			VTPageTableLayerIndex = 0;
 
 			FMaterialUniformExpression* UniformExpression = GetParameterUniformExpression(TextureIndex); 
@@ -7914,53 +7843,21 @@ int32 FHLSLMaterialTranslator::TextureSample(
 			}
 
 			FMaterialUniformExpressionTexture* TextureUniform = UniformExpression->GetTextureUniformExpression();
-			VTPackedUniformName = FString::Printf(TEXT("VTUniform_Unpack(Material.MaterialCacheTagPageTableUniform[%d])"), TextureUniform->GetMaterialCacheTagIndex());
-		}
-		else if (bFromCollection)
-		{
-			FString TextureCode = GetParameterCode(TextureIndex);
-			
-			SampleInfo.PageTableIndex = AddInternalCodeChunk(
-				MCT_Unexposed, TEXT("Texture2D<uint4>"),
-				*FString::Printf(TEXT("Material.TextureCollectionPageTable_%i"), VirtualTextureIndex)
-			);
-
-			FString Uniforms = GetParameterCode(AddInternalCodeChunk(
-				MCT_Unexposed, TEXT("FIndirectVirtualTextureUniform"),
-				*FString::Printf(TEXT("GetIndirectVirtualTextureUniform(%i)"), VirtualTextureIndex)
-			));
-
-			// Manually unpack the page table to redirect the page boundaries
-			SampleInfo.PageTableUniformIndex = AddInternalCodeChunk(
-				MCT_Unexposed, TEXT("VTPageTableUniform"),
-				*FString::Printf(
-					TEXT("UnpackIndirectVirtualTexturePageTableUniform(%s, %s.PackedPageTableUniform[0], %s.PackedPageTableUniform[1])"),
-					*TextureCode, *Uniforms, *Uniforms
-				)
-			);
-
-			VTPackedUniformName = FString::Printf(TEXT("VTUniform_Unpack(%s.PackedUniform)"), *Uniforms);
-			
-			VTStackIndex = AcquireVTStackIndex(SampleInfo, UV_Value, UV_Ddx, UV_Ddy);
-			VTPageTableLayerIndex = 0;
+			VTPackedUniformName = FString::Printf(TEXT("Material.MaterialCacheTagPageTableUniform[%d]"), TextureUniform->GetMaterialCacheTagIndex());
 		}
 		else 
 		{
 			check(VirtualTextureIndex >= 0);
-			VTPackedUniformName = FString::Printf(TEXT("VTUniform_Unpack(Material.VTPackedUniform[%d])"), VirtualTextureIndex);
-
-			SampleInfo.CoordinateIndex = NonLWCCoordinateIndex;
+			VTPackedUniformName = FString::Printf(TEXT("Material.VTPackedUniform[%d]"), VirtualTextureIndex);
 
 			const int32 VTTextureLayerIndex = UniformTextureExpressions[(uint32)EMaterialTextureParameterType::Virtual][VirtualTextureIndex]->GetTextureLayerIndex();
 			if (VTTextureLayerIndex != INDEX_NONE)
 			{
-				SampleInfo.PreallocatedStackTextureIndex = TextureReferenceIndex;
-				
 				// The layer index in the virtual texture stack is already known
 				// Create a page table sample for each new combination of virtual texture and sample parameters
-				VTStackIndex = AcquireVTStackIndex(SampleInfo, UV_Value, UV_Ddx, UV_Ddy);
+				VTStackIndex = AcquireVTStackIndex(VirtualTextureIndex, MipValueMode, StaticAddressX, StaticAddressY, 1.0f, NonLWCCoordinateIndex, MipValue0Index, MipValue1Index, TextureReferenceIndex, UV_Value, UV_Ddx, UV_Ddy, PageTableOverride, PageTableUniformOverride, AdaptiveVirtualTexture, bGenerateFeedback, bMeshPaint, bMaterialCache);
 				VTPageTableLayerIndex = UniformTextureExpressions[(uint32)EMaterialTextureParameterType::Virtual][VirtualTextureIndex]->GetPageTableLayerIndex();
-
+		
 				MaterialCompilationOutput.UniformExpressionSet.VTStacks[VTStackIndex].SetLayer(VTTextureLayerIndex, VirtualTextureIndex);
 			}
 			else
@@ -7975,10 +7872,8 @@ int32 FHLSLMaterialTranslator::TextureSample(
 				// This is because we only need to overlay VT pages from within a given block
 				const float TextureAspectRatio = (float)Texture->Source.GetSizeX() / (float)Texture->Source.GetSizeY();
 
-				SampleInfo.AspectRatio = TextureAspectRatio;
-				
 				// Create a page table sample for each new set of sample parameters
-				VTStackIndex = AcquireVTStackIndex(SampleInfo, UV_Value, UV_Ddx, UV_Ddy);
+				VTStackIndex = AcquireVTStackIndex(VirtualTextureIndex, MipValueMode, StaticAddressX, StaticAddressY, TextureAspectRatio, NonLWCCoordinateIndex, MipValue0Index, MipValue1Index, INDEX_NONE, UV_Value, UV_Ddx, UV_Ddy, PageTableOverride, PageTableUniformOverride, AdaptiveVirtualTexture, bGenerateFeedback, bMeshPaint, bMaterialCache);
 
 				// Find or allocate the layer in the virtual texture stack for this physical sample
 				VTPageTableLayerIndex = MaterialCompilationOutput.UniformExpressionSet.VTStacks[VTStackIndex].FindLayer(VirtualTextureIndex);
@@ -8000,7 +7895,7 @@ int32 FHLSLMaterialTranslator::TextureSample(
 		const FString VTPageTableResult_Finite = GetParameterCode(VTStackEntry.CodeIndex);
 
 		// 'Texture name/sampler', 'PageTableResult', 'LayerIndex', 'PackedUniform'
-		FString SampleCodeFinite = FString::Printf(TEXT("%s(%s, %s, %d, %s)"), *TextureTypeName, *TextureName, *VTPageTableResult_Finite, VTPageTableLayerIndex, *VTPackedUniformName);
+		FString SampleCodeFinite = FString::Printf(TEXT("%s(%s, %s, %d, VTUniform_Unpack(%s))"), *TextureTypeName, *TextureName, *VTPageTableResult_Finite, VTPageTableLayerIndex, *VTPackedUniformName);
 		SampleCodeFinite = ApplySamplerType(SampleCodeFinite, SamplerType);
 
 		if (IsAnalyticDerivEnabled() && IsDerivativeValid(UvDerivativeStatus))
@@ -8011,11 +7906,11 @@ int32 FHLSLMaterialTranslator::TextureSample(
 
 			if (SamplerDebugSupported(TextureType, bVirtualTexture, bDecal))
 			{
-				SampleCodeAnalytic = FString::Printf(TEXT("Debug%s(%s, %s, %d, %s, %s)"), *TextureTypeName, *TextureName, *VTPageTableResult_Analytic, VTPageTableLayerIndex, *VTPackedUniformName, *UV_Scale);
+				SampleCodeAnalytic = FString::Printf(TEXT("Debug%s(%s, %s, %d, VTUniform_Unpack(%s), %s)"), *TextureTypeName, *TextureName, *VTPageTableResult_Analytic, VTPageTableLayerIndex, *VTPackedUniformName, *UV_Scale);
 			}
 			else
 			{
-				SampleCodeAnalytic = FString::Printf(TEXT("%s(%s, %s, %d, %s)"), *TextureTypeName, *TextureName, *VTPageTableResult_Analytic, VTPageTableLayerIndex, *VTPackedUniformName);
+				SampleCodeAnalytic = FString::Printf(TEXT("%s(%s, %s, %d, VTUniform_Unpack(%s))"), *TextureTypeName, *TextureName, *VTPageTableResult_Analytic, VTPageTableLayerIndex, *VTPackedUniformName);
 			}
 
 			SampleCodeAnalytic = ApplySamplerType(SampleCodeAnalytic, SamplerType);
@@ -8168,30 +8063,6 @@ int32 FHLSLMaterialTranslator::TextureFromCollection(int32 TextureCollectionCode
 		return Errorf(TEXT("Expected a texture return type"));
 	}
 
-	if (TextureCollectionExpression->IsVirtualCollection())
-	{
-		if (ResultTextureType != MCT_Texture2D)
-		{
-			return Errorf(TEXT("Texture collection virtual sampling only supports 2d textures"));
-		}
-		
-		FString Uniforms = GetParameterCode(AddInternalCodeChunk(
-			MCT_Unexposed, TEXT("FIndirectVirtualTextureUniform"),
-			*FString::Printf(TEXT("GetIndirectVirtualTextureUniform(%i)"), TextureCollectionExpression->GetTextureCollectionTypePrefixIndex())
-		));
-		
-		return AddInternalCodeChunk(
-			static_cast<EMaterialValueType>(ResultTextureType | MCT_TextureVirtual | MCT_TextureCollection | MCT_Unexposed), TEXT("FIndirectVirtualTextureEntry"),
-			*FString::Printf(
-				TEXT("LoadIndirectVirtualTexture(%s, %s, %s)"),
-				*GetParameterCode(TextureCollectionCodeIndex),
-				*Uniforms,
-				*GetParameterCode(IndexIntoCollectionCodeIndex)
-			),
-			TextureCollectionExpression
-		);
-	}
-	
 	return AddCodeChunk(
 		EMaterialValueType(ResultTextureType | MCT_TextureCollection),
 		TEXT("TextureFromCollection_%s(%s, %s)"),
@@ -8944,10 +8815,8 @@ int32 FHLSLMaterialTranslator::TextureCollection(UTextureCollection* InTextureCo
 		return INDEX_NONE;
 	}
 
-	const bool bIsVirtualCollection = InTextureCollection->IsVirtualCollection();
-
 	FAddUniformExpressionScope Scope(this);
-	return AddUniformExpression(Scope, new FMaterialUniformExpressionTextureCollection(TextureCollectionReferenceIndex, UINT32_MAX, bIsVirtualCollection), MCT_TextureCollection, TEXT(""));
+	return AddUniformExpression(Scope, new FMaterialUniformExpressionTextureCollection(TextureCollectionReferenceIndex), MCT_TextureCollection, TEXT(""));
 }
 
 int32 FHLSLMaterialTranslator::TextureCollectionParameter(FName ParameterName, UTextureCollection* InDefaultValue, int32& TextureCollectionReferenceIndex)
@@ -8970,13 +8839,11 @@ int32 FHLSLMaterialTranslator::TextureCollectionParameter(FName ParameterName, U
 		return INDEX_NONE;
 	}
 
-	const bool bIsVirtualCollection = InDefaultValue->IsVirtualCollection();
-
 	FMaterialParameterInfo ParameterInfo = GetParameterAssociationInfo();
 	ParameterInfo.Name = ParameterName;
 
 	FAddUniformExpressionScope Scope(this);
-	return AddUniformExpression(Scope, new FMaterialUniformExpressionTextureCollectionParameter(ParameterInfo, TextureCollectionReferenceIndex, UINT32_MAX, bIsVirtualCollection), MCT_TextureCollection, TEXT(""));
+	return AddUniformExpression(Scope, new FMaterialUniformExpressionTextureCollectionParameter(ParameterInfo, TextureCollectionReferenceIndex), MCT_TextureCollection, TEXT(""));
 }
 
 int32 FHLSLMaterialTranslator::TextureCollectionCount(int32 InTextureCollectionCodeIndex)
@@ -8997,19 +8864,7 @@ int32 FHLSLMaterialTranslator::TextureCollectionCount(int32 InTextureCollectionC
 	{
 		return Errorf(TEXT("Expected a texture collection expression"));
 	}
-	
-	if (TextureCollectionExpression->IsVirtualCollection())
-	{
-		// Make sure the uniform is present
-		AccessUniformExpression(InTextureCollectionCodeIndex);
 
-		return AddCodeChunk(
-			EMaterialValueType(MCT_UInt1),
-			TEXT("GetIndirectVirtualTextureCount(GetIndirectVirtualTextureUniform(%u))"),
-			TextureCollectionExpression->GetTextureCollectionTypePrefixIndex()
-		);
-	}
-	
 	return AddCodeChunk(
 		EMaterialValueType(MCT_UInt1),
 		TEXT("GetCountFromResourceCollection(%s)"),
@@ -16198,7 +16053,7 @@ void FHLSLMaterialTranslator::PrepareEnvironmentDefines()
 	uint32 NumVirtualTextureFeedbackRequests = 0;
 	for (int i = 0; i < VTStacks.Num(); ++i)
 	{
-		if (VTStacks[i].SampleInfo.bGenerateFeedback)
+		if (VTStacks[i].bGenerateFeedback)
 		{
 			++NumVirtualTextureFeedbackRequests;
 		}
@@ -16216,22 +16071,18 @@ void FHLSLMaterialTranslator::PrepareEnvironmentDefines()
 		{
 			EnvironmentDefines->VirtualPageTypes[i] |= FEnvironmentDefines::LOCAL_TABLE1;
 		}
-		if (VTStacks[i].SampleInfo.bAdaptive)
+		if (VTStacks[i].bAdaptive)
 		{
 			EnvironmentDefines->VirtualPageTypes[i] |= FEnvironmentDefines::LOCAL_ADAPTIVE_INDIRECTION;
 		}
-		if (VTStacks[i].SampleInfo.bMeshPaint)
+		if (VTStacks[i].bMeshPaint)
 		{
 			EnvironmentDefines->VirtualPageTypes[i] |= FEnvironmentDefines::TABLE_MESH_PAINT;
 		}
-		if (VTStacks[i].SampleInfo.bMaterialCache)
+		if (VTStacks[i].bMaterialCache)
 		{
 			EnvironmentDefines->VirtualPageTypes[i] |= FEnvironmentDefines::TABLE_MATERIAL_CACHE;
 		}
-		if (VTStacks[i].SampleInfo.bCollection)
-		{
-			EnvironmentDefines->VirtualPageTypes[i] |= FEnvironmentDefines::TABLE_COLLECTION;
-		}
 	}
 
 	// Set all the shading models for this material here 
diff --git a/Engine/Source/Runtime/Engine/Private/Materials/HLSLMaterialTranslator.h b/Engine/Source/Runtime/Engine/Private/Materials/HLSLMaterialTranslator.h
--- a/Engine/Source/Runtime/Engine/Private/Materials/HLSLMaterialTranslator.h
+++ b/Engine/Source/Runtime/Engine/Private/Materials/HLSLMaterialTranslator.h
@@ -123,7 +123,6 @@ struct FShaderCodeChunk
 
 	/** Whether the code chunk should be inlined or not.  If true, SymbolName is empty and Definition contains the code to inline. */
 	bool bInline;
-	bool bIntermediate;
 	/** The status of partial derivatives at this expression. **/
 	EDerivativeStatus DerivativeStatus;
 
@@ -152,7 +151,6 @@ struct FShaderCodeChunk
 		UsedScopeIndex(INDEX_NONE),
 		ScopeLevel(0),
 		bInline(bInInline),
-		bIntermediate(false),
 		DerivativeStatus(InDerivativeStatus)
 	{}
 
@@ -168,7 +166,6 @@ struct FShaderCodeChunk
 		UsedScopeIndex(INDEX_NONE),
 		ScopeLevel(0),
 		bInline(false),
-		bIntermediate(false),
 		DerivativeStatus(InDerivativeStatus)
 	{}
 };
@@ -238,39 +235,28 @@ class FHLSLMaterialTranslator : public FMaterialCompiler
 {
 	friend class FMaterialDerivativeAutogen;
 protected:
-	struct FMaterialVTSampleInfo
-	{
-		bool operator==(const FMaterialVTSampleInfo&) const = default;
-		
-		int32 VirtualTextureIndex = INDEX_NONE;
-		ETextureMipValueMode MipValueMode = TMVM_None;
-		TextureAddress AddressU = TA_Clamp;
-		TextureAddress AddressV = TA_Clamp;
-		int32 CoordinateIndex = INDEX_NONE;
-		int32 MipValue0Index = INDEX_NONE;
-		int32 MipValue1Index = INDEX_NONE;
-		int32 PreallocatedStackTextureIndex = INDEX_NONE;
-		int32 PageTableIndex = INDEX_NONE;
-		int32 PageTableUniformIndex = INDEX_NONE;
-		bool bAdaptive = false;
-		bool bGenerateFeedback = false;
-		bool bMeshPaint = false;
-		bool bMaterialCache = false;
-		bool bCollection = false;
-		float AspectRatio = 0.0f;
-		int32 CodeIndex = INDEX_NONE;
-	};
-
 	/** Translator specific structure during generation of VT stacks. This is later converted to FMaterialVirtualTextureStack entries in FUniformExpressionSet. */
 	struct FMaterialVTStackEntry
 	{
-		FMaterialVTSampleInfo SampleInfo;
 		uint64 ScopeID;
 		uint64 CoordinateHash;
 		uint64 MipValue0Hash;
 		uint64 MipValue1Hash;
 		uint64 PageTableHash;
 		uint64 PageTableUniformHash;
+		ETextureMipValueMode MipValueMode;
+		TextureAddress AddressU;
+		TextureAddress AddressV;
+		int32 DebugCoordinateIndex;
+		int32 DebugMipValue0Index;
+		int32 DebugMipValue1Index;
+		int32 PreallocatedStackTextureIndex;
+		bool bAdaptive;
+		bool bGenerateFeedback;
+		bool bMeshPaint;
+		bool bMaterialCache;
+		float AspectRatio;
+
 		int32 CodeIndex;
 	};
 
@@ -765,9 +751,6 @@ protected:
 	int32 AddCodeChunkInner(uint64 Hash, const TCHAR* FormattedType, const TCHAR* FormattedCode, EMaterialValueType Type, EDerivativeStatus DerivativeStatus, bool bInlined);
 	int32 AddCodeChunkInner(uint64 Hash, const TCHAR* FormattedCode, EMaterialValueType Type, EDerivativeStatus DerivativeStatus, bool bInlined);
 
-	/** Adds an already formatted inline or referenced code chunk */
-	int32 AddInternalCodeChunk(EMaterialValueType Type, const TCHAR* InternalType, const TCHAR* FormattedCode, const TRefCountPtr<FMaterialUniformExpression>& UniformExpression = nullptr);
-
 public:
 	/** Adds an already formatted inline or referenced code chunk, and notes the derivative status. */
 	int32 AddCodeChunkInnerDeriv(const TCHAR* FormattedCodeFinite, const TCHAR* FormattedCodeAnalytic, EMaterialValueType Type, bool bInlined, EDerivativeStatus DerivativeStatus);
@@ -1073,11 +1056,19 @@ protected:
 	virtual void SetPotentiallyManipulateTexCoords() override;
 
 	uint32 AcquireVTStackIndex(
-		const FMaterialVTSampleInfo& SampleInfo,
+		int32 VirtualTexureIndex,
+		ETextureMipValueMode MipValueMode,
+		TextureAddress AddressU, TextureAddress AddressV, 
+		float AspectRatio, 
+		int32 CoordinateIndex, 
+		int32 MipValue0Index, int32 MipValue1Index, 
+		int32 PreallocatedStackTextureIndex, 
 		const FString& UV_Value,
 		const FString& UV_Ddx,
-		const FString& UV_Ddy
-	);
+		const FString& UV_Ddy,
+		const FString& PageTableOverride,
+		const FString& PageTableUniformOverride,
+		bool bAdaptive, bool bGenerateFeedback, bool bMeshPaint, bool bMaterialCache);
 
 	virtual int32 TextureSample(
 		int32 TextureIndex,
diff --git a/Engine/Source/Runtime/Engine/Private/Materials/MaterialRenderProxy.cpp b/Engine/Source/Runtime/Engine/Private/Materials/MaterialRenderProxy.cpp
--- a/Engine/Source/Runtime/Engine/Private/Materials/MaterialRenderProxy.cpp
+++ b/Engine/Source/Runtime/Engine/Private/Materials/MaterialRenderProxy.cpp
@@ -14,8 +14,6 @@
 #include "TextureResource.h"
 #include "RenderGraphBuilder.h"
 #include "MaterialCache/IMaterialCacheTagProvider.h"
-#include "Engine/TextureCollection.h"
-#include "Engine/VirtualTextureCollection.h"
 #include "VT/RuntimeVirtualTexture.h"
 
 #if WITH_EDITOR
@@ -489,35 +487,6 @@ void FMaterialRenderProxy::EvaluateUniformExpressions(FRHICommandListBase& RHICm
 		bHasMaterialCacheCallbacks = true;
 	}
 
-	for (int32 VTCIndex = 0; VTCIndex < UniformExpressionSet.UniformTextureCollectionParameters.Num(); ++VTCIndex)
-	{
-		const FMaterialTextureCollectionParameterInfo& Parameter = UniformExpressionSet.UniformTextureCollectionParameters[VTCIndex];
-		if (!Parameter.bIsVirtualCollection)
-		{
-			continue;
-		}
-
-		const UTextureCollection* Texture = nullptr;
-		UniformExpressionSet.GetTextureCollectionValue(VTCIndex, Context, Context.Material, Texture);
-
-		// May not be bound, but must be a VTC
-		if (!Texture || !ensure(Texture->IsVirtualCollection()))
-		{
-			continue;
-		}
-
-		// May not have a valid resource
-		FVirtualTextureCollectionResource* Resource = static_cast<FVirtualTextureCollectionResource*>(Texture->GetResource());
-		if (!Resource || !ensure(!Resource->bIsBindless))
-		{
-			continue;
-		}
-
-		// Subscribe to VT destruction
-		GetRendererModule().AddVirtualTextureProducerDestroyedCallback(Resource->GetProducerHandle(), &OnVirtualTextureDestroyedCB, const_cast<FMaterialRenderProxy*>(this));
-		bHasVirtualTextureCallbacks = true;
-	}
-
 	const FRHIUniformBufferLayout* UniformBufferLayout = ShaderMap->GetUniformBufferLayout();
 
 	if (IsValidRef(OutUniformExpressionCache.UniformBuffer))
diff --git a/Engine/Source/Runtime/Engine/Private/Materials/MaterialShared.cpp b/Engine/Source/Runtime/Engine/Private/Materials/MaterialShared.cpp
--- a/Engine/Source/Runtime/Engine/Private/Materials/MaterialShared.cpp
+++ b/Engine/Source/Runtime/Engine/Private/Materials/MaterialShared.cpp
@@ -116,19 +116,6 @@ bool IsUsingNewMaterialTranslatorPrototype()
 
 IMPLEMENT_TYPE_LAYOUT(FHashedMaterialParameterInfo);
 IMPLEMENT_TYPE_LAYOUT(FUniformExpressionSet);
-
-void FUniformExpressionSet::CountTextureCollections(uint32& BindlessCollections, uint32& VirtualCollections) const
-{
-	VirtualCollections = 0;
-	
-	for (const FMaterialTextureCollectionParameterInfo& UniformTextureCollectionParameter : UniformTextureCollectionParameters)
-	{
-		VirtualCollections += UniformTextureCollectionParameter.bIsVirtualCollection;
-	}
-
-	BindlessCollections = UniformTextureCollectionParameters.Num() - VirtualCollections;
-}
-
 IMPLEMENT_TYPE_LAYOUT(FMaterialCompilationOutput);
 IMPLEMENT_TYPE_LAYOUT(FMeshMaterialShaderMap);
 IMPLEMENT_TYPE_LAYOUT(FMaterialProcessedSource);
diff --git a/Engine/Source/Runtime/Engine/Private/Materials/MaterialUniformExpressions.cpp b/Engine/Source/Runtime/Engine/Private/Materials/MaterialUniformExpressions.cpp
--- a/Engine/Source/Runtime/Engine/Private/Materials/MaterialUniformExpressions.cpp
+++ b/Engine/Source/Runtime/Engine/Private/Materials/MaterialUniformExpressions.cpp
@@ -25,7 +25,6 @@
 #include "MaterialCache/MaterialCacheAttribute.h"
 #include "MaterialCache/MaterialCacheTagProvider.h"
 #include "MaterialCache/MaterialCacheTagSceneData.h"
-#include "Engine/VirtualTextureCollection.h"
 
 void WriteMaterialUniformAccess(UE::Shader::EValueComponentType ComponentType, uint32 NumComponents, uint32 UniformOffset, FStringBuilderBase& OutResult)
 {
@@ -395,22 +394,7 @@ FShaderParametersMetadata* FUniformExpressionSet::CreateBufferStruct()
 	// Make sure FUniformExpressionSet::CreateDebugLayout() is in sync
 	TArray<FShaderParametersMetadata::FMember> Members;
 	uint32 NextMemberOffset = 0;
-	
-	if (!UniformTextureCollectionParameters.IsEmpty())
-	{
-		uint32 BindlessCollectionCount = 0;
-		uint32 VirtualCollectionCount  = 0;
-		CountTextureCollections(BindlessCollectionCount, VirtualCollectionCount);
 
-		// Bindless collections do not make use of uniforms
-		if (VirtualCollectionCount)
-		{
-			static_assert(sizeof(FUintVector4) * UE::HLSL::IndirectVirtualTextureUniformDWord4Count == sizeof(UE::HLSL::FIndirectVirtualTextureUniform), "Unexpected size");
-			new(Members) FShaderParametersMetadata::FMember(TEXT("TextureCollectionPackedUniforms"), TEXT("FIndirectVirtualTextureUniform"), __LINE__, NextMemberOffset, UBMT_UINT32, EShaderPrecisionModifier::Float, 1, 4, UE::HLSL::IndirectVirtualTextureUniformDWord4Count * VirtualCollectionCount, NULL);
-			NextMemberOffset += sizeof(UE::HLSL::FIndirectVirtualTextureUniform) * UniformTextureCollectionParameters.Num();
-		}
-	}
-	
 	if (VTStacks.Num())
 	{
 		// 2x uint4 per VTStack
@@ -462,9 +446,6 @@ FShaderParametersMetadata* FUniformExpressionSet::CreateBufferStruct()
 	static FString VolumeTextureSamplerNames[MaxUniformExpressions];
 	static FString ExternalTextureNames[MaxUniformExpressions];
 	static FString TextureCollectionNames[MaxUniformExpressions];
-	static FString TextureCollectionPhysicalNames[MaxUniformExpressions];
-	static FString TextureCollectionSamplerNames[MaxUniformExpressions];
-	static FString TextureCollectionPageTableNames[MaxUniformExpressions];
 	static FString MediaTextureSamplerNames[MaxUniformExpressions];
 	static FString VirtualTexturePageTableNames0[MaxUniformExpressions];
 	static FString VirtualTexturePageTableNames1[MaxUniformExpressions];
@@ -499,9 +480,6 @@ FShaderParametersMetadata* FUniformExpressionSet::CreateBufferStruct()
 			VolumeTextureSamplerNames[i] = FString::Printf(TEXT("VolumeTexture_%dSampler"), i);
 			ExternalTextureNames[i] = FString::Printf(TEXT("ExternalTexture_%d"), i);
 			TextureCollectionNames[i] = FString::Printf(TEXT("TextureCollection_%d"), i);
-			TextureCollectionPhysicalNames[i] = FString::Printf(TEXT("TextureCollectionPhysical_%d"), i);
-			TextureCollectionSamplerNames[i] = FString::Printf(TEXT("TextureCollectionSampler_%d"), i);
-			TextureCollectionPageTableNames[i] = FString::Printf(TEXT("TextureCollectionPageTable_%d"), i);
 			MediaTextureSamplerNames[i] = FString::Printf(TEXT("ExternalTexture_%dSampler"), i);
 			VirtualTexturePageTableNames0[i] = FString::Printf(TEXT("VirtualTexturePageTable0_%d"), i);
 			VirtualTexturePageTableNames1[i] = FString::Printf(TEXT("VirtualTexturePageTable1_%d"), i);
@@ -611,27 +589,8 @@ FShaderParametersMetadata* FUniformExpressionSet::CreateBufferStruct()
 	for (int32 i = 0; i < UniformTextureCollectionParameters.Num(); ++i)
 	{
 		check((NextMemberOffset % SHADER_PARAMETER_POINTER_ALIGNMENT) == 0);
-
-		// Resource collections are RHI constructs, texture collections are software virtual tables
-		if (!UniformTextureCollectionParameters[i].bIsVirtualCollection)
-		{
-			new(Members) FShaderParametersMetadata::FMember(*TextureCollectionNames[i], TEXT("FResourceCollection"), __LINE__, NextMemberOffset, UBMT_RESOURCE_COLLECTION, EShaderPrecisionModifier::Float, 1, 1, 0, NULL);
-			NextMemberOffset += SHADER_PARAMETER_POINTER_ALIGNMENT;
-		}
-		else
-		{
-			new(Members) FShaderParametersMetadata::FMember(*TextureCollectionNames[i], TEXT("FIndirectVirtualTexture"), __LINE__, NextMemberOffset, UBMT_SRV, EShaderPrecisionModifier::Float, 1, 1, 0, NULL);
-			NextMemberOffset += SHADER_PARAMETER_POINTER_ALIGNMENT;
-			
-			new(Members) FShaderParametersMetadata::FMember(*TextureCollectionPageTableNames[i], TEXT("Texture2D<uint4>"), __LINE__, NextMemberOffset, UBMT_TEXTURE, EShaderPrecisionModifier::Float, 1, 1, 0, NULL);
-			NextMemberOffset += SHADER_PARAMETER_POINTER_ALIGNMENT;
-		
-			new(Members) FShaderParametersMetadata::FMember(*TextureCollectionPhysicalNames[i], TEXT("Texture2D"), __LINE__, NextMemberOffset, UBMT_SRV, EShaderPrecisionModifier::Float, 1, 1, 0, NULL);
-			NextMemberOffset += SHADER_PARAMETER_POINTER_ALIGNMENT;
-	
-			new(Members) FShaderParametersMetadata::FMember(*TextureCollectionSamplerNames[i], TEXT("SamplerState"), __LINE__, NextMemberOffset, UBMT_SAMPLER, EShaderPrecisionModifier::Float, 1, 1, 0, NULL);
-			NextMemberOffset += SHADER_PARAMETER_POINTER_ALIGNMENT;
-		}
+		new(Members) FShaderParametersMetadata::FMember(*TextureCollectionNames[i], TEXT("FResourceCollection"), __LINE__, NextMemberOffset, UBMT_RESOURCE_COLLECTION, EShaderPrecisionModifier::Float, 1, 1, 0, NULL);
+		NextMemberOffset += SHADER_PARAMETER_POINTER_ALIGNMENT;
 	}
 
 	for (int32 i = 0; i < VTStacks.Num(); ++i)
@@ -920,20 +879,6 @@ void FUniformExpressionSet::GetTextureValue(int32 Index, const FMaterialRenderCo
 	}
 }
 
-void FUniformExpressionSet::GetTextureCollectionValue(int32 Index, const FMaterialRenderContext& Context, const FMaterial& Material, const UTextureCollection*& OutValue) const
-{
-	check(IsInParallelRenderingThread());
-	if (ensure(Index < GetNumTextureCollections()))
-	{
-		const FMaterialTextureCollectionParameterInfo& Parameter = GetTextureCollectionParameter(Index);
-		Context.GetTextureCollectionParameterValue(Parameter.ParameterInfo, Parameter.TextureCollectionIndex, OutValue);
-	}
-	else
-	{
-		OutValue = nullptr;
-	}
-}
-
 void FUniformExpressionSet::GetGameThreadTextureCollectionValue(int32 Index, const UMaterialInterface* MaterialInterface, const FMaterial& Material, UTextureCollection*& OutValue) const
 {
 	check(IsInGameThread() || IsInParallelGameThread());
@@ -1039,34 +984,7 @@ void FUniformExpressionSet::FillUniformBuffer(const FMaterialRenderContext& Mate
 
 		void* BufferCursor = TempBuffer;
 		check(BufferCursor <= TempBuffer + TempBufferSize);
-		
-		if (!UniformTextureCollectionParameters.IsEmpty())
-		{
-			// Bindless collections do not make use of uniforms
-			for (const FMaterialTextureCollectionParameterInfo& ParameterInfo : UniformTextureCollectionParameters)
-			{
-				if (!ParameterInfo.bIsVirtualCollection)
-				{
-					continue;
-				}
-				
-				const UTextureCollection* TextureCollection = nullptr;
-				ParameterInfo.GetTextureCollection(MaterialRenderContext, TextureCollection);
-		
-				auto* VTPackedPageTableUniform = static_cast<UE::HLSL::FIndirectVirtualTextureUniform*>(BufferCursor);
-				BufferCursor = VTPackedPageTableUniform + 1;
 
-				if (FVirtualTextureCollectionResource* Resource = static_cast<FVirtualTextureCollectionResource*>(TextureCollection->GetResource()))
-				{
-					*VTPackedPageTableUniform = Resource->GetVirtualPackedUniform();
-				}
-				else
-				{
-					*VTPackedPageTableUniform = UE::HLSL::FIndirectVirtualTextureUniform{};
-				}
-			}
-		}
-		
 		// Dump virtual texture per page table uniform data
 		check(AllocatedVTs.Num() == VTStacks.Num());
 		for ( int32 VTStackIndex = 0; VTStackIndex < VTStacks.Num(); ++VTStackIndex)
@@ -1280,10 +1198,6 @@ void FUniformExpressionSet::FillUniformBuffer(const FMaterialRenderContext& Mate
 				NumPageTableIndirectionTextures++;
 			}
 	
-			uint32 BindlessCollectionCount = 0;
-			uint32 VirtualCollectionCount  = 0;
-			CountTextureCollections(BindlessCollectionCount, VirtualCollectionCount);
-		
 			check(UniformBufferLayout->Resources.Num() == 
 				UniformTextureParameters[(uint32)EMaterialTextureParameterType::Standard2D].Num() * 2
 				+ UniformTextureParameters[(uint32)EMaterialTextureParameterType::Cube].Num() * 2
@@ -1292,8 +1206,7 @@ void FUniformExpressionSet::FillUniformBuffer(const FMaterialRenderContext& Mate
 				+ UniformTextureParameters[(uint32)EMaterialTextureParameterType::Volume].Num() * 2
 				+ UniformTextureParameters[(uint32)EMaterialTextureParameterType::SparseVolume].Num() * 4
 				+ UniformExternalTextureParameters.Num() * 2
-				+ BindlessCollectionCount * 1
-				+ VirtualCollectionCount * 4
+				+ UniformTextureCollectionParameters.Num()
 				+ UniformTextureParameters[(uint32)EMaterialTextureParameterType::Virtual].Num() * 2
 				+ MaterialCacheTagStacks.Num() * (MaterialCacheMaxRuntimeLayers + 2)
 				+ NumPageTableTextures
@@ -1630,56 +1543,25 @@ void FUniformExpressionSet::FillUniformBuffer(const FMaterialRenderContext& Mate
 			}
 		}
 
-		if (!UniformTextureCollectionParameters.IsEmpty())
+		for (const FMaterialTextureCollectionParameterInfo& ParameterInfo : UniformTextureCollectionParameters)
 		{
-			for (const FMaterialTextureCollectionParameterInfo& ParameterInfo : UniformTextureCollectionParameters)
-			{
-				const UTextureCollection* TextureCollection = nullptr;
-				ParameterInfo.GetTextureCollection(MaterialRenderContext, TextureCollection);
-
-				FTextureCollectionResource* Resource = TextureCollection ? TextureCollection->GetResource() : nullptr;
-
-				if (!ParameterInfo.bIsVirtualCollection)
-				{
-					void** ResourceTableBufferPtr = (void**)((uint8*)BufferCursor + 0 * SHADER_PARAMETER_POINTER_ALIGNMENT);
-					
-					if (FBindlessTextureCollectionResource* BindlessResource = static_cast<FBindlessTextureCollectionResource*>(Resource))
-					{
-						*ResourceTableBufferPtr = BindlessResource->GetResourceCollectionRHI();
-					}
-					else
-					{
-						*ResourceTableBufferPtr = GEmptyResourceCollection.ResourceCollection;
-					}
-		
-					BufferCursor = ((uint8*)BufferCursor) + (SHADER_PARAMETER_POINTER_ALIGNMENT);
-					check(BufferCursor <= TempBuffer + TempBufferSize);
-				}
-				else
-				{
-					void** ResourceTableBufferPtr          = (void**)((uint8*)BufferCursor + 0 * SHADER_PARAMETER_POINTER_ALIGNMENT);
-					void** ResourceTablePageTablePtr       = (void**)((uint8*)BufferCursor + 1 * SHADER_PARAMETER_POINTER_ALIGNMENT);
-					void** ResourceTablePhysicalTexturePtr = (void**)((uint8*)BufferCursor + 2 * SHADER_PARAMETER_POINTER_ALIGNMENT);
-					void** ResourceTablePhysicalSamplerPtr = (void**)((uint8*)BufferCursor + 3 * SHADER_PARAMETER_POINTER_ALIGNMENT);
-					
-					if (FVirtualTextureCollectionResource* VirtualResource = static_cast<FVirtualTextureCollectionResource*>(Resource))
-					{
-						*ResourceTableBufferPtr          = VirtualResource->GetVirtualCollectionRHI();
-						*ResourceTablePageTablePtr       = VirtualResource->GetVirtualPageTable();
-						*ResourceTablePhysicalTexturePtr = VirtualResource->GetVirtualPhysicalTextureSRV();
-					}
-					else
-					{
-						*ResourceTableBufferPtr          = GEmptyVertexBufferUInt4WithUAV->ShaderResourceViewRHI;
-						*ResourceTablePageTablePtr       = GBlackUintTexture->TextureRHI.GetReference();
-						*ResourceTablePhysicalTexturePtr = GBlackTexture->TextureRHI.GetReference();
-					}
+			void** ResourceTableBufferPtr = reinterpret_cast<void**>(BufferCursor);
+			BufferCursor = reinterpret_cast<uint8*>(BufferCursor) + SHADER_PARAMETER_POINTER_ALIGNMENT;
+			check(BufferCursor <= TempBuffer + TempBufferSize);
 
-					*ResourceTablePhysicalSamplerPtr = TStaticSamplerState<SF_AnisotropicPoint, AM_Clamp, AM_Clamp, AM_Clamp, 0, 8>::GetRHI();
+			const UTextureCollection* TextureCollection = nullptr;
+			ParameterInfo.GetTextureCollection(MaterialRenderContext, TextureCollection);
 
-					BufferCursor = ((uint8*)BufferCursor) + (SHADER_PARAMETER_POINTER_ALIGNMENT * 4);
-					check(BufferCursor <= TempBuffer + TempBufferSize);
-				}
+			if (TextureCollection != nullptr)
+			{
+				FTextureCollectionResource* Resource = TextureCollection->GetResource();
+				check(Resource);
+				*ResourceTableBufferPtr = Resource->GetRHI();
+			}
+			else
+			{
+				// TODO: christopher.waters - implement default texture collections
+				unimplemented();
 			}
 		}
 
@@ -1952,10 +1834,8 @@ bool FMaterialUniformExpressionTexture::IsIdentical(const FMaterialUniformExpres
 
 // FMaterialUniformExpressionTextureCollection
 
-FMaterialUniformExpressionTextureCollection::FMaterialUniformExpressionTextureCollection(int32 InTextureCollectionIndex, int32 InTextureCollectionTypePrefixIndex, bool bIsVirtualCollection)
+FMaterialUniformExpressionTextureCollection::FMaterialUniformExpressionTextureCollection(int32 InTextureCollectionIndex)
 	: TextureCollectionIndex(InTextureCollectionIndex)
-	, TextureCollectionTypePrefixIndex(InTextureCollectionTypePrefixIndex)
-    , bIsVirtualCollection(bIsVirtualCollection)
 {
 }
 
@@ -1978,13 +1858,12 @@ FMaterialUniformExpressionTextureCollectionParameter* FMaterialUniformExpression
 void FMaterialUniformExpressionTextureCollection::GetTextureCollectionParameterInfo(FMaterialTextureCollectionParameterInfo& OutParameter) const
 {
 	OutParameter.TextureCollectionIndex = this->TextureCollectionIndex;
-	OutParameter.bIsVirtualCollection = this->bIsVirtualCollection;
 }
 
 // FMaterialUniformExpressionTextureCollectionParameter
 
-FMaterialUniformExpressionTextureCollectionParameter::FMaterialUniformExpressionTextureCollectionParameter(const FMaterialParameterInfo& InParameterInfo, int32 InTextureCollectionIndex, int32 InTextureCollectionTypePrefixIndex, bool bIsVirtualCollection)
-	: Super(InTextureCollectionIndex, InTextureCollectionTypePrefixIndex, bIsVirtualCollection)
+FMaterialUniformExpressionTextureCollectionParameter::FMaterialUniformExpressionTextureCollectionParameter(const FMaterialParameterInfo& InParameterInfo, int32 InTextureCollectionIndex)
+	: Super(InTextureCollectionIndex)
 	, ParameterInfo(InParameterInfo)
 {
 }
diff --git a/Engine/Source/Runtime/Engine/Private/Materials/MaterialUniformExpressions.h b/Engine/Source/Runtime/Engine/Private/Materials/MaterialUniformExpressions.h
--- a/Engine/Source/Runtime/Engine/Private/Materials/MaterialUniformExpressions.h
+++ b/Engine/Source/Runtime/Engine/Private/Materials/MaterialUniformExpressions.h
@@ -140,7 +140,7 @@ class FMaterialUniformExpressionTextureCollection : public FMaterialUniformExpre
 	DECLARE_MATERIALUNIFORMEXPRESSION_TYPE(FMaterialUniformExpressionTextureCollection);
 
 public:
-	FMaterialUniformExpressionTextureCollection(int32 InTextureCollectionIndex, int32 InTextureCollectionTypePrefixIndex, bool bIsVirtualCollection);
+	FMaterialUniformExpressionTextureCollection(int32 InTextureCollectionIndex);
 
 	//~ Begin FMaterialUniformExpression Interface.
 	virtual class FMaterialUniformExpressionTextureCollection* GetTextureCollectionUniformExpression() final { return this; }
@@ -155,25 +155,8 @@ public:
 		return TextureCollectionIndex;
 	}
 
-	void SetTextureCollectionTypePrefixIndex(uint32 InPrefix)
-	{
-		TextureCollectionTypePrefixIndex = InPrefix;
-	}
-
-	int32 GetTextureCollectionTypePrefixIndex() const
-	{
-		return TextureCollectionTypePrefixIndex;
-	}
-
-	bool IsVirtualCollection() const
-	{
-		return bIsVirtualCollection;
-	}
-
 protected:
 	int32 TextureCollectionIndex;
-	int32 TextureCollectionTypePrefixIndex;
-	bool bIsVirtualCollection;
 };
 
 class FMaterialUniformExpressionTextureCollectionParameter : public FMaterialUniformExpressionTextureCollection
@@ -182,7 +165,7 @@ class FMaterialUniformExpressionTextureCollectionParameter : public FMaterialUni
 	DECLARE_MATERIALUNIFORMEXPRESSION_TYPE(FMaterialUniformExpressionTextureCollectionParameter);
 
 public:
-	FMaterialUniformExpressionTextureCollectionParameter(const FMaterialParameterInfo& InParameterInfo, int32 InTextureCollectionIndex, int32 InTextureCollectionTypePrefixIndex, bool bIsVirtualCollection);
+	FMaterialUniformExpressionTextureCollectionParameter(const FMaterialParameterInfo& InParameterInfo, int32 InTextureCollectionIndex);
 
 	//~ Begin FMaterialUniformExpression Interface.
 	virtual bool IsIdentical(const FMaterialUniformExpression* OtherExpression) const override;
@@ -206,6 +189,7 @@ public:
 
 private:
 	FHashedMaterialParameterInfo ParameterInfo;
+	int32 ParameterIndex;
 };
 
 class FMaterialUniformExpressionExternalTextureBase : public FMaterialUniformExpression
diff --git a/Engine/Source/Runtime/Engine/Private/TextureCollection.cpp b/Engine/Source/Runtime/Engine/Private/TextureCollection.cpp
--- a/Engine/Source/Runtime/Engine/Private/TextureCollection.cpp
+++ b/Engine/Source/Runtime/Engine/Private/TextureCollection.cpp
@@ -18,11 +18,6 @@
 #define LOCTEXT_NAMESPACE "UTextureCollection"
 
 FTextureCollectionResource::FTextureCollectionResource(UTextureCollection* InParent)
-{
-	CollectionName = InParent->GetFName();
-}
-
-FBindlessTextureCollectionResource::FBindlessTextureCollectionResource(UTextureCollection* InParent) : FTextureCollectionResource(InParent)
 {
 	TConstArrayView<TObjectPtr<UTexture>> ParentTextures = InParent->Textures;
 	InputTextureResources.Reserve(ParentTextures.Num());
@@ -35,23 +30,25 @@ FBindlessTextureCollectionResource::FBindlessTextureCollectionResource(UTextureC
 	}
 }
 
-void FBindlessTextureCollectionResource::InitRHI(FRHICommandListBase& RHICmdList)
+void FTextureCollectionResource::InitRHI(FRHICommandListBase& RHICmdList)
 {
+	InputTextures.Empty(InputTextureResources.Num());
+
 	TArray<FRHIResourceCollectionMember> CollectionMembers;
 	CollectionMembers.Reserve(InputTextureResources.Num());
 
 	for (FTextureResource* TextureResource : InputTextureResources)
 	{
 		FRHITextureReference* TextureReference = TextureResource ? TextureResource->GetTextureReference() : nullptr;
+		InputTextures.Emplace(TextureReference);
 		CollectionMembers.Emplace(TextureReference);
 	}
 
 	ResourceCollectionRHI = RHICmdList.CreateResourceCollection(CollectionMembers);
 }
 
-void FBindlessTextureCollectionResource::ReleaseRHI()
+void FTextureCollectionResource::ReleaseRHI()
 {
-	
 }
 
 #if WITH_EDITOR
@@ -83,6 +80,8 @@ void UTextureCollection::NotifyMaterials()
 void UTextureCollection::PostEditChangeProperty(FPropertyChangedEvent& PropertyChangedEvent)
 {
 	UpdateResource();
+
+	NotifyMaterials();
 }
 #endif // WITH_EDITOR
 
@@ -150,14 +149,7 @@ FTextureCollectionResource* UTextureCollection::GetResource()
 
 FTextureCollectionResource* UTextureCollection::CreateResource()
 {
-	// TODO: christopher.waters - We should add a way to query the RHI to see if bindless is enabled.
-	ERHIBindlessConfiguration BindlessResourcesConfig = RHIGetRuntimeBindlessConfiguration(GMaxRHIShaderPlatform);
-	if (BindlessResourcesConfig == ERHIBindlessConfiguration::Disabled || !GRHIGlobals.bSupportsBindless)
-	{
-		return nullptr;
-	}
-
-	return new FBindlessTextureCollectionResource(this);
+	return new FTextureCollectionResource(this);
 }
 
 void UTextureCollection::ReleaseResource()
@@ -188,6 +180,13 @@ void UTextureCollection::UpdateResource()
 		return;
 	}
 
+	// TODO: christopher.waters - We should add a way to query the RHI to see if bindless is enabled.
+	ERHIBindlessConfiguration BindlessConfig = RHIGetRuntimeBindlessConfiguration(GMaxRHIShaderPlatform);
+	if (BindlessConfig == ERHIBindlessConfiguration::Disabled || !GRHIGlobals.bSupportsBindless)
+	{
+		return;
+	}
+
 	if (FTextureCollectionResource* NewResource = CreateResource())
 	{
 		LLM_SCOPE(ELLMTag::Textures);
@@ -207,10 +206,6 @@ void UTextureCollection::UpdateResource()
 	{
 		SetResource(nullptr);
 	}
-
-#if WITH_EDITOR
-	NotifyMaterials();
-#endif // WITH_EDITOR
 }
 
 #undef LOCTEXT_NAMESPACE
diff --git a/Engine/Source/Runtime/Engine/Private/VT/VirtualTextureAdapter.cpp b/Engine/Source/Runtime/Engine/Private/VT/VirtualTextureAdapter.cpp
--- a/Engine/Source/Runtime/Engine/Private/VT/VirtualTextureAdapter.cpp
+++ b/Engine/Source/Runtime/Engine/Private/VT/VirtualTextureAdapter.cpp
@@ -14,15 +14,48 @@
 #include "TextureResource.h"
 #include "VirtualTextureEnum.h"
 #include "VirtualTexturing.h"
-#include "VT/CopyCompressShader.h"
 #include "VT/VirtualTextureBuildSettings.h"
 
 #include UE_INLINE_GENERATED_CPP_BY_NAME(VirtualTextureAdapter)
 
-IMPLEMENT_GLOBAL_SHADER(FCopyCompressCS, "/Engine/Private/VirtualTextureAdapter.usf", "CopyCompressCS", SF_Compute);
-
 namespace VirtualTextureAdapter
 {
+	/** Copy with compression shader. */
+	class FCopyCompressCS : public FGlobalShader
+	{
+	public:
+		DECLARE_GLOBAL_SHADER(FCopyCompressCS);
+		SHADER_USE_PARAMETER_STRUCT(FCopyCompressCS, FGlobalShader);
+
+		BEGIN_SHADER_PARAMETER_STRUCT(FParameters, )
+			SHADER_PARAMETER_SRV(Texture2D, SourceTextureA)
+			SHADER_PARAMETER_RDG_TEXTURE_SRV(Texture2D, SourceTextureB)
+			SHADER_PARAMETER_SAMPLER(SamplerState, TextureSampler)
+			SHADER_PARAMETER_RDG_TEXTURE_UAV(RWTexture2D, DestTexture)
+			SHADER_PARAMETER_RDG_TEXTURE_UAV(RWTexture2D, DestCompressTexture_64bit)
+			SHADER_PARAMETER_RDG_TEXTURE_UAV(RWTexture2D, DestCompressTexture_128bit)
+			SHADER_PARAMETER(FVector2f, SourceUV)
+			SHADER_PARAMETER(FVector2f, TexelSize)
+			SHADER_PARAMETER(FVector2f, TexelOffsets)
+			SHADER_PARAMETER(FIntVector4, DestRect)
+		END_SHADER_PARAMETER_STRUCT()
+
+		static const int32 GroupSize = 8;
+
+		class FSourceTextureSelector : SHADER_PERMUTATION_BOOL("SOURCE_TEXTURE_A");
+		class FDestSrgb : SHADER_PERMUTATION_BOOL("TEXTURE_SRGB");
+		class FCompressionFormatDim : SHADER_PERMUTATION_INT("COMPRESSION_FORMAT", 7);
+		using FPermutationDomain = TShaderPermutationDomain<FSourceTextureSelector, FDestSrgb, FCompressionFormatDim>;
+
+		/** Get index to use for FCompressionFormatDim. */
+		static int32 GetCompressionPermutation(EPixelFormat InFormat)
+		{
+			const EPixelFormat CompressedFormats[] = { PF_DXT1, PF_DXT5, PF_BC4, PF_BC5, PF_BC6H, PF_BC7 };
+			return MakeArrayView(CompressedFormats).Find(InFormat) + 1;
+		}
+	};
+	IMPLEMENT_GLOBAL_SHADER(FCopyCompressCS, "/Engine/Private/VirtualTextureAdapter.usf", "CopyCompressCS", SF_Compute);
+
 	/** Final copy to output RDG parameters. */
 	BEGIN_SHADER_PARAMETER_STRUCT(FCopyToOutputParameters, )
 		RDG_TEXTURE_ACCESS(Input, ERHIAccess::CopySrc)
diff --git a/Engine/Source/Runtime/Engine/Private/VT/VirtualTextureRecreate.cpp b/Engine/Source/Runtime/Engine/Private/VT/VirtualTextureRecreate.cpp
--- a/Engine/Source/Runtime/Engine/Private/VT/VirtualTextureRecreate.cpp
+++ b/Engine/Source/Runtime/Engine/Private/VT/VirtualTextureRecreate.cpp
@@ -9,7 +9,6 @@
 #include "TextureResource.h"
 #include "UObject/UObjectIterator.h"
 #include "VirtualTexturing.h"
-#include "Engine/VirtualTextureCollection.h"
 #include "VT/RuntimeVirtualTexture.h"
 
 namespace VirtualTexture
@@ -46,11 +45,6 @@ namespace VirtualTexture
 				Owner->MarkRenderStateDirty();
 			}
 		}
-		
-		for (TObjectIterator<UVirtualTextureCollection> It(RF_ClassDefaultObject, true, EInternalObjectFlags::Garbage); It; ++It)
-		{
-			It->UpdateResource();
-		}
 	}
 
 	void Recreate(TConstArrayView< TEnumAsByte<EPixelFormat> > InFormat)
@@ -134,13 +128,5 @@ namespace VirtualTexture
 				}
 			}
 		}
-		
-		for (TObjectIterator<UVirtualTextureCollection> It(RF_ClassDefaultObject, true, EInternalObjectFlags::Garbage); It; ++It)
-		{
-			if (InFormat.Num() == 1 && It->RuntimePixelFormat == InFormat[0])
-			{
-				It->UpdateResource();
-			}
-		}
 	}
 }
diff --git a/Engine/Source/Runtime/Engine/Private/VirtualTextureCollection.cpp b/Engine/Source/Runtime/Engine/Private/VirtualTextureCollection.cpp
deleted file mode 100644
diff --git a/Engine/Source/Runtime/Engine/Public/MaterialShared.h b/Engine/Source/Runtime/Engine/Public/MaterialShared.h
--- a/Engine/Source/Runtime/Engine/Public/MaterialShared.h
+++ b/Engine/Source/Runtime/Engine/Public/MaterialShared.h
@@ -569,7 +569,7 @@ class FMaterialTextureCollectionParameterInfo
 public:
 	friend inline bool operator==(const FMaterialTextureCollectionParameterInfo& Lhs, const FMaterialTextureCollectionParameterInfo& Rhs)
 	{
-		return Lhs.ParameterInfo == Rhs.ParameterInfo && Lhs.TextureCollectionIndex == Rhs.TextureCollectionIndex && Lhs.bIsVirtualCollection == Rhs.bIsVirtualCollection;
+		return Lhs.ParameterInfo == Rhs.ParameterInfo && Lhs.TextureCollectionIndex == Rhs.TextureCollectionIndex;
 	}
 	friend inline bool operator!=(const FMaterialTextureCollectionParameterInfo& Lhs, const FMaterialTextureCollectionParameterInfo& Rhs)
 	{
@@ -581,7 +581,6 @@ public:
 
 	LAYOUT_FIELD(FHashedMaterialParameterInfo, ParameterInfo);
 	LAYOUT_FIELD(int32, TextureCollectionIndex);
-	LAYOUT_FIELD(bool, bIsVirtualCollection);
 };
 
 class FMaterialExternalTextureParameterInfo
@@ -714,7 +713,6 @@ public:
 
 	int32 GetNumTextureCollections() const { return UniformTextureCollectionParameters.Num(); }
 	ENGINE_API void GetGameThreadTextureCollectionValue(int32 Index, const UMaterialInterface* MaterialInterface, const FMaterial& Material, UTextureCollection*& OutValue) const;
-	ENGINE_API void GetTextureCollectionValue(int32 Index, const FMaterialRenderContext& Context, const FMaterial& Material, const UTextureCollection*& OutValue) const;
 
 	int32 FindOrAddTextureParameter(EMaterialTextureParameterType Type, const FMaterialTextureParameterInfo& Info);
 	int32 FindOrAddTextureCollectionParameter(const FMaterialTextureCollectionParameterInfo& Info);
@@ -748,8 +746,6 @@ public:
 	/// @param EntryGenerator Callback that writes the preshader opcodes.
 	void WriteUniformPreshaderEntry(uint32 BufferOffset, UE::Shader::EValueType EntryType, const TFunction<void(UE::Shader::FPreshaderData&)>& EntryGenerator);
 
-	void CountTextureCollections(uint32& BindlessCollections, uint32& VirtualCollections) const;
-
 protected:
 	union FVTPackedStackAndLayerIndex
 	{
diff --git a/Engine/Source/Runtime/Engine/Public/MaterialValueType.h b/Engine/Source/Runtime/Engine/Public/MaterialValueType.h
--- a/Engine/Source/Runtime/Engine/Public/MaterialValueType.h
+++ b/Engine/Source/Runtime/Engine/Public/MaterialValueType.h
@@ -76,9 +76,6 @@ enum EMaterialValueType : uint64
 	/** Material cache */
 	MCT_MaterialCacheABuffer  = 1ull << 35,
 
-	/** An internal type */
-	MCT_Unexposed             = 1ull << 36,
-
 	/** Reserved range for licensee modifications, any bits in the (inclusive) range is free for modifications */
 	MCT_LicenseeReservedBegin     = 1ull << 48,
 	MCT_LicenseeReservedEnd       = 1ull << 63,
diff --git a/Engine/Source/Runtime/RenderCore/Private/GlobalRenderResources.cpp b/Engine/Source/Runtime/RenderCore/Private/GlobalRenderResources.cpp
--- a/Engine/Source/Runtime/RenderCore/Private/GlobalRenderResources.cpp
+++ b/Engine/Source/Runtime/RenderCore/Private/GlobalRenderResources.cpp
@@ -104,33 +104,6 @@ public:
 	}
 };
 
-class FEmptyUInt4VertexBuffer : public FVertexBufferWithSRV
-{
-public:
-	virtual void InitRHI(FRHICommandListBase& RHICmdList) override
-	{
-		const FRHIBufferCreateDesc CreateDesc =
-			FRHIBufferCreateDesc::CreateVertex(TEXT("EmptyUInt4VertexBuffer"), sizeof(FUintVector4))
-			.AddUsage(BUF_Static | BUF_ShaderResource | BUF_UnorderedAccess)
-			.DetermineInitialState();
-		
-		const FUintVector4 InitialData(ForceInitToZero);
-		VertexBufferRHI = UE::RHIResourceUtils::CreateBufferWithValue(RHICmdList, CreateDesc, InitialData);
-
-		ShaderResourceViewRHI = RHICmdList.CreateShaderResourceView(
-			VertexBufferRHI,
-			FRHIViewDesc::CreateBufferSRV()
-				.SetType(FRHIViewDesc::EBufferType::Typed)
-				.SetFormat(PF_R32G32B32A32_UINT));
-
-		UnorderedAccessViewRHI = RHICmdList.CreateUnorderedAccessView(
-			VertexBufferRHI,
-			FRHIViewDesc::CreateBufferUAV()
-				.SetType(FRHIViewDesc::EBufferType::Typed)
-				.SetFormat(PF_R32G32B32A32_UINT));
-	}
-};
-
 class FEmptyStructuredBuffer : public FVertexBufferWithSRV
 {
 public:
@@ -254,7 +227,6 @@ FTexture* GBlackTexture = GBlackTextureWithSRV;
 FTexture* GTransparentBlackTexture = GTransparentBlackTextureWithSRV;
 
 FVertexBufferWithSRV* GEmptyVertexBufferWithUAV = new TGlobalResource<FEmptyVertexBuffer, FRenderResource::EInitPhase::Pre>;
-FVertexBufferWithSRV* GEmptyVertexBufferUInt4WithUAV = new TGlobalResource<FEmptyUInt4VertexBuffer, FRenderResource::EInitPhase::Pre>;
 FVertexBufferWithSRV* GEmptyStructuredBufferWithUAV = new TGlobalResource<FEmptyStructuredBuffer, FRenderResource::EInitPhase::Pre>;
 FVertexBufferWithSRV* GBlackFloat4StructuredBufferWithSRV = new TGlobalResource<FBlackFloat4StructuredBufferWithSRV, FRenderResource::EInitPhase::Pre>;
 FVertexBufferWithSRV* GBlackFloat4VertexBufferWithSRV = new TGlobalResource<FBlackFloat4VertexBuffer, FRenderResource::EInitPhase::Pre>;
@@ -767,18 +739,6 @@ void FTwoTrianglesIndexBuffer::InitRHI(FRHICommandListBase& RHICmdList)
 
 TGlobalResource<FTwoTrianglesIndexBuffer, FRenderResource::EInitPhase::Pre> GTwoTrianglesIndexBuffer;
 
-void FEmptyResourceCollection::InitRHI(FRHICommandListBase& RHICmdList)
-{
-	ResourceCollection = RHICmdList.CreateResourceCollection({});
-}
-
-void FEmptyResourceCollection::ReleaseRHI()
-{
-	ResourceCollection.SafeRelease();
-}
-
-TGlobalResource<FEmptyResourceCollection, FRenderResource::EInitPhase::Pre> GEmptyResourceCollection;
-
 /*------------------------------------------------------------------------------
 	FGlobalDynamicVertexBuffer implementation.
 ------------------------------------------------------------------------------*/
diff --git a/Engine/Source/Runtime/RenderCore/Private/ShaderParameterMetadata.cpp b/Engine/Source/Runtime/RenderCore/Private/ShaderParameterMetadata.cpp
--- a/Engine/Source/Runtime/RenderCore/Private/ShaderParameterMetadata.cpp
+++ b/Engine/Source/Runtime/RenderCore/Private/ShaderParameterMetadata.cpp
@@ -383,10 +383,6 @@ EShaderCodeResourceBindingType ParseShaderResourceBindingType(const TCHAR* Shade
 	{
 		BindingType = EShaderCodeResourceBindingType::ResourceCollection;
 	}
-	else if (!bIsRWResource && ShaderTypeLength == 23 && FCString::Strncmp(ComparedShaderType, TEXT("FIndirectVirtualTexture"), ShaderTypeLength) == 0)
-	{
-		BindingType = EShaderCodeResourceBindingType::Buffer;
-	}
 	else if (bIsRasterizerOrderedResource)
 	{
 		BindingType = EShaderCodeResourceBindingType::RasterizerOrderedTexture2D;
@@ -781,7 +777,7 @@ void FShaderParametersMetadata::InitializeLayout(FRHIUniformBufferLayoutInitiali
 				}
 				else if (BindingType == EShaderCodeResourceBindingType::ResourceCollection)
 				{
-					bIsValidBindingType = BaseType == UBMT_RESOURCE_COLLECTION;
+					bIsValidBindingType = (BaseType == UBMT_RESOURCE_COLLECTION);
 				}
 				else
 				{
diff --git a/Engine/Source/Runtime/RenderCore/Public/GlobalRenderResources.h b/Engine/Source/Runtime/RenderCore/Public/GlobalRenderResources.h
--- a/Engine/Source/Runtime/RenderCore/Public/GlobalRenderResources.h
+++ b/Engine/Source/Runtime/RenderCore/Public/GlobalRenderResources.h
@@ -28,9 +28,6 @@ extern RENDERCORE_API FTexture* GDummyTransitionTexture;
 // An empty vertex buffer with a UAV
 extern RENDERCORE_API FVertexBufferWithSRV* GEmptyVertexBufferWithUAV;
 
-// An empty uint4 vertex buffer with a UAV
-extern RENDERCORE_API FVertexBufferWithSRV* GEmptyVertexBufferUInt4WithUAV;
-
 // An empty structured buffer with a UAV
 extern RENDERCORE_API FVertexBufferWithSRV* GEmptyStructuredBufferWithUAV;
 
@@ -167,19 +164,6 @@ public:
 };
 extern RENDERCORE_API TGlobalResource<FTwoTrianglesIndexBuffer, FRenderResource::EInitPhase::Pre> GTwoTrianglesIndexBuffer;
 
-class FEmptyResourceCollection : public FRenderResource
-{
-public:
-	FEmptyResourceCollection() = default;
-	~FEmptyResourceCollection() = default;
-
-	RENDERCORE_API virtual void InitRHI(FRHICommandListBase& RHICmdList) override;
-	RENDERCORE_API virtual void ReleaseRHI() override;
-
-	FRHIResourceCollectionRef ResourceCollection;
-};
-
-extern RENDERCORE_API TGlobalResource<FEmptyResourceCollection, FRenderResource::EInitPhase::Pre> GEmptyResourceCollection;
 
 //////////////////////////////////////////////////////////////////////////////////////////////////////////////////////
 // FGlobalDynamicVertexBuffer
diff --git a/Engine/Source/Runtime/RenderCore/Public/RendererInterface.h b/Engine/Source/Runtime/RenderCore/Public/RendererInterface.h
--- a/Engine/Source/Runtime/RenderCore/Public/RendererInterface.h
+++ b/Engine/Source/Runtime/RenderCore/Public/RendererInterface.h
@@ -820,8 +820,6 @@ public:
 	virtual uint32 RemoveAllVirtualTextureProducerDestroyedCallbacks(const void* Baton) = 0;
 	virtual void ReleaseVirtualTexturePendingResources() = 0;
 
-	virtual IVirtualTexture* FindProducer(const FVirtualTextureProducerHandle& Handle) = 0;
-
 	virtual void RequestVirtualTextureTiles(const FVector2D& InScreenSpaceSize, int32 InMipLevel) = 0;
 	virtual void RequestVirtualTextureTiles(const FMaterialRenderProxy* InMaterialRenderProxy, const FVector2D& InScreenSpaceSize, ERHIFeatureLevel::Type InFeatureLevel) = 0;
 
diff --git a/Engine/Source/Runtime/Renderer/Private/RendererModule.h b/Engine/Source/Runtime/Renderer/Private/RendererModule.h
--- a/Engine/Source/Runtime/Renderer/Private/RendererModule.h
+++ b/Engine/Source/Runtime/Renderer/Private/RendererModule.h
@@ -112,7 +112,6 @@ public:
 	virtual void RequestVirtualTextureTiles(const FVector2D& InScreenSpaceSize, int32 InMipLevel) override;
 	virtual void RequestVirtualTextureTiles(const FMaterialRenderProxy* InMaterialRenderProxy, const FVector2D& InScreenSpaceSize, ERHIFeatureLevel::Type InFeatureLevel) override;
 	virtual void RequestVirtualTextureTiles(IAllocatedVirtualTexture* AllocatedVT, const FVector2D& InScreenSpaceSize, const FVector2D& InViewportPosition, const FVector2D& InViewportSize, const FVector2D& InUV0, const FVector2D& InUV1, int32 InMipLevel) override;
-	virtual IVirtualTexture* FindProducer(const FVirtualTextureProducerHandle& Handle) override;
 
 	UE_DEPRECATED(5.4, "Use RequestVirtualTextureTiles() overloads that takes similar parameters. Make sure not to negate the InViewportPosition.")
 	virtual void RequestVirtualTextureTilesForRegion(IAllocatedVirtualTexture* AllocatedVT, const FVector2D& InScreenSpaceSize, const FVector2D& InViewportPosition, const FVector2D& InViewportSize, const FVector2D& InUV0, const FVector2D& InUV1, int32 InMipLevel) override;
diff --git a/Engine/Source/Runtime/Renderer/Private/SceneRendering.cpp b/Engine/Source/Runtime/Renderer/Private/SceneRendering.cpp
--- a/Engine/Source/Runtime/Renderer/Private/SceneRendering.cpp
+++ b/Engine/Source/Runtime/Renderer/Private/SceneRendering.cpp
@@ -5521,16 +5521,6 @@ void FRendererModule::RequestVirtualTextureTiles(IAllocatedVirtualTexture* Alloc
 	FVirtualTextureSystem::Get().RequestTiles(AllocatedVT, InScreenSpaceSize, InViewportPosition, InViewportSize, InUV0, InUV1, InMipLevel);
 }
 
-IVirtualTexture* FRendererModule::FindProducer(const FVirtualTextureProducerHandle& Handle)
-{
-	if (FVirtualTextureProducer* Producer = FVirtualTextureSystem::Get().FindProducer(Handle))
-	{
-		return Producer->GetVirtualTexture();
-	}
-
-	return nullptr;
-}
-
 void FRendererModule::RequestVirtualTextureTilesForRegion(IAllocatedVirtualTexture* AllocatedVT, const FVector2D& InScreenSpaceSize, const FVector2D& InViewportPosition, const FVector2D& InViewportSize, const FVector2D& InUV0, const FVector2D& InUV1, int32 InMipLevel)
 {
 PRAGMA_DISABLE_DEPRECATION_WARNINGS
diff --git a/Engine/Source/Runtime/Renderer/Private/VT/VirtualTextureSystem.h b/Engine/Source/Runtime/Renderer/Private/VT/VirtualTextureSystem.h
--- a/Engine/Source/Runtime/Renderer/Private/VT/VirtualTextureSystem.h
+++ b/Engine/Source/Runtime/Renderer/Private/VT/VirtualTextureSystem.h
@@ -182,7 +182,6 @@ private:
 	friend class FVirtualTextureProducerCollection;
 	friend class FTexturePageMap;
 	friend class FTexturePagePool;
-	friend class FRendererModule;
 
 	FVirtualTextureSystem();
 	~FVirtualTextureSystem();
