Changed UInstancedSkinnedMeshComponent::GetComponentDesc to use ShaderPlatform instead of Feature level to fix a cook time check.

#tests Nanite test cook and editor
#rnx

diff --git a/Engine/Source/Runtime/Engine/Classes/Components/InstancedSkinnedMeshComponent.h b/Engine/Source/Runtime/Engine/Classes/Components/InstancedSkinnedMeshComponent.h
--- a/Engine/Source/Runtime/Engine/Classes/Components/InstancedSkinnedMeshComponent.h
+++ b/Engine/Source/Runtime/Engine/Classes/Components/InstancedSkinnedMeshComponent.h
@@ -238,7 +238,7 @@ protected:
 	/** Applies the cached component instance data to a newly blueprint constructed component. */
 	void ApplyComponentInstanceData(struct FInstancedSkinnedMeshComponentInstanceData* ComponentInstanceData);
 
-	FInstanceDataManagerSourceDataDesc GetComponentDesc(ERHIFeatureLevel::Type FeatureLevel);
+	FInstanceDataManagerSourceDataDesc GetComponentDesc(EShaderPlatform InShaderPlatform);
 
 #if WITH_EDITOR
 	/** One bit per instance if the instance is selected. */
diff --git a/Engine/Source/Runtime/Engine/Private/Components/InstancedSkinnedMeshComponent.cpp b/Engine/Source/Runtime/Engine/Private/Components/InstancedSkinnedMeshComponent.cpp
--- a/Engine/Source/Runtime/Engine/Private/Components/InstancedSkinnedMeshComponent.cpp
+++ b/Engine/Source/Runtime/Engine/Private/Components/InstancedSkinnedMeshComponent.cpp
@@ -212,7 +212,7 @@ void UInstancedSkinnedMeshComponent::Serialize(FArchive& Ar)
 #if WITH_EDITOR
 		else if (Ar.IsSaving())
 		{
-			InstanceDataManager.WriteCookedRenderData(Ar, GetComponentDesc(GMaxRHIFeatureLevel));
+			InstanceDataManager.WriteCookedRenderData(Ar, GetComponentDesc(GMaxRHIShaderPlatform));
 		}
 #endif
 	}
@@ -288,7 +288,7 @@ inline void ReorderArray(ArrayType& InOutDataArray, const TArray<int32>& OldInde
 void UInstancedSkinnedMeshComponent::OptimizeInstanceData(bool bShouldRetainIdMap)
 {
 	// compute the optimal order 
-	TArray<int32> IndexRemap = InstanceDataManager.Optimize(GetComponentDesc(GMaxRHIFeatureLevel), bShouldRetainIdMap);
+	TArray<int32> IndexRemap = InstanceDataManager.Optimize(GetComponentDesc(GMaxRHIShaderPlatform), bShouldRetainIdMap);
 	
 	if (!IndexRemap.IsEmpty())
 	{
@@ -358,11 +358,11 @@ void UInstancedSkinnedMeshComponent::PostAssetCompilation()
 
 #endif 
 
-FInstanceDataManagerSourceDataDesc UInstancedSkinnedMeshComponent::GetComponentDesc(ERHIFeatureLevel::Type FeatureLevel)
+FInstanceDataManagerSourceDataDesc UInstancedSkinnedMeshComponent::GetComponentDesc(EShaderPlatform InShaderPlatform)
 {
 	FInstanceDataManagerSourceDataDesc ComponentDesc;
 
-	ComponentDesc.PrimitiveMaterialDesc = GetUsedMaterialPropertyDesc(GetFeatureLevelShaderPlatform_Checked(FeatureLevel));
+	ComponentDesc.PrimitiveMaterialDesc = GetUsedMaterialPropertyDesc(InShaderPlatform);
 
 	FInstanceDataFlags Flags;
 	Flags.bHasPerInstanceRandom = ComponentDesc.PrimitiveMaterialDesc.bAnyMaterialHasPerInstanceRandom;
@@ -462,7 +462,7 @@ void UInstancedSkinnedMeshComponent::SendRenderInstanceData_Concurrent()
 		if (SceneProxy != nullptr)
 		{
 			// Make sure the instance data proxy is up to date:
-			if (InstanceDataManager.FlushChanges(GetComponentDesc(SceneProxy->GetScene().GetFeatureLevel())))
+			if (InstanceDataManager.FlushChanges(GetComponentDesc(SceneProxy->GetScene().GetShaderPlatform())))
 			{
 				UpdateBounds();
 				GetWorld()->Scene->UpdatePrimitiveInstances(this);
@@ -590,7 +590,7 @@ FPrimitiveSceneProxy* UInstancedSkinnedMeshComponent::CreateSceneProxy()
 	// If GPU-only flag set, instance data is entirely GPU driven, don't upload from CPU.
 	if (Result && !bIsInstanceDataGPUOnly)
 	{
-		InstanceDataManager.FlushChanges(GetComponentDesc(Result->GetScene().GetFeatureLevel()));
+		InstanceDataManager.FlushChanges(GetComponentDesc(Result->GetScene().GetShaderPlatform()));
 	}
 
 #if !(UE_BUILD_SHIPPING || UE_BUILD_TEST)
