Minimize ShaderMapKeyString allocation.
#rnx

diff --git a/Engine/Source/Runtime/Engine/Private/Materials/MaterialShader.cpp b/Engine/Source/Runtime/Engine/Private/Materials/MaterialShader.cpp
--- a/Engine/Source/Runtime/Engine/Private/Materials/MaterialShader.cpp
+++ b/Engine/Source/Runtime/Engine/Private/Materials/MaterialShader.cpp
@@ -218,6 +218,7 @@ FString GetMaterialShaderMapKeyString(const FMaterialShaderMapId& ShaderMapId,
 	FMaterialKeyGeneratorContext Context(ShaderMapKeyString, Platform);
 	Context.SetFlags(EMaterialKeyInclude::ShaderDependencies, bIncludeKeyStringShaderDependencies);
 	RecordOrEmitMaterialShaderMapKey(Context, ShaderMapId, ShaderParameters);
+	ShaderMapKeyString.Shrink();
 
 	return ShaderMapKeyString;
 }
