[First Person] Fixed incorrect initialization of PrevTanAndInvTanHalfFOV and PrevFirstPersonTanAndInvTanHalfFOV view UB members.

diff --git a/Engine/Source/Runtime/Engine/Private/SceneView.cpp b/Engine/Source/Runtime/Engine/Private/SceneView.cpp
--- a/Engine/Source/Runtime/Engine/Private/SceneView.cpp
+++ b/Engine/Source/Runtime/Engine/Private/SceneView.cpp
@@ -2596,7 +2596,7 @@ void SetupViewRectUniformBufferParameters(
 		ViewUniformShaderParameters.TanAndInvTanHalfFOV = TanAndInvTanFOV;
 		ViewUniformShaderParameters.PrevTanAndInvTanHalfFOV = InPrevViewMatrices.GetTanAndInvTanHalfFOV();
 		ViewUniformShaderParameters.FirstPersonTanAndInvTanHalfFOV = InViewMatrices.GetFirstPersonTanAndInvTanHalfFOV();
-		ViewUniformShaderParameters.PrevTanAndInvTanHalfFOV = InPrevViewMatrices.GetFirstPersonTanAndInvTanHalfFOV();
+		ViewUniformShaderParameters.PrevFirstPersonTanAndInvTanHalfFOV = InPrevViewMatrices.GetFirstPersonTanAndInvTanHalfFOV();
 
 		if (InViewMatrices.IsPerspectiveProjection())
 		{
