World Partition - Cook Package Splitter: Added option to allow External Data Layer packages to no longer inherit chunk assignments from their generator package
* Added World Partition setting: "External Data Layers packages inherit World chunk assignments" (on by default to retain prior behavior)
* When inheritance is disabled, EDL packages will rely on their own UAssetManager::GetPackageChunkIds assignment instead of blindly inheriting from their generator.
* Non-EDL packages will be unaffected and continue to inherit chunk assignments as before.
* Uses bCopyChunkAssignmentFromGenerator = false, which was introduced specifically for this use case in 44755683

diff --git a/Engine/Config/BaseEngine.ini b/Engine/Config/BaseEngine.ini
--- a/Engine/Config/BaseEngine.ini
+++ b/Engine/Config/BaseEngine.ini
@@ -3934,6 +3934,7 @@ SendOSData=true
 EditorHashDefaultClass=/Script/CoreUObject.Class'/Script/Engine.WorldPartitionEditorSpatialHash'
 RuntimeHashDefaultClass=/Script/CoreUObject.Class'/Script/Engine.WorldPartitionRuntimeHashSet'
 UnsupportedActorClassesPerHLODLayerType=((Instancing, (ActorClasses=("/Script/CoreUObject.Class'/Script/Engine.SplineMeshActor'"))))
+bShouldEDLPackagesInheritWorldChunkAssignmentsDuringCook=true
 
 [/Script/Engine.WorldPartitionRuntimeCellTransformerSettings]
 +IgnoredComponentClasses=/Script/CoreUObject.Class'/Script/Engine.LevelInstanceComponent'
diff --git a/Engine/Source/Runtime/Engine/Private/WorldPartition/WorldPartitionCookPackageSplitter.cpp b/Engine/Source/Runtime/Engine/Private/WorldPartition/WorldPartitionCookPackageSplitter.cpp
--- a/Engine/Source/Runtime/Engine/Private/WorldPartition/WorldPartitionCookPackageSplitter.cpp
+++ b/Engine/Source/Runtime/Engine/Private/WorldPartition/WorldPartitionCookPackageSplitter.cpp
@@ -10,6 +10,7 @@
 #include "WorldPartition/WorldPartition.h"
 #include "WorldPartition/WorldPartitionLog.h"
 #include "WorldPartition/WorldPartitionRuntimeCell.h"
+#include "WorldPartition/WorldPartitionSettings.h"
 
 TArray<ICookPackageSplitter::FGeneratedPackage> FWorldPartitionCookPackageSplitter::GetGenerateList(const UPackage* OwnerPackage)
 {
@@ -109,6 +110,8 @@ void FWorldPartitionCookPackageSplitter::Teardown(ICookPackageSplitter::ETeardow
 
 void FWorldPartitionCookPackageSplitter::BuildPackagesToGenerateList(TArray<ICookPackageSplitter::FGeneratedPackage>& PackagesToGenerate) const
 {
+	const bool bEDLCopyChunkAssignmentFromGenerator = UWorldPartitionSettings::Get()->ShouldEDLPackagesInheritWorldChunkAssignmentsDuringCook();
+
 	for (const IWorldPartitionCookPackageGenerator* CookPackageGenerator : CookContext.GetCookPackageGenerators())
 	{
 		if (const TArray<FWorldPartitionCookPackage*>* CookPackages = CookContext.GetCookPackages(CookPackageGenerator))
@@ -124,11 +127,16 @@ void FWorldPartitionCookPackageSplitter::BuildPackagesToGenerateList(TArray<ICoo
 
 				UE_LOG(LogWorldPartition, Log, TEXT("Adding Cell %s with GenerationHash: %s to packages to generate"), *CookPackage->GetFullGeneratedPath(), *LexToString(CookPackage->GenerationHash))
 
-				CookPackage->Type == FWorldPartitionCookPackage::EType::Level ? GeneratedPackage.SetCreateAsMap(true) : GeneratedPackage.SetCreateAsMap(false);
+				GeneratedPackage.SetCreateAsMap(CookPackage->Type == FWorldPartitionCookPackage::EType::Level);
 
 				// Fill generated package dependencies for iterative cooking
 				if (UWorldPartitionRuntimeCell* Cell = CookPackageGenerator->GetCellForPackage(*CookPackage))
 				{
+					check(CookPackage->Type == FWorldPartitionCookPackage::EType::Level);
+
+					// If this package belongs to an External Data Layer, do NOT inherit the generator's chunk assignment if this option is enabled for the project
+					GeneratedPackage.bCopyChunkAssignmentFromGenerator = bEDLCopyChunkAssignmentFromGenerator || Cell->GetExternalDataLayer().IsNone();
+
 					TSet<FName> ActorPackageNames = Cell->GetActorPackageNames();
 					GeneratedPackage.PackageDependencies.Reset(ActorPackageNames.Num());
 					for (FName ActorPackageName : ActorPackageNames)
@@ -140,6 +148,11 @@ void FWorldPartitionCookPackageSplitter::BuildPackagesToGenerateList(TArray<ICoo
 							UE::AssetRegistry::EDependencyProperty::Hard | UE::AssetRegistry::EDependencyProperty::Build));
 					}
 				}
+				else
+				{
+					// Copy chunk assignment only for level packages.
+					GeneratedPackage.bCopyChunkAssignmentFromGenerator = bEDLCopyChunkAssignmentFromGenerator || CookPackage->Type == FWorldPartitionCookPackage::EType::Level;
+				}
 			}
 		}
 	}
diff --git a/Engine/Source/Runtime/Engine/Public/WorldPartition/WorldPartitionSettings.h b/Engine/Source/Runtime/Engine/Public/WorldPartition/WorldPartitionSettings.h
--- a/Engine/Source/Runtime/Engine/Public/WorldPartition/WorldPartitionSettings.h
+++ b/Engine/Source/Runtime/Engine/Public/WorldPartition/WorldPartitionSettings.h
@@ -35,6 +35,7 @@ public:
 	
 	TSubclassOf<UWorldPartitionEditorHash> GetEditorHashDefaultClass() const { return EditorHashDefaultClass; }
 	TSubclassOf<UWorldPartitionRuntimeHash> GetRuntimeHashDefaultClass() const { return RuntimeHashDefaultClass; }
+	bool ShouldEDLPackagesInheritWorldChunkAssignmentsDuringCook() const { return bShouldEDLPackagesInheritWorldChunkAssignmentsDuringCook;	}
 
 	UWorldPartitionPropertyOverridePolicy* GetPropertyOverridePolicy() const
 	{
@@ -78,6 +79,9 @@ protected:
 
 	UPROPERTY(Config)
 	TMap<EHLODLayerType, FHLODLayerTypeUnsupportedActorClasses> UnsupportedActorClassesPerHLODLayerType;
+
+	UPROPERTY(EditAnywhere, Config, Category = WorldPartition, meta=(DisplayName="Cook - External Data Layers packages inherit World chunk assignments"))
+	bool bShouldEDLPackagesInheritWorldChunkAssignmentsDuringCook = true;
 };
 
 #undef UE_API
