[MaterialExpression]
- Replace FApp::IsGame() which can dynamically change when starting or exiting PIE to a more viable and deterministic per-object approach

diff --git a/Engine/Source/Runtime/Engine/Private/Materials/MaterialExpressions.cpp b/Engine/Source/Runtime/Engine/Private/Materials/MaterialExpressions.cpp
--- a/Engine/Source/Runtime/Engine/Private/Materials/MaterialExpressions.cpp
+++ b/Engine/Source/Runtime/Engine/Private/Materials/MaterialExpressions.cpp
@@ -2076,7 +2076,7 @@ void UMaterialExpression::ConnectExpression( FExpressionInput* Input, int32 Outp
 void UMaterialExpression::UpdateMaterialExpressionGuid(bool bForceGeneration, bool bAllowMarkingPackageDirty)
 {
 	// If we are in the editor, and we don't have a valid GUID yet, generate one.
-	if (GIsEditor && !FApp::IsGame())
+	if (GIsEditor && !IsRunningGame() && !GetPackage()->HasAnyPackageFlags(PKG_PlayInEditor))
 	{
 		FGuid& Guid = GetMaterialExpressionId();
 
