[Shaders] Dump all pending compile jobs associated with a hung shadermap before we crash on that event.

This also logs all dumped jobs with their individual time since they have been assigned to a worker (here at the example of a low timeout for testing):
LogShaderCompilers: Job [WorldGridMaterial_fa3686479b6322ac/Default/TGPUSkinAPEXClothVertexFactoryUnlimited/FDebugViewModeVS/0:/Engine/Private/DebugViewModeVertexShader.usf|Main] pending for 11.6 seconds

#rnx

diff --git a/Engine/Source/Runtime/Engine/Internal/ShaderCompiler/ShaderCompilerInternal.cpp b/Engine/Source/Runtime/Engine/Internal/ShaderCompiler/ShaderCompilerInternal.cpp
--- a/Engine/Source/Runtime/Engine/Internal/ShaderCompiler/ShaderCompilerInternal.cpp
+++ b/Engine/Source/Runtime/Engine/Internal/ShaderCompiler/ShaderCompilerInternal.cpp
@@ -27,3 +27,10 @@ void FShaderCompileInternalUtilities::EnableDumpDebugInfoForRetry(FShaderCommonC
 		}
 	);
 }
+
+void FShaderCompileInternalUtilities::DumpDebugInfo(FShaderCommonCompileJob& Job)
+{
+	FShaderCompileInternalUtilities::EnableDumpDebugInfoForRetry(Job);
+	FShaderDebugDataContext Ctx;
+	Job.OnComplete(Ctx);
+}
diff --git a/Engine/Source/Runtime/Engine/Internal/ShaderCompiler/ShaderCompilerInternal.h b/Engine/Source/Runtime/Engine/Internal/ShaderCompiler/ShaderCompilerInternal.h
--- a/Engine/Source/Runtime/Engine/Internal/ShaderCompiler/ShaderCompilerInternal.h
+++ b/Engine/Source/Runtime/Engine/Internal/ShaderCompiler/ShaderCompilerInternal.h
@@ -19,5 +19,8 @@ public:
 
 	/** Ensures DumpDebugInfoPath is assigned and tracked by the shader compiling manager. */
 	static ENGINE_API void EnableDumpDebugInfoForRetry(FShaderCommonCompileJob& Job);
+
+	/** Explicitly dumps debug information for this shader compile job. Call this if debug info must be dumped earlier than the regular compile job completion, e.g. right before a Fatal error. */
+	static ENGINE_API void DumpDebugInfo(FShaderCommonCompileJob& Job);
 };
 
diff --git a/Engine/Source/Runtime/Engine/Private/ShaderCompiler/ShaderCompiler.cpp b/Engine/Source/Runtime/Engine/Private/ShaderCompiler/ShaderCompiler.cpp
--- a/Engine/Source/Runtime/Engine/Private/ShaderCompiler/ShaderCompiler.cpp
+++ b/Engine/Source/Runtime/Engine/Private/ShaderCompiler/ShaderCompiler.cpp
@@ -2603,6 +2603,19 @@ void FShaderMapCompileResults::CheckIfHung()
 
 			if (GCrashOnHungShaderMaps)
 			{
+				// Dump all pending compile jobs that are assigned to this shader map before we crash with the following Fatal error
+				GShaderCompilingManager->DumpDebugInfoForEachPendingJob([this](FShaderCommonCompileJob* Job) -> bool {
+					if (Job->PendingShaderMap.GetReference() == this && !Job->bSucceeded)
+					{
+						Job->ForEachSingleShaderJob([](FShaderCompileJob& SingleJob) -> void {
+							const double JobDuration = FPlatformTime::Seconds() - SingleJob.TimeAssignedToExecution;
+							UE_LOG(LogShaderCompilers, Log, TEXT("Job [%s] pending for %.1f seconds"), *SingleJob.Input.GenerateShaderName(), JobDuration);
+						});
+						return true;
+					}
+					return false;
+				});
+
 				UE_LOG(LogShaderCompilers, Fatal, TEXT("Crashing on a hung shadermap, time spent compiling: %f seconds, NumPendingJobs: %d, FinishedJobs: %d"),
 					DurationSoFar,
 					NumPendingJobs.GetValue(),
@@ -2749,6 +2762,24 @@ void FShaderCompilingManager::ProcessAsyncResults(bool bLimitExecutionTime, bool
 	ProcessAsyncResults(TimeSlice, bBlockOnGlobalShaderCompletion);
 }
 
+void FShaderCompilingManager::DumpDebugInfoForEachPendingJob(const FShaderCompileJobCallback& PendingJobCallback) const
+{
+	check(PendingJobCallback);
+
+	for (const TUniquePtr<FShaderCompileThreadRunnableBase>& ShaderCompileThread : Threads)
+	{
+		ShaderCompileThread->ForEachPendingJob([this, &PendingJobCallback](FShaderCommonCompileJob* Job) -> bool {
+			check(Job);
+			if (PendingJobCallback(Job))
+			{
+				// Dump debug information for current job
+				FShaderCompileInternalUtilities::DumpDebugInfo(*Job);
+			}
+			return true; // Always perform full iteration, don't exit early
+		});
+	}
+}
+
 void FShaderCompilingManager::ProcessAsyncResults(float TimeSlice, bool bBlockOnGlobalShaderCompletion)
 {
 	LLM_SCOPE_BYTAG(ShaderCompiler);
diff --git a/Engine/Source/Runtime/Engine/Private/ShaderCompiler/ShaderCompilerDistributed.cpp b/Engine/Source/Runtime/Engine/Private/ShaderCompiler/ShaderCompilerDistributed.cpp
--- a/Engine/Source/Runtime/Engine/Private/ShaderCompiler/ShaderCompilerDistributed.cpp
+++ b/Engine/Source/Runtime/Engine/Private/ShaderCompiler/ShaderCompilerDistributed.cpp
@@ -560,9 +560,7 @@ int32 FShaderCompileDistributedThreadRunnable_Interface::CompilingLoop()
 					UE_LOG(LogShaderCompilers, Display, TEXT("Executing %s"), *JobDiagnostics);
 
 					// Dumping debug info in case the compile job crashes the cooker
-					FShaderCompileInternalUtilities::EnableDumpDebugInfoForRetry(*Job);
-					FShaderDebugDataContext Ctx;
-					Job->OnComplete(Ctx);
+					FShaderCompileInternalUtilities::DumpDebugInfo(*Job);
 
 					FShaderCompileUtilities::ExecuteShaderCompileJob(*Job);
 				}
@@ -616,6 +614,21 @@ int32 FShaderCompileDistributedThreadRunnable_Interface::CompilingLoop()
 	return Manager->AllJobs.GetNumOutstandingJobs() > 0;
 }
 
+void FShaderCompileDistributedThreadRunnable_Interface::ForEachPendingJob(const FShaderCompileJobCallback& PendingJobCallback) const
+{
+	for (auto Iter = DispatchedTasks.CreateConstIterator(); Iter; ++Iter)
+	{
+		const FDistributedShaderCompilerTask* Task = *Iter;
+		for (const FShaderCommonCompileJobPtr& Job : Task->ShaderJobs)
+		{
+			if (!PendingJobCallback(Job.GetReference()))
+			{
+				return;
+			}
+		}
+	}
+}
+
 const TCHAR* FShaderCompileDistributedThreadRunnable_Interface::GetThreadName() const
 {
 	return TEXT("ShaderCompilingThread-Distributed");
diff --git a/Engine/Source/Runtime/Engine/Private/ShaderCompiler/ShaderCompilerThreadRunnable.cpp b/Engine/Source/Runtime/Engine/Private/ShaderCompiler/ShaderCompilerThreadRunnable.cpp
--- a/Engine/Source/Runtime/Engine/Private/ShaderCompiler/ShaderCompilerThreadRunnable.cpp
+++ b/Engine/Source/Runtime/Engine/Private/ShaderCompiler/ShaderCompilerThreadRunnable.cpp
@@ -295,6 +295,21 @@ void FShaderCompileThreadRunnable::OnMachineResourcesChanged()
 	}
 }
 
+void FShaderCompileThreadRunnable::ForEachPendingJob(const FShaderCompileJobCallback& PendingJobCallback) const
+{
+	FScopeLock WorkerScopeLock(&WorkerInfosLock);
+	for (const TUniquePtr<FShaderCompileWorkerInfo>& WorkerInfo : WorkerInfos)
+	{
+		for (const FShaderCommonCompileJobPtr& Job : WorkerInfo->QueuedJobs)
+		{
+			if (!PendingJobCallback(Job.GetReference()))
+			{
+				return;
+			}
+		}
+	}
+}
+
 /** Entry point for the shader compiling thread. */
 uint32 FShaderCompileThreadRunnableBase::Run()
 {
diff --git a/Engine/Source/Runtime/Engine/Public/ShaderCompiler.h b/Engine/Source/Runtime/Engine/Public/ShaderCompiler.h
--- a/Engine/Source/Runtime/Engine/Public/ShaderCompiler.h
+++ b/Engine/Source/Runtime/Engine/Public/ShaderCompiler.h
@@ -95,6 +95,7 @@ public:
 
 	int32 GetNumPendingJobs() const;
 
+	/** Despite being named a getter, this is drawing entries out of the list of pending compile jobs, i.e. the returned jobs will be removed from the list of pending jobs. */
 	int32 GetPendingJobs(EShaderCompilerWorkerType InWorkerType, EShaderCompileJobPriority InPriority, int32 MinNumJobs, int32 MaxNumJobs, TArray<FShaderCommonCompileJobPtr>& OutJobs);
 
 private:
@@ -139,6 +140,9 @@ struct FShaderCompileMemoryUsage
 	uint64 PhysicalMemory;
 };
 
+/** Function signature to iterate over shader compile jobs (single and pipeline jobs). */
+using FShaderCompileJobCallback = TFunction<bool(FShaderCommonCompileJob*)>;
+
 class FShaderCompileThreadRunnableBase : public FRunnable
 {
 	friend class FShaderCompilingManager;
@@ -209,6 +213,14 @@ public:
 
 	/** Events from the manager */
 	virtual void OnMachineResourcesChanged() {}
+
+	/**
+	 * Invokes the specified callback on each pending job this shader compile thread is processing.
+	 * Use this to find compile jobs that are already sent to a worker, since once a worker processes them,
+	 * they are no longer in the "pending list" of the shader compile manager (i.e. FShaderCompileJobCollection specifically).
+	 * @param PendingJobCallback Function callback that is invoked on each pending job. Returns false to exit the iteration early.
+	 */
+	virtual void ForEachPendingJob(const FShaderCompileJobCallback& PendingJobCallback) const {}
 };
 
 /** 
@@ -287,6 +299,8 @@ private:
 
 	virtual void OnMachineResourcesChanged() override;
 
+	virtual void ForEachPendingJob(const FShaderCompileJobCallback& PendingJobCallback) const override;
+
 	void PrintWorkerMemoryUsageWithLockTaken();
 
 	/** Returns the total number of workers this thread is handling. */
@@ -379,6 +393,8 @@ public:
 
 	static bool IsSupported();
 
+	virtual void ForEachPendingJob(const FShaderCompileJobCallback& PendingJobCallback) const override;
+
 protected:
 	
 	IDistributedBuildController& CachedController;
@@ -1268,6 +1284,9 @@ public:
 	/** Version of ProcessAsyncResults that specifies use of ProcessGameThreadTargetTime for the timeslice. */
 	ENGINE_API void ProcessAsyncResults(bool bLimitExecutionTime, bool bBlockOnGlobalShaderCompletion);
 
+	/** Dumps the debug information for each pending job for which PendingJobCallback returns true. */
+	ENGINE_API void DumpDebugInfoForEachPendingJob(const FShaderCompileJobCallback& PendingJobCallback) const;
+
 	/**
 	 * Returns true if the given shader compile worker is still running.
 	 */
