[HLOD] Custom HLOD support
* Added WorldPartitionCustomHLODActor, an actor you can place in the world to provide a custom HLOD representation (using static mesh component).
* Added a new HLOD Layer type: "Custom HLOD Actor".
* Custom HLOD actors assigned to "Custom HLOD Actor" layer are injected as-is into the HLOD runtime partition.
* The "Custom HLOD Actor" layer can specify a parent layer. In that case custom HLOD actors are also used to generate parent layer's HLOD representation.
* Custom HLOD actors can also be assigned to other (non-Custom) HLOD Layers. In that case they are used only during HLOD generation and are not pulled into the HLOD partition themselves.

diff --git a/Engine/Plugins/Editor/WorldPartitionHLODUtilities/Source/Private/WorldPartition/HLOD/Builders/HLODBuilderCustomHLODActor.cpp b/Engine/Plugins/Editor/WorldPartitionHLODUtilities/Source/Private/WorldPartition/HLOD/Builders/HLODBuilderCustomHLODActor.cpp
new file mode 100644
--- /dev/null
+++ b/Engine/Plugins/Editor/WorldPartitionHLODUtilities/Source/Private/WorldPartition/HLOD/Builders/HLODBuilderCustomHLODActor.cpp
@@ -0,0 +1,25 @@
+// Copyright Epic Games, Inc. All Rights Reserved.
+
+#include "WorldPartition/HLOD/Builders/HLODBuilderCustomHLODActor.h"
+
+#include UE_INLINE_GENERATED_CPP_BY_NAME(HLODBuilderCustomHLODActor)
+
+UHLODBuilderCustomHLODActorSettings::UHLODBuilderCustomHLODActorSettings(const FObjectInitializer& ObjectInitializer)
+	: Super(ObjectInitializer)
+{
+}
+
+UHLODBuilderCustomHLODActor::UHLODBuilderCustomHLODActor(const FObjectInitializer& ObjectInitializer)
+	: Super(ObjectInitializer)
+{
+}
+
+TSubclassOf<UHLODBuilderSettings> UHLODBuilderCustomHLODActor::GetSettingsClass() const
+{
+	return UHLODBuilderCustomHLODActorSettings::StaticClass();
+}
+
+TArray<UActorComponent*> UHLODBuilderCustomHLODActor::Build(const FHLODBuildContext& InHLODBuildContext, const TArray<UActorComponent*>& InSourceComponents) const
+{
+	return {};
+}
\ No newline at end of file
diff --git a/Engine/Plugins/Editor/WorldPartitionHLODUtilities/Source/Private/WorldPartition/HLOD/Utilities/WorldPartitionHLODUtilities.cpp b/Engine/Plugins/Editor/WorldPartitionHLODUtilities/Source/Private/WorldPartition/HLOD/Utilities/WorldPartitionHLODUtilities.cpp
--- a/Engine/Plugins/Editor/WorldPartitionHLODUtilities/Source/Private/WorldPartition/HLOD/Utilities/WorldPartitionHLODUtilities.cpp
+++ b/Engine/Plugins/Editor/WorldPartitionHLODUtilities/Source/Private/WorldPartition/HLOD/Utilities/WorldPartitionHLODUtilities.cpp
@@ -15,6 +15,7 @@
 #include "WorldPartition/HLOD/HLODSourceActorsFromCell.h"
 #include "WorldPartition/HLOD/HLODStats.h"
 #include "WorldPartition/HLOD/HLODInstancedStaticMeshComponent.h"
+#include "WorldPartition/HLOD/Builders/HLODBuilderCustomHLODActor.h"
 #include "WorldPartition/HLOD/Builders/HLODBuilderInstancing.h"
 #include "WorldPartition/HLOD/Builders/HLODBuilderMeshMerge.h"
 #include "WorldPartition/HLOD/Builders/HLODBuilderMeshSimplify.h"
@@ -197,6 +198,20 @@ TArray<AWorldPartitionHLOD*> FWorldPartitionHLODUtilities::CreateHLODActors(FHLO
 
 			if (UHLODLayer* HLODLayer = Cast<UHLODLayer>(ActorDescView.GetHLODLayer().TryLoad()))
 			{
+				if (HLODLayer->GetLayerType() == EHLODLayerType::CustomHLODActor)
+				{
+					// Actors in the CustomHLODActor layer will be directly included in the HLOD representation, so we can use it as a source actor in the parent layer
+					// Or skip it if there's no parent layer
+					if (HLODLayer->GetParentLayer())
+					{
+						HLODLayer = HLODLayer->GetParentLayer();
+					}
+					else
+					{
+						continue;
+					}
+				}
+
 				TMap<FGuid, FWorldPartitionRuntimeCellObjectMapping>& SubActors = SourceActorsPerHLODLayer.FindOrAdd(HLODLayer);
 				AddSourceActor(ActorDescView, WorldPackageName, SubActors);
 			}
@@ -406,6 +421,9 @@ TSubclassOf<UHLODBuilder> FWorldPartitionHLODUtilities::GetHLODBuilderClass(cons
 	case EHLODLayerType::Custom:
 		return InHLODLayer->GetHLODBuilderClass();
 
+	case EHLODLayerType::CustomHLODActor:
+		return UHLODBuilderCustomHLODActor::StaticClass();
+
 	default:
 		checkf(false, TEXT("Unsupported type"));
 		return nullptr;
diff --git a/Engine/Plugins/Editor/WorldPartitionHLODUtilities/Source/Public/WorldPartition/HLOD/Builders/HLODBuilderCustomHLODActor.h b/Engine/Plugins/Editor/WorldPartitionHLODUtilities/Source/Public/WorldPartition/HLOD/Builders/HLODBuilderCustomHLODActor.h
new file mode 100644
--- /dev/null
+++ b/Engine/Plugins/Editor/WorldPartitionHLODUtilities/Source/Public/WorldPartition/HLOD/Builders/HLODBuilderCustomHLODActor.h
@@ -0,0 +1,25 @@
+// Copyright Epic Games, Inc. All Rights Reserved.
+
+#pragma once
+
+#include "Templates/SubclassOf.h"
+#include "WorldPartition/HLOD/HLODBuilder.h"
+#include "HLODBuilderCustomHLODActor.generated.h"
+
+UCLASS(Blueprintable, Config = Engine, PerObjectConfig)
+class UHLODBuilderCustomHLODActorSettings : public UHLODBuilderSettings
+{
+	GENERATED_UCLASS_BODY()
+
+	virtual bool IsReusingSourceMaterials() const override { return true; }
+};
+
+UCLASS(HideDropdown)
+class UHLODBuilderCustomHLODActor : public UHLODBuilder
+{
+	GENERATED_UCLASS_BODY()
+
+public:
+	virtual TSubclassOf<UHLODBuilderSettings> GetSettingsClass() const override;
+	virtual TArray<UActorComponent*> Build(const FHLODBuildContext& InHLODBuildContext, const TArray<UActorComponent *>& InSourceComponents) const override;
+};
\ No newline at end of file
diff --git a/Engine/Plugins/Experimental/FastGeoStreaming/Source/FastGeoStreaming/Internal/FastGeoHLOD.h b/Engine/Plugins/Experimental/FastGeoStreaming/Source/FastGeoStreaming/Internal/FastGeoHLOD.h
--- a/Engine/Plugins/Experimental/FastGeoStreaming/Source/FastGeoStreaming/Internal/FastGeoHLOD.h
+++ b/Engine/Plugins/Experimental/FastGeoStreaming/Source/FastGeoStreaming/Internal/FastGeoHLOD.h
@@ -34,6 +34,8 @@ public:
 	virtual const FGuid& GetSourceCellGuid() const override;
 	virtual bool IsStandalone() const override;
 	virtual const FGuid& GetStandaloneHLODGuid() const override;
+	virtual bool IsCustomHLOD() const override;
+	virtual const FGuid& GetCustomHLODGuid() const override;
 	// ~End IWorldPartitionHLODObject interface
 
 #if WITH_EDITOR
@@ -45,6 +47,7 @@ public:
 	void SetSourceCellGuid(const FGuid& InSourceCellGuid);
 	void SetRequireWarmup(bool bInRequireWarmup);
 	void SetStandaloneHLODGuid(const FGuid& InStandaloneHLODGuid);
+	void SetCustomHLODGuid(const FGuid& InCustomHLODGuid);
 #endif
 
 private:
@@ -56,4 +59,5 @@ private:
 	bool bRequireWarmup;
 	FGuid SourceCellGuid;
 	FGuid StandaloneHLODGuid;
+	FGuid CustomHLODGuid;
 };
\ No newline at end of file
diff --git a/Engine/Plugins/Experimental/FastGeoStreaming/Source/FastGeoStreaming/Private/FastGeoHLOD.cpp b/Engine/Plugins/Experimental/FastGeoStreaming/Source/FastGeoStreaming/Private/FastGeoHLOD.cpp
--- a/Engine/Plugins/Experimental/FastGeoStreaming/Source/FastGeoStreaming/Private/FastGeoHLOD.cpp
+++ b/Engine/Plugins/Experimental/FastGeoStreaming/Source/FastGeoStreaming/Private/FastGeoHLOD.cpp
@@ -32,6 +32,7 @@ void FFastGeoHLOD::Serialize(FArchive& Ar)
 	Ar << bRequireWarmup;
 	Ar << SourceCellGuid;
 	Ar << StandaloneHLODGuid;
+	Ar << CustomHLODGuid;
 }
 
 ULevel* FFastGeoHLOD::GetHLODLevel() const
@@ -112,6 +113,16 @@ const FGuid& FFastGeoHLOD::GetStandaloneHLODGuid() const
 	return StandaloneHLODGuid;
 }
 
+bool FFastGeoHLOD::IsCustomHLOD() const
+{
+	return CustomHLODGuid.IsValid();
+}
+
+const FGuid& FFastGeoHLOD::GetCustomHLODGuid() const
+{
+	return CustomHLODGuid;
+}
+
 void FFastGeoHLOD::OnRegister()
 {
 	GetLevel()->GetWorld()->GetSubsystem<UWorldPartitionHLODRuntimeSubsystem>()->RegisterHLODObject(this);
@@ -138,6 +149,11 @@ void FFastGeoHLOD::SetStandaloneHLODGuid(const FGuid& InStandaloneHLODGuid)
 	StandaloneHLODGuid = InStandaloneHLODGuid;
 }
 
+void FFastGeoHLOD::SetCustomHLODGuid(const FGuid& InCustomHLODGuid)
+{
+	CustomHLODGuid = InCustomHLODGuid;
+}
+
 FFastGeoComponent& FFastGeoHLOD::AddComponent(FFastGeoElementType InComponentType)
 {
 	FFastGeoComponent& NewComponent = Super::AddComponent(InComponentType);
diff --git a/Engine/Plugins/Experimental/FastGeoStreaming/Source/FastGeoStreaming/Private/FastGeoWorldPartitionRuntimeCellTransformer.cpp b/Engine/Plugins/Experimental/FastGeoStreaming/Source/FastGeoStreaming/Private/FastGeoWorldPartitionRuntimeCellTransformer.cpp
--- a/Engine/Plugins/Experimental/FastGeoStreaming/Source/FastGeoStreaming/Private/FastGeoWorldPartitionRuntimeCellTransformer.cpp
+++ b/Engine/Plugins/Experimental/FastGeoStreaming/Source/FastGeoStreaming/Private/FastGeoWorldPartitionRuntimeCellTransformer.cpp
@@ -302,6 +302,7 @@ void UFastGeoWorldPartitionRuntimeCellTransformer::Transform(ULevel* InLevel)
 				FastGeoHLOD->SetSourceCellGuid(HLODActor->GetSourceCellGuid());
 				FastGeoHLOD->SetRequireWarmup(HLODActor->DoesRequireWarmup());
 				FastGeoHLOD->SetStandaloneHLODGuid(HLODActor->GetStandaloneHLODGuid());
+				FastGeoHLOD->SetCustomHLODGuid(HLODActor->GetCustomHLODGuid());
 				CurrentComponentCluster = FastGeoHLOD.Get();
 			}
 
diff --git a/Engine/Source/Runtime/Engine/Private/WorldPartition/HLOD/CustomHLODActor.cpp b/Engine/Source/Runtime/Engine/Private/WorldPartition/HLOD/CustomHLODActor.cpp
new file mode 100644
--- /dev/null
+++ b/Engine/Source/Runtime/Engine/Private/WorldPartition/HLOD/CustomHLODActor.cpp
@@ -0,0 +1,131 @@
+// Copyright Epic Games, Inc. All Rights Reserved.
+
+#include "WorldPartition/HLOD/CustomHLODActor.h"
+
+#include "Components/StaticMeshComponent.h"
+#include "Engine/StaticMesh.h"
+#include "Engine/World.h"
+#include "Materials/MaterialInterface.h"
+#include "UObject/ObjectSaveContext.h"
+#include "WorldPartition/HLOD/HLODRuntimeSubsystem.h"
+
+#include UE_INLINE_GENERATED_CPP_BY_NAME(CustomHLODActor)
+
+AWorldPartitionCustomHLOD::AWorldPartitionCustomHLOD(const FObjectInitializer& ObjectInitializer)
+	: Super(ObjectInitializer)
+{
+	StaticMeshComponent = CreateDefaultSubobject<UStaticMeshComponent>(TEXT("RootComponent"));
+	StaticMeshComponent->Mobility = EComponentMobility::Static;
+	RootComponent = StaticMeshComponent;
+}
+
+UObject* AWorldPartitionCustomHLOD::GetUObject() const
+{
+	return const_cast<AWorldPartitionCustomHLOD*>(this);
+}
+
+ULevel* AWorldPartitionCustomHLOD::GetHLODLevel() const
+{
+	return GetLevel();
+}
+
+FString AWorldPartitionCustomHLOD::GetHLODNameOrLabel() const
+{
+	return GetActorNameOrLabel();
+}
+
+bool AWorldPartitionCustomHLOD::DoesRequireWarmup() const
+{
+	return true;
+}
+
+TSet<UObject*> AWorldPartitionCustomHLOD::GetAssetsToWarmup() const
+{
+	TSet<UObject*> AssetsToWarmup;
+	
+	ForEachComponent<UStaticMeshComponent>(false, [&](UStaticMeshComponent* SMC)
+	{
+		for (int32 iMaterialIndex = 0; iMaterialIndex < SMC->GetNumMaterials(); ++iMaterialIndex)
+		{
+			if (UMaterialInterface* Material = SMC->GetMaterial(iMaterialIndex))
+			{
+				AssetsToWarmup.Add(Material);
+			}
+		}
+
+		if (UStaticMesh* StaticMesh = SMC->GetStaticMesh())
+		{
+			AssetsToWarmup.Add(StaticMesh);
+		}
+	});
+
+	return AssetsToWarmup;
+}
+
+void AWorldPartitionCustomHLOD::SetVisibility(bool bInVisible)
+{
+	ForEachComponent<USceneComponent>(false, [bInVisible](USceneComponent* SceneComponent)
+	{
+		if (SceneComponent && (SceneComponent->GetVisibleFlag() != bInVisible))
+		{
+			SceneComponent->SetVisibility(bInVisible, false);
+		}
+	});
+}
+
+const FGuid& AWorldPartitionCustomHLOD::GetSourceCellGuid() const
+{
+	static const FGuid InvalidGuid;
+	return InvalidGuid;
+}
+
+bool AWorldPartitionCustomHLOD::IsStandalone() const
+{
+	return false;
+}
+
+const FGuid& AWorldPartitionCustomHLOD::GetStandaloneHLODGuid() const
+{
+	static const FGuid InvalidGuid;
+	return InvalidGuid;
+}
+
+bool AWorldPartitionCustomHLOD::IsCustomHLOD() const
+{
+	return true;
+}
+
+const FGuid& AWorldPartitionCustomHLOD::GetCustomHLODGuid() const
+{
+#if WITH_EDITOR
+	return GetActorInstanceGuid();
+#else
+	return HLODInstanceGuid;
+#endif
+}
+
+void AWorldPartitionCustomHLOD::BeginPlay()
+{
+	Super::BeginPlay();
+	GetWorld()->GetSubsystem<UWorldPartitionHLODRuntimeSubsystem>()->RegisterHLODObject(this);
+}
+
+void AWorldPartitionCustomHLOD::EndPlay(const EEndPlayReason::Type EndPlayReason)
+{
+	GetWorld()->GetSubsystem<UWorldPartitionHLODRuntimeSubsystem>()->UnregisterHLODObject(this);
+	Super::EndPlay(EndPlayReason);
+}
+
+#if WITH_EDITOR
+void AWorldPartitionCustomHLOD::PreSave(FObjectPreSaveContext ObjectSaveContext)
+{
+	Super::PreSave(ObjectSaveContext);
+
+	if (ObjectSaveContext.IsCooking())
+	{
+		HLODInstanceGuid = GetActorInstanceGuid();
+	}
+
+	SetActorEnableCollision(false);
+}
+#endif
\ No newline at end of file
diff --git a/Engine/Source/Runtime/Engine/Private/WorldPartition/HLOD/CustomHLODPlaceholderActor.cpp b/Engine/Source/Runtime/Engine/Private/WorldPartition/HLOD/CustomHLODPlaceholderActor.cpp
new file mode 100644
--- /dev/null
+++ b/Engine/Source/Runtime/Engine/Private/WorldPartition/HLOD/CustomHLODPlaceholderActor.cpp
@@ -0,0 +1,51 @@
+// Copyright Epic Games, Inc. All Rights Reserved.
+
+#include "WorldPartition/HLOD/CustomHLODPlaceholderActor.h"
+
+#include "WorldPartition/ActorDescContainerInstance.h"
+#include "WorldPartition/WorldPartitionActorDescInstance.h"
+#include "WorldPartition/HLOD/CustomHLODPlaceholderActorDesc.h"
+
+#include UE_INLINE_GENERATED_CPP_BY_NAME(CustomHLODPlaceholderActor)
+
+AWorldPartitionCustomHLODPlaceholder::AWorldPartitionCustomHLODPlaceholder(const FObjectInitializer& ObjectInitializer)
+	: Super(ObjectInitializer)
+{
+	RootComponent = CreateDefaultSubobject<USceneComponent>(TEXT("Root"));
+}
+
+#if WITH_EDITOR
+void AWorldPartitionCustomHLODPlaceholder::InitFrom(const FWorldPartitionActorDescInstance* InCustomHLODActorDescInstance)
+{
+	CustomHLODActorDescInstance = InCustomHLODActorDescInstance;
+	CustomHLODActorGuid = CustomHLODActorDescInstance->GetContainerInstance()->GetContainerID().GetActorGuid(CustomHLODActorDescInstance->GetGuid());
+
+	SetActorTransform(CustomHLODActorDescInstance->GetActorTransform());
+	SetRuntimeGrid(CustomHLODActorDescInstance->GetRuntimeGrid());
+	SetIsSpatiallyLoaded(CustomHLODActorDescInstance->GetActorDesc()->GetIsSpatiallyLoadedRaw());
+	bIsEditorOnlyActor = CustomHLODActorDescInstance->GetActorIsEditorOnly();
+}
+
+const FGuid& AWorldPartitionCustomHLODPlaceholder::GetCustomHLODActorGuid() const
+{
+	return CustomHLODActorGuid;
+}
+
+void AWorldPartitionCustomHLODPlaceholder::GetStreamingBounds(FBox& OutRuntimeBounds, FBox& OutEditorBounds) const
+{
+	if (CustomHLODActorDescInstance)
+	{
+		OutRuntimeBounds = CustomHLODActorDescInstance->GetActorDesc()->GetRuntimeBounds();
+		OutEditorBounds = CustomHLODActorDescInstance->GetActorDesc()->GetEditorBounds();
+	}
+	else
+	{
+		OutRuntimeBounds = OutEditorBounds = FBox(EForceInit::ForceInit);
+	}
+}
+
+TUniquePtr<FWorldPartitionActorDesc> AWorldPartitionCustomHLODPlaceholder::CreateClassActorDesc() const
+{
+	return TUniquePtr<FWorldPartitionActorDesc>(new FCustomHLODPlaceholderActorDesc());
+}
+#endif
\ No newline at end of file
diff --git a/Engine/Source/Runtime/Engine/Private/WorldPartition/HLOD/CustomHLODPlaceholderActorDesc.cpp b/Engine/Source/Runtime/Engine/Private/WorldPartition/HLOD/CustomHLODPlaceholderActorDesc.cpp
new file mode 100644
--- /dev/null
+++ b/Engine/Source/Runtime/Engine/Private/WorldPartition/HLOD/CustomHLODPlaceholderActorDesc.cpp
@@ -0,0 +1,47 @@
+// Copyright Epic Games, Inc. All Rights Reserved.
+
+#include "WorldPartition/HLOD/CustomHLODPlaceholderActorDesc.h"
+
+#include "WorldPartition/HLOD/CustomHLODPlaceholderActor.h"
+
+#if WITH_EDITOR
+void FCustomHLODPlaceholderActorDesc::Init(const AActor* InActor)
+{
+	FWorldPartitionActorDesc::Init(InActor);
+
+	if (!bIsDefaultActorDesc)
+	{
+		if (const AWorldPartitionCustomHLODPlaceholder* PlaceholderActor = CastChecked<AWorldPartitionCustomHLODPlaceholder>(InActor))
+		{
+			CustomHLODActorGuid = PlaceholderActor->GetCustomHLODActorGuid();
+		}
+	}
+}
+
+bool FCustomHLODPlaceholderActorDesc::Equals(const FWorldPartitionActorDesc* Other) const
+{
+	if (FWorldPartitionActorDesc::Equals(Other))
+	{
+		const FCustomHLODPlaceholderActorDesc* CustomHLODPlaceholderActorDesc = (FCustomHLODPlaceholderActorDesc*)Other;
+		return CustomHLODActorGuid == CustomHLODPlaceholderActorDesc->CustomHLODActorGuid;
+	}
+	return false;
+}
+
+const FGuid& FCustomHLODPlaceholderActorDesc::GetCustomHLODActorGuid() const
+{
+	return CustomHLODActorGuid;
+}
+
+uint32 FCustomHLODPlaceholderActorDesc::GetSizeOf() const
+{
+	return sizeof(FCustomHLODPlaceholderActorDesc);
+}
+
+void FCustomHLODPlaceholderActorDesc::Serialize(FArchive& Ar)
+{
+	FWorldPartitionActorDesc::Serialize(Ar);
+
+	Ar << CustomHLODActorGuid;
+}
+#endif
\ No newline at end of file
diff --git a/Engine/Source/Runtime/Engine/Private/WorldPartition/HLOD/HLODActor.cpp b/Engine/Source/Runtime/Engine/Private/WorldPartition/HLOD/HLODActor.cpp
--- a/Engine/Source/Runtime/Engine/Private/WorldPartition/HLOD/HLODActor.cpp
+++ b/Engine/Source/Runtime/Engine/Private/WorldPartition/HLOD/HLODActor.cpp
@@ -107,6 +107,17 @@ const FGuid& AWorldPartitionHLOD::GetStandaloneHLODGuid() const
 	return StandaloneHLODGuid;
 }
 
+bool AWorldPartitionHLOD::IsCustomHLOD() const
+{
+	return false;
+}
+
+const FGuid& AWorldPartitionHLOD::GetCustomHLODGuid() const
+{
+	static const FGuid InvalidGuid;
+	return InvalidGuid;
+}
+
 void AWorldPartitionHLOD::SetVisibility(bool bInVisible)
 {
 	ForEachComponent<USceneComponent>(false, [bInVisible](USceneComponent* SceneComponent)
diff --git a/Engine/Source/Runtime/Engine/Private/WorldPartition/HLOD/HLODRuntimeSubsystem.cpp b/Engine/Source/Runtime/Engine/Private/WorldPartition/HLOD/HLODRuntimeSubsystem.cpp
--- a/Engine/Source/Runtime/Engine/Private/WorldPartition/HLOD/HLODRuntimeSubsystem.cpp
+++ b/Engine/Source/Runtime/Engine/Private/WorldPartition/HLOD/HLODRuntimeSubsystem.cpp
@@ -386,7 +386,33 @@ void UWorldPartitionHLODRuntimeSubsystem::RegisterHLODObject(IWorldPartitionHLOD
 {
 	TRACE_CPUPROFILER_EVENT_SCOPE(UWorldPartitionHLODRuntimeSubsystem::RegisterHLODObject);
 	
-	if (FCellData* SourceCellData = GetCellData(InWorldPartitionHLOD))
+	if (InWorldPartitionHLOD->IsCustomHLOD())
+	{
+		if (const UWorldPartition* WorldPartition = FHLODRuntimeSubsystem::GetWorldPartition(InWorldPartitionHLOD))
+		{
+			if (const FGuid* CustomHLODSourceCellGuidPtr = WorldPartition->RuntimeHash->GetCustomHLODActorSourceCellOverride(InWorldPartitionHLOD->GetCustomHLODGuid()))
+			{
+				FGuid CustomHLODSourceCellGuid = *CustomHLODSourceCellGuidPtr;
+				if (FCellData* CustomHLODSourceCellData = GetCellData(WorldPartition, CustomHLODSourceCellGuid))
+				{
+					// Overriden source cell found in WorldPartitionHLODRuntimeData. Add the HLOD Actor to its LoadedHLODs.
+					UE_LOG(LogHLODRuntimeSubsystem, Verbose, TEXT("Registering Custom HLOD %s for overriden source cell '%s'"), *InWorldPartitionHLOD->GetHLODNameOrLabel(), *CustomHLODSourceCellGuid.ToString());
+
+					CustomHLODSourceCellData->LoadedHLODs.Add(InWorldPartitionHLOD);
+					InWorldPartitionHLOD->SetVisibility(UWorldPartitionHLODRuntimeSubsystem::WorldPartitionHLODEnabled && !CustomHLODSourceCellData->bIsCellVisible);
+				}
+				else
+				{
+					UE_LOG(LogHLODRuntimeSubsystem, Log, TEXT("Override source cell '%s' not found for Custom HLOD actor %s"), *CustomHLODSourceCellGuid.ToString(), *InWorldPartitionHLOD->GetHLODNameOrLabel());
+				}
+			}
+			else
+			{
+				UE_LOG(LogHLODRuntimeSubsystem, Log, TEXT("Override source cell not found for Custom HLOD actor %s"), *InWorldPartitionHLOD->GetHLODNameOrLabel());
+			}
+		}
+	}
+	else if (FCellData* SourceCellData = GetCellData(InWorldPartitionHLOD))
 	{
 		UE_LOG(LogHLODRuntimeSubsystem, Verbose, TEXT("Registering HLOD %s for cell %s"), *InWorldPartitionHLOD->GetHLODNameOrLabel(), *InWorldPartitionHLOD->GetSourceCellGuid().ToString());
 
@@ -465,7 +491,24 @@ void UWorldPartitionHLODRuntimeSubsystem::UnregisterHLODObject(IWorldPartitionHL
 {
 	TRACE_CPUPROFILER_EVENT_SCOPE(UWorldPartitionHLODRuntimeSubsystem::UnregisterHLODObject);
 
-	if (FCellData* CellData = GetCellData(InWorldPartitionHLOD))
+	if (InWorldPartitionHLOD->IsCustomHLOD())
+	{
+		if (const UWorldPartition* WorldPartition = FHLODRuntimeSubsystem::GetWorldPartition(InWorldPartitionHLOD))
+		{
+			if (const FGuid* CustomHLODSourceCellGuidPtr = WorldPartition->RuntimeHash->GetCustomHLODActorSourceCellOverride(InWorldPartitionHLOD->GetCustomHLODGuid()))
+			{
+				const FGuid CustomHLODSourceCellGuid = *CustomHLODSourceCellGuidPtr;
+				if (FCellData* CustomHLODSourceCellData = GetCellData(WorldPartition, CustomHLODSourceCellGuid))
+				{
+					// Overriden source cell found. Remove the HLOD object from its LoadedHLODs.
+					UE_LOG(LogHLODRuntimeSubsystem, Verbose, TEXT("Unregistering Custom HLOD %s for overriden source cell %s"), *InWorldPartitionHLOD->GetHLODNameOrLabel(), *CustomHLODSourceCellGuid.ToString());
+
+					verify(CustomHLODSourceCellData->LoadedHLODs.Remove(InWorldPartitionHLOD));
+				}
+			}
+		}
+	}
+	else if (FCellData* CellData = GetCellData(InWorldPartitionHLOD))
 	{
 		UE_LOG(LogHLODRuntimeSubsystem, Verbose, TEXT("Unregistering HLOD %s for cell %s"), *InWorldPartitionHLOD->GetHLODNameOrLabel(), *InWorldPartitionHLOD->GetSourceCellGuid().ToString());
 
diff --git a/Engine/Source/Runtime/Engine/Private/WorldPartition/RuntimeHashSet/WorldPartitionRuntimeHashSet.cpp b/Engine/Source/Runtime/Engine/Private/WorldPartition/RuntimeHashSet/WorldPartitionRuntimeHashSet.cpp
--- a/Engine/Source/Runtime/Engine/Private/WorldPartition/RuntimeHashSet/WorldPartitionRuntimeHashSet.cpp
+++ b/Engine/Source/Runtime/Engine/Private/WorldPartition/RuntimeHashSet/WorldPartitionRuntimeHashSet.cpp
@@ -1357,3 +1357,8 @@ const FGuid* UWorldPartitionRuntimeHashSet::GetStandaloneHLODActorSourceCellOver
 {
 	return StandaloneHLODActorToSourceCellsMap.Find(InActorGuid);
 }
+
+const FGuid* UWorldPartitionRuntimeHashSet::GetCustomHLODActorSourceCellOverride(const FGuid& InActorGuid) const
+{
+	return CustomHLODActorToSourceCellsMap.Find(InActorGuid);
+}
\ No newline at end of file
diff --git a/Engine/Source/Runtime/Engine/Private/WorldPartition/RuntimeHashSet/WorldPartitionRuntimeHashSetStreamingGeneration.cpp b/Engine/Source/Runtime/Engine/Private/WorldPartition/RuntimeHashSet/WorldPartitionRuntimeHashSetStreamingGeneration.cpp
--- a/Engine/Source/Runtime/Engine/Private/WorldPartition/RuntimeHashSet/WorldPartitionRuntimeHashSetStreamingGeneration.cpp
+++ b/Engine/Source/Runtime/Engine/Private/WorldPartition/RuntimeHashSet/WorldPartitionRuntimeHashSetStreamingGeneration.cpp
@@ -11,6 +11,8 @@
 #include "WorldPartition/HLOD/HLODLayer.h"
 #include "WorldPartition/HLOD/HLODActor.h"
 #include "WorldPartition/HLOD/HLODActorDesc.h"
+#include "WorldPartition/HLOD/CustomHLODPlaceholderActor.h"
+#include "WorldPartition/HLOD/CustomHLODPlaceholderActorDesc.h"
 
 #if WITH_EDITOR
 bool UWorldPartitionRuntimeHashSet::GenerateRuntimePartitionsStreamingDescs(const IStreamingGenerationContext* StreamingGenerationContext, TMap<URuntimePartition*, TArray<URuntimePartition::FCellDescInstance>>& OutRuntimeCellDescs) const
@@ -169,6 +171,8 @@ bool UWorldPartitionRuntimeHashSet::GenerateStreaming(UWorldPartitionStreamingPo
 		return RuntimeCell;
 	};
 
+	CustomHLODActorToSourceCellsMap.Empty();
+
 	TMap<FGuid, FGuid> StandaloneHLODActorToCell;
 	TArray<UWorldPartitionRuntimeCell*> RuntimeCells;
 	TMap<URuntimePartition*, FRuntimePartitionStreamingData> RuntimePartitionsStreamingData;
@@ -185,6 +189,23 @@ bool UWorldPartitionRuntimeHashSet::GenerateStreaming(UWorldPartitionStreamingPo
 				RuntimeCell->SetIsAlwaysLoaded(bIsCellAlwaysLoaded);
 				RuntimeCell->SetIsSpatiallyLoaded(CellDescInstance.bIsSpatiallyLoaded);
 
+				// Handle Custom HLOD Placeholder actors.
+				// Save information about mapping from Custom HLOD actor GUID to Source cell and remove placeholder actor from the list of actors to be placed in this cell.
+				// For more information see FWorldPartitionStreamingGenerator::CreateActorDescViewMap
+				bool bRemovedAnyCustomHLODPlaceholderActors = false;
+				for (TArray<IStreamingGenerationContext::FActorInstance>::TIterator ActorInstanceIterator = CellActorInstances.CreateIterator(); ActorInstanceIterator; ++ActorInstanceIterator)
+				{
+					if (ActorInstanceIterator->GetActorDescView().GetActorNativeClass()->IsChildOf<AWorldPartitionCustomHLODPlaceholder>())
+					{
+						const FCustomHLODPlaceholderActorDesc* CustomHLODPlaceholderActorDesc = static_cast<const FCustomHLODPlaceholderActorDesc*>(ActorInstanceIterator->GetActorDescView().GetActorDesc());
+						CustomHLODActorToSourceCellsMap.Add(CustomHLODPlaceholderActorDesc->GetCustomHLODActorGuid(), RuntimeCell->GetGuid());
+						ActorInstanceIterator.RemoveCurrent();
+						bRemovedAnyCustomHLODPlaceholderActors = true;
+					}
+				}
+				// If, after removing all placeholder actors, there are no actors left in this cell, mark the cell as Custom HLOD Placeholder cell, so that we can still use it at runtime to control Custom HLOD actor's visibility.
+				RuntimeCell->SetIsCustomHLODPlaceholderCell(bRemovedAnyCustomHLODPlaceholderActors && CellActorInstances.IsEmpty());
+
 				PopulateRuntimeCell(RuntimeCell, CellActorInstances, OutPackagesToGenerate);
 
 				// Save HLOD actor GUID -> Cell GUID mapping to use it for Standalone HLOD source cell override
diff --git a/Engine/Source/Runtime/Engine/Private/WorldPartition/WorldPartitionRuntimeLevelStreamingCell.cpp b/Engine/Source/Runtime/Engine/Private/WorldPartition/WorldPartitionRuntimeLevelStreamingCell.cpp
--- a/Engine/Source/Runtime/Engine/Private/WorldPartition/WorldPartitionRuntimeLevelStreamingCell.cpp
+++ b/Engine/Source/Runtime/Engine/Private/WorldPartition/WorldPartitionRuntimeLevelStreamingCell.cpp
@@ -639,7 +639,7 @@ UWorldPartitionLevelStreamingDynamic* UWorldPartitionRuntimeLevelStreamingCell::
 	}
 #else
 	// In Runtime, always loaded cell level is handled by World directly
-	check(LevelStreaming || IsAlwaysLoaded());
+	check(LevelStreaming || IsAlwaysLoaded() || GetIsCustomHLODPlaceholderCell());
 #endif
 
 #if !WITH_EDITOR
@@ -692,6 +692,10 @@ void UWorldPartitionRuntimeLevelStreamingCell::Activate() const
 	{
 		LocalLevelStreaming->Activate();
 	}
+	else if (GetIsCustomHLODPlaceholderCell())
+	{
+		OnCellShown();
+	}
 }
 
 void UWorldPartitionRuntimeLevelStreamingCell::SetStreamingPriority(int32 InStreamingPriority) const
@@ -715,31 +719,39 @@ bool UWorldPartitionRuntimeLevelStreamingCell::CanUnload() const
 void UWorldPartitionRuntimeLevelStreamingCell::Unload() const
 {
 #if WITH_EDITOR
-	check(LevelStreaming || GetActorCount());
+	check(LevelStreaming || GetActorCount() || GetIsCustomHLODPlaceholderCell());
 #else
 	// In Runtime, always loaded cell level is handled by World directly
-	check(LevelStreaming || IsAlwaysLoaded());
+	check(LevelStreaming || IsAlwaysLoaded() || GetIsCustomHLODPlaceholderCell());
 #endif
 
 	if (LevelStreaming)
 	{
 		LevelStreaming->Unload();
 	}
+	else if (GetIsCustomHLODPlaceholderCell())
+	{
+		OnCellHidden();
+	}
 }
 
 void UWorldPartitionRuntimeLevelStreamingCell::Deactivate() const
 {
 #if WITH_EDITOR
-	check(LevelStreaming || GetActorCount());
+	check(LevelStreaming || GetActorCount() || GetIsCustomHLODPlaceholderCell());
 #else
 	// In Runtime, always loaded cell level is handled by World directly
-	check(LevelStreaming || IsAlwaysLoaded());
+	check(LevelStreaming || IsAlwaysLoaded() || GetIsCustomHLODPlaceholderCell());
 #endif
 
 	if (LevelStreaming)
 	{
 		LevelStreaming->Deactivate();
 	}
+	else if (GetIsCustomHLODPlaceholderCell())
+	{
+		OnCellHidden();
+	}
 }
 
 void UWorldPartitionRuntimeLevelStreamingCell::OnLevelShown()
diff --git a/Engine/Source/Runtime/Engine/Private/WorldPartition/WorldPartitionStreamingGeneration.cpp b/Engine/Source/Runtime/Engine/Private/WorldPartition/WorldPartitionStreamingGeneration.cpp
--- a/Engine/Source/Runtime/Engine/Private/WorldPartition/WorldPartitionStreamingGeneration.cpp
+++ b/Engine/Source/Runtime/Engine/Private/WorldPartition/WorldPartitionStreamingGeneration.cpp
@@ -34,11 +34,14 @@
 #include "WorldPartition/DataLayer/WorldDataLayers.h"
 #include "WorldPartition/ErrorHandling/WorldPartitionStreamingGenerationNullErrorHandler.h"
 #include "WorldPartition/ErrorHandling/WorldPartitionStreamingGenerationLogErrorHandler.h"
+#include "WorldPartition/HLOD/CustomHLODActor.h"
+#include "WorldPartition/HLOD/CustomHLODPlaceholderActor.h"
 #include "WorldPartition/HLOD/HLODActor.h"
 #include "WorldPartition/HLOD/HLODActorDesc.h"
 #include "WorldPartition/HLOD/HLODLayer.h"
 #include "WorldPartition/HLOD/StandaloneHLODActor.h"
 #include "WorldPartition/HLOD/StandaloneHLODSubsystem.h"
+#include "WorldPartition/RuntimeHashSet/WorldPartitionRuntimeHashSet.h"
 #include "WorldPartition/WorldPartitionSubsystem.h"
 #include "WorldPartition/WorldPartitionRuntimeContainerResolving.h"
 #include "HAL/FileManager.h"
@@ -665,6 +668,9 @@ class FWorldPartitionStreamingGenerator
 		/** List of Standalone HLOD actor descriptor views. */
 		TArray<TUniquePtr<FStreamingGenerationTempActorDescInstance>> StandaloneHLODActorInstances;
 
+		/** List of transient placeholder actors used to determine source cell for Custom HLOD actors */
+		TArray<TUniquePtr<FStreamingGenerationTempActorDescInstance>> CustomHLODActorPlaceholderInstances;
+
 		/** List of actor clusters for this descriptor */
 		TArray<TArray<FGuid>> Clusters;
 
@@ -818,6 +824,7 @@ class FWorldPartitionStreamingGenerator
 		TArray<FStreamingGenerationActorDescView>& OutContainerInstances = InContainerCollectionInstanceDescriptor.ContainerCollectionInstanceViews;
 		TArray<TUniquePtr<FStreamingGenerationTempActorDescInstance>>& OutUnsavedDirtyInstances = InContainerCollectionInstanceDescriptor.UnsavedDirtyInstances;
 		TArray<TUniquePtr<FStreamingGenerationTempActorDescInstance>>& OutStandaloneHLODActorInstances = InContainerCollectionInstanceDescriptor.StandaloneHLODActorInstances;
+		TArray<TUniquePtr<FStreamingGenerationTempActorDescInstance>>& OutCustomHLODActorPlaceholderInstances = InContainerCollectionInstanceDescriptor.CustomHLODActorPlaceholderInstances;
 		TMap<FGuid, FTransform>& EditorOnlyParentActorTransforms = InContainerCollectionInstanceDescriptor.EditorOnlyParentActorTransforms;
 
 		// Should we handle unsaved or newly created actors?
@@ -888,6 +895,47 @@ class FWorldPartitionStreamingGenerator
 				}
 			}
 
+			if (Iterator->GetActorNativeClass()->IsChildOf<AWorldPartitionCustomHLOD>())
+			{
+				// Custom HLOD actors, unlike regular HLOD actors don't have their SourceCell set during creation, so we need to compute and store that information during streaming generation
+				// We're doing this by adding a placeholder actor descriptor for each custom HLOD actor. Placeholder actor has the same transform, bounds, etc. as the Custom HLOD actor.
+				// Later, when placing actors in the appropriate cells we'll save that information and remove the placeholder actors.
+				// The easiest way to add actor desc view is to spawn an actual actor and use the same logic that is used to handle unsaved dirty actors. 
+				// Right after we have actor desc view, we can destroy actor to avoid undesired side effects.
+				// To keep streaming generation deterministic, we use a deterministic actor GUID and name.
+				FActorSpawnParameters SpawnParameters;
+				SpawnParameters.ObjectFlags = RF_Transient;
+				SpawnParameters.bCreateActorPackage = false;
+				SpawnParameters.OverrideActorGuid = FGuid::NewDeterministicGuid(TEXT("CustomHLODPlaceholder") + Iterator->GetActorSoftPath().ToString());
+				SpawnParameters.Name = FName(SpawnParameters.OverrideActorGuid.ToString());
+				SpawnParameters.NameMode = FActorSpawnParameters::ESpawnActorNameMode::Requested;
+
+				// Spawn placeholder actor and initialize properties based on Custom HLOD actor
+				AWorldPartitionCustomHLODPlaceholder* CustomHLODPlaceholderActor = WorldPartitionContext->GetWorld()->SpawnActor<AWorldPartitionCustomHLODPlaceholder>(SpawnParameters);
+				CustomHLODPlaceholderActor->InitFrom(*Iterator);
+				
+				// Register placeholder actor desc view
+				TUniquePtr<FStreamingGenerationTempActorDescInstance>& CustomHLODPlaceholderRef = OutCustomHLODActorPlaceholderInstances.Add_GetRef(MakeUnique<FStreamingGenerationTempActorDescInstance>(Iterator->GetContainerInstance(), CustomHLODPlaceholderActor));
+				RegisterActorDescView(FStreamingGenerationActorDescView(OutActorDescViewMap, CustomHLODPlaceholderRef.Get(), false));
+
+				// Destroy placeholder actor since we only need its actor descriptor.
+				WorldPartitionContext->GetWorld()->DestroyActor(CustomHLODPlaceholderActor, false, false);
+				FName UniqueTrashName = MakeUniqueObjectName(CustomHLODPlaceholderActor->GetOuter(), CustomHLODPlaceholderActor->GetClass(), NAME_TrashedPackage);
+				CustomHLODPlaceholderActor->Rename(*UniqueTrashName.ToString(), nullptr, REN_DontCreateRedirectors | REN_NonTransactional | REN_DoNotDirty);
+
+				// If Custom HLOD actor is used with some HLOD layer different than CustomHLODActor layer don't register its actor descriptor
+				if (bIncludeCustomHLODActorsWithOtherLayerType == false)
+				{
+					if (UHLODLayer* HLODLayer = Cast<UHLODLayer>(Iterator->GetHLODLayer().TryLoad()))
+					{
+						if (HLODLayer->GetLayerType() != EHLODLayerType::CustomHLODActor)
+						{
+							continue;
+						}
+					}
+				}
+			}
+
 			AActor* Actor = nullptr;
 			if (ShouldRegisterActorDesc(*Iterator, &Actor))
 			{
@@ -1007,11 +1055,11 @@ class FWorldPartitionStreamingGenerator
 			}
 		}
 	}
-
+	
 	void CreateActorDescriptorViewsRecursive(FContainerCollectionInstanceDescriptor&& InContainerCollectionInstanceDescriptor)
 	{
 		// Inherited parent per-instance data logic
-		auto InheritParentContainerPerInstanceData = [](const FContainerCollectionInstanceDescriptor& ParentContainerCollectionInstanceDescriptor, const FStreamingGenerationActorDescView& InActorDescView)
+		auto InheritParentContainerPerInstanceData = [this](const FContainerCollectionInstanceDescriptor& ParentContainerCollectionInstanceDescriptor, const FStreamingGenerationActorDescView& InActorDescView)
 		{
 			FContainerCollectionInstanceDescriptor::FPerInstanceData ResultPerInstanceData;
 
@@ -1021,6 +1069,28 @@ class FWorldPartitionStreamingGenerator
 			// Runtime grid is inherited from the main world if the actor has its runtime grid set to none.
 			ResultPerInstanceData.RuntimeGrid = (ParentContainerCollectionInstanceDescriptor.ID.IsMainContainer() || ParentContainerCollectionInstanceDescriptor.InstanceData.RuntimeGrid.IsNone()) ? InActorDescView.GetRuntimeGrid() : ParentContainerCollectionInstanceDescriptor.InstanceData.RuntimeGrid;
 
+			// Set RuntimeGrid for Custom HLOD Actor based on actor's HLOD Layer to have it directly included in the appropriate grid
+			if (InActorDescView.GetActorNativeClass()->IsChildOf<AWorldPartitionCustomHLOD>())
+			{
+				if (UHLODLayer* HLODLayer = Cast<UHLODLayer>(InActorDescView.GetHLODLayer().TryLoad()))
+				{
+					if (HLODLayer->GetLayerType() == EHLODLayerType::CustomHLODActor)
+					{
+						if (WorldPartitionContext)
+						{
+							if (const UWorldPartitionRuntimeHashSet* HashSet = Cast<UWorldPartitionRuntimeHashSet>(WorldPartitionContext->RuntimeHash))
+							{
+								if (const URuntimePartition* RuntimePartition = HashSet->ResolveRuntimePartitionForHLODLayer(ResultPerInstanceData.RuntimeGrid, InActorDescView.GetHLODLayer()))
+								{
+									FName GridName = ResultPerInstanceData.RuntimeGrid.IsNone() ? DefaultGrid : ResultPerInstanceData.RuntimeGrid;
+									ResultPerInstanceData.RuntimeGrid = FName(GridName.ToString() + TEXT(":") + RuntimePartition->Name.ToString());
+								}
+							}
+						}
+					}
+				}
+			}
+
 			// Data layers are accumulated down the hierarchy chain, since level instances supports data layers assignation on actors
 			ResultPerInstanceData.DataLayers = InActorDescView.GetRuntimeDataLayerInstanceNames().ToArray();
 			ResultPerInstanceData.DataLayers.Append(ParentContainerCollectionInstanceDescriptor.InstanceData.DataLayers);
@@ -1849,6 +1919,7 @@ public:
 			, bHandleUnsavedActors(false)
 			, bIsPIE(false)
 			, bIncludeStandaloneHLODActors(false)
+			, bIncludeCustomHLODActorsWithOtherLayerType(false)
 		{}
 
 		const UWorldPartition* WorldPartitionContext;
@@ -1858,6 +1929,7 @@ public:
 		bool bHandleUnsavedActors;
 		bool bIsPIE;
 		bool bIncludeStandaloneHLODActors;
+		bool bIncludeCustomHLODActorsWithOtherLayerType;
 		TMap<FGuid, const UActorDescContainerInstance*> ActorGuidsToContainerInstanceMap;
 		TArray<TSubclassOf<AActor>> FilteredClasses;
 		FName DefaultGrid;
@@ -1871,6 +1943,7 @@ public:
 		FWorldPartitionStreamingGeneratorParams& SetHandleUnsavedActors(bool bInHandleUnsavedActors) { bHandleUnsavedActors = bInHandleUnsavedActors; return *this; }
 		FWorldPartitionStreamingGeneratorParams& SetIsPIE(bool bInIsPIE) { bIsPIE = bInIsPIE; return *this; }
 		FWorldPartitionStreamingGeneratorParams& SetIncludeStandaloneHLODActors(bool bInIncludeStandaloneHLODActors) { bIncludeStandaloneHLODActors = bInIncludeStandaloneHLODActors; return *this; }
+		FWorldPartitionStreamingGeneratorParams& SetIncludeCustomHLODActorsWithOtherLayerType(bool bInIncludeCustomHLODActorsWithOtherLayerType) { bIncludeCustomHLODActorsWithOtherLayerType = bInIncludeCustomHLODActorsWithOtherLayerType; return *this; }
 		FWorldPartitionStreamingGeneratorParams& SetActorGuidsToContainerInstanceMap(const TMap<FGuid, const UActorDescContainerInstance*>& InActorGuidsToContainerInstanceMap) { ActorGuidsToContainerInstanceMap = InActorGuidsToContainerInstanceMap; return *this; }
 		FWorldPartitionStreamingGeneratorParams& SetFilteredClasses(const TArray<TSubclassOf<AActor>>& InFilteredClasses) { FilteredClasses = InFilteredClasses; return *this; }
 		FWorldPartitionStreamingGeneratorParams& SetDefaultGrid(FName InDefaultGrid) { DefaultGrid = InDefaultGrid; return *this; }
@@ -1887,6 +1960,7 @@ public:
 		, bHandleUnsavedActors(Params.bHandleUnsavedActors)
 		, bIsPIE(Params.bIsPIE)
 		, bIncludeStandaloneHLODActors(Params.bIncludeStandaloneHLODActors)
+		, bIncludeCustomHLODActorsWithOtherLayerType(Params.bIncludeCustomHLODActorsWithOtherLayerType)
 		, FilteredClasses(Params.FilteredClasses)
 		, DefaultGrid(Params.DefaultGrid)
 		, IsValidGrid(Params.IsValidGrid)
@@ -2061,6 +2135,7 @@ private:
 	bool bHandleUnsavedActors;
 	bool bIsPIE;
 	bool bIncludeStandaloneHLODActors;
+	bool bIncludeCustomHLODActorsWithOtherLayerType;
 	TArray<TSubclassOf<AActor>> FilteredClasses;
 	FName DefaultGrid;
 	TFunction<bool(FName, const UClass*)> IsValidGrid;
@@ -2314,7 +2389,8 @@ void UWorldPartition::SetupHLODActors(const FSetupHLODActorsParams& Params)
 			.SetFilteredClasses({ AWorldPartitionHLOD::StaticClass() })
 			.SetDefaultGrid(RuntimeHash ? RuntimeHash->GetDefaultGrid() : NAME_None)
 			.SetIsValidGrid([this](FName GridName, const UClass* ActorClass) { return RuntimeHash->IsValidGrid(GridName, ActorClass); })
-			.SetIsValidHLODLayer([this](FName GridName, const FSoftObjectPath& HLODLayerPath) { return RuntimeHash->IsValidHLODLayer(GridName, HLODLayerPath); });
+			.SetIsValidHLODLayer([this](FName GridName, const FSoftObjectPath& HLODLayerPath) { return RuntimeHash->IsValidHLODLayer(GridName, HLODLayerPath); })
+			.SetIncludeCustomHLODActorsWithOtherLayerType(true);
 
 		FWorldPartitionStreamingGenerator StreamingGenerator(StreamingGeneratorParams);
 		StreamingGenerator.PreparationPhase(InContainerInstanceCollection);
diff --git a/Engine/Source/Runtime/Engine/Public/WorldPartition/HLOD/CustomHLODActor.h b/Engine/Source/Runtime/Engine/Public/WorldPartition/HLOD/CustomHLODActor.h
new file mode 100644
--- /dev/null
+++ b/Engine/Source/Runtime/Engine/Public/WorldPartition/HLOD/CustomHLODActor.h
@@ -0,0 +1,47 @@
+// Copyright Epic Games, Inc. All Rights Reserved.
+
+#pragma once
+
+#include "GameFramework/Actor.h"
+#include "WorldPartition/HLOD/IWorldPartitionHLODObject.h"
+
+#include "CustomHLODActor.generated.h"
+
+UCLASS(MinimalAPI)
+class AWorldPartitionCustomHLOD : public AActor, public IWorldPartitionHLODObject
+{
+	GENERATED_BODY()
+
+public:
+	ENGINE_API AWorldPartitionCustomHLOD(const FObjectInitializer& ObjectInitializer);
+
+	// ~Begin IWorldPartitionHLODObject interface
+	virtual UObject* GetUObject() const override;
+	ENGINE_API virtual ULevel* GetHLODLevel() const override;
+	ENGINE_API virtual FString GetHLODNameOrLabel() const override;
+	virtual bool DoesRequireWarmup() const override;
+	ENGINE_API virtual TSet<UObject*> GetAssetsToWarmup() const override;
+	ENGINE_API virtual void SetVisibility(bool bIsVisible) override;
+	ENGINE_API virtual const FGuid& GetSourceCellGuid() const override;
+	ENGINE_API virtual bool IsStandalone() const override;
+	ENGINE_API virtual const FGuid& GetStandaloneHLODGuid() const override;
+	ENGINE_API virtual bool IsCustomHLOD() const override;
+	ENGINE_API virtual const FGuid& GetCustomHLODGuid() const override;
+	// ~End IWorldPartitionHLODObject interface
+
+	// ~Begin AActor Interface
+	ENGINE_API virtual void BeginPlay() override;
+	ENGINE_API virtual void EndPlay(const EEndPlayReason::Type EndPlayReason) override;
+	// ~End AActor Interface
+
+protected:
+#if WITH_EDITOR
+	ENGINE_API virtual void PreSave(FObjectPreSaveContext ObjectSaveContext) override;
+#endif
+private:
+	UPROPERTY(Category=CustomHLOD, VisibleAnywhere, BlueprintReadOnly, meta = (ExposeFunctionCategories = "Mesh,Rendering,Physics,Components|StaticMesh", AllowPrivateAccess = "true"))
+	TObjectPtr<class UStaticMeshComponent> StaticMeshComponent;
+
+	UPROPERTY()
+	FGuid HLODInstanceGuid;
+};
\ No newline at end of file
diff --git a/Engine/Source/Runtime/Engine/Public/WorldPartition/HLOD/CustomHLODPlaceholderActor.h b/Engine/Source/Runtime/Engine/Public/WorldPartition/HLOD/CustomHLODPlaceholderActor.h
new file mode 100644
--- /dev/null
+++ b/Engine/Source/Runtime/Engine/Public/WorldPartition/HLOD/CustomHLODPlaceholderActor.h
@@ -0,0 +1,30 @@
+// Copyright Epic Games, Inc. All Rights Reserved.
+
+#pragma once
+
+#include "GameFramework/Actor.h"
+#include "Misc/Guid.h"
+
+#include "CustomHLODPlaceholderActor.generated.h"
+
+class FWorldPartitionActorDescInstance;
+
+UCLASS(MinimalAPI, NotPlaceable, Transient)
+class AWorldPartitionCustomHLODPlaceholder : public AActor
+{
+	GENERATED_BODY()
+
+public:
+	ENGINE_API AWorldPartitionCustomHLODPlaceholder(const FObjectInitializer& ObjectInitializer);
+
+#if WITH_EDITOR
+	void InitFrom(const FWorldPartitionActorDescInstance* InCustomHLODActorDescInstance);
+	virtual void GetStreamingBounds(FBox& OutRuntimeBounds, FBox& OutEditorBounds) const override;
+	ENGINE_API virtual TUniquePtr<class FWorldPartitionActorDesc> CreateClassActorDesc() const override;
+	ENGINE_API const FGuid& GetCustomHLODActorGuid() const;
+
+private:
+	FGuid CustomHLODActorGuid;
+	const FWorldPartitionActorDescInstance* CustomHLODActorDescInstance;
+#endif
+};
\ No newline at end of file
diff --git a/Engine/Source/Runtime/Engine/Public/WorldPartition/HLOD/CustomHLODPlaceholderActorDesc.h b/Engine/Source/Runtime/Engine/Public/WorldPartition/HLOD/CustomHLODPlaceholderActorDesc.h
new file mode 100644
--- /dev/null
+++ b/Engine/Source/Runtime/Engine/Public/WorldPartition/HLOD/CustomHLODPlaceholderActorDesc.h
@@ -0,0 +1,21 @@
+// Copyright Epic Games, Inc. All Rights Reserved.
+
+#pragma once
+
+#include "WorldPartition/WorldPartitionActorDesc.h"
+
+#if WITH_EDITOR
+class FCustomHLODPlaceholderActorDesc : public FWorldPartitionActorDesc
+{
+public:
+	virtual void Init(const AActor* InActor) override;
+	virtual bool Equals(const FWorldPartitionActorDesc* Other) const override;
+	const FGuid& GetCustomHLODActorGuid() const;
+
+protected:
+	virtual uint32 GetSizeOf() const override;
+	virtual void Serialize(FArchive& Ar) override;
+
+	FGuid CustomHLODActorGuid;
+};
+#endif
\ No newline at end of file
diff --git a/Engine/Source/Runtime/Engine/Public/WorldPartition/HLOD/HLODActor.h b/Engine/Source/Runtime/Engine/Public/WorldPartition/HLOD/HLODActor.h
--- a/Engine/Source/Runtime/Engine/Public/WorldPartition/HLOD/HLODActor.h
+++ b/Engine/Source/Runtime/Engine/Public/WorldPartition/HLOD/HLODActor.h
@@ -48,6 +48,8 @@ public:
 	ENGINE_API virtual const FGuid& GetSourceCellGuid() const override;
 	ENGINE_API virtual bool IsStandalone() const override;
 	ENGINE_API virtual const FGuid& GetStandaloneHLODGuid() const override;
+	ENGINE_API virtual bool IsCustomHLOD() const override;
+	ENGINE_API virtual const FGuid& GetCustomHLODGuid() const override;
 	// ~End IWorldPartitionHLODObject interface
 
 #if WITH_EDITOR
diff --git a/Engine/Source/Runtime/Engine/Public/WorldPartition/HLOD/HLODLayer.h b/Engine/Source/Runtime/Engine/Public/WorldPartition/HLOD/HLODLayer.h
--- a/Engine/Source/Runtime/Engine/Public/WorldPartition/HLOD/HLODLayer.h
+++ b/Engine/Source/Runtime/Engine/Public/WorldPartition/HLOD/HLODLayer.h
@@ -33,7 +33,8 @@ enum class EHLODLayerType : uint8
 	MeshMerge				UMETA(DisplayName = "Merged Mesh"),
 	MeshSimplify			UMETA(DisplayName = "Simplified Mesh"),
 	MeshApproximate			UMETA(DisplayName = "Approximated Mesh"),
-	Custom					UMETA(DisplayName = "Custom"),
+	Custom					UMETA(DisplayName = "Custom Builder"),
+	CustomHLODActor			UMETA(DisplayName = "Custom HLOD Actor"),
 };
 
 UCLASS(Blueprintable, MinimalAPI)
diff --git a/Engine/Source/Runtime/Engine/Public/WorldPartition/HLOD/IWorldPartitionHLODObject.h b/Engine/Source/Runtime/Engine/Public/WorldPartition/HLOD/IWorldPartitionHLODObject.h
--- a/Engine/Source/Runtime/Engine/Public/WorldPartition/HLOD/IWorldPartitionHLODObject.h
+++ b/Engine/Source/Runtime/Engine/Public/WorldPartition/HLOD/IWorldPartitionHLODObject.h
@@ -38,4 +38,10 @@ public:
 
 	/** Returns Standalone HLOD Guid */
 	virtual const FGuid& GetStandaloneHLODGuid() const = 0;
+
+	/** Returns whether the HLOD objects represents custom HLOD */
+	virtual bool IsCustomHLOD() const = 0;
+
+	/** Returns Custom HLOD Guid */
+	virtual const FGuid& GetCustomHLODGuid() const = 0;
 };
\ No newline at end of file
diff --git a/Engine/Source/Runtime/Engine/Public/WorldPartition/RuntimeHashSet/WorldPartitionRuntimeHashSet.h b/Engine/Source/Runtime/Engine/Public/WorldPartition/RuntimeHashSet/WorldPartitionRuntimeHashSet.h
--- a/Engine/Source/Runtime/Engine/Public/WorldPartition/RuntimeHashSet/WorldPartitionRuntimeHashSet.h
+++ b/Engine/Source/Runtime/Engine/Public/WorldPartition/RuntimeHashSet/WorldPartitionRuntimeHashSet.h
@@ -167,6 +167,7 @@ public:
 	ENGINE_API virtual bool Draw2D(FWorldPartitionDraw2DContext& DrawContext) const override;
 	ENGINE_API virtual void Draw3D(const TArray<FWorldPartitionStreamingSource>& Sources) const override;
 	ENGINE_API virtual const FGuid* GetStandaloneHLODActorSourceCellOverride(const FGuid& InActorGuid) const override;
+	ENGINE_API virtual const FGuid* GetCustomHLODActorSourceCellOverride(const FGuid& InActorGuid) const override;
 	ENGINE_API void ForEachHLODLayer(TFunctionRef<bool(FName, FName, int32)> Func) const; // RuntimePartitionName, HLODSetupName, HLODSetupIndex
 	ENGINE_API int32 ComputeHLODHierarchyDepth() const;
 
@@ -236,9 +237,9 @@ private:
 public:
 	ENGINE_API const URuntimePartition* ResolveRuntimePartition(FName GridName, bool bMainPartitionLayer = false) const;
 
-private:
 	const URuntimePartition* ResolveRuntimePartitionForHLODLayer(FName GridName, const FSoftObjectPath& HLODLayerPath) const;
 
+private:
 	void FixupHLODSetup(FRuntimePartitionDesc& RuntimePartition);
 
 	void RemoveIrrelevantCells(FRuntimePartitionStreamingData& StreamingData);
@@ -256,6 +257,9 @@ private:
 	UPROPERTY()
 	TMap<FGuid, FGuid> StandaloneHLODActorToSourceCellsMap;
 
+	UPROPERTY()
+	TMap<FGuid, FGuid> CustomHLODActorToSourceCellsMap;
+
 	// Optimized data
 	TMap<FName, TArray<const FRuntimePartitionStreamingData*>> RuntimeSpatiallyLoadedDataGridMap;
 	TArray<const FRuntimePartitionStreamingData*> RuntimeNonSpatiallyLoadedDataGridList;
diff --git a/Engine/Source/Runtime/Engine/Public/WorldPartition/WorldPartitionRuntimeCell.h b/Engine/Source/Runtime/Engine/Public/WorldPartition/WorldPartitionRuntimeCell.h
--- a/Engine/Source/Runtime/Engine/Public/WorldPartition/WorldPartitionRuntimeCell.h
+++ b/Engine/Source/Runtime/Engine/Public/WorldPartition/WorldPartitionRuntimeCell.h
@@ -350,6 +350,9 @@ class UWorldPartitionRuntimeCell : public UObject, public IWorldPartitionCell, p
 
 	void SetIsHLOD(bool bInIsHLOD) { bIsHLOD = bInIsHLOD; }
 	bool GetIsHLOD() const { return bIsHLOD; }
+
+	void SetIsCustomHLODPlaceholderCell(bool bInIsCustomHLODPlaceholderCell) { bIsCustomHLODPlaceholderCell = bInIsCustomHLODPlaceholderCell; }
+	bool GetIsCustomHLODPlaceholderCell() const { return bIsCustomHLODPlaceholderCell; }
 	
 	const FGuid& GetGuid() const { return CellGuid; }
 	void SetGuid(const FGuid& InCellGuid) { CellGuid = InCellGuid; }
@@ -398,6 +401,9 @@ private:
 	UPROPERTY()
 	bool bIsHLOD;
 
+	UPROPERTY()
+	bool bIsCustomHLODPlaceholderCell;
+
 	UPROPERTY()
 	bool bBlockOnSlowLoading;
 
@@ -444,8 +450,8 @@ public:
 	/** Returns the cell's bounds, which is the uniform size of the cell. */
 	inline FBox GetCellBounds() const { return RuntimeCellData->GetCellBounds(); }
 
-	/** Returns the cell's streaming bounds, which is what the underlying runtime hash uses to intersect cells.  */
-	inline FBox GetStreamingBounds() const { return RuntimeCellData->GetStreamingBounds(); }
+	/** Returns the cell's streaming bounds, which is what the underlying runtime hash uses to intersect cells. If the cell is Custom HLOD Placeholder cell, returns cell bounds. */
+	inline FBox GetStreamingBounds() const { return bIsCustomHLODPlaceholderCell ? RuntimeCellData->GetCellBounds() : RuntimeCellData->GetStreamingBounds(); }
 
 	ENGINE_API virtual bool IsDebugShown() const;
 	//~End UWorldPartitionRuntimeCellData Proxy
diff --git a/Engine/Source/Runtime/Engine/Public/WorldPartition/WorldPartitionRuntimeHash.h b/Engine/Source/Runtime/Engine/Public/WorldPartition/WorldPartitionRuntimeHash.h
--- a/Engine/Source/Runtime/Engine/Public/WorldPartition/WorldPartitionRuntimeHash.h
+++ b/Engine/Source/Runtime/Engine/Public/WorldPartition/WorldPartitionRuntimeHash.h
@@ -280,6 +280,7 @@ public:
 	virtual bool IsStreaming3D() const { return true; }
 	virtual bool GetShouldMergeStreamingSourceInfo() const { return false; }
 	virtual const FGuid* GetStandaloneHLODActorSourceCellOverride(const FGuid& InActorGuid) const { return nullptr; }
+	virtual const FGuid* GetCustomHLODActorSourceCellOverride(const FGuid& InActorGuid) const { return nullptr; }
 
 	static ENGINE_API URuntimeHashExternalStreamingObjectBase* CreateExternalStreamingObject(TSubclassOf<URuntimeHashExternalStreamingObjectBase> InClass, UObject* InOuter, UWorld* InOuterWorld);
 	ENGINE_API UWorldPartitionRuntimeCell* CreateRuntimeCell(UClass* CellClass, UClass* CellDataClass, const FString& CellName, const FString& CellInstanceSuffix, UObject* InOuter = nullptr);
