PR #13561: [Volume] More detail in logging when a volume is missing a brush component to help track down which one

diff --git a/Engine/Source/Runtime/Engine/Private/Volume.cpp b/Engine/Source/Runtime/Engine/Private/Volume.cpp
--- a/Engine/Source/Runtime/Engine/Private/Volume.cpp
+++ b/Engine/Source/Runtime/Engine/Private/Volume.cpp
@@ -111,7 +111,7 @@ bool AVolume::EncompassesPoint(FVector Point, float SphereRadius/*=0.f*/, float*
 	}
 	else
 	{
-		UE_LOG(LogVolume, Log, TEXT("AVolume::EncompassesPoint : No BrushComponent"));
+		UE_LOG(LogVolume, Log, TEXT("AVolume::EncompassesPoint : No BrushComponent on %s"), *GetActorNameOrLabel());
 		return false;
 	}
 }
