[CustomHLOD] Fix missing pointer initialization

diff --git a/Engine/Source/Runtime/Engine/Private/WorldPartition/HLOD/CustomHLODPlaceholderActor.cpp b/Engine/Source/Runtime/Engine/Private/WorldPartition/HLOD/CustomHLODPlaceholderActor.cpp
--- a/Engine/Source/Runtime/Engine/Private/WorldPartition/HLOD/CustomHLODPlaceholderActor.cpp
+++ b/Engine/Source/Runtime/Engine/Private/WorldPartition/HLOD/CustomHLODPlaceholderActor.cpp
@@ -10,6 +10,9 @@
 
 AWorldPartitionCustomHLODPlaceholder::AWorldPartitionCustomHLODPlaceholder(const FObjectInitializer& ObjectInitializer)
 	: Super(ObjectInitializer)
+#if WITH_EDITOR
+	, CustomHLODActorDescInstance(nullptr)
+#endif
 {
 	RootComponent = CreateDefaultSubobject<USceneComponent>(TEXT("Root"));
 }
