Fixed UAssetManager::GatherPublicAssetsForPackage including non-top-level assets in its results

Only top-level assets can be public API, as non-top-level assets (eg, OFPA packages) are baked out during cook

#rnx

diff --git a/Engine/Source/Runtime/Engine/Private/AssetManager.cpp b/Engine/Source/Runtime/Engine/Private/AssetManager.cpp
--- a/Engine/Source/Runtime/Engine/Private/AssetManager.cpp
+++ b/Engine/Source/Runtime/Engine/Private/AssetManager.cpp
@@ -4818,20 +4818,24 @@ void UAssetManager::GatherPublicAssetsForPackage(FName PackagePath, TSet<FSoftOb
 
 	GetAssetRegistry().EnumerateAssets(Filter, [&AssetsToCook](const FAssetData& AssetData)
 	{
-		// this package can be externally referenced; include it in the cook
-		if (!AssetData.PackageName.IsNone())
+		// Only TopLevelAssets can be API exposed
+		if (AssetData.IsTopLevelAsset())
 		{
-			UE_LOG(LogAssetManager, Verbose,
-				TEXT("GatherPublicAssetsForPackage: Adding public package [%s] (instigator: [%s]"),
-				*AssetData.PackageName.ToString(),
-				*AssetData.AssetName.ToString());
-			AssetsToCook.Add(AssetData.GetSoftObjectPath());
-		}
-		else
-		{
-			UE_LOG(LogAssetManager, Error,
-				TEXT("GatherPublicAssetsForPackage: Failed to resolve package for asset [%s]"),
-				*AssetData.AssetName.ToString());
+			// this package can be externally referenced; include it in the cook
+			if (!AssetData.PackageName.IsNone())
+			{
+				UE_LOG(LogAssetManager, Verbose,
+					TEXT("GatherPublicAssetsForPackage: Adding public package [%s] (instigator: [%s]"),
+					*AssetData.PackageName.ToString(),
+					*AssetData.AssetName.ToString());
+				AssetsToCook.Add(AssetData.GetSoftObjectPath());
+			}
+			else
+			{
+				UE_LOG(LogAssetManager, Error,
+					TEXT("GatherPublicAssetsForPackage: Failed to resolve package for asset [%s]"),
+					*AssetData.AssetName.ToString());
+			}
 		}
 
 		return true;
