Decouple IO request from GC. Added an overrideable FastCancel method to FRenderAssetUpdate that is called (behind a CVar) in UStreamableRenderAsset::IsReadyForFinishDestroy.
If a particular asset update wishes to handle this, it needs to decouple its internal async task (IO or similar) by ensuring any callbacks are still able to access the asset update state itself (for e.g. by holding a reference within the callback).

#tests Local tests with induced IO latency, RRs, Presubmit
#rnx

diff --git a/Engine/Source/Runtime/Engine/Private/StreamableRenderAsset.cpp b/Engine/Source/Runtime/Engine/Private/StreamableRenderAsset.cpp
--- a/Engine/Source/Runtime/Engine/Private/StreamableRenderAsset.cpp
+++ b/Engine/Source/Runtime/Engine/Private/StreamableRenderAsset.cpp
@@ -475,6 +475,11 @@ bool UStreamableRenderAsset::IsReadyForFinishDestroy()
 
 	if (PendingUpdate)
 	{
+		if (GEnableAssetUpdateFastCancels)
+		{
+			PendingUpdate->FastCancel();
+		}
+		
 		// To avoid async tasks from timing-out the GC, we tick as Async to force completion if this is relevant.
 		// This could lead the asset from releasing the PendingUpdate, which will be deleted once the async task completes.
 		if (PendingUpdate->GetRelevantThread() == FRenderAssetUpdate::TT_Async)
diff --git a/Engine/Source/Runtime/Engine/Private/Streaming/RenderAssetUpdate.cpp b/Engine/Source/Runtime/Engine/Private/Streaming/RenderAssetUpdate.cpp
--- a/Engine/Source/Runtime/Engine/Private/Streaming/RenderAssetUpdate.cpp
+++ b/Engine/Source/Runtime/Engine/Private/Streaming/RenderAssetUpdate.cpp
@@ -27,6 +27,14 @@ static TAutoConsoleVariable<int32> CVarStreamingStressTestExtraAsyncLatency(
 	TEXT("An extra latency in milliseconds for each async task when doing the stress test."),
 	ECVF_Cheat);
 
+ENGINE_API bool GEnableAssetUpdateFastCancels = false;
+static FAutoConsoleVariableRef CVarEnableAssetUpdateFastCancels(
+	TEXT("r.Streaming.EnableAssetUpdateFastCancels"),
+	GEnableAssetUpdateFastCancels,
+	TEXT("Enable asset update fast cancel support, enables calls to FRenderAssetUpdate::FastCancel which can then be used to optionally decouple Asset Update internals that stall/prevent objects from being GC'd"),
+	ECVF_Default
+);
+
 
 std::atomic<int32> GRenderAssetStreamingSuspension = 0;
 
diff --git a/Engine/Source/Runtime/Engine/Private/Streaming/Texture2DStreamIn_IO.cpp b/Engine/Source/Runtime/Engine/Private/Streaming/Texture2DStreamIn_IO.cpp
--- a/Engine/Source/Runtime/Engine/Private/Streaming/Texture2DStreamIn_IO.cpp
+++ b/Engine/Source/Runtime/Engine/Private/Streaming/Texture2DStreamIn_IO.cpp
@@ -22,6 +22,15 @@ CSV_DECLARE_CATEGORY_EXTERN(TextureStreaming);
 #include "RenderUtils.h"
 #endif
 
+#if !UE_BUILD_SHIPPING
+float GTextureStreamInFastCancelCallbackIOLatency = 0.0f;
+static FAutoConsoleVariableRef CVarTextureStreamInFastCancelCallbackIOLatency(
+	TEXT("r.Streaming.TextureStreamInFastCancelCallbackIOLatency"),
+	GTextureStreamInFastCancelCallbackIOLatency,
+	TEXT("Texture stream in latency (in ms), only applied to the bulk request callkbacks when fast cancels are enabled, useful to induce a case where FastCancel is called while the IO request is in flight"),
+	ECVF_Default);
+#endif // #if !UE_BUILD_SHIPPING
+
 FTexture2DStreamIn_IO::FTexture2DStreamIn_IO(UTexture2D* InTexture, bool InPrioritizedIORequest)
 	: FTexture2DStreamIn(InTexture)
 	, bPrioritizedIORequest(InPrioritizedIORequest)
@@ -114,7 +123,7 @@ void FTexture2DStreamIn_IO::SetIORequests(const FContext& Context)
 
 	TaskSynchronization.Increment();
 
-	Batch.Issue([this, InFlightBulkSize](FBulkDataBatchRequest::EStatus Status)
+	FBulkDataRequest::FCompletionCallback OnBulkDataRequestComplete = [this, InFlightBulkSize](FBulkDataBatchRequest::EStatus Status)
 	{
 		TaskSynchronization.Decrement();
 
@@ -138,8 +147,44 @@ void FTexture2DStreamIn_IO::SetIORequests(const FContext& Context)
 		// The tick here is intended to schedule the success or cancel callback.
 		// Using TT_None ensure gets which could create a dead lock.
 		Tick(FTexture2DUpdate::TT_None);
-	},
-	BatchRequest);
+	};
+
+	if (!GEnableAssetUpdateFastCancels)
+	{
+		Batch.Issue(MoveTemp(OnBulkDataRequestComplete), BatchRequest);
+	}
+	else
+	{
+		// Explicitly add/release ref to keep this object alive as the owning object can release their references if FastCanceled.
+		AddRef();
+
+		Batch.Issue([this, OnBulkDataRequestComplete = MoveTemp(OnBulkDataRequestComplete)](FBulkDataBatchRequest::EStatus Status)
+		{
+#if !UE_BUILD_SHIPPING
+			// This is separate from FRenderAssetStreamingSettings::ExtraIOLatency to make it easier to test fast cancels.
+			// FRenderAssetStreamingSettings::ExtraIOLatency is applied to many parts of the asset update that can in turn stall GC.
+			// Stalled GCs would not trigger the call to FastCancel, testing which is the purpose of this CVar.
+			if (GTextureStreamInFastCancelCallbackIOLatency > 0)
+			{
+				FPlatformProcess::Sleep(GTextureStreamInFastCancelCallbackIOLatency * .001f); // Slow down the streaming.
+			}
+#endif // !UE_BUILD_SHIPPING
+
+			// Handle fast cancels with early out.
+			if (TaskSynchronization.GetValue() <= 0)
+			{
+				check(bIsCancelled);
+				UE_LOG(LogContentStreaming, Warning, TEXT("Texture stream-in request was FastCanceled"));
+			}
+			else
+			{
+				OnBulkDataRequestComplete(Status);
+			}
+
+			Release();
+		},
+		BatchRequest);
+	}
 }
 
 void FTexture2DStreamIn_IO::CancelIORequests()
@@ -185,6 +230,16 @@ void FTexture2DStreamIn_IO::Abort()
 	}
 }
 
+void FTexture2DStreamIn_IO::FastCancel()
+{
+	check(GEnableAssetUpdateFastCancels);
+	if (!IsCompleted() && IsCancelled() && BatchRequest.IsPending() && TaskSynchronization.GetValue() > 0)
+	{
+		TaskSynchronization.Set(0);
+		UE_LOG(LogContentStreaming, Warning, TEXT("Decoupled FTexture2DStreamIn_IO from pending IO request in order to unblock GC"));
+	}
+}
+
 void FTexture2DStreamIn_IO::FCancelIORequestsTask::DoWork()
 {
 	check(PendingUpdate);
diff --git a/Engine/Source/Runtime/Engine/Private/Streaming/Texture2DStreamIn_IO.h b/Engine/Source/Runtime/Engine/Private/Streaming/Texture2DStreamIn_IO.h
--- a/Engine/Source/Runtime/Engine/Private/Streaming/Texture2DStreamIn_IO.h
+++ b/Engine/Source/Runtime/Engine/Private/Streaming/Texture2DStreamIn_IO.h
@@ -39,6 +39,8 @@ protected:
 	// Start an async task to cancel pending IO requests.
 	void Abort() override;
 
+	void FastCancel() override;
+
 private:
 
 	class FCancelIORequestsTask : public FNonAbandonableTask
diff --git a/Engine/Source/Runtime/Engine/Public/RenderAssetUpdate.h b/Engine/Source/Runtime/Engine/Public/RenderAssetUpdate.h
--- a/Engine/Source/Runtime/Engine/Public/RenderAssetUpdate.h
+++ b/Engine/Source/Runtime/Engine/Public/RenderAssetUpdate.h
@@ -21,6 +21,8 @@ class UStreamableRenderAsset;
 
 ENGINE_API bool IsAssetStreamingSuspended();
 
+extern ENGINE_API bool GEnableAssetUpdateFastCancels;
+
 /**
 * This class provides a framework for loading and unloading the texture/mesh LODs.
 * Each thread essentially calls Tick() until the job is done.
@@ -75,6 +77,14 @@ public:
 		MarkAsCancelled();
 	}
 
+	/**
+	* Allows the update to decouple any internal tasks it is awaiting on in preparation to be garbage collected.
+	* This is called when we query IsReadyForFinishDestroy on the UStreamableRenderAsset.
+	*/
+	virtual void FastCancel()
+	{	
+	}
+
 	/** Returns whether the task was aborted through Abort() or cancelled.  */
 	bool IsCancelled() const
 	{
