FObjectReplicator's CallRepNotifies method is now const.

diff --git a/Engine/Source/Runtime/Engine/Private/DataReplication.cpp b/Engine/Source/Runtime/Engine/Private/DataReplication.cpp
--- a/Engine/Source/Runtime/Engine/Private/DataReplication.cpp
+++ b/Engine/Source/Runtime/Engine/Private/DataReplication.cpp
@@ -2423,7 +2423,7 @@ void FObjectReplicator::StartBecomingDormant()
 	bLastUpdateEmpty = false; // Ensure we get one more attempt to update properties
 }
 
-void FObjectReplicator::CallRepNotifies(bool bSkipIfChannelHasQueuedBunches)
+void FObjectReplicator::CallRepNotifies(bool bSkipIfChannelHasQueuedBunches) const
 {
 	// This logic is mostly a copy of FRepLayout::CallRepNotifies, and they should be merged.
 
diff --git a/Engine/Source/Runtime/Engine/Public/Net/DataReplication.h b/Engine/Source/Runtime/Engine/Public/Net/DataReplication.h
--- a/Engine/Source/Runtime/Engine/Public/Net/DataReplication.h
+++ b/Engine/Source/Runtime/Engine/Public/Net/DataReplication.h
@@ -218,7 +218,7 @@ public:
 
 	ENGINE_API void StartBecomingDormant();
 
-	ENGINE_API void CallRepNotifies(bool bSkipIfChannelHasQueuedBunches);
+	ENGINE_API void CallRepNotifies(bool bSkipIfChannelHasQueuedBunches) const;
 
 	ENGINE_API void UpdateUnmappedObjects(bool& bOutHasMoreUnmapped);
 
