Fixed crash which could occur if an actor channel was removed during replay playback scrubbing.

diff --git a/Engine/Source/Runtime/Engine/Private/DemoNetDriver.cpp b/Engine/Source/Runtime/Engine/Private/DemoNetDriver.cpp
--- a/Engine/Source/Runtime/Engine/Private/DemoNetDriver.cpp
+++ b/Engine/Source/Runtime/Engine/Private/DemoNetDriver.cpp
@@ -3052,11 +3052,21 @@ void UDemoNetDriver::FinalizeFastForward(const double StartTime)
 	// Flush all pending RepNotifies that were built up during the fast-forward.
 	if (ServerConnection != nullptr)
 	{
-		for (auto& ChannelPair : ServerConnection->ActorChannelMap())
+		// OnReps can cause the channel map to be modified. Make a copy and lookup the object in the active channel map.
+		const FActorChannelMap& ChannelMap = ServerConnection->ActorChannelMap();
+		const FActorChannelMap ChannelMapCopy = ChannelMap;
+		for (const TPair<TWeakObjectPtr<AActor>, UActorChannel*>& ChannelPair : ChannelMapCopy)
 		{
-			if (ChannelPair.Value != nullptr)
+			const TWeakObjectPtr<AActor>& Actor = ChannelPair.Key;
+			if (!Actor.IsValid())
 			{
-				for (auto& ReplicatorPair : ChannelPair.Value->ReplicationMap)
+				continue;
+			}
+
+			// Look up the actor in the original channel map
+			if (const UActorChannel*const* Channel = ChannelMap.Find(Actor))
+			{
+				for (const TPair<UObject*, TSharedRef<FObjectReplicator>>& ReplicatorPair : (*Channel)->ReplicationMap)
 				{
 					ReplicatorPair.Value->CallRepNotifies(true);
 				}
