Added initialization of TexCoords variables in generated material shader to circumvent a DXC issue concerning uninitialized variables.

#rnx

diff --git a/Engine/Source/Runtime/Engine/Private/Materials/HLSLMaterialTranslator.cpp b/Engine/Source/Runtime/Engine/Private/Materials/HLSLMaterialTranslator.cpp
--- a/Engine/Source/Runtime/Engine/Private/Materials/HLSLMaterialTranslator.cpp
+++ b/Engine/Source/Runtime/Engine/Private/Materials/HLSLMaterialTranslator.cpp
@@ -2154,7 +2154,8 @@ void FHLSLMaterialTranslator::TranslateMaterial()
 	check(SavedNumUserTexCoords == GetNumUserTexCoords());
 
 	FString InterpolatorsOffsetsDefinitionCode;
-	TBitArray<> FinalAllocatedCoords = GetVertexInterpolatorsOffsets(InterpolatorsOffsetsDefinitionCode);
+	TArray<uint32> Dummmy;
+	TBitArray<> FinalAllocatedCoords = GetVertexInterpolatorsOffsets(InterpolatorsOffsetsDefinitionCode, Dummmy);
 
 	// Finished compilation, verify final interpolator count restrictions
 	if (CurrentCustomVertexInterpolatorOffset > 0)
@@ -3103,14 +3104,14 @@ void FHLSLMaterialTranslator::GetMaterialEnvironment(EShaderPlatform InPlatform,
 }
 
 // Assign custom interpolators to slots, packing them as much as possible in unused slots.
-TBitArray<> FHLSLMaterialTranslator::GetVertexInterpolatorsOffsets(FString& VertexInterpolatorsOffsetsDefinitionCode) const
+TBitArray<> FHLSLMaterialTranslator::GetVertexInterpolatorsOffsets(FString& VertexInterpolatorsOffsetsDefinitionCode, TArray<uint32>& UsedTextureCoordsIndices) const
 {
 	TBitArray<> AllocatedCoords = AllocatedUserTexCoords; // Don't mess with the already assigned sets of UV coords
 
 	int32 CurrentSlot = INDEX_NONE;
 	int32 EndAllocatedSlot = INDEX_NONE;
 
-	auto GetNextUVSlot = [&CurrentSlot, &EndAllocatedSlot, &AllocatedCoords]() -> int32
+	auto GetNextUVSlot = [&CurrentSlot, &EndAllocatedSlot, &AllocatedCoords, &UsedTextureCoordsIndices]() -> int32
 	{
 		if (CurrentSlot == EndAllocatedSlot)
 		{
@@ -3129,6 +3130,8 @@ TBitArray<> FHLSLMaterialTranslator::GetVertexInterpolatorsOffsets(FString& Vert
 		int32 ResultUVSlot = CurrentSlot / 2;
 		CurrentSlot++;
 
+		UsedTextureCoordsIndices.AddUnique(ResultUVSlot);
+
 		return ResultUVSlot;
 	};
 
@@ -17082,8 +17085,8 @@ void FHLSLMaterialTranslator::PrepareMaterialSourceStringParameters()
 
 	// Assign slots to vertex interpolators
 	FString VertexInterpolatorsOffsetsDefinition;
-	TBitArray<> FinalAllocatedCoords = GetVertexInterpolatorsOffsets(VertexInterpolatorsOffsetsDefinition);
-
+	TArray<uint32> UsedTexCoordIndices;
+	TBitArray<> FinalAllocatedCoords = GetVertexInterpolatorsOffsets(VertexInterpolatorsOffsetsDefinition, UsedTexCoordIndices);
 	const uint32 NumUserVertexTexCoords = GetNumUserVertexTexCoords();
 	const uint32 NumUserTexCoords = GetNumUserTexCoords();
 	const uint32 NumCustomVectors = FMath::DivideAndRoundUp((uint32)CurrentCustomVertexInterpolatorOffset, 2u);
@@ -17223,6 +17226,11 @@ void FHLSLMaterialTranslator::PrepareMaterialSourceStringParameters()
 	// Print custom vertex shader interpolator assignments
 	FString CustomInterpolatorAssignments;
 
+	for (uint32 InterpolatorIndex : UsedTexCoordIndices)
+	{
+		CustomInterpolatorAssignments.Appendf(TEXT("\tOutTexCoords[%i] = 0.0f.xx;") HLSL_LINE_TERMINATOR, InterpolatorIndex);
+	}
+
 	for (UMaterialExpressionVertexInterpolator* Interpolator : CustomVertexInterpolators)
 	{
 		if (Interpolator->InterpolatorOffset != INDEX_NONE)
@@ -17236,19 +17244,19 @@ void FHLSLMaterialTranslator::PrepareMaterialSourceStringParameters()
 			const int32 Index = Interpolator->InterpolatorIndex;
 
 			// Note: We reference the UV define directly to avoid having to pre-accumulate UV counts before property translation
-			CustomInterpolatorAssignments += FString::Printf(TEXT("\tOutTexCoords[VERTEX_INTERPOLATOR_%i_TEXCOORDS_X].%s = VertexInterpolator%i(Parameters).x;") HLSL_LINE_TERMINATOR, Index, Swizzle[Offset % 2], Index);
+			CustomInterpolatorAssignments.Appendf(TEXT("\tOutTexCoords[VERTEX_INTERPOLATOR_%i_TEXCOORDS_X].%s = VertexInterpolator%i(Parameters).x;") HLSL_LINE_TERMINATOR, Index, Swizzle[Offset % 2], Index);
 
 			if (Type >= MCT_Float2)
 			{
-				CustomInterpolatorAssignments += FString::Printf(TEXT("\tOutTexCoords[VERTEX_INTERPOLATOR_%i_TEXCOORDS_Y].%s = VertexInterpolator%i(Parameters).y;") HLSL_LINE_TERMINATOR, Index, Swizzle[(Offset + 1) % 2], Index);
+				CustomInterpolatorAssignments.Appendf(TEXT("\tOutTexCoords[VERTEX_INTERPOLATOR_%i_TEXCOORDS_Y].%s = VertexInterpolator%i(Parameters).y;") HLSL_LINE_TERMINATOR, Index, Swizzle[(Offset + 1) % 2], Index);
 
 				if (Type >= MCT_Float3)
 				{
-					CustomInterpolatorAssignments += FString::Printf(TEXT("\tOutTexCoords[VERTEX_INTERPOLATOR_%i_TEXCOORDS_Z].%s = VertexInterpolator%i(Parameters).z;") HLSL_LINE_TERMINATOR, Index, Swizzle[(Offset + 2) % 2], Index);
+					CustomInterpolatorAssignments.Appendf(TEXT("\tOutTexCoords[VERTEX_INTERPOLATOR_%i_TEXCOORDS_Z].%s = VertexInterpolator%i(Parameters).z;") HLSL_LINE_TERMINATOR, Index, Swizzle[(Offset + 2) % 2], Index);
 
 					if (Type == MCT_Float4)
 					{
-						CustomInterpolatorAssignments += FString::Printf(TEXT("\tOutTexCoords[VERTEX_INTERPOLATOR_%i_TEXCOORDS_W].%s = VertexInterpolator%i(Parameters).w;") HLSL_LINE_TERMINATOR, Index, Swizzle[(Offset + 3) % 2], Index);
+						CustomInterpolatorAssignments.Appendf(TEXT("\tOutTexCoords[VERTEX_INTERPOLATOR_%i_TEXCOORDS_W].%s = VertexInterpolator%i(Parameters).w;") HLSL_LINE_TERMINATOR, Index, Swizzle[(Offset + 3) % 2], Index);
 					}
 				}
 			}
diff --git a/Engine/Source/Runtime/Engine/Private/Materials/HLSLMaterialTranslator.h b/Engine/Source/Runtime/Engine/Private/Materials/HLSLMaterialTranslator.h
--- a/Engine/Source/Runtime/Engine/Private/Materials/HLSLMaterialTranslator.h
+++ b/Engine/Source/Runtime/Engine/Private/Materials/HLSLMaterialTranslator.h
@@ -685,7 +685,7 @@ public:
 	void GetMaterialEnvironment(EShaderPlatform InPlatform, FShaderCompilerEnvironment& OutEnvironment);
 	
 	// Assign custom interpolators to slots, packing them as much as possible in unused slots.
-	TBitArray<> GetVertexInterpolatorsOffsets(FString& VertexInterpolatorsOffsetsDefinitionCode) const;
+	TBitArray<> GetVertexInterpolatorsOffsets(FString& VertexInterpolatorsOffsetsDefinitionCode, TArray<uint32>& UsedTexCoordIndices) const;
 
 	void GetSharedInputsMaterialCode(FString& PixelMembersDeclaration, FString& NormalAssignment, FString& PixelMembersInitializationEpilog, ECompiledPartialDerivativeVariation DerivativeVariation);
 
