Change LevelBounds to use new tick options in FTickableGameObject, which fixes an ensure in async loading and cleans up the editor-only logic

diff --git a/Engine/Source/Runtime/Engine/Classes/Engine/LevelBounds.h b/Engine/Source/Runtime/Engine/Classes/Engine/LevelBounds.h
--- a/Engine/Source/Runtime/Engine/Classes/Engine/LevelBounds.h
+++ b/Engine/Source/Runtime/Engine/Classes/Engine/LevelBounds.h
@@ -9,17 +9,6 @@
 #include "LevelBounds.generated.h"
 
 /**
- *  Ticks only in the editor, regardless of viewport 'Realtime' option
- */
-class FEditorTickableLevelBounds 
-#if WITH_EDITOR
-	: public FTickableGameObject
-#endif
-{
-};
-
-/**
- *
  * Defines level bounds
  * Updates bounding box automatically based on actors transformation changes or holds fixed user defined bounding box
  * Uses only actors where AActor::IsLevelBoundsRelevant() == true
@@ -27,9 +16,12 @@ class FEditorTickableLevelBounds
 UCLASS(hidecategories=(Advanced, Collision, Display, Rendering, Physics, Input), showcategories=("Input|MouseInput", "Input|TouchInput"), MinimalAPI)
 class ALevelBounds
 	: public AActor
-	, public FEditorTickableLevelBounds 
+	, public FTickableGameObject 
 {
-	GENERATED_UCLASS_BODY()
+	GENERATED_BODY()
+
+public:
+	ENGINE_API ALevelBounds(const FObjectInitializer& ObjectInitializer);
 		
 	/** Bounding box for the level bounds. */
 	UPROPERTY(EditAnywhere, Category = LevelBounds)
@@ -70,11 +62,10 @@ class ALevelBounds
 	virtual bool CanChangeIsSpatiallyLoadedFlag() const override { return false; }
 	
 protected:
-	/** FTickableGameObject interface */
+	/** FTickableGameObject interface, this ticks only in the editor regardless of viewport 'Realtime' option */
 	ENGINE_API virtual void Tick(float DeltaTime) override;
 	virtual UWorld* GetTickableGameObjectWorld() const override { return GetWorld(); }
 	ENGINE_API virtual TStatId GetStatId() const override;
-	ENGINE_API virtual ETickableTickType GetTickableTickType() const override;
 	ENGINE_API virtual bool IsTickable() const override;
 	ENGINE_API virtual bool IsTickableInEditor() const override;
 	
@@ -109,5 +100,10 @@ protected:
 
 private:
 	virtual bool ActorTypeIsMainWorldOnly() const override { return true; }
-#endif
+
+#else
+	// This never actually ticks as it starts disabled
+	virtual void Tick(float DeltaTime) override {}
+	virtual TStatId GetStatId() const override { return TStatId(); }
+#endif // WITH_EDITOR
 };
diff --git a/Engine/Source/Runtime/Engine/Private/LevelBounds.cpp b/Engine/Source/Runtime/Engine/Private/LevelBounds.cpp
--- a/Engine/Source/Runtime/Engine/Private/LevelBounds.cpp
+++ b/Engine/Source/Runtime/Engine/Private/LevelBounds.cpp
@@ -14,6 +14,7 @@ static const FVector DefaultLevelSize = FVector(1000.f);
 
 ALevelBounds::ALevelBounds(const FObjectInitializer& ObjectInitializer)
 	: Super(ObjectInitializer)
+	, FTickableGameObject(ETickableTickType::Never)
 {
 	BoxComponent = CreateDefaultSubobject<UBoxComponent>(TEXT("BoxComponent0"));
 	RootComponent = BoxComponent;
@@ -138,6 +139,8 @@ void ALevelBounds::Tick(float DeltaTime)
 	{
 		UpdateLevelBounds();
 		bLevelBoundsDirty = false;
+
+		SetTickableTickType(ETickableTickType::Never);
 	}
 }
 
@@ -146,11 +149,6 @@ TStatId ALevelBounds::GetStatId() const
 	RETURN_QUICK_DECLARE_CYCLE_STAT(ALevelBounds, STATGROUP_Tickables);
 }
 
-ETickableTickType ALevelBounds::GetTickableTickType() const
-{
-	return ((GIsEditor && !IsTemplate()) ? ETickableTickType::Conditional : ETickableTickType::Never);
-}
-
 bool ALevelBounds::IsTickable() const
 {
 	if (!IsValidChecked(this) || HasAnyFlags(RF_BeginDestroyed|RF_FinishDestroyed))
@@ -213,6 +211,11 @@ void ALevelBounds::UpdateLevelBounds()
 void ALevelBounds::MarkLevelBoundsDirty()
 {
 	bLevelBoundsDirty = true;
+
+	if (GIsEditor)
+	{
+		SetTickableTickType(ETickableTickType::Conditional);
+	}
 }
 
 bool ALevelBounds::IsUsingDefaultBounds() const
