Fix calculation of virtual texture producer name hashes.
We were hashing an FName which is not guaranteed to be case consistent between dev and shipping builds.
The inconsistency meant that the cinematic prestreaming data couldn't resolve which runtime virtual textures to prestream.
Also fix some existing static validator warnings about potential Texture2DResource nullptr.

diff --git a/Engine/Source/Runtime/Engine/Private/Texture2D.cpp b/Engine/Source/Runtime/Engine/Private/Texture2D.cpp
--- a/Engine/Source/Runtime/Engine/Private/Texture2D.cpp
+++ b/Engine/Source/Runtime/Engine/Private/Texture2D.cpp
@@ -1552,6 +1552,7 @@ FVirtualTexture2DResource::FVirtualTexture2DResource(const UTexture2D* InOwner,
 
 	TextureName = InOwner->GetFName();
 	PackageName = InOwner->GetOutermost()->GetFName();
+	FullNameHash = GetTypeHash(InOwner->GetName());
 
 	UTextureLODSettings* TextureLODSettings = UDeviceProfileManager::Get().GetActiveProfile()->GetTextureLODSettings();
 	Filter = (ESamplerFilter)TextureLODSettings->GetSamplerFilter(InOwner);
@@ -1635,7 +1636,7 @@ void FVirtualTexture2DResource::InitRHI(FRHICommandListBase& RHICmdList)
 
 	FVTProducerDescription ProducerDesc;
 	ProducerDesc.Name = TextureName;
-	ProducerDesc.FullNameHash = GetTypeHash(TextureName);
+	ProducerDesc.FullNameHash = FullNameHash;
 	ProducerDesc.Dimensions = 2;
 	ProducerDesc.TileSize = VTData->TileSize;
 	ProducerDesc.TileBorderSize = VTData->TileBorderSize;
@@ -2000,7 +2001,7 @@ bool UTexture2D::StreamIn(int32 NewMipCount, bool bHighPrio)
 			if (HasDerivedData(this))
 			{
 				// If the future texture is to be a virtual texture, use the virtual stream in path.
-				if (Texture2DResource->bUsePartiallyResidentMips)
+				if (Texture2DResource && Texture2DResource->bUsePartiallyResidentMips)
 				{
 					PendingUpdate = new UE::FTexture2DStreamIn_DerivedData_Virtual(this, bHighPrio);
 				}
@@ -2018,7 +2019,7 @@ bool UTexture2D::StreamIn(int32 NewMipCount, bool bHighPrio)
 			else
 			{
 				// If the future texture is to be a virtual texture, use the virtual stream in path.
-				if (Texture2DResource->bUsePartiallyResidentMips)
+				if (Texture2DResource && Texture2DResource->bUsePartiallyResidentMips)
 				{
 					PendingUpdate = new FTexture2DStreamIn_IO_Virtual(this, bHighPrio);
 				}
@@ -2050,7 +2051,7 @@ bool UTexture2D::StreamIn(int32 NewMipCount, bool bHighPrio)
 				DefaultMipDataProvider = new FTexture2DMipDataProvider_IO(this, bHighPrio);
 			}
 
-			if (Texture2DResource->bUsePartiallyResidentMips)
+			if (Texture2DResource && Texture2DResource->bUsePartiallyResidentMips)
 			{
 				MipAllocator = new FTexture2DMipAllocator_Virtual(this);
 			}
@@ -2090,7 +2091,7 @@ bool UTexture2D::StreamOut(int32 NewMipCount)
 	FTexture2DResource* Texture2DResource = TextureResource ? TextureResource->GetTexture2DResource() : nullptr;
 	if (!HasPendingInitOrStreaming() && CachedSRRState.StreamOut(NewMipCount) && ensure(Texture2DResource))
 	{
-		if (Texture2DResource->bUsePartiallyResidentMips)
+		if (Texture2DResource && Texture2DResource->bUsePartiallyResidentMips)
 		{
 			PendingUpdate = new FTexture2DStreamOut_Virtual(this);
 		}
diff --git a/Engine/Source/Runtime/Engine/Public/TextureResource.h b/Engine/Source/Runtime/Engine/Public/TextureResource.h
--- a/Engine/Source/Runtime/Engine/Public/TextureResource.h
+++ b/Engine/Source/Runtime/Engine/Public/TextureResource.h
@@ -224,6 +224,8 @@ protected:
 	FName TextureName;
 	/** The FName of the texture package for stats */
 	FName PackageName;
+	/** A hash of the texture asset name */
+	uint32 FullNameHash = 0;
 	/** Cached sampler config */
 	TEnumAsByte<ESamplerFilter> Filter = SF_Point;
 	TEnumAsByte<ESamplerAddressMode> AddressU = AM_Wrap;
