[StandaloneHLOD] Fix crash when adding a standalone LI with standalone HLOD to a non-WP world

diff --git a/Engine/Source/Runtime/Engine/Private/WorldPartition/HLOD/StandaloneHLODSubsystem.cpp b/Engine/Source/Runtime/Engine/Private/WorldPartition/HLOD/StandaloneHLODSubsystem.cpp
--- a/Engine/Source/Runtime/Engine/Private/WorldPartition/HLOD/StandaloneHLODSubsystem.cpp
+++ b/Engine/Source/Runtime/Engine/Private/WorldPartition/HLOD/StandaloneHLODSubsystem.cpp
@@ -291,7 +291,7 @@ void UWorldPartitionStandaloneHLODSubsystem::UpdateStandaloneHLODActors(UWorldPa
 			SpawnParams.Name = FName(ActorLabel);
 			SpawnParams.NameMode = FActorSpawnParameters::ESpawnActorNameMode::Requested;
 
-			UWorldPartitionRuntimeHashSet* WorldHash = Cast<UWorldPartitionRuntimeHashSet>(GetWorld()->GetWorldPartition()->RuntimeHash);
+			UWorldPartitionRuntimeHashSet* WorldHash = GetWorld()->GetWorldPartition() ? Cast<UWorldPartitionRuntimeHashSet>(GetWorld()->GetWorldPartition()->RuntimeHash) : nullptr;
 			if (!WorldHash)
 			{
 				UE_LOG(LogStandaloneHLODSubsystem, Log, TEXT("RuntimeHashSet not found in World %s"), *GetWorld()->GetName());
