Iris Beta
- Removed #if UE_WITH_IRIS for Engine/...

diff --git a/Engine/Plugins/Experimental/DaySequence/Source/DaySequence/Private/DaySequencePlayer.cpp b/Engine/Plugins/Experimental/DaySequence/Source/DaySequence/Private/DaySequencePlayer.cpp
--- a/Engine/Plugins/Experimental/DaySequence/Source/DaySequence/Private/DaySequencePlayer.cpp
+++ b/Engine/Plugins/Experimental/DaySequence/Source/DaySequence/Private/DaySequencePlayer.cpp
@@ -22,9 +22,7 @@
 #include "GameFramework/WorldSettings.h"
 #include "GameFramework/PlayerController.h"
 #include "GameFramework/PlayerState.h"
-#if UE_WITH_IRIS
 #include "Iris/ReplicationSystem/ReplicationFragmentUtil.h"
-#endif
 #include "Net/UnrealNetwork.h"
 #include "Engine/NetDriver.h"
 
@@ -1919,12 +1917,10 @@ bool UDaySequencePlayer::CallRemoteFunction(UFunction* Function, void* Parameter
 	return false;
 }
 
-#if UE_WITH_IRIS
 void UDaySequencePlayer::RegisterReplicationFragments(UE::Net::FFragmentRegistrationContext& Context, UE::Net::EFragmentRegistrationFlags RegistrationFlags)
 {
 	UE::Net::FReplicationFragmentUtil::CreateAndRegisterFragmentsForObject(this, Context, RegistrationFlags);
 }
-#endif
 
 bool UDaySequencePlayer::NeedsQueueLatentAction() const
 {
diff --git a/Engine/Plugins/Experimental/DaySequence/Source/DaySequence/Private/DaySequencePlayer.h b/Engine/Plugins/Experimental/DaySequence/Source/DaySequence/Private/DaySequencePlayer.h
--- a/Engine/Plugins/Experimental/DaySequence/Source/DaySequence/Private/DaySequencePlayer.h
+++ b/Engine/Plugins/Experimental/DaySequence/Source/DaySequence/Private/DaySequencePlayer.h
@@ -413,9 +413,7 @@ protected:
 	UE_API virtual bool CallRemoteFunction(UFunction* Function, void* Parameters, FOutParmRec* OutParms, FFrame* Stack) override;
 	UE_API virtual void PostNetReceive() override;
 	UE_API virtual void BeginDestroy() override;
-#if UE_WITH_IRIS
 	UE_API virtual void RegisterReplicationFragments(UE::Net::FFragmentRegistrationContext& Context, UE::Net::EFragmentRegistrationFlags RegistrationFlags) override;
-#endif
 	/*~ End UObject interface */
 	
 	bool CanPlay() const { return WeakOwner.IsValid(); }
diff --git a/Engine/Plugins/Experimental/Iris/Source/Iris/Private/Iris/IrisModule.cpp b/Engine/Plugins/Experimental/Iris/Source/Iris/Private/Iris/IrisModule.cpp
--- a/Engine/Plugins/Experimental/Iris/Source/Iris/Private/Iris/IrisModule.cpp
+++ b/Engine/Plugins/Experimental/Iris/Source/Iris/Private/Iris/IrisModule.cpp
@@ -7,9 +7,7 @@ class FIrisModule : public IModuleInterface
 private:
 	virtual void StartupModule() override
 	{
-#if UE_WITH_IRIS
 		FModuleManager::Get().LoadModule("IrisCore", ELoadModuleFlags::None);
-#endif
 	}
 };
 
diff --git a/Engine/Plugins/Online/OnlineSubsystemUtils/Source/OnlineSubsystemUtils/Private/OnlineBeacon.cpp b/Engine/Plugins/Online/OnlineSubsystemUtils/Source/OnlineSubsystemUtils/Private/OnlineBeacon.cpp
--- a/Engine/Plugins/Online/OnlineSubsystemUtils/Source/OnlineSubsystemUtils/Private/OnlineBeacon.cpp
+++ b/Engine/Plugins/Online/OnlineSubsystemUtils/Source/OnlineSubsystemUtils/Private/OnlineBeacon.cpp
@@ -4,9 +4,7 @@
 #include "Engine/Channel.h"
 #include "Engine/NetConnection.h"
 #include "Engine/Engine.h"
-#if UE_WITH_IRIS
 #include "Iris/ReplicationSystem/ReplicationSystem.h"
-#endif // UE_WITH_IRIS
 
 #include UE_INLINE_GENERATED_CPP_BY_NAME(OnlineBeacon)
 
@@ -98,7 +96,6 @@ void AOnlineBeacon::OnActorChannelOpen(FInBunch& Bunch, UNetConnection* Connecti
 	Connection->OwningActor = this;
 	Super::OnActorChannelOpen(Bunch, Connection);
 
-#if UE_WITH_IRIS
 	{
 		// Enable replication
 		if (UReplicationSystem* ReplicationSystem = Connection->Driver->GetReplicationSystem())
@@ -106,7 +103,6 @@ void AOnlineBeacon::OnActorChannelOpen(FInBunch& Bunch, UNetConnection* Connecti
 			ReplicationSystem->SetReplicationEnabledForConnection(Connection->GetConnectionHandle().GetParentConnectionId(), true);
 		}
 	}
-#endif // UE_WITH_IRIS
 }
 
 bool AOnlineBeacon::IsRelevancyOwnerFor(const AActor* ReplicatedActor, const AActor* ActorOwner, const AActor* ConnectionActor) const
diff --git a/Engine/Plugins/Runtime/GameplayAbilities/Source/GameplayAbilities/Private/Abilities/GameplayAbility.cpp b/Engine/Plugins/Runtime/GameplayAbilities/Source/GameplayAbilities/Private/Abilities/GameplayAbility.cpp
--- a/Engine/Plugins/Runtime/GameplayAbilities/Source/GameplayAbilities/Private/Abilities/GameplayAbility.cpp
+++ b/Engine/Plugins/Runtime/GameplayAbilities/Source/GameplayAbilities/Private/Abilities/GameplayAbility.cpp
@@ -21,9 +21,7 @@
 
 #define LOCTEXT_NAMESPACE "GameplayAbility"
 
-#if UE_WITH_IRIS
 #include "Iris/ReplicationSystem/ReplicationFragmentUtil.h"
-#endif // UE_WITH_IRIS
 
 #include UE_INLINE_GENERATED_CPP_BY_NAME(GameplayAbility)
 
@@ -2321,12 +2319,10 @@ void UGameplayAbility::GetLifetimeReplicatedProps(TArray< class FLifetimePropert
 	}
 }
 
-#if UE_WITH_IRIS
 void UGameplayAbility::RegisterReplicationFragments(UE::Net::FFragmentRegistrationContext& Context, UE::Net::EFragmentRegistrationFlags RegistrationFlags)
 {
 	// Build descriptors and allocate PropertyReplicationFragments for this object
 	UE::Net::FReplicationFragmentUtil::CreateAndRegisterFragmentsForObject(this, Context, RegistrationFlags);
 }
-#endif // UE_WITH_IRIS
 
 #undef LOCTEXT_NAMESPACE
\ No newline at end of file
diff --git a/Engine/Plugins/Runtime/GameplayAbilities/Source/GameplayAbilities/Private/AbilitySystemGlobals.cpp b/Engine/Plugins/Runtime/GameplayAbilities/Source/GameplayAbilities/Private/AbilitySystemGlobals.cpp
--- a/Engine/Plugins/Runtime/GameplayAbilities/Source/GameplayAbilities/Private/AbilitySystemGlobals.cpp
+++ b/Engine/Plugins/Runtime/GameplayAbilities/Source/GameplayAbilities/Private/AbilitySystemGlobals.cpp
@@ -19,11 +19,9 @@
 #include "Modules/ModuleManager.h"
 #include "UObject/UObjectIterator.h"
 
-#if UE_WITH_IRIS
 #include "Serialization/GameplayAbilityTargetDataHandleNetSerializer.h"
 #include "Serialization/GameplayEffectContextHandleNetSerializer.h"
 #include "Serialization/PredictionKeyNetSerializer.h"
-#endif // UE_WITH_IRIS
 
 #if WITH_EDITOR
 #include "Editor.h"
diff --git a/Engine/Plugins/Runtime/GameplayAbilities/Source/GameplayAbilities/Private/AttributeSet.cpp b/Engine/Plugins/Runtime/GameplayAbilities/Source/GameplayAbilities/Private/AttributeSet.cpp
--- a/Engine/Plugins/Runtime/GameplayAbilities/Source/GameplayAbilities/Private/AttributeSet.cpp
+++ b/Engine/Plugins/Runtime/GameplayAbilities/Source/GameplayAbilities/Private/AttributeSet.cpp
@@ -16,9 +16,7 @@
 #include "Net/Core/PushModel/PushModel.h"
 #include "UObject/UObjectThreadContext.h"
 
-#if UE_WITH_IRIS
 #include "Iris/ReplicationSystem/ReplicationFragmentUtil.h"
-#endif // UE_WITH_IRIS
 
 #include UE_INLINE_GENERATED_CPP_BY_NAME(AttributeSet)
 
@@ -427,7 +425,6 @@ void UAttributeSet::SetNetAddressable()
 	bNetAddressable = true;
 }
 
-#if UE_WITH_IRIS
 void UAttributeSet::RegisterReplicationFragments(UE::Net::FFragmentRegistrationContext& Context, UE::Net::EFragmentRegistrationFlags RegistrationFlags)
 {
 	using namespace UE::Net;
@@ -435,7 +432,6 @@ void UAttributeSet::RegisterReplicationFragments(UE::Net::FFragmentRegistrationC
 	// Build descriptors and allocate PropertyReplicationFragments for this object
 	FReplicationFragmentUtil::CreateAndRegisterFragmentsForObject(this, Context, RegistrationFlags);
 }
-#endif // UE_WITH_IRIS
 
 void UAttributeSet::InitFromMetaDataTable(const UDataTable* DataTable)
 {
diff --git a/Engine/Plugins/Runtime/GameplayAbilities/Source/GameplayAbilities/Private/Serialization/GameplayAbilityRepAnimMontageNetSerializer.cpp b/Engine/Plugins/Runtime/GameplayAbilities/Source/GameplayAbilities/Private/Serialization/GameplayAbilityRepAnimMontageNetSerializer.cpp
--- a/Engine/Plugins/Runtime/GameplayAbilities/Source/GameplayAbilities/Private/Serialization/GameplayAbilityRepAnimMontageNetSerializer.cpp
+++ b/Engine/Plugins/Runtime/GameplayAbilities/Source/GameplayAbilities/Private/Serialization/GameplayAbilityRepAnimMontageNetSerializer.cpp
@@ -4,7 +4,6 @@
 
 #include UE_INLINE_GENERATED_CPP_BY_NAME(GameplayAbilityRepAnimMontageNetSerializer)
 
-#if UE_WITH_IRIS
 
 #include "Animation/AnimSequenceBase.h"
 #include "Animation/AnimMontage.h"
@@ -583,5 +582,4 @@ void FGameplayAbilityRepAnimMontageNetSerializer::FNetSerializerRegistryDelegate
 
 }
 
-#endif // UE_WITH_IRIS
 
diff --git a/Engine/Plugins/Runtime/GameplayAbilities/Source/GameplayAbilities/Private/Serialization/GameplayAbilityTargetDataHandleNetSerializer.cpp b/Engine/Plugins/Runtime/GameplayAbilities/Source/GameplayAbilities/Private/Serialization/GameplayAbilityTargetDataHandleNetSerializer.cpp
--- a/Engine/Plugins/Runtime/GameplayAbilities/Source/GameplayAbilities/Private/Serialization/GameplayAbilityTargetDataHandleNetSerializer.cpp
+++ b/Engine/Plugins/Runtime/GameplayAbilities/Source/GameplayAbilities/Private/Serialization/GameplayAbilityTargetDataHandleNetSerializer.cpp
@@ -2,7 +2,6 @@
 
 #include "Serialization/GameplayAbilityTargetDataHandleNetSerializer.h"
 
-#if UE_WITH_IRIS
 
 #include "Abilities/GameplayAbilityTargetTypes.h"
 #include "Iris/Serialization/NetSerializerDelegates.h"
@@ -267,5 +266,3 @@ void InitGameplayAbilityTargetDataHandleNetSerializerTypeCache()
 {
 	UE::Net::FGameplayAbilityTargetDataHandleNetSerializer::InitTypeCache();
 }
-
-#endif // UE_WITH_IRIS
\ No newline at end of file
diff --git a/Engine/Plugins/Runtime/GameplayAbilities/Source/GameplayAbilities/Private/Serialization/GameplayAbilityTargetingLocationInfoNetSerializer.cpp b/Engine/Plugins/Runtime/GameplayAbilities/Source/GameplayAbilities/Private/Serialization/GameplayAbilityTargetingLocationInfoNetSerializer.cpp
--- a/Engine/Plugins/Runtime/GameplayAbilities/Source/GameplayAbilities/Private/Serialization/GameplayAbilityTargetingLocationInfoNetSerializer.cpp
+++ b/Engine/Plugins/Runtime/GameplayAbilities/Source/GameplayAbilities/Private/Serialization/GameplayAbilityTargetingLocationInfoNetSerializer.cpp
@@ -4,7 +4,6 @@
 
 #include UE_INLINE_GENERATED_CPP_BY_NAME(GameplayAbilityTargetingLocationInfoNetSerializer)
 
-#if UE_WITH_IRIS
 
 #include "Abilities/GameplayAbilityTargetTypes.h"
 #include "Iris/Core/BitTwiddling.h"
@@ -475,4 +474,3 @@ void FGameplayAbilityTargetingLocationInfoNetSerializer::FNetSerializerRegistryD
 
 }
 
-#endif // UE_WITH_IRIS
diff --git a/Engine/Plugins/Runtime/GameplayAbilities/Source/GameplayAbilities/Private/Serialization/GameplayEffectContextHandleNetSerializer.cpp b/Engine/Plugins/Runtime/GameplayAbilities/Source/GameplayAbilities/Private/Serialization/GameplayEffectContextHandleNetSerializer.cpp
--- a/Engine/Plugins/Runtime/GameplayAbilities/Source/GameplayAbilities/Private/Serialization/GameplayEffectContextHandleNetSerializer.cpp
+++ b/Engine/Plugins/Runtime/GameplayAbilities/Source/GameplayAbilities/Private/Serialization/GameplayEffectContextHandleNetSerializer.cpp
@@ -4,7 +4,6 @@
 
 #include UE_INLINE_GENERATED_CPP_BY_NAME(GameplayEffectContextHandleNetSerializer)
 
-#if UE_WITH_IRIS
 
 #include "GameplayEffectTypes.h"
 #include "Iris/Serialization/PolymorphicNetSerializerImpl.h"
@@ -100,4 +99,3 @@ void InitGameplayEffectContextHandleNetSerializerTypeCache()
 	UE::Net::FGameplayEffectContextHandleNetSerializer::InitTypeCache();
 }
 
-#endif // UE_WITH_IRIS
diff --git a/Engine/Plugins/Runtime/GameplayAbilities/Source/GameplayAbilities/Private/Serialization/GameplayEffectContextNetSerializer.cpp b/Engine/Plugins/Runtime/GameplayAbilities/Source/GameplayAbilities/Private/Serialization/GameplayEffectContextNetSerializer.cpp
--- a/Engine/Plugins/Runtime/GameplayAbilities/Source/GameplayAbilities/Private/Serialization/GameplayEffectContextNetSerializer.cpp
+++ b/Engine/Plugins/Runtime/GameplayAbilities/Source/GameplayAbilities/Private/Serialization/GameplayEffectContextNetSerializer.cpp
@@ -5,7 +5,6 @@
 
 #include UE_INLINE_GENERATED_CPP_BY_NAME(GameplayEffectContextNetSerializer)
 
-#if UE_WITH_IRIS
 
 #include "Engine/HitResult.h"
 #include "GameplayEffectTypes.h"
@@ -603,7 +602,6 @@ void FGameplayEffectContextNetSerializer::FNetSerializerRegistryDelegates::OnPos
 
 }
 
-#endif // UE_WITH_IRIS
 
 void FGameplayEffectContextAccessorForNetSerializer::CopyReplicatedFieldsFrom(const FGameplayEffectContextAccessorForNetSerializer& GE)
 {
diff --git a/Engine/Plugins/Runtime/GameplayAbilities/Source/GameplayAbilities/Private/Serialization/InternalMinimalGameplayCueReplicationProxyNetSerializer.cpp b/Engine/Plugins/Runtime/GameplayAbilities/Source/GameplayAbilities/Private/Serialization/InternalMinimalGameplayCueReplicationProxyNetSerializer.cpp
--- a/Engine/Plugins/Runtime/GameplayAbilities/Source/GameplayAbilities/Private/Serialization/InternalMinimalGameplayCueReplicationProxyNetSerializer.cpp
+++ b/Engine/Plugins/Runtime/GameplayAbilities/Source/GameplayAbilities/Private/Serialization/InternalMinimalGameplayCueReplicationProxyNetSerializer.cpp
@@ -4,7 +4,6 @@
 
 #include UE_INLINE_GENERATED_CPP_BY_NAME(InternalMinimalGameplayCueReplicationProxyNetSerializer)
 
-#if UE_WITH_IRIS
 
 #include "GameplayCueInterface.h"
 
@@ -20,4 +19,3 @@ void FMinimalGameplayCueReplicationProxyForNetSerializer::AssignReplicatedFields
 	ReplicationProxy.ReplicatedLocations = Locations;
 }
 
-#endif
diff --git a/Engine/Plugins/Runtime/GameplayAbilities/Source/GameplayAbilities/Private/Serialization/InternalMinimalReplicationTagCountMapNetSerializer.cpp b/Engine/Plugins/Runtime/GameplayAbilities/Source/GameplayAbilities/Private/Serialization/InternalMinimalReplicationTagCountMapNetSerializer.cpp
--- a/Engine/Plugins/Runtime/GameplayAbilities/Source/GameplayAbilities/Private/Serialization/InternalMinimalReplicationTagCountMapNetSerializer.cpp
+++ b/Engine/Plugins/Runtime/GameplayAbilities/Source/GameplayAbilities/Private/Serialization/InternalMinimalReplicationTagCountMapNetSerializer.cpp
@@ -4,7 +4,6 @@
 
 #include UE_INLINE_GENERATED_CPP_BY_NAME(InternalMinimalReplicationTagCountMapNetSerializer)
 
-#if UE_WITH_IRIS
 
 #include "GameplayEffectTypes.h"
 
@@ -31,4 +30,3 @@ void FMinimalReplicationTagCountMapForNetSerializer::ClampTagCount(SIZE_T MaxTag
 	}
 }
 
-#endif
diff --git a/Engine/Plugins/Runtime/GameplayAbilities/Source/GameplayAbilities/Private/Serialization/MinimalGameplayCueReplicationProxyNetSerializer.cpp b/Engine/Plugins/Runtime/GameplayAbilities/Source/GameplayAbilities/Private/Serialization/MinimalGameplayCueReplicationProxyNetSerializer.cpp
--- a/Engine/Plugins/Runtime/GameplayAbilities/Source/GameplayAbilities/Private/Serialization/MinimalGameplayCueReplicationProxyNetSerializer.cpp
+++ b/Engine/Plugins/Runtime/GameplayAbilities/Source/GameplayAbilities/Private/Serialization/MinimalGameplayCueReplicationProxyNetSerializer.cpp
@@ -4,7 +4,6 @@
 
 #include UE_INLINE_GENERATED_CPP_BY_NAME(MinimalGameplayCueReplicationProxyNetSerializer)
 
-#if UE_WITH_IRIS
 
 #include "GameplayCueInterface.h"
 #include "Iris/Core/NetObjectReference.h"
@@ -249,4 +248,3 @@ void FMinimalGameplayCueReplicationProxyNetSerializer::FNetSerializerRegistryDel
 
 }
 
-#endif // UE_WITH_IRIS
diff --git a/Engine/Plugins/Runtime/GameplayAbilities/Source/GameplayAbilities/Private/Serialization/MinimalGameplayCueReplicationProxyReplicationFragment.cpp b/Engine/Plugins/Runtime/GameplayAbilities/Source/GameplayAbilities/Private/Serialization/MinimalGameplayCueReplicationProxyReplicationFragment.cpp
--- a/Engine/Plugins/Runtime/GameplayAbilities/Source/GameplayAbilities/Private/Serialization/MinimalGameplayCueReplicationProxyReplicationFragment.cpp
+++ b/Engine/Plugins/Runtime/GameplayAbilities/Source/GameplayAbilities/Private/Serialization/MinimalGameplayCueReplicationProxyReplicationFragment.cpp
@@ -2,7 +2,6 @@
 
 #include "Serialization/MinimalGameplayCueReplicationProxyReplicationFragment.h"
 
-#if UE_WITH_IRIS
 
 #include "Engine/NetConnection.h"
 #include "Iris/ReplicationState/ReplicationStateDescriptor.h"
@@ -245,4 +244,3 @@ void FMinimalGameplayCueReplicationProxyReplicationFragment::CallRepNotify(FRepl
 
 }
 
-#endif
diff --git a/Engine/Plugins/Runtime/GameplayAbilities/Source/GameplayAbilities/Private/Serialization/MinimalReplicationTagCountMapReplicationFragment.cpp b/Engine/Plugins/Runtime/GameplayAbilities/Source/GameplayAbilities/Private/Serialization/MinimalReplicationTagCountMapReplicationFragment.cpp
--- a/Engine/Plugins/Runtime/GameplayAbilities/Source/GameplayAbilities/Private/Serialization/MinimalReplicationTagCountMapReplicationFragment.cpp
+++ b/Engine/Plugins/Runtime/GameplayAbilities/Source/GameplayAbilities/Private/Serialization/MinimalReplicationTagCountMapReplicationFragment.cpp
@@ -2,7 +2,6 @@
 
 #include "Serialization/MinimalReplicationTagCountMapReplicationFragment.h"
 
-#if UE_WITH_IRIS
 
 #include "Engine/NetConnection.h"
 #include "Iris/ReplicationState/ReplicationStateDescriptor.h"
@@ -171,4 +170,3 @@ void FMinimalReplicationTagCountMapReplicationFragment::CallRepNotify(FReplicati
 
 }
 
-#endif
diff --git a/Engine/Plugins/Runtime/GameplayAbilities/Source/GameplayAbilities/Private/Serialization/MinimalReplicationTagCountMapReplicationFragment.h b/Engine/Plugins/Runtime/GameplayAbilities/Source/GameplayAbilities/Private/Serialization/MinimalReplicationTagCountMapReplicationFragment.h
--- a/Engine/Plugins/Runtime/GameplayAbilities/Source/GameplayAbilities/Private/Serialization/MinimalReplicationTagCountMapReplicationFragment.h
+++ b/Engine/Plugins/Runtime/GameplayAbilities/Source/GameplayAbilities/Private/Serialization/MinimalReplicationTagCountMapReplicationFragment.h
@@ -2,7 +2,6 @@
 
 #pragma once
 
-#if UE_WITH_IRIS
 
 #include "Iris/ReplicationState/PropertyReplicationState.h"
 #include "Iris/ReplicationSystem/ReplicationFragment.h"
@@ -44,4 +43,3 @@ private:
 
 }
 
-#endif
diff --git a/Engine/Plugins/Runtime/GameplayAbilities/Source/GameplayAbilities/Private/Serialization/MinimalReplicationTagCountNetSerializer.cpp b/Engine/Plugins/Runtime/GameplayAbilities/Source/GameplayAbilities/Private/Serialization/MinimalReplicationTagCountNetSerializer.cpp
--- a/Engine/Plugins/Runtime/GameplayAbilities/Source/GameplayAbilities/Private/Serialization/MinimalReplicationTagCountNetSerializer.cpp
+++ b/Engine/Plugins/Runtime/GameplayAbilities/Source/GameplayAbilities/Private/Serialization/MinimalReplicationTagCountNetSerializer.cpp
@@ -6,7 +6,6 @@
 
 #include UE_INLINE_GENERATED_CPP_BY_NAME(MinimalReplicationTagCountNetSerializer)
 
-#if UE_WITH_IRIS
 
 #include "GameplayEffectTypes.h"
 #include "Iris/Serialization/NetSerializerDelegates.h"
@@ -282,4 +281,3 @@ void FMinimalReplicationTagCountMapNetSerializer::FNetSerializerRegistryDelegate
 
 }
 
-#endif
diff --git a/Engine/Plugins/Runtime/GameplayAbilities/Source/GameplayAbilities/Private/Serialization/PredictionKeyNetSerializer.cpp b/Engine/Plugins/Runtime/GameplayAbilities/Source/GameplayAbilities/Private/Serialization/PredictionKeyNetSerializer.cpp
--- a/Engine/Plugins/Runtime/GameplayAbilities/Source/GameplayAbilities/Private/Serialization/PredictionKeyNetSerializer.cpp
+++ b/Engine/Plugins/Runtime/GameplayAbilities/Source/GameplayAbilities/Private/Serialization/PredictionKeyNetSerializer.cpp
@@ -4,7 +4,6 @@
 
 #include UE_INLINE_GENERATED_CPP_BY_NAME(PredictionKeyNetSerializer)
 
-#if UE_WITH_IRIS
 
 #include "GameplayPrediction.h"
 #include "Engine/NetConnection.h"
@@ -206,4 +205,3 @@ void FPredictionKeyNetSerializer::FNetSerializerRegistryDelegates::OnPreFreezeNe
 
 }
 
-#endif // UE_WITH_IRIS
diff --git a/Engine/Plugins/Runtime/GameplayAbilities/Source/GameplayAbilities/Public/Abilities/GameplayAbility.h b/Engine/Plugins/Runtime/GameplayAbilities/Source/GameplayAbilities/Public/Abilities/GameplayAbility.h
--- a/Engine/Plugins/Runtime/GameplayAbilities/Source/GameplayAbilities/Public/Abilities/GameplayAbility.h
+++ b/Engine/Plugins/Runtime/GameplayAbilities/Source/GameplayAbilities/Public/Abilities/GameplayAbility.h
@@ -522,10 +522,8 @@ public:
 	/** Overridden to allow Blueprint replicated properties to work */
 	UE_API virtual void GetLifetimeReplicatedProps(TArray< class FLifetimeProperty >& OutLifetimeProps) const;
 
-#if UE_WITH_IRIS
 	/** Register all replication fragments */
 	UE_API virtual void RegisterReplicationFragments(UE::Net::FFragmentRegistrationContext& Context, UE::Net::EFragmentRegistrationFlags RegistrationFlags) override;
-#endif // UE_WITH_IRIS
 
 	// --------------------------------------
 	//	IGameplayTaskOwnerInterface
diff --git a/Engine/Plugins/Runtime/GameplayAbilities/Source/GameplayAbilities/Public/AttributeSet.h b/Engine/Plugins/Runtime/GameplayAbilities/Source/GameplayAbilities/Public/AttributeSet.h
--- a/Engine/Plugins/Runtime/GameplayAbilities/Source/GameplayAbilities/Public/AttributeSet.h
+++ b/Engine/Plugins/Runtime/GameplayAbilities/Source/GameplayAbilities/Public/AttributeSet.h
@@ -258,10 +258,8 @@ public:
 	UE_API virtual void PreNetReceive() override;
 	UE_API virtual void PostNetReceive() override;
 
-#if UE_WITH_IRIS
 	/** Register all replication fragments */
 	UE_API virtual void RegisterReplicationFragments(UE::Net::FFragmentRegistrationContext& Context, UE::Net::EFragmentRegistrationFlags RegistrationFlags) override;
-#endif // UE_WITH_IRIS
 protected:
 	/** Is this attribute set safe to ID over the network by name?  */
 	uint32 bNetAddressable : 1;
diff --git a/Engine/Plugins/Runtime/GameplayAbilities/Source/GameplayAbilities/Public/GameplayCueInterface.h b/Engine/Plugins/Runtime/GameplayAbilities/Source/GameplayAbilities/Public/GameplayCueInterface.h
--- a/Engine/Plugins/Runtime/GameplayAbilities/Source/GameplayAbilities/Public/GameplayCueInterface.h
+++ b/Engine/Plugins/Runtime/GameplayAbilities/Source/GameplayAbilities/Public/GameplayCueInterface.h
@@ -14,13 +14,11 @@
 
 #define UE_API GAMEPLAYABILITIES_API
 
-#if UE_WITH_IRIS
 struct FMinimalGameplayCueReplicationProxyForNetSerializer;
 namespace UE::Net
 {
 	class FMinimalGameplayCueReplicationProxyReplicationFragment;
 }
-#endif
 
 /** Interface for actors that wish to handle GameplayCue events from GameplayEffects. Native only because blueprints can't implement interfaces with native functions */
 UINTERFACE(MinimalAPI, meta = (CannotImplementInterfaceInBlueprint))
@@ -237,10 +235,8 @@ struct FMinimalGameplayCueReplicationProxy
 	bool operator!=(const FActiveGameplayCueContainer& Other) const { return !(*this == Other); }
 
 private:
-#if UE_WITH_IRIS
 	friend FMinimalGameplayCueReplicationProxyForNetSerializer;
 	friend UE::Net::FMinimalGameplayCueReplicationProxyReplicationFragment;
-#endif
 
 	enum { NumInlineTags = 16 };
 
diff --git a/Engine/Plugins/Runtime/GameplayAbilities/Source/GameplayAbilities/Public/Serialization/GameplayAbilityTargetDataHandleNetSerializer.h b/Engine/Plugins/Runtime/GameplayAbilities/Source/GameplayAbilities/Public/Serialization/GameplayAbilityTargetDataHandleNetSerializer.h
--- a/Engine/Plugins/Runtime/GameplayAbilities/Source/GameplayAbilities/Public/Serialization/GameplayAbilityTargetDataHandleNetSerializer.h
+++ b/Engine/Plugins/Runtime/GameplayAbilities/Source/GameplayAbilities/Public/Serialization/GameplayAbilityTargetDataHandleNetSerializer.h
@@ -2,7 +2,6 @@
 
 #pragma once
 
-#if UE_WITH_IRIS
 
 #include "Iris/Serialization/NetSerializer.h"
 
@@ -15,4 +14,3 @@ UE_NET_DECLARE_SERIALIZER(FGameplayAbilityTargetDataHandleNetSerializer, GAMEPLA
 
 void InitGameplayAbilityTargetDataHandleNetSerializerTypeCache();
 
-#endif // UE_WITH_IRIS
diff --git a/Engine/Plugins/Runtime/GameplayAbilities/Source/GameplayAbilities/Public/Serialization/MinimalGameplayCueReplicationProxyReplicationFragment.h b/Engine/Plugins/Runtime/GameplayAbilities/Source/GameplayAbilities/Public/Serialization/MinimalGameplayCueReplicationProxyReplicationFragment.h
--- a/Engine/Plugins/Runtime/GameplayAbilities/Source/GameplayAbilities/Public/Serialization/MinimalGameplayCueReplicationProxyReplicationFragment.h
+++ b/Engine/Plugins/Runtime/GameplayAbilities/Source/GameplayAbilities/Public/Serialization/MinimalGameplayCueReplicationProxyReplicationFragment.h
@@ -2,7 +2,6 @@
 
 #pragma once
 
-#if UE_WITH_IRIS
 
 #include "Iris/ReplicationState/PropertyReplicationState.h"
 #include "Iris/ReplicationSystem/ReplicationFragment.h"
@@ -44,4 +43,3 @@ private:
 
 }
 
-#endif
diff --git a/Engine/Plugins/Runtime/InstancedActors/Source/InstancedActors/Private/InstancedActorsData.cpp b/Engine/Plugins/Runtime/InstancedActors/Source/InstancedActors/Private/InstancedActorsData.cpp
--- a/Engine/Plugins/Runtime/InstancedActors/Source/InstancedActors/Private/InstancedActorsData.cpp
+++ b/Engine/Plugins/Runtime/InstancedActors/Source/InstancedActors/Private/InstancedActorsData.cpp
@@ -18,9 +18,7 @@
 #include "Components/InstancedStaticMeshComponent.h"
 #include "Engine/AssetManager.h"
 #include "Engine/StaticMesh.h"
-#if UE_WITH_IRIS
 #include "Iris/ReplicationSystem/ReplicationFragmentUtil.h"
-#endif
 #include "InstancedActorsVisualizationProcessor.h"
 #include "MassActorSubsystem.h"
 #include "MassCommonFragments.h"
@@ -1030,12 +1028,10 @@ void UInstancedActorsData::GetLifetimeReplicatedProps(TArray<FLifetimeProperty>&
 	DOREPLIFETIME_WITH_PARAMS_FAST(UInstancedActorsData, InstanceDeltas, RepParams);
 }
 
-#if UE_WITH_IRIS
 void UInstancedActorsData::RegisterReplicationFragments(UE::Net::FFragmentRegistrationContext& Context, UE::Net::EFragmentRegistrationFlags RegistrationFlags)
 {
 	UE::Net::FReplicationFragmentUtil::CreateAndRegisterFragmentsForObject(this, Context, RegistrationFlags);
 }
-#endif
 
 void UInstancedActorsData::SetInstanceCurrentLifecyclePhase(FInstancedActorsInstanceIndex InstanceIndex, uint8 InCurrentLifecyclePhaseIndex)
 {
diff --git a/Engine/Plugins/Runtime/InstancedActors/Source/InstancedActors/Private/InstancedActorsManager.cpp b/Engine/Plugins/Runtime/InstancedActors/Source/InstancedActors/Private/InstancedActorsManager.cpp
--- a/Engine/Plugins/Runtime/InstancedActors/Source/InstancedActors/Private/InstancedActorsManager.cpp
+++ b/Engine/Plugins/Runtime/InstancedActors/Source/InstancedActors/Private/InstancedActorsManager.cpp
@@ -28,9 +28,7 @@
 #include "MassRepresentationSubsystem.h"
 #include "Math/NumericLimits.h"
 
-#if UE_WITH_IRIS
 #include "Net/Iris/ReplicationSystem/ReplicationSystemUtil.h"
-#endif
 
 #include UE_INLINE_GENERATED_CPP_BY_NAME(InstancedActorsManager)
 
@@ -221,7 +219,6 @@ AInstancedActorsManager::AInstancedActorsManager()
 	// @todo need default implementation of GUID and set it here via SetSavedActorGUID
 }
 
-#if UE_WITH_IRIS
 void AInstancedActorsManager::OnReplicationStartedForIris(const FOnReplicationStartedParams& Params)
 {
 	Super::OnReplicationStartedForIris(Params);
@@ -237,7 +234,6 @@ void AInstancedActorsManager::OnReplicationStartedForIris(const FOnReplicationSt
 	// replicated at the maximum configured frequency.
 	UE::Net::FReplicationSystemUtil::SetStaticPriority(this, 1.0f);
 }
-#endif
 
 void AInstancedActorsManager::BeginPlay()
 {
diff --git a/Engine/Plugins/Runtime/InstancedActors/Source/InstancedActors/Public/InstancedActorsData.h b/Engine/Plugins/Runtime/InstancedActors/Source/InstancedActors/Public/InstancedActorsData.h
--- a/Engine/Plugins/Runtime/InstancedActors/Source/InstancedActors/Public/InstancedActorsData.h
+++ b/Engine/Plugins/Runtime/InstancedActors/Source/InstancedActors/Public/InstancedActorsData.h
@@ -450,9 +450,7 @@ protected:
 	UE_API virtual void PostDuplicate(EDuplicateMode::Type DuplicateMode) override;
 	UE_API virtual void PostLoad() override;
 	UE_API virtual void GetLifetimeReplicatedProps(TArray<FLifetimeProperty>& OutLifetimeProps) const override;
-#if UE_WITH_IRIS
 	UE_API virtual void RegisterReplicationFragments(UE::Net::FFragmentRegistrationContext& Context, UE::Net::EFragmentRegistrationFlags RegistrationFlags) override;
-#endif
 	virtual bool IsNameStableForNetworking() const override { return true; }
 	//~ End UObject Overrides
 
diff --git a/Engine/Plugins/Runtime/InstancedActors/Source/InstancedActors/Public/InstancedActorsManager.h b/Engine/Plugins/Runtime/InstancedActors/Source/InstancedActors/Public/InstancedActorsManager.h
--- a/Engine/Plugins/Runtime/InstancedActors/Source/InstancedActors/Public/InstancedActorsManager.h
+++ b/Engine/Plugins/Runtime/InstancedActors/Source/InstancedActors/Public/InstancedActorsManager.h
@@ -262,9 +262,7 @@ protected:
 	virtual void OnInitializeModifyAndSpawnEntities() {}
 
 	//~ Begin AActor Overrides
-#if UE_WITH_IRIS
 	UE_API virtual void OnReplicationStartedForIris(const FOnReplicationStartedParams&) override;
-#endif
 	UE_API virtual void BeginPlay() override;
 	UE_API virtual void EndPlay(const EEndPlayReason::Type EndPlayReason) override;
 	UE_API virtual bool IsHLODRelevant() const override;
diff --git a/Engine/Source/Developer/CQTest/Public/Components/PIENetworkComponent.h b/Engine/Source/Developer/CQTest/Public/Components/PIENetworkComponent.h
--- a/Engine/Source/Developer/CQTest/Public/Components/PIENetworkComponent.h
+++ b/Engine/Source/Developer/CQTest/Public/Components/PIENetworkComponent.h
@@ -17,10 +17,8 @@
 #include "Net/NetIDVariant.h"
 #include "UnrealEdGlobals.h"
 
-#if UE_WITH_IRIS
 #include "Iris/ReplicationSystem/ReplicationSystem.h"
 #include "Iris/ReplicationSystem/ObjectReplicationBridge.h"
-#endif
 
 /*
 //Example boiler plate
diff --git a/Engine/Source/Developer/CQTest/Public/Impl/PIENetworkComponent.inl b/Engine/Source/Developer/CQTest/Public/Impl/PIENetworkComponent.inl
--- a/Engine/Source/Developer/CQTest/Public/Impl/PIENetworkComponent.inl
+++ b/Engine/Source/Developer/CQTest/Public/Impl/PIENetworkComponent.inl
@@ -219,14 +219,10 @@ inline FPIENetworkComponent<NetworkDataType>& FPIENetworkComponent<NetworkDataTy
 		UE::Net::FNetIDVariant NetIDVariant;
 		if (State.World->GetNetDriver()->IsUsingIrisReplication())
 		{
-#if UE_WITH_IRIS
 			UReplicationSystem* ReplicationSystem = State.World->GetNetDriver()->GetReplicationSystem();
 			UObjectReplicationBridge* ReplicationBridge = Cast<UObjectReplicationBridge>(ReplicationSystem->GetReplicationBridge());
 			checkf(IsValid(ReplicationBridge), TEXT("Unable to create a ReplicationBridge."));
 			NetIDVariant = UE::Net::FNetIDVariant(ReplicationBridge->GetReplicatedRefHandle(*ServerActor));
-#else
-			checkf(false, TEXT("NetDriver is set to use Iris Replication, but the Engine is not configured with Iris support."));
-#endif // UE_WITH_IRIS
 		}
 		else
 		{
@@ -267,14 +263,10 @@ inline bool FPIENetworkComponent<NetworkDataType>::ReplicateToClients(NetworkDat
 		AActor* ClientActor = nullptr;
 		if (ClientState.World->GetNetDriver()->IsUsingIrisReplication())
 		{
-#if UE_WITH_IRIS
 			UReplicationSystem* ReplicationSystem = ClientState.World->GetNetDriver()->GetReplicationSystem();
 			UObjectReplicationBridge* ReplicationBridge = Cast<UObjectReplicationBridge>(ReplicationSystem->GetReplicationBridge());
 			checkf(IsValid(ReplicationBridge), TEXT("Unable to create a ReplicationBridge."));
 			ClientActor = Cast<AActor>(ReplicationBridge->GetReplicatedObject(NetIDVariant.GetVariant().template Get<UE::Net::FNetRefHandle>()));
-#else
-			checkf(false, TEXT("NetDriver is set to use Iris Replication, but the Engine is not configured with Iris support."));
-#endif // UE_WITH_IRIS
 		}
 		else
 		{
diff --git a/Engine/Source/Runtime/CoreUObject/Private/UObject/Obj.cpp b/Engine/Source/Runtime/CoreUObject/Private/UObject/Obj.cpp
--- a/Engine/Source/Runtime/CoreUObject/Private/UObject/Obj.cpp
+++ b/Engine/Source/Runtime/CoreUObject/Private/UObject/Obj.cpp
@@ -5863,11 +5863,9 @@ void UObject::GetReplicatedCustomConditionState(FCustomPropertyConditionState& O
 
 }
 
-#if UE_WITH_IRIS
 void UObject::RegisterReplicationFragments(UE::Net::FFragmentRegistrationContext& Context, UE::Net::EFragmentRegistrationFlags RegistrationFlags)
 {
 }
-#endif // UE_WITH_IRIS
 
 /** Called right before receiving a bunch */
 void UObject::PreNetReceive()
diff --git a/Engine/Source/Runtime/CoreUObject/Public/UObject/CoreNet.h b/Engine/Source/Runtime/CoreUObject/Public/UObject/CoreNet.h
--- a/Engine/Source/Runtime/CoreUObject/Public/UObject/CoreNet.h
+++ b/Engine/Source/Runtime/CoreUObject/Public/UObject/CoreNet.h
@@ -50,11 +50,9 @@ namespace UE::Net
 {
 	struct FNetResult;
 
-#if UE_WITH_IRIS
 	class FReplicationFragment;
 	struct FReplicationStateDescriptor;
 	typedef FReplicationFragment* (*CreateAndRegisterReplicationFragmentFunc)(UObject* Owner, const FReplicationStateDescriptor* Descriptor, FFragmentRegistrationContext& Context);
-#endif
 	
 	class FNetTokenExportContext;
 	class FNetTokenResolveContext;
@@ -306,9 +304,7 @@ public:
 	ELifetimeCondition Condition;
 	ELifetimeRepNotifyCondition RepNotifyCondition;
 	bool bIsPushBased;
-#if UE_WITH_IRIS
 	UE::Net::CreateAndRegisterReplicationFragmentFunc CreateAndRegisterReplicationFragmentFunction = nullptr;
-#endif
 
 	FLifetimeProperty()
 		: RepIndex(0)
@@ -344,9 +340,7 @@ public:
 			check(Condition == Other.Condition);
 			check(RepNotifyCondition == Other.RepNotifyCondition);
 			check(bIsPushBased == Other.bIsPushBased);
-#if UE_WITH_IRIS
 			check(CreateAndRegisterReplicationFragmentFunction == Other.CreateAndRegisterReplicationFragmentFunction);
-#endif
 			return true;
 		}
 
diff --git a/Engine/Source/Runtime/CoreUObject/Public/UObject/Object.h b/Engine/Source/Runtime/CoreUObject/Public/UObject/Object.h
--- a/Engine/Source/Runtime/CoreUObject/Public/UObject/Object.h
+++ b/Engine/Source/Runtime/CoreUObject/Public/UObject/Object.h
@@ -41,7 +41,6 @@ struct FPropertyChangedChainEvent;
 struct FTopLevelAssetPath;
 class UClass;
 
-#if UE_WITH_IRIS
 namespace UE::Net
 {
 	class FFragmentRegistrationContext;
@@ -51,7 +50,6 @@ namespace UE::Net
 		struct FNetHandleLegacyPushModelHelper;
 	}
 }
-#endif // UE_WITH_IRIS
 
 DECLARE_LOG_CATEGORY_EXTERN(LogObj, Log, All);
 
@@ -1036,7 +1034,6 @@ public:
 
 	/** Called when this object begins replicating to initialize the state of custom property conditions */
 	COREUOBJECT_API virtual void GetReplicatedCustomConditionState(FCustomPropertyConditionState& OutActiveState) const;
-#if UE_WITH_IRIS
 	/**
 	 * RegisterReplicationFragments is called when we an object is added to the ReplicationSystem, it allows an object to register new or existing ReplicationFragments describing data to be replicated and how it should be accessed
 	 * For more information about ReplicationFragments see ReplicationFragment.h
@@ -1045,7 +1042,6 @@ public:
 	 * @param RegistrationFlags Flags specifying what should be registered in the call
 	 */
 	COREUOBJECT_API virtual void RegisterReplicationFragments(UE::Net::FFragmentRegistrationContext& Context, UE::Net::EFragmentRegistrationFlags RegistrationFlags);
-#endif // UE_WITH_IRIS
 
 	/** IsNameStableForNetworking means an object can be referred to its path name (relative to outer) over the network */
 	COREUOBJECT_API virtual bool IsNameStableForNetworking() const;
@@ -1818,9 +1814,7 @@ struct FObjectNetPushIdHelper
 private:
 	friend struct FNetPrivatePushIdHelper;
 	friend struct FNetObjectManagerPushIdHelper;
-#if UE_WITH_IRIS
 	friend struct UE::Net::Private::FNetHandleLegacyPushModelHelper;
-#endif // UE_WITH_IRIS
 
 	static void SetNetPushIdDynamic(UObject* Object, const uint64 NewNetPushId)
 	{
diff --git a/Engine/Source/Runtime/Engine/Classes/Components/ActorComponent.h b/Engine/Source/Runtime/Engine/Classes/Components/ActorComponent.h
--- a/Engine/Source/Runtime/Engine/Classes/Components/ActorComponent.h
+++ b/Engine/Source/Runtime/Engine/Classes/Components/ActorComponent.h
@@ -744,7 +744,6 @@ public:
 	/** Allows components to wait on outstanding tasks prior to sending EOF update data. Executed on Game Thread and may await tasks. */
 	virtual void OnPreEndOfFrameSync() {}
 
-#if UE_WITH_IRIS
 	/** Register all replication fragments */
 	ENGINE_API virtual void RegisterReplicationFragments(UE::Net::FFragmentRegistrationContext& Context, UE::Net::EFragmentRegistrationFlags RegistrationFlags) override;
 	
@@ -768,14 +767,11 @@ public:
 	/** Tell component to end replication, should not be called explicitly. */
 	UE_DEPRECATED(5.7, "This is being deprecated as part of iris beta. Override OnStopReplicationForIris instead.")
 	ENGINE_API virtual void EndReplication();
-#endif // UE_WITH_IRIS
 
 private:
 
-#if UE_WITH_IRIS
 	//ENGINE_API static void StartReplicatingForIris(UActorComponent* ActorComponent);
 	//ENGINE_API static void StopReplicatingForIris(UActorComponent* ActorComponent);
-#endif
 
 	/** If WorldPrivate isn't set this will determine the world from outers */
 	ENGINE_API UWorld* GetWorld_Uncached() const;
diff --git a/Engine/Source/Runtime/Engine/Classes/Components/ChildActorComponent.h b/Engine/Source/Runtime/Engine/Classes/Components/ChildActorComponent.h
--- a/Engine/Source/Runtime/Engine/Classes/Components/ChildActorComponent.h
+++ b/Engine/Source/Runtime/Engine/Classes/Components/ChildActorComponent.h
@@ -227,10 +227,8 @@ public:
 	ENGINE_API void SetEditorTreeViewVisualizationMode(EChildActorComponentTreeViewVisualizationMode InMode);
 #endif
 
-#if UE_WITH_IRIS
 	/** Register all replication fragments */
 	ENGINE_API virtual void RegisterReplicationFragments(UE::Net::FFragmentRegistrationContext& Context, UE::Net::EFragmentRegistrationFlags RegistrationFlags) override;
-#endif // UE_WITH_IRIS
 
 
 private:
diff --git a/Engine/Source/Runtime/Engine/Classes/Engine/NetConnection.h b/Engine/Source/Runtime/Engine/Classes/Engine/NetConnection.h
--- a/Engine/Source/Runtime/Engine/Classes/Engine/NetConnection.h
+++ b/Engine/Source/Runtime/Engine/Classes/Engine/NetConnection.h
@@ -797,12 +797,10 @@ public:
 
 	void AddDestructionInfo(FActorDestructionInfo* DestructionInfo)
 	{
-#if UE_WITH_IRIS
 		if (Driver && Driver->GetReplicationSystem())
 		{
 			return;
 		}
-#endif // UE_WITH_IRIS
 		if (ReplicationConnectionDriver)
 		{
 			ReplicationConnectionDriver->NotifyAddDestructionInfo(DestructionInfo);
@@ -815,12 +813,10 @@ public:
 
 	void RemoveDestructionInfo(FActorDestructionInfo* DestructionInfo)
 	{
-#if UE_WITH_IRIS
 		if (Driver && Driver->GetReplicationSystem())
 		{
 			return;
 		}
-#endif // UE_WITH_IRIS
 		if (ReplicationConnectionDriver)
 		{
 			ReplicationConnectionDriver->NotifyRemoveDestructionInfo(DestructionInfo);
@@ -833,12 +829,10 @@ public:
 	
 	void ResetDestructionInfos()
 	{
-#if UE_WITH_IRIS
 		if (Driver && Driver->GetReplicationSystem())
 		{
 			return;
 		}
-#endif // UE_WITH_IRIS
 		if (ReplicationConnectionDriver)
 		{
 			ReplicationConnectionDriver->NotifyResetDestructionInfo();
diff --git a/Engine/Source/Runtime/Engine/Classes/Engine/NetDriver.h b/Engine/Source/Runtime/Engine/Classes/Engine/NetDriver.h
--- a/Engine/Source/Runtime/Engine/Classes/Engine/NetDriver.h
+++ b/Engine/Source/Runtime/Engine/Classes/Engine/NetDriver.h
@@ -25,9 +25,7 @@
 #include "Net/NetAnalyticsTypes.h"
 #include "Net/NetConnectionIdHandler.h"
 #include "HAL/LowLevelMemTracker.h"
-#if UE_WITH_IRIS
 #include "Templates/PimplPtr.h"
-#endif
 #include "NetDriver.generated.h"
 
 /**
@@ -350,13 +348,11 @@ struct FReplicatedStaticActorDestructionInfo;
 
 enum class ECreateReplicationChangelistMgrFlags;
 enum class EEngineNetworkRuntimeFeatures : uint16;
-#if UE_WITH_IRIS
 class UReplicationSystem;
 namespace UE::Net
 {
 	class FNetObjectGroupHandle;
 }
-#endif // UE_WITH_IRIS
 
 namespace UE::Net
 {
@@ -1573,10 +1569,8 @@ public:
 	/** Initialized the replication driver class to use for this driver */
 	ENGINE_API virtual bool InitReplicationDriverClass();
 
-#if UE_WITH_IRIS
 	/** Initialized the replication bridge class to use for this driver if using iris replication*/
 	ENGINE_API virtual bool InitReplicationBridgeClass();
-#endif
 	/** Shutdown all connections managed by this net driver */
 	ENGINE_API virtual void Shutdown();
 
@@ -1963,17 +1957,12 @@ public:
 	/** Returns if this netdriver is initialized to replicate using the Iris replication system or the Legacy replication system. */
 	inline bool IsUsingIrisReplication() const
 	{
-#if UE_WITH_IRIS
 		return bIsUsingIris;
-#else
-		return false;
-#endif //UE_WITH_IRIS
 	}
 
 	/** Returns the bitflag telling which network features are activated for this NetDriver. */
 	ENGINE_API EEngineNetworkRuntimeFeatures GetNetworkRuntimeFeatures() const;
 
-#if UE_WITH_IRIS
 	/**
 	 * Destroy and recreate the iris replication system for an active netdriver.
 	 * This will re-add all existing replicated actors back in the system.
@@ -1986,12 +1975,10 @@ public:
 
 	/** Create the Iris replication system after it got destroyed. */
 	ENGINE_API void RecreateIrisSystem();
-#endif // UE_WITH_IRIS
 
 	template<class T>
 	T* GetReplicationDriver() const { return Cast<T>(ReplicationDriver); }
 
-#if UE_WITH_IRIS
 	inline UReplicationSystem* GetReplicationSystem() { return ReplicationSystem; }
 	inline UReplicationSystem* GetReplicationSystem() const { return ReplicationSystem; }
 
@@ -2022,11 +2009,6 @@ public:
 		return OnNetTokenStoreReadyDelegate;
 	}
 
-#else
-	/** Returns NetTokenStore that is required to create and serialize NetTokens */
-	const UE::Net::FNetTokenStore* GetNetTokenStore() const { return nullptr; }
-	UE::Net::FNetTokenStore* GetNetTokenStore() { return nullptr; }
-#endif // UE_WITH_IRIS
 
 
 	ENGINE_API void RemoveClientConnection(UNetConnection* ClientConnectionToRemove);
@@ -2464,14 +2446,12 @@ private:
 	/** Handle to FNetDelegates::OnSyncLoadDetected delegate */
 	FDelegateHandle ReportSyncLoadDelegateHandle;
 
-#if UE_WITH_IRIS
 	void InitIrisSettings(FName NewDriverName);
 	void SetReplicationSystem(UReplicationSystem* ReplicationSystem);
 	void CreateReplicationSystem(bool bInitAsClient);
 	void UpdateIrisReplicationViews() const;
 	void SendClientMoveAdjustments();
 	void PostDispatchSendUpdate();
-#endif
 
 	/** Tell the registered NetAnalytics to send their analytics via the provider */
 	void SendNetAnalytics();
@@ -2500,7 +2480,6 @@ private:
 	UPROPERTY(transient)
 	TObjectPtr<UReplicationDriver> ReplicationDriver;
 
-#if UE_WITH_IRIS
 	TUniquePtr<UE::Net::FNetTokenStore> NetTokenStore;
 
 	UReplicationSystem* ReplicationSystem = nullptr;
@@ -2513,7 +2492,6 @@ private:
 
 	// For FindOrAddNetworkObjectInfo
 	TPimplPtr<FNetworkObjectInfo> DummyNetworkObjectInfo;
-#endif
 
 	/** Stores the list of objects to replicate into the replay stream. This should be a TUniquePtr, but it appears the generated.cpp file needs the full definition of the pointed-to type. */
 	TSharedPtr<FNetworkObjectList> NetworkObjects;
diff --git a/Engine/Source/Runtime/Engine/Classes/Engine/World.h b/Engine/Source/Runtime/Engine/Classes/Engine/World.h
--- a/Engine/Source/Runtime/Engine/Classes/Engine/World.h
+++ b/Engine/Source/Runtime/Engine/Classes/Engine/World.h
@@ -458,9 +458,7 @@ private:
 
 	friend class UPackageMapClient;
 
-#if UE_WITH_IRIS
 	friend class UNetActorFactory;
-#endif // UE_WITH_IRIS
 
 	/* Is the actor remotely owned. This should only be set true by the package map when it is creating an actor on a client that was replicated from the server. */
 	uint8	bRemoteOwned:1;
diff --git a/Engine/Source/Runtime/Engine/Classes/GameFramework/Actor.h b/Engine/Source/Runtime/Engine/Classes/GameFramework/Actor.h
--- a/Engine/Source/Runtime/Engine/Classes/GameFramework/Actor.h
+++ b/Engine/Source/Runtime/Engine/Classes/GameFramework/Actor.h
@@ -47,10 +47,8 @@ class UExternalDataLayerAsset;
 class UDataLayerInstance;
 class AWorldDataLayers;
 class IWorldPartitionCell;
-#if UE_WITH_IRIS
 struct FActorReplicationParams;
 struct FStopReplicatingActorParams;
-#endif // UE_WITH_IRIS
 class UActorFolder;
 struct FActorDataLayer;
 struct FHitResult;
@@ -3459,7 +3457,6 @@ public:
 	/** Destroys the constructed components. */
 	ENGINE_API void DestroyConstructedComponents();
 
-#if UE_WITH_IRIS
 	ENGINE_API virtual void RegisterReplicationFragments(UE::Net::FFragmentRegistrationContext& Context, UE::Net::EFragmentRegistrationFlags RegistrationFlags) override;
 
 	struct FFillReplicationParamsContext
@@ -3516,7 +3513,6 @@ protected:
 
 	friend struct UE::Net::FReplicationSystemUtil;
 
-#endif // UE_WITH_IRIS
 
 	/**
 	 * Updates the ReplicatePhysics condition. That information needs to be pushed to the ReplicationSystem.
diff --git a/Engine/Source/Runtime/Engine/Classes/GameFramework/CharacterMovementReplication.h b/Engine/Source/Runtime/Engine/Classes/GameFramework/CharacterMovementReplication.h
--- a/Engine/Source/Runtime/Engine/Classes/GameFramework/CharacterMovementReplication.h
+++ b/Engine/Source/Runtime/Engine/Classes/GameFramework/CharacterMovementReplication.h
@@ -9,9 +9,7 @@
 #include "Containers/BitArray.h"
 #include "Net/Core/NetToken/NetTokenExportContext.h"
 #include "Components/PrimitiveComponent.h"
-#if UE_WITH_IRIS
 #include "Iris/Serialization/IrisObjectReferencePackageMap.h"
-#endif
 
 #include "CharacterMovementReplication.generated.h"
 
@@ -60,11 +58,9 @@ struct FCharacterNetworkSerializationPackedBits
 	// TInlineAllocator used with TBitArray takes the number of 32-bit dwords, but the define is in number of bits, so convert here by dividing by 32.
 	TBitArray<TInlineAllocator<CHARACTER_SERIALIZATION_PACKEDBITS_RESERVED_SIZE / NumBitsPerDWORD>> DataBits;
 
-#if UE_WITH_IRIS
 	// Since this struct uses custom serialization path we need to explicitly capture exports in order to forward them to iris
 	// This is managed by the use of a custom packagemap.
 	UE::Net::FIrisPackageMapExports PackageMapExports;
-#endif
 
 	// Since we capturing data outside of the normal serialization path we also need to store exports to inject when actually sending the data.
 	UE::Net::FNetTokenExportContext::FNetTokenExports NetTokensPendingExport;
diff --git a/Engine/Source/Runtime/Engine/Classes/GameFramework/PlayerController.h b/Engine/Source/Runtime/Engine/Classes/GameFramework/PlayerController.h
--- a/Engine/Source/Runtime/Engine/Classes/GameFramework/PlayerController.h
+++ b/Engine/Source/Runtime/Engine/Classes/GameFramework/PlayerController.h
@@ -376,9 +376,7 @@ public:
 #endif
 
 protected:
-#if UE_WITH_IRIS
 	ENGINE_API virtual void OnReplicationStartedForIris(const FOnReplicationStartedParams&) override;
-#endif // UE_WITH_IRIS
 
 private:
 
diff --git a/Engine/Source/Runtime/Engine/Private/Actor.cpp b/Engine/Source/Runtime/Engine/Private/Actor.cpp
--- a/Engine/Source/Runtime/Engine/Private/Actor.cpp
+++ b/Engine/Source/Runtime/Engine/Private/Actor.cpp
@@ -46,12 +46,10 @@
 #include "Engine/AutoDestroySubsystem.h"
 #include "WorldPartition/WorldPartitionRuntimeCellInterface.h"
 #include "AutoRTFM.h"
-#if UE_WITH_IRIS
 #include "Net/Iris/ReplicationSystem/EngineReplicationBridge.h"
 #include "Iris/ReplicationSystem/ReplicationSystem.h"
 #include "Iris/ReplicationSystem/Conditionals/ReplicationCondition.h"
 #include "Net/Iris/ReplicationSystem/ReplicationSystemUtil.h"
-#endif // UE_WITH_IRIS
 #include "LevelUtils.h"
 #include "Components/DecalComponent.h"
 #include "Components/LightComponent.h"
@@ -108,14 +106,12 @@ CSV_DEFINE_CATEGORY(ActorMigration, UE_WITH_REMOTE_OBJECT_HANDLE);
 
 FMakeNoiseDelegate AActor::MakeNoiseDelegate = FMakeNoiseDelegate::CreateStatic(&AActor::MakeNoiseImpl);
 
-#if UE_WITH_IRIS
-/** Controls if a replicated component should call EndReplication() immediately when updating its replication status. */
-static bool GbCallEndReplicationForActorComponentWhenUpdatingReplicatedComponent = false;
-static FAutoConsoleVariableRef CVarEndReplicationForActorComponentWhenUpdatingReplicatedComponent(
-	TEXT("net.Iris.CallEndReplicationForActorComponentWhenUpdatingReplicatedComponent"),
-	GbCallEndReplicationForActorComponentWhenUpdatingReplicatedComponent,
-	TEXT("If true, replicated ActorComponentrs will call endreplication if they are no longer replicated after a call to UpdateReplicatedComponent."));
-#endif
+/** Controls if a replicated component should call FReplicationSystemUtil::StopReplicatingActorComponent immediately when updating its replication status. */
+static bool GbCallStopReplicationForActorComponentWhenUpdatingReplicatedComponent = false;
+static FAutoConsoleVariableRef CVarStopReplicationForActorComponentWhenUpdatingReplicatedComponent(
+	TEXT("net.Iris.CallStopReplicationForActorComponentWhenUpdatingReplicatedComponent"),
+	GbCallStopReplicationForActorComponentWhenUpdatingReplicatedComponent,
+	TEXT("If true, replicated ActorComponentrs will call stop replication if they are no longer replicated after a call to UpdateReplicatedComponent.."));
 
 /** Selects if actor and actorcomponents will replicate subobjects using the registration list by default. */
 extern bool GDefaultUseSubObjectReplicationList;
@@ -1329,9 +1325,7 @@ void AActor::PostMigrate(const FUObjectMigrationContext& MigrationContext)
 			if (World)
 			{
 				// Notify replication systems that this actor has been destroyed.
-#if UE_WITH_IRIS
-				EndReplication(EEndPlayReason::Destroyed);
-#endif // UE_WITH_IRIS
+				UE::Net::FReplicationSystemUtil::StopReplicatingActor(this, FStopReplicatingActorParams(EEndPlayReason::Destroyed));
 
 				UE::Net::Private::FNetPropertyConditionManager::Get().NotifyObjectDestroyed(this);
 
@@ -1393,9 +1387,7 @@ void AActor::PostMigrate(const FUObjectMigrationContext& MigrationContext)
 		{
 			ensureMsgf(!bActorIsBeingDestroyed, TEXT("We do not want to be assigned ownership of Actors that have already been destroyed. We do not yet know what to do with it."));
 
-			#if UE_WITH_IRIS
 			UE::Net::FReplicationSystemUtil::StartReplicatingActor(this);
-			#endif // UE_WITH_IRIS
 
 			// Each component will be in charge of registering itself in its own PostMigrate, but we must register
 			// the Actor (which we'll do be saying it's currently the opposite of the deserialized value, then flipping it)
@@ -2604,9 +2596,7 @@ void AActor::SetOwner(AActor* NewOwner)
 			MarkOwnerRelevantComponentsDirty(this);
 		}
 
-#if UE_WITH_IRIS
 		UpdateOwningNetConnection();
-#endif // UE_WITH_IRIS
 	}
 }
 
@@ -3219,9 +3209,7 @@ void AActor::RouteEndPlay(const EEndPlayReason::Type EndPlayReason)
 			if (UWorld* World = GetWorld())
 			{
 				World->RemoveNetworkActor(this);
-#if UE_WITH_IRIS
 				UE::Net::FReplicationSystemUtil::StopReplicatingActor(this, FStopReplicatingActorParams(EndPlayReason));
-#endif
 			}
 		}
 
@@ -3246,10 +3234,8 @@ void AActor::EndPlay(const EEndPlayReason::Type EndPlayReason)
 
 		ActorHasBegunPlay = EActorBeginPlayState::HasNotBegunPlay;
 
-#if UE_WITH_IRIS
 		// This must be called otherwise the ReplicationSystem will keep a reference to the actor forever.
 		UE::Net::FReplicationSystemUtil::StopReplicatingActor(this, FStopReplicatingActorParams(EndPlayReason));
-#endif
 
 		// Dispatch the blueprint events
 		ReceiveEndPlay(EndPlayReason);
@@ -3871,7 +3857,6 @@ void AActor::UpdateReplicatedComponent(UActorComponent* Component)
 		ReplicatedComponentsInfo.Emplace(FReplicatedComponentInfo(Component, NetCondition));
 	}
 
-#if UE_WITH_IRIS
 	if (HasAuthority())
 	{
 		if (NetCondition != COND_Never && HasActorBegunPlay())
@@ -3881,7 +3866,7 @@ void AActor::UpdateReplicatedComponent(UActorComponent* Component)
 		}
 		else if (bWasRemovedFromReplicatedComponents || NetCondition == COND_Never)
 		{
-			if (GbCallEndReplicationForActorComponentWhenUpdatingReplicatedComponent)
+			if (GbCallStopReplicationForActorComponentWhenUpdatingReplicatedComponent)
 			{
 				UE::Net::FReplicationSystemUtil::StopReplicatingActorComponent(Component);
 			}
@@ -3891,7 +3876,6 @@ void AActor::UpdateReplicatedComponent(UActorComponent* Component)
 			}
 		}
 	}
-#endif // UE_WITH_IRIS
 }
 
 void AActor::UpdateAllReplicatedComponents()
@@ -3904,13 +3888,12 @@ void AActor::UpdateAllReplicatedComponents()
 	for (FReplicatedComponentInfo& RepComponentInfo : ReplicatedComponentsInfo)
 	{
 		RepComponentInfo.NetCondition = COND_Never;
-#if UE_WITH_IRIS
 		if (HasAuthority())
 		{
 			// Need to stop replication for all components that should no longer replicate
 			if (RepComponentInfo.Component && !RepComponentInfo.Component->GetIsReplicated())
 			{
-				if (GbCallEndReplicationForActorComponentWhenUpdatingReplicatedComponent)
+				if (GbCallStopReplicationForActorComponentWhenUpdatingReplicatedComponent)
 				{
 					UE::Net::FReplicationSystemUtil::StopReplicatingActorComponent(RepComponentInfo.Component);
 				}
@@ -3920,7 +3903,6 @@ void AActor::UpdateAllReplicatedComponents()
 				}
 			}
 		}
-#endif
 	}
 
 	for (UActorComponent* Component : OwnedComponents)
@@ -3941,7 +3923,6 @@ void AActor::UpdateAllReplicatedComponents()
 					Component->ReadyForReplication();
 				}
 
-#if UE_WITH_IRIS
 				if (HasAuthority())
 				{
 					if (NetCondition != COND_Never)
@@ -3951,7 +3932,7 @@ void AActor::UpdateAllReplicatedComponents()
 					}
 					else
 					{
-						if (GbCallEndReplicationForActorComponentWhenUpdatingReplicatedComponent)
+						if (GbCallStopReplicationForActorComponentWhenUpdatingReplicatedComponent)
 						{
 							UE::Net::FReplicationSystemUtil::StopReplicatingActorComponent(Component);
 						}
@@ -3961,7 +3942,6 @@ void AActor::UpdateAllReplicatedComponents()
 						}
 					}
 				}
-#endif
 			}
 		}
 	}
@@ -4563,7 +4543,6 @@ void AActor::SetReplicates(bool bInReplicates)
 						MyWorld->AddNetworkActor(this);
 					}
 
-#if UE_WITH_IRIS
 					if (HasActorBegunPlay())
 					{
 						if (bNewlyReplicates)
@@ -4575,17 +4554,14 @@ void AActor::SetReplicates(bool bInReplicates)
 							UpdateOwningNetConnection();
 						}
 					}
-#endif // UE_WITH_IRIS
 
 					ForcePropertyCompare();
 				}
 			}
-#if UE_WITH_IRIS
 			else if (HasActorBegunPlay())
 			{
 				UE::Net::FReplicationSystemUtil::StopReplicatingActor(this, FStopReplicatingActorParams(EEndPlayReason::RemovedFromWorld));
 			}
-#endif
 
 			MARK_PROPERTY_DIRTY_FROM_NAME(AActor, RemoteRole, this);
 		}
@@ -4622,7 +4598,6 @@ void AActor::SetAutonomousProxy(const bool bInAutonomousProxy, const bool bAllow
 			// This is because RemoteRole is special since it will change between connections, so we have to special case
 			ForcePropertyCompare();
 
-#if UE_WITH_IRIS
 			// Update autonomous role condition
 			if (UReplicationSystem* ReplicationSystem = UE::Net::FReplicationSystemUtil::GetReplicationSystem(this))
 			{
@@ -4637,7 +4612,6 @@ void AActor::SetAutonomousProxy(const bool bInAutonomousProxy, const bool bAllow
 					}
 				}
 			}
-#endif // UE_WITH_IRIS
 		}
 	}
 	else
@@ -4712,7 +4686,6 @@ EActorUpdateOverlapsMethod AActor::GetUpdateOverlapsMethodDuringLevelStreaming()
 
 void AActor::DispatchBeginPlay(bool bFromLevelStreaming)
 {
-#if UE_WITH_IRIS
 	// If we are spawned from networking, the actor is not ready for begin play until the initial state has been applied.
 	if (bActorIsPendingPostNetInit)
 	{
@@ -4721,7 +4694,6 @@ void AActor::DispatchBeginPlay(bool bFromLevelStreaming)
 			return;
 		}
 	}
-#endif
 	
 	UWorld* World = (!HasActorBegunPlay() && IsValidChecked(this) ? GetWorld() : nullptr);
 
@@ -4742,9 +4714,7 @@ void AActor::DispatchBeginPlay(bool bFromLevelStreaming)
 
 		ActorHasBegunPlay = EActorBeginPlayState::BeginningPlay;
 
-#if UE_WITH_IRIS
 		UE::Net::FReplicationSystemUtil::StartReplicatingActor(this);
-#endif // UE_WITH_IRIS
 
 #if UE_SUPPORT_FOR_ACTOR_TICK_DISABLE
 		if (World->IsActorTickAndUserCallbacksEnabled())
diff --git a/Engine/Source/Runtime/Engine/Private/ActorReplication.cpp b/Engine/Source/Runtime/Engine/Private/ActorReplication.cpp
--- a/Engine/Source/Runtime/Engine/Private/ActorReplication.cpp
+++ b/Engine/Source/Runtime/Engine/Private/ActorReplication.cpp
@@ -15,7 +15,6 @@
 #include "DrawDebugHelpers.h"
 #include "Net/Core/PushModel/PushModel.h"
 #include "Interfaces/Interface_ActorSubobject.h"
-#if UE_WITH_IRIS
 #include "Engine/NetConnection.h"
 #include "Iris/IrisConfig.h"
 #include "Iris/ReplicationSystem/ReplicationFragment.h"
@@ -24,7 +23,6 @@
 #include "Iris/ReplicationSystem/Conditionals/ReplicationCondition.h"
 #include "Net/Iris/ReplicationSystem/ReplicationSystemUtil.h"
 #include "Net/Iris/ReplicationSystem/EngineReplicationBridge.h"
-#endif // UE_WITH_IRIS
 #include "Physics/Experimental/PhysScene_Chaos.h"
 
 /*-----------------------------------------------------------------------------
@@ -474,9 +472,7 @@ void AActor::GatherCurrentMovement()
 			// If RepPhysics has changed value then notify the ReplicationSystem
 			if (bPrevRepPhysics != ReplicatedMovement.bRepPhysics)
 			{
-#if UE_WITH_IRIS
 				UpdateReplicatePhysicsCondition();
-#endif // UE_WITH_IRIS
 				bWasRepMovementModified = true;
 			}
 		}
@@ -617,9 +613,7 @@ void AActor::SetReplicatedComponentNetCondition(const UActorComponent* Replicate
 	{
 		if (NetCondition != ComponentInfo->NetCondition)
 		{
-#if UE_WITH_IRIS
 			UE::Net::FReplicationSystemUtil::SetNetConditionForActorComponent(ReplicatedComponent, NetCondition);
-#endif
 			ComponentInfo->NetCondition = NetCondition;
 		}
 	}
@@ -669,12 +663,10 @@ void AActor::AddComponentForReplication(UActorComponent* Component)
 		Component->ReadyForReplication();
 	}
 
-#if UE_WITH_IRIS
     if (NetCondition != COND_Never)
     {
 		UE::Net::FReplicationSystemUtil::StartReplicatingActorComponent(this, Component);
     }
-#endif
 }
 
 void AActor::RemoveReplicatedComponent(UActorComponent* Component)
@@ -683,12 +675,10 @@ void AActor::RemoveReplicatedComponent(UActorComponent* Component)
 	if (Index != INDEX_NONE)
 	{
 		ReplicatedComponentsInfo.RemoveAtSwap(Index);
-#if UE_WITH_IRIS
 		if (HasAuthority())
 		{
 			UE::Net::FReplicationSystemUtil::StopReplicatingActorComponent(Component);
 		}
-#endif
 	}
 }
 
@@ -717,12 +707,10 @@ void AActor::AddReplicatedSubObject(UObject* SubObject, ELifetimeCondition NetCo
 	ensureMsgf(Result != FSubObjectRegistry::EResult::NetConditionConflict, TEXT("%s(0x%p) Registered subobject %s (0x%p) again with a different net condition. Active [%s] New [%s]. New condition will be ignored"),
 		*GetName(), this, *SubObject->GetName(), SubObject, *StaticEnum<ELifetimeCondition>()->GetValueAsString(ReplicatedSubObjects.GetNetCondition(SubObject)), *StaticEnum<ELifetimeCondition>()->GetValueAsString(NetCondition));
 
-#if UE_WITH_IRIS
 	if (GetIsReplicated() && (ActorHasBegunPlay != EActorBeginPlayState::HasNotBegunPlay))
 	{
 		UE::Net::FReplicationSystemUtil::StartReplicatingActorSubObject(this, SubObject, NetCondition);
 	}
-#endif
 }
 
 bool AActor::RemoveReplicatedSubObjectFromList(UObject* SubObject)
@@ -737,12 +725,10 @@ bool AActor::RemoveReplicatedSubObjectFromList(UObject* SubObject)
 void AActor::RemoveReplicatedSubObject(UObject* SubObject)
 {
 	const bool bWasRemoved = RemoveReplicatedSubObjectFromList(SubObject);
-#if UE_WITH_IRIS
 	if (bWasRemoved && HasAuthority())
 	{
 		UE::Net::FReplicationSystemUtil::StopReplicatingActorSubObject(this, SubObject);
 	}
-#endif // UE_WITH_IRIS
 }
 
 void AActor::DestroyReplicatedSubObjectOnRemotePeers(UObject* SubObject)
@@ -894,12 +880,10 @@ void AActor::AddActorComponentReplicatedSubObject(UActorComponent* OwnerComponen
 		ensureMsgf(Result != FSubObjectRegistry::EResult::NetConditionConflict, TEXT("%s::%s (0x%p) Registered subobject %s (0x%p) again with a different net condition. Active [%s] New [%s]. New condition will be ignored"),
 			*GetName(), *OwnerComponent->GetName(), OwnerComponent, *SubObject->GetName(), SubObject, *StaticEnum<ELifetimeCondition>()->GetValueAsString(ComponentInfo->SubObjects.GetNetCondition(SubObject)), *StaticEnum<ELifetimeCondition>()->GetValueAsString(NetCondition));
 
-#if UE_WITH_IRIS
 		if (OwnerComponent->GetIsReplicated() && OwnerComponent->HasBegunPlay())
 		{
 			UE::Net::FReplicationSystemUtil::StartReplicatingActorComponentSubObject(OwnerComponent, SubObject, NetCondition);
 		}
-#endif
 	}
 	else
 	{
@@ -936,12 +920,10 @@ void AActor::RemoveActorComponentReplicatedSubObject(UActorComponent* OwnerCompo
 {
 	const bool bWasRemoved = RemoveActorComponentReplicatedSubObjectFromList(OwnerComponent, SubObject);
 	
-#if UE_WITH_IRIS
 	if (bWasRemoved && HasAuthority())
 	{
 		UE::Net::FReplicationSystemUtil::StopReplicatingActorComponentSubObject(OwnerComponent, SubObject);
 	}
-#endif
 }
 
 bool AActor::IsReplicatedSubObjectRegistered(const UObject* SubObject) const
@@ -1033,7 +1015,6 @@ void AActor::OnRep_Owner()
 
 }
 
-#if UE_WITH_IRIS
 void AActor::RegisterReplicationFragments(UE::Net::FFragmentRegistrationContext& Context, UE::Net::EFragmentRegistrationFlags RegistrationFlags)
 {
 	using namespace UE::Net;
@@ -1141,11 +1122,9 @@ void AActor::UpdateOwningNetConnection()
 	} while (EveryChildren.Num());
 }
 
-#endif // UE_WITH_IRIS
 
 void AActor::UpdateReplicatePhysicsCondition()
 {
-#if UE_WITH_IRIS
 	using namespace UE::Net;
 
 	const FNetHandle NetHandle = FReplicationSystemUtil::GetNetHandle(this);
@@ -1153,7 +1132,6 @@ void AActor::UpdateReplicatePhysicsCondition()
 	{
 		FReplicationSystemUtil::SetReplicationCondition(NetHandle, EReplicationCondition::ReplicatePhysics, GetReplicatedMovement().bRepPhysics);
 	}
-#endif // UE_WITH_IRIS
 }
 
 bool AActor::CanTriggerResimulation() const
diff --git a/Engine/Source/Runtime/Engine/Private/Components/ActorComponent.cpp b/Engine/Source/Runtime/Engine/Private/Components/ActorComponent.cpp
--- a/Engine/Source/Runtime/Engine/Private/Components/ActorComponent.cpp
+++ b/Engine/Source/Runtime/Engine/Private/Components/ActorComponent.cpp
@@ -54,12 +54,10 @@
 #include "WorldPartition/WorldPartitionActorDesc.h"
 #endif
 
-#if UE_WITH_IRIS
 #include "Iris/Core/IrisLog.h"
 #include "Iris/ReplicationSystem/ReplicationFragment.h"
 #include "Iris/ReplicationSystem/ReplicationFragmentUtil.h"
 #include "Net/Iris/ReplicationSystem/ReplicationSystemUtil.h"
-#endif // UE_WITH_IRIS
 
 #if UE_WITH_REMOTE_OBJECT_HANDLE
 #include "UObject/UObjectMigrationContext.h"
@@ -1630,12 +1628,10 @@ void UActorComponent::EndPlay(const EEndPlayReason::Type EndPlayReason)
 
 	check(bHasBegunPlay);
 
-#if UE_WITH_IRIS
 	if (EndPlayReason != EEndPlayReason::EndPlayInEditor && EndPlayReason != EEndPlayReason::Quit)
 	{
 		UE::Net::FReplicationSystemUtil::StopReplicatingActorComponent(this);
 	}
-#endif
 
 	// If we're in the process of being garbage collected it is unsafe to call out to blueprints
 	if (!HasAnyFlags(RF_BeginDestroyed) && !IsUnreachable() && (GetClass()->HasAnyClassFlags(CLASS_CompiledFromBlueprint) || !GetClass()->HasAnyClassFlags(CLASS_Native)))
@@ -3052,7 +3048,6 @@ bool UActorComponent::CanEditChange(const FProperty* InProperty) const
 }
 #endif
 
-#if UE_WITH_IRIS
 void UActorComponent::RegisterReplicationFragments(UE::Net::FFragmentRegistrationContext& Context, UE::Net::EFragmentRegistrationFlags RegistrationFlags)
 {
 	if (CreationMethod == EComponentCreationMethod::UserConstructionScript)
@@ -3084,7 +3079,6 @@ void UActorComponent::EndReplication()
 {
 }
 
-#endif // UE_WITH_IRIS
 
 bool UActorComponent::IsEditableWhenInherited() const
 {
diff --git a/Engine/Source/Runtime/Engine/Private/Components/CharacterMovementComponent.cpp b/Engine/Source/Runtime/Engine/Private/Components/CharacterMovementComponent.cpp
--- a/Engine/Source/Runtime/Engine/Private/Components/CharacterMovementComponent.cpp
+++ b/Engine/Source/Runtime/Engine/Private/Components/CharacterMovementComponent.cpp
@@ -28,11 +28,9 @@
 #include "PBDRigidsSolver.h"
 #include "Engine/NetworkObjectList.h"
 #include "Net/PerfCountersHelpers.h"
-#if UE_WITH_IRIS
 #include "Net/Iris/ReplicationSystem/ReplicationSystemUtil.h"
 #include "Net/Iris/ReplicationSystem/EngineReplicationBridge.h"
 #include "Iris/ReplicationSystem/ReplicationSystem.h"
-#endif
 
 #include UE_INLINE_GENERATED_CPP_BY_NAME(CharacterMovementComponent)
 
@@ -8978,7 +8976,6 @@ void UCharacterMovementComponent::ReplicateMoveToServer(float DeltaTime, const F
 	ClientData->PendingMove = NULL;
 }
 
-#if UE_WITH_IRIS
 namespace UE::Private
 {
 
@@ -9017,7 +9014,6 @@ static UIrisObjectReferencePackageMap* GetIrisPackageMapToReadReferences(const U
 }
 
 }
-#endif
 
 void UCharacterMovementComponent::CallServerMovePacked(const FSavedMove_Character* NewMove, const FSavedMove_Character* PendingMove, const FSavedMove_Character* OldMove)
 {
@@ -9033,13 +9029,11 @@ void UCharacterMovementComponent::CallServerMovePacked(const FSavedMove_Characte
 	static FCharacterServerMovePackedBits PackedBits;
 	UNetConnection* NetConnection = CharacterOwner->GetNetConnection();	
 
-#if UE_WITH_IRIS
 	if (UPackageMap* PackageMap = UE::Private::GetIrisPackageMapToCaptureReferences(NetConnection, PackedBits))
 	{
 		ServerMoveBitWriter.PackageMap = PackageMap;
 	}
 	else
-#endif
 	{
 		// Extract the net package map used for serializing object references.
 		ServerMoveBitWriter.PackageMap = NetConnection ? ToRawPtr(NetConnection->PackageMap) : nullptr;
@@ -9777,13 +9771,11 @@ void UCharacterMovementComponent::ServerMovePacked_ServerReceive(const FCharacte
 	// Reuse bit reader to avoid allocating memory each time.
 	ServerMoveBitReader.SetData((uint8*)PackedBits.DataBits.GetData(), NumBits);
 
-#if UE_WITH_IRIS
 	if (UIrisObjectReferencePackageMap* PackageMap = UE::Private::GetIrisPackageMapToReadReferences(CharacterOwner->GetNetConnection(), PackedBits))
 	{
 		ServerMoveBitReader.PackageMap = PackageMap;
 	}
 	else
-#endif
 	{
 		ServerMoveBitReader.PackageMap = PackedBits.GetPackageMap();
 	}
@@ -10647,13 +10639,11 @@ void UCharacterMovementComponent::MoveResponsePacked_ClientReceive(const FCharac
 	// Reuse bit reader to avoid allocating memory each time.
 	MoveResponseBitReader.SetData((uint8*)PackedBits.DataBits.GetData(), NumBits);
 
-#if UE_WITH_IRIS
 	if (UPackageMap* PackageMap = UE::Private::GetIrisPackageMapToReadReferences(CharacterOwner->GetNetConnection(), PackedBits))
 	{
 		MoveResponseBitReader.PackageMap = PackageMap;
 	}
 	else
-#endif
 	{
 		MoveResponseBitReader.PackageMap = PackedBits.GetPackageMap();
 	}
@@ -10692,13 +10682,11 @@ void UCharacterMovementComponent::ServerSendMoveResponse(const FClientAdjustment
 	UNetConnection* NetConnection = CharacterOwner->GetNetConnection();	
 
 	// Extract the net package map used for serializing object references.
-#if UE_WITH_IRIS
 	if (UPackageMap* PackageMap = UE::Private::GetIrisPackageMapToCaptureReferences(NetConnection, PackedBits))	
 	{
 		MoveResponseBitWriter.PackageMap = PackageMap;
 	}
 	else
-#endif
 	{
 		MoveResponseBitWriter.PackageMap = NetConnection ? ToRawPtr(NetConnection->PackageMap) : nullptr;
 	}
diff --git a/Engine/Source/Runtime/Engine/Private/Components/ChildActorComponent.cpp b/Engine/Source/Runtime/Engine/Private/Components/ChildActorComponent.cpp
--- a/Engine/Source/Runtime/Engine/Private/Components/ChildActorComponent.cpp
+++ b/Engine/Source/Runtime/Engine/Private/Components/ChildActorComponent.cpp
@@ -15,9 +15,7 @@
 #include "WorldPartition/HLOD/HLODBuilder.h"
 #include "Misc/DataValidation.h"
 #endif
-#if UE_WITH_IRIS
 #include "Iris/ReplicationSystem/PropertyReplicationFragment.h"
-#endif // UE_WITH_IRIS
 
 DEFINE_LOG_CATEGORY_STATIC(LogChildActorComponent, Log, All);
 
@@ -992,7 +990,6 @@ bool UChildActorComponent::IsHLODRelevant() const
 	return bIsHLODRelevant;
 }
 
-#if UE_WITH_IRIS
 void UChildActorComponent::RegisterReplicationFragments(UE::Net::FFragmentRegistrationContext& Context, UE::Net::EFragmentRegistrationFlags RegistrationFlags)
 {
 	using namespace UE::Net;
@@ -1000,7 +997,6 @@ void UChildActorComponent::RegisterReplicationFragments(UE::Net::FFragmentRegist
 	// Fallback on default since archetype used on client and server differs giving different default values
 	Super::RegisterReplicationFragments(Context, RegistrationFlags | UE::Net::EFragmentRegistrationFlags::InitializeDefaultStateFromClassDefaults);
 }
-#endif // UE_WITH_IRIS
 
 
 #if WITH_EDITOR
diff --git a/Engine/Source/Runtime/Engine/Private/DataChannel.cpp b/Engine/Source/Runtime/Engine/Private/DataChannel.cpp
--- a/Engine/Source/Runtime/Engine/Private/DataChannel.cpp
+++ b/Engine/Source/Runtime/Engine/Private/DataChannel.cpp
@@ -11,9 +11,7 @@
 #include "Engine/Level.h"
 #include "Engine/World.h"
 #include "GameFramework/WorldSettings.h"
-#if UE_WITH_IRIS
 #include "Iris/ReplicationSystem/ReplicationSystem.h"
-#endif
 #include "Interfaces/ReplicatedObjectInterface.h"
 #include "Misc/MemStack.h"
 #include "Misc/ScopeExit.h"
@@ -1975,12 +1973,10 @@ void UControlChannel::ReceivedBunch( FInBunch& Bunch )
 			uint64 NetRefHandleId = 0;
 			if (FNetControlMessage<NMT_IrisProtocolMismatch>::Receive(Bunch, NetRefHandleId))
 			{
-#if UE_WITH_IRIS
 				if (UReplicationSystem* IrisRepSystem = Connection->Driver->GetReplicationSystem())
 				{
 					IrisRepSystem->ReportProtocolMismatch(NetRefHandleId, Connection->GetConnectionHandle().GetParentConnectionId());
 				}
-#endif
 			}
 		}
 		else if (MessageType == NMT_IrisNetRefHandleError)
@@ -1990,12 +1986,10 @@ void UControlChannel::ReceivedBunch( FInBunch& Bunch )
 			TArray<uint64> ExtraNetRefs;
 			if (FNetControlMessage<NMT_IrisNetRefHandleError>::Receive(Bunch, NetRefHandleError, NetRefHandleId, ExtraNetRefs))
 			{
-#if UE_WITH_IRIS
 				if (UReplicationSystem* IrisRepSystem = Connection->Driver->GetReplicationSystem())
 				{
 					IrisRepSystem->ReportErrorWithNetRefHandle(NetRefHandleError, NetRefHandleId, Connection->GetConnectionHandle().GetParentConnectionId(), ExtraNetRefs);
 				}
-#endif
 			}
 		}
 		else if (Connection->Driver->Notify != nullptr)
diff --git a/Engine/Source/Runtime/Engine/Private/Engine/HitResultNetSerializer.cpp b/Engine/Source/Runtime/Engine/Private/Engine/HitResultNetSerializer.cpp
--- a/Engine/Source/Runtime/Engine/Private/Engine/HitResultNetSerializer.cpp
+++ b/Engine/Source/Runtime/Engine/Private/Engine/HitResultNetSerializer.cpp
@@ -2,7 +2,6 @@
 #include "Engine/HitResultNetSerializer.h"
 #include UE_INLINE_GENERATED_CPP_BY_NAME(HitResultNetSerializer)
 
-#if UE_WITH_IRIS
 
 #include "Engine/HitResult.h"
 #include "Net/Core/NetBitArray.h"
@@ -524,4 +523,3 @@ void FHitResultNetSerializer::FNetSerializerRegistryDelegates::OnPostFreezeNetSe
 }
 
 
-#endif
\ No newline at end of file
diff --git a/Engine/Source/Runtime/Engine/Private/GameFramework/CharacterNetworkSerializationPackedBitsNetSerializer.cpp b/Engine/Source/Runtime/Engine/Private/GameFramework/CharacterNetworkSerializationPackedBitsNetSerializer.cpp
--- a/Engine/Source/Runtime/Engine/Private/GameFramework/CharacterNetworkSerializationPackedBitsNetSerializer.cpp
+++ b/Engine/Source/Runtime/Engine/Private/GameFramework/CharacterNetworkSerializationPackedBitsNetSerializer.cpp
@@ -4,7 +4,6 @@
 
 #include UE_INLINE_GENERATED_CPP_BY_NAME(CharacterNetworkSerializationPackedBitsNetSerializer)
 
-#if UE_WITH_IRIS
 
 #include "GameFramework/CharacterNetworkSerializationPackedBitsNetSerializer.h"
 #include "EngineLogs.h"
@@ -319,4 +318,3 @@ void FCharacterNetworkSerializationPackedBitsNetSerializer::FNetSerializerRegist
 
 }
 
-#endif // UE_WITH_IRIS
diff --git a/Engine/Source/Runtime/Engine/Private/GameFramework/PlayerStateCountLimiterConfig.cpp b/Engine/Source/Runtime/Engine/Private/GameFramework/PlayerStateCountLimiterConfig.cpp
--- a/Engine/Source/Runtime/Engine/Private/GameFramework/PlayerStateCountLimiterConfig.cpp
+++ b/Engine/Source/Runtime/Engine/Private/GameFramework/PlayerStateCountLimiterConfig.cpp
@@ -7,12 +7,10 @@
 UPlayerStateCountLimiterConfig::UPlayerStateCountLimiterConfig()
 : Super()
 {
-#if UE_WITH_IRIS
 	Mode = ENetObjectCountLimiterMode::Fill;
 	MaxObjectCount = 2;
 	Priority = 1.0f;
 	// We want the owning connection's player state to have high priority. 
 	OwningConnectionPriority = 2.0f;
 	bEnableOwnedObjectsFastLane = true;
-#endif
 }
diff --git a/Engine/Source/Runtime/Engine/Private/GameFramework/RootMotionSourceGroupNetSerializer.cpp b/Engine/Source/Runtime/Engine/Private/GameFramework/RootMotionSourceGroupNetSerializer.cpp
--- a/Engine/Source/Runtime/Engine/Private/GameFramework/RootMotionSourceGroupNetSerializer.cpp
+++ b/Engine/Source/Runtime/Engine/Private/GameFramework/RootMotionSourceGroupNetSerializer.cpp
@@ -4,7 +4,6 @@
 
 #include UE_INLINE_GENERATED_CPP_BY_NAME(RootMotionSourceGroupNetSerializer)
 
-#if UE_WITH_IRIS
 
 #include "GameFramework/RootMotionSource.h"
 #include "HAL/PlatformMemory.h"
@@ -562,4 +561,3 @@ void FRootMotionSourceGroupNetSerializer::FNetSerializerRegistryDelegates::OnLoa
 
 }
 
-#endif
diff --git a/Engine/Source/Runtime/Engine/Private/GameFramework/UniqueNetIdReplNetSerializer.cpp b/Engine/Source/Runtime/Engine/Private/GameFramework/UniqueNetIdReplNetSerializer.cpp
--- a/Engine/Source/Runtime/Engine/Private/GameFramework/UniqueNetIdReplNetSerializer.cpp
+++ b/Engine/Source/Runtime/Engine/Private/GameFramework/UniqueNetIdReplNetSerializer.cpp
@@ -4,7 +4,6 @@
 
 #include UE_INLINE_GENERATED_CPP_BY_NAME(UniqueNetIdReplNetSerializer)
 
-#if UE_WITH_IRIS
 
 #include "GameFramework/InternalUniqueNetIdReplNetSerializer.h"
 #include "GameFramework/OnlineReplStructs.h"
@@ -567,4 +566,3 @@ void FUniqueNetIdReplNetSerializer::FNetSerializerRegistryDelegates::OnPostFreez
 
 }
 
-#endif // UE_WITH_IRIS
diff --git a/Engine/Source/Runtime/Engine/Private/LevelStreaming.cpp b/Engine/Source/Runtime/Engine/Private/LevelStreaming.cpp
--- a/Engine/Source/Runtime/Engine/Private/LevelStreaming.cpp
+++ b/Engine/Source/Runtime/Engine/Private/LevelStreaming.cpp
@@ -48,11 +48,7 @@ namespace LevelStreamingCVars
 	// There are cases where we might have multiple visibility requests (and data) in flight leading to the server 
 	// starting to replicate data based on an older visibility/streamingstatus update which can lead to broken channels
 	// to mitigate this problem we assign a TransactionId to each request/update to make sure that we are acting on the correct data
-#if UE_WITH_IRIS
 	static bool bDefaultAllowClientUseMakingInvisibleTransactionRequests = true;
-#else
-	static bool bDefaultAllowClientUseMakingInvisibleTransactionRequests = false;
-#endif
 	FAutoConsoleVariableRef CVarDefaultAllowClientUseMakingInvisibleTransactionRequests(
 		TEXT("LevelStreaming.DefaultAllowClientUseMakingInvisibleTransactionRequests"),
 		bDefaultAllowClientUseMakingInvisibleTransactionRequests,
@@ -70,11 +66,7 @@ namespace LevelStreamingCVars
 		TEXT("0: Disable, 1: Enable"),
 		ECVF_Default);
 
-#if UE_WITH_IRIS
 	static bool bShouldServerUseMakingVisibleTransactionRequest = true;
-#else
-	static bool bShouldServerUseMakingVisibleTransactionRequest = false;
-#endif
 	FAutoConsoleVariableRef CVarShouldServerUseMakingVisibleTransactionRequest(
 		TEXT("LevelStreaming.ShouldServerUseMakingVisibleTransactionRequest"),
 		bShouldServerUseMakingVisibleTransactionRequest,
diff --git a/Engine/Source/Runtime/Engine/Private/Net/Experimental/Iris/DataStreamChannel.cpp b/Engine/Source/Runtime/Engine/Private/Net/Experimental/Iris/DataStreamChannel.cpp
--- a/Engine/Source/Runtime/Engine/Private/Net/Experimental/Iris/DataStreamChannel.cpp
+++ b/Engine/Source/Runtime/Engine/Private/Net/Experimental/Iris/DataStreamChannel.cpp
@@ -6,7 +6,6 @@
 
 #include UE_INLINE_GENERATED_CPP_BY_NAME(DataStreamChannel)
 
-#if UE_WITH_IRIS
 // IWYU pragma: begin_keep
 #include "Engine/NetConnection.h"
 #include "Iris/IrisConfig.h"
@@ -30,7 +29,6 @@
 #include "PacketHandler.h"
 #include "ProfilingDebugging/CsvProfiler.h"
 // IWYU pragma: end_keep
-#endif // UE_WITH_IRIS
 
 namespace UE::Net::Private
 {
@@ -64,7 +62,6 @@ void UDataStreamChannel::InitDataStreamManager()
 {
 	check(Connection);
 
-#if UE_WITH_IRIS
 	if (UReplicationSystem* ReplicationSystem = Connection->Driver->GetReplicationSystem())
 	{
 		DataStreamManager = NewObject<UDataStreamManager>();
@@ -78,7 +75,6 @@ void UDataStreamChannel::InitDataStreamManager()
 	{
 		bIsReadyToHandshake = 0U;
 	}
-#endif
 }
 
 void UDataStreamChannel::Init(UNetConnection* InConnection, int32 InChIndex, EChannelCreateFlags CreateFlags)
@@ -94,7 +90,6 @@ void UDataStreamChannel::ReInit()
 
 bool UDataStreamChannel::CleanUp(const bool bForDestroy, EChannelCloseReason CloseReason)
 {
-#if UE_WITH_IRIS
 	if (IsValid(DataStreamManager))
 	{
 		DataStreamManager->Deinit();
@@ -103,14 +98,12 @@ bool UDataStreamChannel::CleanUp(const bool bForDestroy, EChannelCloseReason Clo
 	}
 
 	WriteRecords.Reset();
-#endif // UE_WITH_IRIS
 
 	return Super::CleanUp(bForDestroy, CloseReason);
 }
 
 void UDataStreamChannel::ReceivedBunch(FInBunch& Bunch)
 {
-#if UE_WITH_IRIS
 	using namespace UE::Net;
 
 	IRIS_PROFILER_SCOPE(UDataStreamChannel_ReceivedBunch);
@@ -161,12 +154,10 @@ void UDataStreamChannel::ReceivedBunch(FInBunch& Bunch)
 		Bunch.SetError();
 	}
 
-#endif // UE_WITH_IRIS
 }
 
 void UDataStreamChannel::SendOpenBunch()
 {
-#if UE_WITH_IRIS
 
 	// We send this only once
 	if (!bHandshakeSent && NumOutRec == 0)
@@ -190,12 +181,10 @@ void UDataStreamChannel::SendOpenBunch()
 		}
 	}
 
-#endif // UE_WITH_IRIS
 }
 
 void UDataStreamChannel::PostTickDispatch()
 {
-#if UE_WITH_IRIS
 	using namespace UE::Net;
 
 	if (!Connection->Driver->IsUsingIrisReplication() || !bHandshakeComplete)
@@ -223,12 +212,10 @@ void UDataStreamChannel::PostTickDispatch()
 
 	WriteData(EDataStreamWriteMode::PostTickDispatch);
 
-#endif
 }
 
 void UDataStreamChannel::Tick()
 {
-#if UE_WITH_IRIS
 	using namespace UE::Net;
 
 	if (!Connection->Driver->IsUsingIrisReplication() || !bIsReadyToHandshake)
@@ -263,12 +250,10 @@ void UDataStreamChannel::Tick()
 
 	WriteData(EDataStreamWriteMode::Full);
 
-#endif
 }
 
 void UDataStreamChannel::WriteData(UE::Net::EDataStreamWriteMode WriteMode)
 {
-#if UE_WITH_IRIS
 	using namespace UE::Net;
 
 	// Limit the amount of bits to minimum of a bunch and our buffer. NetBitStreamWriter requires the number of bytes to be a multiple of 4.
@@ -437,7 +422,6 @@ void UDataStreamChannel::WriteData(UE::Net::EDataStreamWriteMode WriteMode)
 		Connection->FlushNet();	
 	}
 
-#endif // UE_WITH_IRIS
 }
 
 bool UDataStreamChannel::CanStopTicking() const
@@ -452,7 +436,6 @@ FString UDataStreamChannel::Describe()
 
 void UDataStreamChannel::ReceivedAck(int32 PacketId)
 {
-#if UE_WITH_IRIS
 
 	if (!bHandshakeComplete)
 	{
@@ -470,12 +453,10 @@ void UDataStreamChannel::ReceivedAck(int32 PacketId)
 		WriteRecords.Pop();
 	}
 
-#endif // UE_WITH_IRIS
 }
 
 void UDataStreamChannel::ReceivedNak(int32 PacketId)
 {
-#if UE_WITH_IRIS
 
 	if (!bHandshakeComplete)
 	{
@@ -492,20 +473,17 @@ void UDataStreamChannel::ReceivedNak(int32 PacketId)
 		WriteRecords.Pop();
 	}
 
-#endif // UE_WITH_IRIS
 }
 
 // Some DataStreams require perfect acking. If the ack sequence window is full we would get NAKs for packets thay may have been received.
 bool UDataStreamChannel::IsPacketWindowFull() const
 {
-#if UE_WITH_IRIS
 	const uint32 IrisPacketSequenceSafetyMarginUnsigned = static_cast<uint32>(FPlatformMath::Max(0, UE::Net::Private::IrisPacketSequenceSafetyMargin));
 	if (Connection->IsPacketSequenceWindowFull(IrisPacketSequenceSafetyMarginUnsigned))
 	{
 		UE_LOG(LogIris, Verbose, TEXT("Packet window full."));
 		return true;
 	}
-#endif // UE_WITH_IRIS
 
 	return WriteRecords.Count() >= WriteRecords.AllocatedCapacity();
 }
@@ -513,12 +491,10 @@ bool UDataStreamChannel::IsPacketWindowFull() const
 void UDataStreamChannel::AddReferencedObjects(UObject* Object, FReferenceCollector& Collector)
 {
 	UDataStreamChannel* Channel = CastChecked<UDataStreamChannel>(Object);
-#if UE_WITH_IRIS
 	if (Channel->DataStreamManager)
 	{
 		Collector.AddReferencedObject(Channel->DataStreamManager);
 	}
-#endif
 
 	Super::AddReferencedObjects(Channel, Collector);
 }
@@ -539,12 +515,10 @@ void UDataStreamChannel::AppendMustBeMappedGuids(FOutBunch* Bunch)
 
 bool UDataStreamChannel::HasAcknowledgedAllReliableData() const
 {
-#if UE_WITH_IRIS
 	if (DataStreamManager)
 	{
 		return DataStreamManager->HasAcknowledgedAllReliableData();
 	}
-#endif
 
 	return true;
 }
diff --git a/Engine/Source/Runtime/Engine/Private/Net/Experimental/Iris/DataStreamChannel.h b/Engine/Source/Runtime/Engine/Private/Net/Experimental/Iris/DataStreamChannel.h
--- a/Engine/Source/Runtime/Engine/Private/Net/Experimental/Iris/DataStreamChannel.h
+++ b/Engine/Source/Runtime/Engine/Private/Net/Experimental/Iris/DataStreamChannel.h
@@ -84,9 +84,7 @@ private:
 
 	void WriteData(UE::Net::EDataStreamWriteMode WriteMode);
 
-#if UE_WITH_IRIS
 	TObjectPtr<UDataStreamManager> DataStreamManager = nullptr;
-#endif
 
 	TResizableCircularQueue<FDataStreamChannelRecord> WriteRecords;
 
diff --git a/Engine/Source/Runtime/Engine/Private/Net/Iris/ReplicationSystem/NetEngineFactories.cpp b/Engine/Source/Runtime/Engine/Private/Net/Iris/ReplicationSystem/NetEngineFactories.cpp
--- a/Engine/Source/Runtime/Engine/Private/Net/Iris/ReplicationSystem/NetEngineFactories.cpp
+++ b/Engine/Source/Runtime/Engine/Private/Net/Iris/ReplicationSystem/NetEngineFactories.cpp
@@ -2,14 +2,11 @@
 
 #include "Net/Iris/ReplicationSystem/NetEngineFactories.h"
 
-#if UE_WITH_IRIS
 #include "Net/Iris/ReplicationSystem/NetActorFactory.h"
 #include "Net/Iris/ReplicationSystem/NetSubObjectFactory.h"
 
 #include "Iris/ReplicationSystem/NetObjectFactoryRegistry.h"
-#endif 
 
-#if UE_WITH_IRIS
 
 namespace UE::Net
 {
@@ -44,5 +41,3 @@ void ShutdownEngineNetObjectFactories()
 }
 
 } // end namespace UE::Net
-
-#endif // UE_WITH_IRIS
\ No newline at end of file
diff --git a/Engine/Source/Runtime/Engine/Private/Net/Iris/ReplicationSystem/ReplicationSystemUtil.cpp b/Engine/Source/Runtime/Engine/Private/Net/Iris/ReplicationSystem/ReplicationSystemUtil.cpp
--- a/Engine/Source/Runtime/Engine/Private/Net/Iris/ReplicationSystem/ReplicationSystemUtil.cpp
+++ b/Engine/Source/Runtime/Engine/Private/Net/Iris/ReplicationSystem/ReplicationSystemUtil.cpp
@@ -2,7 +2,6 @@
 
 #include "Net/Iris/ReplicationSystem/ReplicationSystemUtil.h"
 
-#if UE_WITH_IRIS
 
 #include "EngineUtils.h"
 #include "Engine/Engine.h"
@@ -752,4 +751,3 @@ void FReplicationSystemUtil::SetPollFrequency(const UObject* Object, float PollF
 
 } // end namespace UE::Net
 
-#endif
diff --git a/Engine/Source/Runtime/Engine/Private/Net/NetIDVariant.cpp b/Engine/Source/Runtime/Engine/Private/Net/NetIDVariant.cpp
--- a/Engine/Source/Runtime/Engine/Private/Net/NetIDVariant.cpp
+++ b/Engine/Source/Runtime/Engine/Private/Net/NetIDVariant.cpp
@@ -11,12 +11,10 @@ FNetIDVariant::FNetIDVariant(FNetworkGUID NetGUID)
 {
 }
 
-#if UE_WITH_IRIS
 FNetIDVariant::FNetIDVariant(FNetRefHandle NetRefHandle)
 	: Variant(TInPlaceType<FNetRefHandle>(), NetRefHandle)
 {
 }
-#endif
 
 bool FNetIDVariant::operator==(const FNetIDVariant& RHS) const
 {
@@ -29,12 +27,10 @@ bool FNetIDVariant::operator==(const FNetIDVariant& RHS) const
 	{
 		return Variant.Get<FNetworkGUID>() == RHS.Variant.Get<FNetworkGUID>();
 	}
-#if UE_WITH_IRIS
 	else if (Variant.IsType<FNetRefHandle>())
 	{
 		return Variant.Get<FNetRefHandle>() == RHS.Variant.Get<FNetRefHandle>();
 	}
-#endif
 
 	return true;
 }
@@ -46,13 +42,11 @@ bool FNetIDVariant::IsValid() const
 		return Variant.Get<FNetworkGUID>().IsValid();
 	}
 	else
-#if UE_WITH_IRIS
 	if (Variant.IsType<FNetRefHandle>())
 	{
 		return Variant.Get<FNetRefHandle>().IsValid();
 	}
 	else
-#endif
 	{
 		return false;
 	}
@@ -60,11 +54,7 @@ bool FNetIDVariant::IsValid() const
 
 FArchive& operator<<(FArchive& Ar, FNetIDVariant& NetID)
 {
-#if UE_WITH_IRIS
 	constexpr SIZE_T ExpectedSize = 3;
-#else
-	constexpr SIZE_T ExpectedSize = 2;
-#endif
 	static_assert(TVariantSize_V<FNetIDVariant::FVariantType> == ExpectedSize, "FNetIDVariant variant size changed, potential serialization incompatibility.");
 
 	uint32 TypeIndex = static_cast<uint32>(NetID.Variant.GetIndex());
@@ -78,13 +68,11 @@ FArchive& operator<<(FArchive& Ar, FNetIDVariant& NetID)
 			FNetworkGUID WriteGUID = NetID.Variant.Get<FNetworkGUID>();
 			Ar << WriteGUID;
 		}
-#if UE_WITH_IRIS
 		else if (NetID.Variant.IsType<FNetRefHandle>())
 		{
 			FNetRefHandle WriteHandle = NetID.Variant.Get<FNetRefHandle>();
 			Ar << WriteHandle;
 		}
-#endif
 	}
 	else
 	{
@@ -94,14 +82,12 @@ FArchive& operator<<(FArchive& Ar, FNetIDVariant& NetID)
 			Ar << ReadGUID;
 			NetID.Variant.Set<FNetworkGUID>(ReadGUID);
 		}
-#if UE_WITH_IRIS
 		else if (TypeIndex == FNetIDVariant::FVariantType::IndexOfType<FNetRefHandle>())
 		{
 			FNetRefHandle ReadRefHandle;
 			Ar << ReadRefHandle;
 			NetID.Variant.Set<FNetRefHandle>(ReadRefHandle);
 		}
-#endif
 		else if(TypeIndex == FNetIDVariant::FVariantType::IndexOfType<FNetIDVariant::FEmptyID>())
 		{
 			NetID.Variant = FNetIDVariant::FVariantType();
@@ -123,12 +109,10 @@ FString FNetIDVariant::ToString() const
 	{
 		return Variant.Get<FNetworkGUID>().ToString();
 	}
-#if UE_WITH_IRIS
 	else if (Variant.IsType<FNetRefHandle>())
 	{
 		return Variant.Get<FNetRefHandle>().ToString();
 	}
-#endif
 	return FString(TEXT("Invalid"));
 }
 
diff --git a/Engine/Source/Runtime/Engine/Private/NetConnection.cpp b/Engine/Source/Runtime/Engine/Private/NetConnection.cpp
--- a/Engine/Source/Runtime/Engine/Private/NetConnection.cpp
+++ b/Engine/Source/Runtime/Engine/Private/NetConnection.cpp
@@ -45,13 +45,11 @@
 #include "Net/NetSubObjectRegistryGetter.h"
 #include "Net/RepLayout.h"
 #include "Net/Subsystems/NetworkSubsystem.h"
-#if UE_WITH_IRIS
 #include "Iris/IrisConfig.h"
 #include "Iris/ReplicationSystem/ReplicationSystem.h"
 #include "Iris/ReplicationSystem/Filtering/NetObjectFilter.h"
 #include "Net/Iris/ReplicationSystem/EngineReplicationBridge.h"
 #include "Iris/ReplicationSystem/NetTokenStore.h"
-#endif // UE_WITH_IRIS
 
 #include UE_INLINE_GENERATED_CPP_BY_NAME(NetConnection)
 
@@ -575,13 +573,11 @@ void UNetConnection::InitBase(UNetDriver* InDriver,class FSocket* InSocket, cons
 	}
 
 	// For now, NetToken are only enabled if we compile with iris.
-#if UE_WITH_IRIS
 	// Init RemoteNetTokenState.
 	if (Driver->GetNetTokenStore())
 	{
 		Driver->GetNetTokenStore()->InitRemoteNetTokenStoreState(GetConnectionHandle().GetParentConnectionId());
 	}
-#endif
 
 	// Create package map.
 	UPackageMapClient* PackageMapClient = NewObject<UPackageMapClient>(this, PackageMapClass);
@@ -625,13 +621,11 @@ void UNetConnection::InitBase(UNetDriver* InDriver,class FSocket* InSocket, cons
 	UE_NET_TRACE_CONNECTION_CREATED(NetTraceId, GetConnectionHandle().GetParentConnectionId());
 	UE_NET_TRACE_CONNECTION_STATE_UPDATED(NetTraceId, GetConnectionHandle().GetParentConnectionId(), static_cast<uint8>(GetConnectionState()));
 
-#if UE_WITH_IRIS
 	if (UReplicationSystem* ReplicationSystem = Driver->GetReplicationSystem())
 	{
 		ReplicationSystem->AddConnection(GetConnectionHandle().GetParentConnectionId());
 		ReplicationSystem->SetConnectionUserData(GetConnectionHandle().GetParentConnectionId(), this);
 	}
-#endif // UE_WITH_IRIS
 }
 
 /**
@@ -1160,12 +1154,10 @@ void UNetConnection::GracefulClose(FNetResult&& CloseReason)
 			SetConnectionState(EConnectionState::USOCK_Closing);
 			GracefulCloseTimeoutDeadline = Driver->GetElapsedTime() + Driver->GracefulCloseConnectionTimeout;
 
-#if UE_WITH_IRIS
 			if (UReplicationSystem* ReplicationSystem = Driver->GetReplicationSystem())
 			{
 				ReplicationSystem->SetConnectionGracefullyClosing(ConnectionHandle.GetParentConnectionId());
 			}
-#endif
 
 			// Might as well try to close immediately in case there isn't any pending reliable data.
 			TryClosePendingGracefulClose();
@@ -1437,12 +1429,10 @@ void UNetConnection::CleanUp()
 		// Otherwise we'd be able to do the appropriate logic in Add/Remove Client/ServerConnection
 		if (ConnectionHandle.IsParentConnection())
 		{
-#if UE_WITH_IRIS
 			if (UReplicationSystem* ReplicationSystem = Driver->GetReplicationSystem())
 			{
 				ReplicationSystem->RemoveConnection(ConnectionHandle.GetParentConnectionId());
 			}
-#endif // UE_WITH_IRIS
 			Driver->FreeConnectionId(ConnectionHandle.GetParentConnectionId());
 		}
 	}
@@ -1714,7 +1704,6 @@ FNetLevelVisibilityTransactionId UNetConnection::UpdateLevelStreamStatusChangedT
 	{
 		ClientVisibleLevelNames.Remove(PackageName);
 		ClientMakingVisibleLevelNames.Remove(PackageName);
-#if UE_WITH_IRIS
 		if (UReplicationSystem* ReplicationSystem = Driver->GetReplicationSystem())
 		{
 			if (const ULevel* Level = LevelObject->GetLoadedLevel())
@@ -1727,7 +1716,6 @@ FNetLevelVisibilityTransactionId UNetConnection::UpdateLevelStreamStatusChangedT
 			}
 		}
 		else
-#endif
 		{
 			UpdateCachedLevelVisibility(PackageName);
 		}
@@ -1867,7 +1855,6 @@ void UNetConnection::UpdateLevelVisibilityInternal(const FUpdateLevelVisibilityL
 			// when the streaming level is loaded, so mark them active for this connection
 			UWorld* LevelWorld = nullptr;
 
-#if UE_WITH_IRIS
 			if (UReplicationSystem* ReplicationSystem = Driver->GetReplicationSystem())
 			{
 				// We need to be careful here if the server would load a level after the client, 
@@ -1883,7 +1870,6 @@ void UNetConnection::UpdateLevelVisibilityInternal(const FUpdateLevelVisibilityL
 				}
 			}
 			else 
-#endif // UE_WITH_IRIS
 			{
 				LevelWorld = VisibilityResult.Package ? (UWorld*)FindObjectWithOuter(VisibilityResult.Package, UWorld::StaticClass()) : nullptr;
 
@@ -1962,7 +1948,6 @@ void UNetConnection::UpdateLevelVisibilityInternal(const FUpdateLevelVisibilityL
 			ReplicationConnectionDriver->NotifyClientVisibleLevelNamesRemove(LevelVisibility.PackageName);
 		}
 
-#if UE_WITH_IRIS
 		if (UReplicationSystem* ReplicationSystem = Driver->GetReplicationSystem())
 		{
 			// See if the level is loaded, if it is disable filtering, if not filter should already be removed
@@ -1981,7 +1966,6 @@ void UNetConnection::UpdateLevelVisibilityInternal(const FUpdateLevelVisibilityL
 				}
 			}
 		}
-#endif // UE_WITH_IRIS
 
 		QUICK_SCOPE_CYCLE_COUNTER(STAT_NetUpdateLevelVisibility_LevelUnloadChannelClose);
 
@@ -2027,13 +2011,11 @@ void UNetConnection::SetClientWorldPackageName(FName NewClientWorldPackageName)
 {
 	ClientWorldPackageName = NewClientWorldPackageName;
 	
-#if UE_WITH_IRIS
 	if (Driver->IsUsingIrisReplication())
 	{
 		Driver->UpdateGroupFilterStatusForAllLevelsForConnection(GetConnectionHandle());
 	}
 	else
-#endif
 	{
 		UpdateAllCachedLevelVisibility();
 	}
@@ -4108,7 +4090,6 @@ void UNetConnection::DispatchPacket( FBitReader& Reader, int32 PacketId, bool& b
 				else
 				{
 					UE_LOG(LogNetTraffic, Error, TEXT("Received corrupted packet data with SequenceId: %d from server %s"), PacketId, *LowLevelGetRemoteAddress());
-#if UE_WITH_IRIS
 					// If Iris reports errors they are unrecoverable.
 					if (Driver->GetReplicationSystem() != nullptr)
 					{
@@ -4116,7 +4097,6 @@ void UNetConnection::DispatchPacket( FBitReader& Reader, int32 PacketId, bool& b
 						Close(AddToAndConsumeChainResultPtr(Bunch.ExtendedError, ENetCloseResult::CorruptData));
 						return;
 					}
-#endif
 				}
 			}
 			// In replay, if the bunch generated an error but the channel isn't actually open, clean it up so the channel index remains free
@@ -5146,7 +5126,6 @@ void UNetConnection::HandleClientPlayer( APlayerController *PC, UNetConnection*
 	PlayerController = PC;
 	OwningActor = PC;
 
-#if UE_WITH_IRIS
 	{
 		// Enable replication
 		if (UReplicationSystem* ReplicationSystem = Driver->GetReplicationSystem())
@@ -5154,7 +5133,6 @@ void UNetConnection::HandleClientPlayer( APlayerController *PC, UNetConnection*
 			ReplicationSystem->SetReplicationEnabledForConnection(GetConnectionHandle().GetParentConnectionId(), true);
 		}
 	}
-#endif // UE_WITH_IRIS
 
 	UWorld* World = PlayerController->GetWorld();
 	// if we have already loaded some sublevels, tell the server about them
@@ -5692,7 +5670,6 @@ void UNetConnection::SetClientLoginState( const EClientLoginState::Type NewState
 				TEXT("UNetConnection::SetClientLoginState: State changing from %s to %s"),
 				EClientLoginState::ToString(ClientLoginState), EClientLoginState::ToString(NewState));
 
-#if UE_WITH_IRIS
 	if (UReplicationSystem* ReplicationSystem = Driver->GetReplicationSystem())
 	{
 		if (NewState == EClientLoginState::ReceivedJoin)
@@ -5701,7 +5678,6 @@ void UNetConnection::SetClientLoginState( const EClientLoginState::Type NewState
 			ReplicationSystem->SetReplicationEnabledForConnection(GetConnectionHandle().GetParentConnectionId(), true);
 		}
 	}
-#endif // UE_WITH_IRIS
 
 	ClientLoginState = NewState;
 }
diff --git a/Engine/Source/Runtime/Engine/Private/NetDriver.cpp b/Engine/Source/Runtime/Engine/Private/NetDriver.cpp
--- a/Engine/Source/Runtime/Engine/Private/NetDriver.cpp
+++ b/Engine/Source/Runtime/Engine/Private/NetDriver.cpp
@@ -56,7 +56,6 @@
 #include "Templates/Greater.h"
 #include "UObject/RemoteExecutor.h"
 #include "UObject/Stack.h"
-#if UE_WITH_IRIS
 #include "Iris/IrisConfig.h"
 #include "Iris/Core/IrisDebugging.h"
 #include "Iris/Core/IrisProfiler.h"
@@ -72,7 +71,6 @@
 #include "Iris/ReplicationSystem/NameTokenStore.h"
 #include "Iris/ReplicationSystem/StringTokenStore.h"
 #include "GameplayTagTokenStore.h"
-#endif // UE_WITH_IRIS
 #include "Interfaces/ReplicatedObjectInterface.h"
 
 
@@ -262,7 +260,6 @@ namespace UE::Net::Private
 
 	static bool bIgnoreStaticActorDestruction = false;
 
-#if UE_WITH_IRIS
 	void ApplyReplicationSystemConfig(const FNetDriverReplicationSystemConfig& ReplicationSystemConfig, UReplicationSystem::FReplicationSystemParams& OutParams)
 	{
 		if (ReplicationSystemConfig.MaxReplicatedObjectCount != 0)
@@ -302,7 +299,6 @@ namespace UE::Net::Private
 
 		OutParams.bAllowParallelTasks = ReplicationSystemConfig.bAllowParallelTasks;
 	}
-#endif // UE_WITH_IRIS
 
 	bool IsGuidInOuterChain(const FNetGUIDCache& GuidCache, const FNetGuidCacheObject* CacheObj, FNetworkGUID GuidMatch)
 	{
@@ -697,9 +693,7 @@ UNetDriver::UNetDriver(const FObjectInitializer& ObjectInitializer)
 ,	ProcessQueuedBunchesCurrentFrameMilliseconds(0.0f)
 ,	DDoS()
 ,	LocalAddr(nullptr)
-#if UE_WITH_IRIS
 ,	DummyNetworkObjectInfo(MakePimpl<FNetworkObjectInfo>())
-#endif // UE_WITH_IRIS
 ,	NetworkObjects(new FNetworkObjectList)
 ,	LagState(ENetworkLagState::NotLagging)
 ,	DuplicateLevelID(INDEX_NONE)
@@ -876,13 +870,11 @@ void UNetDriver::NotifyGameInstanceUpdated()
 #endif
 
 #if WITH_EDITOR
-#if UE_WITH_IRIS
 	if (ReplicationSystem && LocalWorld)
 	{
 		ReplicationSystem->SetPIEInstanceID(World->GetPackage()->GetPIEInstanceID());
 	}
 #endif
-#endif
 }
 
 bool UNetDriver::HasSupportForNoPawnConnection()
@@ -926,14 +918,12 @@ void UNetDriver::AddNetworkActor(AActor* Actor)
 
 	UE_LOG(LogNetLifecycle, Log, TEXT("[%s] Registering for replication: Actor: %s"), *GetName(), *GetFullNameSafe(Actor));
 
-#if UE_WITH_IRIS
 	// We currently only want to do this on the server as we want to replicate the actor to the client
 	// but we rely on the old system for actor instantiation
 	if (ReplicationSystem)
 	{
 		return;
 	}
-#endif // UE_WITH_IRIS
 
 	if (!IsDormInitialStartupActor(Actor))
 	{
@@ -965,14 +955,12 @@ FNetworkObjectInfo* UNetDriver::FindOrAddNetworkObjectInfo(const AActor* InActor
 
 	bool bWasAdded = false;
 
-#if UE_WITH_IRIS
 	if (ReplicationSystem)
 	{
 		PRAGMA_DISABLE_DEPRECATION_WARNINGS
 		return FindNetworkObjectInfo(InActor);
 		PRAGMA_ENABLE_DEPRECATION_WARNINGS
 	}
-#endif // UE_WITH_IRIS
 
 	if (TSharedPtr<FNetworkObjectInfo>* InfoPtr = GetNetworkObjectList().FindOrAdd(const_cast<AActor*>(InActor), this, &bWasAdded))
 	{
@@ -989,7 +977,6 @@ FNetworkObjectInfo* UNetDriver::FindOrAddNetworkObjectInfo(const AActor* InActor
 
 FNetworkObjectInfo* UNetDriver::FindNetworkObjectInfo(const AActor* InActor)
 {
-#if UE_WITH_IRIS
 	if (ReplicationSystem)
 	{
 		if (const UObjectReplicationBridge* Bridge = ReplicationSystem->GetReplicationBridge())
@@ -1006,7 +993,6 @@ FNetworkObjectInfo* UNetDriver::FindNetworkObjectInfo(const AActor* InActor)
 
 		return nullptr;
 	}
-#endif // UE_WITH_IRIS
 
 	return GetNetworkObjectList().Find(InActor).Get();
 }
@@ -1145,7 +1131,6 @@ void UNetDriver::TickFlush(float DeltaSeconds)
 #if WITH_SERVER_CODE
 		CSV_SCOPED_TIMING_STAT_EXCLUSIVE(ServerReplicateActors);
 
-#if UE_WITH_IRIS
 		if (ReplicationSystem)
 		{
 			UpdateIrisReplicationViews();
@@ -1153,7 +1138,6 @@ void UNetDriver::TickFlush(float DeltaSeconds)
 			ReplicationSystem->NetUpdate(DeltaSeconds);
 		}
 		else
-#endif // UE_WITH_IRIS
 		if (ClientConnections.Num() > 0)
 		{
 #if UE_WITH_REMOTE_OBJECT_HANDLE
@@ -1173,7 +1157,6 @@ void UNetDriver::TickFlush(float DeltaSeconds)
 		}
 #endif // WITH_SERVER_CODE
 	}
-#if UE_WITH_IRIS
 	else if (!IsServer() && ServerConnection != nullptr)
 	{
 		if (ReplicationSystem)
@@ -1182,7 +1165,6 @@ void UNetDriver::TickFlush(float DeltaSeconds)
 			ReplicationSystem->NetUpdate(DeltaSeconds);
 		}
 	}
-#endif // UE_WITH_IRIS
 
 	// Reset queued bunch amortization timer
 	ProcessQueuedBunchesCurrentFrameMilliseconds = 0.0f;
@@ -1646,12 +1628,10 @@ void UNetDriver::PostTickFlush()
 	// We are still ticking here
 	{
 		TGuardValue<bool> GuardInNetTick(bInTick, true);
-#if UE_WITH_IRIS
 		if (ReplicationSystem && (ServerConnection != nullptr || (ClientConnections.Num() > 0) || UE::Net::Private::bUpdateReplicationSystemWithNoConnections))
 		{
 			ReplicationSystem->PostSendUpdate();
 		}
-#endif // UE_WITH_IRIS
 
 		if (World && !bSkipClearVoicePackets)
 		{
@@ -1770,7 +1750,6 @@ bool UNetDriver::InitBase(bool bInitAsClient, FNetworkNotify* InNotify, const FU
 
 	Notify = InNotify;
 
-#if UE_WITH_IRIS
 	// At the time being using NetTokens and NetTokenStore is only available when compiling with iris.
 	// Init NetToken stores
 	if (!GetNetTokenStore())
@@ -1792,14 +1771,11 @@ bool UNetDriver::InitBase(bool bInitAsClient, FNetworkNotify* InNotify, const FU
 		NetTokenStore->CreateAndRegisterDataStore<FGameplayTagTokenStore>();
 		OnNetTokenStoreReadyDelegate.Broadcast(this);
 	}
-#endif
 	
-#if UE_WITH_IRIS
 	if (IsUsingIrisReplication() && !ReplicationSystem)
 	{
 		CreateReplicationSystem(bInitAsClient);
 	}
-#endif // UE_WITH_IRIS
 
 	
 	if (NetDriverDefinition == NAME_GameNetDriver)
@@ -2250,12 +2226,10 @@ FString UNetDriver::GetReplicationModelName() const
 	{
 		return RepDriver->GetClass()->GetName();
 	}
-#if UE_WITH_IRIS
 	else if (ReplicationSystem)
 	{
 		return TEXT("Iris");
 	}
-#endif
 	else
 	{
 		return TEXT("Generic");
@@ -2382,21 +2356,17 @@ void UNetDriver::InternalTickDispatch(float DeltaSeconds)
 {
 	TGuardValue<bool> GuardInNetTick(bInTick, true);
 
-#if UE_WITH_IRIS
 	if (ReplicationSystem)
 	{
 		ReplicationSystem->PreReceiveUpdate();
 	}
-#endif
 
 	TickDispatch(DeltaSeconds);
 
-#if UE_WITH_IRIS
 	if (ReplicationSystem)
 	{
 		ReplicationSystem->PostReceiveUpdate();
 	}
-#endif
 }
 
 void UNetDriver::InternalTickFlush(float DeltaSeconds)
@@ -2504,9 +2474,7 @@ void UNetDriver::Shutdown()
 	ConnectionlessHandler.Reset();
 
 	SetReplicationDriver(nullptr);
-#if UE_WITH_IRIS
 	SetReplicationSystem(nullptr);
-#endif // UE_WITH_IRIS
 
 	// End NetTrace session for this instance
 	UE_NET_TRACE_END_SESSION(GetNetTraceId());
@@ -2585,7 +2553,6 @@ EEngineNetworkRuntimeFeatures UNetDriver::GetNetworkRuntimeFeatures() const
 
 void UNetDriver::TryUpgradeNetworkFeatures(EEngineNetworkRuntimeFeatures RemoteFeatures)
 {
-#if UE_WITH_IRIS
 	const bool bIsRemoteUsingIris = EnumHasAnyFlags(RemoteFeatures, EEngineNetworkRuntimeFeatures::IrisEnabled);
 
 	// If our netdriver needs to enable Iris
@@ -2637,7 +2604,6 @@ void UNetDriver::TryUpgradeNetworkFeatures(EEngineNetworkRuntimeFeatures RemoteF
 	}
 
 	ensureMsgf(GetNetworkRuntimeFeatures() == RemoteFeatures, TEXT("UNetDriver::TryUpgradeNetworkFeatures failed for %s"), *GetDescription());
-#endif
 }
 
 // ----------------------------------------------------------------------------------------
@@ -2858,9 +2824,7 @@ void UNetDriver::PostTickDispatch()
 		}
 	}
 
-#if UE_WITH_IRIS
 	PostDispatchSendUpdate();
-#endif
 
 	if (ReplicationDriver)
 	{
@@ -4310,7 +4274,6 @@ bool UNetDriver::CanCreateDestructionInfo() const
 
 void UNetDriver::NotifyActorDestroyed(AActor* ThisActor, bool IsSeamlessTravel, EChannelCloseReason CloseReason)
 {
-#if UE_WITH_IRIS
 	if (ReplicationSystem)
 	{
 		// For actors that are DORM_Initial when destroyed we need to create and replicate a destruction info.
@@ -4342,7 +4305,6 @@ void UNetDriver::NotifyActorDestroyed(AActor* ThisActor, bool IsSeamlessTravel,
 		// For the other cases, Iris handles replication of destroy through EndReplication
 		return;
 	}
-#endif // UE_WITH_IRIS
 
 	const bool bIsServer = IsServer();
 	
@@ -4463,7 +4425,6 @@ void UNetDriver::RemoveNetworkActor(AActor* Actor)
 
 void UNetDriver::DeleteSubObjectOnClients(AActor* Actor, UObject* SubObject)
 {
-#if UE_WITH_IRIS
 	if (ReplicationSystem)
 	{
 		if (UEngineReplicationBridge* Bridge = Cast<UEngineReplicationBridge>(ReplicationSystem->GetReplicationBridge()))
@@ -4473,7 +4434,6 @@ void UNetDriver::DeleteSubObjectOnClients(AActor* Actor, UObject* SubObject)
 		}
 	}
 	else
-#endif 
 	{
 #if UE_REPLICATED_OBJECT_REFCOUNTING
 		NetworkObjects->SetSubObjectForDeletion(Actor, SubObject);
@@ -4483,7 +4443,6 @@ void UNetDriver::DeleteSubObjectOnClients(AActor* Actor, UObject* SubObject)
 
 void UNetDriver::TearOffSubObjectOnClients(AActor* Actor, UObject* SubObject)
 {
-#if UE_WITH_IRIS
 	if (ReplicationSystem)
 	{
 		if (UEngineReplicationBridge* Bridge = Cast<UEngineReplicationBridge>(ReplicationSystem->GetReplicationBridge()))
@@ -4492,7 +4451,6 @@ void UNetDriver::TearOffSubObjectOnClients(AActor* Actor, UObject* SubObject)
 		}
 	}
 	else
-#endif 
 	{
 #if UE_REPLICATED_OBJECT_REFCOUNTING
 		NetworkObjects->SetSubObjectForTearOff(Actor, SubObject);
@@ -4558,7 +4516,6 @@ void UNetDriver::NotifyActorRenamed(AActor* ThisActor, UObject* PreviousOuter, F
 					ReplicationDriver->NotifyActorRenamed(ThisActor, PreviousOuter, PreviousName);
 				}
 
-#if UE_WITH_IRIS
 				if (ReplicationSystem)
 				{
 					if (UEngineReplicationBridge* Bridge = ReplicationSystem->GetReplicationBridgeAs<UEngineReplicationBridge>())
@@ -4567,7 +4524,6 @@ void UNetDriver::NotifyActorRenamed(AActor* ThisActor, UObject* PreviousOuter, F
 					}
 				}
 				else
-#endif // UE_WITH_IRIS
 
 				if (UE::Net::Private::CleanUpRenamedDynamicActors)
 				{
@@ -4609,12 +4565,10 @@ void UNetDriver::NotifyStreamingLevelUnload(ULevel* Level)
 			}
 		}
 
-#if UE_WITH_IRIS
 		if (ReplicationSystem)
 		{
 			ReplicationSystem->NotifyStreamingLevelUnload(Level);
 		}
-#endif // UE_WITH_IRIS
 		
 		TArray<FNetworkGUID> RemovedGUIDs;
 		for (auto It = DestroyedStartupOrDormantActors.CreateIterator(); It; ++It)
@@ -4696,7 +4650,6 @@ void UNetDriver::NotifyActorTearOff(AActor* Actor)
 	{
 		ReplicationDriver->NotifyActorTearOff(Actor);
 	}
-#if UE_WITH_IRIS
 	if (ReplicationSystem)
 	{
 		if (UObjectReplicationBridge* Bridge = ReplicationSystem->GetReplicationBridge())
@@ -4706,17 +4659,14 @@ void UNetDriver::NotifyActorTearOff(AActor* Actor)
 			ReplicationSystem->TearOffNextUpdate(ActorRefHandle);
 		}
 	}
-#endif // UE_WITH_IRIS
 }
 
 void UNetDriver::NotifyActorIsTraveling(AActor* TravelingActor)
 {
-#if UE_WITH_IRIS
 	if (ReplicationSystem)
 	{
 		return;
 	}
-#endif
 
 	NetworkObjects->OnActorIsTraveling(TravelingActor);
 }
@@ -4730,7 +4680,6 @@ void UNetDriver::ForceNetUpdate(AActor* Actor)
 		return;
 	}
 
-#if UE_WITH_IRIS
 	if (ReplicationSystem)
 	{
 		if (UObjectReplicationBridge* Bridge = ReplicationSystem->GetReplicationBridge())
@@ -4740,7 +4689,6 @@ void UNetDriver::ForceNetUpdate(AActor* Actor)
 		}
 		return;
 	}
-#endif // UE_WITH_IRIS
 	
 	// Legacy implementation
 	PRAGMA_DISABLE_DEPRECATION_WARNINGS
@@ -4799,13 +4747,11 @@ void UNetDriver::FlushActorDormancy(AActor* Actor, bool bWasDormInitial)
 		ReplicationDriver->FlushNetDormancy(Actor, bWasDormInitial);
 	}
 
-#if UE_WITH_IRIS
 	if (ReplicationSystem)
 	{
 		UE::Net::FReplicationSystemUtil::FlushNetDormancy(ReplicationSystem, Actor, bWasDormInitial);
 	}
 	else
-#endif // UE_WITH_IRIS
 	{
 		FlushActorDormancyInternal(Actor);
 	}
@@ -4818,13 +4764,11 @@ void UNetDriver::NotifyActorDormancyChange(AActor* Actor, ENetDormancy OldDorman
 		ReplicationDriver->NotifyActorDormancyChange(Actor, OldDormancyState);
 	}
 
-#if UE_WITH_IRIS
 	if (ReplicationSystem)
 	{
 		UE::Net::FReplicationSystemUtil::NotifyActorDormancyChange(ReplicationSystem, Actor, OldDormancyState);
 	}
 	else
-#endif // UE_WITH_IRIS
 	{
 		FlushActorDormancyInternal(Actor);
 	}
@@ -6481,9 +6425,7 @@ void UNetDriver::ReinitBase()
 
 void UNetDriver::PostCreation(bool bInitializeWithIris)
 {
-#if UE_WITH_IRIS
 	bIsUsingIris = bInitializeWithIris;
-#endif //UE_WITH_IRIS
 
 	if (NetDriverDefinition == NAME_GameNetDriver)
 	{
@@ -6599,9 +6541,7 @@ void UNetDriver::DrawNetDriverDebug()
 	// Used to draw additional boxes to show more state
 	const FBox BoxExpansion = FBox(&Extent, 1);
 
-#if UE_WITH_IRIS
 	const UObjectReplicationBridge* Bridge = ReplicationSystem ? ReplicationSystem->GetReplicationBridge() : nullptr;
-#endif
 
 	for (FActorIterator It(LocalWorld); It; ++It)
 	{
@@ -6666,7 +6606,6 @@ void UNetDriver::DrawNetDriverDebug()
 		// Draw NetGUID or NetHandle
 		FVector AboveActor(0.f, 0.f, 32.f);
 #if !UE_BUILD_SHIPPING
-#if UE_WITH_IRIS 
 		if (Bridge)
 		{
 			AActor* Actor = *It;
@@ -6682,7 +6621,6 @@ void UNetDriver::DrawNetDriverDebug()
 			}
 		}
 		else
-#endif
 		{
 			FNetworkGUID NetGUID = Connection->PackageMap->GetNetGUIDFromObject(*It);
 			if (NetGUID.IsValid())
@@ -6703,7 +6641,7 @@ void UNetDriver::DrawNetDriverDebug()
 	}
 
 	// Draw state
-#if UE_WITH_IRIS && !UE_BUILD_SHIPPING
+#if !UE_BUILD_SHIPPING
 	if (Bridge)
 	{
 		UE::Net::FNetRefHandle DebugRefHandle = UE::Net::IrisDebugHelper::GetDebugNetRefHandle();
@@ -6717,7 +6655,7 @@ void UNetDriver::DrawNetDriverDebug()
 			}
 		}
 	}
-#endif
+#endif // !UE_BUILD_SHIPPING
 
 
 #endif
@@ -6857,7 +6795,6 @@ void UNetDriver::CreateReplicatedStaticActorDestructionInfo(ULevel* Level, const
 {
 	check(Level);
 
-#if UE_WITH_IRIS
 	if (ReplicationSystem)
 	{
 		if (UEngineReplicationBridge* Bridge = ReplicationSystem->GetReplicationBridgeAs<UEngineReplicationBridge>())
@@ -6875,7 +6812,6 @@ void UNetDriver::CreateReplicatedStaticActorDestructionInfo(ULevel* Level, const
 		}
 		return;
 	}
-#endif // UE_WITH_IRIS
 
 	FNetworkGUID NetGUID = GetNetGuidCache()->AssignNewNetGUIDFromPath_Server(Info.PathName.ToString(), Info.ObjOuter.Get(), Info.ObjClass);
 	if (NetGUID.IsDefault())
@@ -7062,7 +6998,6 @@ void UNetDriver::SetWorld(class UWorld* InWorld)
 		Notify = InWorld;
 		RegisterTickEvents(InWorld);
 
-#if UE_WITH_IRIS
 		if (IsUsingIrisReplication())
 		{
 			if (bSkipStartReplicatingWorldForIris)
@@ -7077,7 +7012,6 @@ void UNetDriver::SetWorld(class UWorld* InWorld)
 			}
 		}
 		else
-#endif // UE_WITH_IRIS
 		{
 			GetNetworkObjectList().AddInitialObjects(InWorld, this);
 		}
@@ -7238,14 +7172,12 @@ void UNetDriver::SetReplicationDriver(UReplicationDriver* NewReplicationDriver)
 		}
 	}
 
-#if UE_WITH_IRIS
 	if (IsUsingIrisReplication() && NewReplicationDriver)
 	{
 		checkf(false, TEXT("ReplicationDriver (%s) are not supported with a NetDriver (%s) configured to use Iris"), *NewReplicationDriver->GetName(), *NetDriverName.ToString());
 		NewReplicationDriver->MarkAsGarbage();
 		return;
 	}
-#endif
 
 	ReplicationDriver = NewReplicationDriver;
 	if (ReplicationDriver)
@@ -7294,7 +7226,6 @@ UNetConnection* UNetDriver::GetConnectionByHandle(UE::Net::FConnectionHandle Con
 	return nullptr;
 }
 
-#if UE_WITH_IRIS
 bool UNetDriver::InitReplicationBridgeClass()
 {
 	if (ReplicationBridgeClass != nullptr)
@@ -7699,7 +7630,6 @@ void UNetDriver::PostDispatchSendUpdate()
 	ReplicationSystem->PostSendUpdate();
 }
 
-#endif // UE_WITH_IRIS
 
 void UNetDriver::InitNetTraceId()
 {
@@ -8049,7 +7979,6 @@ void UNetDriver::ProcessRemoteFunction(
 			LocalOutParms = UE::Net::Private::CopyOutParametersToLocalParameters(Function, OutParms, Parameters, TargetObj);
 		}
 
-#if UE_WITH_IRIS
 		if (ReplicationSystem)
 		{
 			if (bIsServerMulticast)
@@ -8073,7 +8002,6 @@ void UNetDriver::ProcessRemoteFunction(
 			// If we are using Iris replication, we should never fall back on normal replication path
 			return;
 		}
-#endif // UE_WITH_IRIS
 
 		// Forward to replication Driver if there is one
 		if (ReplicationDriver && ReplicationDriver->ProcessRemoteFunction(Actor, Function, Parameters, OutParms, Stack, SubObject))
@@ -9100,7 +9028,6 @@ UE::Net::FNetIDVariant UNetDriver::CreateNetIDForObject(UObject* Object)
 {
 	using namespace UE::Net;
 
-#if UE_WITH_IRIS
 	if (bIsUsingIris)
 	{
 		if (ReplicationSystem)
@@ -9125,7 +9052,6 @@ UE::Net::FNetIDVariant UNetDriver::CreateNetIDForObject(UObject* Object)
 		return FNetIDVariant(FNetRefHandle());
 	}
 	else
-#endif
 	{
 		FNetworkGUID GUID = GuidCache->AssignNewNetGUID_Server(Object);
 		return FNetIDVariant(GUID);
@@ -9134,7 +9060,6 @@ UE::Net::FNetIDVariant UNetDriver::CreateNetIDForObject(UObject* Object)
 
 void UNetDriver::AssignNetIDToObject(UE::Net::FNetIDVariant NetID, const UObject* Object)
 {
-#if UE_WITH_IRIS
 	if (bIsUsingIris)
 	{
 		if (NetID.GetVariant().IsType<UE::Net::FNetRefHandle>())
@@ -9164,7 +9089,6 @@ void UNetDriver::AssignNetIDToObject(UE::Net::FNetIDVariant NetID, const UObject
 		}
 	}
 	else
-#endif
 	{
 		if (NetID.GetVariant().IsType<FNetworkGUID>())
 		{
diff --git a/Engine/Source/Runtime/Engine/Private/NetworkObjectList.cpp b/Engine/Source/Runtime/Engine/Private/NetworkObjectList.cpp
--- a/Engine/Source/Runtime/Engine/Private/NetworkObjectList.cpp
+++ b/Engine/Source/Runtime/Engine/Private/NetworkObjectList.cpp
@@ -43,12 +43,10 @@ void FNetworkObjectList::AddInitialObjects(UWorld* const World, UNetDriver* NetD
 		return;
 	}
 
-#if UE_WITH_IRIS
 	if (NetDriver->GetReplicationSystem())
 	{
 		return;
 	}
-#endif // UE_WITH_IRIS
 
 
 	for (FActorIterator Iter(World); Iter; ++Iter)
@@ -87,12 +85,10 @@ TSharedPtr<FNetworkObjectInfo>* FNetworkObjectList::FindOrAdd(AActor* const Acto
 		return nullptr;
 	}
 
-#if UE_WITH_IRIS
 	if (NetDriver && NetDriver->GetReplicationSystem())
 	{
 		return nullptr;
 	}
-#endif // UE_WITH_IRIS
 
 	TSharedPtr<FNetworkObjectInfo>* NetworkObjectInfo = AllNetworkObjects.Find(Actor);
 	
diff --git a/Engine/Source/Runtime/Engine/Private/PackageMapClient.cpp b/Engine/Source/Runtime/Engine/Private/PackageMapClient.cpp
--- a/Engine/Source/Runtime/Engine/Private/PackageMapClient.cpp
+++ b/Engine/Source/Runtime/Engine/Private/PackageMapClient.cpp
@@ -23,9 +23,7 @@
 #include "Misc/CommandLine.h"
 #include "Net/NetSubObjectRegistryGetter.h"
 
-#if UE_WITH_IRIS
 #include "Iris/ReplicationSystem/NetTokenStore.h"
-#endif
 
 #include UE_INLINE_GENERATED_CPP_BY_NAME(PackageMapClient)
 
@@ -297,14 +295,12 @@ void UPackageMapClient::Initialize(UNetConnection * InConnection, TSharedPtr<FNe
 	ExportNetGUIDCount = 0;
 	OverrideAckState = &AckState;
 
-#if UE_WITH_IRIS
 	// Initialize NetTokenResolveContext
 	NetTokenResolveContext.NetTokenStore = InConnection->GetDriver()->GetNetTokenStore();
 	if (NetTokenResolveContext.NetTokenStore)
 	{
 		NetTokenResolveContext.RemoteNetTokenStoreState = NetTokenResolveContext.NetTokenStore->GetRemoteNetTokenStoreState(InConnection->GetConnectionHandle().GetParentConnectionId());
 	}
-#endif
 }
 
 /**
@@ -2238,7 +2234,6 @@ FOutBunch* UPackageMapClient::CreateCustomExportsBunch(const FOutBunch& Outgoing
 	FOutBunch* ExportsBunch = nullptr;
 
 	// NetTokenExports is still an experimental feature that depends on Iris code.
-#if UE_WITH_IRIS
 	// NetTokens - begin
 	using namespace UE::Net;
 	TArray<UE::Net::FNetToken, TInlineAllocator<4>> NetTokensExportedInThisBunch;
@@ -2300,7 +2295,6 @@ FOutBunch* UPackageMapClient::CreateCustomExportsBunch(const FOutBunch& Outgoing
 	}
 
 	check(!ExportsBunch->IsError());
-#endif // UE_WITH_IRIS
 
 	return ExportsBunch;
 }
@@ -2311,7 +2305,6 @@ void UPackageMapClient::ReceiveCustomExportsBunch(FInBunch &InBunch)
 	check(InBunch.bPartialCustomExportsFinal);
 
 	// NetTokenExports is still an experimental feature that depends on Iris code.
-#if UE_WITH_IRIS
 	using namespace UE::Net;
 
 	// Read NetTokenExports
@@ -2332,7 +2325,6 @@ void UPackageMapClient::ReceiveCustomExportsBunch(FInBunch &InBunch)
 		// Continuation bit
 		InBunch.SerializeBits(&bHasExportsToRead, 1);
 	}
-#endif // UE_WITH_IRIS
 }
 
 int32 UPackageMapClient::GetNumExportBunches() const
diff --git a/Engine/Source/Runtime/Engine/Private/Pawn.cpp b/Engine/Source/Runtime/Engine/Private/Pawn.cpp
--- a/Engine/Source/Runtime/Engine/Private/Pawn.cpp
+++ b/Engine/Source/Runtime/Engine/Private/Pawn.cpp
@@ -656,10 +656,8 @@ void APawn::PossessedBy(AController* NewController)
 	SetController(NewController);
 	ForceNetUpdate();
 
-#if UE_WITH_IRIS
 	// The owning connection depends on the Controller having the new value.
 	UpdateOwningNetConnection();
-#endif
 
 	if (GetController()->PlayerState != nullptr)
 	{
@@ -701,10 +699,8 @@ void APawn::UnPossessed()
 	SetOwner(nullptr);
 	SetController(nullptr);
 
-#if UE_WITH_IRIS
 	// The owning connection depends on the Controller having the new value.
 	UpdateOwningNetConnection();
-#endif
 
 	// Unregister input component if we created one
 	DestroyPlayerInputComponent();
diff --git a/Engine/Source/Runtime/Engine/Private/PhysicsEngine/NetworkPhysicsComponent.cpp b/Engine/Source/Runtime/Engine/Private/PhysicsEngine/NetworkPhysicsComponent.cpp
--- a/Engine/Source/Runtime/Engine/Private/PhysicsEngine/NetworkPhysicsComponent.cpp
+++ b/Engine/Source/Runtime/Engine/Private/PhysicsEngine/NetworkPhysicsComponent.cpp
@@ -12,9 +12,7 @@
 #include "Chaos/PhysicsObjectInternalInterface.h"
 #include "Net/Core/PushModel/PushModel.h"
 
-#if UE_WITH_IRIS
 #include "Iris/ReplicationState/PropertyNetSerializerInfoRegistry.h"
-#endif // UE_WITH_IRIS
 
 #include UE_INLINE_GENERATED_CPP_BY_NAME(NetworkPhysicsComponent)
 
@@ -190,7 +188,6 @@ FNetworkPhysicsRewindDataProxy& FNetworkPhysicsRewindDataProxy::operator=(const
 	return *this;
 }
 
-#if UE_WITH_IRIS
 UE_NET_IMPLEMENT_NAMED_STRUCT_LASTRESORT_NETSERIALIZER_AND_REGISTRY_DELEGATES(NetworkPhysicsRewindDataInputProxy);
 UE_NET_IMPLEMENT_NAMED_STRUCT_LASTRESORT_NETSERIALIZER_AND_REGISTRY_DELEGATES(NetworkPhysicsRewindDataRemoteInputProxy);
 UE_NET_IMPLEMENT_NAMED_STRUCT_LASTRESORT_NETSERIALIZER_AND_REGISTRY_DELEGATES(NetworkPhysicsRewindDataStateProxy);
@@ -198,7 +195,6 @@ UE_NET_IMPLEMENT_NAMED_STRUCT_LASTRESORT_NETSERIALIZER_AND_REGISTRY_DELEGATES(Ne
 UE_NET_IMPLEMENT_NAMED_STRUCT_LASTRESORT_NETSERIALIZER_AND_REGISTRY_DELEGATES(NetworkPhysicsRewindDataImportantStateProxy);
 UE_NET_IMPLEMENT_NAMED_STRUCT_LASTRESORT_NETSERIALIZER_AND_REGISTRY_DELEGATES(NetworkPhysicsRewindDataDeltaSourceInputProxy);
 UE_NET_IMPLEMENT_NAMED_STRUCT_LASTRESORT_NETSERIALIZER_AND_REGISTRY_DELEGATES(NetworkPhysicsRewindDataDeltaSourceStateProxy);
-#endif // UE_WITH_IRIS
 
 bool FNetworkPhysicsRewindDataInputProxy::NetSerialize(FArchive& Ar, class UPackageMap* Map, bool& bOutSuccess)
 {
diff --git a/Engine/Source/Runtime/Engine/Private/PlayerController.cpp b/Engine/Source/Runtime/Engine/Private/PlayerController.cpp
--- a/Engine/Source/Runtime/Engine/Private/PlayerController.cpp
+++ b/Engine/Source/Runtime/Engine/Private/PlayerController.cpp
@@ -69,11 +69,9 @@
 #include "PBDRigidsSolver.h"
 #include "PhysicsEngine/PhysicsSettings.h"
 
-#if UE_WITH_IRIS
 #include "Iris/ReplicationSystem/ReplicationSystem.h"
 #include "Net/Iris/ReplicationSystem/EngineReplicationBridge.h"
 #include "Net/Iris/ReplicationSystem/ReplicationSystemUtil.h"
-#endif // UE_WITH_IRIS
 
 #if UE_WITH_REMOTE_OBJECT_HANDLE
 #include "AutoRTFM.h"
@@ -4949,7 +4947,6 @@ void APlayerController::ReceivedSpectatorClass(TSubclassOf<ASpectatorPawn> Spect
 
 void APlayerController::SetPawn(APawn* InPawn)
 {
-#if UE_WITH_IRIS
 	if (GetLocalRole() == ROLE_Authority && UE::Net::FReplicationSystemUtil::GetNetHandle(this).IsValid())
 	{
 		if (const UReplicationSystem* ReplicationSystem = UE::Net::FReplicationSystemUtil::GetReplicationSystem(this))
@@ -4968,7 +4965,6 @@ void APlayerController::SetPawn(APawn* InPawn)
 			}
 		}
 	}
-#endif
 
 	if (InPawn == NULL)
 	{
@@ -5244,10 +5240,8 @@ void APlayerController::SetPlayer( UPlayer* InPlayer )
 		{
 			NetConnection->OwningActor = this;
 
-#if UE_WITH_IRIS
 			UpdateOwningNetConnection();
 			UE::Net::FReplicationSystemUtil::UpdateSubObjectGroupMemberships(this);
-#endif // UE_WITH_IRIS
 		}
 	}
 
@@ -6458,21 +6452,16 @@ void APlayerController::IncludeInNetConditionGroup(FName NetGroup)
 	checkf(!NetGroup.IsNone(), TEXT("Invalid netcondition group: NONE"));
 	NetConditionGroups.AddUnique(NetGroup);
 
-#if UE_WITH_IRIS
 	UE::Net::FReplicationSystemUtil::UpdateSubObjectGroupMemberships(this);
-#endif
 }
 
 void APlayerController::RemoveFromNetConditionGroup(FName NetGroup)
 {
 	NetConditionGroups.RemoveSingleSwap(NetGroup);
-#if UE_WITH_IRIS
 	UE::Net::FReplicationSystemUtil::RemoveSubObjectGroupMembership(this, NetGroup);
-#endif
 }
 
 
-#if UE_WITH_IRIS
 void APlayerController::OnReplicationStartedForIris(const FOnReplicationStartedParams& Params)
 {
 	using namespace UE::Net;
@@ -6486,7 +6475,6 @@ void APlayerController::OnReplicationStartedForIris(const FOnReplicationStartedP
 	// Enable groups once owner is set!!
 	FReplicationSystemUtil::UpdateSubObjectGroupMemberships(this);
 }
-#endif // UE_WITH_IRIS
 
 void APlayerController::ExecuteAsyncPhysicsCommand(const FAsyncPhysicsTimestamp& AsyncPhysicsTimestamp, UObject* OwningObject, const TFunction<void()>& Command, const bool bEnableResim)
 {
diff --git a/Engine/Source/Runtime/Engine/Private/PlayerState.cpp b/Engine/Source/Runtime/Engine/Private/PlayerState.cpp
--- a/Engine/Source/Runtime/Engine/Private/PlayerState.cpp
+++ b/Engine/Source/Runtime/Engine/Private/PlayerState.cpp
@@ -7,11 +7,9 @@
 #include "Net/OnlineEngineInterface.h"
 #include "GameFramework/GameStateBase.h"
 #include "Net/Core/PushModel/PushModel.h"
-#if UE_WITH_IRIS
 #include "Net/Iris/ReplicationSystem/ReplicationSystemUtil.h"
 #include "Iris/ReplicationSystem/ReplicationSystem.h"
 #include "Iris/ReplicationSystem/Prioritization/NetObjectPrioritizer.h"
-#endif
 
 #include UE_INLINE_GENERATED_CPP_BY_NAME(PlayerState)
 
diff --git a/Engine/Source/Runtime/Engine/Private/RepLayout.cpp b/Engine/Source/Runtime/Engine/Private/RepLayout.cpp
--- a/Engine/Source/Runtime/Engine/Private/RepLayout.cpp
+++ b/Engine/Source/Runtime/Engine/Private/RepLayout.cpp
@@ -33,10 +33,8 @@
 
 #include "Stats/StatsTrace.h"
 #include "UObject/EnumProperty.h"
-#if UE_WITH_IRIS
 #include "Iris/IrisConfig.h"
 #include "Net/Iris/ReplicationSystem/ReplicationSystemUtil.h"
-#endif // UE_WITH_IRIS
 
 DECLARE_CYCLE_STAT(TEXT("RepLayout InitFromObjectClass"), STAT_RepLayout_InitFromObjectClass, STATGROUP_Game);
 DECLARE_CYCLE_STAT(TEXT("RepLayout DeltaSerializeFastArray"), STAT_RepLayout_DeltaSerializeFastArray, STATGROUP_Game);
diff --git a/Engine/Source/Runtime/Engine/Private/RepMovementNetSerializer.cpp b/Engine/Source/Runtime/Engine/Private/RepMovementNetSerializer.cpp
--- a/Engine/Source/Runtime/Engine/Private/RepMovementNetSerializer.cpp
+++ b/Engine/Source/Runtime/Engine/Private/RepMovementNetSerializer.cpp
@@ -4,7 +4,6 @@
 
 #include UE_INLINE_GENERATED_CPP_BY_NAME(RepMovementNetSerializer)
 
-#if UE_WITH_IRIS
 #include "Engine/ReplicatedState.h"
 #include "Iris/Serialization/BitPacking.h"
 #include "Iris/Serialization/NetBitStreamReader.h"
@@ -977,4 +976,3 @@ void FRepMovementNetSerializer::FNetSerializerRegistryDelegates::InitNetSerializ
 
 }
 
-#endif
diff --git a/Engine/Source/Runtime/Engine/Private/UnrealEngine.cpp b/Engine/Source/Runtime/Engine/Private/UnrealEngine.cpp
--- a/Engine/Source/Runtime/Engine/Private/UnrealEngine.cpp
+++ b/Engine/Source/Runtime/Engine/Private/UnrealEngine.cpp
@@ -132,10 +132,8 @@ UnrealEngine.cpp: Implements the UEngine class and helpers.
 #include "AssetCompilingManager.h"
 #include "UObject/UObjectSerializeContext.h"
 
-#if UE_WITH_IRIS
 #include "Iris/IrisConfig.h"
 #include "Iris/ReplicationSystem/NetObjectFactoryRegistry.h"
-#endif
 
 #include "IUniversalObjectLocatorModule.h"
 
@@ -240,9 +238,7 @@ UnrealEngine.cpp: Implements the UEngine class and helpers.
 #include "Streaming/StreamingManagerTexture.h"
 #endif
 
-#if UE_WITH_IRIS
 #include "Net/Iris/ReplicationSystem/NetEngineFactories.h"
-#endif
 
 #include "HAL/FileManagerGeneric.h"
 #include "UObject/ReferenceChainSearch.h"
@@ -362,13 +358,11 @@ void FEngineModule::StartupModule()
 
 	UE::Anim::FSkeletonRemappingRegistry::Init();
 
-#if UE_WITH_IRIS
 	FDelayedAutoRegisterHelper(EDelayedRegisterRunPhase::ObjectSystemReady,
 	[]
 	{
 		UE::Net::InitEngineNetObjectFactories();
 	});
-#endif
 
 	IUniversalObjectLocatorModule& UolModule = FModuleManager::Get().LoadModuleChecked<IUniversalObjectLocatorModule>("UniversalObjectLocator");
 
@@ -404,9 +398,7 @@ void FEngineModule::ShutdownModule()
 	IPrimitiveComponent::RemoveImplementer(UPrimitiveComponent::StaticClass());
 	IStaticMeshComponent::RemoveImplementer(UStaticMeshComponent::StaticClass());
 
-#if UE_WITH_IRIS
 	UE::Net::ShutdownEngineNetObjectFactories();
-#endif
 
 #if TRACE_FILTERING_ENABLED
 	FTraceFilter::Destroy();
@@ -14387,7 +14379,6 @@ const FIrisNetDriverConfig* UEngine::GetIrisNetDriverConfig(FName InNetDriverDef
 	 *
 	 */
 
-#if UE_WITH_IRIS
 	// Search for the exact name match
 	if (IrisConfig == nullptr)
 	{
@@ -14415,14 +14406,12 @@ const FIrisNetDriverConfig* UEngine::GetIrisNetDriverConfig(FName InNetDriverDef
 			return Config.NetDriverDefinition.IsNone() ? false : Config.NetDriverDefinition == InNetDriverDefinition;
 		});
 	}
-#endif //#if UE_WITH_IRIS
 
 	return IrisConfig;
 }
 
 bool UEngine::WillNetDriverUseIris(const FWorldContext& Context, FName InNetDriverDefinition, FName InNetDriverName) const
 {
-#if UE_WITH_IRIS
 	const FIrisNetDriverConfig* IrisConfig = GetIrisNetDriverConfig(InNetDriverDefinition, InNetDriverName);
 
 	const bool bConfigCanUseIris = IrisConfig && IrisConfig->bCanUseIris;
@@ -14494,9 +14483,6 @@ bool UEngine::WillNetDriverUseIris(const FWorldContext& Context, FName InNetDriv
 	}
 
 	return bUseIrisRepSystem;
-#else
-	return false;
-#endif //UE_WITH_IRIS
 }
 
 namespace UE::Private
@@ -14638,12 +14624,10 @@ namespace UE::Private
 				const FName DriverName = InNetDriverName.IsNone() ? ReturnVal->GetFName() : InNetDriverName;
 				const bool bInitializeWithIris = Engine->WillNetDriverUseIris(Context, NetDriverDefinition, DriverName);
 
-#if UE_WITH_IRIS
 				if (bInitializeWithIris)
 				{
 					UE::Net::FNetObjectFactoryRegistry::SetFactoryRegistrationAllowed(false);
 				}
-#endif
 
 				ReturnVal->SetNetDriverName(DriverName);
 				ReturnVal->SetNetDriverDefinition(NetDriverDefinition);
@@ -14721,7 +14705,6 @@ namespace UE::Private
 			}
 		}
 
-#if UE_WITH_IRIS
 		auto IsAnyNetDriverUsingIris = []()
 		{
 			const TIndirectArray<FWorldContext>& Worlds = GEngine->GetWorldContexts();
@@ -14742,7 +14725,6 @@ namespace UE::Private
 		{
 			UE::Net::FNetObjectFactoryRegistry::SetFactoryRegistrationAllowed(true);
 		}
-#endif // UE_WITH_IRIS
 	}
 } // end namespace UE::Private
 
diff --git a/Engine/Source/Runtime/Engine/Private/UnrealNetwork.cpp b/Engine/Source/Runtime/Engine/Private/UnrealNetwork.cpp
--- a/Engine/Source/Runtime/Engine/Private/UnrealNetwork.cpp
+++ b/Engine/Source/Runtime/Engine/Private/UnrealNetwork.cpp
@@ -52,9 +52,7 @@ void RegisterReplicatedLifetimeProperty(
 
 		FLifetimeProperty LifetimeProp(RepIndex, Params.Condition, Params.RepNotifyCondition, Params.bIsPushBased);
 
-#if UE_WITH_IRIS
 		LifetimeProp.CreateAndRegisterReplicationFragmentFunction = Params.CreateAndRegisterReplicationFragmentFunction;
-#endif
 
 		if (RegisteredPropertyPtr)
 		{
diff --git a/Engine/Source/Runtime/Engine/Private/World.cpp b/Engine/Source/Runtime/Engine/Private/World.cpp
--- a/Engine/Source/Runtime/Engine/Private/World.cpp
+++ b/Engine/Source/Runtime/Engine/Private/World.cpp
@@ -135,10 +135,8 @@
 #include "ProfilingDebugging/CountersTrace.h"
 #include "Streaming/AsyncRegisterLevelContext.h"
 
-#if UE_WITH_IRIS
 #include "Iris/ReplicationSystem/ReplicationSystem.h"
 #include "Net/Iris/ReplicationSystem/EngineReplicationBridge.h"
-#endif // UE_WITH_IRIS
 
 #include "ChaosSolversModule.h"
 #include "HAL/LowLevelMemStats.h"
@@ -8728,11 +8726,7 @@ UWorld* FSeamlessTravelHandler::Tick()
 				if (AGameModeBase* GameMode = LoadedWorld->GetAuthGameMode())
 				{
 					const EReplicationSystem LoadedGameModeRepSystem = GameMode->GetGameNetDriverReplicationSystem();
-#if UE_WITH_IRIS
 					const EReplicationSystem IrisCmdlineRepSystem = UE::Net::GetUseIrisReplicationCmdlineValue();
-#else
-					const EReplicationSystem IrisCmdlineRepSystem = EReplicationSystem::Default;
-#endif
 
 					const bool bIsNetDriverCompatible = LoadedGameModeRepSystem == EReplicationSystem::Default || 
 														IrisCmdlineRepSystem != EReplicationSystem::Default ||
diff --git a/Engine/Source/Runtime/Engine/Public/Net/Iris/ReplicationSystem/ReplicationSystemUtil.h b/Engine/Source/Runtime/Engine/Public/Net/Iris/ReplicationSystem/ReplicationSystemUtil.h
--- a/Engine/Source/Runtime/Engine/Public/Net/Iris/ReplicationSystem/ReplicationSystemUtil.h
+++ b/Engine/Source/Runtime/Engine/Public/Net/Iris/ReplicationSystem/ReplicationSystemUtil.h
@@ -4,7 +4,6 @@
 
 // HEADER_UNIT_SKIP - Can't find includes if UE_WITH_IRIS is not set
 
-#if UE_WITH_IRIS
 
 #include "Iris/ReplicationSystem/Conditionals/ReplicationCondition.h"
 #include "Net/Core/NetHandle/NetHandle.h"
@@ -151,4 +150,3 @@ struct FReplicationSystemUtil
 
 }
 
-#endif
diff --git a/Engine/Source/Runtime/Engine/Public/Net/NetIDVariant.h b/Engine/Source/Runtime/Engine/Public/Net/NetIDVariant.h
--- a/Engine/Source/Runtime/Engine/Public/Net/NetIDVariant.h
+++ b/Engine/Source/Runtime/Engine/Public/Net/NetIDVariant.h
@@ -2,9 +2,7 @@
 
 #pragma once
 
-#if UE_WITH_IRIS
 #include "Iris/ReplicationSystem/NetRefHandle.h"
-#endif
 #include "Misc/NetworkGuid.h"
 #include "Misc/TVariant.h"
 #include "Templates/TypeHash.h"
@@ -24,19 +22,13 @@ public:
 	/** Placeholder to represent empty/unknown type. Always considered invalid. */
 	struct FEmptyID {};
 
-#if UE_WITH_IRIS
 	using FVariantType = TVariant<FEmptyID, FNetworkGUID, FNetRefHandle>;
-#else
-	using FVariantType = TVariant<FEmptyID, FNetworkGUID>;
-#endif
 
 	/** Constructs a variant holding an FEmptyID */
 	FNetIDVariant() = default;
 
 	explicit ENGINE_API FNetIDVariant(FNetworkGUID NetGUID);
-#if UE_WITH_IRIS
 	explicit ENGINE_API FNetIDVariant(FNetRefHandle NetRefHandle);
-#endif
 
 	/** Serializes or deserializes the stored ID or handle, suitable for networking */
 	friend ENGINE_API FArchive& operator<<(FArchive& Ar, FNetIDVariant& NetID);
@@ -70,12 +62,10 @@ inline uint32 GetTypeHash(const FNetIDVariant& NetID)
 	{
 		IDHash = GetTypeHash(NetID.GetVariant().Get<FNetworkGUID>());
 	}
-#if UE_WITH_IRIS
 	else if (NetID.GetVariant().IsType<FNetRefHandle>())
 	{
 		IDHash = GetTypeHash(NetID.GetVariant().Get<FNetRefHandle>());
 	}
-#endif
 
 	return HashCombineFast(TypeIndex, IDHash);
 }
diff --git a/Engine/Source/Runtime/Engine/Public/Net/UnrealNetwork.h b/Engine/Source/Runtime/Engine/Public/Net/UnrealNetwork.h
--- a/Engine/Source/Runtime/Engine/Public/Net/UnrealNetwork.h
+++ b/Engine/Source/Runtime/Engine/Public/Net/UnrealNetwork.h
@@ -145,10 +145,8 @@ struct FDoRepLifetimeParams
 	/** Whether or not this property uses Push Model. See PushModel.h */
 	bool bIsPushBased = false;
 
-#if UE_WITH_IRIS
 	/** Function to create and register a ReplicationFragment for the property */
 	UE::Net::CreateAndRegisterReplicationFragmentFunc CreateAndRegisterReplicationFragmentFunction = nullptr;
-#endif
 
 };
 
@@ -378,7 +376,6 @@ struct CGetFastArrayCreateReplicationFragmentFuncable
 	auto Requires(T* FastArray) -> decltype(FastArray->GetFastArrayCreateReplicationFragmentFunction());
 };
 
-#if UE_WITH_IRIS
 template<typename T>
 inline typename TEnableIf<TModels_V<CGetFastArrayCreateReplicationFragmentFuncable, T>, const FDoRepLifetimeParams>::Type FixupParams(const FDoRepLifetimeParams& Params)
 {
@@ -392,7 +389,6 @@ inline typename TEnableIf<TModels_V<CGetFastArrayCreateReplicationFragmentFuncab
 	NewParams.CreateAndRegisterReplicationFragmentFunction = T::GetFastArrayCreateReplicationFragmentFunction();
 	return NewParams;
 }
-#endif
 
 template<typename T>
 inline typename TEnableIf<!TModels_V<CGetFastArrayCreateReplicationFragmentFuncable, T>, const FDoRepLifetimeParams&>::Type FixupParams(const FDoRepLifetimeParams& Params)
diff --git a/Engine/Source/Runtime/Engine/Public/RepMovementNetSerializer.h b/Engine/Source/Runtime/Engine/Public/RepMovementNetSerializer.h
--- a/Engine/Source/Runtime/Engine/Public/RepMovementNetSerializer.h
+++ b/Engine/Source/Runtime/Engine/Public/RepMovementNetSerializer.h
@@ -11,7 +11,6 @@ struct FRepMovementNetSerializerConfig : public FNetSerializerConfig
 	GENERATED_BODY()
 };
 
-#if UE_WITH_IRIS
 
 namespace UE::Net
 {
@@ -20,4 +19,3 @@ UE_NET_DECLARE_SERIALIZER(FRepMovementNetSerializer, ENGINE_API);
 
 }
 
-#endif // UE_WITH_IRIS
diff --git a/Engine/Source/Runtime/Experimental/GeometryCollectionEngine/Private/GeometryCollection/GeometryCollectionComponent.cpp b/Engine/Source/Runtime/Experimental/GeometryCollectionEngine/Private/GeometryCollection/GeometryCollectionComponent.cpp
--- a/Engine/Source/Runtime/Experimental/GeometryCollectionEngine/Private/GeometryCollection/GeometryCollectionComponent.cpp
+++ b/Engine/Source/Runtime/Experimental/GeometryCollectionEngine/Private/GeometryCollection/GeometryCollectionComponent.cpp
@@ -82,9 +82,7 @@
 #include "Misc/UObjectToken.h"
 #endif //!(UE_BUILD_SHIPPING || UE_BUILD_TEST)
 
-#if UE_WITH_IRIS
 #include "Iris/ReplicationState/PropertyNetSerializerInfoRegistry.h"
-#endif
 
 #if INTEL_ISPC
 
@@ -253,9 +251,7 @@ bool FGeometryCollectionRepData::HasChanged(const FGeometryCollectionRepData& Ba
 	return false;
 }
 
-#if UE_WITH_IRIS
 UE_NET_IMPLEMENT_NAMED_STRUCT_LASTRESORT_NETSERIALIZER_AND_REGISTRY_DELEGATES(GeometryCollectionRepData);
-#endif //  UE_WITH_IRIS
 
 bool FGeometryCollectionRepData::NetSerialize(FArchive& Ar, class UPackageMap* Map, bool& bOutSuccess)
 {
@@ -341,9 +337,7 @@ bool  FGeometryCollectionRepStateData::HasChanged(const FGeometryCollectionRepSt
 	return (Version == 0) || (BrokenState != BaseData.BrokenState) || (bIsRootAnchored != BaseData.bIsRootAnchored);
 }
 
-#if UE_WITH_IRIS
 UE_NET_IMPLEMENT_NAMED_STRUCT_LASTRESORT_NETSERIALIZER_AND_REGISTRY_DELEGATES(GeometryCollectionRepStateData);
-#endif //  UE_WITH_IRIS
 
 bool FGeometryCollectionRepStateData::NetSerialize(FArchive& Ar, class UPackageMap* Map, bool& bOutSuccess)
 {
@@ -440,9 +434,7 @@ bool FGeometryCollectionRepDynamicData::HasChanged(const FGeometryCollectionRepD
 	return false;
 }
 
-#if UE_WITH_IRIS
 UE_NET_IMPLEMENT_NAMED_STRUCT_LASTRESORT_NETSERIALIZER_AND_REGISTRY_DELEGATES(GeometryCollectionRepDynamicData);
-#endif //  UE_WITH_IRIS
 
 bool FGeometryCollectionRepDynamicData::NetSerialize(FArchive& Ar, class UPackageMap* Map, bool& bOutSuccess)
 {
diff --git a/Engine/Source/Runtime/GameplayDebugger/Private/GameplayDebuggerCategoryReplicator.cpp b/Engine/Source/Runtime/GameplayDebugger/Private/GameplayDebuggerCategoryReplicator.cpp
--- a/Engine/Source/Runtime/GameplayDebugger/Private/GameplayDebuggerCategoryReplicator.cpp
+++ b/Engine/Source/Runtime/GameplayDebugger/Private/GameplayDebuggerCategoryReplicator.cpp
@@ -13,10 +13,8 @@
 #include "Engine/NetDriver.h"
 #include "VisualLogger/VisualLogger.h"
 
-#if UE_WITH_IRIS
 #include "Iris/IrisConfig.h"
 #include "Net/Iris/ReplicationSystem/ReplicationSystemUtil.h"
-#endif
 
 #include UE_INLINE_GENERATED_CPP_BY_NAME(GameplayDebuggerCategoryReplicator)
 
@@ -106,25 +104,20 @@ public:
 
 void AGameplayDebuggerCategoryReplicator::PreReplication(IRepChangedPropertyTracker& ChangedPropertyTracker)
 {
-#if UE_WITH_IRIS
 	if (UE::Net::FReplicationSystemUtil::GetNetHandle(this).IsValid())
 	{
 		ReplicatedData.PopulateFromOwner();
 	}
-#endif
 }
 
 void AGameplayDebuggerCategoryReplicator::OnRep_ReplicatedData()
 {
-#if UE_WITH_IRIS
 	if (UE::Net::FReplicationSystemUtil::GetReplicationSystem(this))
 	{
 		ReplicatedData.ApplyToOwner();
 	}
-#endif
 }
 
-#if UE_WITH_IRIS
 void FGameplayDebuggerNetPack::PopulateFromOwner()
 {
 	// The current implementation of the FGameplayDebuggerNetPack replication does not behave like most other 
@@ -222,7 +215,6 @@ void FGameplayDebuggerNetPack::ApplyToOwner()
 		Owner->MarkComponentsRenderStateDirty();
 	}
 }
-#endif // #if UE_WITH_IRIS
 
 bool FGameplayDebuggerNetPack::NetDeltaSerialize(FNetDeltaSerializeInfo& DeltaParms)
 {
@@ -574,7 +566,6 @@ void AGameplayDebuggerCategoryReplicator::BeginPlay()
 	Init();
 }
 
-#if UE_WITH_IRIS
 void AGameplayDebuggerCategoryReplicator::OnReplicationStartedForIris(const FOnReplicationStartedParams& Params)
 {
 	if (UE::Net::FReplicationSystemUtil::GetReplicationSystem(this))
@@ -592,7 +583,6 @@ const AActor* AGameplayDebuggerCategoryReplicator::GetNetOwner() const
 {
 	return OwnerPC.Get();
 }
-#endif
 
 #if WITH_EDITOR
 void AGameplayDebuggerCategoryReplicator::InitForEditor()
diff --git a/Engine/Source/Runtime/GameplayDebugger/Private/GameplayDebuggerModule.cpp b/Engine/Source/Runtime/GameplayDebugger/Private/GameplayDebuggerModule.cpp
--- a/Engine/Source/Runtime/GameplayDebugger/Private/GameplayDebuggerModule.cpp
+++ b/Engine/Source/Runtime/GameplayDebugger/Private/GameplayDebuggerModule.cpp
@@ -14,9 +14,7 @@
 #include "GameplayDebuggerExtension_Spectator.h"
 #include "GameplayDebuggerExtension_HUD.h"
 
-#if UE_WITH_IRIS
 #include "Iris/IrisConfig.h"
-#endif // UE_WITH_IRIS
 
 
 static bool bAutoCreateGameplayDebuggerManager = true;
diff --git a/Engine/Source/Runtime/GameplayDebugger/Public/GameplayDebuggerCategoryReplicator.h b/Engine/Source/Runtime/GameplayDebugger/Public/GameplayDebuggerCategoryReplicator.h
--- a/Engine/Source/Runtime/GameplayDebugger/Public/GameplayDebuggerCategoryReplicator.h
+++ b/Engine/Source/Runtime/GameplayDebugger/Public/GameplayDebuggerCategoryReplicator.h
@@ -69,12 +69,10 @@ struct FGameplayDebuggerNetPack
 	bool NetDeltaSerialize(FNetDeltaSerializeInfo & DeltaParms);
 	void OnCategoriesChanged();
 
-#if UE_WITH_IRIS
 	// When using Iris we cannot use the same appoach as the current NetSerializeDelta as it does polling of data and also modifies data outside of the replicated state during serialization which is not allowed when using iris.
 	// Therefore we need explicit methods to populate the state from the owner and to apply it to the owner which can be invoked before and after serialization as needed.
 	void PopulateFromOwner();
 	void ApplyToOwner();
-#endif
 
 private:
 	UPROPERTY()
@@ -127,11 +125,9 @@ class AGameplayDebuggerCategoryReplicator : public AActor
 protected:
 	GAMEPLAYDEBUGGER_API virtual void BeginPlay() override;
 
-#if UE_WITH_IRIS
 	GAMEPLAYDEBUGGER_API virtual void OnReplicationStartedForIris(const FOnReplicationStartedParams&) override;
 
 	GAMEPLAYDEBUGGER_API virtual const AActor* GetNetOwner() const override;
-#endif
 
 public:
 	GAMEPLAYDEBUGGER_API virtual void EndPlay(const EEndPlayReason::Type EndPlayReason) override;
diff --git a/Engine/Source/Runtime/GameplayTags/Private/GameplayTagContainer.cpp b/Engine/Source/Runtime/GameplayTags/Private/GameplayTagContainer.cpp
--- a/Engine/Source/Runtime/GameplayTags/Private/GameplayTagContainer.cpp
+++ b/Engine/Source/Runtime/GameplayTags/Private/GameplayTagContainer.cpp
@@ -14,10 +14,8 @@
 #include "JsonObjectConverter.h"
 
 // For dynamic serialization support
-#if UE_WITH_IRIS
 #include "Net/Core/NetToken/NetTokenExportContext.h"
 #include "GameplayTagTokenStore.h"
-#endif
 #include <Net/Core/Trace/NetTrace.h>
 
 #include UE_INLINE_GENERATED_CPP_BY_NAME(GameplayTagContainer)
@@ -1389,7 +1387,6 @@ bool FGameplayTag::NetSerialize_ForReplayUsingFastReplication(FArchive& Ar, UPac
 // DynamicSerialization currently relies on experimental code only available when compiling with Iris.
 namespace UE::GameplayTags::GameplayTagDynamicSerialization
 {
-#if UE_WITH_IRIS
 	// Can we make this generic and handle arbitrary export payloads? Probably something that should be handled in PackageMapClient
 	static bool NetSerialize_ForReplay(FGameplayTag& GameplayTag, FArchive& Ar, UPackageMapClient& PackageMapClient)
 	{
@@ -1557,13 +1554,6 @@ namespace UE::GameplayTags::GameplayTagDynamicSerialization
 		}
 		return false;
 	}
-#else
-	bool NetSerialize(FGameplayTag& GameplayTag, FArchive& Ar, UPackageMap* Map)
-	{
-		LowLevelFatalError(TEXT("Cannot use dynamic serialization without compiling with Iris."));
-		return false;
-	}
-#endif
 }
 
 
diff --git a/Engine/Source/Runtime/GameplayTags/Private/GameplayTagContainerNetSerializer.cpp b/Engine/Source/Runtime/GameplayTags/Private/GameplayTagContainerNetSerializer.cpp
--- a/Engine/Source/Runtime/GameplayTags/Private/GameplayTagContainerNetSerializer.cpp
+++ b/Engine/Source/Runtime/GameplayTags/Private/GameplayTagContainerNetSerializer.cpp
@@ -4,7 +4,6 @@
 
 #include UE_INLINE_GENERATED_CPP_BY_NAME(GameplayTagContainerNetSerializer)
 
-#if UE_WITH_IRIS
 
 #include "InternalGameplayTagContainerNetSerializer.h"
 #include "GameplayTagContainer.h"
@@ -370,4 +369,3 @@ void FGameplayTagContainerNetSerializer::FNetSerializerRegistryDelegates::OnPost
 
 }
 
-#endif // UE_WITH_IRIS
diff --git a/Engine/Source/Runtime/GameplayTags/Private/GameplayTagNetSerializer.cpp b/Engine/Source/Runtime/GameplayTags/Private/GameplayTagNetSerializer.cpp
--- a/Engine/Source/Runtime/GameplayTags/Private/GameplayTagNetSerializer.cpp
+++ b/Engine/Source/Runtime/GameplayTags/Private/GameplayTagNetSerializer.cpp
@@ -4,7 +4,6 @@
 
 #include UE_INLINE_GENERATED_CPP_BY_NAME(GameplayTagNetSerializer)
 
-#if UE_WITH_IRIS
 
 #include "GameplayTagContainer.h"
 #include "GameplayTagsManager.h"
@@ -279,4 +278,3 @@ void FGameplayTagNetSerializer::FNetSerializerRegistryDelegates::OnPreFreezeNetS
 
 }
 
-#endif // UE_WITH_IRIS
diff --git a/Engine/Source/Runtime/GameplayTags/Private/GameplayTagTokenStore.cpp b/Engine/Source/Runtime/GameplayTags/Private/GameplayTagTokenStore.cpp
--- a/Engine/Source/Runtime/GameplayTags/Private/GameplayTagTokenStore.cpp
+++ b/Engine/Source/Runtime/GameplayTags/Private/GameplayTagTokenStore.cpp
@@ -1,8 +1,5 @@
 // Copyright Epic Games, Inc. All Rights Reserved.
 
-// NetTokenExports is still an experimental feature that depends on Iris code.
-#if UE_WITH_IRIS
-
 #include "GameplayTagTokenStore.h"
 #include "GameplayTagContainer.h"
 #include "GameplayTagsManager.h"
@@ -26,5 +23,3 @@ FGameplayTag FGameplayTagTokenStore::ResolveToken(FNetToken Token, const FNetTok
 }
 
 }
-
-#endif // UE_WITH_IRIS
\ No newline at end of file
diff --git a/Engine/Source/Runtime/GameplayTags/Public/GameplayTagTokenStore.h b/Engine/Source/Runtime/GameplayTags/Public/GameplayTagTokenStore.h
--- a/Engine/Source/Runtime/GameplayTags/Public/GameplayTagTokenStore.h
+++ b/Engine/Source/Runtime/GameplayTags/Public/GameplayTagTokenStore.h
@@ -3,8 +3,6 @@
 #pragma once
 #include "CoreMinimal.h"
 
-// NetTokenExports is still an experimental feature that depends on Iris code.
-#if UE_WITH_IRIS
 
 #include "Iris/ReplicationSystem/NameTokenStore.h"
 
@@ -34,5 +32,4 @@ private:
 	inline static FName GameplayTokenStoreName = TEXT("GameplayTagTokenStore");
 };
 
-}
-#endif // UE_WITH_IRIS
\ No newline at end of file
+}
\ No newline at end of file
diff --git a/Engine/Source/Runtime/GameplayTasks/Classes/GameplayTask.h b/Engine/Source/Runtime/GameplayTasks/Classes/GameplayTask.h
--- a/Engine/Source/Runtime/GameplayTasks/Classes/GameplayTask.h
+++ b/Engine/Source/Runtime/GameplayTasks/Classes/GameplayTask.h
@@ -310,10 +310,8 @@ protected:
 	GAMEPLAYTASKS_API virtual void OnGameplayTaskDeactivated(UGameplayTask& Task) override;
 	// IGameplayTaskOwnerInterface END
 
-#if UE_WITH_IRIS
 	/** Register all replication fragments */
 	GAMEPLAYTASKS_API virtual void RegisterReplicationFragments(UE::Net::FFragmentRegistrationContext& Context, UE::Net::EFragmentRegistrationFlags RegistrationFlags) override;
-#endif // UE_WITH_IRIS
 
 private:
 	friend UGameplayTasksComponent;
diff --git a/Engine/Source/Runtime/GameplayTasks/Private/GameplayTask.cpp b/Engine/Source/Runtime/GameplayTasks/Private/GameplayTask.cpp
--- a/Engine/Source/Runtime/GameplayTasks/Private/GameplayTask.cpp
+++ b/Engine/Source/Runtime/GameplayTasks/Private/GameplayTask.cpp
@@ -7,9 +7,7 @@
 #include "GameplayTaskResource.h"
 #include "GameplayTasksComponent.h"
 
-#if UE_WITH_IRIS
 #include "Iris/ReplicationSystem/ReplicationFragmentUtil.h"
-#endif // UE_WITH_IRIS
 
 #include UE_INLINE_GENERATED_CPP_BY_NAME(GameplayTask)
 
@@ -473,10 +471,8 @@ void UGameplayTask::OnGameplayTaskInitialized(UGameplayTask& Task)
 	ChildTask = &Task;
 }
 
-#if UE_WITH_IRIS
 void UGameplayTask::RegisterReplicationFragments(UE::Net::FFragmentRegistrationContext& Context, UE::Net::EFragmentRegistrationFlags RegistrationFlags)
 {
 	// Build descriptors and allocate PropertyReplicationFragments for this object.
 	UE::Net::FReplicationFragmentUtil::CreateAndRegisterFragmentsForObject(this, Context, RegistrationFlags);
 }
-#endif // UE_WITH_IRIS
diff --git a/Engine/Source/Runtime/MovieScene/Private/MovieSceneSequencePlayer.cpp b/Engine/Source/Runtime/MovieScene/Private/MovieSceneSequencePlayer.cpp
--- a/Engine/Source/Runtime/MovieScene/Private/MovieSceneSequencePlayer.cpp
+++ b/Engine/Source/Runtime/MovieScene/Private/MovieSceneSequencePlayer.cpp
@@ -26,10 +26,8 @@
 #include "Algo/BinarySearch.h"
 #include "UniversalObjectLocatorResolveParameterBuffer.inl"
 
-#if UE_WITH_IRIS
 #include "Iris/ReplicationSystem/ReplicationFragmentUtil.h"
 #include "Net/Iris/ReplicationSystem/ReplicationSystemUtil.h"
-#endif // UE_WITH_IRIS
 #include "MovieSceneObjectBindingID.h"
 
 
@@ -2227,12 +2225,10 @@ bool UMovieSceneSequencePlayer::CallRemoteFunction(UFunction* Function, void* Pa
 	return false;
 }
 
-#if UE_WITH_IRIS
 void UMovieSceneSequencePlayer::RegisterReplicationFragments(UE::Net::FFragmentRegistrationContext& Context, UE::Net::EFragmentRegistrationFlags RegistrationFlags)
 {
 	UE::Net::FReplicationFragmentUtil::CreateAndRegisterFragmentsForObject(this, Context, RegistrationFlags);
 }
-#endif
 
 bool UMovieSceneSequencePlayer::NeedsQueueLatentAction() const
 {
diff --git a/Engine/Source/Runtime/MovieScene/Public/MovieSceneSequencePlayer.h b/Engine/Source/Runtime/MovieScene/Public/MovieSceneSequencePlayer.h
--- a/Engine/Source/Runtime/MovieScene/Public/MovieSceneSequencePlayer.h
+++ b/Engine/Source/Runtime/MovieScene/Public/MovieSceneSequencePlayer.h
@@ -588,9 +588,7 @@ protected:
 	MOVIESCENE_API virtual bool CallRemoteFunction(UFunction* Function, void* Parameters, FOutParmRec* OutParms, FFrame* Stack) override;
 	MOVIESCENE_API virtual void PostNetReceive() override;
 	MOVIESCENE_API virtual void BeginDestroy() override;
-#if UE_WITH_IRIS
 	MOVIESCENE_API virtual void RegisterReplicationFragments(UE::Net::FFragmentRegistrationContext& Context, UE::Net::EFragmentRegistrationFlags RegistrationFlags) override;
-#endif
 	/*~ End UObject interface */
 
 	//~ Begin IMovieSceneSequenceTickManagerClient interface
diff --git a/Engine/Source/Runtime/Net/Core/Private/Net/Core/PropertyConditions/PropertyConditions.cpp b/Engine/Source/Runtime/Net/Core/Private/Net/Core/PropertyConditions/PropertyConditions.cpp
--- a/Engine/Source/Runtime/Net/Core/Private/Net/Core/PropertyConditions/PropertyConditions.cpp
+++ b/Engine/Source/Runtime/Net/Core/Private/Net/Core/PropertyConditions/PropertyConditions.cpp
@@ -12,10 +12,8 @@ DECLARE_CYCLE_STAT(TEXT("PropertyConditions PostGarbageCollect"), STAT_PropertyC
 
 namespace UE::Net::Private
 {
-#if UE_WITH_IRIS
 FPropertyConditionDelegates::FOnPropertyCustomConditionChanged FPropertyConditionDelegates::OnPropertyCustomConditionChangedDelegate;
 FPropertyConditionDelegates::FOnPropertyDynamicConditionChanged FPropertyConditionDelegates::OnPropertyDynamicConditionChangedDelegate;
-#endif // UE_WITH_IRIS
 
 FNetPropertyConditionManager::FNetPropertyConditionManager()
 {
diff --git a/Engine/Source/Runtime/Net/Core/Private/Net/Core/PropertyConditions/RepChangedPropertyTracker.cpp b/Engine/Source/Runtime/Net/Core/Private/Net/Core/PropertyConditions/RepChangedPropertyTracker.cpp
--- a/Engine/Source/Runtime/Net/Core/Private/Net/Core/PropertyConditions/RepChangedPropertyTracker.cpp
+++ b/Engine/Source/Runtime/Net/Core/Private/Net/Core/PropertyConditions/RepChangedPropertyTracker.cpp
@@ -21,12 +21,10 @@ void FRepChangedPropertyTracker::SetCustomIsActiveOverride(const UObject* Owning
 	}
 #endif
 
-#if UE_WITH_IRIS
 	if (bOldActive != bIsActive && UE::Net::Private::FPropertyConditionDelegates::GetOnPropertyCustomConditionChangedDelegate().IsBound())
 	{
 		UE::Net::Private::FPropertyConditionDelegates::GetOnPropertyCustomConditionChangedDelegate().Broadcast(OwningObject, RepIndex, bIsActive);
 	}
-#endif // UE_WITH_IRIS
 }
 
 void FRepChangedPropertyTracker::SetDynamicCondition(const UObject* OwningObject, const uint16 RepIndex, const ELifetimeCondition Condition)
@@ -45,13 +43,11 @@ void FRepChangedPropertyTracker::SetDynamicCondition(const UObject* OwningObject
 	MARK_PROPERTY_DIRTY_UNSAFE(OwningObject, RepIndex);
 #endif
 
-#if UE_WITH_IRIS
 	FPropertyConditionDelegates::FOnPropertyDynamicConditionChanged& Delegate = FPropertyConditionDelegates::GetOnPropertyDynamicConditionChangedDelegate();
 	if (Delegate.IsBound())
 	{
 		Delegate.Broadcast(OwningObject, RepIndex, Condition);
 	}
-#endif
 }
 
 void FRepChangedPropertyTracker::CountBytes(FArchive& Ar) const
diff --git a/Engine/Source/Runtime/Net/Core/Private/Net/Core/PushModel/PushModel.cpp b/Engine/Source/Runtime/Net/Core/Private/Net/Core/PushModel/PushModel.cpp
--- a/Engine/Source/Runtime/Net/Core/Private/Net/Core/PushModel/PushModel.cpp
+++ b/Engine/Source/Runtime/Net/Core/Private/Net/Core/PushModel/PushModel.cpp
@@ -451,7 +451,6 @@ namespace UEPushModelPrivate
 		FConsoleCommandDelegate::CreateLambda([]() { PushObjectManager.PrintAssignedHandles(); })
 	);
 
-#if UE_WITH_IRIS
 	static FIrisMarkPropertyDirty IrisMarkPropertyDirtyDelegate;
 	static FIrisMarkPropertiesDirty IrisMarkPropertiesDirtyDelegate;
 
@@ -464,18 +463,15 @@ namespace UEPushModelPrivate
 	{
 		IrisMarkPropertiesDirtyDelegate = Delegate;
 	}
-#endif // UE_WITH_IRIS
 
 	void MarkPropertyDirty(const UObject* Object, const FNetPushObjectId ObjectId, const int32 RepIndex)
 	{
-#if UE_WITH_IRIS
 		if (ObjectId.IsIrisId())
 		{
 			// Assume the delegate is bound. If it isn't we have an invalid ID, which is a bug that needs to be tracked down.
 			IrisMarkPropertyDirtyDelegate.Execute(Object, ObjectId.GetIrisPushObjectId(), RepIndex);
 		}
 		else
-#endif // UE_WITH_IRIS
 		{
 			PushObjectManager.MarkPropertyDirty(ObjectId.GetLegacyPushObjectId(), RepIndex);
 		}
@@ -483,13 +479,11 @@ namespace UEPushModelPrivate
 
 	void MarkPropertyDirty(const UObject* Object, const FNetPushObjectId ObjectId, const int32 StartRepIndex, const int32 EndRepIndex)
 	{
-#if UE_WITH_IRIS
 		if (ObjectId.IsIrisId())
 		{
 			IrisMarkPropertiesDirtyDelegate.Execute(Object, ObjectId.GetIrisPushObjectId(), StartRepIndex, EndRepIndex);
 		}
 		else
-#endif // UE_WITH_IRIS
 		{
 			PushObjectManager.MarkPropertyDirty(ObjectId.GetLegacyPushObjectId(), StartRepIndex, EndRepIndex);
 		}
diff --git a/Engine/Source/Runtime/Net/Core/Public/Net/Core/PropertyConditions/PropertyConditionsDelegates.h b/Engine/Source/Runtime/Net/Core/Public/Net/Core/PropertyConditions/PropertyConditionsDelegates.h
--- a/Engine/Source/Runtime/Net/Core/Public/Net/Core/PropertyConditions/PropertyConditionsDelegates.h
+++ b/Engine/Source/Runtime/Net/Core/Public/Net/Core/PropertyConditions/PropertyConditionsDelegates.h
@@ -6,7 +6,6 @@
 // We are considering to move Iris into NetCore when we leave the experimental stage at that point we can call directly into iris code
 // rather than going though a delegate.
 
-#if UE_WITH_IRIS
 
 #include "Delegates/Delegate.h"
 
@@ -31,4 +30,3 @@ private:
 
 } // UE::Net::Private
 
-#endif
diff --git a/Engine/Source/Runtime/Net/Core/Public/Net/Core/PushModel/PushModel.h b/Engine/Source/Runtime/Net/Core/Public/Net/Core/PushModel/PushModel.h
--- a/Engine/Source/Runtime/Net/Core/Public/Net/Core/PushModel/PushModel.h
+++ b/Engine/Source/Runtime/Net/Core/Public/Net/Core/PushModel/PushModel.h
@@ -410,14 +410,12 @@ namespace UEPushModelPrivate
 
 	NETCORE_API void LogMemory(FOutputDevice& Ar);
 
-#if UE_WITH_IRIS
 	// For internal use only.
 	DECLARE_DELEGATE_ThreeParams(FIrisMarkPropertyDirty, const UObject*, UEPushModelPrivate::FNetIrisPushObjectId, const int32);
 	DECLARE_DELEGATE_FourParams(FIrisMarkPropertiesDirty, const UObject*, UEPushModelPrivate::FNetIrisPushObjectId, const int32, const int32);
 
 	NETCORE_API void SetIrisMarkPropertyDirtyDelegate(const FIrisMarkPropertyDirty& Delegate);
 	NETCORE_API void SetIrisMarkPropertiesDirtyDelegate(const FIrisMarkPropertiesDirty& Delegate);
-#endif // UE_WITH_IRIS
 }
 
 
@@ -468,10 +466,8 @@ namespace UEPushModelPrivate
 #define COMPARE_ASSIGN_AND_MARK_PROPERTY_DIRTY(ClassName, PropertyName, NewValue, Object) \
 	if (NewValue != PropertyName) { PropertyName = NewValue; MARK_PROPERTY_DIRTY_FROM_NAME(ClassName, PropertyName, Object); }
 
-#if UE_WITH_IRIS
 #define UE_NET_SET_IRIS_MARK_PROPERTY_DIRTY_DELEGATE(Delegate) UEPushModelPrivate::SetIrisMarkPropertyDirtyDelegate(Delegate)
 #define UE_NET_SET_IRIS_MARK_PROPERTIES_DIRTY_DELEGATE(Delegate) UEPushModelPrivate::SetIrisMarkPropertiesDirtyDelegate(Delegate)
-#endif // UE_WITH_IRIS
 
 #else // WITH_PUSH_MODEL
 
@@ -494,9 +490,7 @@ namespace UEPushModelPrivate
 #define COMPARE_ASSIGN_AND_MARK_PROPERTY_DIRTY(ClassName, PropertyName, NewValue, Object) \
 	if (NewValue != PropertyName) { PropertyName = NewValue; }
 
-#if UE_WITH_IRIS
 #define UE_NET_SET_IRIS_MARK_PROPERTY_DIRTY_DELEGATE(...)
 #define UE_NET_SET_IRIS_MARK_PROPERTIES_DIRTY_DELEGATE(...)
-#endif // UE_WITH_IRIS
 
 #endif
diff --git a/Engine/Source/Runtime/Net/Iris/Private/Iris/Serialization/InstancedStructNetSerializer.cpp b/Engine/Source/Runtime/Net/Iris/Private/Iris/Serialization/InstancedStructNetSerializer.cpp
--- a/Engine/Source/Runtime/Net/Iris/Private/Iris/Serialization/InstancedStructNetSerializer.cpp
+++ b/Engine/Source/Runtime/Net/Iris/Private/Iris/Serialization/InstancedStructNetSerializer.cpp
@@ -3,7 +3,6 @@
 #include "Iris/Serialization/InstancedStructNetSerializer.h"
 #include UE_INLINE_GENERATED_CPP_BY_NAME(InstancedStructNetSerializer)
 
-#if UE_WITH_IRIS
 
 #include "StructUtils/InstancedStruct.h"
 #include "HAL/IConsoleManager.h"
@@ -752,5 +751,4 @@ TRefCountPtr<const FReplicationStateDescriptor> FInstancedStructDescriptorCache:
 
 }
 
-#endif
 
diff --git a/Engine/Source/Runtime/Net/Iris/Public/Iris/ReplicationSystem/NetTokenStore.h b/Engine/Source/Runtime/Net/Iris/Public/Iris/ReplicationSystem/NetTokenStore.h
--- a/Engine/Source/Runtime/Net/Iris/Public/Iris/ReplicationSystem/NetTokenStore.h
+++ b/Engine/Source/Runtime/Net/Iris/Public/Iris/ReplicationSystem/NetTokenStore.h
@@ -29,8 +29,6 @@
 *
 * A current example that us used by both systems is GameplayTags, For Iris: See GameplayTagNetSerializer.cpp, for the old replication system: See: GameplayTagContainer.cpp.
 *
-* Note: this is still an experimental feature and is currently only available if compiling UE_WITH_IRIS. 
-* Implementation details is subject to change.
 */
 
 /* NetBitArray validation support. */
