Add additional safety on subsystem removal to make sure we don't access stale UObjects

#rnx

diff --git a/Engine/Source/Runtime/Engine/Private/Subsystems/SubsystemCollection.cpp b/Engine/Source/Runtime/Engine/Private/Subsystems/SubsystemCollection.cpp
--- a/Engine/Source/Runtime/Engine/Private/Subsystems/SubsystemCollection.cpp
+++ b/Engine/Source/Runtime/Engine/Private/Subsystems/SubsystemCollection.cpp
@@ -134,6 +134,37 @@ void FSubsystemCollectionBase::ForEachSubsystemOfClass(UClass* SubsystemClass, T
 	}
 }
 
+void FSubsystemCollectionBase::RemoveSubsystemsInPackages(TConstArrayView<UPackage*> Packages)
+{
+	// safety first
+	for (TPair<UClass*, TUniquePtr<FSubsystemArray>>& Pair : SubsystemArrayMap)
+	{
+		UE_CLOG(Pair.Value->bIsIterating, LogSubsystemCollection, Fatal,
+			TEXT("FSubsystemCollectionBase::RemoveSubsystemsInPackages called while iterating subsystems of type %s"), *Pair.Key->GetPathName());
+	}
+
+	// release any subsystem class & instances that matches the packages being unloaded
+	for (auto Iter = SubsystemArrayMap.CreateIterator(); Iter; ++Iter)
+	{
+		const UClass* KeyClass = Iter.Key();
+
+		if (Packages.Contains(KeyClass->GetPackage()))
+		{
+			UE_LOG(LogSubsystemCollection, Verbose, TEXT("FSubsystemCollectionBase::RemoveSubsystemsInPackages: removing %s"), *KeyClass->GetFullName());
+			
+			for (USubsystem* Subsystem : Iter.Value()->Subsystems)
+			{
+				if (IsValid(Subsystem) && Subsystem->GetClass() == KeyClass)
+				{
+					Subsystem->Deinitialize();
+					Subsystem->InternalOwningSubsystem = nullptr;
+				}
+			}
+			Iter.RemoveCurrent();
+		}
+	}
+}
+
 void FSubsystemCollectionBase::Initialize(UObject* NewOuter)
 {
 	if (Outer != nullptr)
@@ -184,6 +215,8 @@ void FSubsystemCollectionBase::Initialize(UObject* NewOuter)
 			}
 		}
 
+		ModulesUnloadedHandle = FCoreUObjectDelegates::CompiledInUObjectsRemovedDelegate.AddRaw(this, &FSubsystemCollectionBase::RemoveSubsystemsInPackages);
+
 		// Statically track collections
 		GlobalSubsystemCollections.Add(this);
 	}
@@ -214,6 +247,12 @@ void FSubsystemCollectionBase::Deinitialize()
 	{
 		return;
 	}
+
+	if (ModulesUnloadedHandle.IsValid())
+	{
+		FCoreUObjectDelegates::CompiledInUObjectsRemovedDelegate.Remove(ModulesUnloadedHandle);
+		ModulesUnloadedHandle.Reset();
+	}
 	
 	// Remove static tracking 
 	GlobalSubsystemCollections.Remove(this);
diff --git a/Engine/Source/Runtime/Engine/Public/Subsystems/SubsystemCollection.h b/Engine/Source/Runtime/Engine/Public/Subsystems/SubsystemCollection.h
--- a/Engine/Source/Runtime/Engine/Public/Subsystems/SubsystemCollection.h
+++ b/Engine/Source/Runtime/Engine/Public/Subsystems/SubsystemCollection.h
@@ -89,6 +89,9 @@ protected:
 	/** Perform an operation on all subsystems that derive from the given class */
 	ENGINE_API void ForEachSubsystemOfClass(UClass* SubsystemClass, TFunctionRef<void(USubsystem*)> Operation) const;
 
+	/** Remove all subsystems in this list of packages */
+	ENGINE_API void RemoveSubsystemsInPackages(TConstArrayView<UPackage*> Packages);
+
 private:
 	ENGINE_API USubsystem* AddAndInitializeSubsystem(UClass* SubsystemClass);
 
@@ -105,6 +108,8 @@ private:
 	bool bPopulating;
 	mutable bool bIterating = false; // True if iterating over SubsystemMap
 
+	FDelegateHandle ModulesUnloadedHandle;
+
 private:
 	friend class FSubsystemModuleWatcher;
 
