Control Rig: Fix crash when resetting preview mesh to default
#rnx

diff --git a/Engine/Source/Runtime/Engine/Private/SkeletalDebugRendering.cpp b/Engine/Source/Runtime/Engine/Private/SkeletalDebugRendering.cpp
--- a/Engine/Source/Runtime/Engine/Private/SkeletalDebugRendering.cpp
+++ b/Engine/Source/Runtime/Engine/Private/SkeletalDebugRendering.cpp
@@ -377,7 +377,8 @@ void DrawBonesInternal(
 	{
 		const int32 BoneIndex = RequiredBones[Index];
 		if ((!BoneColors.IsEmpty() && BoneIndex >= BoneColors.Num()) ||
-			BoneIndex >= WorldTransforms.Num())
+			BoneIndex >= WorldTransforms.Num() ||
+			!BonesToDraw.IsValidIndex(BoneIndex))
 		{
 			continue;
 		}
