Ensure robust error messaging when processing deleted parameter references in material cached data, regardless of missing asset path ref.

diff --git a/Engine/Source/Runtime/Engine/Private/Materials/MaterialCachedData.cpp b/Engine/Source/Runtime/Engine/Private/Materials/MaterialCachedData.cpp
--- a/Engine/Source/Runtime/Engine/Private/Materials/MaterialCachedData.cpp
+++ b/Engine/Source/Runtime/Engine/Private/Materials/MaterialCachedData.cpp
@@ -290,8 +290,16 @@ bool FMaterialCachedExpressionData::AddParameter(const FMaterialParameterInfo& P
 	{
 		auto GetEditorInfoAsText = [EditorOnlyData = EditorOnlyData](const FMaterialCachedParameterEditorInfo& InEditorInfo) -> FText
 			{
-				FString AssetPath = EditorOnlyData->AssetPaths[InEditorInfo.AssetIndex];
-				return FText::Format(LOCTEXT("ReportParameterMetaData", "group:'{0}', asset:'{1}'"), FText::FromName(InEditorInfo.Group), FText::FromString(AssetPath));
+				if (EditorOnlyData->AssetPaths.IsValidIndex(InEditorInfo.AssetIndex))
+				{
+					FString AssetPath = EditorOnlyData->AssetPaths[InEditorInfo.AssetIndex];
+					return FText::Format(LOCTEXT("ReportParameterMetaData", "group:'{0}', asset:'{1}'"), FText::FromName(InEditorInfo.Group), FText::FromString(AssetPath));
+				}
+				else
+				{
+					return FText::Format(LOCTEXT("ReportInvalidatedParameterMetaData", "group:'{0}', description:'{1}'"), FText::FromName(InEditorInfo.Group), FText::FromString(InEditorInfo.Description));
+				}
+
 			};
 
 		auto ReportError = [&GetEditorInfoAsText](const FString& InCachedValue, const FString& InValue, const FMaterialCachedParameterEditorInfo& InCachedEditorInfo, const FMaterialCachedParameterEditorInfo& InEditorInfo) -> FText
