Cloth - Fixed rendering crash when pausing cloth physics on a LOD > 0 and switching the render back to a lower LOD.

[RN] Cloth - Fixed rendering crash when pausing cloth physics on a LOD > 0 and switching the render back to a lower LOD.

diff --git a/Engine/Source/Runtime/Engine/Private/SkeletalRenderGPUSkin.cpp b/Engine/Source/Runtime/Engine/Private/SkeletalRenderGPUSkin.cpp
--- a/Engine/Source/Runtime/Engine/Private/SkeletalRenderGPUSkin.cpp
+++ b/Engine/Source/Runtime/Engine/Private/SkeletalRenderGPUSkin.cpp
@@ -1074,6 +1074,15 @@ void FSkeletalMeshObjectGPUSkin::ProcessUpdatedDynamicData(FRHICommandList& RHIC
 		{
 			ClothShaderData = &ClothVertexFactory->GetClothShaderData();
 			ClothSimulationData = DynamicData->ClothingSimData.Find(Section.CorrespondClothAssetIndex);
+			if (ClothSimulationData && ClothSimulationData->LODIndex > LODIndex)
+			{
+				// Can only render deform the simulation with positive LODBias which is
+				// when the sim data LOD is lower than the requested render LODIndex.
+				// Otherwise the cloth deformer mapping would be out of bounds.
+				// This can happen when the physics is paused on a LOD >0 and then switching down the LODs.
+				ClothSimulationData = nullptr;
+			}
+
 			ClothShaderData->bEnabled = ClothSimulationData != nullptr;
 
 			if (ClothSimulationData)
