- Remove ColorCurve during cook on serialize
- This fixes "Non optional objects referenced by optional objects needs to be present in cooked data"

#rnx

diff --git a/Engine/Source/Runtime/Engine/Private/Materials/MaterialExpressions.cpp b/Engine/Source/Runtime/Engine/Private/Materials/MaterialExpressions.cpp
--- a/Engine/Source/Runtime/Engine/Private/Materials/MaterialExpressions.cpp
+++ b/Engine/Source/Runtime/Engine/Private/Materials/MaterialExpressions.cpp
@@ -4867,6 +4867,17 @@ UMaterialExpressionColorRamp::UMaterialExpressionColorRamp(const FObjectInitiali
 #endif // WITH_EDITOR
 }
 
+void UMaterialExpressionColorRamp::Serialize(FArchive& Ar)
+{
+	TOptional<TGuardValue<TObjectPtr<UCurveLinearColor>>> ColorCurveGuard;
+	if (Ar.IsSaving() && Ar.IsCooking())
+	{
+		ColorCurveGuard.Emplace(ColorCurve, nullptr);
+	}
+
+	Super::Serialize(Ar);
+}
+
 #if WITH_EDITOR
 void UMaterialExpressionColorRamp::HandleCurvePropertyChanged(UCurveBase* Curve, EPropertyChangeType::Type ChangeType)
 {
diff --git a/Engine/Source/Runtime/Engine/Public/Materials/MaterialExpressionColorRamp.h b/Engine/Source/Runtime/Engine/Public/Materials/MaterialExpressionColorRamp.h
--- a/Engine/Source/Runtime/Engine/Public/Materials/MaterialExpressionColorRamp.h
+++ b/Engine/Source/Runtime/Engine/Public/Materials/MaterialExpressionColorRamp.h
@@ -10,13 +10,12 @@
 #include "Curves/CurveLinearColor.h"
 #include "MaterialExpressionColorRamp.generated.h"
 
-UCLASS(MinimalAPI, meta = (Private))
+UCLASS(MinimalAPI)
 class UMaterialExpressionColorRamp : public UMaterialExpression
 {
 	GENERATED_BODY()
 
 public:
-
 	UPROPERTY(meta = (RequiredInput = "false", ToolTip = "Defaults to 'ConstInput' if not specified"))
 	FExpressionInput Input;
 
@@ -28,6 +27,10 @@ public:
 
 	UMaterialExpressionColorRamp(const FObjectInitializer& ObjectInitializer);
 
+	//~ Begin UObject Interface
+	virtual void Serialize(FArchive& Ar) override;
+	//~ End UObject Interface
+
 	//~ Begin UMaterialExpression Interface
 #if WITH_EDITOR
 	virtual void Build(MIR::FEmitter& Emitter) override;
