Enable support for Volumetric fog on platform supporting GeometryShaders.

Old condition was too restrictive, preventing VolumetricFog to run on compatible platform.

#test EngineTest editor dx11, dx12, vulkan

diff --git a/Engine/Shaders/Public/ShaderVersion.ush b/Engine/Shaders/Public/ShaderVersion.ush
--- a/Engine/Shaders/Public/ShaderVersion.ush
+++ b/Engine/Shaders/Public/ShaderVersion.ush
@@ -3,4 +3,4 @@
 // Each time the console command is executed it generates a new GUID. As this file is included
 // in Platform.ush (which should be included in any shader) it allows to invalidate the shader DDC.
 // 
-#pragma message("UESHADERMETADATA_VERSION 6289045B-BD8C-4169-846F-3D7A21539311")
+#pragma message("UESHADERMETADATA_VERSION 0D093781-606A-4FF5-9E23-6A81699DFDA2")
diff --git a/Engine/Source/Runtime/Engine/Private/ShaderCompiler/ShaderCompiler.cpp b/Engine/Source/Runtime/Engine/Private/ShaderCompiler/ShaderCompiler.cpp
--- a/Engine/Source/Runtime/Engine/Private/ShaderCompiler/ShaderCompiler.cpp
+++ b/Engine/Source/Runtime/Engine/Private/ShaderCompiler/ShaderCompiler.cpp
@@ -3136,11 +3136,12 @@ namespace
 
 static bool IsSubstrateSupportForShaderPipeline(const FShaderCompilerInput& Input)
 {
-	// Substrate requires HLSL2021 which must be cross-compiled for D3D11 to be consumed by FXC compiler.
-	// This cross-compilation toolchain does not support geometry shaders.
+	// If a material requires geometry shaders, ensure the target platform supports geometry shaders. 
+	// For instance Substrate requires HLSL2021 which can be cross-compiled, but cross-compilation 
+	// toolchain might not support geometry shaders.
 	bool bPipelineContainsGeometryShader = false;
 	Input.Environment.GetCompileArgument(TEXT("PIPELINE_CONTAINS_GEOMETRYSHADER"), bPipelineContainsGeometryShader);
-	const bool bCanRHICompileHlsl2021GeometryShaders = !(GetMaxSupportedFeatureLevel((EShaderPlatform)Input.Target.Platform) == ERHIFeatureLevel::SM5);
+	const bool bCanRHICompileHlsl2021GeometryShaders = RHISupportsGeometryShaders((EShaderPlatform)Input.Target.Platform);
 	const bool bIsSubstrateSupportedForPipeline = !bPipelineContainsGeometryShader || bCanRHICompileHlsl2021GeometryShaders;
 	return bIsSubstrateSupportedForPipeline;
 }
diff --git a/Engine/Source/Runtime/Renderer/Private/VolumetricFog.cpp b/Engine/Source/Runtime/Renderer/Private/VolumetricFog.cpp
--- a/Engine/Source/Runtime/Renderer/Private/VolumetricFog.cpp
+++ b/Engine/Source/Runtime/Renderer/Private/VolumetricFog.cpp
@@ -1190,10 +1190,7 @@ IMPLEMENT_GLOBAL_SHADER(FVolumetricFogFinalIntegrationCS, "/Engine/Private/Volum
 
 bool DoesPlatformSupportVolumetricFogVoxelization(const FStaticShaderPlatform Platform)
 {
-	// Substrate requires HLSL2021 which must be cross-compiled for D3D11 and Vulkan to be consumed by FXC compiler.
-	// This cross-compilation toolchain does not support geometry shaders.
-	const bool bCanRHICompileGeometryShaders =!Substrate::IsSubstrateEnabled() || GetMaxSupportedFeatureLevel(Platform) > ERHIFeatureLevel::SM5;
-	return !IsMobileVolumetricFog(Platform) && bCanRHICompileGeometryShaders;
+	return !IsMobileVolumetricFog(Platform) && RHISupportsGeometryShaders(Platform);
 }
 bool ShouldRenderVolumetricFog(const FScene* Scene, const FSceneViewFamily& ViewFamily)
 {
