Fixed crash occurring when the material parent of an instance is change through scripting.

[RN] Changing the parent of a Mateiral Instance through scripting no longer crashes the editor.

diff --git a/Engine/Source/Runtime/Engine/Private/Materials/MaterialInstanceConstant.cpp b/Engine/Source/Runtime/Engine/Private/Materials/MaterialInstanceConstant.cpp
--- a/Engine/Source/Runtime/Engine/Private/Materials/MaterialInstanceConstant.cpp
+++ b/Engine/Source/Runtime/Engine/Private/Materials/MaterialInstanceConstant.cpp
@@ -81,8 +81,10 @@ void UMaterialInstanceConstant::SetParentEditorOnly(UMaterialInterface* NewParen
 
 	if (SetParentInternal(NewParent, RecacheShader))
 	{
+		FMaterialUpdateContext UpdateContext;
 		ValidateStaticPermutationAllowed();
 		UpdateCachedData();
+		RecacheUniformExpressions(true);
 	}
 }
 
