PR #13753: Add asset registry tag to find material instances causing shader permutations

#13753
[RN] minor | Rendering - Materials and Shaders

diff --git a/Engine/Source/Runtime/Engine/Private/Materials/MaterialInstanceConstant.cpp b/Engine/Source/Runtime/Engine/Private/Materials/MaterialInstanceConstant.cpp
--- a/Engine/Source/Runtime/Engine/Private/Materials/MaterialInstanceConstant.cpp
+++ b/Engine/Source/Runtime/Engine/Private/Materials/MaterialInstanceConstant.cpp
@@ -8,6 +8,7 @@
 #include "Materials/Material.h"
 #include "Materials/MaterialInstanceSupport.h"
 #include "MaterialCachedData.h"
+#include "UObject/AssetRegistryTagsContext.h"
 
 #include UE_INLINE_GENERATED_CPP_BY_NAME(MaterialInstanceConstant)
 
@@ -165,6 +166,13 @@ void UMaterialInstanceConstant::ClearParameterValuesEditorOnly()
 	ClearParameterValuesInternal();
 }
 
+void UMaterialInstanceConstant::GetAssetRegistryTags(FAssetRegistryTagsContext Context) const
+{
+	Super::GetAssetRegistryTags(Context);
+
+	Context.AddTag(FAssetRegistryTag("HasStaticPermutationResource", bHasStaticPermutationResource ? TEXT("True") : TEXT("False"), FAssetRegistryTag::TT_Alphabetical));
+}
+
 #if WITH_EDITOR
 void FMaterialInstanceCachedData::InitializeForConstant(const FMaterialLayersFunctions* Layers, const FMaterialLayersFunctions* ParentLayers)
 {
diff --git a/Engine/Source/Runtime/Engine/Public/Materials/MaterialInstanceConstant.h b/Engine/Source/Runtime/Engine/Public/Materials/MaterialInstanceConstant.h
--- a/Engine/Source/Runtime/Engine/Public/Materials/MaterialInstanceConstant.h
+++ b/Engine/Source/Runtime/Engine/Public/Materials/MaterialInstanceConstant.h
@@ -39,6 +39,7 @@ class UMaterialInstanceConstant : public UMaterialInstance
 
 	virtual ENGINE_API void PostEditChangeProperty(FPropertyChangedEvent& PropertyChangedEvent) override;
 	virtual ENGINE_API void UpdateCachedData() override final;
+	virtual ENGINE_API void GetAssetRegistryTags(FAssetRegistryTagsContext Context) const override;
 #endif
 
 	/** Physical material mask to use for this graphics material. Used for sounds, effects etc.*/
