Change cvar name (with deprecation support).
"r.VT.Support16BitPageTable" -> "r.VT.PageTableMode"
New cvar also includes the behavior from removed "r.VT.SplitPhysicalPoolSize".
The cvar can now be modified live at run time instead of being read only.
This opens the door for auto-changing to 32bit page table mode if we hit a panic condition about virtual texture space limits.

diff --git a/Engine/Source/Runtime/Engine/Classes/Engine/Texture2D.h b/Engine/Source/Runtime/Engine/Classes/Engine/Texture2D.h
--- a/Engine/Source/Runtime/Engine/Classes/Engine/Texture2D.h
+++ b/Engine/Source/Runtime/Engine/Classes/Engine/Texture2D.h
@@ -389,7 +389,7 @@ public:
 	virtual bool IsVirtualTexturedWithSinglePhysicalSpace() const { return false;  }
 
 	/**
-	 * Returns true if this virtual texture requires that it's physical pool is never split by the r.VT.SplitPhysicalPoolSize setting.
+	 * Returns true if this virtual texture requires that it's physical pool is never split by the r.VT.PageTableSupportType=2 setting.
 	 * This is necessary if the texture is part of a system that assumes that all textures in the system are in the same physical pool.
 	 * Note that textures with this setting will necessarily be put in different physical spaces to those without the setting.
 	 */
diff --git a/Engine/Source/Runtime/Engine/Classes/VT/VirtualTexturePoolConfig.h b/Engine/Source/Runtime/Engine/Classes/VT/VirtualTexturePoolConfig.h
--- a/Engine/Source/Runtime/Engine/Classes/VT/VirtualTexturePoolConfig.h
+++ b/Engine/Source/Runtime/Engine/Classes/VT/VirtualTexturePoolConfig.h
@@ -108,6 +108,8 @@ struct VirtualTexturePool
 	ENGINE_API static float GetPoolSizeScale();
 	/** Get if the virtual texture physical pools should auto grow on oversubscription. */
 	ENGINE_API static bool GetPoolAutoGrow();
+	/** Get if we support 16bit page table formats. */
+	ENGINE_API static bool Support16BitPageTables();
 	/** Get the scalability setting whether we can create multiple physical pools for the same format. */
 	ENGINE_API static int32 GetSplitPhysicalPoolSize();
 	/** Get a hash of all settings that affect virtual texture physical pool creation. */
diff --git a/Engine/Source/Runtime/Engine/Private/VT/VirtualTexturePoolConfig.cpp b/Engine/Source/Runtime/Engine/Private/VT/VirtualTexturePoolConfig.cpp
--- a/Engine/Source/Runtime/Engine/Private/VT/VirtualTexturePoolConfig.cpp
+++ b/Engine/Source/Runtime/Engine/Private/VT/VirtualTexturePoolConfig.cpp
@@ -31,14 +31,17 @@ static TAutoConsoleVariable<bool> CVarVTPoolAutoGrow(
 	ECVF_RenderThreadSafe
 );
 
-static TAutoConsoleVariable<int32> CVarVTSplitPhysicalPoolSize(
-	TEXT("r.VT.SplitPhysicalPoolSize"),
-	0,
-	TEXT("Create multiple physical pools per format to keep pools at this maximum size in tiles.\n")
-	TEXT("A value of 64 tiles will force 16bit page tables. This can be a page table memory optimization for large physical pools.\n")
-	TEXT("Defaults to 0 (off)."),
-	ECVF_RenderThreadSafe
-);
+static int32 GPageTableMode = 1;
+static FAutoConsoleVariableRef CVarVTPageTableMode(
+	TEXT("r.VT.PageTableMode"),
+	GPageTableMode,
+	TEXT("Sets whether to use 16bit or 32bit page table entries. Defaults to mode 1.\n")
+	TEXT("0: Always use 32bit bit page table entries.\n")
+	TEXT("1: Use 16bit bit page table entries for physical pools with fewer than 4K pages which can be 16bit addressed. (Saves page table memory but can increase virtual texture space count).\n")
+	TEXT("2: Try to force 16bit bit page table entries by splitting the physical pools to 4K page chunks. (Saves page table memory but with less efficient physical pool usage).\n"),
+	ECVF_RenderThreadSafe);
+
+static FAutoConsoleVariableDeprecated CVarVTSupport16BitPageTable_Deprecated(TEXT("r.VT.Support16BitPageTable"), TEXT("r.VT.PageTableMode"), TEXT("5.7"));
 
 /** Track changes and apply to relevant systems. This allows us to dynamically change the scalability settings. */
 void OnVirtualTexturePoolConfigUpdate()
@@ -284,16 +287,21 @@ bool VirtualTexturePool::GetPoolAutoGrow()
 	return CVarVTPoolAutoGrow.GetValueOnAnyThread();
 }
 
+bool VirtualTexturePool::Support16BitPageTables()
+{
+	return GPageTableMode != 0;
+}
+
 int32 VirtualTexturePool::GetSplitPhysicalPoolSize()
 {
-	return FMath::Max(CVarVTSplitPhysicalPoolSize.GetValueOnAnyThread(), 0);
+	return GPageTableMode == 2 ? 64 : 0;
 }
 
 uint32 VirtualTexturePool::GetConfigHash()
 {
 	uint32 Hash = GetTypeHash(GetPoolSizeScale());
 	Hash = HashCombine(Hash, GetTypeHash(GetPoolAutoGrow()));
-	Hash = HashCombine(Hash, GetTypeHash(GetSplitPhysicalPoolSize()));
+	Hash = HashCombine(Hash, GetTypeHash(GPageTableMode));
 	Hash = HashCombine(Hash, GetTypeHash(GVirtualTexturePoolVersion.Get()));
 
 	return Hash;
diff --git a/Engine/Source/Runtime/Engine/Private/VT/VirtualTextureScalability.cpp b/Engine/Source/Runtime/Engine/Private/VT/VirtualTextureScalability.cpp
--- a/Engine/Source/Runtime/Engine/Private/VT/VirtualTextureScalability.cpp
+++ b/Engine/Source/Runtime/Engine/Private/VT/VirtualTextureScalability.cpp
@@ -250,24 +250,4 @@ namespace VirtualTextureScalability
 	{
 		return CVarVTMaxAnisotropy.GetValueOnAnyThread();
 	}
-
-	// Begin deprecated functions.
-	// Can remove include of VirtualTexturePoolConfig.h when these are removed.
-	float GetPoolSizeScale() 
-	{
-		return VirtualTexturePool::GetPoolSizeScale(); 
-	}
-	float GetPoolSizeScale(uint32 GroupIndex) 
-	{
-		return VirtualTexturePool::GetPoolSizeScale(); 
-	}
-	int32 GetSplitPhysicalPoolSize() 
-	{
-		return VirtualTexturePool::GetSplitPhysicalPoolSize(); 
-	}
-	uint32 GetPhysicalPoolSettingsHash() 
-	{
-		return VirtualTexturePool::GetConfigHash(); 
-	}
-	// End deprecated functions.
 }
diff --git a/Engine/Source/Runtime/Engine/Public/VT/VirtualTextureScalability.h b/Engine/Source/Runtime/Engine/Public/VT/VirtualTextureScalability.h
--- a/Engine/Source/Runtime/Engine/Public/VT/VirtualTextureScalability.h
+++ b/Engine/Source/Runtime/Engine/Public/VT/VirtualTextureScalability.h
@@ -30,13 +30,4 @@ namespace VirtualTextureScalability
 	ENGINE_API bool IsAnisotropicFilteringEnabled();
 	/** Get maximum anisotropy when virtual texture sampling. This is also clamped per virtual texture according to the tile border size. */
 	ENGINE_API int32 GetMaxAnisotropy();
-
-	UE_DEPRECATED(5.4, "Use VirtualTexturePool::GetPoolSizeScale() instead.")
-	ENGINE_API float GetPoolSizeScale();
-	UE_DEPRECATED(5.4, "Use VirtualTexturePool::GetPoolSizeScale() instead.")
-	ENGINE_API float GetPoolSizeScale(uint32 GroupIndex);
-	UE_DEPRECATED(5.4, "Use VirtualTexturePool::GetSplitPhysicalPoolSize() instead.")
-	ENGINE_API int32 GetSplitPhysicalPoolSize();
-	UE_DEPRECATED(5.4, "Use VirtualTexturePool::GetConfigHash() instead.")
-	ENGINE_API uint32 GetPhysicalPoolSettingsHash();
 }
diff --git a/Engine/Source/Runtime/Renderer/Private/VT/AllocatedVirtualTexture.cpp b/Engine/Source/Runtime/Renderer/Private/VT/AllocatedVirtualTexture.cpp
--- a/Engine/Source/Runtime/Renderer/Private/VT/AllocatedVirtualTexture.cpp
+++ b/Engine/Source/Runtime/Renderer/Private/VT/AllocatedVirtualTexture.cpp
@@ -3,21 +3,12 @@
 #include "AllocatedVirtualTexture.h"
 
 #include "Misc/StringBuilder.h"
+#include "VT/VirtualTexturePoolConfig.h"
 #include "VT/VirtualTextureScalability.h"
 #include "VT/VirtualTextureSystem.h"
 #include "VT/VirtualTextureSpace.h"
 #include "VT/VirtualTexturePhysicalSpace.h"
 
-bool GSupport16BitPageTable = true;
-static FAutoConsoleVariableRef CVarVTSupport16BitPageTable(
-	TEXT("r.VT.Support16BitPageTable"),
-	GSupport16BitPageTable,
-	TEXT("Enable support for 16 bit page table entries.\n")
-	TEXT("This can reduce page table memory when only 16bit addressing is needed.\n")
-	TEXT("But this can increase the number of page table spaces required when a mixture of 16bit and 32bit addressing is needed.\n")
-	TEXT("Defaults on.\n"),
-	ECVF_ReadOnly);
-
 FAllocatedVirtualTexture::FAllocatedVirtualTexture(
 	FRHICommandListBase& RHICmdList,
 	FVirtualTextureSystem* InSystem,
@@ -87,7 +78,7 @@ FAllocatedVirtualTexture::FAllocatedVirtualTexture(
 	LockOrUnlockTiles(InSystem, true);
 
 	// Use 16bit page table entries if all physical spaces are small enough
-	bool bSupport16BitPageTable = GSupport16BitPageTable;
+	bool bSupport16BitPageTable = VirtualTexturePool::Support16BitPageTables();
 	for (int32 Index = 0; Index < UniquePageTableLayers.Num(); ++Index)
 	{
 		const FVirtualTexturePhysicalSpace* PhysicalSpace = UniquePageTableLayers[Index].PhysicalSpace;
