Limit to 200 the number of shader compilation requests we make for each ODSC server request: if the game starts accumulating a lot of materials, we may end up with a single request with 5000+ shaders to compile. Receiving them all at once was triggering OOMs on some platform

Snoop PendingMeshMaterialThreadedRequests every time we're done with a single request instead of waiting for all requests to finish. This may reduce batching efficiency, but since ODSC server response time is still significant compared to frame time, it should not be a big problem in practice

diff --git a/Engine/Source/Runtime/Engine/Private/ODSC/ODSCThread.cpp b/Engine/Source/Runtime/Engine/Private/ODSC/ODSCThread.cpp
--- a/Engine/Source/Runtime/Engine/Private/ODSC/ODSCThread.cpp
+++ b/Engine/Source/Runtime/Engine/Private/ODSC/ODSCThread.cpp
@@ -11,6 +11,14 @@
 #include "PrimitiveSceneInfo.h"
 #include "Components/PrimitiveComponent.h"
 
+
+static int32 GODSCMaxPayloadPerRequests = 200;
+static FAutoConsoleVariableRef CVarODSCMaxPayloadPerRequests(
+	TEXT("ODSC.MaxPayloadPerRequests"),
+	GODSCMaxPayloadPerRequests,
+	TEXT("Specify how many shaders compilations requests we send in a single message to the ODSC server \n")
+);
+
 FODSCRequestPayload::FODSCRequestPayload(
 	EShaderPlatform InShaderPlatform,
 	ERHIFeatureLevel::Type InFeatureLevel,
@@ -517,27 +525,38 @@ void FODSCThread::Exit()
 
 }
 
-void FODSCThread::Process()
+void CollectPendingMeshMaterialRequests(TQueue<FODSCRequestPayload, EQueueMode::Mpsc>& PendingMeshMaterialThreadedRequests, TArray<FODSCMessageHandler*>& PendingRequestsPipeline, 
+										std::atomic<uint32>& NumPendingMaterialsShaders, const int32 MaxPayloadPerRequests)
 {
-	// cache all pending pipeline requests
+	TArray<FODSCRequestPayload> PayloadsToAggregate;
+	FODSCRequestPayload Payload;
+	while (PendingMeshMaterialThreadedRequests.Dequeue(Payload))
 	{
-		TArray<FODSCRequestPayload> PayloadsToAggregate;
-		FODSCRequestPayload Payload;
-		while (PendingMeshMaterialThreadedRequests.Dequeue(Payload))
-		{
-			PayloadsToAggregate.Add(Payload);
-		}
+		PayloadsToAggregate.Add(Payload);
+	}
 
-		if (PayloadsToAggregate.Num())
+	if (PayloadsToAggregate.Num())
+	{
+		FODSCMessageHandler* RequestHandler = new FODSCMessageHandler(PayloadsToAggregate[0].ShaderPlatform, PayloadsToAggregate[0].FeatureLevel, PayloadsToAggregate[0].QualityLevel, ODSCRecompileCommand::Material);
+		for (const FODSCRequestPayload& payload : PayloadsToAggregate)
 		{
-			FODSCMessageHandler* RequestHandler = new FODSCMessageHandler(PayloadsToAggregate[0].ShaderPlatform, PayloadsToAggregate[0].FeatureLevel, PayloadsToAggregate[0].QualityLevel, ODSCRecompileCommand::Material);
-			for (const FODSCRequestPayload& payload : PayloadsToAggregate)
+			if (RequestHandler->NumPayloads() >= MaxPayloadPerRequests)
 			{
-				RequestHandler->AddPayload(payload);
+				PendingRequestsPipeline.Add(RequestHandler);
+				NumPendingMaterialsShaders += RequestHandler->NumPayloads();
+				RequestHandler = new FODSCMessageHandler(PayloadsToAggregate[0].ShaderPlatform, PayloadsToAggregate[0].FeatureLevel, PayloadsToAggregate[0].QualityLevel, ODSCRecompileCommand::Material);;
 			}
-			PendingRequestsPipeline.Add(RequestHandler);
+
+			RequestHandler->AddPayload(payload);
 		}
+		PendingRequestsPipeline.Add(RequestHandler);
+		NumPendingMaterialsShaders += RequestHandler->NumPayloads();
 	}
+}
+
+void FODSCThread::Process()
+{
+	MaxPayloadPerRequests = FMath::Max(1, GODSCMaxPayloadPerRequests);
 
 	// cache all pending material/global requests
 	{
@@ -610,17 +629,15 @@ void FODSCThread::Process()
 
 	RequestsToStart = MoveTemp(PendingRequestsPipeline);
 
-	uint32 NumPipelines = 0;
-	for (FODSCMessageHandler* NextRequest : RequestsToStart)
-	{
-		NumPipelines += NextRequest->NumPayloads();
-	}
+	NumPendingMaterialsShaders.store(0, std::memory_order_release);
 
-	NumPendingMaterialsShaders.store(NumPipelines, std::memory_order_release);
+	CollectPendingMeshMaterialRequests(PendingMeshMaterialThreadedRequests, RequestsToStart, NumPendingMaterialsShaders, MaxPayloadPerRequests);
 
-	// process any specific mesh material shader requests.
-	for (FODSCMessageHandler* NextRequest : RequestsToStart)
+	int32 NextRequestIndex = 0;
+
+	while (NextRequestIndex < RequestsToStart.Num())
 	{
+		FODSCMessageHandler* NextRequest = RequestsToStart[NextRequestIndex];
 		if (SendMessageToServer(NextRequest))
 		{
 			CompletedThreadedRequests.Enqueue(NextRequest);
@@ -630,6 +647,9 @@ void FODSCThread::Process()
 			PendingRequestsPipeline.Add(NextRequest);
 		}
 		NumPendingMaterialsShaders -= NextRequest->NumPayloads();
+
+		CollectPendingMeshMaterialRequests(PendingMeshMaterialThreadedRequests, RequestsToStart, NumPendingMaterialsShaders, MaxPayloadPerRequests);
+		++NextRequestIndex;
 	}
 
 	NumPendingMaterialsShaders.store(0, std::memory_order_release);
diff --git a/Engine/Source/Runtime/Engine/Private/ODSC/ODSCThread.h b/Engine/Source/Runtime/Engine/Private/ODSC/ODSCThread.h
--- a/Engine/Source/Runtime/Engine/Private/ODSC/ODSCThread.h
+++ b/Engine/Source/Runtime/Engine/Private/ODSC/ODSCThread.h
@@ -317,6 +317,8 @@ private:
 
 	FString ODSCHostIP;
 
+	int32 MaxPayloadPerRequests = INT32_MAX;
+
 	std::atomic<bool> bIsConnectedToODSCServer = false;
 	std::atomic<bool> bHasPendingGlobalShaders = false;
 	std::atomic<uint32> NumPendingMaterialsRecompile = 0;
