Small refinement of my fix submitted in CL 45623939: removing redundant call to RecacheUniformExpressions(true), since FMaterialUpdateContext already calls the function.

#rnx

diff --git a/Engine/Source/Runtime/Engine/Private/Materials/MaterialInstanceConstant.cpp b/Engine/Source/Runtime/Engine/Private/Materials/MaterialInstanceConstant.cpp
--- a/Engine/Source/Runtime/Engine/Private/Materials/MaterialInstanceConstant.cpp
+++ b/Engine/Source/Runtime/Engine/Private/Materials/MaterialInstanceConstant.cpp
@@ -85,7 +85,6 @@ void UMaterialInstanceConstant::SetParentEditorOnly(UMaterialInterface* NewParen
 		FMaterialUpdateContext UpdateContext;
 		ValidateStaticPermutationAllowed();
 		UpdateCachedData();
-		RecacheUniformExpressions(true);
 	}
 }
 
