Material translator, fixed VTStack deduplication
- Compare mips on hashes, not indices

diff --git a/Engine/Source/Runtime/Core/Private/UObject/FortniteMainBranchObjectVersion.cpp b/Engine/Source/Runtime/Core/Private/UObject/FortniteMainBranchObjectVersion.cpp
--- a/Engine/Source/Runtime/Core/Private/UObject/FortniteMainBranchObjectVersion.cpp
+++ b/Engine/Source/Runtime/Core/Private/UObject/FortniteMainBranchObjectVersion.cpp
@@ -13,7 +13,7 @@ TMap<FGuid, FGuid> FFortniteMainBranchObjectVersion::GetSystemGuids()
 	SystemGuids.Add(DevGuids.MATERIALSHADERMAP_DERIVEDDATA_VER, FGuid("F051877974B44D409435F321BDB39DBB"));
 	SystemGuids.Add(DevGuids.NIAGARASHADERMAP_DERIVEDDATA_VER, FGuid("6360A977062842A29ED30E3A7ACB0E64"));
 	SystemGuids.Add(DevGuids.Niagara_LatestScriptCompileVersion, FGuid("783FF3700BB641C799B2D43552FBF91E"));
-	SystemGuids.Add(DevGuids.MaterialTranslationDDCVersion, FGuid("5609B29D40D44A89967F303EAEDED2C0"));
+	SystemGuids.Add(DevGuids.MaterialTranslationDDCVersion, FGuid("5AE92FF51E864B55B973ED11646FA24B"));
 
 	// These GUIDs are in AppendRestrictedFortniteMainStreamObjectVersionGuids as they should rarely be changed.
 	//SystemGuids.Add(DevGuids.NANITE_DERIVEDDATA_VER, FGuid("ACF98184978742A592537DA223A1E6BE"));
diff --git a/Engine/Source/Runtime/Engine/Private/Materials/HLSLMaterialTranslator.cpp b/Engine/Source/Runtime/Engine/Private/Materials/HLSLMaterialTranslator.cpp
--- a/Engine/Source/Runtime/Engine/Private/Materials/HLSLMaterialTranslator.cpp
+++ b/Engine/Source/Runtime/Engine/Private/Materials/HLSLMaterialTranslator.cpp
@@ -7095,6 +7095,7 @@ uint32 FHLSLMaterialTranslator::AcquireVTStackIndex(
 	Flags |= SampleInfo.bGenerateFeedback ? (1 << 1) : 0;
 	Flags |= SampleInfo.bMeshPaint ? (1 << 2) : 0;
 	Flags |= SampleInfo.bMaterialCache ? (1 << 3) : 0;
+	Flags |= SampleInfo.bCollection ? (1 << 4) : 0;
 
 	uint64 Hash = CityHash128to64({ CurrentScopeID, CoordinateHash });
 	Hash = CityHash128to64({ Hash, MipValue0Hash });
@@ -7123,7 +7124,16 @@ uint32 FHLSLMaterialTranslator::AcquireVTStackIndex(
 			Entry.MipValue1Hash == MipValue1Hash &&
 			Entry.PageTableHash == PageTableHash &&
 			Entry.PageTableUniformHash == PageTableUniformHash &&
-			Entry.SampleInfo == SampleInfo)
+			Entry.SampleInfo.MipValueMode == SampleInfo.MipValueMode &&
+			Entry.SampleInfo.AddressU == SampleInfo.AddressU &&
+			Entry.SampleInfo.AddressV == SampleInfo.AddressV &&
+			Entry.SampleInfo.AspectRatio == SampleInfo.AspectRatio &&
+			Entry.SampleInfo.PreallocatedStackTextureIndex == SampleInfo.PreallocatedStackTextureIndex &&
+			Entry.SampleInfo.bAdaptive == SampleInfo.bAdaptive &&
+			Entry.SampleInfo.bGenerateFeedback == SampleInfo.bGenerateFeedback &&
+			Entry.SampleInfo.bMeshPaint == SampleInfo.bMeshPaint &&
+			Entry.SampleInfo.bMaterialCache == SampleInfo.bMaterialCache &&
+			Entry.SampleInfo.bCollection == SampleInfo.bCollection)
 		{
 			return Index;
 		}
