Do linear search in array when finding a VTStack during material translation.
Previous hash map search is out of order so that we may miss finding the virtual texture index in an existing stack.
Alternative would be to do two pass hash map search instead (first for existing index, then for non-empty stacks).
But any performance gain for hash map look up is likely not worth the complexity. Plus for low stack numbers, a linear search may be faster anyway.
#rnx

diff --git a/Engine/Source/Runtime/Core/Private/UObject/FortniteMainBranchObjectVersion.cpp b/Engine/Source/Runtime/Core/Private/UObject/FortniteMainBranchObjectVersion.cpp
--- a/Engine/Source/Runtime/Core/Private/UObject/FortniteMainBranchObjectVersion.cpp
+++ b/Engine/Source/Runtime/Core/Private/UObject/FortniteMainBranchObjectVersion.cpp
@@ -13,7 +13,7 @@ TMap<FGuid, FGuid> FFortniteMainBranchObjectVersion::GetSystemGuids()
 	SystemGuids.Add(DevGuids.MATERIALSHADERMAP_DERIVEDDATA_VER, FGuid("F051877974B44D409435F321BDB39DBB"));
 	SystemGuids.Add(DevGuids.NIAGARASHADERMAP_DERIVEDDATA_VER, FGuid("6360A977062842A29ED30E3A7ACB0E64"));
 	SystemGuids.Add(DevGuids.Niagara_LatestScriptCompileVersion, FGuid("783FF3700BB641C799B2D43552FBF91E"));
-	SystemGuids.Add(DevGuids.MaterialTranslationDDCVersion, FGuid("5AE92FF51E864B55B973ED11646FA24B"));
+	SystemGuids.Add(DevGuids.MaterialTranslationDDCVersion, FGuid("8F124D15BF2B4522933D5C6C52D6F3E0"));
 
 	// These GUIDs are in AppendRestrictedFortniteMainStreamObjectVersionGuids as they should rarely be changed.
 	//SystemGuids.Add(DevGuids.NANITE_DERIVEDDATA_VER, FGuid("ACF98184978742A592537DA223A1E6BE"));
diff --git a/Engine/Source/Runtime/Engine/Private/Materials/HLSLMaterialTranslator.cpp b/Engine/Source/Runtime/Engine/Private/Materials/HLSLMaterialTranslator.cpp
--- a/Engine/Source/Runtime/Engine/Private/Materials/HLSLMaterialTranslator.cpp
+++ b/Engine/Source/Runtime/Engine/Private/Materials/HLSLMaterialTranslator.cpp
@@ -7090,27 +7090,8 @@ uint32 FHLSLMaterialTranslator::AcquireVTStackIndex(
 	const uint64 PageTableHash = GetParameterHash(SampleInfo.PageTableIndex);
 	const uint64 PageTableUniformHash = GetParameterHash(SampleInfo.PageTableUniformIndex);
 	
-	uint32 Flags = 0;
-	Flags |= SampleInfo.bAdaptive ? (1 << 0) : 0;
-	Flags |= SampleInfo.bGenerateFeedback ? (1 << 1) : 0;
-	Flags |= SampleInfo.bMeshPaint ? (1 << 2) : 0;
-	Flags |= SampleInfo.bMaterialCache ? (1 << 3) : 0;
-	Flags |= SampleInfo.bCollection ? (1 << 4) : 0;
-
-	uint64 Hash = CityHash128to64({ CurrentScopeID, CoordinateHash });
-	Hash = CityHash128to64({ Hash, MipValue0Hash });
-	Hash = CityHash128to64({ Hash, MipValue1Hash });
-	Hash = CityHash128to64({ Hash, PageTableHash });
-	Hash = CityHash128to64({ Hash, PageTableUniformHash });
-	Hash = CityHash128to64({ Hash, (uint64)SampleInfo.MipValueMode });
-	Hash = CityHash128to64({ Hash, (uint64)SampleInfo.AddressU });
-	Hash = CityHash128to64({ Hash, (uint64)SampleInfo.AddressV });
-	Hash = CityHash128to64({ Hash, (uint64)(SampleInfo.AspectRatio * 1000.0f) });
-	Hash = CityHash128to64({ Hash, (uint64)SampleInfo.PreallocatedStackTextureIndex });
-	Hash = CityHash128to64({ Hash, (uint64)Flags });
-
 	// First check to see if we have an existing VTStack that matches this key, that can still fit another layer
-	for (int32 Index = VTStackHash.First(Hash); VTStackHash.IsValid(Index); Index = VTStackHash.Next(Index))
+	for (int32 Index = 0; Index < VTStacks.Num(); ++Index)
 	{
 		const FMaterialVirtualTextureStack& Stack = MaterialCompilationOutput.UniformExpressionSet.VTStacks[Index];
 		const FMaterialVTStackEntry& Entry = VTStacks[Index];
@@ -7141,7 +7122,6 @@ uint32 FHLSLMaterialTranslator::AcquireVTStackIndex(
 
 	// Need to allocate a new VTStack
 	const int32 StackIndex = VTStacks.AddDefaulted();
-	VTStackHash.Add(Hash, StackIndex);
 	FMaterialVTStackEntry& Entry = VTStacks[StackIndex];
 	Entry.SampleInfo = SampleInfo;
 	Entry.ScopeID = CurrentScopeID;
diff --git a/Engine/Source/Runtime/Engine/Private/Materials/HLSLMaterialTranslator.h b/Engine/Source/Runtime/Engine/Private/Materials/HLSLMaterialTranslator.h
--- a/Engine/Source/Runtime/Engine/Private/Materials/HLSLMaterialTranslator.h
+++ b/Engine/Source/Runtime/Engine/Private/Materials/HLSLMaterialTranslator.h
@@ -369,7 +369,6 @@ protected:
 
 	/** VT Stacks */
 	TArray<FMaterialVTStackEntry> VTStacks;
-	FHashTable VTStackHash;
 	
 	/** All material tags discovered */
 	TMap<FGuid, FMaterialCacheTag> MaterialCacheTags;
