Fix rare crash on shutdown that can occur if a FStreamableHandle outlives its manager.
#tests Engineered a repro locally. Could not repro naturally due to the issue being highly timing dependent

diff --git a/Engine/Source/Runtime/Engine/Private/StreamableManager.cpp b/Engine/Source/Runtime/Engine/Private/StreamableManager.cpp
--- a/Engine/Source/Runtime/Engine/Private/StreamableManager.cpp
+++ b/Engine/Source/Runtime/Engine/Private/StreamableManager.cpp
@@ -1055,16 +1055,13 @@ FStreamableHandle::~FStreamableHandle()
 {
 	checkf(IsInGameThread(), TEXT("Streamable handles aren't thread-safe and must be destroyed on the game thread"));
 
-	if (IsActive())
+	if (IsActive() && bLoadStarted)
 	{
-		if (!bStalled)
+		for (const FSoftObjectPath&  RequestedAsset : RequestedAssets)
 		{
-			for (const FSoftObjectPath&  RequestedAsset : RequestedAssets)
+			if (FStreamable* Streamable = OwningManager->FindStreamable(RequestedAsset))
 			{
-				if (FStreamable* Streamable = OwningManager->FindStreamable(RequestedAsset))
-				{
-					RemoveActiveHandle(*Streamable, *this);
-				}
+				RemoveActiveHandle(*Streamable, *this);
 			}
 		}
 
