Iris
* Reduce glue code needed when user code interacts with the Iris system and bridge in FReplicationSystemUtil
* Added ForEachBridge / ForEachSystem / ForEachSystemAndBridge:  they only need a world
* Added ForEachServerBridge / ForEachServerSystem that only trigger on server worlds

#rnx

diff --git a/Engine/Source/Runtime/Engine/Private/Net/Iris/ReplicationSystem/ReplicationSystemUtil.cpp b/Engine/Source/Runtime/Engine/Private/Net/Iris/ReplicationSystem/ReplicationSystemUtil.cpp
--- a/Engine/Source/Runtime/Engine/Private/Net/Iris/ReplicationSystem/ReplicationSystemUtil.cpp
+++ b/Engine/Source/Runtime/Engine/Private/Net/Iris/ReplicationSystem/ReplicationSystemUtil.cpp
@@ -2,52 +2,221 @@
 
 #include "Net/Iris/ReplicationSystem/ReplicationSystemUtil.h"
 
-
-#include "EngineUtils.h"
 #include "Engine/Engine.h"
+#include "Engine/GameInstance.h"
 #include "Engine/NetConnection.h"
 #include "Engine/NetDriver.h"
 #include "Engine/World.h"
+
+#include "EngineUtils.h"
+
+#include "GameFramework/PlayerController.h"
+
+#include "HAL/IConsoleManager.h"
+
+#include "Iris/Core/IrisLog.h"
 #include "Iris/Core/IrisProfiler.h"
+
 #include "Iris/ReplicationSystem/Filtering/NetObjectFilter.h"
 #include "Iris/ReplicationSystem/ReplicationSystem.h"
-#include "Iris/Core/IrisLog.h"
-#include "Net/Iris/ReplicationSystem/EngineReplicationBridge.h"
-#include "Net/NetSubObjectRegistryGetter.h"
-#include "HAL/IConsoleManager.h"
-#include "Templates/Casts.h"
-#include "Net/Subsystems/NetworkSubsystem.h"
+
 #include "Net/Core/Connection/ConnectionHandle.h"
 #include "Net/Core/Misc/NetConditionGroupManager.h"
 #include "Net/Core/NetHandle/NetHandleManager.h"
-#include "GameFramework/PlayerController.h"
+
+#include "Net/Iris/ReplicationSystem/EngineReplicationBridge.h"
+#include "Net/NetSubObjectRegistryGetter.h"
+#include "Net/Subsystems/NetworkSubsystem.h"
+
+#include "Templates/Casts.h"
 #include "Templates/Function.h"
 
+namespace UE::Net::Private
+{
+	bool DoesTypeMatch(FReplicationSystemUtil::EType Type, UReplicationSystem* System)
+	{
+		switch(Type)
+		{
+			case FReplicationSystemUtil::EType::Any:
+			{
+				return true;
+			} break;
+
+			case FReplicationSystemUtil::EType::Server:
+			{
+				return System->IsServer();
+			} break;
+
+			case FReplicationSystemUtil::EType::Client:
+			{
+				return !System->IsServer();
+			} break;
+
+			default:
+			{
+				ensureMsgf(false, TEXT("EType entry not handled"));
+				return false;
+			} break;
+		}
+	}
+}
 
 namespace UE::Net
 {
 
-void FReplicationSystemUtil::ForEachReplicationSystem(const UEngine* Engine, const UWorld* World, TFunctionRef<void(UReplicationSystem*)> Function)
+void FReplicationSystemUtil::ForEachBridge(const UWorld* World, TFunctionRef<void(UEngineReplicationBridge*)> Function, EType Type)
+{
+	using namespace UE::Net::Private;
+
+	if (!World)
 	{
-		if (Engine == nullptr || World == nullptr)
+		return;
+	}
+
+	UGameInstance* GameInstance = World->GetGameInstance();
+	if (!GameInstance)
+	{
+		return;
+	}
+
+	UEngine* Engine = GameInstance->GetEngine();
+	if (!Engine)
+	{
+		return;
+	}
+
+	const FWorldContext* Context = Engine->GetWorldContextFromWorld(World);
+	if (!Context)
+	{
+		return;
+	}
+
+	for (const FNamedNetDriver& NamedNetDriver : Context->ActiveNetDrivers)
+	{
+		if (UNetDriver* NetDriver = NamedNetDriver.NetDriver)
 		{
-			return;
+			if (UReplicationSystem* ReplicationSystem = NetDriver->GetReplicationSystem())
+			{
+				if (DoesTypeMatch(Type, ReplicationSystem))
+				{
+					// There should always be a bridge attached to the system
+					UEngineReplicationBridge* ReplicationBridge = CastChecked<UEngineReplicationBridge>(ReplicationSystem->GetReplicationBridge());
+
+					Function(ReplicationBridge);
+				}
+			}
+		}
+	}
+}
+
+void FReplicationSystemUtil::ForEachSystem(const UWorld* World, TFunctionRef<void(UReplicationSystem*)> Function, EType Type)
+{
+	using namespace UE::Net::Private;
+
+	if (!World)
+	{
+		return;
+	}
+
+	UGameInstance* GameInstance = World->GetGameInstance();
+	if (!GameInstance)
+	{
+		return;
+	}
+
+	UEngine* Engine = GameInstance->GetEngine();
+	if (!Engine)
+	{
+		return;
+	}
+
+	const FWorldContext* Context = Engine->GetWorldContextFromWorld(World);
+	if (!Context)
+	{
+		return;
+	}
+
+	for (const FNamedNetDriver& NamedNetDriver : Context->ActiveNetDrivers)
+	{
+		if (UNetDriver* NetDriver = NamedNetDriver.NetDriver)
+		{
+			if (UReplicationSystem* ReplicationSystem = NetDriver->GetReplicationSystem())
+			{
+				if (DoesTypeMatch(Type, ReplicationSystem))
+				{
+					Function(ReplicationSystem);
+				}
+			}
+		}
+	}
+}
+
+void FReplicationSystemUtil::ForEachSystemAndBridge(const UWorld* World, TFunctionRef<void(UReplicationSystem*, UEngineReplicationBridge*)> Function, EType Type)
+{
+	using namespace UE::Net::Private;
+
+	if (!World)
+	{
+		return;
+	}
+
+	UGameInstance* GameInstance = World->GetGameInstance();
+	if (!GameInstance)
+	{
+		return;
+	}
+
+	UEngine* Engine = GameInstance->GetEngine();
+	if (!Engine)
+	{
+		return;
+	}
+
+	const FWorldContext* Context = Engine->GetWorldContextFromWorld(World);
+	if (!Context)
+	{
+		return;
+	}
+
+	for (const FNamedNetDriver& NamedNetDriver : Context->ActiveNetDrivers)
+	{
+		if (UNetDriver* NetDriver = NamedNetDriver.NetDriver)
+		{
+			if (UReplicationSystem* ReplicationSystem = NetDriver->GetReplicationSystem())
+			{
+				if (DoesTypeMatch(Type, ReplicationSystem))
+				{
+					// There should always be a bridge attached to the system
+					UEngineReplicationBridge* ReplicationBridge = CastChecked<UEngineReplicationBridge>(ReplicationSystem->GetReplicationBridge());
+
+					Function(ReplicationSystem, ReplicationBridge);
+				}
+			}
 		}
+	}
+}
+
+void FReplicationSystemUtil::ForEachReplicationSystem(const UEngine* Engine, const UWorld* World, TFunctionRef<void(UReplicationSystem*)> Function)
+{
+	if (Engine == nullptr || World == nullptr)
+	{
+		return;
+	}
 
-		if (const FWorldContext* Context = Engine->GetWorldContextFromWorld(World))
+	if (const FWorldContext* Context = Engine->GetWorldContextFromWorld(World))
+	{
+		for (const FNamedNetDriver& NamedNetDriver : Context->ActiveNetDrivers)
 		{
-			for (const FNamedNetDriver& NamedNetDriver : Context->ActiveNetDrivers)
+			if (UNetDriver* NetDriver = NamedNetDriver.NetDriver)
 			{
-				if (UNetDriver* NetDriver = NamedNetDriver.NetDriver)
+				if (UReplicationSystem* ReplicationSystem = NetDriver->GetReplicationSystem())
 				{
-					if (UReplicationSystem* ReplicationSystem = NetDriver->GetReplicationSystem())
-					{
-						Function(ReplicationSystem);
-					}
+					Function(ReplicationSystem);
 				}
 			}
 		}
 	}
+}
 
 void FReplicationSystemUtil::ForEachReplicationSystem(TFunctionRef<void(UReplicationSystem*)> Function)
 	{
diff --git a/Engine/Source/Runtime/Engine/Public/Net/Iris/ReplicationSystem/ReplicationSystemUtil.h b/Engine/Source/Runtime/Engine/Public/Net/Iris/ReplicationSystem/ReplicationSystemUtil.h
--- a/Engine/Source/Runtime/Engine/Public/Net/Iris/ReplicationSystem/ReplicationSystemUtil.h
+++ b/Engine/Source/Runtime/Engine/Public/Net/Iris/ReplicationSystem/ReplicationSystemUtil.h
@@ -34,6 +34,41 @@ namespace UE::Net
 // Helper methods to interact with ReplicationSystem from engine code
 struct FReplicationSystemUtil
 {
+	/** Enum representing the type of system you want to iterate on */
+	enum class EType : uint8 
+	{
+		Any,
+		Client,
+		Server,
+	};
+
+	/** Call a function for all ReplicationBridges of a specific world and type. */
+	ENGINE_API static void ForEachBridge(const UWorld* World, TFunctionRef<void(UEngineReplicationBridge*)> Function, EType Type = EType::Any);
+
+	/** Call a function for all ReplicationSystems of a specific world and type. */
+	ENGINE_API static void ForEachSystem(const UWorld* World, TFunctionRef<void(UReplicationSystem*)> Function, EType Type = EType::Any);
+
+	/** Call a function for all ReplicationSystems + ReplicationBridges of a specific world and type. */
+	ENGINE_API static void ForEachSystemAndBridge(const UWorld* World, TFunctionRef<void(UReplicationSystem*, UEngineReplicationBridge*)> Function, EType Type = EType::Any);
+
+	/** Call a function for all Server ReplicationBridges */
+	static void ForEachServerBridge(const UWorld* World, TFunctionRef<void(UEngineReplicationBridge*)> Function)
+	{
+		ForEachBridge(World, Function, EType::Server);
+	}
+
+	/** Call a function for all Server ReplicationSystems */
+	static void ForEachServerSystem(const UWorld* World, TFunctionRef<void(UReplicationSystem*)> Function)
+	{
+		ForEachSystem(World, Function, EType::Server);
+	}
+
+	/** Call a function for all Server ReplicationSystems + ReplicationBridges */
+	static void ForEachServerSystemAndBridge(const UWorld* World, TFunctionRef<void(UReplicationSystem*, UEngineReplicationBridge* Bridge)> Function)
+	{
+		ForEachSystemAndBridge(World, Function, EType::Server);
+	}
+
 	/** Call a function on all ReplicationSystems of a specific world. */
 	ENGINE_API static void ForEachReplicationSystem(const UEngine* Engine, const UWorld* World, TFunctionRef<void(UReplicationSystem*)> Function);
 
