Iris
* Change wrong reference to UReplicationBridge to UObjectReplicationBridge

#rnx

diff --git a/Engine/Source/Runtime/Engine/Private/Net/Iris/ReplicationSystem/EngineReplicationBridge.cpp b/Engine/Source/Runtime/Engine/Private/Net/Iris/ReplicationSystem/EngineReplicationBridge.cpp
--- a/Engine/Source/Runtime/Engine/Private/Net/Iris/ReplicationSystem/EngineReplicationBridge.cpp
+++ b/Engine/Source/Runtime/Engine/Private/Net/Iris/ReplicationSystem/EngineReplicationBridge.cpp
@@ -93,7 +93,7 @@ bool IsActorValidForIrisReplication(const AActor* Actor)
 	return IsValid(Actor) && !Actor->IsActorBeingDestroyed() && !Actor->IsUnreachable();
 }
 
-void ActorReplicationBridgePreUpdateFunction(TArrayView<UObject*> Instances, const UReplicationBridge* Bridge)
+void ActorReplicationBridgePreUpdateFunction(TArrayView<UObject*> Instances, const UObjectReplicationBridge* Bridge)
 {
 	UNetDriver* NetDriver = CastChecked<const UEngineReplicationBridge>(Bridge)->GetNetDriver();
 	for (UObject* Instance : Instances)
