Adding categories to trace regions

diff --git a/Engine/Source/Runtime/Core/Public/ProfilingDebugging/MiscTrace.h b/Engine/Source/Runtime/Core/Public/ProfilingDebugging/MiscTrace.h
--- a/Engine/Source/Runtime/Core/Public/ProfilingDebugging/MiscTrace.h
+++ b/Engine/Source/Runtime/Core/Public/ProfilingDebugging/MiscTrace.h
@@ -158,7 +158,8 @@ if (UE_TRACE_CHANNELEXPR_IS_ENABLED(BookmarkChannel)) \
 FMiscTrace::OutputBeginRegion(RegionName __VA_OPT__(,) __VA_ARGS__);
 
 /**
- * Start a Timing Region, identified by an Id
+ * Start a Timing Region with a name and optionally a category. Returns an Id that may be used with TRACE_END_REGION_WITH_ID to end tracking.
+ * Both the name and the category should be null-terminated TCHAR* strings, and may be dynamic strings or string constants.
  * Example Usage:
  * @code
  * uint64 OpenRegionId = TRACE_BEGIN_REGION_WITH_ID(TEXT("MyRegion"), TEXT("MyCategory"))
diff --git a/Engine/Source/Runtime/CoreUObject/Private/UObject/GarbageCollection.cpp b/Engine/Source/Runtime/CoreUObject/Private/UObject/GarbageCollection.cpp
--- a/Engine/Source/Runtime/CoreUObject/Private/UObject/GarbageCollection.cpp
+++ b/Engine/Source/Runtime/CoreUObject/Private/UObject/GarbageCollection.cpp
@@ -5564,7 +5564,7 @@ void PreCollectGarbageImpl(EObjectFlags KeepFlags)
 
 	DECLARE_SCOPE_CYCLE_COUNTER( TEXT( "CollectGarbageInternal" ), STAT_CollectGarbageInternal, STATGROUP_GC );
 	STAT_ADD_CUSTOMMESSAGE_NAME( STAT_NamedMarker, TEXT( "GarbageCollection - Begin" ) );
-	TRACE_BEGIN_REGION(TEXT("GarbageCollection"));
+	TRACE_BEGIN_REGION(TEXT("GarbageCollection"), TEXT("GarbageCollection"));
 
 	// We can't collect garbage while there's a load in progress. E.g. one potential issue is Import.XObject
 	check(!IsLoading());
diff --git a/Engine/Source/Runtime/Engine/Private/World.cpp b/Engine/Source/Runtime/Engine/Private/World.cpp
--- a/Engine/Source/Runtime/Engine/Private/World.cpp
+++ b/Engine/Source/Runtime/Engine/Private/World.cpp
@@ -3626,7 +3626,7 @@ void UWorld::AddToWorld(ULevel* Level, const FTransform& LevelTransform, bool bC
 	if( bExecuteNextStep && !bIsLevelMakingVisible && !IsLevelMakingInvisible(Level) )
 	{
 		TRACE_CPUPROFILER_EVENT_SCOPE(AddToWorld_CurrentLevelPendingInvisibility);
-		TRACE_BEGIN_REGION(*WriteToString<256>(TEXT("AddToWorld: "), Level->GetOutermost()->GetName()));
+		TRACE_BEGIN_REGION(*WriteToString<256>(TEXT("AddToWorld: "), Level->GetOutermost()->GetName()), TEXT("LevelStreaming"));
 
 		Level->OwningWorld = this;
 
@@ -4076,7 +4076,7 @@ void UWorld::RemoveFromWorld(ULevel* Level, bool bAllowIncrementalRemoval, const
 			{
 				if (!Timeout.IsExpired())
 				{
-					TRACE_BEGIN_REGION(*WriteToString<256>(TEXT("RemoveFromWorld: "), Level->GetPackage()->GetName()));
+					TRACE_BEGIN_REGION(*WriteToString<256>(TEXT("RemoveFromWorld: "), Level->GetPackage()->GetName()), TEXT("LevelStreaming"));
 
 					// Mark level as being in process of being made invisible. 
 					// This will prevent this level from being unloaded or made visible in the meantime
