Add GetTimeElapsed and GetTimeElapsedFraction functions to FCachedAnimTransitionData

diff --git a/Engine/Source/Runtime/Engine/Private/Animation/CachedAnimData.cpp b/Engine/Source/Runtime/Engine/Private/Animation/CachedAnimData.cpp
--- a/Engine/Source/Runtime/Engine/Private/Animation/CachedAnimData.cpp
+++ b/Engine/Source/Runtime/Engine/Private/Animation/CachedAnimData.cpp
@@ -329,3 +329,28 @@ float FCachedAnimTransitionData::GetCrossfadeDuration(UAnimInstance& InAnimInsta
 	return 0.f;
 }
 
+float FCachedAnimTransitionData::GetTimeElapsed(UAnimInstance& InAnimInstance) const
+{
+	CacheIndices(InAnimInstance);
+
+	if (MachineIndex != INDEX_NONE && TransitionIndex != INDEX_NONE)
+	{
+		return InAnimInstance.GetInstanceTransitionTimeElapsed(MachineIndex, TransitionIndex);
+	}
+
+	return 0.f;
+}
+
+float FCachedAnimTransitionData::GetTimeElapsedFraction(UAnimInstance& InAnimInstance) const
+{
+	CacheIndices(InAnimInstance);
+
+	if (MachineIndex != INDEX_NONE && TransitionIndex != INDEX_NONE)
+	{
+		return InAnimInstance.GetInstanceTransitionTimeElapsedFraction(MachineIndex, TransitionIndex);
+	}
+
+	return 0.f;
+}
+
+
diff --git a/Engine/Source/Runtime/Engine/Public/Animation/CachedAnimData.h b/Engine/Source/Runtime/Engine/Public/Animation/CachedAnimData.h
--- a/Engine/Source/Runtime/Engine/Public/Animation/CachedAnimData.h
+++ b/Engine/Source/Runtime/Engine/Public/Animation/CachedAnimData.h
@@ -172,6 +172,8 @@ struct FCachedAnimTransitionData
 	FName ToStateName;
 
 	ENGINE_API float GetCrossfadeDuration(UAnimInstance& InAnimInstance) const;
+	ENGINE_API float GetTimeElapsed(UAnimInstance& InAnimInstance) const;
+	ENGINE_API float GetTimeElapsedFraction(UAnimInstance& InAnimInstance) const;
 
 private:
 	ENGINE_API void CacheIndices(UAnimInstance& InAnimInstance) const;
