Fix typo and incorrect CVar/Settings description.

diff --git a/Engine/Source/Runtime/Engine/Classes/Engine/RendererSettings.h b/Engine/Source/Runtime/Engine/Classes/Engine/RendererSettings.h
--- a/Engine/Source/Runtime/Engine/Classes/Engine/RendererSettings.h
+++ b/Engine/Source/Runtime/Engine/Classes/Engine/RendererSettings.h
@@ -1250,7 +1250,7 @@ class URendererSettings : public UDeveloperSettings
 	UPROPERTY(config, EditAnywhere, Category = Substrate, meta = (
 		ConsoleVariable = "r.Substrate.ProjectClosuresPerPixelOverride",
 		DisplayName = "Substrate Closure-Per-Pixel Override",
-		ToolTip = "Select how the maximum number of closures evaluated per pixel is defined for each platforms.\n - Use Platform Closure Count: Use the mininum between the platform's closure count and the project's closure count.\n - Force Closure Count : Force the project's closure count on platforms supporting Adaptive GBuffer.",
+		ToolTip = "Select how the maximum number of closures evaluated per pixel is defined for each platform.\n - Use Platform Closure Count: Use the minimum between the platform's closure count and the project's closure count.\n - Force Closure Count: Force the project's closure count on platforms supporting Adaptive GBuffer.",
 		EditCondition = "SubstrateGBufferFormat == 1",
 		ConfigRestartRequired = true))
 	TEnumAsByte<ESubstrateClosureConfig::Type> SubstrateClosurePerPixelOverride;
diff --git a/Engine/Source/Runtime/RenderCore/Private/RenderUtils.cpp b/Engine/Source/Runtime/RenderCore/Private/RenderUtils.cpp
--- a/Engine/Source/Runtime/RenderCore/Private/RenderUtils.cpp
+++ b/Engine/Source/Runtime/RenderCore/Private/RenderUtils.cpp
@@ -1960,7 +1960,7 @@ static TAutoConsoleVariable<int32> CVarSubstrateProjectClosuresPerPixel(
 
 static TAutoConsoleVariable<int32> CVarSubstrateProjectClosuresPerPixelOverride(
 	TEXT("r.Substrate.ProjectClosuresPerPixelOverride"),
-	0,	// By default use platform's max closure count (i.e., ESubstrateClosureConfig::Type::UsePlatformMaxClosuresPerPixel)
+	0,	// By default, use the mininum between the platform's closure count and the project's closure count (i.e., ESubstrateClosureConfig::Type::UseMinProjectAndPlatforClosuresPerPixel)
 	TEXT("Select how the maximum number of closures evaluated per pixel is defined for each platforms.\n - Use Platform Closure Count: Use the mininum between the platform's closure count and the project's closure count.\n - Force Closure Count : Force the project's closure count on platforms supporting Adaptive GBuffer."),
 	ECVF_ReadOnly | ECVF_RenderThreadSafe | ECVF_Preview);
 
