Fix for static mesh render proxy reregister assert cutting and pasting between editor versions, caused by duplicate pointers to the same UActorComponent in the input array.  These are now gracefully handled by ignoring the duplicates.

#rnx

diff --git a/Engine/Source/Runtime/Engine/Private/StaticMeshResources.cpp b/Engine/Source/Runtime/Engine/Private/StaticMeshResources.cpp
--- a/Engine/Source/Runtime/Engine/Private/StaticMeshResources.cpp
+++ b/Engine/Source/Runtime/Engine/Private/StaticMeshResources.cpp
@@ -121,16 +121,47 @@ FStaticMeshComponentBulkReregisterContext::FStaticMeshComponentBulkReregisterCon
 	TArray<UPrimitiveComponent*> ComponentsWithRenderState;
 	ComponentsWithRenderState.Reserve(InComponents.Num());
 
-	for (UActorComponent* Component : InComponents)
+	// In some rare cases, the caller may pass multiple pointers to the same component.  In that case, we must only reregister it once.
+	// The duplicate tracker is only updated when we have a possible duplicate (detected by checking if bBulkReregister is already set),
+	// to avoid overhead in in the common case where there are no duplicates.
+	TSet<UActorComponent*> DuplicateTracker;
+
+	for (int32 ComponentIndex = 0; ComponentIndex < InComponents.Num(); ++ComponentIndex)
 	{
+		UActorComponent* Component = InComponents[ComponentIndex];
 		UStaticMeshComponent* StaticMeshComponent = Cast<UStaticMeshComponent>(Component);
 
 		// When running a render state update, only batch process components that are dirty and not hidden
 		if (StaticMeshComponent && StaticMeshComponent->SceneProxy &&
 			(ReregisterType == EBulkReregister::Component || (StaticMeshComponent->IsRenderStateDirty() && StaticMeshComponent->ShouldComponentAddToScene())))
 		{
+			// Do we have a potential duplicate?
+			if (StaticMeshComponent->bBulkReregister)
+			{
+				// Initialize the duplicate tracker with the elements before the current one.
+				if (DuplicateTracker.IsEmpty())
+				{
+					DuplicateTracker.Reserve(InComponents.Num());
+					for (int32 PreviousIndex = 0; PreviousIndex < ComponentIndex; ++PreviousIndex)
+					{
+						DuplicateTracker.Add(InComponents[PreviousIndex]);
+					}
+				}
+			}
+
+			// Check if the current item is a duplicate.
+			if (!DuplicateTracker.IsEmpty())
+			{
+				if (DuplicateTracker.Find(Component))
+				{
+					continue;
+				}
+				DuplicateTracker.Add(Component);
+			}
+
 			StaticMeshComponents.Add(StaticMeshComponent);
 
+			// If bBulkReregister is set, and it's NOT a duplicate, it implies a code bug (for example, nesting of a reregister context, which is not supported).
 			check(StaticMeshComponent->bBulkReregister == false);
 			StaticMeshComponent->bBulkReregister = true;
 
