Instanced Skinned Mesh component fixes:
- Change GetComponentDesc to use the templated helper method introduced by FastGeo changes (appears to have been an oversight from CL 41411512)
- Fix by-reference captures of local variables in a lambda that is returned from a function.
- Added a new Gather routine for instance data writers when writing a single value to all elements, and simplified the instance data writing for transform provider instance data.

#rnx

diff --git a/Engine/Source/Runtime/Engine/Private/Components/InstancedSkinnedMeshComponent.cpp b/Engine/Source/Runtime/Engine/Private/Components/InstancedSkinnedMeshComponent.cpp
--- a/Engine/Source/Runtime/Engine/Private/Components/InstancedSkinnedMeshComponent.cpp
+++ b/Engine/Source/Runtime/Engine/Private/Components/InstancedSkinnedMeshComponent.cpp
@@ -360,87 +360,7 @@ void UInstancedSkinnedMeshComponent::PostAssetCompilation()
 
 FInstanceDataManagerSourceDataDesc UInstancedSkinnedMeshComponent::GetComponentDesc(EShaderPlatform InShaderPlatform)
 {
-	FInstanceDataManagerSourceDataDesc ComponentDesc;
-
-	ComponentDesc.PrimitiveMaterialDesc = GetUsedMaterialPropertyDesc(InShaderPlatform);
-
-	FInstanceDataFlags Flags;
-	Flags.bHasPerInstanceRandom = ComponentDesc.PrimitiveMaterialDesc.bAnyMaterialHasPerInstanceRandom;
-	Flags.bHasPerInstanceCustomData = ComponentDesc.PrimitiveMaterialDesc.bAnyMaterialHasPerInstanceCustomData && NumCustomDataFloats != 0;
-#if WITH_EDITOR
-	Flags.bHasPerInstanceEditorData = GIsEditor != 0 && bHasPerInstanceHitProxies;
-#endif
-	
-	const bool bForceRefPose = CVarInstancedSkinnedMeshesForceRefPose.GetValueOnAnyThread() != 0;
-	const bool bValidTransformProvider = !bForceRefPose && TransformProvider != nullptr && TransformProvider->IsEnabled();
-
-	const FTransform& ComponentTransform = GetComponentTransform();
-
-	Flags.bHasPerInstanceHierarchyOffset = false;
-	Flags.bHasPerInstanceLocalBounds = TransformProvider ? TransformProvider->HasAnimationBounds() : false;
-	Flags.bHasPerInstanceDynamicData = false;
-	Flags.bHasPerInstanceSkinningData = true;
-
-	Flags.bHasPerInstanceLMSMUVBias = false;//IsStaticLightingAllowed();
-
-	ComponentDesc.Flags = Flags;
-	
-	// TODO: rename
-	ComponentDesc.MeshBounds = GetSkinnedAsset()->GetBounds();
-	ComponentDesc.NumCustomDataFloats = NumCustomDataFloats;
-	ComponentDesc.NumInstances = InstanceData.Num();
-
-	ComponentDesc.PrimitiveLocalToWorld = GetRenderMatrix();
-	ComponentDesc.ComponentMobility = Mobility;
-
-	ComponentDesc.BuildChangeSet = [&, MeshBounds = ComponentDesc.MeshBounds](FInstanceUpdateChangeSet& ChangeSet)
-	{
-		// publish data
-		ChangeSet.GetTransformWriter().Gather([&](int32 InstanceIndex) -> FRenderTransform { return FRenderTransform(InstanceData[InstanceIndex].Transform.ToMatrixWithScale()); });
-		ChangeSet.GetCustomDataWriter().Gather(MakeArrayView(InstanceCustomData), NumCustomDataFloats);
-
-		ChangeSet.GetSkinningDataWriter().Gather(
-			[&](int32 InstanceIndex) -> uint32 
-			{
-				if (bValidTransformProvider)
-				{
-					check(InstanceData.IsValidIndex(InstanceIndex));
-					const FSkinnedMeshInstanceData& Instance = InstanceData[InstanceIndex];
-					return TransformProvider->GetSkinningDataOffset(InstanceIndex, ComponentTransform, Instance);
-				}
-
-				return 0u;
-			});
-
-		ChangeSet.GetLocalBoundsWriter().Gather(
-			[&](int32 InstanceIndex) -> FRenderBounds
-			{
-				const uint32 AnimationIndex = InstanceData[InstanceIndex].AnimationIndex;
-				if (bValidTransformProvider && Flags.bHasPerInstanceLocalBounds)
-				{
-					FRenderBounds AnimationBounds;
-					if (TransformProvider->GetAnimationBounds(AnimationIndex, AnimationBounds))
-					{
-						return AnimationBounds;
-					}
-				}
-				
-				return MeshBounds;
-			});
-
-#if WITH_EDITOR
-	if (ChangeSet.Flags.bHasPerInstanceEditorData)
-	{
-		// TODO: the way hit proxies are managed seems daft, why don't we just add them when needed and store them in an array alongside the instances?
-		//       this will always force us to update all the hit proxy data for every instances.
-		TArray<TRefCountPtr<HHitProxy>> HitProxies;
-		CreateHitProxyData(HitProxies);
-		ChangeSet.SetEditorData(HitProxies, SelectedInstances);
-	}
-#endif
-	};
-
-	return ComponentDesc;
+	return FInstancedSkinnedMeshComponentHelper::GetComponentDesc(*this, InShaderPlatform);
 }
 
 void UInstancedSkinnedMeshComponent::SendRenderInstanceData_Concurrent()
diff --git a/Engine/Source/Runtime/Engine/Public/InstanceData/InstanceUpdateChangeSet.h b/Engine/Source/Runtime/Engine/Public/InstanceData/InstanceUpdateChangeSet.h
--- a/Engine/Source/Runtime/Engine/Public/InstanceData/InstanceUpdateChangeSet.h
+++ b/Engine/Source/Runtime/Engine/Public/InstanceData/InstanceUpdateChangeSet.h
@@ -145,6 +145,42 @@ struct TDeltaSetup
 				}
 			}
 		}
+
+		void Gather(const FElement& SingleSourceElement)
+		{
+			if (!Setup.bIsEnabled)
+			{
+				return;
+			}
+
+			if (Setup.Delta.IsEmpty())
+			{
+				Setup.DeltaDataArray.Reset();
+			}
+			else 
+			{
+				const int32 TotalNumItems = Setup.Delta.GetNumItems() * Setup.ElementStride;
+				bool bIsFull = !Setup.Delta.IsDelta();
+
+				if (bIsFull)
+				{
+					// full update, bulk initialization
+					Setup.DeltaDataArray.Init(SingleSourceElement, TotalNumItems);
+				}
+				else
+				{
+					Setup.DeltaDataArray.Reset();
+					for (auto It = Setup.Delta.GetIterator(); It; ++It)
+					{
+						check(Setup.DeltaDataArray.Num() + Setup.ElementStride <= TotalNumItems);
+						for (int32 i = 0; i < Setup.ElementStride; ++i)
+						{
+							Setup.DeltaDataArray.Add(SingleSourceElement);
+						}
+					}
+				}
+			}
+		}
 	};
 
 	/**
diff --git a/Engine/Source/Runtime/Engine/Public/InstancedSkinnedMeshComponentHelper.h b/Engine/Source/Runtime/Engine/Public/InstancedSkinnedMeshComponentHelper.h
--- a/Engine/Source/Runtime/Engine/Public/InstancedSkinnedMeshComponentHelper.h
+++ b/Engine/Source/Runtime/Engine/Public/InstancedSkinnedMeshComponentHelper.h
@@ -45,10 +45,8 @@ FInstanceDataManagerSourceDataDesc FInstancedSkinnedMeshComponentHelper::GetComp
 	}
 #endif
 
-	USkinnedAsset* SkinnedAsset = InComponent.GetSkinnedAsset();
-
+	const USkinnedAsset* SkinnedAsset = InComponent.GetSkinnedAsset();
 	const UTransformProviderData* TransformProvider = InComponent.GetTransformProvider();
-
 	const bool bForceRefPose = UInstancedSkinnedMeshComponent::ShouldForceRefPose();
 	const bool bValidTransformProvider = !bForceRefPose && TransformProvider != nullptr && TransformProvider->IsEnabled();
 
@@ -57,7 +55,6 @@ FInstanceDataManagerSourceDataDesc FInstancedSkinnedMeshComponentHelper::GetComp
 	Flags.bHasPerInstanceDynamicData = false;
 	Flags.bHasPerInstanceSkinningData = true;
 
-
 	Flags.bHasPerInstanceLMSMUVBias = false;//IsStaticLightingAllowed();
 
 	ComponentDesc.Flags = Flags;
@@ -72,40 +69,52 @@ FInstanceDataManagerSourceDataDesc FInstancedSkinnedMeshComponentHelper::GetComp
 
 	const FTransform& ComponentTransform = InComponent.GetComponentTransform();
 
-	ComponentDesc.BuildChangeSet = [&, MeshBounds = ComponentDesc.MeshBounds](FInstanceUpdateChangeSet& ChangeSet)
+	ComponentDesc.BuildChangeSet =
+	[
+		&InComponent,
+		ComponentTransform,
+		TransformProvider = bValidTransformProvider ? TransformProvider : nullptr,
+		MeshBounds = ComponentDesc.MeshBounds,
+		bHasPerInstanceLocalBounds = Flags.bHasPerInstanceLocalBounds
+	] (FInstanceUpdateChangeSet& ChangeSet)
 	{
 		// publish data
-		ChangeSet.GetTransformWriter().Gather([&](int32 InstanceIndex) -> FRenderTransform { return FRenderTransform(InComponent.InstanceData[InstanceIndex].Transform.ToMatrixWithScale()); });
+		ChangeSet.GetTransformWriter().Gather([&InComponent](int32 InstanceIndex) -> FRenderTransform { return FRenderTransform(InComponent.InstanceData[InstanceIndex].Transform.ToMatrixWithScale()); });
 		ChangeSet.GetCustomDataWriter().Gather(MakeArrayView(InComponent.InstanceCustomData), InComponent.NumCustomDataFloats);
 
-		ChangeSet.GetSkinningDataWriter().Gather(
-			[&](int32 InstanceIndex) -> uint32
-			{
-				if (bValidTransformProvider)
+		if (TransformProvider)
+		{
+			ChangeSet.GetSkinningDataWriter().Gather(
+				[&InComponent, TransformProvider, ComponentTransform](int32 InstanceIndex)->uint32
 				{
 					check(InComponent.InstanceData.IsValidIndex(InstanceIndex));
 					const FSkinnedMeshInstanceData& Instance = InComponent.InstanceData[InstanceIndex];
 					return TransformProvider->GetSkinningDataOffset(InstanceIndex, ComponentTransform, Instance);
-				}
-
-				return 0u;
-			});
+				});
+		}
+		else
+		{
+			ChangeSet.GetSkinningDataWriter().Gather(0u);
+		}
 
-		ChangeSet.GetLocalBoundsWriter().Gather(
-			[&](int32 InstanceIndex) -> FRenderBounds
-			{
-				const uint32 AnimationIndex = InComponent.InstanceData[InstanceIndex].AnimationIndex;
-				if (bValidTransformProvider && Flags.bHasPerInstanceLocalBounds)
+		if (TransformProvider && bHasPerInstanceLocalBounds)
+		{
+			ChangeSet.GetLocalBoundsWriter().Gather(
+				[&InComponent, TransformProvider, MeshBounds](int32 InstanceIndex) -> FRenderBounds
 				{
+					const uint32 AnimationIndex = InComponent.InstanceData[InstanceIndex].AnimationIndex;
 					FRenderBounds AnimationBounds;
 					if (TransformProvider->GetAnimationBounds(AnimationIndex, AnimationBounds))
 					{
 						return AnimationBounds;
 					}
-				}
-
-				return MeshBounds;
-			});
+					return MeshBounds;
+				});
+		}
+		else
+		{
+			ChangeSet.GetLocalBoundsWriter().Gather(MeshBounds);
+		}
 
 #if WITH_EDITOR
 		if constexpr (bSupportHitProxies)
