SplineComponent: Broadcast OnSplineChanged after undo/redo

diff --git a/Engine/Source/Runtime/Engine/Classes/Components/SplineComponent.h b/Engine/Source/Runtime/Engine/Classes/Components/SplineComponent.h
--- a/Engine/Source/Runtime/Engine/Classes/Components/SplineComponent.h
+++ b/Engine/Source/Runtime/Engine/Classes/Components/SplineComponent.h
@@ -361,6 +361,7 @@ public:
 	ENGINE_API virtual bool GetIgnoreBoundsForEditorFocus() const override;
 #if WITH_EDITOR
 	ENGINE_API virtual void PostEditChangeChainProperty(FPropertyChangedChainEvent& PropertyChangedEvent) override;
+	ENGINE_API virtual void PostTransacted(const FTransactionObjectEvent& TransactionEvent) override;
 #endif
 	//~ End UObject Interface
 
@@ -1004,7 +1005,7 @@ protected:
 	/** Returns the length of the specified spline segment up to the parametric value given */
 	ENGINE_API float GetSegmentLength(const int32 Index, const float Param = 1.0f) const;
 
-	/** Delegate broadcast any time the spline is mutated. */
+	/** Delegate broadcast any time the spline has potentially been mutated. */
 	FSimpleMulticastDelegate OnSplineChanged;
 
 	/** Delegate broadcast any time UpdateSpline is called (regardless of whether or not the spline has changed). */
diff --git a/Engine/Source/Runtime/Engine/Private/Components/SplineComponent.cpp b/Engine/Source/Runtime/Engine/Private/Components/SplineComponent.cpp
--- a/Engine/Source/Runtime/Engine/Private/Components/SplineComponent.cpp
+++ b/Engine/Source/Runtime/Engine/Private/Components/SplineComponent.cpp
@@ -24,6 +24,7 @@
 #include "Materials/Material.h"
 #include "DynamicMeshBuilder.h"
 #include "StaticMeshResources.h"
+#include "Misc/TransactionObjectEvent.h"
 #include "UObject/ICookInfo.h"
 
 
@@ -3291,6 +3292,17 @@ void USplineComponent::PostEditChangeChainProperty(FPropertyChangedChainEvent& P
 
 	Super::PostEditChangeChainProperty(PropertyChangedEvent);
 }
+
+void USplineComponent::PostTransacted(const FTransactionObjectEvent& TransactionEvent)
+{
+	Super::PostTransacted(TransactionEvent);
+
+	// Notify listeners that the spline data may now be invalid.
+	if (TransactionEvent.GetEventType() == ETransactionObjectEventType::UndoRedo)
+	{
+		OnSplineChanged.Broadcast();
+	}
+}
 #endif
 
 void FSplinePositionLinearApproximation::Build(const FSplineCurves& InCurves, TArray<FSplinePositionLinearApproximation>& OutPoints, float InDensity)
