Temporarily fixing android PSO compile crash.

diff --git a/Engine/Source/Runtime/Engine/Private/Materials/MaterialShader.cpp b/Engine/Source/Runtime/Engine/Private/Materials/MaterialShader.cpp
--- a/Engine/Source/Runtime/Engine/Private/Materials/MaterialShader.cpp
+++ b/Engine/Source/Runtime/Engine/Private/Materials/MaterialShader.cpp
@@ -3439,6 +3439,23 @@ FPSOPrecacheDataArray FMaterialShaderMap::CollectPSOPrecacheData(const FMaterial
 #endif // PSO_PRECACHING_VALIDATE
 
 				PSOCollector->CollectPSOInitializers(SceneTexturesConfig, *PrecacheParams.Material, PrecacheParams.VertexFactoryData, PrecacheParams.PrecachePSOParams, PSOInitializers);
+
+#if PLATFORM_ANDROID
+				// TODO: AJB temporarily un-set subpass index if no subpass hint was specified.
+				for (FPSOPrecacheData& PrecacheData : PSOInitializers)
+				{
+					switch (PrecacheData.Type)
+					{
+						case FPSOPrecacheData::EType::Graphics:
+						{
+							if (PrecacheData.GraphicsPSOInitializer.SubpassHint == ESubpassHint::None)
+							{
+								PrecacheData.GraphicsPSOInitializer.SubpassIndex = 0;
+							}
+						}
+					}
+				}
+#endif
 				delete PSOCollector;
 			}
 		}
diff --git a/Engine/Source/Runtime/VulkanRHI/Private/VulkanPipeline.cpp b/Engine/Source/Runtime/VulkanRHI/Private/VulkanPipeline.cpp
--- a/Engine/Source/Runtime/VulkanRHI/Private/VulkanPipeline.cpp
+++ b/Engine/Source/Runtime/VulkanRHI/Private/VulkanPipeline.cpp
@@ -2293,7 +2293,9 @@ void FVulkanPipelineStateCacheManager::CreateGfxEntry(const FGraphicsPipelineSta
 		DescriptorSetLayoutInfo.FinalizeBindings<false>(*Device, UBGatherInfo, ImmutableSamplers, (NumBindlessShaders != 0));
 	}
 
-	OutGfxEntry->SubpassIndex = PSOInitializer.SubpassIndex;
+
+	// TODO: AJB - temporarily un-set subpass index if no subpass hint was specified.
+	OutGfxEntry->SubpassIndex = PSOInitializer.SubpassHint == ESubpassHint::None ? 0 : PSOInitializer.SubpassIndex;
 
 	FVulkanBlendState* BlendState = ResourceCast(PSOInitializer.BlendState);
 
