SpringMath API Updates for 5.7 release
* Make spring math API experimental (it was added in 5.7 cycle but not marked experimental)
* Complete documentation for SpringMath
* Harmonize naming in spring math
* Add experimental blueprint library for spring math to expose functions to BP
* Rename SpringDamperQuat to ExponentialSmoothingApproxQuat to mirror the existing function FMath::ExponentialSmoothingApprox

diff --git a/Engine/Plugins/Animation/AnimationWarping/Source/Runtime/Private/BoneControllers/AnimNode_Steering.cpp b/Engine/Plugins/Animation/AnimationWarping/Source/Runtime/Private/BoneControllers/AnimNode_Steering.cpp
--- a/Engine/Plugins/Animation/AnimationWarping/Source/Runtime/Private/BoneControllers/AnimNode_Steering.cpp
+++ b/Engine/Plugins/Animation/AnimationWarping/Source/Runtime/Private/BoneControllers/AnimNode_Steering.cpp
@@ -148,7 +148,7 @@ void FAnimNode_Steering::EvaluateSkeletalControl_AnyThread(FComponentSpacePoseCo
 						{
 							// Apply linear correction
 							FQuat LinearCorrection = FQuat::Identity;
-							SpringMath::SimpleSpringDamperQuat(IN OUT LinearCorrection, IN OUT AngularVelocity, DeltaToTargetOrientation, SpringMath::SpringHalfLifeToSmoothingTime(ProceduralTargetTime), DeltaSeconds);
+							SpringMath::SimpleSpringDamperQuat(IN OUT LinearCorrection, IN OUT AngularVelocity, DeltaToTargetOrientation, SpringMath::HalfLifeToSmoothingTime(ProceduralTargetTime), DeltaSeconds);
 
 							UE_VLOG_ARROW(Output.AnimInstanceProxy->GetAnimInstanceObject(), "Steering", Display,
 								RootBoneTransform.GetLocation(),
diff --git a/Engine/Plugins/Experimental/Mover/Source/Mover/Private/DefaultMovementSet/Modes/SmoothWalkingMode.cpp b/Engine/Plugins/Experimental/Mover/Source/Mover/Private/DefaultMovementSet/Modes/SmoothWalkingMode.cpp
--- a/Engine/Plugins/Experimental/Mover/Source/Mover/Private/DefaultMovementSet/Modes/SmoothWalkingMode.cpp
+++ b/Engine/Plugins/Experimental/Mover/Source/Mover/Private/DefaultMovementSet/Modes/SmoothWalkingMode.cpp
@@ -74,7 +74,7 @@ void USmoothWalkingMode::GenerateWalkMove_Implementation(FMoverTickStartData& St
 				SpringState.IntermediateVelocity,
 				DesiredVelocity.GetSafeNormal() * SpringState.IntermediateVelocity.Length(),
 				DeltaSeconds,
-				SpringMath::SpringStrengthToSmoothingTime(TurningStrength));
+				SpringMath::StrengthToSmoothingTime(TurningStrength));
 		}
 	}
 
@@ -117,7 +117,7 @@ void USmoothWalkingMode::GenerateWalkMove_Implementation(FMoverTickStartData& St
 
 	// Get the lag associated with the current velocity smoothing time and perform the same process to estimate the velocity we need to track to
 	// avoid lagging behind
-	const float LagSeconds = DeltaSeconds + VelocitySmoothingCompensation * SpringMath::HalfLifeToLag(SpringMath::SpringSmoothingTimeToHalfLife(VelocitySmoothingTime));
+	const float LagSeconds = DeltaSeconds + (VelocitySmoothingCompensation * VelocitySmoothingTime);
 
 	FVector TrackVelocity =
 		VelocityDifference.Length() > DesiredAcceleration.Length() * LagSeconds ?
diff --git a/Engine/Plugins/Experimental/UAF/UAFWarping/Source/UAFWarping/Private/SteeringTrait.cpp b/Engine/Plugins/Experimental/UAF/UAFWarping/Source/UAFWarping/Private/SteeringTrait.cpp
--- a/Engine/Plugins/Experimental/UAF/UAFWarping/Source/UAFWarping/Private/SteeringTrait.cpp
+++ b/Engine/Plugins/Experimental/UAF/UAFWarping/Source/UAFWarping/Private/SteeringTrait.cpp
@@ -198,7 +198,7 @@ void FAnimNextSteeringTask::Execute(UE::UAF::FEvaluationVM& VM) const
 						{
 							// Apply linear correction
 							FQuat LinearCorrection = FQuat::Identity;
-							SpringMath::SimpleSpringDamperQuat(LinearCorrection, InstanceData->AngularVelocity, DeltaToTargetOrientation, SpringMath::SpringHalfLifeToSmoothingTime(SharedData->ProceduralTargetTime), InstanceData->DeltaTime);
+							SpringMath::SimpleSpringDamperQuat(LinearCorrection, InstanceData->AngularVelocity, DeltaToTargetOrientation, SpringMath::HalfLifeToSmoothingTime(SharedData->ProceduralTargetTime), InstanceData->DeltaTime);
 
 #if ENABLE_ANIM_DEBUG
 							UE_VLOG_ARROW(InstanceData->HostObject, "Steering", Display,
diff --git a/Engine/Source/Runtime/Engine/Classes/Kismet/BlueprintSpringMathLibrary.h b/Engine/Source/Runtime/Engine/Classes/Kismet/BlueprintSpringMathLibrary.h
new file mode 100644
--- /dev/null
+++ b/Engine/Source/Runtime/Engine/Classes/Kismet/BlueprintSpringMathLibrary.h
@@ -0,0 +1,95 @@
+// Copyright Epic Games, Inc. All Rights Reserved.
+
+#pragma once
+
+#include "BlueprintFunctionLibrary.h"
+
+#include "BlueprintSpringMathLibrary.generated.h"
+
+UCLASS(MinimalAPI, Experimental, meta=(BlueprintThreadSafe))
+class UBlueprintSpringMathLibrary : public UBlueprintFunctionLibrary
+{
+	GENERATED_BODY()
+
+public:
+	// Simple Spring //
+
+	UFUNCTION(BlueprintCallable, meta=(Category="Math|Spring|Experimental"))
+	UE_EXPERIMENTAL(5.7, "SpringMath is experimental")
+	static ENGINE_API void SpringDampFloat(UPARAM(ref) float& InOutX, UPARAM(ref) float& InOutV, const float& TargetX, float DeltaTime, float SmoothingTime = 0.2f);
+
+	UE_EXPERIMENTAL(5.7, "SpringMath is experimental")
+	UFUNCTION(BlueprintCallable, meta=(Category="Math|Spring|Experimental"))
+	static ENGINE_API void SpringDampVector(UPARAM(ref) FVector& InOutX, UPARAM(ref) FVector& InOutV, const FVector& TargetX, float DeltaTime, float SmoothingTime = 0.2f);
+
+	UE_EXPERIMENTAL(5.7, "SpringMath is experimental")
+	UFUNCTION(BlueprintCallable, meta=(Category="Math|Spring|Experimental"))
+	static ENGINE_API void SpringDampVector2D(UPARAM(ref) FVector2D& InOutX, UPARAM(ref) FVector2D& InOutV, const FVector2D& TargetX, float DeltaTime, float SmoothingTime = 0.2f);
+
+	UE_EXPERIMENTAL(5.7, "SpringMath is experimental")
+	UFUNCTION(BlueprintCallable, meta=(Category="Math|Spring|Experimental"))
+	static ENGINE_API void SpringDampQuat(UPARAM(ref) FQuat& InOutRotation, UPARAM(ref) FVector& InOutAngularVelocity, const FQuat& TargetRotation, float DeltaTime, float SmoothingTime = 0.2f);
+
+	UE_EXPERIMENTAL(5.7, "SpringMath is experimental")
+	UFUNCTION(BlueprintCallable, meta=(Category="Math|Spring|Experimental"))
+	static ENGINE_API void SpringDampRotator(UPARAM(ref) FRotator& InOutRotation, UPARAM(ref) FVector& InOutAngularVelocity, const FRotator& TargetRotation, float DeltaTime, float SmoothingTime = 0.2f);
+
+	// Velocity Spring //
+
+	UFUNCTION(BlueprintCallable, meta=(Category="Math|Spring|Experimental"))
+	UE_EXPERIMENTAL(5.7, "SpringMath is experimental")
+	static ENGINE_API void VelocitySpringDampFloat(UPARAM(ref) float& InOutX, UPARAM(ref) float& InOutV, UPARAM(ref) float& InOutVi, float TargetX, float MaxSpeed, float DeltaTime, float SmoothingTime = 0.2f);
+
+	UE_EXPERIMENTAL(5.7, "SpringMath is experimental")
+	UFUNCTION(BlueprintCallable, meta=(Category="Math|Spring|Experimental"))
+	static ENGINE_API void VelocitySpringDampVector(UPARAM(ref) FVector& InOutX, UPARAM(ref) FVector& InOutV, UPARAM(ref) FVector& InOutVi, const FVector& TargetX, float MaxSpeed, float DeltaTime,  float SmoothingTime = 0.2f);
+
+	UE_EXPERIMENTAL(5.7, "SpringMath is experimental")
+	UFUNCTION(BlueprintCallable, meta=(Category="Math|Spring|Experimental"))
+	static ENGINE_API void VelocitySpringDampVector2D(UPARAM(ref) FVector2D& InOutX, UPARAM(ref) FVector2D& InOutV, UPARAM(ref) FVector2D& InOutVi, const FVector2D& TargetX, float MaxSpeed, float DeltaTime,  float SmoothingTime = 0.2f);
+	
+	// Damper, aka ExponentialSmoothingApprox // 
+
+	UE_EXPERIMENTAL(5.7, "SpringMath is experimental")
+	UFUNCTION(BlueprintPure, meta=(Category="Math|Spring|Experimental"))
+	static ENGINE_API float DampFloat(float Value, float Target, float DeltaTime, float SmoothingTime = 0.2f);
+
+	UE_EXPERIMENTAL(5.7, "SpringMath is experimental")
+	UFUNCTION(BlueprintPure, meta=(Category="Math|Spring|Experimental"))
+	static ENGINE_API FVector DampVector(const FVector& Value, const FVector& Target, float DeltaTime, float SmoothingTime = 0.2f);
+
+	UE_EXPERIMENTAL(5.7, "SpringMath is experimental")
+	UFUNCTION(BlueprintPure, meta=(Category="Math|Spring|Experimental"))
+	static ENGINE_API FVector2D DampVector2D(const FVector2D& Value, const FVector2D& Target, float DeltaTime, float SmoothingTime = 0.2f);
+
+	UE_EXPERIMENTAL(5.7, "SpringMath is experimental")
+	UFUNCTION(BlueprintPure, meta=(Category="Math|Spring|Experimental"))
+	static ENGINE_API FQuat DampQuat(const FQuat& Value, const FQuat& Target, float DeltaTime, float SmoothingTime = 0.2f);
+
+	UE_EXPERIMENTAL(5.7, "SpringMath is experimental")
+	UFUNCTION(BlueprintPure, meta=(Category="Math|Spring|Experimental"))
+	static ENGINE_API FRotator DampRotator(const FRotator& Value, const FRotator& Target, float DeltaTime, float SmoothingTime = 0.2f);
+
+	// Spring Character
+	UE_EXPERIMENTAL(5.7, "SpringMath is experimental")
+	UFUNCTION(BlueprintCallable, meta=(Category="Math|Spring|Experimental"))
+	static ENGINE_API void SpringCharacterUpdate(UPARAM(ref) FVector& InOutPosition, UPARAM(ref) FVector& InOutVelocity, UPARAM(ref) FVector& InOutAcceleration, const FVector& TargetVelocity, float DeltaTime, float SmoothingTime = 0.2f);
+
+	// Velocity Spring Character
+	UE_EXPERIMENTAL(5.7, "SpringMath is experimental")
+	UFUNCTION(BlueprintCallable, meta=(Category="Math|Spring|Experimental"))
+	static ENGINE_API void VelocitySpringCharacterUpdate(UPARAM(ref) FVector& InOutPosition, UPARAM(ref) FVector& InOutVelocity, UPARAM(ref) FVector& InOutVelocityIntermediate, UPARAM(ref) FVector& InOutAcceleration, const FVector& TargetVelocity, float DeltaTime, float SmoothingTime = 0.2f, float MaxAcceleration = 1.0f);
+
+	// Conversion Methods
+	UFUNCTION(BlueprintPure, meta=(Category="Math|Spring|Experimental"))
+	static ENGINE_API float ConvertSmoothingTimeToStrength(float SmoothingTime);
+
+	UFUNCTION(BlueprintPure, meta=(Category="Math|Spring|Experimental"))
+	static ENGINE_API float ConvertStrengthToSmoothingTime(float Strength);
+
+	UFUNCTION(BlueprintPure, meta=(Category="Math|Spring|Experimental"))
+	static ENGINE_API float ConvertHalfLifeToSmoothingTime(float HalfLife);
+
+	UFUNCTION(BlueprintPure, meta=(Category="Math|Spring|Experimental"))
+	static ENGINE_API float ConvertSmoothingTimeToHalfLife(float SmoothingTime);
+};
diff --git a/Engine/Source/Runtime/Engine/Private/BlueprintSpringMathLibrary.cpp b/Engine/Source/Runtime/Engine/Private/BlueprintSpringMathLibrary.cpp
new file mode 100644
--- /dev/null
+++ b/Engine/Source/Runtime/Engine/Private/BlueprintSpringMathLibrary.cpp
@@ -0,0 +1,139 @@
+// Copyright Epic Games, Inc. All Rights Reserved.
+
+#include "Kismet/BlueprintSpringMathLibrary.h"
+
+#include "Animation/SpringMath.h"
+
+DEFINE_LOG_CATEGORY_STATIC(LogBlueprintSpingMathLibrary, Log, All);
+
+void UBlueprintSpringMathLibrary::SpringDampVector(FVector& InOutX, FVector& InOutV, const FVector& TargetX, float DeltaTime, float SmoothingTime)
+{
+	SpringMath::SimpleSpringDamper(InOutX, InOutV, TargetX, SmoothingTime, DeltaTime);
+}
+
+void UBlueprintSpringMathLibrary::SpringDampVector2D(FVector2D& InOutX, FVector2D& InOutV, const FVector2D& TargetX, float DeltaTime, float SmoothingTime)
+{
+	SpringMath::SimpleSpringDamper(InOutX, InOutV, TargetX, SmoothingTime, DeltaTime);
+}
+
+void UBlueprintSpringMathLibrary::SpringDampFloat(float& InOutX, float& InOutV, const float& TargetX, float DeltaTime, float SmoothingTime)
+{
+	SpringMath::SimpleSpringDamper(InOutX, InOutV, TargetX, SmoothingTime, DeltaTime);
+}
+
+void UBlueprintSpringMathLibrary::SpringDampQuat(FQuat& InOutRotation, FVector& InOutAngularVelocity, const FQuat& TargetRotation,
+	float DeltaTime, float SmoothingTime)
+{
+	SpringMath::SimpleSpringDamperQuat(InOutRotation, InOutAngularVelocity, TargetRotation, SmoothingTime, DeltaTime);
+}
+
+void UBlueprintSpringMathLibrary::SpringDampRotator(FRotator& InOutRotation, FVector& InOutAngularVelocity, const FRotator& TargetRotation,
+	float DeltaTime, float SmoothingTime)
+{
+	FQuat InOutRotationQuat = InOutRotation.Quaternion();
+	SpringMath::SimpleSpringDamperQuat(InOutRotationQuat, InOutAngularVelocity, TargetRotation.Quaternion(), SmoothingTime, DeltaTime);
+	InOutRotation = InOutRotationQuat.Rotator();
+}
+
+void UBlueprintSpringMathLibrary::VelocitySpringDampFloat(float& InOutX, float& InOutV, float& InOutVi, float TargetX, float MaxSpeed,
+                                                          float DeltaTime, float SmoothingTime)
+{
+	if (MaxSpeed < 0.0f)
+	{
+		UE_LOG(LogBlueprintSpingMathLibrary, Warning, TEXT("UBlueprintSpringMathLibrary::VelocitySpringDampFloat TargetSpeed cannot be negative"));
+		return;
+	}
+
+	SpringMath::VelocitySpringDamperF(InOutX, InOutV, InOutVi, TargetX, MaxSpeed, SmoothingTime, DeltaTime);
+}
+
+void UBlueprintSpringMathLibrary::VelocitySpringDampVector(FVector& InOutX, FVector& InOutV, FVector& InOutVi, const FVector& TargetX, float MaxSpeed,
+	float DeltaTime, float SmoothingTime)
+{
+	if (MaxSpeed < 0.0f)
+	{
+		UE_LOG(LogBlueprintSpingMathLibrary, Warning, TEXT("UBlueprintSpringMathLibrary::VelocitySpringDampVector TargetSpeed cannot be negative"));
+		return;
+	}
+
+	SpringMath::VelocitySpringDamper(InOutX, InOutV, InOutVi, TargetX, MaxSpeed, SmoothingTime, DeltaTime);
+}
+
+void UBlueprintSpringMathLibrary::VelocitySpringDampVector2D(FVector2D& InOutX, FVector2D& InOutV, FVector2D& InOutVi, const FVector2D& TargetX,
+	float MaxSpeed, float DeltaTime, float SmoothingTime)
+{
+	if (MaxSpeed < 0.0f)
+	{
+		UE_LOG(LogBlueprintSpingMathLibrary, Warning, TEXT("UBlueprintSpringMathLibrary::VelocitySpringDampVector2D TargetSpeed cannot be negative"));
+		return;
+	}
+
+	SpringMath::VelocitySpringDamper(InOutX, InOutV, InOutVi, TargetX, MaxSpeed, SmoothingTime, DeltaTime);
+}
+
+float UBlueprintSpringMathLibrary::DampFloat(float Value, float Target, float DeltaTime, float SmoothingTime)
+{
+	float DampedValue = Value;
+	FMath::ExponentialSmoothingApprox(DampedValue, Target, DeltaTime, SmoothingTime);
+	return DampedValue;
+}
+
+FVector UBlueprintSpringMathLibrary::DampVector(const FVector& Value, const FVector& Target, float DeltaTime, float SmoothingTime)
+{
+	FVector DampedValue = Value;
+	FMath::ExponentialSmoothingApprox(DampedValue, Target, DeltaTime, SmoothingTime);
+	return DampedValue;
+}
+
+FVector2D UBlueprintSpringMathLibrary::DampVector2D(const FVector2D& Value, const FVector2D& Target, float DeltaTime, float SmoothingTime)
+{
+	FVector2D DampedValue = Value;
+	FMath::ExponentialSmoothingApprox(DampedValue, Target, DeltaTime, SmoothingTime);
+	return DampedValue;
+}
+
+FQuat UBlueprintSpringMathLibrary::DampQuat(const FQuat& Value, const FQuat& Target, float DeltaTime, float SmoothingTime)
+{
+	FQuat DampedValue = Value;
+	SpringMath::ExponentialSmoothingApproxQuat(DampedValue, Target, DeltaTime, SmoothingTime);
+	return DampedValue;
+}
+
+FRotator UBlueprintSpringMathLibrary::DampRotator(const FRotator& Value, const FRotator& Target, float DeltaTime, float SmoothingTime)
+{
+	FQuat DampedValue = Value.Quaternion();
+	SpringMath::ExponentialSmoothingApproxQuat(DampedValue, Target.Quaternion(), DeltaTime, SmoothingTime);
+	return DampedValue.Rotator();
+}
+
+void UBlueprintSpringMathLibrary::SpringCharacterUpdate(FVector& InOutPosition, FVector& InOutVelocity, FVector& InOutAcceleration,
+                                                        const FVector& TargetVelocity, float DeltaTime, float SmoothingTime)
+{
+	SpringMath::SpringCharacterUpdate(InOutPosition, InOutVelocity, InOutAcceleration, TargetVelocity, SmoothingTime, DeltaTime);
+}
+
+void UBlueprintSpringMathLibrary::VelocitySpringCharacterUpdate(FVector& InOutPosition, FVector& InOutVelocity, FVector& InOutVelocityIntermediate,
+	FVector& InOutAcceleration, const FVector& TargetVelocity, float DeltaTime, float SmoothingTime, float MaxAcceleration)
+{
+	SpringMath::VelocitySpringCharacterUpdate(InOutPosition, InOutVelocity, InOutVelocityIntermediate, InOutAcceleration, TargetVelocity, SmoothingTime, MaxAcceleration, DeltaTime);
+}
+
+float UBlueprintSpringMathLibrary::ConvertSmoothingTimeToStrength(float SmoothingTime)
+{
+	return SpringMath::SmoothingTimeToStrength(SmoothingTime);
+}
+
+float UBlueprintSpringMathLibrary::ConvertStrengthToSmoothingTime(float Strength)
+{
+	return SpringMath::StrengthToSmoothingTime(Strength);
+}
+
+float UBlueprintSpringMathLibrary::ConvertHalfLifeToSmoothingTime(float HalfLife)
+{
+	return SpringMath::HalfLifeToSmoothingTime(HalfLife);
+}
+
+float UBlueprintSpringMathLibrary::ConvertSmoothingTimeToHalfLife(float SmoothingTime)
+{
+	return SpringMath::SmoothingTimeToHalfLife(SmoothingTime);
+}
diff --git a/Engine/Source/Runtime/Engine/Public/Animation/SpringMath.h b/Engine/Source/Runtime/Engine/Public/Animation/SpringMath.h
--- a/Engine/Source/Runtime/Engine/Public/Animation/SpringMath.h
+++ b/Engine/Source/Runtime/Engine/Public/Animation/SpringMath.h
@@ -6,66 +6,87 @@
 #include "Math/Vector.h"
 #include "Math/Quat.h"
 
+// Collection of useful spring methods which can be used for damping, simulating characters etc.
+// Reference https://theorangeduck.com/page/spring-roll-call
+
 namespace SpringMath
 {
-	// Reference https://theorangeduck.com/page/spring-roll-call 
-
-	constexpr float SpringHalfLifeToDamping(float HalfLife, float Eps = 1e-5f)
-	{
-		return (4.0f * UE_LN2) / (HalfLife + Eps);
-	}
 
-	constexpr float SpringDampingToHalfLife(float Damping)
-	{
-		return Damping / (4.0f * UE_LN2);
-	}
-
-	constexpr float SpringSmoothingTimeToHalfLife(float SmoothingTime)
-	{
-		return SmoothingTime / UE_LN2;
-	}
-
-	constexpr float SpringHalfLifeToSmoothingTime(float HalfLife)
-	{
-		return HalfLife * UE_LN2;
-	}
-
-	constexpr float SpringSmoothingTimeToDamping(float SmoothingTime, float Eps = 1e-5f)
+	/** Convert a halflife to a damping ratio
+	 * 
+	 * @param HalfLife The half life of the spring. How long it takes the value to get halfway towards the target.
+	 * @return The damping ratio
+	 */
+	UE_EXPERIMENTAL(5.7, "SpringMath is experimental")
+	constexpr float HalfLifeToDampingRatio(float HalfLife)
 	{
-		return 4.0f / FMath::Max(SmoothingTime, Eps);
+		return (4.0f * UE_LN2) / FMath::Max(HalfLife, UE_KINDA_SMALL_NUMBER);
 	}
 
-	constexpr float SpringSmoothingTimeToStrength(float SmoothingTime, float Eps = 1e-5f)
+	/** Convert a damping ratio to a halflife
+	 * 
+	 * @param DampingRatio
+	 * @return The half life of the spring. How long it takes the value to get halfway towards the target.
+	 */
+	UE_EXPERIMENTAL(5.7, "SpringMath is experimental")
+	constexpr float DampingRatioToHalfLife(float DampingRatio)
 	{
-		return 2.0f / FMath::Max(SmoothingTime, Eps);
+		return DampingRatio / (4.0f * UE_LN2);
 	}
 
-	constexpr float SpringStrengthToSmoothingTime(float Strength, float Eps = 1e-5f)
+	/** Convert a smoothing time to a half life
+	 * 
+	 * @param SmoothingTime The smoothing time of the spring in seconds. It is the time by which the output lags the input when critically damped, when the input is changing at a constant rate. 
+	 * @return The half life of the spring. How long it takes the value to get halfway towards the target. 
+	 */
+	UE_EXPERIMENTAL(5.7, "SpringMath is experimental")
+	constexpr float SmoothingTimeToHalfLife(float SmoothingTime)
 	{
-		return 2.0f / FMath::Max(Strength, Eps);
+		return SmoothingTime * UE_LN2;
 	}
 
-	constexpr float HalfLifeToLag(float HalfLife)
+	/** Convert a halflife to a smoothing time
+	 * 
+	 * @param HalfLife 
+	 * @return 
+	 */
+	UE_EXPERIMENTAL(5.7, "SpringMath is experimental")
+	constexpr float HalfLifeToSmoothingTime(float HalfLife)
 	{
 		return HalfLife / UE_LN2;
 	}
 
-	constexpr float LagToHalfLife(float Lag)
+	/** Convert from smoothing time to a damping ratio.
+	 * 
+	 * @param SmoothingTime 
+	 * @return Damping Ratio
+	 */
+	UE_EXPERIMENTAL(5.7, "SpringMath is experimental")
+	constexpr float SmoothingTimeToDampingRatio(float SmoothingTime)
 	{
-		return Lag * UE_LN2;
+		return 4.0f / FMath::Max(SmoothingTime, UE_KINDA_SMALL_NUMBER);
 	}
 
-	inline FVector QuatToScaledAngleAxis(const FQuat& Q)
+	/** Convert from smoothing time to spring strength.
+	  * 
+	 * @param SmoothingTime The smoothing time of the spring in seconds. It is the time by which the output lags the input when critically damped, when the input is changing at a constant rate.
+	 * @return The spring strength. This corresponds to the undamped frequency of the spring in hz.
+	 */
+	UE_EXPERIMENTAL(5.7, "SpringMath is experimental")
+	constexpr float SmoothingTimeToStrength(float SmoothingTime)
 	{
-		FQuat QLog = Q.Log();
-		return 2.0f * FVector(QLog.X, QLog.Y, QLog.Z);
+		return 2.0f / FMath::Max(SmoothingTime, UE_KINDA_SMALL_NUMBER);
 	}
 
-	inline FQuat QuatFromScaledAngleAxis(FVector V)
+	/** Convert from spring strength to smoothing time.
+	 * 
+	 * @param Strength The spring strength. This corresponds to the undamped frequency of the spring in hz.
+	 * @return The smoothing time of the spring in seconds. It is the time by which the output lags the input when critically damped, when the input is changing at a constant rate.
+	 */
+	UE_EXPERIMENTAL(5.7, "SpringMath is experimental")
+	constexpr float StrengthToSmoothingTime(float Strength)
 	{
-		V /= 2.0f;
-		FQuat QLog(V.X, V.Y, V.Z, 0.0f);
-		return QLog.Exp();
+		return 2.0f / FMath::Max(Strength, UE_KINDA_SMALL_NUMBER);
 	}
 
 	/** Simplified version of FMath::CriticallyDampedSmoothing where v_goal is assumed to be 0
@@ -77,7 +98,8 @@ namespace SpringMath
 	 * @param SmoothingTime The smoothing time to use for the spring. Longer times lead to more damped behaviour. A time of 0 will snap X to XGoal
 	 * @param DeltaTime Timestep in seconds
 	 */
-	template<typename T>
+	template <typename T>
+	UE_EXPERIMENTAL(5.7, "SpringMath is experimental")
 	void SimpleSpringDamper(
 		T& InOutX,
 		T& InOutV,
@@ -85,7 +107,7 @@ namespace SpringMath
 		float SmoothingTime,
 		float DeltaTime)
 	{
-		float Y = SpringSmoothingTimeToDamping(SmoothingTime) / 2.0f;
+		float Y = SmoothingTimeToDampingRatio(SmoothingTime) / 2.0f;
 		T J0 = InOutX - TargetX;
 		T J1 = InOutV + J0 * Y;
 		float EyDt = FMath::InvExpApprox(Y * DeltaTime);
@@ -102,6 +124,7 @@ namespace SpringMath
 	 * @param SmoothingTime he smoothing time to use for the spring. Longer times lead to more damped behaviour. A time of 0 will snap InOutRotation to TargetRotation
 	 * @param DeltaTime 
 	 */
+	UE_EXPERIMENTAL(5.7, "SpringMath is experimental")
 	inline void SimpleSpringDamperQuat(
 		FQuat& InOutRotation,
 		FVector& InOutAngularVelocityRadians,
@@ -109,16 +132,19 @@ namespace SpringMath
 		float SmoothingTime,
 		float DeltaTime)
 	{
-		float Y = SpringSmoothingTimeToDamping(SmoothingTime) / 2.0f;
+		PRAGMA_DISABLE_EXPERIMENTAL_WARNINGS
+		float Y = SmoothingTimeToDampingRatio(SmoothingTime) / 2.0f;
+		PRAGMA_ENABLE_EXPERIMENTAL_WARNINGS
 
 		FQuat Diff = InOutRotation * TargetRotation.Inverse();
 		Diff.EnforceShortestArcWith(FQuat::Identity);
-		FVector J0 = QuatToScaledAngleAxis(Diff);
+		FVector J0 = Diff.ToRotationVector();
 		FVector J1 = InOutAngularVelocityRadians + J0 * Y;
 
 		float EyDt = FMath::InvExpApprox(Y * DeltaTime);
 
-		InOutRotation = QuatFromScaledAngleAxis(EyDt * (J0 + J1 * DeltaTime)) * TargetRotation;
+		
+		InOutRotation = FQuat::MakeFromRotationVector(EyDt * (J0 + J1 * DeltaTime)) * TargetRotation;
 		InOutAngularVelocityRadians = EyDt * (InOutAngularVelocityRadians - J1 * Y * DeltaTime);
 	}
 
@@ -130,92 +156,109 @@ namespace SpringMath
 	 * @param InOutV The velocity of the value to be damped
 	 * @param InOutXi The intermediate target of the value to be damped
 	 * @param TargetX The target value of X to damp towards
-	 * @param TargetV The desired linear velocity of X to achieve while damping towards X
+	 * @param MaxSpeed The desired speed to achieve while damping towards X
 	 * @param SmoothingTime The smoothing time to use while damping towards X. Higher values will give more smoothed behaviour. A value of 0 will give a linear interpolation of X to Target
 	 * @param DeltaTime The timestep in seconds
 	 */
 	template <typename TFloat>
+	UE_EXPERIMENTAL(5.7, "SpringMath is experimental")
 	void VelocitySpringDamperF(
 		TFloat& InOutX,
 		TFloat& InOutV,
 		TFloat& InOutXi,
 		TFloat TargetX,
-		TFloat TargetV,
+		TFloat MaxSpeed,
 		float SmoothingTime,
 		float DeltaTime)
 	{
-		float HalfLife = SpringSmoothingTimeToHalfLife(SmoothingTime);
 		static_assert(std::is_floating_point_v<TFloat>, "TFloat must be floating point");
 
-		TFloat XDiff = ((TargetX - InOutXi) > 0.0f ? 1.0f : -1.0f) * TargetV;
+		MaxSpeed = FMath::Max(MaxSpeed, 0.0f); // MaxSpeed can't be negative
+
+		TFloat XDiff = ((TargetX - InOutXi) > 0.0f
+			? 1.0f
+			: -1.0f) * MaxSpeed;
 
-		float TGoalFuture = HalfLifeToLag(HalfLife);
-		TFloat XGoalFuture = fabs(TargetX - InOutXi) > TGoalFuture * TargetV ? InOutXi + XDiff * TGoalFuture : TargetX;
+		float TGoalFuture = SmoothingTime;
+		TFloat XGoalFuture = fabs(TargetX - InOutXi) > TGoalFuture * MaxSpeed
+			? InOutXi + XDiff * TGoalFuture
+			: TargetX;
 
-		SimpleSpringDamper(InOutX, InOutV, XGoalFuture, HalfLife, DeltaTime);
+		SimpleSpringDamper(InOutX, InOutV, XGoalFuture, SmoothingTime, DeltaTime);
 
-		InOutXi = fabs(TargetX - InOutXi) > DeltaTime * TargetV ? InOutXi + XDiff * DeltaTime : TargetX;
+		InOutXi = FMath::Abs(TargetX - InOutXi) > DeltaTime * MaxSpeed
+			? InOutXi + XDiff * DeltaTime
+			: TargetX;
 	}
 
 	/** A velocity spring will damp towards a target that follows a fixed linear target velocity, allowing control of the interpolation speed
 	* while still giving a smoothed beheaviour. A SmoothingTime of 0 will give a linear interpolation between X and TargetX
 	 * 
-	 * @tparam TVector 
+	 * @tparam TVector The type of vector to use (e.g. FVector2D, FVector4, FVector etc)
 	 * @param InOutX The value to be damped
 	 * @param InOutV The velocity of the value to be damped
 	 * @param InOutXi The intermediate target of the value to be damped
 	 * @param TargetX The target value of X to damp towards
-	 * @param MaxV The max velocity to use for the intermediate target interpolation
+	 * @param MaxSpeed The max velocity to use for the intermediate target interpolation
 	 * @param SmoothingTime The smoothing time to use while damping towards X. Higher values will give more smoothed behaviour. A value of 0 will give a linear interpolation of X to Target
 	 * @param DeltaTime The timestep in seconds
 	 */
 	template <typename TVector>
+	UE_EXPERIMENTAL(5.7, "SpringMath is experimental")
 	void VelocitySpringDamper(
 		TVector& InOutX,
 		TVector& InOutV,
 		TVector& InOutXi,
 		TVector TargetX,
-		float MaxV,
+		float MaxSpeed,
 		float SmoothingTime,
 		float DeltaTime)
 	{
-		float HalfLife = SpringSmoothingTimeToHalfLife(SmoothingTime);
 		TVector XDiff = TargetX - InOutXi;
 		float XDiffLength = XDiff.Length();
-		TVector XDiffDir = XDiffLength > FLT_EPSILON ? XDiff / XDiffLength : TVector::ZeroVector;
+		TVector XDiffDir = XDiffLength > FLT_EPSILON
+			? XDiff / XDiffLength
+			: TVector::ZeroVector;
 
-		float TGoalFuture = HalfLifeToLag(HalfLife);
-		TVector XGoalFuture = XDiffLength > TGoalFuture * MaxV ? InOutXi + (XDiffDir * MaxV) * TGoalFuture : TargetX;
+		float TGoalFuture = SmoothingTime;
+		TVector XGoalFuture = XDiffLength > TGoalFuture * MaxSpeed
+			? InOutXi + (XDiffDir * MaxSpeed) * TGoalFuture
+			: TargetX;
 
-		SimpleSpringDamper(InOutX, InOutV, XGoalFuture, HalfLife, DeltaTime);
+		SimpleSpringDamper(InOutX, InOutV, XGoalFuture, SmoothingTime, DeltaTime);
 
-		InOutXi = XDiffLength > DeltaTime * MaxV ? InOutXi + XDiffDir * MaxV * DeltaTime : TargetX;
+		InOutXi = XDiffLength > DeltaTime * MaxSpeed
+			? InOutXi + XDiffDir * MaxSpeed * DeltaTime
+			: TargetX;
 	}
 
-	/** Update a position representing a character given a target velocity using a simple damped spring
+	/** Update the position of a character given a target velocity using a simple damped spring
 	 * 
-	 * @param InOutPosition 
-	 * @param InOutVelocity 
-	 * @param InOutAcceleration 
-	 * @param TargetVelocity 
-	 * @param SmoothingTime 
-	 * @param DeltaTime 
-	 * @param VDeadzone 
-	 * @param ADeadzone 
+	 * @tparam TVector The type of vector to use (e.g. FVector2D, FVector4, FVector etc)
+	 * @param InOutPosition The position of the character
+	 * @param InOutVelocity The velocity of the character. Needs to be stored and persisted by the caller. Usually initialized to zero and not modified by the caller.
+	 * @param InOutAcceleration The acceleration of the character. Needs to be stored and persisted by the caller. Usually initialized to zero and not modified by the caller.
+	 * @param TargetVelocity The target velocity of the character.
+	 * @param SmoothingTime The time over which to smooth velocity. It takes roughly the smoothing time in order for the character to reach the target velocity.
+	 * @param DeltaTime The delta time to tick the character
+	 * @param VDeadzone Deadzone for velocity. Current velocity will snap to target velocity when within the deadzone
+	 * @param ADeadzone Deadzone for acceleration. Acceleration will snap to zero when within the deadzone
 	 */
-	inline void SpringCharacterUpdate(
-		FVector& InOutPosition,
-		FVector& InOutVelocity,
-		FVector& InOutAcceleration,
-		const FVector& TargetVelocity,
+	template <typename TVector>
+	UE_EXPERIMENTAL(5.7, "SpringMath is experimental")
+	void SpringCharacterUpdate(
+		TVector& InOutPosition,
+		TVector& InOutVelocity,
+		TVector& InOutAcceleration,
+		const TVector& TargetVelocity,
 		float SmoothingTime,
 		float DeltaTime,
 		float VDeadzone = 1e-2,
 		float ADeadzone = 1e-4)
 	{
-		float Y = SpringSmoothingTimeToDamping(SmoothingTime) / 2.0f;
-		FVector J0 = InOutVelocity - TargetVelocity;
-		FVector J1 = InOutAcceleration + J0 * Y;
+		float Y = SmoothingTimeToDampingRatio(SmoothingTime) / 2.0f;
+		TVector J0 = InOutVelocity - TargetVelocity;
+		TVector J1 = InOutAcceleration + J0 * Y;
 		float EyDt = FMath::InvExpApprox(Y * DeltaTime);
 
 		InOutPosition = EyDt * (((-J1) / (Y * Y)) + ((-J0 - J1 * DeltaTime) / Y)) +
@@ -230,140 +273,249 @@ namespace SpringMath
 
 			if (InOutAcceleration.SquaredLength() < FMath::Square(ADeadzone))
 			{
-				InOutAcceleration = FVector::ZeroVector;
+				InOutAcceleration = TVector::ZeroVector;
 			}
 		}
 	}
 
-	/** Gives predicted positions, velocities and accelerations for SpringCharacterUpdate
-	 * 
-	 * @param OutPredictedPositions 
-	 * @param OutPredictedVelocities 
-	 * @param OutPredictedAccelerations 
-	 * @param PredictCount 
-	 * @param CurrentPosition 
-	 * @param CurrentVelocity 
-	 * @param CurrentAcceleration 
-	 * @param VGoal 
-	 * @param SmoothingTime 
-	 * @param SecondsPerPredictionStep 
-	 * @param VDeadzone 
+	/** Gives predicted positions, velocities and accelerations for SpringCharacterUpdate. Useful for generating a predicted trajectory given known initial start conditions.
+	 *
+	 * @tparam TVector The type of vector to use (e.g. FVector2D, FVector4, FVector etc)
+	 * @param OutPredictedPositions ArrayView of output buffer to put the predicted positions. ArrayView should be the same size as PredictCount
+	 * @param OutPredictedVelocities ArrayView of output buffer to put the predicted velocities. ArrayView should be the same size as PredictCount
+	 * @param OutPredictedAccelerations ArrayView of output buffer to put the predicted accelerations. ArrayView should be the same size as PredictCount
+	 * @param PredictCount How many points to predict. Must be greater than 0
+	 * @param CurrentPosition The initial position of the character
+	 * @param CurrentVelocity The initial velocity of the character
+	 * @param CurrentAcceleration The initial acceleration of the character
+	 * @param TargetVelocity The target velocity of the character
+	 * @param SmoothingTime The smoothing time of the character. It takes roughly the smoothing time in order for the character to reach the target velocity.
+	 * @param SecondsPerPredictionStep How much time in between each prediction step.
+	 * @param VDeadzone Deadzone for velocity. Current velocity will snap to target velocity when within the deadzone
+	 * @param ADeadzone Deadzone for acceleration. Acceleration will snap to zero when within the deadzone
 	 */
-	inline void SpringCharacterPredict(
-		TArrayView<FVector> OutPredictedPositions,
-		TArrayView<FVector> OutPredictedVelocities,
-		TArrayView<FVector> OutPredictedAccelerations,
-		int PredictCount,
-		const FVector& CurrentPosition,
-		const FVector& CurrentVelocity,
-		const FVector& CurrentAcceleration,
-		const FVector& VGoal,
+	template <typename TVector>
+	UE_EXPERIMENTAL(5.7, "SpringMath is experimental")
+	void SpringCharacterPredict(
+		TArrayView<TVector> OutPredictedPositions,
+		TArrayView<TVector> OutPredictedVelocities,
+		TArrayView<TVector> OutPredictedAccelerations,
+		const TVector& CurrentPosition,
+		const TVector& CurrentVelocity,
+		const TVector& CurrentAcceleration,
+		const TVector& TargetVelocity,
 		float SmoothingTime,
 		float SecondsPerPredictionStep,
-		float VDeadzone)
+		float VDeadzone = 1e-2,
+		float ADeadzone = 1e-4)
 	{
-		for (int i = 0; i < PredictCount; i++)
+		int32 PredictCount = OutPredictedPositions.Num();
+		check(PredictCount > 0);
+		check(OutPredictedVelocities.Num() == PredictCount);
+		check(OutPredictedAccelerations.Num() == PredictCount);
+		
+		for (int32 i = 0; i < PredictCount; i++)
 		{
 			OutPredictedPositions[i] = CurrentPosition;
 			OutPredictedVelocities[i] = CurrentVelocity;
 			OutPredictedAccelerations[i] = CurrentAcceleration;
 		}
 
-		for (int i = 0; i < PredictCount; i++)
+		for (int32 i = 0; i < PredictCount; i++)
 		{
 			const float PredictTime = (i + 1) * SecondsPerPredictionStep; // Note i+1 since we want index 0 to be the first prediction step
-			SpringCharacterUpdate(OutPredictedPositions[i], OutPredictedVelocities[i], OutPredictedAccelerations[i], VGoal, SmoothingTime,
-				PredictTime, VDeadzone);
+			PRAGMA_DISABLE_EXPERIMENTAL_WARNINGS
+			SpringCharacterUpdate(OutPredictedPositions[i], OutPredictedVelocities[i], OutPredictedAccelerations[i], TargetVelocity, SmoothingTime,
+				PredictTime, VDeadzone, ADeadzone);
+			PRAGMA_ENABLE_EXPERIMENTAL_WARNINGS
 		}
 	}
 
-	/** Update a position representing a character given a target velocity using a velocity spring
+	/** Update a position representing a character given a target velocity using a velocity spring.
+	 * A velocity spring tracks an intermediate velocity which moves at a maximum acceleration linearly towards a target.
+	 * This means unlike the "SpringCharacterUpdate", it will take longer to reach a target velocity that is further away from the current velocity.
 	 * 
-	 * @tparam TVector 
-	 * @param InOutPosition 
-	 * @param InOutVelocity 
-	 * @param InOutVelocityIntermediate 
-	 * @param InOutAccel 
-	 * @param TargetVelocity 
-	 * @param SmoothingTime 
-	 * @param MacAccel 
-	 * @param DeltaTime 
+	 * @tparam TVector The type of vector to use (e.g. FVector2D, FVector4, FVector etc)
+	 * @param InOutPosition The position of the character
+	 * @param InOutVelocity The velocity of the character. Needs to be stored and persisted by the caller. Usually initialized to zero and not modified by the caller.
+	 * @param InOutVelocityIntermediate The intermediate velocity of the character. Needs to be stored and persisted by the caller. Usually initialized to zero and not modified by the caller.
+	 * @param InOutAcceleration The acceleration of the character. Needs to be stored and persisted by the caller. Usually initialized to zero and not modified by the caller.
+	 * @param TargetVelocity The target velocity of the character.
+	 * @param SmoothingTime The time over which to smooth velocity. It takes roughly the smoothing time in order for the character to reach the target velocity.
+	 * @param MaxAcceleration Puts a limit on the maximum acceleration that the intermediate velocity can do each frame. If MaxAccel is very large, the behaviour wil lbe the same as SpringCharacterUpdate
+	 * @param DeltaTime The delta time to tick the character
+	 * @param VDeadzone Deadzone for velocity. Current velocity will snap to target velocity when within the deadzone
+	 * @param ADeadzone Deadzone for acceleration. Acceleration will snap to zero when within the deadzone
 	 */
 	template <typename TVector>
+	UE_EXPERIMENTAL(5.7, "SpringMath is experimental")
 	void VelocitySpringCharacterUpdate(
 		TVector& InOutPosition,
 		TVector& InOutVelocity,
 		TVector& InOutVelocityIntermediate,
-		TVector& InOutAccel,
+		TVector& InOutAcceleration,
 		TVector TargetVelocity,
 		float SmoothingTime,
-		float MacAccel,
-		float DeltaTime)
+		float MaxAcceleration,
+		float DeltaTime,
+		float VDeadzone = 1e-2,
+		float ADeadzone = 1e-4)
 	{
-		float HalfLife = SpringSmoothingTimeToHalfLife(SmoothingTime);
 		TVector VDiff = TargetVelocity - InOutVelocityIntermediate;
 		float VDiffLength = VDiff.Length();
-		TVector VDiffDir = VDiffLength > 0.0001f ? VDiff / VDiffLength : TVector::ZeroVector;
+		TVector VDiffDir = VDiffLength > 0.0001f
+			? VDiff / VDiffLength
+			: TVector::ZeroVector;
 
-		float TGoalFuture = HalfLifeToLag(HalfLife);
-		TVector MaxVFuture = VDiffLength > TGoalFuture * MacAccel ? InOutVelocityIntermediate + (VDiffDir * MacAccel) * TGoalFuture : TargetVelocity;
+		float TGoalFuture = SmoothingTime;
+		TVector MaxVFuture = VDiffLength > TGoalFuture * MaxAcceleration
+			? InOutVelocityIntermediate + (VDiffDir * MaxAcceleration) * TGoalFuture
+			: TargetVelocity;
 
-		float Y = SpringHalfLifeToDamping(HalfLife) / 2.0f;
+		float Y = SmoothingTimeToDampingRatio(SmoothingTime) / 2.0f;
 		TVector J0 = InOutVelocity - MaxVFuture;
-		TVector J1 = InOutAccel + J0 * Y;
+		TVector J1 = InOutAcceleration + J0 * Y;
 		float EyDt = FMath::InvExpApprox(Y * DeltaTime);
 
 		InOutPosition = EyDt * (((-J1) / (Y * Y)) + ((-J0 - J1 * DeltaTime) / Y)) +
 			(J1 / (Y * Y)) + J0 / Y + MaxVFuture * DeltaTime + InOutPosition;
 		InOutVelocity = EyDt * (J0 + J1 * DeltaTime) + MaxVFuture;
-		InOutAccel = EyDt * (InOutAccel - J1 * Y * DeltaTime);
-		InOutVelocityIntermediate = VDiffLength > DeltaTime * MacAccel ? InOutVelocityIntermediate + VDiffDir * MacAccel * DeltaTime : TargetVelocity;
-	}
+		InOutAcceleration = EyDt * (InOutAcceleration - J1 * Y * DeltaTime);
+		InOutVelocityIntermediate = VDiffLength > DeltaTime * MaxAcceleration
+			? InOutVelocityIntermediate + VDiffDir * MaxAcceleration * DeltaTime
+			: TargetVelocity;
 
+		if ((TargetVelocity - InOutVelocity).SquaredLength() < FMath::Square(VDeadzone))
+		{
+			// We reached our target
+			InOutVelocity = TargetVelocity;
 
+			if (InOutAcceleration.SquaredLength() < FMath::Square(ADeadzone))
+			{
+				InOutAcceleration = TVector::ZeroVector;
+			}
+		}
+	}
 
-	/** Specialized quaternion damper, similar to FMath::ExponentialSmoothingApprox but for quats
-	 * 
-	 * @param CurrentRotation 
-	 * @param TargetRotation 
-	 * @param SmoothingTime 
-	 * @param DeltaTime 
-	 * @param Eps 
-	 * @return 
+	/** Gives predicted positions, velocities and accelerations for SpringCharacterUpdate. Useful for generating a predicted trajectory given known initial start conditions.
+	 *
+	 * @tparam TVector The type of vector to use (e.g. FVector2D, FVector4, FVector etc)
+	 * @param OutPredictedPositions ArrayView of output buffer to put the predicted positions. ArrayView should be the same size as PredictCount
+	 * @param OutPredictedVelocities ArrayView of output buffer to put the predicted velocities. ArrayView should be the same size as PredictCount
+	 * @param OutPredictedIntermediateVelocities ArrayView of output buffer to put the predicted intermediate velocities. ArrayView should be the same size as PredictCount
+	 * @param OutPredictedAccelerations ArrayView of output buffer to put the predicted accelerations. ArrayView should be the same size as PredictCount
+	 * @param PredictCount How many points to predict. Must be greater than 0
+	 * @param CurrentPosition The initial position of the character
+	 * @param CurrentVelocity The initial velocity of the character
+	 * @param CurrentIntermediateVelocity The initial intermediate velocity of the character
+	 * @param CurrentAcceleration The initial acceleration of the character
+	 * @param TargetVelocity The target velocity of the character
+	 * @param SmoothingTime The smoothing time of the character. It takes roughly the smoothing time in order for the character to reach the target velocity.
+	 * @param MaxAcceleration Puts a limit on the maximum acceleration that the intermediate velocity can do each frame. If MaxAccel is very large, the behaviour wil lbe the same as SpringCharacterUpdate
+	 * @param SecondsPerPredictionStep How much time in between each prediction step.
+	 * @param VDeadzone Deadzone for velocity. Current velocity will snap to target velocity when within the deadzone
+	 * @param ADeadzone Deadzone for acceleration. Acceleration will snap to zero when within the deadzone
 	 */
-	FORCEINLINE FQuat SpringDamperQuat(const FQuat& CurrentRotation, const FQuat& TargetRotation, float SmoothingTime, float DeltaTime, float Eps = 1e-5f)
+	template <typename TVector>
+	UE_EXPERIMENTAL(5.7, "SpringMath is experimental")
+	void VelocitySpringCharacterPredict(
+		TArrayView<TVector> OutPredictedPositions,
+		TArrayView<TVector> OutPredictedVelocities,
+		TArrayView<TVector> OutPredictedIntermediateVelocities,
+		TArrayView<TVector> OutPredictedAccelerations,
+		const TVector& CurrentPosition,
+		const TVector& CurrentVelocity,
+		const TVector& CurrentIntermediateVelocity,
+		const TVector& CurrentAcceleration,
+		const TVector& TargetVelocity,
+		float SmoothingTime,
+		float MaxAcceleration,
+		float SecondsPerPredictionStep,
+		float VDeadzone = 1e-2,
+		float ADeadzone = 1e-4)
 	{
-		if (SmoothingTime < Eps)
-			return TargetRotation;
+		int32 PredictCount = OutPredictedPositions.Num();
+		check(PredictCount > 0);
+		check(OutPredictedVelocities.Num() == PredictCount);
+		check(OutPredictedAccelerations.Num() == PredictCount);
+		check(OutPredictedIntermediateVelocities.Num() == PredictCount);
+		
+		for (int32 i = 0; i < PredictCount; i++)
+		{
+			OutPredictedPositions[i] = CurrentPosition;
+			OutPredictedVelocities[i] = CurrentVelocity;
+			OutPredictedIntermediateVelocities[i] = CurrentIntermediateVelocity;
+			OutPredictedAccelerations[i] = CurrentAcceleration;
+		}
 
-		return FQuat::Slerp(CurrentRotation, TargetRotation, 1.0f - FMath::InvExpApprox((DeltaTime) / (SmoothingTime)));
+		for (int32 i = 0; i < PredictCount; i++)
+		{
+			const float PredictTime = (i + 1) * SecondsPerPredictionStep; // Note i+1 since we want index 0 to be the first prediction step
+			VelocitySpringCharacterUpdate(OutPredictedPositions[i],
+				OutPredictedVelocities[i],
+				OutPredictedIntermediateVelocities[i],
+				OutPredictedAccelerations[i],
+				TargetVelocity,
+				SmoothingTime,
+				MaxAcceleration,
+				PredictTime,
+				VDeadzone,
+				ADeadzone);
+		}
 	}
 
-	/** Prediction of SpringDamperQuat
+	/** Prediction of SimpleSpringDamperQuat
 	 * 
-	 * @param OutPredictedRotations 
-	 * @param OutPredictedAngularVelocities 
-	 * @param PredictorCount 
-	 * @param CurrentRotation 
-	 * @param CurrentAngularVelocity 
-	 * @param TargetRotation 
-	 * @param SmoothingTime 
-	 * @param SecondsPerPredictionStep 
+	 * @param OutPredictedRotations ArrayView of output buffer to put the predicted rotations. ArrayView should be the same size as PredictCount
+	 * @param OutPredictedAngularVelocities ArrayView of output buffer to put the predicted angular velocities. ArrayView should be the same size as PredictCount
+	 * @param PredictCount How many points to predict. Must be greater than 0
+	 * @param CurrentRotation Initial rotation at t = 0
+	 * @param CurrentAngularVelocity Initial angular velocity at t = 0
+	 * @param TargetRotation The target rotation
+	 * @param SmoothingTime The smoothing time
+	 * @param SecondsPerPredictionStep How many seconds per prediction step
 	 */
+	UE_EXPERIMENTAL(5.7, "SpringMath is experimental")
 	inline void SimpleSpringDamperQuatPredict(TArrayView<FQuat> OutPredictedRotations, TArrayView<FVector> OutPredictedAngularVelocities,
-	                                          int PredictorCount, const FQuat& CurrentRotation, const FVector& CurrentAngularVelocity,
+	                                          int32 PredictCount, const FQuat& CurrentRotation, const FVector& CurrentAngularVelocity,
 	                                          const FQuat& TargetRotation, float SmoothingTime, float SecondsPerPredictionStep)
 	{
-		for (int i = 0; i < PredictorCount; i++)
+		check(OutPredictedRotations.Num() == PredictCount);
+		check(OutPredictedAngularVelocities.Num() == PredictCount);
+		
+		for (int32 i = 0; i < PredictCount; i++)
 		{
 			OutPredictedRotations[i] = CurrentRotation;
 			OutPredictedAngularVelocities[i] = CurrentAngularVelocity;
 		}
 
-		for (int i = 0; i < PredictorCount; i++)
+		for (int32 i = 0; i < PredictCount; i++)
 		{
 			const float PredictTime = (i + 1) * SecondsPerPredictionStep; // Note i+1 since we want index 0 to be the first prediction step
+			PRAGMA_DISABLE_EXPERIMENTAL_WARNINGS
 			SimpleSpringDamperQuat(OutPredictedRotations[i], OutPredictedAngularVelocities[i], TargetRotation, SmoothingTime, PredictTime);
+			PRAGMA_ENABLE_EXPERIMENTAL_WARNINGS
+		}
+	}
+
+	/** Specialized quaternion damper, similar to FMath::ExponentialSmoothingApprox but for quaternions.
+	 *  Smooths a value using exponential damping towards a target.
+	 * 
+	 * @param InOutRotation			The value to be smoothed
+	 * @param InTargetRotation		The target to smooth towards
+	 * @param InDeltaTime			Time interval
+	 * @param InSmoothingTime		Timescale over which to smooth. Larger values result in more smoothed behaviour. Can be zero.
+	 */
+	UE_EXPERIMENTAL(5.7, "SpringMath is experimental")
+	static constexpr void ExponentialSmoothingApproxQuat(FQuat& InOutRotation, const FQuat& InTargetRotation, const float InDeltaTime, const float InSmoothingTime)
+	{
+		if (InSmoothingTime > UE_KINDA_SMALL_NUMBER)
+		{
+			InOutRotation = FQuat::Slerp(InOutRotation, InTargetRotation, 1.0f - FMath::InvExpApprox(InDeltaTime / InSmoothingTime));
+		}
+		else
+		{
+			InOutRotation = InTargetRotation;
 		}
 	}
-}
\ No newline at end of file
+}
