[Anim - SkinnedMesh] Expose SkinnedAsset and TransformProvider on FSoftSkinnedMeshComponentDescriptor

* To have parity with the non-soft version

#rnx

diff --git a/Engine/Source/Runtime/Engine/Classes/Animation/AnimBank.h b/Engine/Source/Runtime/Engine/Classes/Animation/AnimBank.h
--- a/Engine/Source/Runtime/Engine/Classes/Animation/AnimBank.h
+++ b/Engine/Source/Runtime/Engine/Classes/Animation/AnimBank.h
@@ -528,10 +528,10 @@ struct FSoftSkinnedMeshComponentDescriptor : public FSkinnedMeshComponentDescrip
 	}
 
 public:
-	UPROPERTY()
+	UPROPERTY(EditAnywhere, Category = "Component Settings", meta = (DisplayAfter = "ComponentClass"))
 	TSoftObjectPtr<class USkinnedAsset> SkinnedAsset;
 
-	UPROPERTY()
+	UPROPERTY(EditAnywhere, Category = "Component Settings", meta = (DisplayAfter = "SkinnedAsset"))
 	TSoftObjectPtr<class UTransformProviderData> TransformProvider;
 };
 
