- Added missing check in Lerp that the alpha argument has primitive type.
- Turned a hard crash caused by calling DerivType on a non-float type into an internal material translator error instead.

[RN] Fixed uncommon crash in Lerp material expression when the alpha pin was connected to a non-primitive output.

diff --git a/Engine/Source/Runtime/Engine/Private/Materials/HLSLMaterialTranslator.cpp b/Engine/Source/Runtime/Engine/Private/Materials/HLSLMaterialTranslator.cpp
--- a/Engine/Source/Runtime/Engine/Private/Materials/HLSLMaterialTranslator.cpp
+++ b/Engine/Source/Runtime/Engine/Private/Materials/HLSLMaterialTranslator.cpp
@@ -3326,9 +3326,14 @@ EDerivativeStatus FHLSLMaterialTranslator::GetDerivativeStatus(int32 Index) cons
 	return AtParameterCodeChunk(Index).DerivativeStatus;
 }
 
-FDerivInfo FHLSLMaterialTranslator::GetDerivInfo(int32 Index, bool bAllowNonFloat) const
+FDerivInfo FHLSLMaterialTranslator::GetDerivInfo(int32 Index, bool bAllowNonFloat)
 {
 	const FShaderCodeChunk& CodeChunk = AtParameterCodeChunk(Index);
+	if (!bAllowNonFloat && (CodeChunk.Type & MCT_Numeric) == 0)
+	{
+		Errorf(TEXT("Internal error: unexpected non-float derivative argument type found (%s)."), DescribeType(CodeChunk.Type));
+		return FDerivInfo { CodeChunk.Type, EDerivativeType::None, EDerivativeStatus::NotValid };
+	}
 	return FDerivInfo { CodeChunk.Type, GetDerivType(CodeChunk.Type, bAllowNonFloat), CodeChunk.DerivativeStatus };
 }
 
@@ -10736,6 +10741,12 @@ int32 FHLSLMaterialTranslator::Lerp(int32 X,int32 Y,int32 A)
 
 	EMaterialValueType ResultType = GetArithmeticResultType(X,Y);
 	EMaterialValueType AlphaType = (*CurrentScopeChunks)[A].Type;
+	
+	if (!IsPrimitiveType(AlphaType))
+	{
+		Errorf(TEXT("Lerp alpha argument is not primitive (it is '%s' instead)"), DescribeType(AlphaType));
+		return INDEX_NONE;
+	}
 
 	if ((AlphaType & MCT_Float) != 0 && ExpressionA && ExpressionA->IsConstant())
 	{
diff --git a/Engine/Source/Runtime/Engine/Private/Materials/HLSLMaterialTranslator.h b/Engine/Source/Runtime/Engine/Private/Materials/HLSLMaterialTranslator.h
--- a/Engine/Source/Runtime/Engine/Private/Materials/HLSLMaterialTranslator.h
+++ b/Engine/Source/Runtime/Engine/Private/Materials/HLSLMaterialTranslator.h
@@ -692,7 +692,7 @@ public:
 
 	const FShaderCodeChunk& AtParameterCodeChunk(int32 Index) const;
 	EDerivativeStatus GetDerivativeStatus(int32 Index) const;
-	FDerivInfo GetDerivInfo(int32 Index, bool bAllowNonFloat = false) const;
+	FDerivInfo GetDerivInfo(int32 Index, bool bAllowNonFloat = false);
 
 protected:
 	bool GetParameterOverrideValueForCurrentFunction(EMaterialParameterType ParameterType, FName ParameterName, FMaterialParameterMetadata& OutResult) const;
