Improvement to 45825273 to better detect erroneous DerivType usage in the material translator.

#rnx

diff --git a/Engine/Source/Runtime/Engine/Private/Materials/HLSLMaterialTranslator.cpp b/Engine/Source/Runtime/Engine/Private/Materials/HLSLMaterialTranslator.cpp
--- a/Engine/Source/Runtime/Engine/Private/Materials/HLSLMaterialTranslator.cpp
+++ b/Engine/Source/Runtime/Engine/Private/Materials/HLSLMaterialTranslator.cpp
@@ -3329,9 +3329,9 @@ EDerivativeStatus FHLSLMaterialTranslator::GetDerivativeStatus(int32 Index) cons
 FDerivInfo FHLSLMaterialTranslator::GetDerivInfo(int32 Index, bool bAllowNonFloat)
 {
 	const FShaderCodeChunk& CodeChunk = AtParameterCodeChunk(Index);
-	if (!bAllowNonFloat && (CodeChunk.Type & MCT_Numeric) == 0)
+	if (!bAllowNonFloat && GetDerivType(CodeChunk.Type, true) == EDerivativeType::None)
 	{
-		Errorf(TEXT("Internal error: unexpected non-float derivative argument type found (%s)."), DescribeType(CodeChunk.Type));
+		Errorf(TEXT("Internal error: unexpected non-numeric derivative argument type found (%s)."), DescribeType(CodeChunk.Type));
 		return FDerivInfo { CodeChunk.Type, EDerivativeType::None, EDerivativeStatus::NotValid };
 	}
 	return FDerivInfo { CodeChunk.Type, GetDerivType(CodeChunk.Type, bAllowNonFloat), CodeChunk.DerivativeStatus };
