Mobile: Make SSR always available in a deferred shading. Enable it on a 'High' scalability bucket
Clarify that project settings like SSR, AO and Clustered reflections affect only Forward shading
Set gen4 TAA an engine default for mobile

diff --git a/Engine/Config/Android/AndroidScalability.ini b/Engine/Config/Android/AndroidScalability.ini
--- a/Engine/Config/Android/AndroidScalability.ini
+++ b/Engine/Config/Android/AndroidScalability.ini
@@ -85,6 +85,23 @@ r.VolumetricFog=0
 r.ContactShadows=0
 r.DFShadowQuality=2
 
+;-----------------------------------------------------------------------------------------------------------------
+
+[ReflectionQuality@0]
+r.SSR.Quality=0
+
+[ReflectionQuality@1]
+r.SSR.Quality=0
+
+[ReflectionQuality@2]
+r.SSR.Quality=0
+
+[ReflectionQuality@3]
+r.SSR.Quality=2
+
+[ReflectionQuality@Cine]
+r.SSR.Quality=2
+
 ; ------------------------------------------------------------------------------
 ; Postprocess settings
 ; ------------------------------------------------------------------------------
diff --git a/Engine/Config/IOS/IOSScalability.ini b/Engine/Config/IOS/IOSScalability.ini
--- a/Engine/Config/IOS/IOSScalability.ini
+++ b/Engine/Config/IOS/IOSScalability.ini
@@ -85,6 +85,23 @@ r.VolumetricFog=0
 r.ContactShadows=0
 r.DFShadowQuality=2
 
+;-----------------------------------------------------------------------------------------------------------------
+
+[ReflectionQuality@0]
+r.SSR.Quality=0
+
+[ReflectionQuality@1]
+r.SSR.Quality=0
+
+[ReflectionQuality@2]
+r.SSR.Quality=0
+
+[ReflectionQuality@3]
+r.SSR.Quality=2
+
+[ReflectionQuality@Cine]
+r.SSR.Quality=2
+
 ; ------------------------------------------------------------------------------
 ; Postprocess settings
 ; ------------------------------------------------------------------------------
diff --git a/Engine/Source/Runtime/Core/Private/HAL/ConsoleManager.cpp b/Engine/Source/Runtime/Core/Private/HAL/ConsoleManager.cpp
--- a/Engine/Source/Runtime/Core/Private/HAL/ConsoleManager.cpp
+++ b/Engine/Source/Runtime/Core/Private/HAL/ConsoleManager.cpp
@@ -4549,11 +4549,11 @@ static TAutoConsoleVariable<int32> CVarSkyLightingQuality(
 
 static TAutoConsoleVariable<int32> CVarMobileDefaultAntiAliasing(
 	TEXT("r.Mobile.AntiAliasing"),
-	1,
+	2,
 	TEXT("Mobile default AntiAliasingMethod\n"
 		 " 0: off (no anti-aliasing)\n"
-		 " 1: FXAA (default, faster than TemporalAA but much more shimmering for non static cases)\n"
-		 " 2: TemporalAA(it will fallback to FXAA if SupportsGen4TAA is disabled) \n"
+		 " 1: FXAA (faster than TemporalAA but much more shimmering for non static cases)\n"
+		 " 2: TemporalAA (default, it will fallback to FXAA if SupportsGen4TAA is disabled) \n"
 		 " 3: MSAA \n"
 		 " 4: TSR(it will fallback to TAA if SupportsGen5TAA is disabled) \n"
 		 " 5: SMAA"),
diff --git a/Engine/Source/Runtime/Engine/Classes/Engine/RendererSettings.h b/Engine/Source/Runtime/Engine/Classes/Engine/RendererSettings.h
--- a/Engine/Source/Runtime/Engine/Classes/Engine/RendererSettings.h
+++ b/Engine/Source/Runtime/Engine/Classes/Engine/RendererSettings.h
@@ -316,24 +316,12 @@ class URendererSettings : public UDeveloperSettings
 	TEnumAsByte<EMobileShadingPath::Type> MobileShadingPath;
 
 	UPROPERTY(config, EditAnywhere, Category = Mobile, meta = (
-		ConsoleVariable = "r.Mobile.AllowDeferredShadingOpenGL", DisplayName = "Allow Deferred Shading on OpenGL",
-		ToolTip = "Whether to allow Deferred Shading on OpenGL, requires the DXC shader compiler and Mobile Shading set to deferred",
-		ConfigRestartRequired = true))
-	uint32 bMobileSupportDeferredOnOpenGL : 1;
-
-	UPROPERTY(config, EditAnywhere, Category = Mobile, meta = (
-		ConsoleVariable = "r.Mobile.SupportGPUScene", DisplayName = "Enable GPUScene on Mobile",
-		ToolTip = "Whether to enable GPUScene on mobile. GPUScene is required for mesh auto-instancing. Changing this setting requires restarting the editor.",
-		ConfigRestartRequired = true))
-	uint32 bMobileSupportGPUScene : 1;
-
-	UPROPERTY(config, EditAnywhere, Category = Mobile, meta = (
-		ConsoleVariable = "r.Mobile.AntiAliasing", DisplayName = "Mobile Anti-Aliasing Method",
+		ConsoleVariable = "r.Mobile.AntiAliasing", DisplayName = "Anti-Aliasing Method",
 		ToolTip = "The mobile default anti-aliasing method."))
 	TEnumAsByte<EMobileAntiAliasingMethod::Type> MobileAntiAliasing;
 
 	UPROPERTY(config, EditAnywhere, Category = Mobile, meta = (
-		ConsoleVariable = "r.Mobile.FloatPrecisionMode", DisplayName = "Mobile Float Precision Mode",
+		ConsoleVariable = "r.Mobile.FloatPrecisionMode", DisplayName = "Float Precision Mode",
 		ToolTip = "Project wide mobile float precision mode for shaders and materials. Changing this setting requires restarting the editor.",
 		ConfigRestartRequired = true))
 	TEnumAsByte<EMobileFloatPrecisionMode::Type> MobileFloatPrecisionMode;
@@ -1388,14 +1376,14 @@ class URendererSettings : public UDeveloperSettings
 
 	UPROPERTY(config, EditAnywhere, Category = Mobile, meta = (
 		ConsoleVariable = "r.Mobile.Forward.EnableLocalLights",
-		DisplayName = "Mobile Local Light Setting",
+		DisplayName = "Local Light Setting (When using ForwardShading)",
 		ToolTip = "Select which Local Light Setting to use for Mobile. Changing this setting requires restarting the editor.",
 		ConfigRestartRequired = true))
 		TEnumAsByte<EMobileLocalLightSetting> MobileLocalLightSetting;
 
 	UPROPERTY(config, EditAnywhere, Category = Mobile, meta = (
 		ConsoleVariable = "r.Mobile.Forward.EnableClusteredReflections",
-		DisplayName = "Enable clustered reflections on mobile forward",
+		DisplayName = "Clustered reflections (When using ForwardShading)",
 		ToolTip = "Whether to enable clustered reflections on mobile forward (including translucency in deferred). Always supported for opaque geometry on mobile deferred. Changing this setting requires restarting the editor.",
 		ConfigRestartRequired = true))
 		uint32 bMobileForwardEnableClusteredReflections : 1;
@@ -1438,6 +1426,18 @@ class URendererSettings : public UDeveloperSettings
 		ConfigRestartRequired = true))
 		uint32 bSupportReversedIndexBuffers : 1;
 
+	UPROPERTY(config, EditAnywhere, Category = Mobile, meta = (
+		ConsoleVariable = "r.Mobile.AmbientOcclusion", DisplayName = "Ambient Occlusion (When using ForwardShading)",
+		ToolTip = "Mobile Ambient Occlusion. Causion: An extra sampler will be occupied in mobile base pass pixel shader after enable the mobile ambient occlusion. Changing this setting requires restarting the editor.",
+		ConfigRestartRequired = true))
+	uint32 bMobileAmbientOcclusion : 1;
+
+	UPROPERTY(config, EditAnywhere, Category = Mobile, meta = (
+		ConsoleVariable = "r.Mobile.ScreenSpaceReflections", DisplayName = "Screen Space Reflections (When using Forward Shading)",
+		ToolTip = "Support Screen Space Reflections with mobile rendering. Screen Space Reflections on mobile require TAA Changing this setting requires restarting the editor.",
+		ConfigRestartRequired = true))
+	uint32 bMobileScreenSpaceReflections : 1;
+
 	UPROPERTY(config, EditAnywhere, Category = Mobile, meta = (
 		ConsoleVariable = "r.Mobile.DBuffer", DisplayName = "Mobile DBuffer Decals",
 		ToolTip = "Whether to accumulate decal properties to a buffer before the base pass with mobile rendering. DBuffer enabled forces a full prepass. Changing this setting requires restarting the editor.",
@@ -1487,18 +1487,6 @@ class URendererSettings : public UDeveloperSettings
 	UPROPERTY()
 	uint8 MobilePlanarReflectionMode;
 
-	UPROPERTY(config, EditAnywhere, Category = Mobile, meta = (
-		ConsoleVariable = "r.Mobile.ScreenSpaceReflections", DisplayName = "Support Screen Space Reflections on mobile",
-		ToolTip = "Support Screen Space Reflections with mobile rendering. Screen Space Reflections on mobile require TAA Changing this setting requires restarting the editor.",
-		ConfigRestartRequired = true))
-		uint32 bMobileScreenSpaceReflections : 1;
-
-	UPROPERTY(config, EditAnywhere, Category = Mobile, meta = (
-		ConsoleVariable = "r.Mobile.SupportsGen4TAA", DisplayName = "Support desktop Gen4 TAA on mobile",
-		ToolTip = "Support desktop Gen4 TAA with mobile rendering. Changing this setting requires restarting the editor.",
-		ConfigRestartRequired = true))
-		uint32 bMobileSupportsGen4TAA : 1;
-
 	UPROPERTY(config, EditAnywhere, Category = Nanite, meta = (
 		ConsoleVariable = "r.Nanite.Foliage", DisplayName = "Nanite Foliage (Experimental)",
 		ToolTip = "Enables experimental features required to support Nanite Foliage.",
diff --git a/Engine/Source/Runtime/Engine/Private/RendererSettings.cpp b/Engine/Source/Runtime/Engine/Private/RendererSettings.cpp
--- a/Engine/Source/Runtime/Engine/Private/RendererSettings.cpp
+++ b/Engine/Source/Runtime/Engine/Private/RendererSettings.cpp
@@ -354,10 +354,6 @@ void URendererSettings::PostEditChangeProperty(FPropertyChangedEvent& PropertyCh
 
 		if ((PropertyChangedEvent.Property->GetFName() == GET_MEMBER_NAME_CHECKED(URendererSettings, MobileShadingPath)))
 		{
-			if (MobileShadingPath.GetValue() == 0)
-			{
-				bMobileSupportDeferredOnOpenGL = 0;
-			}
 			if(MobileShadingPath.GetValue() == 1)
 			{
 				bMobilePostProcessing = 1;
@@ -492,11 +488,6 @@ bool URendererSettings::CanEditChange(const FProperty* InProperty) const
 		return !bForwardShading;
 	}
 
-	if ((InProperty->GetFName() == GET_MEMBER_NAME_CHECKED(URendererSettings, bMobileSupportDeferredOnOpenGL)))
-	{
-		return MobileShadingPath.GetValue() > 0;
-	}
-
 	if ((InProperty->GetFName() == GET_MEMBER_NAME_CHECKED(URendererSettings, bMobilePostProcessing)))
 	{
 		return MobileShadingPath.GetValue() == 0;
diff --git a/Engine/Source/Runtime/RenderCore/Private/RenderUtils.cpp b/Engine/Source/Runtime/RenderCore/Private/RenderUtils.cpp
--- a/Engine/Source/Runtime/RenderCore/Private/RenderUtils.cpp
+++ b/Engine/Source/Runtime/RenderCore/Private/RenderUtils.cpp
@@ -553,12 +553,14 @@ RENDERCORE_API uint32 GetPlatformShadingModelsMask(const FStaticShaderPlatform P
 
 using ShaderPlatformMaskType = TStaticBitArray<EShaderPlatform::SP_NumPlatforms>;
 
+// Only for a mobile ForwardShading
 RENDERCORE_API bool IsMobileAmbientOcclusionEnabled(const FStaticShaderPlatform Platform)
 {
 	static FShaderPlatformCachedIniValue<bool> MobileAmbientOcclusionIniValue(TEXT("r.Mobile.AmbientOcclusion"));
 	return IsMobilePlatform(Platform) && MobileAmbientOcclusionIniValue.Get(Platform);
 }
 
+// Only for a mobile ForwardShading
 RENDERCORE_API bool AreMobileScreenSpaceReflectionsEnabled(const FStaticShaderPlatform Platform)
 {
 	static FShaderPlatformCachedIniValue<bool> MobileScreenSpaceReflectionsIniValue(TEXT("r.Mobile.ScreenSpaceReflections"));
diff --git a/Engine/Source/Runtime/Renderer/Private/MobileBasePassRendering.h b/Engine/Source/Runtime/Renderer/Private/MobileBasePassRendering.h
--- a/Engine/Source/Runtime/Renderer/Private/MobileBasePassRendering.h
+++ b/Engine/Source/Runtime/Renderer/Private/MobileBasePassRendering.h
@@ -437,7 +437,7 @@ public:
 		OutEnvironment.SetDefine(TEXT("MOBILE_TRANSLUCENT_COLOR_TRANSMITTANCE_DUAL_SRC_BLENDING"), TranslucentColorTransmittanceMode == EMobileTranslucentColorTransmittanceMode::DUAL_SRC_BLENDING ? 1u : 0u);
 		OutEnvironment.SetDefine(TEXT("MOBILE_TRANSLUCENT_COLOR_TRANSMITTANCE_PROGRAMMABLE_BLENDING"), TranslucentColorTransmittanceMode == EMobileTranslucentColorTransmittanceMode::PROGRAMMABLE_BLENDING ? 1u : 0u);
 		OutEnvironment.SetDefine(TEXT("MOBILE_TRANSLUCENT_COLOR_TRANSMITTANCE_SINGLE_SRC_BLENDING"), TranslucentColorTransmittanceMode == EMobileTranslucentColorTransmittanceMode::SINGLE_SRC_BLENDING ? 1u : 0u);
-		OutEnvironment.SetDefine(TEXT("MOBILE_SSR_ENABLED"), AreMobileScreenSpaceReflectionsEnabled(Parameters.Platform) ? 1u : 0u);
+		OutEnvironment.SetDefine(TEXT("MOBILE_SSR_ENABLED"), !bDeferredShadingEnabled && AreMobileScreenSpaceReflectionsEnabled(Parameters.Platform) ? 1u : 0u);
 	}
 	
 	/** Initialization constructor. */
diff --git a/Engine/Source/Runtime/Renderer/Private/MobileDeferredShadingPass.cpp b/Engine/Source/Runtime/Renderer/Private/MobileDeferredShadingPass.cpp
--- a/Engine/Source/Runtime/Renderer/Private/MobileDeferredShadingPass.cpp
+++ b/Engine/Source/Runtime/Renderer/Private/MobileDeferredShadingPass.cpp
@@ -361,11 +361,6 @@ public:
 			return false;
 		}
 
-		if (PermutationVector.Get<FMobileSSRQuality>() != EMobileSSRQuality::Disabled && !AreMobileScreenSpaceReflectionsEnabled(Parameters.Platform))
-		{
-			return false;
-		}
-
 		if (PermutationVector.Get<FSkyShadowing>() && !IsMobileDistanceFieldAOEnabled(Parameters.Platform))
 		{
 			return false;
diff --git a/Engine/Source/Runtime/Renderer/Private/MobileSSR.cpp b/Engine/Source/Runtime/Renderer/Private/MobileSSR.cpp
--- a/Engine/Source/Runtime/Renderer/Private/MobileSSR.cpp
+++ b/Engine/Source/Runtime/Renderer/Private/MobileSSR.cpp
@@ -13,6 +13,7 @@ int32 GMobileScreenSpaceReflectionsSupported = 0;
 static FAutoConsoleVariableRef CVarMobileScreenSpaceReflections(
 	TEXT("r.Mobile.ScreenSpaceReflections"),
 	GMobileScreenSpaceReflectionsSupported,
+	TEXT("Caution: Affects only mobile Forward shading. Screen Space Reflections on mobile require TAA.\n")
 	TEXT("0: Mobile Renderer Screen Space Reflections disabled (default)\n"
 		 "1: Mobile Renderer Screen Space Reflections enabled\n"),
 	ECVF_ReadOnly | ECVF_RenderThreadSafe);
@@ -26,7 +27,14 @@ static FAutoConsoleVariableRef CVarMobileMobileSSRIntensity(
 
 bool IsMobileSSREnabled(const FViewInfo& View)
 {
-	return GMobileScreenSpaceReflectionsSupported && ScreenSpaceRayTracing::ShouldRenderScreenSpaceReflections(View) && (GMobileMobileSSRIntensity < 0.0f || GMobileMobileSSRIntensity > 0.0f);
+	const EShaderPlatform Platform = View.GetShaderPlatform();
+	
+	bool bDeferredWithFullDepth = false;
+	if (IsMobileDeferredShadingEnabled(Platform))
+	{
+		bDeferredWithFullDepth = (MobileUsesFullDepthPrepass(Platform) || !MobileAllowFramebufferFetch(Platform));
+	}
+	return (GMobileScreenSpaceReflectionsSupported || bDeferredWithFullDepth) && ScreenSpaceRayTracing::ShouldRenderScreenSpaceReflections(View) && (GMobileMobileSSRIntensity < 0.0f || GMobileMobileSSRIntensity > 0.0f);
 }
 
 EMobileSSRQuality ActiveMobileSSRQuality(const FViewInfo& View, bool bHasVelocityTexture)
diff --git a/Engine/Source/Runtime/Renderer/Private/MobileShadingRenderer.cpp b/Engine/Source/Runtime/Renderer/Private/MobileShadingRenderer.cpp
--- a/Engine/Source/Runtime/Renderer/Private/MobileShadingRenderer.cpp
+++ b/Engine/Source/Runtime/Renderer/Private/MobileShadingRenderer.cpp
@@ -324,8 +324,7 @@ FMobileSceneRenderer::FMobileSceneRenderer(const FSceneViewFamily* InViewFamily,
 	bIsMaskedOnlyDepthPrepassEnabled = Scene->EarlyZPassMode == DDM_MaskedOnly;
 	bEnableClusteredLocalLights = MobileForwardEnableLocalLights(ShaderPlatform);
 	bEnableClusteredReflections = MobileForwardEnableClusteredReflections(ShaderPlatform);
-	bRequiresScreenSpaceReflections = AreMobileScreenSpaceReflectionsEnabled(ShaderPlatform);
-	
+		
 	StandardTranslucencyPass = ViewFamily.AllowTranslucencyAfterDOF() ? ETranslucencyPass::TPT_TranslucencyStandard : ETranslucencyPass::TPT_AllTranslucency;
 	StandardTranslucencyMeshPass = TranslucencyPassToMeshPass(StandardTranslucencyPass);
 
diff --git a/Engine/Source/Runtime/Renderer/Private/SceneRendering.h b/Engine/Source/Runtime/Renderer/Private/SceneRendering.h
--- a/Engine/Source/Runtime/Renderer/Private/SceneRendering.h
+++ b/Engine/Source/Runtime/Renderer/Private/SceneRendering.h
@@ -2948,7 +2948,6 @@ private:
 	bool bShouldRenderCustomDepth;
 	bool bShouldRenderVelocities;
 	bool bShouldRenderHZB;
-	bool bRequiresScreenSpaceReflections;
 	bool bIsFullDepthPrepassEnabled;
 	bool bIsMaskedOnlyDepthPrepassEnabled;
 	bool bRequiresSceneDepthAux;
