[Soundbase][PropertySheets] Move injection from serialize to presave and postload to avoid errors on cook with uobject dependencies harvest

#rnx

diff --git a/Engine/Source/Runtime/Engine/Classes/Sound/SoundBase.h b/Engine/Source/Runtime/Engine/Classes/Sound/SoundBase.h
--- a/Engine/Source/Runtime/Engine/Classes/Sound/SoundBase.h
+++ b/Engine/Source/Runtime/Engine/Classes/Sound/SoundBase.h
@@ -107,6 +107,8 @@ public:
 UCLASS(config=Engine, hidecategories=Object, abstract, editinlinenew, BlueprintType, MinimalAPI)
 class USoundBase : public UObject, public IAudioPropertiesSheetAssetUserInterface, public IInterface_AssetUserData
 {
+
+private:
 	GENERATED_UCLASS_BODY()
 
 public:
@@ -276,6 +278,7 @@ public:
 #endif
 	ENGINE_API virtual bool CanBeClusterRoot() const override;
 	ENGINE_API virtual bool CanBeInCluster() const override;
+	ENGINE_API virtual void PreSave(FObjectPreSaveContext SaveContext) override;
 	ENGINE_API virtual void Serialize(FArchive& Ar) override;
 
 #if WITH_EDITOR
diff --git a/Engine/Source/Runtime/Engine/Private/SoundBase.cpp b/Engine/Source/Runtime/Engine/Private/SoundBase.cpp
--- a/Engine/Source/Runtime/Engine/Private/SoundBase.cpp
+++ b/Engine/Source/Runtime/Engine/Private/SoundBase.cpp
@@ -5,11 +5,11 @@
 #include "Engine/AssetUserData.h"
 #include "IAudioParameterTransmitter.h"
 #include "UObject/UObjectGlobals.h"
+#include "UObject/ObjectSaveContext.h"
 #include "Sound/SoundSubmix.h"
 
 #include UE_INLINE_GENERATED_CPP_BY_NAME(SoundBase)
 
-
 USoundBase::USoundBase(const FObjectInitializer& ObjectInitializer)
 	: Super(ObjectInitializer)
 	, VirtualizationMode(EVirtualizationMode::Restart)
@@ -252,6 +252,7 @@ void USoundBase::PostLoad()
 	{
 		AudioPropertiesSheet->ConditionalPostLoad();
 		BindToPropertySheetChanges();
+		InjectPropertySheet();
 	}
 }
 #endif // WITH_EDITORONLY_DATA
@@ -266,6 +267,15 @@ bool USoundBase::CanBeInCluster() const
 	return true;
 }
 
+void USoundBase::PreSave(FObjectPreSaveContext SaveContext)
+{
+#if WITH_EDITORONLY_DATA
+	InjectPropertySheet();
+#endif
+	
+	Super::PreSave(SaveContext);
+}
+
 void USoundBase::Serialize(FArchive& Ar)
 {
 	Super::Serialize(Ar);
@@ -279,8 +289,6 @@ void USoundBase::Serialize(FArchive& Ar)
 			SoundConcurrencySettings_DEPRECATED = nullptr;
 		}
 	}
-
-	InjectPropertySheet();
 #endif // WITH_EDITORONLY_DATA
 }
 
