Fixing issue where attaching an actor to the socket of another skinned me actor would not transact correctly in Multi User or as part of an undo operation.

This is a 5.7 backport of a fix that is already in 5.8. The fix in 5.8 uses the PostEditUndo override in USkinnedMeshComponent, which was added in 5.8 but is not available in 5.7. Therefore as a backport, we make the same call UpdateCachedSocketAttachmentFlag in PostEditChangeProperty, which whilst not strictly correct for the use case, is called within the same call stack as PostEditUndo as required for this fix.

Due to not wanting this change to end up in addition to the long term fix in main, it is explicitly excluded from the robomerge into main.

diff --git a/Engine/Source/Runtime/Engine/Private/Components/SkinnedMeshComponent.cpp b/Engine/Source/Runtime/Engine/Private/Components/SkinnedMeshComponent.cpp
--- a/Engine/Source/Runtime/Engine/Private/Components/SkinnedMeshComponent.cpp
+++ b/Engine/Source/Runtime/Engine/Private/Components/SkinnedMeshComponent.cpp
@@ -1480,6 +1480,14 @@ void USkinnedMeshComponent::PostEditChangeProperty(FPropertyChangedEvent& Proper
 			MeshDeformerInstanceSettings = ActiveMeshDeformer ? ActiveMeshDeformer->CreateSettingsInstance(this) : nullptr;
 		}
 	}
+
+	if (PropertyChangedEvent.ChangeType == EPropertyChangeType::Unspecified)
+	{
+		// NOTE: From 5.8, this call has been simplified and moved to PostEditUndo in USkinnedMeshComponent.
+		// If the PropertyChangedEvent is not a specified change, it may have originated from a transaction such as undo or MU. 
+		// Therefore, make sure that any socket attachments are up to date.
+		UpdateCachedSocketAttachmentFlag();
+	}
 }
 
 EDataValidationResult USkinnedMeshComponent::IsDataValid(class FDataValidationContext& Context) const
