fix FKShapeElem copy constructor accidentally changing the CollisionEnabled flag instead of copying it

diff --git a/Engine/Source/Runtime/Engine/Classes/PhysicsEngine/ShapeElem.h b/Engine/Source/Runtime/Engine/Classes/PhysicsEngine/ShapeElem.h
--- a/Engine/Source/Runtime/Engine/Classes/PhysicsEngine/ShapeElem.h
+++ b/Engine/Source/Runtime/Engine/Classes/PhysicsEngine/ShapeElem.h
@@ -65,7 +65,7 @@ struct FKShapeElem
 	, Name(Copy.Name)
 	, ShapeType(Copy.ShapeType)
 	, bContributeToMass(Copy.bContributeToMass)
-	, CollisionEnabled(ECollisionEnabled::QueryAndPhysics)
+	, CollisionEnabled(Copy.CollisionEnabled)
 	, UserData(this)
 	{
 	}
