Fix for anim bank using anim skeleton bone count instead of mesh skeleton bone count for key frames.

diff --git a/Engine/Source/Runtime/Engine/Private/Animation/AnimBank.cpp b/Engine/Source/Runtime/Engine/Private/Animation/AnimBank.cpp
--- a/Engine/Source/Runtime/Engine/Private/Animation/AnimBank.cpp
+++ b/Engine/Source/Runtime/Engine/Private/Animation/AnimBank.cpp
@@ -1507,7 +1507,7 @@ void BuildSkinnedAssetMapping(const USkinnedAsset& Asset, FSkinnedAssetMapping&
 	const int32 NumAnimBones = AnimRefSkeleton.GetRawBoneNum();
 
 	// Get Number of RawBones (no virtual)
-	Mapping.BoneCount = AnimRefSkeleton.GetRawBoneNum();
+	Mapping.BoneCount = NumMeshBones;
 
 	// Get the skeleton reference pose in local space;
 	const TArray<FTransform>& MeshLocalRefPose = MeshRefSkeleton.GetRawRefBonePose();
