[CellTransformer] Replace references to objects that could be recreated during RerunConstructionScripts triggered from ApplyRuntimeCellsTransformerStack in Standalone Editor

Follow up to CL 44081786. That change used FCoreUObjectDelegates::OnObjectsReplaced to track references to replace, but that delegate only fires when UEditorEngine is present, which is not true in Standalone Editor.
This change introduces a new delegate to track component replacement during RerunConstructionScripts, and uses it to update references in both the Editor and Standalone Editor paths.

diff --git a/Engine/Source/Runtime/Engine/Classes/GameFramework/Actor.h b/Engine/Source/Runtime/Engine/Classes/GameFramework/Actor.h
--- a/Engine/Source/Runtime/Engine/Classes/GameFramework/Actor.h
+++ b/Engine/Source/Runtime/Engine/Classes/GameFramework/Actor.h
@@ -185,6 +185,9 @@ DECLARE_DELEGATE_SixParams(FMakeNoiseDelegate, AActor*, float /*Loudness*/, clas
 
 #if WITH_EDITOR
 DECLARE_EVENT_TwoParams(AActor, FActorOnPackagingModeChanged, AActor*, bool /* bExternal */);
+
+using FReplacementComponentMap = TMap<UObject*, UObject*>;
+DECLARE_MULTICAST_DELEGATE_OneParam(FActorOnComponentsReplaced, const FReplacementComponentMap&);
 #endif
 
 #if !UE_BUILD_SHIPPING
@@ -4345,6 +4348,10 @@ public:
 	/** Returns the instance components array */
 	ENGINE_API const TArray<UActorComponent*>& GetInstanceComponents() const;
 
+#if WITH_EDITOR
+	FActorOnComponentsReplaced OnComponentsReplaced;
+#endif
+
 	//~=============================================================================
 	// Navigation/AI related functions
 	
diff --git a/Engine/Source/Runtime/Engine/Private/ActorConstruction.cpp b/Engine/Source/Runtime/Engine/Private/ActorConstruction.cpp
--- a/Engine/Source/Runtime/Engine/Private/ActorConstruction.cpp
+++ b/Engine/Source/Runtime/Engine/Private/ActorConstruction.cpp
@@ -784,6 +784,13 @@ void AActor::RerunConstructionScripts()
 			GEditor->NotifyToolsOfObjectReplacement(OldToNewComponentMapping);
 		}
 
+		// NotifyToolsOfObjectReplacement works only with UEditorEngine
+		// OnComponentsReplaced delegate broadcasts in standalone editor as well
+		if (OldToNewComponentMapping.Num() > 0)
+		{
+			OnComponentsReplaced.Broadcast(OldToNewComponentMapping);
+		}
+
 		if (bErrorFree)
 		{
 			CurrentTransactionAnnotation = nullptr;
diff --git a/Engine/Source/Runtime/Engine/Private/WorldPartition/WorldPartition.cpp b/Engine/Source/Runtime/Engine/Private/WorldPartition/WorldPartition.cpp
--- a/Engine/Source/Runtime/Engine/Private/WorldPartition/WorldPartition.cpp
+++ b/Engine/Source/Runtime/Engine/Private/WorldPartition/WorldPartition.cpp
@@ -2508,22 +2508,23 @@ void UWorldPartition::ApplyRuntimeCellsTransformerStack(ULevel* InLevel)
 			});
 		
 			TMap<UObject*, UObject*> CombinedReplacementMap;
-			FDelegateHandle OnObjectsReplacedHandle = FCoreUObjectDelegates::OnObjectsReplaced.AddLambda([&CombinedReplacementMap](const TMap<UObject*, UObject*>& ReplacementMap)
-			{
-				CombinedReplacementMap.Append(ReplacementMap);
-			});
 
 			for (TObjectPtr<AActor>& LevelActor : LevelActors)
 			{
 				// Deal with deleted child actors that might have been deleted from previous calls to RerunConstructionScripts
 				if (IsValid(LevelActor))
 				{
+					FDelegateHandle OnComponentsReplacedHandle = LevelActor->OnComponentsReplaced.AddLambda([&CombinedReplacementMap](const TMap<UObject*, UObject*>& ReplacementMap)
+					{
+						CombinedReplacementMap.Append(ReplacementMap);
+					});
+
 					LevelActor->RerunConstructionScripts();
+					
+					LevelActor->OnComponentsReplaced.Remove(OnComponentsReplacedHandle);
 				}
 			}
 		
-			FCoreUObjectDelegates::OnObjectsReplaced.Remove(OnObjectsReplacedHandle);
-		
 			TMap<UObject*, UObject*> FinalReplacementMap;
 			TSet<UObject*> CurrentlyResolving;
 			TFunction<UObject*(UObject*)> ResolveFinalReplacement = [&CombinedReplacementMap, &FinalReplacementMap, &CurrentlyResolving, &ResolveFinalReplacement](UObject* Object) -> UObject*
