[Streaming] Add configurable CSV stats collection for render asset streaming, via new define.
- Introduce UE_WITH_RENDER_ASSET_STREAMING_CSV_STATS define to control when CSV-only statistics are collected, defaulting to UE_BUILD_TEST (current behavior).
- Enabled UE_WITH_RENDER_ASSET_STREAMING_CSV_STATS for internal project.

diff --git a/Engine/Source/Runtime/Engine/Private/Streaming/AsyncTextureStreaming.cpp b/Engine/Source/Runtime/Engine/Private/Streaming/AsyncTextureStreaming.cpp
--- a/Engine/Source/Runtime/Engine/Private/Streaming/AsyncTextureStreaming.cpp
+++ b/Engine/Source/Runtime/Engine/Private/Streaming/AsyncTextureStreaming.cpp
@@ -1068,7 +1068,7 @@ void FRenderAssetStreamingMipCalcTask::DoWork()
 
 #if STATS
 	UpdateStats_Async();
-#elif UE_BUILD_TEST
+#elif UE_WITH_RENDER_ASSET_STREAMING_CSV_STATS
 	UpdateCSVOnlyStats_Async();
 #endif // STATS
 }
@@ -1137,7 +1137,7 @@ void FRenderAssetStreamingMipCalcTask::UpdateStats_Async()
 	Stats.RequiredVirtualTexturePool = 0;
 	Stats.UsedVirtualTexturePool = 0;
 
-	for (FStreamingRenderAsset& StreamingRenderAsset : StreamingRenderAssets)
+	for (const FStreamingRenderAsset& StreamingRenderAsset : StreamingRenderAssets)
 	{
 		if (IsAborted()) break;
 		if (!StreamingRenderAsset.RenderAsset) continue;
diff --git a/Engine/Source/Runtime/Engine/Private/Streaming/StreamingManagerTexture.cpp b/Engine/Source/Runtime/Engine/Private/Streaming/StreamingManagerTexture.cpp
--- a/Engine/Source/Runtime/Engine/Private/Streaming/StreamingManagerTexture.cpp
+++ b/Engine/Source/Runtime/Engine/Private/Streaming/StreamingManagerTexture.cpp
@@ -1952,7 +1952,7 @@ void FRenderAssetStreamingManager::UpdateResourceStreaming( float DeltaTime, boo
 		STAT(GatheredStats.CallbacksCycles = 0);
 #if STATS
 		UpdateStats();
-#elif UE_BUILD_TEST
+#elif UE_WITH_RENDER_ASSET_STREAMING_CSV_STATS
 		UpdateCSVOnlyStats();
 #endif // STATS
 	}
@@ -2053,9 +2053,9 @@ void FRenderAssetStreamingManager::UpdateResourceStreaming( float DeltaTime, boo
 
 		STAT(GatheredStats.StreamRenderAssetsCycles += FPlatformTime::Cycles();)
 #if STATS
-			UpdateStats();
-#elif UE_BUILD_TEST
-			UpdateCSVOnlyStats();
+		UpdateStats();
+#elif UE_WITH_RENDER_ASSET_STREAMING_CSV_STATS
+		UpdateCSVOnlyStats();
 #endif // STATS
 	}
 
diff --git a/Engine/Source/Runtime/Engine/Private/Streaming/TextureStreamingHelpers.h b/Engine/Source/Runtime/Engine/Private/Streaming/TextureStreamingHelpers.h
--- a/Engine/Source/Runtime/Engine/Private/Streaming/TextureStreamingHelpers.h
+++ b/Engine/Source/Runtime/Engine/Private/Streaming/TextureStreamingHelpers.h
@@ -141,6 +141,10 @@ typedef TArray<const UStreamableRenderAsset*, TInlineAllocator<12> > FRemovedRen
 #define MAX_MIPDELTA			5.0f
 #define MAX_LASTRENDERTIME		90.0f
 
+#ifndef UE_WITH_RENDER_ASSET_STREAMING_CSV_STATS
+#define UE_WITH_RENDER_ASSET_STREAMING_CSV_STATS UE_BUILD_TEST
+#endif
+
 class UPrimitiveComponent;
 class FDynamicRenderAssetInstanceManager;
 template<typename T>
