Fix check/crash in Convex environmental collisions code.

diff --git a/Engine/Source/Runtime/Engine/Private/PhysicsEngine/EnvironmentalCollisions.cpp b/Engine/Source/Runtime/Engine/Private/PhysicsEngine/EnvironmentalCollisions.cpp
--- a/Engine/Source/Runtime/Engine/Private/PhysicsEngine/EnvironmentalCollisions.cpp
+++ b/Engine/Source/Runtime/Engine/Private/PhysicsEngine/EnvironmentalCollisions.cpp
@@ -143,18 +143,27 @@ void FEnvironmentalCollisions::AppendCollisionDataFromEnvironment(const USceneCo
 										const int32 PointIndex = ImplicitConvex.GetPlaneVertex(FaceIndex, Index);
 										Faces[FaceIndex].Indices[Index] = PointIndex;
 
-										PointFaces[PointIndex].Add(FaceIndex);
+										if (ensure(PointFaces.IsValidIndex(PointIndex)))
+										{
+											PointFaces[PointIndex].Add(FaceIndex);
+										}
 									}
 								}
 
 								for (int32 PointIndex = 0; PointIndex < NumSurfacePoints; ++PointIndex)
 								{
-									check(PointFaces[PointIndex].Num() >= 3);
-									const int32 Index0 = PointFaces[PointIndex][0];
-									const int32 Index1 = PointFaces[PointIndex][1];
-									const int32 Index2 = PointFaces[PointIndex][2];
-
-									if (!FMath::IntersectPlanes3(SurfacePoints[PointIndex], Faces[Index0].Plane, Faces[Index1].Plane, Faces[Index2].Plane))
+									if (ensure(PointFaces[PointIndex].Num() >= 3))
+									{
+										const int32 Index0 = PointFaces[PointIndex][0];
+										const int32 Index1 = PointFaces[PointIndex][1];
+										const int32 Index2 = PointFaces[PointIndex][2];
+
+										if (!FMath::IntersectPlanes3(SurfacePoints[PointIndex], Faces[Index0].Plane, Faces[Index1].Plane, Faces[Index2].Plane))
+										{
+											SurfacePoints[PointIndex] = ComponentToClothMatrix.TransformPosition(FVector(ImplicitConvex.GetVertex(PointIndex)));
+										}
+									}
+									else
 									{
 										SurfacePoints[PointIndex] = ComponentToClothMatrix.TransformPosition(FVector(ImplicitConvex.GetVertex(PointIndex)));
 									}
