Add way to define the chunkId you want for an encrypted plugin.
Pipe the encrypted plugin chunk Id through with added logging.

This allows teams to stabilize the chunk from release to release when new encrypted packages go in and out of a build the chunkID can be consistent.

diff --git a/Engine/Source/Runtime/Engine/Classes/Engine/ContentEncryptionConfig.h b/Engine/Source/Runtime/Engine/Classes/Engine/ContentEncryptionConfig.h
--- a/Engine/Source/Runtime/Engine/Classes/Engine/ContentEncryptionConfig.h
+++ b/Engine/Source/Runtime/Engine/Classes/Engine/ContentEncryptionConfig.h
@@ -34,6 +34,7 @@ public:
 
 		EGroupType GroupType = EGroupType::Root;
 		EAllowedReferences AllowedReferences = EAllowedReferences::None;
+		int32 DesiredChunkId = INDEX_NONE;
 	};
 
 	typedef TMap<FName, FGroup> TGroupMap;
@@ -63,6 +64,11 @@ public:
 		PackageGroups.FindOrAdd(InGroupName).AllowedReferences = InAllowedReferences;
 	}
 
+	void SetDesiredChunkId(FName InGroupName, int32 InChunkId)
+	{
+		PackageGroups.FindOrAdd(InGroupName).DesiredChunkId = InChunkId;
+	}
+
 	void AddReleasedKey(FGuid InKey)
 	{
 		ReleasedKeys.Add(InKey);
