Fix validation errors on Metal with RT1 in DebugViewModePixelShader. Previously, RT 1 could have a missing or incorrect RT bound in viewmodes that do not use DebugAux. To prevent this, a dummy texture is bound by default.

#rnx

diff --git a/Engine/Shaders/Private/DebugViewModePixelShader.usf b/Engine/Shaders/Private/DebugViewModePixelShader.usf
--- a/Engine/Shaders/Private/DebugViewModePixelShader.usf
+++ b/Engine/Shaders/Private/DebugViewModePixelShader.usf
@@ -510,7 +510,7 @@ void Main(
 	in FDebugPSIn DebugInputs
 	OPTIONAL_IsFrontFace,
 	out float4 OutColor : SV_Target0,
-	out uint4 OutDebugAux : SV_Target1
+	out uint OutDebugAux : SV_Target1
 	)
 {
 	StereoSetupPS(DebugInputs.StereoInput);
@@ -520,7 +520,7 @@ void Main(
 
 	// If we see pink something went wrong.
 	float4 VisualizeColor = float4(1, 0, 1, 1);
-	uint4 DebugAux = uint4(0, 0, 0, 0);
+	uint DebugAux = 0;
 	bool bApplyDebugShading = false;
 
 	BRANCH
@@ -567,7 +567,7 @@ void Main(
 			discard;
 			return;
 		}
-		DebugAux = uint4(InstanceId, 0, 0, 0);
+		DebugAux = InstanceId;
 		VisualizeColor = float4(1, 1, 1, 1);
 		bApplyDebugShading = true;
 	}
diff --git a/Engine/Source/Runtime/Engine/Private/DebugViewModeInterface.cpp b/Engine/Source/Runtime/Engine/Private/DebugViewModeInterface.cpp
--- a/Engine/Source/Runtime/Engine/Private/DebugViewModeInterface.cpp
+++ b/Engine/Source/Runtime/Engine/Private/DebugViewModeInterface.cpp
@@ -15,7 +15,45 @@ FDebugViewModeInterface* FDebugViewModeInterface::Singleton = nullptr;
 
 void FDebugViewModeInterface::SetDrawRenderState(EDebugViewShaderMode DebugViewMode, EBlendMode InBlendMode, FRenderState& DrawRenderState, bool bHasDepthPrepassForMaskedMaterial) const
 {
-	if (DebugViewMode == DVSM_QuadComplexity || DebugViewMode == DVSM_ShaderComplexityBleedingQuadOverhead || DebugViewMode == DVSM_ShaderComplexityContainedQuadOverhead || DebugViewMode == DVSM_ShaderComplexity)
+	// No writes to RT1
+	FBlendStateRHIRef DefaultOpaqueBlendState = TStaticBlendState<
+			/* RT0ColorWriteMask */ CW_RGBA,
+			/* RT0ColorBlendOp   */ BO_Add,
+			/* RT0ColorSrcBlend  */ BF_One,
+			/* RT0ColorDestBlend */ BF_Zero,
+			/* RT0AlphaBlendOp   */ BO_Add,
+			/* RT0AlphaSrcBlend  */ BF_One,
+			/* RT0AlphaDestBlend */ BF_Zero,
+			/* RT1ColorWriteMask */ CW_NONE
+		>::GetRHI();
+
+	// Default values
+	if (IsTranslucentBlendMode(InBlendMode))
+	{
+		// shaders will use an hardcoded alpha
+		DrawRenderState.BlendState = TStaticBlendState<CW_RGBA, BO_Add, BF_SourceAlpha, BF_InverseSourceAlpha, BO_Add, BF_Zero, BF_InverseSourceAlpha, CW_NONE>::GetRHI();
+		DrawRenderState.DepthStencilState = TStaticDepthStencilState<false, CF_DepthNearOrEqual>::GetRHI();
+	}
+	else
+	{
+		DrawRenderState.BlendState = DefaultOpaqueBlendState;
+
+		// If not selected, use depth equal to make alpha test stand out (goes with EarlyZPassMode = DDM_AllOpaque) 
+		if (IsMaskedBlendMode(InBlendMode) && bHasDepthPrepassForMaskedMaterial)
+		{
+			DrawRenderState.DepthStencilState = TStaticDepthStencilState<false, CF_Equal>::GetRHI();
+		}
+		else
+		{
+			DrawRenderState.DepthStencilState = TStaticDepthStencilState<>::GetRHI();
+		}
+	}
+
+	// Viewmode overrides
+	if (DebugViewMode == DVSM_QuadComplexity 
+		|| DebugViewMode == DVSM_ShaderComplexityBleedingQuadOverhead 
+		|| DebugViewMode == DVSM_ShaderComplexityContainedQuadOverhead 
+		|| DebugViewMode == DVSM_ShaderComplexity)
 	{
 		if (IsOpaqueBlendMode(InBlendMode))
 		{
@@ -36,34 +74,23 @@ void FDebugViewModeInterface::SetDrawRenderState(EDebugViewShaderMode DebugViewM
 		{
 			DrawRenderState.DepthStencilState = TStaticDepthStencilState<false, CF_DepthNearOrEqual>::GetRHI();
 		}
-		DrawRenderState.BlendState = TStaticBlendState<CW_RGBA, BO_Add, BF_One, BF_One, BO_Add, BF_Zero, BF_One>::GetRHI();
+		
+		DrawRenderState.BlendState = TStaticBlendState<CW_RGBA, BO_Add, BF_One, BF_One, BO_Add, BF_Zero, BF_One, CW_NONE>::GetRHI();
 	}
 	else if (DebugViewMode == DVSM_OutputMaterialTextureScales)
 	{
-		DrawRenderState.BlendState = TStaticBlendState<>::GetRHI();
 		DrawRenderState.DepthStencilState = TStaticDepthStencilState<false, CF_Always>::GetRHI();
 	}
-	else
+	else if (DebugViewMode == DVSM_ShadowCasters)
 	{
+		// enable RT 1
 		if (IsTranslucentBlendMode(InBlendMode))
 		{
-			// Otherwise, force translucent blend mode (shaders will use an hardcoded alpha).
-			DrawRenderState.BlendState = TStaticBlendState<CW_RGBA, BO_Add, BF_SourceAlpha, BF_InverseSourceAlpha, BO_Add, BF_Zero, BF_InverseSourceAlpha>::GetRHI();
-			DrawRenderState.DepthStencilState = TStaticDepthStencilState<false, CF_DepthNearOrEqual>::GetRHI();
+			DrawRenderState.BlendState = TStaticBlendState<CW_RGBA, BO_Add, BF_SourceAlpha, BF_InverseSourceAlpha, BO_Add, BF_Zero, BF_InverseSourceAlpha, CW_RED>::GetRHI();
 		}
 		else
 		{
-			DrawRenderState.BlendState = TStaticBlendState<>::GetRHI();
-
-			// If not selected, use depth equal to make alpha test stand out (goes with EarlyZPassMode = DDM_AllOpaque) 
-			if (IsMaskedBlendMode(InBlendMode) && bHasDepthPrepassForMaskedMaterial)
-			{
-				DrawRenderState.DepthStencilState = TStaticDepthStencilState<false, CF_Equal>::GetRHI();
-			}
-			else
-			{
-				DrawRenderState.DepthStencilState = TStaticDepthStencilState<>::GetRHI();
-			}
+			DrawRenderState.BlendState = TStaticBlendState<CW_RGBA, BO_Add, BF_One, BF_Zero, BO_Add, BF_One, BF_Zero, CW_RED>::GetRHI();
 		}
 	}
 }
diff --git a/Engine/Source/Runtime/Renderer/Private/BasePassRendering.cpp b/Engine/Source/Runtime/Renderer/Private/BasePassRendering.cpp
--- a/Engine/Source/Runtime/Renderer/Private/BasePassRendering.cpp
+++ b/Engine/Source/Runtime/Renderer/Private/BasePassRendering.cpp
@@ -1581,7 +1581,7 @@ void FDeferredShadingSceneRenderer::RenderBasePassInternal(
 		else if (ViewFamily.UseDebugViewPS())
 		{
 			// Override the base pass with one of the debug view shader mode (see EDebugViewShaderMode) if required.
-			RenderDebugViewMode(GraphBuilder, InViews, SceneTextures.DebugAux, BasePassRenderTargets);
+			RenderDebugViewMode(GraphBuilder, InViews, SceneTextures.DebugAux, BasePassRenderTargets[0], BasePassRenderTargets.DepthStencil);
 		}
 	}
 	else
diff --git a/Engine/Source/Runtime/Renderer/Private/DebugViewModeRendering.cpp b/Engine/Source/Runtime/Renderer/Private/DebugViewModeRendering.cpp
--- a/Engine/Source/Runtime/Renderer/Private/DebugViewModeRendering.cpp
+++ b/Engine/Source/Runtime/Renderer/Private/DebugViewModeRendering.cpp
@@ -242,16 +242,28 @@ void ClearDebugAux(FRDGBuilder& GraphBuilder, const FViewFamilyInfo& ViewFamily,
 	}
 }
 
-void RenderDebugViewMode(FRDGBuilder& GraphBuilder, TArrayView<FViewInfo> Views, FRDGTextureRef DebugAuxTexture, const FRenderTargetBindingSlots& RenderTargets)
+void RenderDebugViewMode(FRDGBuilder& GraphBuilder, TArrayView<FViewInfo> Views, FRDGTextureRef DebugAuxTexture, const FRenderTargetBinding& SceneColorRTBinding, const FDepthStencilBinding& SceneDepthBinding)
 {
 	RDG_EVENT_SCOPE(GraphBuilder, "DebugViewMode");
 
+	check(SceneColorRTBinding.GetTexture());
+
 	if (Views.Num() == 0)
 	{
 		return;
 	}
+	
+	FIntPoint SceneColorExtent = SceneColorRTBinding.GetTexture()->Desc.Extent;
+	
+	// Dummy needs to have the correct size and type to prevent validation errors, but we can use R8 instead of R32.
+	const FRDGTextureDesc RenderTarget1Desc = FRDGTextureDesc::Create2D(SceneColorExtent, PF_R8_UINT, FClearValueBinding::None, TexCreate_ShaderResource | TexCreate_RenderTargetable | TexCreate_UAV);
+	FRDGTextureRef DummyRenderTarget1 = GraphBuilder.CreateTexture(RenderTarget1Desc, TEXT("DummyDebugAux"));
+
+	FRenderTargetBindingSlots DebugRenderTargets{};
+	DebugRenderTargets.DepthStencil = SceneDepthBinding;
+	DebugRenderTargets[0] = SceneColorRTBinding;
+	DebugRenderTargets[1] = FRenderTargetBinding(DummyRenderTarget1, ERenderTargetLoadAction::ENoAction);
 
-	FRenderTargetBindingSlots DebugRenderTargets = RenderTargets;
 	FRDGTextureRef QuadOverdrawTexture = nullptr;
 
 	EDebugViewShaderMode DebugViewMode = Views[0].Family->GetDebugViewShaderMode();
@@ -748,10 +760,11 @@ void ClearDebugAux(FRDGBuilder& GraphBuilder, const FViewFamilyInfo& ViewFamily,
 {}
 
 void RenderDebugViewMode(
-	FRDGBuilder& GraphBuilder,
-	TArrayView<FViewInfo> Views,
-	FRDGTextureRef DebugAuxTexture,
-	const FRenderTargetBindingSlots& RenderTargets)
+	FRDGBuilder& GraphBuilder, 
+	TArrayView<FViewInfo> Views, 
+	FRDGTextureRef DebugAuxTexture, 
+	const FRenderTargetBinding& SceneColorRTBinding, 
+	const FDepthStencilBinding& SceneDepthBinding)
 {}
 
 #endif // WITH_DEBUG_VIEW_MODES
diff --git a/Engine/Source/Runtime/Renderer/Private/DebugViewModeRendering.h b/Engine/Source/Runtime/Renderer/Private/DebugViewModeRendering.h
--- a/Engine/Source/Runtime/Renderer/Private/DebugViewModeRendering.h
+++ b/Engine/Source/Runtime/Renderer/Private/DebugViewModeRendering.h
@@ -217,7 +217,8 @@ public:
 void ClearDebugAux(FRDGBuilder& GraphBuilder, const FViewFamilyInfo& ViewFamily, FRDGTextureRef DebugAuxTexture);
 
 void RenderDebugViewMode(
-	FRDGBuilder& GraphBuilder,
+	FRDGBuilder& GraphBuilder, 
 	TArrayView<FViewInfo> Views,
 	FRDGTextureRef DebugAuxTexture,
-	const FRenderTargetBindingSlots& RenderTargets);
+	const FRenderTargetBinding& SceneColorRTBinding, 
+	const FDepthStencilBinding& SceneDepthBinding);
