Ensure SetLightMapResolution is only executed when the property is actually being updated. Currently this path can lead to a race condition during asset deletion that trips on a missing base UMaterial (if that is part of the deletion process). Currently we are unable to repro deterministically, but we do know that this function is the choke point, and doesn't actually need to be called in the circumstance where a material is being deleted, hence we can protect/isolate this function to only execute only if we require an update.

Perhaps there's another race condition error/fallout in the same area, but currently a repro is impossible, so best option is to move on with this fix for now which will make this call healthier regardless/won't update the value unnecessarily every single time we pass through PostEditChangeProperty.

diff --git a/Engine/Source/Runtime/Engine/Private/StaticMesh.cpp b/Engine/Source/Runtime/Engine/Private/StaticMesh.cpp
--- a/Engine/Source/Runtime/Engine/Private/StaticMesh.cpp
+++ b/Engine/Source/Runtime/Engine/Private/StaticMesh.cpp
@@ -5289,7 +5289,10 @@ void UStaticMesh::PostEditChangeProperty(FPropertyChangedEvent& PropertyChangedE
 	}
 #endif
 
-	SetLightMapResolution(FMath::Max(GetLightMapResolution(), 0));
+	if (PropertyName == GetLightMapResolutionName())
+	{
+		SetLightMapResolution(FMath::Max(GetLightMapResolution(), 0));
+	}
 
 	if (PropertyChangedEvent.MemberProperty 
 		&& (PropertyChangedEvent.MemberProperty->GetFName() == UStaticMesh::GetPositiveBoundsExtensionName() || 
