Fix for preview materials in the substrate material graph not displaying when an output value is not a Substrate result. Substrate previews rely on IsResultSubstrateMaterial being true for the FMatExpressionPreview output, however MF outputs (and possibly other previews) cannot guarantee that, so this introduces a conversion step for the preview compiles.

Also swapped the method from using the SlabBSDF compilation pathway due to incompatibility with raw nodes that don't output SubstrateMaterial (i.e. IsResultSubstrateMaterial returning false), so now the LegacyConversion is default value generator. Behaviour remains the same for other use cases as this functionality is exclusively required for preview compiles.

diff --git a/Engine/Source/Runtime/Engine/Private/Materials/MaterialExpressionSubstrate.cpp b/Engine/Source/Runtime/Engine/Private/Materials/MaterialExpressionSubstrate.cpp
--- a/Engine/Source/Runtime/Engine/Private/Materials/MaterialExpressionSubstrate.cpp
+++ b/Engine/Source/Runtime/Engine/Private/Materials/MaterialExpressionSubstrate.cpp
@@ -936,75 +936,6 @@ int32 UMaterialExpressionSubstrateSlabBSDF::Compile(class FMaterialCompiler* Com
 	return OutputCodeChunk;
 }
 
-int32 UMaterialExpressionSubstrateSlabBSDF::CompileDefaultSlab(FMaterialCompiler* Compiler, FVector3f EmissiveOverride)
-{
-	/**
-	 * This creates a default slab, primarily for material function previews where the input may be optional/missing.
-	 */
-	if (!Substrate::IsSubstrateEnabled())
-	{
-		return Compiler->SubstrateCreateAndRegisterNullMaterial();
-	}
-
-	FGuid PathUniqueId = Compiler->SubstrateTreeStackGetPathUniqueId();
-	FSubstrateOperator& SubstrateOperator = Compiler->SubstrateCompilationGetOperator(PathUniqueId);
-
-	int32 SSSProfileCodeChunk = INDEX_NONE;
-	int32 SpecularProfileCodeChunk = INDEX_NONE;
-	int32 TangentCodeChunk = INDEX_NONE;
-	int32 NormalCodeChunk = Compiler->VertexNormal();
-	int32 ThicknessCodeChunk = Compiler->Constant(SUBSTRATE_LAYER_DEFAULT_THICKNESS_CM);
-
-	const FSubstrateRegisteredSharedLocalBasis NewRegisteredSharedLocalBasis = SubstrateCompilationInfoCreateSharedLocalBasis(Compiler, NormalCodeChunk, TangentCodeChunk);
-	SubstrateOperator.BSDFRegisteredSharedLocalBasis = NewRegisteredSharedLocalBasis;
-
-	int32 DiffuseAlbedoCodeChunk = Compiler->Constant3(0.0f, 0.0f, 0.0f);
-	const float DefaultF0 = 0.04f;
-	int32 F0CodeChunk = Compiler->Constant3(DefaultF0, DefaultF0, DefaultF0);
-	int32 RoughnessCodeChunk = Compiler->Constant(0.5f);
-	int32 AnisotropyCodeChunk = Compiler->Constant(0.0f);
-	int32 F90CodeChunk = Compiler->Constant3(1.0f, 1.0f, 1.0f);
-	int32 SSSMFPCodeChunk = Compiler->Constant3(0.0f, 0.0f, 0.0f);
-	int32 SSSMFPScaleCodeChunk = Compiler->Constant(1.0f);
-	int32 SSSPhaseAnisotropyCodeChunk = Compiler->Constant(0.0f);
-	int32 SecondRoughnessCodeChunk = Compiler->Constant(0.0f);
-	int32 SecondRoughnessWeightCodeChunk = Compiler->Constant(0.0f);
-	int32 FuzzAmountCodeChunk = Compiler->Constant(0.0f);
-	int32 FuzzColorCodeChunk = Compiler->Constant3(0.0f, 0.0f, 0.0f);
-	int32 FuzzRoughnessCodeChunk = Compiler->Constant(0.5f);
-	int32 GlintValueCodeChunk = Compiler->Constant(0.0f);
-	int32 GlintUVCodeChunk = Compiler->Constant2(0.0f, 0.0f);
-
-	return Compiler->SubstrateSlabBSDF(
-		DiffuseAlbedoCodeChunk,
-		F0CodeChunk,
-		F90CodeChunk,
-		RoughnessCodeChunk,
-		AnisotropyCodeChunk,
-		Compiler->Constant(0.0f),
-		SSSMFPCodeChunk,
-		SSSMFPScaleCodeChunk,
-		SSSPhaseAnisotropyCodeChunk,
-		Compiler->Constant(EMaterialSubSurfaceType::MSS_None),
-		Compiler->Constant3(EmissiveOverride.X, EmissiveOverride.Y, EmissiveOverride.Z),
-		SecondRoughnessCodeChunk,
-		SecondRoughnessWeightCodeChunk,
-		Compiler->Constant(0.0f),										// SecondRoughnessAsSimpleClearCoat
-		NormalCodeChunk,
-		FuzzAmountCodeChunk,
-		FuzzColorCodeChunk,
-		FuzzRoughnessCodeChunk,
-		ThicknessCodeChunk,
-		GlintValueCodeChunk,
-		GlintUVCodeChunk,
-		Compiler->Constant(0.0f),
-		false, 															// bIsAtTheBottomOfTopology
-		NormalCodeChunk,
-		TangentCodeChunk,
-		Compiler->GetSubstrateSharedLocalBasisIndexMacro(NewRegisteredSharedLocalBasis),
-		(SubstrateOperator.Index != INDEX_NONE && SubstrateOperator.BSDFIndex != INDEX_NONE) && (!SubstrateOperator.bUseParameterBlending || (SubstrateOperator.bUseParameterBlending && SubstrateOperator.bRootOfParameterBlendingSubTree)) ? &SubstrateOperator : nullptr);
-}
-
 FSubstrateMaterialComplexity UMaterialExpressionSubstrateSlabBSDF::GetHighestComplexity() const
 {
 	// This function returns the highest complexity of a materials.
@@ -4024,6 +3955,63 @@ int32 UMaterialExpressionSubstrateConvertMaterialAttributes::Compile(class FMate
 		DummyInput, DummyInput);
 }
 
+int32 UMaterialExpressionSubstrateConvertMaterialAttributes::CompileDefaultSlab(FMaterialCompiler* Compiler, int32 EmissiveOverride)
+{
+	if (!Substrate::IsSubstrateEnabled())
+	{
+		return Compiler->SubstrateCreateAndRegisterNullMaterial();
+	}
+
+	FGuid PathUniqueId = Compiler->SubstrateTreeStackGetPathUniqueId();
+	FSubstrateOperator& SubstrateOperator = Compiler->SubstrateCompilationGetOperator(PathUniqueId);
+
+	int32 NormalCodeChunk = FMaterialAttributeDefinitionMap::CompileDefaultExpression(Compiler, MP_Normal);
+
+	const FSubstrateRegisteredSharedLocalBasis NewRegisteredSharedLocalBasis = SubstrateCompilationInfoCreateSharedLocalBasis(Compiler, NormalCodeChunk, INDEX_NONE);
+	const FString BasisIndexMacro = Compiler->GetSubstrateSharedLocalBasisIndexMacro(NewRegisteredSharedLocalBasis);
+
+	return Compiler->SubstrateConversionFromLegacy(
+		false,
+		// Metalness workflow
+		FMaterialAttributeDefinitionMap::CompileDefaultExpression(Compiler, MP_BaseColor),
+		FMaterialAttributeDefinitionMap::CompileDefaultExpression(Compiler, MP_Specular),
+		FMaterialAttributeDefinitionMap::CompileDefaultExpression(Compiler, MP_Metallic),
+		// Roughness
+		FMaterialAttributeDefinitionMap::CompileDefaultExpression(Compiler, MP_Roughness),
+		FMaterialAttributeDefinitionMap::CompileDefaultExpression(Compiler, MP_Anisotropy),
+		// SSS
+		FMaterialAttributeDefinitionMap::CompileDefaultExpression(Compiler, MP_SubsurfaceColor),
+		CreateDefaultSubsurfaceProfileParameter(Compiler),
+		// Clear Coat / Custom
+		FMaterialAttributeDefinitionMap::CompileDefaultExpression(Compiler, MP_CustomData0),// Clear coat
+		FMaterialAttributeDefinitionMap::CompileDefaultExpression(Compiler, MP_CustomData1),// Clear coat roughness
+		// Misc
+		EmissiveOverride != INDEX_NONE ? EmissiveOverride : FMaterialAttributeDefinitionMap::CompileDefaultExpression(Compiler, MP_EmissiveColor),
+		FMaterialAttributeDefinitionMap::CompileDefaultExpression(Compiler, MP_Opacity),
+		Compiler->Constant3(0.5f, 0.5f, 0.5f),
+		Compiler->Constant(1.0f),
+		// Water
+		Compiler->Constant3(0.0f, 0.0f, 0.0f),
+		Compiler->Constant3(0.0f, 0.0f, 0.0f),
+		Compiler->Constant(0.0f),
+		Compiler->Constant3(1.0f, 1.0f, 1.0f),
+		// Shading model
+		FMaterialAttributeDefinitionMap::CompileDefaultExpression(Compiler, MP_ShadingModel),
+		NormalCodeChunk,
+		INDEX_NONE,
+		BasisIndexMacro,
+		NormalCodeChunk,
+		INDEX_NONE,
+		BasisIndexMacro,
+		NormalCodeChunk,
+		(SubstrateOperator.Index != INDEX_NONE && SubstrateOperator.BSDFIndex != INDEX_NONE) && (!SubstrateOperator.bUseParameterBlending || (SubstrateOperator.bUseParameterBlending && SubstrateOperator.bRootOfParameterBlendingSubTree)) ? &SubstrateOperator : nullptr);
+}
+
+int32 UMaterialExpressionSubstrateConvertMaterialAttributes::CompileDefaultSlab(FMaterialCompiler* Compiler, FVector3f EmissiveOverride)
+{
+	return CompileDefaultSlab(Compiler, Compiler->Constant3(EmissiveOverride.X, EmissiveOverride.Y, EmissiveOverride.Z));
+}
+
 void UMaterialExpressionSubstrateConvertMaterialAttributes::PostEditChangeProperty(FPropertyChangedEvent& PropertyChangedEvent)
 {
 	if (GraphNode && PropertyChangedEvent.Property != nullptr)
diff --git a/Engine/Source/Runtime/Engine/Private/Materials/MaterialExpressions.cpp b/Engine/Source/Runtime/Engine/Private/Materials/MaterialExpressions.cpp
--- a/Engine/Source/Runtime/Engine/Private/Materials/MaterialExpressions.cpp
+++ b/Engine/Source/Runtime/Engine/Private/Materials/MaterialExpressions.cpp
@@ -16274,7 +16274,7 @@ int32 UMaterialExpressionFunctionInput::CompilePreviewValue(FMaterialCompiler* C
 		}
 		case FunctionInput_Substrate:
 		{
-			return UMaterialExpressionSubstrateSlabBSDF::CompileDefaultSlab(Compiler, FVector3f(PreviewValue.X, PreviewValue.Y, PreviewValue.Z));
+			return UMaterialExpressionSubstrateConvertMaterialAttributes::CompileDefaultSlab(Compiler, FVector3f(PreviewValue.X, PreviewValue.Y, PreviewValue.Z));
 		}
 		case FunctionInput_Texture2D:
 		case FunctionInput_TextureCube:
@@ -16760,13 +16760,17 @@ int32 UMaterialExpressionFunctionOutput::CompilePreview(class FMaterialCompiler*
 {
 	FExpressionInput ATraced = A.GetTracedInput();
 	UMaterialExpression* AExpression = ATraced.Expression;
-	if (AExpression && AExpression->IsResultSubstrateMaterial(ATraced.OutputIndex))
+
+	if (Substrate::IsSubstrateEnabled() && AExpression)
 	{
 		// Compile the SubstrateData output.
 		int32 SubstrateDataCodeChunk = Compile(Compiler, ATraced.OutputIndex);
+		if (!AExpression->IsResultSubstrateMaterial(ATraced.OutputIndex))
+		{
+			SubstrateDataCodeChunk = UMaterialExpressionSubstrateConvertMaterialAttributes::CompileDefaultSlab(Compiler, SubstrateDataCodeChunk);			
+		}
 		// Convert the SubstrateData to a preview color.
-		int32 PreviewCodeChunk = Compiler->SubstrateCompilePreview(SubstrateDataCodeChunk);
-		return PreviewCodeChunk;
+		return Compiler->SubstrateCompilePreview(SubstrateDataCodeChunk);
 	}
 
 	// Compile the preview value, outputting a float type
diff --git a/Engine/Source/Runtime/Engine/Public/Materials/MaterialExpressionSubstrate.h b/Engine/Source/Runtime/Engine/Public/Materials/MaterialExpressionSubstrate.h
--- a/Engine/Source/Runtime/Engine/Public/Materials/MaterialExpressionSubstrate.h
+++ b/Engine/Source/Runtime/Engine/Public/Materials/MaterialExpressionSubstrate.h
@@ -386,7 +386,6 @@ class UMaterialExpressionSubstrateSlabBSDF : public UMaterialExpressionSubstrate
 	//~ Begin UMaterialExpression Interface
 #if WITH_EDITOR
 	virtual int32 Compile(class FMaterialCompiler* Compiler, int32 OutputIndex) override;
-	static int32 CompileDefaultSlab(class FMaterialCompiler* Compiler, FVector3f EmissiveOverride = FVector3f::Zero());
 	virtual void GetCaption(TArray<FString>& OutCaptions) const override;
 	virtual EMaterialValueType GetOutputValueType(int32 OutputIndex) override;
 	virtual EMaterialValueType GetInputValueType(int32 InputIndex) override;
@@ -967,6 +966,8 @@ class UMaterialExpressionSubstrateConvertMaterialAttributes : public UMaterialEx
 #if WITH_EDITOR
 	virtual FExpressionInput* GetInput(int32 InputIndex) override;
 	virtual int32 Compile(class FMaterialCompiler* Compiler, int32 OutputIndex) override;
+	ENGINE_API static int32 CompileDefaultSlab(class FMaterialCompiler* Compiler, int32 EmissiveOverride);
+	ENGINE_API static int32 CompileDefaultSlab(class FMaterialCompiler* Compiler, FVector3f EmissiveOverride = FVector3f::Zero());
 	virtual void GetCaption(TArray<FString>& OutCaptions) const override;
 	virtual EMaterialValueType GetOutputValueType(int32 OutputIndex) override;
 	virtual EMaterialValueType GetInputValueType(int32 InputIndex) override;
