Fixed GeometryCache not binding Position/Tangent/MotionBlur atrributes when creating DummyVF

diff --git a/Engine/Plugins/Runtime/GeometryCache/Source/GeometryCache/Private/GeometryCacheSceneProxy.cpp b/Engine/Plugins/Runtime/GeometryCache/Source/GeometryCache/Private/GeometryCacheSceneProxy.cpp
--- a/Engine/Plugins/Runtime/GeometryCache/Source/GeometryCache/Private/GeometryCacheSceneProxy.cpp
+++ b/Engine/Plugins/Runtime/GeometryCache/Source/GeometryCache/Private/GeometryCacheSceneProxy.cpp
@@ -1519,6 +1519,8 @@ FGeomCacheVertexFactory::FGeomCacheVertexFactory(ERHIFeatureLevel::Type InFeatur
 
 void FGeomCacheVertexFactory::InitDummyComponentVF(FGeometryCacheVertexVertexFactory::FDataType& OutData)
 {
+	OutData.bIsDummyData = true;
+	
 	// Provide empty buffers for all entries
 	FGeomCacheVertexBuffer PositionBuffers[2];
 	FGeomCacheTangentBuffer TangentXBuffer;
diff --git a/Engine/Source/Runtime/Engine/Private/GeometryCacheVertexFactory.cpp b/Engine/Source/Runtime/Engine/Private/GeometryCacheVertexFactory.cpp
--- a/Engine/Source/Runtime/Engine/Private/GeometryCacheVertexFactory.cpp
+++ b/Engine/Source/Runtime/Engine/Private/GeometryCacheVertexFactory.cpp
@@ -190,7 +190,7 @@ void FGeometryCacheVertexVertexFactory::GetVertexElements(
 {
 	check(InputStreamType == EVertexInputStreamType::Default);
 
-	if (StreamComponentData.PositionComponent.VertexBuffer != NULL)
+	if (StreamComponentData.PositionComponent.VertexBuffer != NULL || StreamComponentData.bIsDummyData)
 	{
 		OutElements.Add(AccessStreamComponent(StreamComponentData.PositionComponent, 0, InOutStreams));
 		OutPositionStreamIndex = OutElements.Last().StreamIndex;
@@ -200,7 +200,7 @@ void FGeometryCacheVertexVertexFactory::GetVertexElements(
 	uint8 TangentBasisAttributes[2] = { 1, 2 };
 	for (int32 AxisIndex = 0; AxisIndex < 2; AxisIndex++)
 	{
-		if (StreamComponentData.TangentBasisComponents[AxisIndex].VertexBuffer != NULL)
+		if (StreamComponentData.TangentBasisComponents[AxisIndex].VertexBuffer != NULL || StreamComponentData.bIsDummyData)
 		{
 			OutElements.Add(AccessStreamComponent(StreamComponentData.TangentBasisComponents[AxisIndex], TangentBasisAttributes[AxisIndex], InOutStreams));
 		}
@@ -218,7 +218,7 @@ void FGeometryCacheVertexVertexFactory::GetVertexElements(
 		OutElements.Add(AccessStreamComponent(NullColorComponent, 3, InOutStreams));
 	}
 
-	if (StreamComponentData.MotionBlurDataComponent.VertexBuffer)
+	if (StreamComponentData.MotionBlurDataComponent.VertexBuffer || StreamComponentData.bIsDummyData)
 	{
 		OutElements.Add(AccessStreamComponent(StreamComponentData.MotionBlurDataComponent, 4, InOutStreams));
 	}
diff --git a/Engine/Source/Runtime/Engine/Public/GeometryCacheVertexFactory.h b/Engine/Source/Runtime/Engine/Public/GeometryCacheVertexFactory.h
--- a/Engine/Source/Runtime/Engine/Public/GeometryCacheVertexFactory.h
+++ b/Engine/Source/Runtime/Engine/Public/GeometryCacheVertexFactory.h
@@ -97,6 +97,8 @@ public:
 
 		/** The stream to read the motion blur data from. */
 		FVertexStreamComponent MotionBlurDataComponent;
+		
+		bool bIsDummyData = false;
 	};
 
 	static ENGINE_API void ModifyCompilationEnvironment(const FVertexFactoryShaderPermutationParameters& Parameters, FShaderCompilerEnvironment& OutEnvironment);
