[NetPhysics] Stop log spam when client is out of rewind bounds in networked physics.
- Print a warning informing it's out of bounds and when it's been inside bounds for 5s print a log saying it's back in sync.

diff --git a/Engine/Source/Runtime/Engine/Private/PhysicsEngine/PhysicsReplication.cpp b/Engine/Source/Runtime/Engine/Private/PhysicsEngine/PhysicsReplication.cpp
--- a/Engine/Source/Runtime/Engine/Private/PhysicsEngine/PhysicsReplication.cpp
+++ b/Engine/Source/Runtime/Engine/Private/PhysicsEngine/PhysicsReplication.cpp
@@ -144,6 +144,9 @@ namespace PhysicsReplicationCVars
 
 		bool bDrawDebug = false;
 		static FAutoConsoleVariableRef CVarResimDrawDebug(TEXT("np2.Resim.DrawDebug"), bDrawDebug, TEXT("Resimulation debug draw-calls"));
+
+		float LogOutOfBoundsTimeLimit = 5.0f;
+		static FAutoConsoleVariableRef CVarResimLogOutOfBoundsTimeLimit(TEXT("np2.Resim.LogOutOfBoundsTimeLimit"), LogOutOfBoundsTimeLimit, TEXT("Time that received targets needs to be within rewind bounds again before printing log that we are back in sync again. A new desync log will not be printed before an in-sync log has been printed."));
 	}
 
 	namespace PredictiveInterpolationCVars
@@ -919,6 +922,8 @@ void FPhysicsReplicationAsync::OnPreSimulate_Internal()
 	Chaos::FPBDRigidsSolver* RigidsSolver = static_cast<Chaos::FPBDRigidsSolver*>(GetSolver());
 	check(RigidsSolver);
 
+	ResimErrorLogTimer += RigidsSolver->GetAsyncDeltaTime();
+
 	// Early out if this is a resim frame
 	Chaos::FRewindData* RewindData = RigidsSolver->GetRewindData();
 	const bool bRewindDataExist = RewindData != nullptr;
@@ -1386,10 +1391,28 @@ void FPhysicsReplicationAsync::CheckTargetResimValidity(FReplicatedPhysicsTarget
 
 			Target.RepMode = EPhysicsReplicationMode::PredictiveInterpolation;
 		}
+		
+		if (ResimOutOfBoundsCounter == 0)
+		{
+			UE_LOG(LogPhysics, Warning, TEXT("FPhysicsReplication DESYNCED - received target frame (%d) out of rewind data bounds (%d, %d) - %s - Target will use %s")
+				, LocalFrame, RewindData->GetEarliestFrame_Internal(), RewindData->CurrentFrame()
+					, (LocalFrame < RewindData->GetEarliestFrame_Internal()) 
+					? TEXT("Client is far ahead of the server, server might be dropping frames.") 
+					: TEXT("Client is behind the server, client might be dropping frames."), *UEnum::GetValueAsString(Target.RepMode));
+		}
+
+		ResimOutOfBoundsCounter++;
+		ResimErrorLogTimer = 0;
+	}
+	else
+	{
+		if (ResimOutOfBoundsCounter > 0 && ResimErrorLogTimer > PhysicsReplicationCVars::ResimulationCVars::LogOutOfBoundsTimeLimit)
+		{
+			UE_LOG(LogPhysics, Log, TEXT("FPhysicsReplication IN-SYNC - Received targets have now been within rewind data bounds again for at least %f seconds"), ResimErrorLogTimer);
 
-		UE_LOG(LogPhysics, Warning, TEXT("FPhysicsReplication received target frame (%d) out of rewind data bounds (%d, %d) - %s - Target will use EPhysicsReplicationMode: %s"),
-			LocalFrame, RewindData->GetEarliestFrame_Internal(), RewindData->CurrentFrame(), (LocalFrame < RewindData->GetEarliestFrame_Internal()) ? TEXT("Client is far ahead of the server, server might be dropping frames.") : TEXT("Client is behind the server, client might be dropping frames."), *UEnum::GetValueAsString(Target.RepMode));
-	}	
+			ResimOutOfBoundsCounter = 0;
+		}
+	}
 }
 
 void FPhysicsReplicationAsync::ApplyPhysicsReplicationLOD(Chaos::FConstPhysicsObjectHandle PhysicsObjectHandle, FReplicatedPhysicsTargetAsync& Target, const uint32 LODFLags)
diff --git a/Engine/Source/Runtime/Engine/Public/PhysicsReplication.h b/Engine/Source/Runtime/Engine/Public/PhysicsReplication.h
--- a/Engine/Source/Runtime/Engine/Public/PhysicsReplication.h
+++ b/Engine/Source/Runtime/Engine/Public/PhysicsReplication.h
@@ -221,6 +221,9 @@ private:
 	TArray<int32> ParticlesInResimIslands;
 	TArray<Chaos::FParticleID> ReplicatedParticleIDs;
 
+	int32 ResimOutOfBoundsCounter = 0;
+	float ResimErrorLogTimer = 0.0f;
+
 private:
 	FReplicatedPhysicsTargetAsync* AddObjectToReplication(Chaos::FConstPhysicsObjectHandle PhysicsObject);
 	void RemoveObjectFromReplication(Chaos::FConstPhysicsObjectHandle PhysicsObject);
