Fix crash when tweaking Substrate Closure count in the UI and the value goes below the previous value.

This CL avoids the settings to be actually propagated to the renderer (which pulls value from CVar), and ensure the value will be only used once the editor restarts.

diff --git a/Engine/Source/Runtime/Engine/Classes/Engine/RendererSettings.h b/Engine/Source/Runtime/Engine/Classes/Engine/RendererSettings.h
--- a/Engine/Source/Runtime/Engine/Classes/Engine/RendererSettings.h
+++ b/Engine/Source/Runtime/Engine/Classes/Engine/RendererSettings.h
@@ -1552,8 +1552,8 @@ private:
 	void FixAntiAliasingOnShadingPathChange(FPropertyChangedEvent& PropertyChangedEvent);
 #endif // WITH_EDITOR
 
-	void SanatizeReflectionCaptureResolution();
-	void SanatizeSubstrateSettings();
+	void SanitizeReflectionCaptureResolution();
+	void SanitizeSubstrateSettings(bool bSyncConsoleVariables);
 	void UpdateWorkingColorSpaceAndChromaticities();
 };
 
diff --git a/Engine/Source/Runtime/Engine/Private/RendererSettings.cpp b/Engine/Source/Runtime/Engine/Private/RendererSettings.cpp
--- a/Engine/Source/Runtime/Engine/Private/RendererSettings.cpp
+++ b/Engine/Source/Runtime/Engine/Private/RendererSettings.cpp
@@ -117,11 +117,11 @@ void URendererSettings::PostInitProperties()
 {
 	Super::PostInitProperties();
 	
-	SanatizeReflectionCaptureResolution();
+	SanitizeReflectionCaptureResolution();
 
 	UpdateWorkingColorSpaceAndChromaticities();
 
-	SanatizeSubstrateSettings();
+	SanitizeSubstrateSettings(true /*bSyncConsoleVariables*/);
 
 #if WITH_EDITOR
 	CheckForMissingShaderModels();
@@ -255,7 +255,7 @@ void URendererSettings::PostEditChangeProperty(FPropertyChangedEvent& PropertyCh
 {
 	Super::PostEditChangeProperty(PropertyChangedEvent);
 
-	SanatizeReflectionCaptureResolution();
+	SanitizeReflectionCaptureResolution();
 
 	if (PropertyChangedEvent.Property)
 	{
@@ -443,7 +443,10 @@ void URendererSettings::PostEditChangeProperty(FPropertyChangedEvent& PropertyCh
 		if (PropertyChangedEvent.Property->GetFName() == GET_MEMBER_NAME_CHECKED(URendererSettings, SubstrateGBufferFormat) ||
 			PropertyChangedEvent.Property->GetFName() == GET_MEMBER_NAME_CHECKED(URendererSettings, SubstrateProjectClosuresPerPixel))
 		{
-			SanatizeSubstrateSettings();
+			// These settings require a restart
+			// Do not update the console variable to avoid a mismatch between the settings used for materials compilation, and the 'new' settings. 
+			// The 'new' settings will be applied on the next restart
+			SanitizeSubstrateSettings(false /*SanitizeSubstrateSettings*/);
 		}
 	}
 }
@@ -686,7 +689,7 @@ void URendererSettings::CheckForMissingShaderModels()
 }
 #endif // #if WITH_EDITOR
 
-void URendererSettings::SanatizeReflectionCaptureResolution()
+void URendererSettings::SanitizeReflectionCaptureResolution()
 {
 	const int32 MaxCubemapResolution = GetMaxCubeTextureDimension();
 	const int32 MinCubemapResolution = 8;
@@ -731,7 +734,7 @@ void URendererSettings::SanatizeReflectionCaptureResolution()
 #endif // WITH_EDITOR
 }
 
-void URendererSettings::SanatizeSubstrateSettings()
+void URendererSettings::SanitizeSubstrateSettings(bool bSyncConsoleVariables)
 {
 	// When blendable GBuffer is used, force project closure count to 1 for clarity
 	if (SubstrateGBufferFormat == ESubstrateStorageFormat::Type::BlendableGBuffer)
@@ -741,9 +744,12 @@ void URendererSettings::SanatizeSubstrateSettings()
 
 	// Ensure project settings and the attached CVar are consistent
 	// Don't use ExportValuesToConsoleVariables() as it filters out ReadOnly CVars
-	if (IConsoleVariable* CVar = IConsoleManager::Get().FindConsoleVariable(TEXT("r.Substrate.ProjectClosuresPerPixel")))
+	if (bSyncConsoleVariables)
 	{
-		CVar->Set((int32)SubstrateProjectClosuresPerPixel, ECVF_SetByProjectSetting);
+		if (IConsoleVariable* CVar = IConsoleManager::Get().FindConsoleVariable(TEXT("r.Substrate.ProjectClosuresPerPixel")))
+		{
+			CVar->Set((int32)SubstrateProjectClosuresPerPixel, ECVF_SetByProjectSetting);
+		}
 	}
 }
 
