Improvement to fix 45623939. It moves the material update context to the material editing library instead of inside the SetParentEditorOnly().

#rnx

diff --git a/Engine/Source/Editor/MaterialEditor/Private/MaterialEditingLibrary.cpp b/Engine/Source/Editor/MaterialEditor/Private/MaterialEditingLibrary.cpp
--- a/Engine/Source/Editor/MaterialEditor/Private/MaterialEditingLibrary.cpp
+++ b/Engine/Source/Editor/MaterialEditor/Private/MaterialEditingLibrary.cpp
@@ -1138,6 +1138,7 @@ void UMaterialEditingLibrary::SetMaterialInstanceParent(UMaterialInstanceConstan
 {
 	if (Instance)
 	{
+		FMaterialUpdateContext UpdateContext;
 		Instance->SetParentEditorOnly(NewParent);
 	}
 }
diff --git a/Engine/Source/Runtime/Engine/Private/Materials/MaterialInstanceConstant.cpp b/Engine/Source/Runtime/Engine/Private/Materials/MaterialInstanceConstant.cpp
--- a/Engine/Source/Runtime/Engine/Private/Materials/MaterialInstanceConstant.cpp
+++ b/Engine/Source/Runtime/Engine/Private/Materials/MaterialInstanceConstant.cpp
@@ -82,7 +82,6 @@ void UMaterialInstanceConstant::SetParentEditorOnly(UMaterialInterface* NewParen
 
 	if (SetParentInternal(NewParent, RecacheShader))
 	{
-		FMaterialUpdateContext UpdateContext;
 		ValidateStaticPermutationAllowed();
 		UpdateCachedData();
 	}
