Fix mismatch between CVar and render settings at initialization time.

diff --git a/Engine/Source/Runtime/Engine/Private/RendererSettings.cpp b/Engine/Source/Runtime/Engine/Private/RendererSettings.cpp
--- a/Engine/Source/Runtime/Engine/Private/RendererSettings.cpp
+++ b/Engine/Source/Runtime/Engine/Private/RendererSettings.cpp
@@ -121,7 +121,6 @@ void URendererSettings::PostInitProperties()
 
 	UpdateWorkingColorSpaceAndChromaticities();
 
-	SanitizeSubstrateSettings(true /*bSyncConsoleVariables*/);
 
 #if WITH_EDITOR
 	CheckForMissingShaderModels();
@@ -131,6 +130,9 @@ void URendererSettings::PostInitProperties()
 		ImportConsoleVariableValues();
 	}
 #endif // #if WITH_EDITOR
+
+	// Only sanitize Substrate value "after" the CVar import, to be sure default value are correct.
+	SanitizeSubstrateSettings(true /*bSyncConsoleVariables*/);
 }
 
 #if WITH_EDITOR
