Expose Susbtrate's byte per pixel settings in to the render settings.

diff --git a/Engine/Source/Runtime/Engine/Classes/Engine/RendererSettings.h b/Engine/Source/Runtime/Engine/Classes/Engine/RendererSettings.h
--- a/Engine/Source/Runtime/Engine/Classes/Engine/RendererSettings.h
+++ b/Engine/Source/Runtime/Engine/Classes/Engine/RendererSettings.h
@@ -1243,6 +1243,18 @@ class URendererSettings : public UDeveloperSettings
 		ConfigRestartRequired = true))
 	TEnumAsByte<ESubstrateClosureConfig::Type> SubstrateClosurePerPixelOverride;
 
+	/**
+	"Substrate maximum number of bytes allocated per pixels for storing GBuffer information for this project with Adaptive GBuffer."
+	*/
+	UPROPERTY(config, EditAnywhere, Category = Substrate, meta = (
+		ConsoleVariable = "r.Substrate.BytesPerPixel", 
+		DisplayName = "Substrate Bytes Per Pixel",
+		ToolTip = "Define the maximum number of bytes per pixel allocated. Materials requiring memory beyond this limit will be simplified until it fits the requirements. This only applies when Adaptive GBuffer is enabled. On platforms not supporting Adaptive GBuffer, the material will be simplified to a smaller fix limit.",
+		ClampMin = "20", ClampMax = "256", 
+		EditCondition = "SubstrateGBufferFormat == 1",
+		ConfigRestartRequired = true))
+	uint32 SubstrateProjectBytesPerPixel;
+
 	/**
 	"Enable Substrate opaque material rough refractions effect from top layers over layers below."
 	*/
diff --git a/Engine/Source/Runtime/Engine/Private/RendererSettings.cpp b/Engine/Source/Runtime/Engine/Private/RendererSettings.cpp
--- a/Engine/Source/Runtime/Engine/Private/RendererSettings.cpp
+++ b/Engine/Source/Runtime/Engine/Private/RendererSettings.cpp
@@ -443,7 +443,8 @@ void URendererSettings::PostEditChangeProperty(FPropertyChangedEvent& PropertyCh
 
 		// Update Substrate closure count if project's GBuffer format change
 		if (PropertyChangedEvent.Property->GetFName() == GET_MEMBER_NAME_CHECKED(URendererSettings, SubstrateGBufferFormat) ||
-			PropertyChangedEvent.Property->GetFName() == GET_MEMBER_NAME_CHECKED(URendererSettings, SubstrateProjectClosuresPerPixel))
+			PropertyChangedEvent.Property->GetFName() == GET_MEMBER_NAME_CHECKED(URendererSettings, SubstrateProjectClosuresPerPixel) ||
+			PropertyChangedEvent.Property->GetFName() == GET_MEMBER_NAME_CHECKED(URendererSettings, SubstrateProjectBytesPerPixel))
 		{
 			// These settings require a restart
 			// Do not update the console variable to avoid a mismatch between the settings used for materials compilation, and the 'new' settings. 
@@ -742,6 +743,7 @@ void URendererSettings::SanitizeSubstrateSettings(bool bSyncConsoleVariables)
 	if (SubstrateGBufferFormat == ESubstrateStorageFormat::Type::BlendableGBuffer)
 	{
 		SubstrateProjectClosuresPerPixel = 1;
+		SubstrateProjectBytesPerPixel = 20u;
 	}
 
 	// Ensure project settings and the attached CVar are consistent
@@ -752,6 +754,11 @@ void URendererSettings::SanitizeSubstrateSettings(bool bSyncConsoleVariables)
 		{
 			CVar->Set((int32)SubstrateProjectClosuresPerPixel, ECVF_SetByProjectSetting);
 		}
+
+		if (IConsoleVariable* CVar = IConsoleManager::Get().FindConsoleVariable(TEXT("r.Substrate.BytesPerPixel")))
+		{
+			CVar->Set((int32)SubstrateProjectBytesPerPixel, ECVF_SetByProjectSetting);
+		}
 	}
 }
 
