World Partition - HLOD: Store an HLOD build report property in the HLOD actors
* Contains a ton of information that will help understand why HLOD builds have been performed
* Stored as text directly in the HLOD actor .uasset file, allowing diff tool to pick it up

diff --git a/Engine/Plugins/Editor/WorldPartitionHLODUtilities/Source/Private/WorldPartition/HLOD/Builders/HLODBuilderInstancing.cpp b/Engine/Plugins/Editor/WorldPartitionHLODUtilities/Source/Private/WorldPartition/HLOD/Builders/HLODBuilderInstancing.cpp
--- a/Engine/Plugins/Editor/WorldPartitionHLODUtilities/Source/Private/WorldPartition/HLOD/Builders/HLODBuilderInstancing.cpp
+++ b/Engine/Plugins/Editor/WorldPartitionHLODUtilities/Source/Private/WorldPartition/HLOD/Builders/HLODBuilderInstancing.cpp
@@ -1,6 +1,8 @@
 // Copyright Epic Games, Inc. All Rights Reserved.
 
 #include "WorldPartition/HLOD/Builders/HLODBuilderInstancing.h"
+#include "WorldPartition/HLOD/HLODHashBuilder.h"
+
 #include "Serialization/ArchiveCrc32.h"
 #include "Engine/StaticMesh.h"
 #include "Components/StaticMeshComponent.h"
@@ -18,39 +20,31 @@ UHLODBuilderInstancingSettings::UHLODBuilderInstancingSettings(const FObjectInit
 {
 }
 
-uint32 UHLODBuilderInstancingSettings::GetCRC() const
+void UHLODBuilderInstancingSettings::ComputeHLODHash(FHLODHashBuilder& InHashBuilder) const
 {
-	UHLODBuilderInstancingSettings& This = *const_cast<UHLODBuilderInstancingSettings*>(this);
-
-	FArchiveCrc32 Ar;
-
 	// Base key, changing this will force a rebuild of all HLODs from this builder
-	FString HLODBaseKey = "5197F75BF850402EABA5D29E4D8BB1DA";
-	Ar << HLODBaseKey;
-
-	Ar << This.bDisallowNanite;
-	UE_LOG(LogHLODBuilder, VeryVerbose, TEXT(" - bDisallowNanite = %d"), Ar.GetCrc());
+	FString HLODBaseKey = TEXT("5197F75BF850402EABA5D29E4D8BB1DA");
+	InHashBuilder.HashField(HLODBaseKey, TEXT("MeshInstancingBaseKey"));
+	InHashBuilder.HashField(bDisallowNanite, TEXT("DissallowNanite"));
 
 	if (InstanceFilteringType != EInstanceFilteringType::FilterNone)
 	{
-		Ar << This.InstanceFilteringType;
+		InHashBuilder.HashField(InstanceFilteringType, TEXT("InstanceFilteringType"));
 		switch (InstanceFilteringType)
 		{
 		case EInstanceFilteringType::FilterMinimumExtent:
-			Ar << This.MinimumExtent;
+			InHashBuilder.HashField(MinimumExtent, TEXT("MinimumExtent"));
 			break;
 		case EInstanceFilteringType::FilterMinimumArea:
-			Ar << This.MinimumArea;
+			InHashBuilder.HashField(MinimumArea, TEXT("MinimumArea"));
 			break;
 		case EInstanceFilteringType::FilterMinimumVolume:
-			Ar << This.MinimumVolume;
+			InHashBuilder.HashField(MinimumVolume, TEXT("MinimumVolume"));
 			break;
 		default:
 			unimplemented();
 		}
 	}
-
-	return Ar.GetCrc();
 }
 
 
diff --git a/Engine/Plugins/Editor/WorldPartitionHLODUtilities/Source/Private/WorldPartition/HLOD/Builders/HLODBuilderMeshApproximate.cpp b/Engine/Plugins/Editor/WorldPartitionHLODUtilities/Source/Private/WorldPartition/HLOD/Builders/HLODBuilderMeshApproximate.cpp
--- a/Engine/Plugins/Editor/WorldPartitionHLODUtilities/Source/Private/WorldPartition/HLOD/Builders/HLODBuilderMeshApproximate.cpp
+++ b/Engine/Plugins/Editor/WorldPartitionHLODUtilities/Source/Private/WorldPartition/HLOD/Builders/HLODBuilderMeshApproximate.cpp
@@ -1,6 +1,7 @@
 // Copyright Epic Games, Inc. All Rights Reserved.
 
 #include "WorldPartition/HLOD/Builders/HLODBuilderMeshApproximate.h"
+#include "WorldPartition/HLOD/HLODHashBuilder.h"
 
 #include "Algo/ForEach.h"
 #include "Components/StaticMeshComponent.h"
@@ -40,29 +41,13 @@ UHLODBuilderMeshApproximateSettings::UHLODBuilderMeshApproximateSettings(const F
 #endif
 }
 
-uint32 UHLODBuilderMeshApproximateSettings::GetCRC() const
+void UHLODBuilderMeshApproximateSettings::ComputeHLODHash(FHLODHashBuilder& InHashBuilder) const
 {
-	UHLODBuilderMeshApproximateSettings& This = *const_cast<UHLODBuilderMeshApproximateSettings*>(this);
-
-	FArchiveCrc32 Ar;
-
 	// Base key, changing this will force a rebuild of all HLODs from this builder
-	FString HLODBaseKey = "4E9BB187B38645A08B424895C5268B8A";
-	Ar << HLODBaseKey;
-
-	Ar << This.MeshApproximationSettings;
-	UE_LOG(LogHLODBuilder, VeryVerbose, TEXT(" - MeshApproximationSettings = %d"), Ar.GetCrc());
-
-	uint32 Hash = Ar.GetCrc();
-
-	if (HLODMaterial)
-	{
-		uint32 MaterialCRC = UHLODProxy::GetCRC(HLODMaterial);
-		UE_LOG(LogHLODBuilder, VeryVerbose, TEXT(" - Material = %d"), MaterialCRC);
-		Hash = HashCombineFast(Hash, MaterialCRC);
-	}
-
-	return Hash;
+	FString HLODBaseKey = TEXT("4E9BB187B38645A08B424895C5268B8A");
+	InHashBuilder.HashField(HLODBaseKey, TEXT("MeshApproximateBaseKey"));
+	InHashBuilder.HashField(MeshApproximationSettings, TEXT("MeshApproximationSettingsHash"));
+	InHashBuilder << HLODMaterial;
 }
 
 TSubclassOf<UHLODBuilderSettings> UHLODBuilderMeshApproximate::GetSettingsClass() const
diff --git a/Engine/Plugins/Editor/WorldPartitionHLODUtilities/Source/Private/WorldPartition/HLOD/Builders/HLODBuilderMeshMerge.cpp b/Engine/Plugins/Editor/WorldPartitionHLODUtilities/Source/Private/WorldPartition/HLOD/Builders/HLODBuilderMeshMerge.cpp
--- a/Engine/Plugins/Editor/WorldPartitionHLODUtilities/Source/Private/WorldPartition/HLOD/Builders/HLODBuilderMeshMerge.cpp
+++ b/Engine/Plugins/Editor/WorldPartitionHLODUtilities/Source/Private/WorldPartition/HLOD/Builders/HLODBuilderMeshMerge.cpp
@@ -1,6 +1,7 @@
 // Copyright Epic Games, Inc. All Rights Reserved.
 
 #include "WorldPartition/HLOD/Builders/HLODBuilderMeshMerge.h"
+#include "WorldPartition/HLOD/HLODHashBuilder.h"
 
 #include "Algo/ForEach.h"
 #include "Engine/Engine.h"
@@ -52,40 +53,25 @@ UHLODBuilderMeshMergeSettings::UHLODBuilderMeshMergeSettings(const FObjectInitia
 	}
 }
 
-uint32 UHLODBuilderMeshMergeSettings::GetCRC() const
+void UHLODBuilderMeshMergeSettings::ComputeHLODHash(FHLODHashBuilder& InHashBuilder) const
 {
-	UHLODBuilderMeshMergeSettings& This = *const_cast<UHLODBuilderMeshMergeSettings*>(this);
-
-	FArchiveCrc32 Ar;
-
 	// Base mesh merge key, changing this will force a rebuild of all HLODs from this builder
-	FString HLODBaseKey = "F906DCC49BCB4102A821B20D64C6E93B";
-	Ar << HLODBaseKey;
+	FString HLODBaseKey = TEXT("F906DCC49BCB4102A821B20D64C6E93B");
+	InHashBuilder.HashField(HLODBaseKey, TEXT("MeshMergeBaseKey"));
 
-	Ar << This.MeshMergeSettings;
-	UE_LOG(LogHLODBuilder, VeryVerbose, TEXT(" - MeshMergeSettings = %d"), Ar.GetCrc());
+	InHashBuilder.HashField(MeshMergeSettings, TEXT("MeshMergeSettingsHash"));
 
 	if (MeshMergeSettings.bMergeMaterials)
 	{
 		IMaterialBakingModule& Module = FModuleManager::Get().LoadModuleChecked<IMaterialBakingModule>("MaterialBaking");
-		uint32 MaterialBakingModuleCRC = Module.GetCRC();
-		Ar << MaterialBakingModuleCRC;
+		InHashBuilder.HashField(Module.GetCRC(), TEXT("MaterialBakingModuleHash"));
 
 		static const auto MeshMergeUtilitiesUVGenerationMethodCVar = IConsoleManager::Get().FindConsoleVariable(TEXT("MeshMergeUtilities.UVGenerationMethod"));
 		int32 MeshMergeUtilitiesUVGenerationMethod = (MeshMergeUtilitiesUVGenerationMethodCVar != nullptr) ? MeshMergeUtilitiesUVGenerationMethodCVar->GetInt() : 0;
-		Ar << MeshMergeUtilitiesUVGenerationMethod;
-	}
-
-	uint32 Hash = Ar.GetCrc();
-
-	if (HLODMaterial)
-	{
-		uint32 MaterialCRC = UHLODProxy::GetCRC(HLODMaterial);
-		UE_LOG(LogHLODBuilder, VeryVerbose, TEXT(" - Material = %d"), MaterialCRC);
-		Hash = HashCombineFast(Hash, MaterialCRC);
+		InHashBuilder.HashField(MeshMergeUtilitiesUVGenerationMethod, TEXT("MeshMergeUtilitiesUVGenerationMethod"));
 	}
 
-	return Hash;
+	InHashBuilder << HLODMaterial;
 }
 
 TSubclassOf<UHLODBuilderSettings> UHLODBuilderMeshMerge::GetSettingsClass() const
diff --git a/Engine/Plugins/Editor/WorldPartitionHLODUtilities/Source/Private/WorldPartition/HLOD/Builders/HLODBuilderMeshSimplify.cpp b/Engine/Plugins/Editor/WorldPartitionHLODUtilities/Source/Private/WorldPartition/HLOD/Builders/HLODBuilderMeshSimplify.cpp
--- a/Engine/Plugins/Editor/WorldPartitionHLODUtilities/Source/Private/WorldPartition/HLOD/Builders/HLODBuilderMeshSimplify.cpp
+++ b/Engine/Plugins/Editor/WorldPartitionHLODUtilities/Source/Private/WorldPartition/HLOD/Builders/HLODBuilderMeshSimplify.cpp
@@ -1,6 +1,7 @@
 // Copyright Epic Games, Inc. All Rights Reserved.
 
 #include "WorldPartition/HLOD/Builders/HLODBuilderMeshSimplify.h"
+#include "WorldPartition/HLOD/HLODHashBuilder.h"
 
 #include "Algo/ForEach.h"
 #include "Engine/Engine.h"
@@ -37,37 +38,23 @@ UHLODBuilderMeshSimplifySettings::UHLODBuilderMeshSimplifySettings(const FObject
 #endif
 }
 
-uint32 UHLODBuilderMeshSimplifySettings::GetCRC() const
+void UHLODBuilderMeshSimplifySettings::ComputeHLODHash(FHLODHashBuilder& InHashBuilder) const
 {
-	UHLODBuilderMeshSimplifySettings& This = *const_cast<UHLODBuilderMeshSimplifySettings*>(this);
-
-	FArchiveCrc32 Ar;
-
 	// Base mesh simplify key, changing this will force a rebuild of all HLODs from this builder
-	FString HLODBaseKey = "833ECCC91448424CACC620661712928A";
-	Ar << HLODBaseKey;
+	FString HLODBaseKey = TEXT("833ECCC91448424CACC620661712928A");
+	InHashBuilder.HashField(HLODBaseKey, TEXT("MeshSimplifyBaseKey"));
 
-	Ar << This.MeshSimplifySettings;
-	UE_LOG(LogHLODBuilder, VeryVerbose, TEXT(" - MeshSimplifySettings = %d"), Ar.GetCrc());
+	InHashBuilder.HashField(MeshSimplifySettings, TEXT("MeshSimplifySettingsHash"));
 
 	IMaterialBakingModule& Module = FModuleManager::Get().LoadModuleChecked<IMaterialBakingModule>("MaterialBaking");
 	uint32 MaterialBakingModuleCRC = Module.GetCRC();
-	Ar << MaterialBakingModuleCRC;
+	InHashBuilder.HashField(MaterialBakingModuleCRC, TEXT("MaterialBakingModuleCRC"));
 
 	static const auto MeshMergeUtilitiesUVGenerationMethodCVar = IConsoleManager::Get().FindConsoleVariable(TEXT("MeshMergeUtilities.UVGenerationMethod"));
 	int32 MeshMergeUtilitiesUVGenerationMethod = (MeshMergeUtilitiesUVGenerationMethodCVar != nullptr) ? MeshMergeUtilitiesUVGenerationMethodCVar->GetInt() : 0;
-	Ar << MeshMergeUtilitiesUVGenerationMethod;
-
-	uint32 Hash = Ar.GetCrc();
-
-	if (HLODMaterial)
-	{
-		uint32 MaterialCRC = UHLODProxy::GetCRC(HLODMaterial);
-		UE_LOG(LogHLODBuilder, VeryVerbose, TEXT(" - Material = %d"), MaterialCRC);
-		Hash = HashCombineFast(Hash, MaterialCRC);
-	}
+	InHashBuilder.HashField(MeshMergeUtilitiesUVGenerationMethod, TEXT("MeshMergeUtilitiesUVGenerationMethod"));
 
-	return Hash;
+	InHashBuilder << HLODMaterial;
 }
 
 TSubclassOf<UHLODBuilderSettings> UHLODBuilderMeshSimplify::GetSettingsClass() const
diff --git a/Engine/Plugins/Editor/WorldPartitionHLODUtilities/Source/Private/WorldPartition/HLOD/Utilities/WorldPartitionHLODUtilities.cpp b/Engine/Plugins/Editor/WorldPartitionHLODUtilities/Source/Private/WorldPartition/HLOD/Utilities/WorldPartitionHLODUtilities.cpp
--- a/Engine/Plugins/Editor/WorldPartitionHLODUtilities/Source/Private/WorldPartition/HLOD/Utilities/WorldPartitionHLODUtilities.cpp
+++ b/Engine/Plugins/Editor/WorldPartitionHLODUtilities/Source/Private/WorldPartition/HLOD/Utilities/WorldPartitionHLODUtilities.cpp
@@ -10,6 +10,7 @@
 #include "WorldPartition/ContentBundle/ContentBundleActivationScope.h"
 #include "WorldPartition/DataLayer/DataLayerInstance.h"
 #include "WorldPartition/HLOD/HLODActor.h"
+#include "WorldPartition/HLOD/HLODHashBuilder.h"
 #include "WorldPartition/HLOD/HLODLayer.h"
 #include "WorldPartition/HLOD/HLODModifier.h"
 #include "WorldPartition/HLOD/HLODSourceActorsFromCell.h"
@@ -32,6 +33,7 @@
 #include "Engine/StaticMesh.h"
 #include "Engine/Texture.h"
 #include "Engine/World.h"
+#include "EngineUtils.h"
 #include "Materials/MaterialInstance.h"
 #include "PhysicsEngine/BodySetup.h"
 #include "ProfilingDebugging/ScopedTimers.h"
@@ -50,41 +52,33 @@ TAutoConsoleVariable<bool> CVarHLODUseLegacy32BitNameHash(
 	TEXT("Force the use of 32-bit hashing to use when generating HLOD actor names. Switching to 64-bit will result in new names/new actor files."),
 	ECVF_ReadOnly);
 
-static uint32 ComputeHLODHash(AWorldPartitionHLOD* InHLODActor, const TArray<UActorComponent*>& InSourceComponents)
+static void ComputeHLODHash(FHLODHashBuilder& InHashBuilder, AWorldPartitionHLOD* InHLODActor, const TArray<UActorComponent*>& InSourceComponents)
 {
-	FArchiveCrc32 Ar;
-
 	// Base key, changing this will force a rebuild of all HLODs
-	FString HLODBaseKey = "0D33837AB1A04CC2AEEB04C5C217DE96";
-	Ar << HLODBaseKey;
+	FString HLODBaseKey = TEXT("0D33837AB1A04CC2AEEB04C5C217DE96");
+	InHashBuilder.HashField(HLODBaseKey, TEXT("HLODBaseKey"));
 
 	// HLOD Source Actors
 	if (ensure(InHLODActor->GetSourceActors()))
 	{
-		uint32 SourceActorsHash = InHLODActor->GetSourceActors()->GetHLODHash();
-		UE_LOG(LogHLODHash, VeryVerbose, TEXT(" - SourceActorsHash = %x"), SourceActorsHash);
-		Ar << SourceActorsHash;
+		InHLODActor->GetSourceActors()->ComputeHLODHash(InHashBuilder);
 	}
 	
 	// Min Visible Distance
-	double MinVisibleDistance = InHLODActor->GetMinVisibleDistance();
-	UE_LOG(LogHLODHash, VeryVerbose, TEXT(" - HLOD Min Visible Distance = %.02f"), InHLODActor->GetMinVisibleDistance());
-	Ar << MinVisibleDistance;
+	InHashBuilder.HashField(InHLODActor->GetMinVisibleDistance(), TEXT("MinVisibleDistance"));
 
 	// ISM Component Class
 	TSubclassOf<UHLODInstancedStaticMeshComponent> HLODISMComponentClass = UHLODBuilder::GetInstancedStaticMeshComponentClass();
 	if (HLODISMComponentClass != UHLODInstancedStaticMeshComponent::StaticClass())
 	{
-		UE_LOG(LogHLODHash, VeryVerbose, TEXT(" - HLOD ISM Component Class = %s"), *HLODISMComponentClass->GetName());
-		Ar << HLODISMComponentClass;
+		InHashBuilder.HashField(HLODISMComponentClass->GetPathName(), TEXT("HLODISMComponentClass"));
 	}
 
-	// Append all components CRCs
-	uint32 HLODComponentsHash = UHLODBuilder::ComputeHLODHash(InSourceComponents);
-	UE_LOG(LogHLODHash, VeryVerbose, TEXT(" - HLOD Source Components = %x"), HLODComponentsHash);
-	Ar << HLODComponentsHash;
-
-	return Ar.GetCrc();
+	// Append all components hashes
+	UHLODBuilder::ComputeHLODHash(InHashBuilder, InSourceComponents);
+	
+	FString HashReport = InHashBuilder.BuildHashReport();
+	UE_LOG(LogHLODHash, Warning, TEXT("%s"), *HashReport);
 }
 
 void AddSourceActor(const FStreamingGenerationActorDescView& ActorDescView, const FName WorldPackageName, TMap<FGuid, FWorldPartitionRuntimeCellObjectMapping>& SourceActors)
@@ -302,8 +296,8 @@ TArray<AWorldPartitionHLOD*> FWorldPartitionHLODUtilities::CreateHLODActors(FHLO
 			bool bSourceActorsChanged = HLODSourceActors->GetActors().Num() != SourceActorsArray.Num();
 			if (!bSourceActorsChanged)
 			{
-				const uint32 NewHash = UWorldPartitionHLODSourceActorsFromCell::GetHLODHash(SourceActorsArray);
-				const uint32 PreviousHash = UWorldPartitionHLODSourceActorsFromCell::GetHLODHash(HLODSourceActors->GetActors());
+				const uint32 NewHash = UWorldPartitionHLODSourceActorsFromCell::GetSourceActorsHash(SourceActorsArray);
+				const uint32 PreviousHash = UWorldPartitionHLODSourceActorsFromCell::GetSourceActorsHash(HLODSourceActors->GetActors());
 				bSourceActorsChanged = NewHash != PreviousHash;
 			}
 
@@ -392,6 +386,7 @@ TArray<AWorldPartitionHLOD*> FWorldPartitionHLODUtilities::CreateHLODActors(FHLO
 		if (DirtyReason)
 		{
 			HLODActor->MarkPackageDirty();
+			HLODActor->UpdateHLODBuildReportHeader();
 			UE_CLOG(!bNewActor, LogHLODBuilder, Log, TEXT("Marking existing HLOD actor \"%s\" dirty, reason \"%s\"."), *HLODActor->GetActorLabel(), DirtyReason);
 		}
 
@@ -777,8 +772,11 @@ uint32 FWorldPartitionHLODUtilities::BuildHLOD(AWorldPartitionHLOD* InHLODActor)
 		HLODRelevantComponents = GatherHLODRelevantComponents(LevelStreaming->GetLoadedLevel()->Actors);
 	}
 
+	FHLODHashBuilder HashBuilder;
+	ComputeHLODHash(HashBuilder, InHLODActor, HLODRelevantComponents);
+
+	uint32 NewHLODHash = HashBuilder.GetCrc();
 	uint32 OldHLODHash = bIsDirty ? 0 : InHLODActor->GetHLODHash();
-	uint32 NewHLODHash = ComputeHLODHash(InHLODActor, HLODRelevantComponents);
 
 	if (HLODBuildEvaluatorDelegate.IsBound())
 	{
@@ -991,6 +989,9 @@ uint32 FWorldPartitionHLODUtilities::BuildHLOD(AWorldPartitionHLOD* InHLODActor)
 	InHLODActor->SetStat(FWorldPartitionHLODStats::BuildTimeBuildMilliseconds, BuildTimeMS);
 	InHLODActor->SetStat(FWorldPartitionHLODStats::BuildTimeTotalMilliseconds, TotalTimeMS);
 
+	// Update HLOD hash and provide the detailed hash report to help future investigation into what caused HLODs to be rebuilt.
+	InHLODActor->SetHLODHash(NewHLODHash, HashBuilder.BuildHashReport());
+
 	return NewHLODHash;
 }
 
diff --git a/Engine/Plugins/Editor/WorldPartitionHLODUtilities/Source/Public/WorldPartition/HLOD/Builders/HLODBuilderInstancing.h b/Engine/Plugins/Editor/WorldPartitionHLODUtilities/Source/Public/WorldPartition/HLOD/Builders/HLODBuilderInstancing.h
--- a/Engine/Plugins/Editor/WorldPartitionHLODUtilities/Source/Public/WorldPartition/HLOD/Builders/HLODBuilderInstancing.h
+++ b/Engine/Plugins/Editor/WorldPartitionHLODUtilities/Source/Public/WorldPartition/HLOD/Builders/HLODBuilderInstancing.h
@@ -22,7 +22,7 @@ class UHLODBuilderInstancingSettings : public UHLODBuilderSettings
 {
 	GENERATED_UCLASS_BODY()
 
-	UE_API virtual uint32 GetCRC() const override;
+	UE_API virtual void ComputeHLODHash(FHLODHashBuilder& InHashBuilder) const override;
 
 	/**
      * By definition, instancing means reusing the source assets.
diff --git a/Engine/Plugins/Editor/WorldPartitionHLODUtilities/Source/Public/WorldPartition/HLOD/Builders/HLODBuilderMeshApproximate.h b/Engine/Plugins/Editor/WorldPartitionHLODUtilities/Source/Public/WorldPartition/HLOD/Builders/HLODBuilderMeshApproximate.h
--- a/Engine/Plugins/Editor/WorldPartitionHLODUtilities/Source/Public/WorldPartition/HLOD/Builders/HLODBuilderMeshApproximate.h
+++ b/Engine/Plugins/Editor/WorldPartitionHLODUtilities/Source/Public/WorldPartition/HLOD/Builders/HLODBuilderMeshApproximate.h
@@ -16,7 +16,7 @@ class UHLODBuilderMeshApproximateSettings : public UHLODBuilderSettings
 {
 	GENERATED_UCLASS_BODY()
 
-	WORLDPARTITIONHLODUTILITIES_API virtual uint32 GetCRC() const override;
+	WORLDPARTITIONHLODUTILITIES_API virtual void ComputeHLODHash(FHLODHashBuilder& InHashBuilder) const override;
 
 	/** Mesh approximation settings */
 	UPROPERTY(EditAnywhere, Category = HLOD, meta=(EditInline))
diff --git a/Engine/Plugins/Editor/WorldPartitionHLODUtilities/Source/Public/WorldPartition/HLOD/Builders/HLODBuilderMeshMerge.h b/Engine/Plugins/Editor/WorldPartitionHLODUtilities/Source/Public/WorldPartition/HLOD/Builders/HLODBuilderMeshMerge.h
--- a/Engine/Plugins/Editor/WorldPartitionHLODUtilities/Source/Public/WorldPartition/HLOD/Builders/HLODBuilderMeshMerge.h
+++ b/Engine/Plugins/Editor/WorldPartitionHLODUtilities/Source/Public/WorldPartition/HLOD/Builders/HLODBuilderMeshMerge.h
@@ -16,7 +16,7 @@ class UHLODBuilderMeshMergeSettings : public UHLODBuilderSettings
 {
 	GENERATED_UCLASS_BODY()
 
-	WORLDPARTITIONHLODUTILITIES_API virtual uint32 GetCRC() const override;
+	WORLDPARTITIONHLODUTILITIES_API virtual void ComputeHLODHash(FHLODHashBuilder& InHashBuilder) const override;
 
 	/**
      * Merge mesh will only reuse the source materials when not merging the materials. In this case, the created mesh
diff --git a/Engine/Plugins/Editor/WorldPartitionHLODUtilities/Source/Public/WorldPartition/HLOD/Builders/HLODBuilderMeshSimplify.h b/Engine/Plugins/Editor/WorldPartitionHLODUtilities/Source/Public/WorldPartition/HLOD/Builders/HLODBuilderMeshSimplify.h
--- a/Engine/Plugins/Editor/WorldPartitionHLODUtilities/Source/Public/WorldPartition/HLOD/Builders/HLODBuilderMeshSimplify.h
+++ b/Engine/Plugins/Editor/WorldPartitionHLODUtilities/Source/Public/WorldPartition/HLOD/Builders/HLODBuilderMeshSimplify.h
@@ -16,7 +16,7 @@ class UHLODBuilderMeshSimplifySettings : public UHLODBuilderSettings
 {
 	GENERATED_UCLASS_BODY()
 
-	WORLDPARTITIONHLODUTILITIES_API virtual uint32 GetCRC() const override;
+	WORLDPARTITIONHLODUTILITIES_API virtual void ComputeHLODHash(FHLODHashBuilder& InHashBuilder) const override;
 
 	/** Simplified mesh generation settings */
 	UPROPERTY(EditAnywhere, Category = HLOD)
diff --git a/Engine/Plugins/Experimental/Water/Source/Runtime/Private/WaterBodyHLODBuilder.cpp b/Engine/Plugins/Experimental/Water/Source/Runtime/Private/WaterBodyHLODBuilder.cpp
--- a/Engine/Plugins/Experimental/Water/Source/Runtime/Private/WaterBodyHLODBuilder.cpp
+++ b/Engine/Plugins/Experimental/Water/Source/Runtime/Private/WaterBodyHLODBuilder.cpp
@@ -22,6 +22,8 @@
 #include "Engine/HLODProxy.h"
 #include "SceneView.h"
 
+#include "WorldPartition/HLOD/HLODHashBuilder.h"
+
 #endif
 
 UWaterBodyHLODBuilder::UWaterBodyHLODBuilder(const FObjectInitializer& ObjectInitializer)
@@ -45,23 +47,20 @@ static double GetPixelToWorldUnit(const double InDistanceZ)
 	return WorldPos.X;
 }
 
-uint32 UWaterBodyHLODBuilder::ComputeHLODHash(const UActorComponent* InSourceComponent) const
+bool UWaterBodyHLODBuilder::ComputeHLODHash(FHLODHashBuilder& HashBuilder, const UActorComponent* InSourceComponent) const
 {
-	FArchiveCrc32 Ar;
-
-	if (const UWaterBodyComponent* WaterBodyComponent = Cast<UWaterBodyComponent>(InSourceComponent))
+	const UWaterBodyComponent* WaterBodyComponent = Cast<UWaterBodyComponent>(InSourceComponent);
+	if (!WaterBodyComponent)
 	{
-		uint32 TransformHash = UHLODProxy::GetCRC(WaterBodyComponent->GetComponentTransform());
-		Ar << TransformHash;
-		
-		FMeshDescription HLODMesh = WaterBodyComponent->GetHLODMeshDescription();
-		Ar << HLODMesh;
-
-		UMaterialInterface* HLODMaterial = WaterBodyComponent->GetHLODMaterial();
-		Ar << HLODMaterial;
+		return false;
 	}
+	
+	HashBuilder.HashField(WaterBodyComponent->GetComponentTransform(), TEXT("Transform"));
+	HashBuilder.HashField(WaterBodyComponent->GetHLODMeshDescription(), TEXT("HLODMesh"));
+
+	HashBuilder << WaterBodyComponent->GetHLODMaterial();
 
-	return Ar.GetCrc();
+	return true;
 }
 
 TArray<UActorComponent*> UWaterBodyHLODBuilder::Build(const FHLODBuildContext& InHLODBuildContext, const TArray<UActorComponent*>& InSourceComponents) const
diff --git a/Engine/Plugins/Experimental/Water/Source/Runtime/Public/WaterBodyHLODBuilder.h b/Engine/Plugins/Experimental/Water/Source/Runtime/Public/WaterBodyHLODBuilder.h
--- a/Engine/Plugins/Experimental/Water/Source/Runtime/Public/WaterBodyHLODBuilder.h
+++ b/Engine/Plugins/Experimental/Water/Source/Runtime/Public/WaterBodyHLODBuilder.h
@@ -13,7 +13,7 @@ class UWaterBodyHLODBuilder : public UHLODBuilder
 
 public:
 #if WITH_EDITOR
-	virtual uint32 ComputeHLODHash(const UActorComponent* InSourceComponent) const override;
+	virtual bool ComputeHLODHash(FHLODHashBuilder& HashBuilder, const UActorComponent* InSourceComponent) const override;
 
 	/**
 	 * Components created with this method need to be properly outered & assigned to your target actor.
diff --git a/Engine/Source/Runtime/CoreUObject/Private/Serialization/ArchiveCrc32.cpp b/Engine/Source/Runtime/CoreUObject/Private/Serialization/ArchiveCrc32.cpp
--- a/Engine/Source/Runtime/CoreUObject/Private/Serialization/ArchiveCrc32.cpp
+++ b/Engine/Source/Runtime/CoreUObject/Private/Serialization/ArchiveCrc32.cpp
@@ -8,6 +8,7 @@
 
 FArchiveCrc32::FArchiveCrc32(uint32 InCRC, UObject* InRootObject)
 	: CRC(InCRC)
+	, InitialCRC(InCRC)
 	, RootObject(InRootObject)
 {
 }
@@ -45,3 +46,9 @@ FArchive& FArchiveCrc32::operator<<(class UObject*& Object)
 
 	return *this;
 }
+
+void FArchiveCrc32::Reset()
+{
+	FArchiveUObject::Reset();
+	CRC = InitialCRC;
+}
diff --git a/Engine/Source/Runtime/CoreUObject/Public/Serialization/ArchiveCrc32.h b/Engine/Source/Runtime/CoreUObject/Public/Serialization/ArchiveCrc32.h
--- a/Engine/Source/Runtime/CoreUObject/Public/Serialization/ArchiveCrc32.h
+++ b/Engine/Source/Runtime/CoreUObject/Public/Serialization/ArchiveCrc32.h
@@ -34,6 +34,7 @@ public:
 	COREUOBJECT_API virtual void Serialize(void* Data, int64 Num) override;
 	COREUOBJECT_API virtual FArchive& operator<<(class FName& Name);
 	COREUOBJECT_API virtual FArchive& operator<<(class UObject*& Object);
+	COREUOBJECT_API virtual void Reset() override;
 	virtual FString GetArchiveName() const { return TEXT("FArchiveCrc32"); }
 	//~ End FArchive Interface
 
@@ -54,8 +55,11 @@ public:
 
 	using FArchiveUObject::operator<<; // For visibility of the overloads we don't override
 
+	
+
 private:
 	uint32 CRC;
+	uint32 InitialCRC;
 	UObject* RootObject;
 };
 
diff --git a/Engine/Source/Runtime/Engine/Private/Components/InstancedSkinnedMeshComponent.cpp b/Engine/Source/Runtime/Engine/Private/Components/InstancedSkinnedMeshComponent.cpp
--- a/Engine/Source/Runtime/Engine/Private/Components/InstancedSkinnedMeshComponent.cpp
+++ b/Engine/Source/Runtime/Engine/Private/Components/InstancedSkinnedMeshComponent.cpp
@@ -432,26 +432,14 @@ void UInstancedSkinnedMeshComponent::ComputeHLODHash(FHLODHashBuilder& HashBuild
 {
 	Super::ComputeHLODHash(HashBuilder);
 
-	FHLODHashScope HashScope(HashBuilder, TEXT("UInstancedSkinnedMeshComponent"));
-
-	for (const FSkinnedMeshInstanceData& SkinnedMeshInstanceData : InstanceData)
-	{
-		HashBuilder << FTransform(SkinnedMeshInstanceData.Transform);
-		HashBuilder << SkinnedMeshInstanceData.AnimationIndex;
-	}
-	HashBuilder << FHLODHashContext(GET_MEMBER_NAME_CHECKED(UInstancedSkinnedMeshComponent, InstanceData));
-
-	{
-		UObject* TransformProviderObject = GetTransformProvider();
-		HashBuilder << TransformProviderObject << FHLODHashContext(TEXT("TransformProvider"));
-	}
-
-	HashBuilder << InstanceCustomData << FHLODHashContext(GET_MEMBER_NAME_CHECKED(UInstancedSkinnedMeshComponent, InstanceCustomData));
-	HashBuilder << InstanceMinDrawDistance << FHLODHashContext(GET_MEMBER_NAME_CHECKED(UInstancedSkinnedMeshComponent, InstanceMinDrawDistance));
-	HashBuilder << InstanceStartCullDistance << FHLODHashContext(GET_MEMBER_NAME_CHECKED(UInstancedSkinnedMeshComponent, InstanceStartCullDistance));
-	HashBuilder << InstanceEndCullDistance << FHLODHashContext(GET_MEMBER_NAME_CHECKED(UInstancedSkinnedMeshComponent, InstanceEndCullDistance));
-		
-	HashBuilder << GetSkinnedAsset() << FHLODHashContext(TEXT("SkinnedAsset"));
+	HashBuilder.HashField(InstanceData, GET_MEMBER_NAME_CHECKED(UInstancedSkinnedMeshComponent, InstanceData));
+	HashBuilder.HashField(GetTransformProvider(), TEXT("TransformProvider"));
+	HashBuilder.HashField(InstanceCustomData, GET_MEMBER_NAME_CHECKED(UInstancedSkinnedMeshComponent, InstanceCustomData));
+	HashBuilder.HashField(InstanceMinDrawDistance, GET_MEMBER_NAME_CHECKED(UInstancedSkinnedMeshComponent, InstanceMinDrawDistance));
+	HashBuilder.HashField(InstanceStartCullDistance, GET_MEMBER_NAME_CHECKED(UInstancedSkinnedMeshComponent, InstanceStartCullDistance));
+	HashBuilder.HashField(InstanceEndCullDistance, GET_MEMBER_NAME_CHECKED(UInstancedSkinnedMeshComponent, InstanceEndCullDistance));
+
+	HashBuilder << GetSkinnedAsset();
 }
 #endif
 
diff --git a/Engine/Source/Runtime/Engine/Private/Components/PrimitiveComponent.cpp b/Engine/Source/Runtime/Engine/Private/Components/PrimitiveComponent.cpp
--- a/Engine/Source/Runtime/Engine/Private/Components/PrimitiveComponent.cpp
+++ b/Engine/Source/Runtime/Engine/Private/Components/PrimitiveComponent.cpp
@@ -5322,15 +5322,17 @@ void UPrimitiveComponent::ComputeHLODHash(FHLODHashBuilder& HashBuilder) const
 {
 	Super::ComputeHLODHash(HashBuilder);
 
-	FHLODHashScope HashScope(HashBuilder, TEXT("UPrimitiveComponent"));
-
-	HashBuilder << GetComponentTransform() << FHLODHashContext(TEXT("Transform"));
-	HashBuilder << HLODBatchingPolicy << FHLODHashContext(GET_MEMBER_NAME_CHECKED(UPrimitiveComponent, HLODBatchingPolicy));
-	HashBuilder << CastShadow << FHLODHashContext(GET_MEMBER_NAME_CHECKED(UPrimitiveComponent, CastShadow));
-	HashBuilder << bCastStaticShadow << FHLODHashContext(GET_MEMBER_NAME_CHECKED(UPrimitiveComponent, bCastStaticShadow));
-	HashBuilder << bCastDynamicShadow << FHLODHashContext(GET_MEMBER_NAME_CHECKED(UPrimitiveComponent, bCastDynamicShadow));
-	HashBuilder << bCastFarShadow << FHLODHashContext(GET_MEMBER_NAME_CHECKED(UPrimitiveComponent, bCastFarShadow));
-	HashBuilder << GetCustomPrimitiveData() << FHLODHashContext(TEXT("CustomPrimitiveData"));
+	HashBuilder.HashField(GetComponentTransform(), TEXT("TransformHash"));
+	HashBuilder.HashField(HLODBatchingPolicy, GET_MEMBER_NAME_CHECKED(UPrimitiveComponent, HLODBatchingPolicy));
+	HashBuilder.HashField(CastShadow, GET_MEMBER_NAME_CHECKED(UPrimitiveComponent, CastShadow));
+	HashBuilder.HashField(bCastStaticShadow, GET_MEMBER_NAME_CHECKED(UPrimitiveComponent, bCastStaticShadow));
+	HashBuilder.HashField(bCastDynamicShadow, GET_MEMBER_NAME_CHECKED(UPrimitiveComponent, bCastDynamicShadow));
+	HashBuilder.HashField(bCastFarShadow, GET_MEMBER_NAME_CHECKED(UPrimitiveComponent, bCastFarShadow));
+
+	if (!GetCustomPrimitiveData().Data.IsEmpty())
+	{
+		HashBuilder.HashField(GetCustomPrimitiveData(), TEXT("CustomPrimitiveData"));
+	}
 	
 	// Materials
 	const int32 NumMaterials = GetNumMaterials();
diff --git a/Engine/Source/Runtime/Engine/Private/Components/StaticMeshComponent.cpp b/Engine/Source/Runtime/Engine/Private/Components/StaticMeshComponent.cpp
--- a/Engine/Source/Runtime/Engine/Private/Components/StaticMeshComponent.cpp
+++ b/Engine/Source/Runtime/Engine/Private/Components/StaticMeshComponent.cpp
@@ -3042,22 +3042,24 @@ void UStaticMeshComponent::ComputeHLODHash(FHLODHashBuilder& HashBuilder) const
 {
 	Super::ComputeHLODHash(HashBuilder);
 
-	FHLODHashScope HashScope(HashBuilder, TEXT("UStaticMeshComponent"));
-
-	HashBuilder << ForcedLodModel << FHLODHashContext(GET_MEMBER_NAME_CHECKED(UStaticMeshComponent, ForcedLodModel));
+	if (ForcedLodModel > 0)
+	{
+		HashBuilder.HashField(ForcedLodModel, GET_MEMBER_NAME_CHECKED(UStaticMeshComponent, ForcedLodModel));
+	}
 	
 	int32 Width, Height;
 	if (GetLightMapResolution(Width, Height))
 	{
-		HashBuilder << Width << Height << FHLODHashContext(TEXT("LightMapResolution"));
+		HashBuilder.HashField(FIntVector2(Width, Height), TEXT("LightMapResolution"));
 	}
 	
 	for (const FStaticMeshComponentLODInfo& LODInfo : LODData)
 	{
 		if (LODInfo.OverrideVertexColors)
 		{
-			HashBuilder.Serialize((uint8*)LODInfo.OverrideVertexColors->GetVertexData(), LODInfo.OverrideVertexColors->GetNumVertices() * LODInfo.OverrideVertexColors->GetStride());
-			HashBuilder << FHLODHashContext(TEXT("OverrideVertexColors"));
+			FArchiveCrc32 Ar;
+			Ar.Serialize((uint8*)LODInfo.OverrideVertexColors->GetVertexData(), LODInfo.OverrideVertexColors->GetNumVertices() * LODInfo.OverrideVertexColors->GetStride());
+			HashBuilder.HashField(Ar.GetCrc(), TEXT("OverrideVertexColors"));
 		}
 	}
 	
diff --git a/Engine/Source/Runtime/Engine/Private/InstancedStaticMesh.cpp b/Engine/Source/Runtime/Engine/Private/InstancedStaticMesh.cpp
--- a/Engine/Source/Runtime/Engine/Private/InstancedStaticMesh.cpp
+++ b/Engine/Source/Runtime/Engine/Private/InstancedStaticMesh.cpp
@@ -2438,23 +2438,10 @@ void UInstancedStaticMeshComponent::ComputeHLODHash(FHLODHashBuilder& HashBuilde
 {
 	Super::ComputeHLODHash(HashBuilder);
 
-	FHLODHashScope HashScope(HashBuilder, TEXT("UInstancedStaticMeshComponent"));
-
-	for (const FInstancedStaticMeshInstanceData& InstanceData : PerInstanceSMData)
-	{
-		HashBuilder << FTransform(InstanceData.Transform);
-	}
-	HashBuilder << FHLODHashContext(GET_MEMBER_NAME_CHECKED(UInstancedStaticMeshComponent, PerInstanceSMData));
-
-	HashBuilder << PerInstanceSMCustomData << FHLODHashContext(GET_MEMBER_NAME_CHECKED(UInstancedStaticMeshComponent, PerInstanceSMCustomData));
-	HashBuilder << InstancingRandomSeed << FHLODHashContext(GET_MEMBER_NAME_CHECKED(UInstancedStaticMeshComponent, InstancingRandomSeed));
-		
-	for (const FInstancedStaticMeshRandomSeed& ISMCRandomSeed : AdditionalRandomSeeds)
-	{
-		HashBuilder << ISMCRandomSeed.StartInstanceIndex;
-		HashBuilder << ISMCRandomSeed.RandomSeed;
-	}
-	HashBuilder << FHLODHashContext(GET_MEMBER_NAME_CHECKED(UInstancedStaticMeshComponent, AdditionalRandomSeeds));
+	HashBuilder.HashField(PerInstanceSMData, GET_MEMBER_NAME_CHECKED(UInstancedStaticMeshComponent, PerInstanceSMData));
+	HashBuilder.HashField(PerInstanceSMCustomData, GET_MEMBER_NAME_CHECKED(UInstancedStaticMeshComponent, PerInstanceSMCustomData));
+	HashBuilder.HashField(InstancingRandomSeed, GET_MEMBER_NAME_CHECKED(UInstancedStaticMeshComponent, InstancingRandomSeed));
+	HashBuilder.HashField(AdditionalRandomSeeds, GET_MEMBER_NAME_CHECKED(UInstancedStaticMeshComponent, AdditionalRandomSeeds));
 }
 #endif
 
diff --git a/Engine/Source/Runtime/Engine/Private/WorldPartition/HLOD/HLODActor.cpp b/Engine/Source/Runtime/Engine/Private/WorldPartition/HLOD/HLODActor.cpp
--- a/Engine/Source/Runtime/Engine/Private/WorldPartition/HLOD/HLODActor.cpp
+++ b/Engine/Source/Runtime/Engine/Private/WorldPartition/HLOD/HLODActor.cpp
@@ -322,6 +322,11 @@ void AWorldPartitionHLOD::PreSave(FObjectPreSaveContext ObjectSaveContext)
 {
 	Super::PreSave(ObjectSaveContext);
 
+	if (ObjectSaveContext.IsProceduralSave())
+	{
+		return;
+	}
+
 	// Always disable collisions on HLODs
 	SetActorEnableCollision(false);
 
@@ -571,6 +576,12 @@ uint32 AWorldPartitionHLOD::GetHLODHash() const
 	return HLODHash;
 }
 
+void AWorldPartitionHLOD::SetHLODHash(uint32 InHLODHash, const FString& InHLODBuildReportContent)
+{
+	HLODHash = InHLODHash;
+	UpdateHLODBuildReportContent(InHLODBuildReportContent);
+}
+
 void AWorldPartitionHLOD::BuildHLOD(bool bForceBuild)
 {
 	IWorldPartitionHLODUtilitiesModule* WPHLODUtilitiesModule = FModuleManager::Get().LoadModulePtr<IWorldPartitionHLODUtilitiesModule>("WorldPartitionHLODUtilities");
@@ -581,7 +592,7 @@ void AWorldPartitionHLOD::BuildHLOD(bool bForceBuild)
 			HLODHash = 0;
 		}
 
-		HLODHash = WPHLODUtilities->BuildHLOD(this);
+		WPHLODUtilities->BuildHLOD(this);
 	}
 
 	// When generating WorldPartition HLODs, we have the renderer initialized.
@@ -597,4 +608,109 @@ void AWorldPartitionHLOD::GetAssetRegistryTags(FAssetRegistryTagsContext Context
 	Context.AddTag(UObject::FAssetRegistryTag(NAME_HLODHash_AssetTag, LexToString(HLODHash), UObject::FAssetRegistryTag::TT_Hidden));
 }
 
+const TCHAR* HLOD_REPORT_BEGIN = TEXT("### HLOD_REPORT_BEGIN ###");
+const TCHAR* HLODLastBuildInfoSectionHeader = TEXT("### Last Build Info ###");
+const TCHAR* HLODActorDetailsSectionHeader = TEXT("### HLOD Actor Details ###");
+const TCHAR* HLODBuildDetailsSectionHeader = TEXT("### HLOD Build Details ###");
+const TCHAR* HLOD_REPORT_END = TEXT("### HLOD_REPORT_END ###");
+
+FString AWorldPartitionHLOD::GenerateHLODBuildReportHeaderString() const
+{
+	TStringBuilder<1024> BuildReportHeader;
+
+	// add constant metadata
+	FString BuildVersion = FApp::GetBuildVersion();
+	FString EngineVersion = FEngineVersion::Current().ToString();
+	FString ExecutingJobURL = FApp::GetExecutingJobURL();
+	FString Platform = ANSI_TO_TCHAR(FPlatformProperties::IniPlatformName());
+	FString BuildConfiguration = LexToString(FApp::GetBuildConfiguration());
+	FString CommandLine = FCommandLine::Get();
+	CommandLine.TrimStartAndEndInline();
+	CommandLine.ReplaceInline(TEXT("\n"), TEXT(""));
+	CommandLine.ReplaceInline(TEXT("\r"), TEXT(""));
+	FString EngineMode = FGenericPlatformMisc::GetEngineMode();
+	FString GraphicsRHI = FApp::GetGraphicsRHI();
+	FString DateTimeUTC = FDateTime::UtcNow().ToString(TEXT("%Y-%m-%d %H:%M:%S"));
+
+	BuildReportHeader << HLODLastBuildInfoSectionHeader << LINE_TERMINATOR;
+	BuildReportHeader << TEXT(" * BuildVersion:       ") << BuildVersion       << LINE_TERMINATOR;
+	BuildReportHeader << TEXT(" * EngineVersion:      ") << EngineVersion      << LINE_TERMINATOR;
+	BuildReportHeader << TEXT(" * ExecutingJobURL:    ") << ExecutingJobURL    << LINE_TERMINATOR;
+	BuildReportHeader << TEXT(" * Platform:           ") << Platform           << LINE_TERMINATOR;
+	BuildReportHeader << TEXT(" * BuildConfiguration: ") << BuildConfiguration << LINE_TERMINATOR;
+	BuildReportHeader << TEXT(" * CommandLine:        ") << CommandLine        << LINE_TERMINATOR;
+	BuildReportHeader << TEXT(" * EngineMode:         ") << EngineMode         << LINE_TERMINATOR;
+	BuildReportHeader << TEXT(" * GraphicsRHI:	      ") << GraphicsRHI        << LINE_TERMINATOR;
+	BuildReportHeader << TEXT(" * DateTimeUTC:        ") << DateTimeUTC        << LINE_TERMINATOR;
+	BuildReportHeader << LINE_TERMINATOR;
+	// Actor Details
+	FString ActorDesc = CreateActorDesc()->ToString(FWorldPartitionActorDesc::EToStringMode::ForDiff).TrimEnd();
+
+	BuildReportHeader << HLODActorDetailsSectionHeader << LINE_TERMINATOR;
+	BuildReportHeader << TEXT(" * HLOD Actor Descriptor:   ") << LINE_TERMINATOR << TEXT("\t") << ActorDesc << LINE_TERMINATOR; // Will log on multiple lines
+	BuildReportHeader << TEXT(" * HLOD Actor Build Hash:   ") << FString::Printf(TEXT("%08X"), HLODHash) << LINE_TERMINATOR;
+	BuildReportHeader << TEXT(" * HLOD Actor IsStandalone: ") << (IsStandalone() ? TEXT("1") : TEXT("0")) << LINE_TERMINATOR;
+
+	return FString(BuildReportHeader);
+}
+
+void AWorldPartitionHLOD::UpdateHLODBuildReportHeader()
+{
+	// The HLOD report is formatted as such:
+	//
+	// ### HLOD_REPORT_BEGIN ###
+	// ### Last Build Info ###
+	// * Info1:
+	// * Info2:
+	// * Info3:
+	// * ...
+	// 
+	// ### HLOD Actor Details ###
+	// * Detail1:
+	// * Detail2:
+	// 
+	// ### HLOD Build Details ###
+	// >> HLODBuildReportContent <<
+	// ### HLOD_REPORT_END ###
+
+	// Extract existing content section
+	FString BuildReportContent;
+	int32 Pos = HLODBuildReport.Find(HLODBuildDetailsSectionHeader);
+	if (Pos != INDEX_NONE)
+	{
+		HLODBuildReport.RightChopInline(Pos, EAllowShrinking::No);
+		BuildReportContent = MoveTemp(HLODBuildReport);
+	}
+	else
+	{
+		// No report content was ever written, so create an empty one
+		BuildReportContent = LINE_TERMINATOR;
+		BuildReportContent += HLODBuildDetailsSectionHeader;
+		BuildReportContent += LINE_TERMINATOR;
+		BuildReportContent += HLOD_REPORT_END;
+	}
+
+	HLODBuildReport  = HLOD_REPORT_BEGIN;
+	HLODBuildReport += LINE_TERMINATOR;
+	HLODBuildReport += GenerateHLODBuildReportHeaderString();
+	HLODBuildReport += BuildReportContent;
+}
+
+void AWorldPartitionHLOD::UpdateHLODBuildReportContent(const FString& InHLODBuildReportContent)
+{
+	HLODBuildReport.Reset();
+
+	HLODBuildReport += LINE_TERMINATOR;
+	HLODBuildReport += HLOD_REPORT_BEGIN;
+	HLODBuildReport += LINE_TERMINATOR;
+	HLODBuildReport += GenerateHLODBuildReportHeaderString();
+	HLODBuildReport += LINE_TERMINATOR;
+	HLODBuildReport += HLODBuildDetailsSectionHeader;
+	HLODBuildReport += LINE_TERMINATOR;
+	HLODBuildReport += InHLODBuildReportContent.TrimEnd();
+	HLODBuildReport += LINE_TERMINATOR;
+	HLODBuildReport += HLOD_REPORT_END;
+	HLODBuildReport += LINE_TERMINATOR;
+}
+
 #endif // WITH_EDITOR
\ No newline at end of file
diff --git a/Engine/Source/Runtime/Engine/Private/WorldPartition/HLOD/HLODBuilder.cpp b/Engine/Source/Runtime/Engine/Private/WorldPartition/HLOD/HLODBuilder.cpp
--- a/Engine/Source/Runtime/Engine/Private/WorldPartition/HLOD/HLODBuilder.cpp
+++ b/Engine/Source/Runtime/Engine/Private/WorldPartition/HLOD/HLODBuilder.cpp
@@ -63,40 +63,44 @@ bool UHLODBuilder::RequiresWarmup() const
 	return true;
 }
 
-uint32 UHLODBuilder::ComputeHLODHash(const UActorComponent* InSourceComponent) const
+bool UHLODBuilder::ComputeHLODHash(FHLODHashBuilder& HashBuilder, const UActorComponent* InSourceComponent) const
 {
-	FHLODHashBuilder HashBuilder;
 	InSourceComponent->ComputeHLODHash(HashBuilder);
-	return HashBuilder.GetCrc();
+	return true;
 }
 
-uint32 UHLODBuilder::ComputeHLODHash(const TArray<UActorComponent*>& InSourceComponents)
+void UHLODBuilder::ComputeHLODHash(FHLODHashBuilder& HashBuilder, const TArray<UActorComponent*>& InSourceComponents)
 {
-	// We get the CRC of each component
-	TArray<uint32> ComponentsCRCs;
+	// We get the hash of each component
+	TArray<uint32> ComponentsHashes;
 
 	for (UActorComponent* SourceComponent : InSourceComponents)
 	{
+		FHLODHashScope HashScope(HashBuilder, SourceComponent, FHLODHashScope::EFlags::ResetHash);
+
 		TSubclassOf<UHLODBuilder> HLODBuilderClass = SourceComponent->GetCustomHLODBuilderClass();
 		if (!HLODBuilderClass)
 		{
 			HLODBuilderClass = UHLODBuilder::StaticClass();
 		}
 
-		uint32 ComponentHash = HLODBuilderClass->GetDefaultObject<UHLODBuilder>()->ComputeHLODHash(SourceComponent);
-		if (ComponentHash == 0)
+		bool bValidHash = HLODBuilderClass->GetDefaultObject<UHLODBuilder>()->ComputeHLODHash(HashBuilder, SourceComponent);
+		if (bValidHash)
+		{
+			ComponentsHashes.Add(HashBuilder.GetCrc());
+		}
+		else
 		{
 			UE_LOG(LogHLODBuilder, Warning, TEXT("Can't compute HLOD hash for component of type %s, assuming it is dirty."), *SourceComponent->GetClass()->GetName());
-			ComponentHash = FMath::Rand();
+			ComponentsHashes.Add(FMath::Rand());
 		}
-
-		ComponentsCRCs.Add(ComponentHash);
 	}
 
-	// Sort the components CRCs to ensure the order of components won't have an impact on the final CRC
-	ComponentsCRCs.Sort();
+	// Sort the components hashes to ensure the order of components won't have an impact on the final hash
+	ComponentsHashes.Sort();
 
-	return FCrc::MemCrc32(ComponentsCRCs.GetData(), ComponentsCRCs.NumBytes());
+	// Incorporate all hashes
+	HashBuilder << ComponentsHashes;
 }
 
 TSubclassOf<UHLODInstancedStaticMeshComponent> UHLODBuilder::GetInstancedStaticMeshComponentClass()
diff --git a/Engine/Source/Runtime/Engine/Private/WorldPartition/HLOD/HLODHashBuilder.cpp b/Engine/Source/Runtime/Engine/Private/WorldPartition/HLOD/HLODHashBuilder.cpp
--- a/Engine/Source/Runtime/Engine/Private/WorldPartition/HLOD/HLODHashBuilder.cpp
+++ b/Engine/Source/Runtime/Engine/Private/WorldPartition/HLOD/HLODHashBuilder.cpp
@@ -5,24 +5,24 @@
 #if WITH_EDITOR
 
 #include "Engine/HLODProxy.h"
+#include "Engine/Level.h"
 #include "Engine/SkinnedAsset.h"
 #include "Engine/StaticMesh.h"
 #include "Engine/Texture.h"
 #include "Materials/MaterialInterface.h"
-#include "Misc/TransformUtilities.h"
 #include "WorldPartition/HLOD/HLODBuilder.h"
 
-void FHLODHashBuilder::PushContext(const TCHAR* Context)
+void FHLODHashBuilder::PushObjectContext(const UObject* InObjectContext)
 {
-	LogContext(Context, false);
-
-	++IndentationLevel;
+	UObject* StopOuter = InObjectContext->IsA<UActorComponent>() ? InObjectContext->GetTypedOuter<ULevel>() : nullptr;
+	FString ObjectPath = InObjectContext->GetFullName(StopOuter);
+	ObjectContextStack.Push(*ObjectPath);
 }
 
-void FHLODHashBuilder::PopContext()
+void FHLODHashBuilder::PopObjectContext()
 {
-	check(IndentationLevel > 0);
-	--IndentationLevel;
+	check(!ObjectContextStack.IsEmpty());
+	ObjectContextStack.Pop();
 }
 
 FArchive& FHLODHashBuilder::operator<<(FTransform InTransform)
@@ -31,14 +31,34 @@ FArchive& FHLODHashBuilder::operator<<(FTransform InTransform)
 	return *this;
 }
 
+uint32 FHLODHashBuilder::AddAssetReference(UObject* Asset, TFunctionRef<uint32()> GetHashFunc)
+{
+	check(Asset != nullptr);
+
+	FName AssetName(Asset->GetPathName());
+
+	FAssetHash& AssetHash = AssetsHashes.FindOrAdd(AssetName);
+	if (AssetHash.Hash == 0)
+	{
+		AssetHash.Hash = GetHashFunc();
+		AssetHash.AssetType = Asset->GetClass()->GetFName();
+	}
+
+	if (!ObjectContextStack.IsEmpty())
+	{		
+		ObjectsReferences.FindOrAdd(ObjectContextStack.Top()).ReferencedAssets.Add(AssetName);
+	}
+
+	return AssetHash.Hash;
+}
+
 FArchive& FHLODHashBuilder::operator<<(UObject*& InObject)
 {
 	if (UMaterialInterface* MaterialInterface = Cast<UMaterialInterface>(InObject))
 	{
-		FHLODHashScope HashScope(*this, TEXT("UMaterialInterface"));
+		*this << AddAssetReference(MaterialInterface, [&]() { return UHLODProxy::GetCRC(MaterialInterface); });
 
-		*this << UHLODProxy::GetCRC(MaterialInterface);
-		*this << FHLODHashContext(TEXT("%s (%s)"), *InObject->GetClass()->GetName(), *InObject->GetName());
+		FHLODHashScope HashScopeMaterial(*this, MaterialInterface);
 
 		TArray<UTexture*> Textures;
 		MaterialInterface->GetUsedTextures(Textures);
@@ -46,8 +66,9 @@ FArchive& FHLODHashBuilder::operator<<(UObject*& InObject)
 	
 		if (UMaterialInterface* NaniteOverride = MaterialInterface->GetNaniteOverride())
 		{ 
-			*this << UHLODProxy::GetCRC(NaniteOverride);
-			*this << FHLODHashContext(TEXT("%s (%s)"), *NaniteOverride->GetClass()->GetName(), *NaniteOverride->GetName());
+			*this << AddAssetReference(NaniteOverride, [&]() { return UHLODProxy::GetCRC(NaniteOverride); });
+
+			FHLODHashScope HashScopeNaniteMaterial(*this, NaniteOverride);
 
 			Textures.Reset();
 			NaniteOverride->GetUsedTextures(Textures);
@@ -56,17 +77,11 @@ FArchive& FHLODHashBuilder::operator<<(UObject*& InObject)
 	}
 	else if (UTexture* Texture = Cast<UTexture>(InObject))
 	{
-		FHLODHashScope HashScope(*this, TEXT("UTexture"));
-
-		*this << UHLODProxy::GetCRC(Texture);
-		*this << FHLODHashContext(TEXT("%s (%s)"), *InObject->GetClass()->GetName(), *InObject->GetName());
+		*this << AddAssetReference(Texture, [&]() { return UHLODProxy::GetCRC(Texture); });
 	}
 	else if (UStaticMesh* StaticMesh = Cast<UStaticMesh>(InObject))
 	{
-		FHLODHashScope HashScope(*this, TEXT("UStaticMesh"));
-
-		*this << UHLODProxy::GetCRC(StaticMesh);
-		*this << FHLODHashContext(TEXT("%s (%s)"), *InObject->GetClass()->GetName(), *InObject->GetName());
+		*this << AddAssetReference(StaticMesh, [&]() { return UHLODProxy::GetCRC(StaticMesh); });
 	}
 	else
 	{
@@ -100,18 +115,183 @@ FArchive& FHLODHashBuilder::operator<<(USkinnedAsset* InSkinnedAsset)
 	return *this << Object;
 }
 
-void FHLODHashBuilder::LogContext(const TCHAR* Context, bool bOutputHash)
+static FString GetTypePrefixFromClassName(const FName& ClassFName)
+{
+	const FString C = ClassFName.ToString();
+	if (C.Contains(TEXT("Material")))                             		 return TEXT("MAT");
+	if (C.Contains(TEXT("Texture")))                              		 return TEXT("TEX");
+	if (C.Contains(TEXT("StaticMesh")))                           		 return TEXT("SM");
+	if (C.Contains(TEXT("SkeletalMesh")) || C.Contains(TEXT("Skinned"))) return TEXT("SK");
+	return TEXT("OBJ");
+}
+
+static uint16 StablePathId16(const FName& AssetPath)
+{
+	FString Norm = AssetPath.ToString();
+	Norm.ToLowerInline();
+	Norm.ReplaceInline(TEXT("\\"), TEXT("/"), ESearchCase::CaseSensitive);
+	return static_cast<uint16>(FCrc::StrCrc32(*Norm) & 0xFFFF);
+}
+
+// Build final label from type prefix + 16-bit id.
+static FString MakeTypeAwareLabel16(const FString& TypePrefix, uint16 Id16)
+{
+	return FString::Printf(TEXT("@%s-%04X"), *TypePrefix, Id16);
+}
+
+FString FHLODHashBuilder::BuildHashReport() const
 {
-	FString Indentation = FString::ChrN(IndentationLevel * 4, TEXT(' '));
+	FString Out;
+	static constexpr int32 ReserveSize = 64 * 1024;
+	Out.Reserve(ReserveSize);
+
+	// Collect & sort asset paths
+	TArray<FName> SortedAssetNames;
+	SortedAssetNames.Reserve(AssetsHashes.Num());
+	AssetsHashes.GenerateKeyArray(SortedAssetNames);
+	SortedAssetNames.Sort([](const FName& A, const FName& B)
+	{
+		return A.ToString().Compare(B.ToString(), ESearchCase::IgnoreCase) < 0;
+	});
+
+	// Generate short asset labels ("@TYP-XXXX")
+	TMap<FName, FString> AssetLabels;	// AssetName to label
+	TSet<FString>        UsedLabels;	// "@TYP-XXXX" taken?
+	AssetLabels.Reserve(SortedAssetNames.Num());
+	UsedLabels.Reserve(SortedAssetNames.Num());
+
+	for (const FName& AssetName : SortedAssetNames)
+	{
+		const FAssetHash* Entry = AssetsHashes.Find(AssetName);
+		const FName AssetType = Entry ? Entry->AssetType : NAME_None;
+
+		const FString Prefix = GetTypePrefixFromClassName(AssetType);
+		uint16 Id16 = StablePathId16(AssetName);
+
+		// Handle collisions - increment id until a free one is found
+		// This is deterministic because input order is the sorted asset path order.
+		FString Token;
+		for (uint32 Attempt = 0; Attempt < 65536u; ++Attempt)
+		{
+			Token = MakeTypeAwareLabel16(Prefix, Id16);
+			if (!UsedLabels.Contains(Token))
+			{
+				UsedLabels.Add(Token);
+				break;
+			}
+			Id16 = static_cast<uint16>(Id16 + 1);
+		}
+
+		AssetLabels.Add(AssetName, Token);
+	}
 
-	if (bOutputHash)
+	Out += TEXT("## Global Fields ##\n");
+	for (const TTuple<FName, FString>& Field : GlobalFields)
 	{
-		UE_LOG(LogHLODBuilder, VeryVerbose, TEXT("%s- %s = %x"), *Indentation, Context, GetCrc());
+		Out += FString::Printf(TEXT("\n    * %s: %s"), *Field.Key.ToString(), *Field.Value);
 	}
-	else
+
+	Out += TEXT("\n\n## Referenced Assets ##\n");
+
+ 	// Referenced assets have no fields, we only track their hashes. We don't want to show them in the Objects section.
+	// Keep a list of the keys we consume here so we can skip them later
+	TSet<FString> SuppressedObjectKeys;
+	SuppressedObjectKeys.Reserve(SortedAssetNames.Num());
+
+	// Get a string containing all referenced assets (sorted)
+	auto GetReferencesString = [&AssetLabels](const FObjectHash& ObjectHashes)
+	{
+		TArray<FString> Labels;
+		Labels.Reserve(ObjectHashes.ReferencedAssets.Num());
+		for (const FName& RefAssetName : ObjectHashes.ReferencedAssets)
+		{
+			if (const FString* Label = AssetLabels.Find(RefAssetName))
+			{
+				Labels.Add(*Label);
+			}
+		}
+		Labels.Sort([](const FString& A, const FString& B)
+		{
+			return A.Compare(B, ESearchCase::CaseSensitive) < 0;
+		});
+
+		FString RefsStr;
+		for (int32 i = 0; i < Labels.Num(); ++i)
+		{
+			if (i > 0)
+			{
+				RefsStr += TEXT(" ");
+			}
+			RefsStr += Labels[i];
+		}
+		return RefsStr;
+	};
+
+	// Assets section (sorted by asset path)
+	for (const FName& AssetName : SortedAssetNames)
+	{
+		const FAssetHash* AssetEntry = AssetsHashes.Find(AssetName);
+		const uint32 Hash = AssetEntry ? AssetEntry->Hash : 0u;
+		const FName  Type = AssetEntry ? AssetEntry->AssetType : NAME_None;
+		const FString& Label = AssetLabels.FindChecked(AssetName);
+
+		FString RefsStr;
+		const FString ContextKey = Type.ToString() + TEXT(" ") + AssetName.ToString();
+		if (const FObjectHash* ContextObj = ObjectsReferences.Find(ContextKey))
+		{
+			if (!ContextObj->ReferencedAssets.IsEmpty())
+			{
+				RefsStr = TEXT(", References=") + GetReferencesString(*ContextObj);
+			}
+			SuppressedObjectKeys.Add(ContextKey);
+		}
+
+		Out += FString::Printf(TEXT("%s: %s %s (Hash=%08X%s)\n"),
+			*Label,
+			*Type.ToString(),
+			*AssetName.ToString(),
+			Hash,
+			*RefsStr);
+	}
+
+	// Objects section (sorted), skipping any keys consumed above
+	TArray<FString> SortedObjectPaths;
+	SortedObjectPaths.Reserve(ObjectsReferences.Num());
+	ObjectsReferences.GenerateKeyArray(SortedObjectPaths);
+	SortedObjectPaths.Sort([](const FString& A, const FString& B)
+	{
+		return A.Compare(B, ESearchCase::IgnoreCase) < 0;
+	});
+
+	Out += TEXT("\n## Source Components ##\n");
+
+	for (const FString& ObjPath : SortedObjectPaths)
 	{
-		UE_LOG(LogHLODBuilder, VeryVerbose, TEXT("%s- %s"), *Indentation, Context);
+		if (SuppressedObjectKeys.Contains(ObjPath))
+		{
+			continue; // Object already shown in the "Referenced Assets" section
+		}
+
+		const FObjectHash& ObjectHashes = ObjectsReferences.FindChecked(ObjPath);
+
+		Out += ObjPath;
+
+		for (const TTuple<FName, FString>& Field : ObjectHashes.Fields)
+		{
+			Out += FString::Printf(TEXT("    * %s: %s\n"), *Field.Key.ToString(), *Field.Value);
+		}
+
+		if (!ObjectHashes.ReferencedAssets.IsEmpty())
+		{
+			Out += TEXT("    * References: ");
+			Out += GetReferencesString(ObjectHashes);
+			Out += TEXT("\n");
+		}
+
+		Out += TEXT("\n");
 	}
+
+	return Out;
 }
 
 #endif // #if WITH_EDITOR
diff --git a/Engine/Source/Runtime/Engine/Private/WorldPartition/HLOD/HLODLayer.cpp b/Engine/Source/Runtime/Engine/Private/WorldPartition/HLOD/HLODLayer.cpp
--- a/Engine/Source/Runtime/Engine/Private/WorldPartition/HLOD/HLODLayer.cpp
+++ b/Engine/Source/Runtime/Engine/Private/WorldPartition/HLOD/HLODLayer.cpp
@@ -16,6 +16,8 @@
 #include "Modules/ModuleManager.h"
 #include "UObject/UnrealType.h"
 #include "Widgets/Notifications/SNotificationList.h"
+#include "WorldPartition/HLOD/HLODHashBuilder.h"
+#include "WorldPartition/HLOD/HLODModifier.h"
 #include "WorldPartition/HLOD/IWorldPartitionHLODUtilities.h"
 #include "WorldPartition/HLOD/IWorldPartitionHLODUtilitiesModule.h"
 #include "WorldPartition/WorldPartition.h"
@@ -245,6 +247,23 @@ FName UHLODLayer::GetRuntimeGrid(uint32 InHLODLevel) const
 	return !IsSpatiallyLoaded() ? NAME_None : GetRuntimeGridName(InHLODLevel, CellSize, LoadingRange);
 	PRAGMA_ENABLE_DEPRECATION_WARNINGS
 }
+
+void UHLODLayer::ComputeHLODHash(FHLODHashBuilder& InHashBuilder) const
+{
+	// HLOD Layer
+	InHashBuilder.HashField(LayerType, GET_MEMBER_NAME_CHECKED(UHLODLayer, LayerType));
+	
+	if (HLODBuilderSettings)
+	{
+		HLODBuilderSettings->ComputeHLODHash(InHashBuilder);
+	}
+
+	if (HLODModifierClass.Get())
+	{
+		InHashBuilder.HashField(HLODModifierClass->GetPathName(), GET_MEMBER_NAME_CHECKED(UHLODLayer, HLODModifierClass));
+	}
+}
+
 #endif // WITH_EDITOR
 
 #undef LOCTEXT_NAMESPACE
\ No newline at end of file
diff --git a/Engine/Source/Runtime/Engine/Private/WorldPartition/HLOD/HLODSourceActors.cpp b/Engine/Source/Runtime/Engine/Private/WorldPartition/HLOD/HLODSourceActors.cpp
--- a/Engine/Source/Runtime/Engine/Private/WorldPartition/HLOD/HLODSourceActors.cpp
+++ b/Engine/Source/Runtime/Engine/Private/WorldPartition/HLOD/HLODSourceActors.cpp
@@ -1,41 +1,15 @@
 // Copyright Epic Games, Inc. All Rights Reserved.
 
 #include "WorldPartition/HLOD/HLODSourceActors.h"
+#include "WorldPartition/HLOD/HLODHashBuilder.h"
 
 #if WITH_EDITOR
 #include "WorldPartition/HLOD/HLODActor.h"
 #include "WorldPartition/HLOD/HLODLayer.h"
-#include "WorldPartition/HLOD/HLODModifier.h"
 #endif
 
 #include UE_INLINE_GENERATED_CPP_BY_NAME(HLODSourceActors)
 
-
-#if WITH_EDITOR
-
-static uint32 GetHLODHash(const UHLODLayer* InHLODLayer)
-{
-	UHLODLayer& HLODLayer = *const_cast<UHLODLayer*>(InHLODLayer);
-
-	uint32 HLODHash;
-
-	HLODHash = static_cast<uint8>(HLODLayer.GetLayerType());
-	UE_LOG(LogHLODHash, VeryVerbose, TEXT(" - LayerType = %d"), HLODHash);
-	
-	HLODHash = HashCombineFast(HLODLayer.GetHLODBuilderSettings()->GetCRC(), HLODHash);
-	UE_LOG(LogHLODHash, VeryVerbose, TEXT(" - HLODBuilderSettings = %d"), HLODHash);
-
-	if (TSubclassOf<UWorldPartitionHLODModifier> HLODModifier = HLODLayer.GetHLODModifierClass(); HLODModifier.Get())
-	{
-		HLODHash = FCrc::StrCrc32(*HLODModifier->GetPathName(), HLODHash);
-		UE_LOG(LogHLODHash, VeryVerbose, TEXT(" - HLODModifier = %d"), HLODHash);
-	}
-
-	return HLODHash;
-}
-
-#endif // #if WITH_EDITOR
-
 UWorldPartitionHLODSourceActors::UWorldPartitionHLODSourceActors(const FObjectInitializer& ObjectInitializer)
 	: Super(ObjectInitializer)
 {
@@ -43,16 +17,12 @@ UWorldPartitionHLODSourceActors::UWorldPartitionHLODSourceActors(const FObjectIn
 
 #if WITH_EDITOR
 
-uint32 UWorldPartitionHLODSourceActors::GetHLODHash() const
+void UWorldPartitionHLODSourceActors::ComputeHLODHash(FHLODHashBuilder& InHashBuilder) const
 {
-	uint32 HLODHash = 0;
 	if (HLODLayer)
 	{
-		// HLOD Layer
-		HLODHash = ::GetHLODHash(HLODLayer);
-		UE_LOG(LogHLODHash, VeryVerbose, TEXT(" - HLOD Layer (%s) = %x"), *HLODLayer->GetName(), HLODHash);
+		HLODLayer->ComputeHLODHash(InHashBuilder);
 	}
-	return HLODHash;
 }
 
 void UWorldPartitionHLODSourceActors::SetHLODLayer(const UHLODLayer* InHLODLayer)
diff --git a/Engine/Source/Runtime/Engine/Private/WorldPartition/HLOD/HLODSourceActorsFromCell.cpp b/Engine/Source/Runtime/Engine/Private/WorldPartition/HLOD/HLODSourceActorsFromCell.cpp
--- a/Engine/Source/Runtime/Engine/Private/WorldPartition/HLOD/HLODSourceActorsFromCell.cpp
+++ b/Engine/Source/Runtime/Engine/Private/WorldPartition/HLOD/HLODSourceActorsFromCell.cpp
@@ -4,6 +4,7 @@
 
 #if WITH_EDITOR
 #include "WorldPartition/HLOD/HLODActor.h"
+#include "WorldPartition/HLOD/HLODHashBuilder.h"
 #include "Serialization/ArchiveCrc32.h"
 #include "WorldPartition/WorldPartitionLevelStreamingDynamic.h"
 #endif
@@ -40,7 +41,7 @@ ULevelStreaming* UWorldPartitionHLODSourceActorsFromCell::LoadSourceActors(bool&
 	return LevelStreaming;
 }
 
-uint32 UWorldPartitionHLODSourceActorsFromCell::GetHLODHash(const TArray<FWorldPartitionRuntimeCellObjectMapping>& InSourceActors)
+uint32 UWorldPartitionHLODSourceActorsFromCell::GetSourceActorsHash(const TArray<FWorldPartitionRuntimeCellObjectMapping>& InSourceActors)
 {
 	TArray<FWorldPartitionRuntimeCellObjectMapping>& MutableSourceActors(const_cast<TArray<FWorldPartitionRuntimeCellObjectMapping>&>(InSourceActors));
 
@@ -53,16 +54,13 @@ uint32 UWorldPartitionHLODSourceActorsFromCell::GetHLODHash(const TArray<FWorldP
 	return Ar.GetCrc();
 }
 
-uint32 UWorldPartitionHLODSourceActorsFromCell::GetHLODHash() const
+void UWorldPartitionHLODSourceActorsFromCell::ComputeHLODHash(FHLODHashBuilder& InHashBuilder) const
 {
-	uint32 HLODHash = Super::GetHLODHash();
+	Super::ComputeHLODHash(InHashBuilder);
 
 	// Source Actors
-	uint32 SourceActorsHash = GetHLODHash(Actors);
-	UE_LOG(LogHLODHash, VeryVerbose, TEXT(" - Source Actors (%d actors) = %x"), Actors.Num(), SourceActorsHash);
-	HLODHash = HashCombineFast(HLODHash, SourceActorsHash);
-
-	return HLODHash;
+	uint32 SourceActorsHash = GetSourceActorsHash(Actors);
+	InHashBuilder.HashField(SourceActorsHash, TEXT("SourceActorsHash"));
 }
 
 void UWorldPartitionHLODSourceActorsFromCell::SetActors(const TArray<FWorldPartitionRuntimeCellObjectMapping>&& InSourceActors)
diff --git a/Engine/Source/Runtime/Engine/Private/WorldPartition/HLOD/HLODSourceActorsFromLevel.cpp b/Engine/Source/Runtime/Engine/Private/WorldPartition/HLOD/HLODSourceActorsFromLevel.cpp
--- a/Engine/Source/Runtime/Engine/Private/WorldPartition/HLOD/HLODSourceActorsFromLevel.cpp
+++ b/Engine/Source/Runtime/Engine/Private/WorldPartition/HLOD/HLODSourceActorsFromLevel.cpp
@@ -4,6 +4,7 @@
 
 #if WITH_EDITOR
 #include "WorldPartition/HLOD/HLODActor.h"
+#include "WorldPartition/HLOD/HLODHashBuilder.h"
 #include "Engine/LevelStreamingDynamic.h"
 #include "Engine/World.h"
 #endif
@@ -35,10 +36,11 @@ ULevelStreaming* UWorldPartitionHLODSourceActorsFromLevel::LoadSourceActors(bool
 	return nullptr;
 }
 
-uint32 UWorldPartitionHLODSourceActorsFromLevel::GetHLODHash() const
+void UWorldPartitionHLODSourceActorsFromLevel::ComputeHLODHash(FHLODHashBuilder& InHashBuilder) const
 {
-	// SourceLevel
-	return FCrc::StrCrc32(*SourceLevel.GetUniqueID().GetLongPackageName(), Super::GetHLODHash());
+	Super::ComputeHLODHash(InHashBuilder);
+
+	InHashBuilder.HashField(SourceLevel.GetUniqueID().GetLongPackageName(), GET_MEMBER_NAME_CHECKED(UWorldPartitionHLODSourceActorsFromLevel, SourceLevel));
 }
 
 
diff --git a/Engine/Source/Runtime/Engine/Public/WorldPartition/HLOD/HLODActor.h b/Engine/Source/Runtime/Engine/Public/WorldPartition/HLOD/HLODActor.h
--- a/Engine/Source/Runtime/Engine/Public/WorldPartition/HLOD/HLODActor.h
+++ b/Engine/Source/Runtime/Engine/Public/WorldPartition/HLOD/HLODActor.h
@@ -78,11 +78,15 @@ public:
 
 	ENGINE_API void BuildHLOD(bool bForceBuild = false);
 	ENGINE_API uint32 GetHLODHash() const;
+	ENGINE_API void SetHLODHash(uint32 InHLODHash, const FString& InHLODBuildReportContent);
 
 	ENGINE_API int64 GetStat(FName InStatName) const;
 	void SetStat(FName InStatName, int64 InStatValue) { HLODStats.Add(InStatName, InStatValue); }
 	void ResetStats() { HLODStats.Reset(); }
 
+	ENGINE_API void UpdateHLODBuildReportHeader();
+	ENGINE_API void UpdateHLODBuildReportContent(const FString& InHLODBuildReportContent);
+
 private:
 	const FStats& GetStats() const { return HLODStats; }
 
@@ -133,6 +137,7 @@ protected:
 private:
 #if WITH_EDITOR
 	ENGINE_API void OnWorldCleanup(UWorld* InWorld, bool bSessionEnded, bool bCleanupResources);
+	FString GenerateHLODBuildReportHeaderString() const;
 #endif
 
 #if WITH_EDITORONLY_DATA
@@ -153,6 +158,9 @@ private:
 
 	UPROPERTY()
 	TMap<FName, int64> HLODStats;
+
+	UPROPERTY()
+	FString HLODBuildReport;
 #endif
 
 	UPROPERTY()
diff --git a/Engine/Source/Runtime/Engine/Public/WorldPartition/HLOD/HLODBuilder.h b/Engine/Source/Runtime/Engine/Public/WorldPartition/HLOD/HLODBuilder.h
--- a/Engine/Source/Runtime/Engine/Public/WorldPartition/HLOD/HLODBuilder.h
+++ b/Engine/Source/Runtime/Engine/Public/WorldPartition/HLOD/HLODBuilder.h
@@ -9,6 +9,7 @@
 #include "HLODBuilder.generated.h"
 
 class AActor;
+class FHLODHashBuilder;
 class UActorComponent;
 class UHLODInstancedStaticMeshComponent;
 class UHLODInstancedSkinnedMeshComponent;
@@ -27,9 +28,10 @@ class UHLODBuilderSettings : public UObject
 public:
 	/**
 	 * Hash the class settings.
-	 * This is used to detect settings changes and trigger an HLOD rebuild if necessarry.
+	 * This is used to detect settings changes and trigger an HLOD rebuild if necessary.
+	 * The provided hash builder state will be updated to reflect the hashed content.
 	 */
-	virtual uint32 GetCRC() const { return 0; }
+	virtual void ComputeHLODHash(FHLODHashBuilder& InHashBuilder) const {}
 
 	/**
 	 * Whether these HLOD settings are going to result in the created mesh referencing external materials.
@@ -37,6 +39,9 @@ public:
 	 * By default most HLOD methods are baking down their textures and uses their own unique material instance.
 	 */
 	virtual bool IsReusingSourceMaterials() const { return false; }
+
+	UE_DEPRECATED(5.7, "GetCRC() has been replaced by ComputeHLODHash()")
+	virtual uint32 GetCRC() const final { return 0; }
 #endif // WITH_EDITOR
 };
 
@@ -141,8 +146,10 @@ public:
 	 * Used to detect changes to the source components of an HLOD.
 	 * The base version can only support hashing of static mesh components. HLOD builder subclasses
 	 * should override this method and compute the hash of component types they support as input.
+	 * The provided hash builder state will be updated to reflect the hashed content.
+	 * Will return true if an hash could be computed, false otherwise.
 	 */
-	ENGINE_API virtual uint32 ComputeHLODHash(const UActorComponent* InSourceComponent) const;
+	ENGINE_API virtual bool ComputeHLODHash(class FHLODHashBuilder& InHashBuilder, const UActorComponent* InSourceComponent) const;
 
 	/**
 	 * Components created with this method need to be properly outered & assigned to your target actor.
@@ -151,9 +158,10 @@ public:
 
 	/**
 	 * From a set of components, compute a unique hash from their properties that are relevant to HLOD generation.
-	 * Used to detect changes to the source actors and trigger an HLOD rebuild if necessarry.
+	 * Used to detect changes to the source actors and trigger an HLOD rebuild if necessary.
+	 * The provided hash builder state will be updated to reflect the hashed content.
 	 */
-	static ENGINE_API uint32 ComputeHLODHash(const TArray<UActorComponent*>& InSourceComponents);
+	static ENGINE_API void ComputeHLODHash(class FHLODHashBuilder& HashBuilder, const TArray<UActorComponent*>& InSourceComponents);
 
 	/** 
 	 * Get the InstancedStaticMeshComponent subclass that should be used when creating instanced static mesh HLODs.
@@ -210,7 +218,7 @@ class UNullHLODBuilder : public UHLODBuilder
 
 #if WITH_EDITOR
 	virtual bool RequiresWarmup() const { return false; }
-	virtual uint32 ComputeHLODHash(const UActorComponent* InSourceComponent) const { return 0; }
+	virtual bool ComputeHLODHash(FHLODHashBuilder& HashBuilder, const UActorComponent* InSourceComponent) const { return false; }
 	virtual TArray<UActorComponent*> Build(const FHLODBuildContext& InHLODBuildContext, const TArray<UActorComponent*>& InSourceComponents) const { return {}; }
 #endif
 };
diff --git a/Engine/Source/Runtime/Engine/Public/WorldPartition/HLOD/HLODHashBuilder.h b/Engine/Source/Runtime/Engine/Public/WorldPartition/HLOD/HLODHashBuilder.h
--- a/Engine/Source/Runtime/Engine/Public/WorldPartition/HLOD/HLODHashBuilder.h
+++ b/Engine/Source/Runtime/Engine/Public/WorldPartition/HLOD/HLODHashBuilder.h
@@ -5,16 +5,24 @@
 #if WITH_EDITOR
 
 #include "Math/MathFwd.h"
+#include "Misc/TransformUtilities.h"
 #include "Serialization/ArchiveCrc32.h"
+#include "UObject/Class.h"
+
+class UMaterialInterface;
+class UTexture;
+class UStaticMesh;
+class USkinnedAsset;
+class UObject;
 
 class FHLODHashBuilder : public FArchiveCrc32
 {
 public:
-	/** Push a new context, logging with correct indentation */
-	void PushContext(const TCHAR* Context);
+	/** Push a new context */
+	void PushObjectContext(const class UObject* InObjectContext);
 
-	/** Pop context, ensuring indentation decreases */
-	void PopContext();
+	/** Pop context */
+	void PopObjectContext();
 
 	template <typename TParamType>
 	typename TEnableIf<TIsIntegral<TParamType>::Value, FArchive&>::Type operator<<(TParamType InValue)
@@ -23,14 +31,82 @@ public:
 		return Self << InValue;
 	}
 
-	FArchive& operator<<(FTransform InTransform);
-	FArchive& operator<<(class UMaterialInterface* InMaterialInterface);
-	FArchive& operator<<(class UTexture* InTexture);
-	FArchive& operator<<(class UStaticMesh* InStaticMesh);
-	FArchive& operator<<(class USkinnedAsset* InSkinnedAsset);
+	// Numeric / enum values will be logged as string (order preserved)
+	template <typename T>
+	typename TEnableIf< TIsIntegral<T>::Value || TIsFloatingPoint<T>::Value || TIsEnum<T>::Value, void>::Type
+	HashField(const T& InFieldValue, const FName& InFieldName)
+	{
+		AddField(InFieldName, FieldValueToString(InFieldValue));
+
+		*this << const_cast<T&>(InFieldValue);
+	}
+
+	// Everything else will log a CRC32 as 8-hex string
+	template <typename T>
+	typename TEnableIf< !(TIsIntegral<T>::Value || TIsFloatingPoint<T>::Value || TIsEnum<T>::Value), void>::Type
+	HashField(const T& InFieldValue, const FName& InFieldName)
+	{
+		FArchiveCrc32 FieldHashAr;
+		FieldHashAr << const_cast<T&>(InFieldValue);
+
+		AddField(InFieldName, FString::Printf(TEXT("%08X"), FieldHashAr.GetCrc()));
+
+		*this << const_cast<T&>(InFieldValue);
+	}
+
+	void HashField(const FString& InFieldValue, const FName& InFieldName)
+	{
+		AddField(InFieldName, InFieldValue);
+
+		*this << const_cast<FString&>(InFieldValue);
+	}
+
+	void HashField(const FName& InFieldValue, const FName& InFieldName)
+	{
+		AddField(InFieldName, InFieldValue.ToString());
+
+		*this << const_cast<FName&>(InFieldValue);
+	}
+
+	void HashField(const ANSICHAR* InFieldValue, const FName& InFieldName)
+	{
+		AddField(InFieldName, InFieldValue);
+
+		FString StringValue(InFieldValue);
+		*this << StringValue;
+	}
+
+	void HashField(const WIDECHAR* InFieldValue, const FName& InFieldName)
+	{
+		AddField(InFieldName, InFieldValue);
+
+		FString StringValue(InFieldValue);
+		*this << StringValue;
+	}
+
+	void HashField(const FTransform& InFieldValue, const FName& InFieldName)
+	{
+		AddField(InFieldName, FString::Printf(TEXT("%08X"), TransformUtilities::GetRoundedTransformCRC32(InFieldValue)));
+
+		*this << const_cast<FTransform&>(InFieldValue);
+	}
+
+	template <typename T>
+	void HashField(const TObjectPtr<T>& InFieldValue, const FName& InFieldName)
+	{
+		AddField(InFieldName, InFieldValue ? InFieldValue->GetPathName() : TEXT("<none>"));
+
+		*this << const_cast<T*>(InFieldValue.Get());
+	}
+
+	ENGINE_API FArchive& operator<<(FTransform InTransform);
+	ENGINE_API FArchive& operator<<(UMaterialInterface* InMaterialInterface);
+	ENGINE_API FArchive& operator<<(UTexture* InTexture);
+	ENGINE_API FArchive& operator<<(UStaticMesh* InStaticMesh);
+	ENGINE_API FArchive& operator<<(USkinnedAsset* InSkinnedAsset);
 
 	//~ Begin FArchive Interface
-	virtual FArchive& operator<<(class UObject*& Object) override;
+	ENGINE_API virtual FArchive& operator<<(UObject*& Object) override;
 	//~ End FArchive Interface
 
 	void LogContext(const TCHAR* Context, bool bOutputHash);
@@ -38,8 +114,104 @@ public:
 	// For visibility of the overloads we don't override
 	using FArchiveCrc32::operator<<;
 
+	ENGINE_API FString BuildHashReport() const;
+
 private:
-	int32 IndentationLevel = 0; // Track indentation depth
+	uint32 AddAssetReference(UObject* Asset, TFunctionRef<uint32()> GetHashFunc);
+
+	// Add a field to the current object context
+	void AddField(const FName& InFieldName, const FString& InFieldValue)
+	{
+		FFieldArray& Fields = ObjectContextStack.IsEmpty() ? GlobalFields : ObjectsReferences.FindOrAdd(ObjectContextStack.Top()).Fields;
+		Fields.Emplace(InFieldName, InFieldValue);
+	}
+
+	// Pretty enum string for UENUMs; falls back to the integer value if no UENUM is available.
+	template <typename T>
+	static FString EnumToPrettyString(T InEnumValue)
+	{
+		// non-null only for UENUMs
+		if (const UEnum* EnumType = StaticEnum<T>())
+		{
+			const int64 RawValue = static_cast<int64>(InEnumValue);
+			const FString Name = EnumType->GetNameStringByValue(RawValue);
+			if (!Name.IsEmpty())
+			{
+				return Name + TEXT("(") + LexToString(static_cast<__underlying_type(T)>(InEnumValue)) + TEXT(")");
+			}
+		}
+
+		// Fallback to the integer value
+		return LexToString(static_cast<__underlying_type(T)>(InEnumValue));
+	}
+
+	// Helper to stringify fields
+	template <typename T>
+	static FString FieldValueToString(const T& InFieldValue)
+	{
+		if constexpr (std::is_same<T, bool>::value)
+		{
+			return InFieldValue ? TEXT("1") : TEXT("0");
+		}
+		else if constexpr (TIsEnum<T>::Value)
+		{
+			return EnumToPrettyString(InFieldValue);
+		}
+		else if constexpr (TIsFloatingPoint<T>::Value)
+		{
+			static const int32 Precision = 4;
+			FString S = FString::Printf(TEXT("%.*f"), Precision, InFieldValue);
+			
+			while (S.Len() && S.EndsWith(TEXT("0")))
+			{
+				S.LeftChopInline(1);
+			}
+			
+			if (S.Len() && S.EndsWith(TEXT(".")))
+			{
+				S.LeftChopInline(1);
+			}
+
+			return S.IsEmpty() ? TEXT("0") : S;
+		}
+		else if constexpr (TIsIntegral<T>::Value)
+		{
+			return LexToString(InFieldValue);
+		}
+		else
+		{
+			static_assert(!sizeof(T), "FieldValueToString only for numeric/enum types");
+			return FString();
+		}
+	}
+
+private:
+	// For a given asset, store its hash & type.
+	struct FAssetHash
+	{
+		uint32	Hash = 0;
+		FName	AssetType = NAME_None;
+	};
+	TMap<FName, FAssetHash> AssetsHashes;
+
+	// Array of name/value for fields.
+	typedef TArray<TTuple<FName, FString>>	FFieldArray;
+
+	// Track assets references & fields values for a given object.
+	struct FObjectHash
+	{
+		TSet<FName>	ReferencedAssets;
+		FFieldArray Fields;
+	};
+
+	// Map of objects names to their data.
+	TMap<FString, FObjectHash> ObjectsReferences;
+
+	// Stack of objects.
+	TArray<FString> ObjectContextStack;
+
+	// Global hashing fields not tied to any particular object.
+	FFieldArray GlobalFields;
 };
 
 // Templated operator overloads for handling arrays
@@ -58,48 +230,29 @@ FArchive& operator<<(FHLODHashBuilder& Ar, TArray<TElementType, TAllocatorType>&
 	return Self << InArray;
 }
 
-class FHLODHashContext
-{
-public:
-	template <typename... TArgs>
-	explicit FHLODHashContext(UE::Core::TCheckedFormatString<FString::FmtCharType, TArgs...> Format, TArgs... Args)
-		: Context(FString::Printf(Format, Args...)) {}
-
-	explicit FHLODHashContext(FName InName)
-		: Context(InName.ToString()) {}
-
-	const TCHAR* GetContext() const { return *Context; }
-
-	/** Overload for FHLODHashContext to log structured properties */
-	friend FArchive& operator<<(FArchive& Ar, FHLODHashContext InContext)
-	{
-		static_cast<FHLODHashBuilder&>(Ar).LogContext(InContext.GetContext(), true);
-		return Ar;
-	}
-
-private:
-	FString Context;
-};
-
 class FHLODHashScope
 {
 public:
-	FHLODHashScope(FHLODHashBuilder& InBuilder, const TCHAR* Format)
-		: Builder(InBuilder)
+	enum class EFlags : uint8
 	{
-		Builder.PushContext(Format);
-	}
+		None = 0,
+		ResetHash = 1 << 0,		//!< Reset the hash when entering this scope
+	};
 
-	template <typename... Args>
-	FHLODHashScope(FHLODHashBuilder& InBuilder, const TCHAR* Format, Args... Arguments)
+	FHLODHashScope(FHLODHashBuilder& InBuilder, const UObject* InObjectContext, EFlags InFlags = EFlags::None)
 		: Builder(InBuilder)
 	{
-		Builder.PushContext(*FString::Printf(Format, Arguments...));
+		Builder.PushObjectContext(InObjectContext);
+
+		if (EnumHasAnyFlags(InFlags, EFlags::ResetHash))
+		{
+			Builder.Reset();
+		}
 	}
 
 	~FHLODHashScope()
 	{
-		Builder.PopContext();
+		Builder.PopObjectContext();
 	}
 
 private:
diff --git a/Engine/Source/Runtime/Engine/Public/WorldPartition/HLOD/HLODLayer.h b/Engine/Source/Runtime/Engine/Public/WorldPartition/HLOD/HLODLayer.h
--- a/Engine/Source/Runtime/Engine/Public/WorldPartition/HLOD/HLODLayer.h
+++ b/Engine/Source/Runtime/Engine/Public/WorldPartition/HLOD/HLODLayer.h
@@ -21,6 +21,7 @@ class AActor;
 class AWorldPartitionHLOD;
 class UHLODBuilder;
 class UHLODBuilderSettings;
+class FHLODHashBuilder;
 class UMaterial;
 class UWorldPartition;
 class UWorldPartitionHLODModifier;
@@ -59,6 +60,7 @@ public:
 	UHLODLayer* GetParentLayer() const { return ParentLayer; }
 	void SetParentLayer(UHLODLayer* InParentLayer) { ParentLayer = InParentLayer; }
 	ENGINE_API bool DoesRequireWarmup() const;
+	ENGINE_API void ComputeHLODHash(FHLODHashBuilder& InHLODHashBuilder) const;
 
 	UE_DEPRECATED(5.7, "Deprecated. These streaming grid properties are now specified in the partition's settings.")
 	ENGINE_API FName GetRuntimeGrid(uint32 InHLODLevel) const;
diff --git a/Engine/Source/Runtime/Engine/Public/WorldPartition/HLOD/HLODSourceActors.h b/Engine/Source/Runtime/Engine/Public/WorldPartition/HLOD/HLODSourceActors.h
--- a/Engine/Source/Runtime/Engine/Public/WorldPartition/HLOD/HLODSourceActors.h
+++ b/Engine/Source/Runtime/Engine/Public/WorldPartition/HLOD/HLODSourceActors.h
@@ -5,6 +5,7 @@
 #include "HLODSourceActors.generated.h"
 
 
+class FHLODHashBuilder;
 class ULevelStreaming;
 class UHLODLayer;
 
@@ -17,10 +18,13 @@ class UWorldPartitionHLODSourceActors : public UObject
 public:
 #if WITH_EDITOR
 	virtual ULevelStreaming* LoadSourceActors(bool& bOutDirty) const PURE_VIRTUAL(UWorldPartitionHLODSourceActors::LoadSourceActors, return nullptr; );
-	ENGINE_API virtual uint32 GetHLODHash() const;
+	ENGINE_API virtual void ComputeHLODHash(FHLODHashBuilder& HashBuilder) const;
 
 	ENGINE_API void SetHLODLayer(const UHLODLayer* HLODLayer);
 	ENGINE_API const UHLODLayer* GetHLODLayer() const;
+
+	UE_DEPRECATED(5.7, "GetHLODHash() has been replaced by ComputeHLODHash()")
+	ENGINE_API virtual uint32 GetHLODHash() const final { return 0; }
 #endif
 
 private:
diff --git a/Engine/Source/Runtime/Engine/Public/WorldPartition/HLOD/HLODSourceActorsFromCell.h b/Engine/Source/Runtime/Engine/Public/WorldPartition/HLOD/HLODSourceActorsFromCell.h
--- a/Engine/Source/Runtime/Engine/Public/WorldPartition/HLOD/HLODSourceActorsFromCell.h
+++ b/Engine/Source/Runtime/Engine/Public/WorldPartition/HLOD/HLODSourceActorsFromCell.h
@@ -15,9 +15,9 @@ class UWorldPartitionHLODSourceActorsFromCell : public UWorldPartitionHLODSource
 public:
 #if WITH_EDITOR
 	ENGINE_API virtual ULevelStreaming* LoadSourceActors(bool& bOutDirty) const override;
-	ENGINE_API virtual uint32 GetHLODHash() const override;
+	ENGINE_API virtual void ComputeHLODHash(FHLODHashBuilder& InHashBuilder) const override;
 
-	ENGINE_API static uint32 GetHLODHash(const TArray<FWorldPartitionRuntimeCellObjectMapping>& InSourceActors);
+	ENGINE_API static uint32 GetSourceActorsHash(const TArray<FWorldPartitionRuntimeCellObjectMapping>& InSourceActors);
 
 	ENGINE_API void SetActors(const TArray<FWorldPartitionRuntimeCellObjectMapping>&& InSourceActors);
 	ENGINE_API const TArray<FWorldPartitionRuntimeCellObjectMapping>& GetActors() const;
diff --git a/Engine/Source/Runtime/Engine/Public/WorldPartition/HLOD/HLODSourceActorsFromLevel.h b/Engine/Source/Runtime/Engine/Public/WorldPartition/HLOD/HLODSourceActorsFromLevel.h
--- a/Engine/Source/Runtime/Engine/Public/WorldPartition/HLOD/HLODSourceActorsFromLevel.h
+++ b/Engine/Source/Runtime/Engine/Public/WorldPartition/HLOD/HLODSourceActorsFromLevel.h
@@ -14,7 +14,7 @@ class UWorldPartitionHLODSourceActorsFromLevel : public UWorldPartitionHLODSourc
 public:
 #if WITH_EDITOR
 	ENGINE_API virtual ULevelStreaming* LoadSourceActors(bool& bOutDirty) const override;
-	ENGINE_API virtual uint32 GetHLODHash() const override;
+	ENGINE_API virtual void ComputeHLODHash(FHLODHashBuilder& InHashBuilder) const override;
 
 	ENGINE_API void SetSourceLevel(const UWorld* InSourceLevel);
 	ENGINE_API const TSoftObjectPtr<UWorld>& GetSourceLevel() const;
diff --git a/Engine/Source/Runtime/Landscape/Private/LandscapeHLODBuilder.cpp b/Engine/Source/Runtime/Landscape/Private/LandscapeHLODBuilder.cpp
--- a/Engine/Source/Runtime/Landscape/Private/LandscapeHLODBuilder.cpp
+++ b/Engine/Source/Runtime/Landscape/Private/LandscapeHLODBuilder.cpp
@@ -30,6 +30,8 @@
 #include "Serialization/ArchiveCrc32.h"
 #include "Engine/HLODProxy.h"
 
+#include "WorldPartition/HLOD/HLODHashBuilder.h"
+
 #endif // WITH_EDITOR
 
 ULandscapeHLODBuilder::ULandscapeHLODBuilder(const FObjectInitializer& ObjectInitializer)
@@ -39,137 +41,82 @@ ULandscapeHLODBuilder::ULandscapeHLODBuilder(const FObjectInitializer& ObjectIni
 
 #if WITH_EDITOR
 
-uint32 ULandscapeHLODBuilder::ComputeHLODHash(const UActorComponent* InSourceComponent) const
+bool ULandscapeHLODBuilder::ComputeHLODHash(FHLODHashBuilder& HashBuilder, const UActorComponent* InSourceComponent) const
 {
-	FArchiveCrc32 Ar;
-
-	// Base lanscape HLOD key, changing this will force a rebuild of all landscape HLODs
-	FString HLODBaseKey = "51EDA4343E26409DA6ABEACE14C138D1";
-	Ar << HLODBaseKey;
-
-	if (const ULandscapeComponent* LSComponent = Cast<ULandscapeComponent>(InSourceComponent))
+	const ULandscapeComponent* LSComponent = Cast<ULandscapeComponent>(InSourceComponent);
+	if (LSComponent == nullptr)
 	{
-		UE_LOG(LogHLODBuilder, VeryVerbose, TEXT("ULandscapeHLODBuilder::ComputeHLODHash %s"), *InSourceComponent->GetName());
+		return false;
+	}
+	
+	// Base lanscape HLOD key, changing this will force a rebuild of all landscape HLODs
+	FString HLODBaseKey = TEXT("51EDA4343E26409DA6ABEACE14C138D1");
+	HashBuilder.HashField(HLODBaseKey, TEXT("LandscapeHLODBuilderBaseKey"));
 
-		ALandscapeProxy* LSProxy = LSComponent->GetLandscapeProxy();
+	ALandscapeProxy* LSProxy = LSComponent->GetLandscapeProxy();
 
-		// LS LOD setup
-		TArray<float> LODScreenSize = LSProxy->GetLODScreenSizeArray();
-		Ar << LODScreenSize;
-		UE_LOG(LogHLODBuilder, VeryVerbose, TEXT("     - LODScreenSize = %x"), Ar.GetCrc());
+	// LS LOD setup
+	HashBuilder.HashField(LSProxy->GetLODScreenSizeArray(), TEXT("LODScreenSize"));
 		
-		// LS Transform
-		uint32 TransformHash = UHLODProxy::GetCRC(LSComponent->GetComponentTransform());
-		Ar << TransformHash;
-		UE_LOG(LogHLODBuilder, VeryVerbose, TEXT("     - TransformHash = %x"), Ar.GetCrc());
+	// LS Transform
+	HashBuilder.HashField(LSComponent->GetComponentTransform(), TEXT("Transform"));
 
-		// LS Content - Heightmap & Weightmaps
-		uint32 LSContentHash = LSComponent->ComputeLayerHash(/*InReturnEditingHash=*/ false);
-		Ar << LSContentHash;
-		UE_LOG(LogHLODBuilder, VeryVerbose, TEXT("     - LSContentHash = %x"), Ar.GetCrc());
+	// LS Content - Heightmap & Weightmaps
+	uint32 LSContentHash = LSComponent->ComputeLayerHash(/*InReturnEditingHash=*/ false);
+	HashBuilder.HashField(LSContentHash, TEXT("LSContentHash"));
 
-		TSet<UTexture*> UsedTextures;
+	TSet<UTexture*> UsedTextures;
 
-		// LS Materials
+	// LS Materials
+	{
+		TArray<UMaterialInterface*> UsedMaterials;
+		if (LSProxy->bUseDynamicMaterialInstance)
 		{
-			TArray<UMaterialInterface*> UsedMaterials;
-			if (LSProxy->bUseDynamicMaterialInstance)
-			{
-				UsedMaterials.Append(LSComponent->MaterialInstancesDynamic);
-			}
-			else
-			{
-				UsedMaterials.Append(LSComponent->MaterialInstances);
-			}
-
-			UsedMaterials.Add(LSComponent->OverrideMaterial);
-			UsedMaterials.Add(LSComponent->OverrideHoleMaterial);
-			TArray<uint32> UsedMaterialsCRC;
-			for (UMaterialInterface* MaterialInterface : UsedMaterials)
-			{
-				if (MaterialInterface)
-				{
-					uint32 MaterialCRC = UHLODProxy::GetCRC(MaterialInterface);
-					UsedMaterialsCRC.Add(MaterialCRC);
-
-					TArray<UTexture*> Textures;
-					MaterialInterface->GetUsedTextures(Textures);
-					UsedTextures.Append(Textures);
-				}
-			}
-			UsedMaterialsCRC.Sort();
-			Ar << UsedMaterialsCRC;
-			UE_LOG(LogHLODBuilder, VeryVerbose, TEXT("     - UsedMaterialsCRC = %x"), Ar.GetCrc());
+			HashBuilder << LSComponent->MaterialInstancesDynamic;
 		}
-
-		// LS Textures
+		else
 		{
-			TArray<uint32> UsedTexturesCRC;
-			for (UTexture* Texture : UsedTextures)
-			{
-				uint32 TextureCRC = UHLODProxy::GetCRC(Texture);
-				UsedTexturesCRC.Add(TextureCRC);
-			}
-			UsedTexturesCRC.Sort();
-			Ar << UsedTexturesCRC;
-			UE_LOG(LogHLODBuilder, VeryVerbose, TEXT("     - UsedTexturesCRC = %x"), Ar.GetCrc());
+			HashBuilder << LSComponent->MaterialInstances;
 		}
 
-		// Nanite enabled?
-		bool bNaniteEnabled = LSProxy->IsNaniteEnabled() || GetDefault<UStaticMesh>()->IsNaniteForceEnabled();
-		Ar << bNaniteEnabled;
-		UE_LOG(LogHLODBuilder, VeryVerbose, TEXT("     - NaniteEnabled = %d"), bNaniteEnabled);
-		if (bNaniteEnabled)
-		{
-			int32 NaniteLODIndex = LSProxy->GetNaniteLODIndex();
-			int32 NanitePositionPrecision = LSProxy->GetNanitePositionPrecision();
-			float NaniteMaxEdgeLengthFactor = LSProxy->GetNaniteMaxEdgeLengthFactor();
-			Ar << NaniteLODIndex << NanitePositionPrecision << NaniteMaxEdgeLengthFactor;
-			UE_LOG(LogHLODBuilder, VeryVerbose, TEXT("     - NaniteLODIndex = %d"), NaniteLODIndex);
-			UE_LOG(LogHLODBuilder, VeryVerbose, TEXT("     - NanitePositionPrecision = %d"), NanitePositionPrecision);
-			UE_LOG(LogHLODBuilder, VeryVerbose, TEXT("     - NaniteMaxEdgeLengthFactor = %f"), NaniteMaxEdgeLengthFactor);
-
-			bool bNaniteSkirtEnabled = LSProxy->IsNaniteSkirtEnabled();
-			Ar << bNaniteSkirtEnabled;
-			UE_LOG(LogHLODBuilder, VeryVerbose, TEXT("     - NaniteSkirtEnabled = %d"), bNaniteSkirtEnabled);
-			if (bNaniteSkirtEnabled)
-			{
-				float NaniteSkirtDepth = LSProxy->GetNaniteSkirtDepth();
-				Ar << NaniteSkirtDepth;
-				UE_LOG(LogHLODBuilder, VeryVerbose, TEXT("     - NaniteSkirtDepth = %f"), NaniteSkirtDepth);
-			}
-		}
+		HashBuilder << LSComponent->OverrideMaterial;
+		HashBuilder << LSComponent->OverrideHoleMaterial;
+	}
 
-		// HLODTextureSize
-		ELandscapeHLODTextureSizePolicy HLODTextureSizePolicy = LSProxy->HLODTextureSizePolicy;
-		Ar << HLODTextureSizePolicy;
-		UE_LOG(LogHLODBuilder, VeryVerbose, TEXT("     - HLODTextureSizePolicy = %x"), HLODTextureSizePolicy);
-		if (HLODTextureSizePolicy == ELandscapeHLODTextureSizePolicy::SpecificSize)
-		{
-			int32 HLODTextureSize = LSProxy->HLODTextureSize;
-			Ar << HLODTextureSize;
-			UE_LOG(LogHLODBuilder, VeryVerbose, TEXT("     - HLODTextureSize = %d"), HLODTextureSize);
-		}
+	// Nanite enabled?
+	bool bNaniteEnabled = LSProxy->IsNaniteEnabled() || GetDefault<UStaticMesh>()->IsNaniteForceEnabled();
+	HashBuilder.HashField(bNaniteEnabled, TEXT("NaniteEnabled"));
+	if (bNaniteEnabled)
+	{
+		HashBuilder.HashField(LSProxy->GetNaniteLODIndex(), TEXT("NaniteLODIndex"));
+		HashBuilder.HashField(LSProxy->GetNanitePositionPrecision(), TEXT("NanitePositionPrecision"));
+		HashBuilder.HashField(LSProxy->GetNaniteMaxEdgeLengthFactor(), TEXT("NaniteMaxEdgeLengthFactor"));
 
-		// HLODMeshSourceLOD
-		ELandscapeHLODMeshSourceLODPolicy HLODMeshSourceLODPolicy = LSProxy->HLODMeshSourceLODPolicy;
-		Ar << HLODMeshSourceLODPolicy;
-		UE_LOG(LogHLODBuilder, VeryVerbose, TEXT("     - HLODMeshSourceLODPolicy = %x"), HLODMeshSourceLODPolicy);
-		if (HLODMeshSourceLODPolicy == ELandscapeHLODMeshSourceLODPolicy::SpecificLOD)
+		HashBuilder.HashField(LSProxy->IsNaniteSkirtEnabled(), TEXT("NaniteSkirtEnabled"));
+		if (LSProxy->IsNaniteSkirtEnabled())
 		{
-			int32 HLODMeshSourceLOD = LSProxy->HLODMeshSourceLOD;
-			Ar << HLODMeshSourceLOD;
-			UE_LOG(LogHLODBuilder, VeryVerbose, TEXT("     - HLODMeshSourceLOD = %d"), HLODMeshSourceLOD);
+			HashBuilder.HashField(LSProxy->GetNaniteSkirtDepth(), TEXT("NaniteSkirtDepth"));
 		}
+	}
 
-		// Project max texture size for landscape HLOD textures
-		const ULandscapeSettings* LandscapeSettings = GetDefault<ULandscapeSettings>();
-		int32 ProjectHLODMaxTextureSize = LandscapeSettings->GetHLODMaxTextureSize();
-		Ar << ProjectHLODMaxTextureSize;
-		UE_LOG(LogHLODBuilder, VeryVerbose, TEXT("     - ProjectHLODMaxTextureSize = %d"), ProjectHLODMaxTextureSize);
+	// HLODTextureSize
+	HashBuilder.HashField(LSProxy->HLODTextureSizePolicy, TEXT("HLODTextureSizePolicy"));
+	if (LSProxy->HLODTextureSizePolicy == ELandscapeHLODTextureSizePolicy::SpecificSize)
+	{
+		HashBuilder.HashField(LSProxy->HLODTextureSize, TEXT("HLODTextureSize"));
 	}
 
-	return Ar.GetCrc();
+	// HLODMeshSourceLOD
+	HashBuilder.HashField(LSProxy->HLODMeshSourceLODPolicy, TEXT("HLODMeshSourceLODPolicy"));
+	if (LSProxy->HLODMeshSourceLODPolicy == ELandscapeHLODMeshSourceLODPolicy::SpecificLOD)
+	{
+		HashBuilder.HashField(LSProxy->HLODMeshSourceLOD, TEXT("HLODMeshSourceLOD"));
+	}
+
+	// Project max texture size for landscape HLOD textures
+	HashBuilder.HashField(GetDefault<ULandscapeSettings>()->GetHLODMaxTextureSize(), TEXT("ProjectHLODMaxTextureSize"));
+
+	return true;
 }
 
 static int32 GetMeshTextureSizeFromTargetTexelDensity(const FMeshDescription& InMesh, float InTargetTexelDensity)
diff --git a/Engine/Source/Runtime/Landscape/Public/LandscapeHLODBuilder.h b/Engine/Source/Runtime/Landscape/Public/LandscapeHLODBuilder.h
--- a/Engine/Source/Runtime/Landscape/Public/LandscapeHLODBuilder.h
+++ b/Engine/Source/Runtime/Landscape/Public/LandscapeHLODBuilder.h
@@ -14,7 +14,7 @@ class ULandscapeHLODBuilder : public UHLODBuilder
 
 public:
 #if WITH_EDITOR
-	LANDSCAPE_API virtual uint32 ComputeHLODHash(const UActorComponent* InSourceComponent) const override;
+	LANDSCAPE_API virtual bool ComputeHLODHash(FHLODHashBuilder& HashBuilder, const UActorComponent* InSourceComponent) const override;
 
 	/**
 	 * Components created with this method need to be properly outered & assigned to your target actor.
