Update the "r.DebugSafeZone.Mode" tooltip

#rnx

diff --git a/Engine/Source/Runtime/Engine/Private/HUD.cpp b/Engine/Source/Runtime/Engine/Private/HUD.cpp
--- a/Engine/Source/Runtime/Engine/Private/HUD.cpp
+++ b/Engine/Source/Runtime/Engine/Private/HUD.cpp
@@ -35,14 +35,13 @@ DEFINE_LOG_CATEGORY_STATIC(LogHUD, Log, All);
 FOnShowDebugInfo AHUD::OnShowDebugInfo;
 FOnHUDPostRender AHUD::OnHUDPostRender;
 
-// Should we visualize the safe zone? (and if so, title or action?)
+// Should we visualize the safe zone?
 TAutoConsoleVariable<int32> GSafeZoneVisualizationModeCVar(
 	TEXT("r.DebugSafeZone.Mode"),
 	0,
-	TEXT("The safe zone visualization mode (0..2)\n")
+	TEXT("The safe zone visualization mode (0..1)\n")
 	TEXT(" 0: Disabled (default)\n")
-	TEXT(" 1: Show Title Safe Zone\n")
-	TEXT(" 2: Show Action Safe Zone"));
+	TEXT(" 1: Show Title Safe Zone"));
 
 // How opaque should the safe zone visualization be?
 TAutoConsoleVariable<float> GSafeZoneVisualizationAlphaCVar(
diff --git a/Engine/Source/Runtime/UMG/Public/Components/SafeZone.h b/Engine/Source/Runtime/UMG/Public/Components/SafeZone.h
--- a/Engine/Source/Runtime/UMG/Public/Components/SafeZone.h
+++ b/Engine/Source/Runtime/UMG/Public/Components/SafeZone.h
@@ -19,11 +19,10 @@
  *   r.DebugSafeZone.TitleRatio 0.96
  *   r.DebugActionZone.ActionRatio 0.96
  * 
- * To enable a red band to visualize the safe zone, use this console command,
- * r.DebugSafeZone.Mode controls the debug visualization overlay (0..2, default 0).
+ * To enable a border to visualize the safe zone, use this console command,
+ * r.DebugSafeZone.Mode controls the debug visualization overlay (0..1, default 0).
  *   0: Do not display the safe zone overlay.
  *   1: Display the overlay for the title safe zone.
- *   2: Display the overlay for the action safe zone.
  */
 UCLASS(MinimalAPI)
 class USafeZone : public UContentWidget
