Fix shaders wrongly compiled with VELOCITY_SUPPORT_TEMPORAL_RESPONSIVENESS when r.Velocity.TemporalResponsiveness.Supported is disabled in runtime

#tests

diff --git a/Engine/Source/Runtime/Engine/Private/ShaderCompiler/ShaderCompiler.cpp b/Engine/Source/Runtime/Engine/Private/ShaderCompiler/ShaderCompiler.cpp
--- a/Engine/Source/Runtime/Engine/Private/ShaderCompiler/ShaderCompiler.cpp
+++ b/Engine/Source/Runtime/Engine/Private/ShaderCompiler/ShaderCompiler.cpp
@@ -3825,16 +3825,15 @@ void GlobalBeginCompileShader(
 	}
 
 	{
-		static IConsoleVariable* CVarTemporalResponsivenessSupported = IConsoleManager::Get().FindConsoleVariable(TEXT("r.Velocity.TemporalResponsiveness.Supported"));
-		static IConsoleVariable* CVarPixelShaderMotionVectorWorldOffsetSupported = IConsoleManager::Get().FindConsoleVariable(TEXT("r.Velocity.PixelShaderMotionVectorWorldOffset.Supported"));
-
 		if (VelocityEncodeDepth((EShaderPlatform)Target.Platform))
 		{
+			static FShaderPlatformCachedIniValue<int32> CVarTemporalResponsivenessSupported(TEXT("r.Velocity.TemporalResponsiveness.Supported"));
+			static FShaderPlatformCachedIniValue<int32> CVarPixelShaderMotionVectorWorldOffsetSupported(TEXT("r.Velocity.PixelShaderMotionVectorWorldOffset.Supported"));
+
 			SET_SHADER_DEFINE(Input.Environment, VELOCITY_ENCODE_DEPTH, 1);
-			SET_SHADER_DEFINE(Input.Environment, VELOCITY_SUPPORT_TEMPORAL_RESPONSIVENESS, CVarTemporalResponsivenessSupported ? CVarTemporalResponsivenessSupported->GetInt() != 0 : 0);
-			SET_SHADER_DEFINE(Input.Environment, VELOCITY_SUPPORT_PS_MOTION_VECTOR_WORLD_OFFSET, 
-				CVarPixelShaderMotionVectorWorldOffsetSupported ? CVarPixelShaderMotionVectorWorldOffsetSupported->GetInt() != 0 : 0);
-		}
+			SET_SHADER_DEFINE(Input.Environment, VELOCITY_SUPPORT_TEMPORAL_RESPONSIVENESS, CVarTemporalResponsivenessSupported.Get((EShaderPlatform)Target.Platform) !=0 ? 1 : 0);
+			SET_SHADER_DEFINE(Input.Environment, VELOCITY_SUPPORT_PS_MOTION_VECTOR_WORLD_OFFSET, CVarPixelShaderMotionVectorWorldOffsetSupported.Get((EShaderPlatform)Target.Platform) != 0 ? 1 : 0);
+		}	
 		else
 		{
 			SET_SHADER_DEFINE(Input.Environment, VELOCITY_ENCODE_DEPTH, 0);
