Optimisation - make the physics force component cache the cluster union proxy component.

#FailedValidations Preflight

diff --git a/Engine/Source/Runtime/Engine/Classes/PhysicsEngine/ClusterUnionReplicatedProxyComponent.h b/Engine/Source/Runtime/Engine/Classes/PhysicsEngine/ClusterUnionReplicatedProxyComponent.h
--- a/Engine/Source/Runtime/Engine/Classes/PhysicsEngine/ClusterUnionReplicatedProxyComponent.h
+++ b/Engine/Source/Runtime/Engine/Classes/PhysicsEngine/ClusterUnionReplicatedProxyComponent.h
@@ -53,7 +53,7 @@ public:
 	void MarkPendingDeletion() { bIsPendingDeletion = true; }
 
 	UFUNCTION()
-	bool IsPendingDeletion() { return bIsPendingDeletion; }
+	bool IsPendingDeletion() const { return bIsPendingDeletion; }
 
 	/* reinitialize the internal transient state of the replicated proxy */
 	ENGINE_API void ResetTransientState();
