SplineComponent: Harden functions dealing with 'time' which do not properly convert between input key and time

#FailedValidations Preflight

diff --git a/Engine/Source/Runtime/Engine/Classes/Components/SplineComponent.h b/Engine/Source/Runtime/Engine/Classes/Components/SplineComponent.h
--- a/Engine/Source/Runtime/Engine/Classes/Components/SplineComponent.h
+++ b/Engine/Source/Runtime/Engine/Classes/Components/SplineComponent.h
@@ -7,6 +7,7 @@
 #include "Components/PrimitiveComponent.h"
 #include "Curves/Spline.h"
 #include "Math/InterpCurve.h"
+#include "BoxTypes.h"
 #include "SplineComponent.generated.h"
 
 class FPrimitiveSceneProxy;
@@ -699,6 +700,10 @@ public:
 	UFUNCTION(BlueprintCallable, Category = Spline)
 	ENGINE_API FVector GetDefaultUpVector(ESplineCoordinateSpace::Type CoordinateSpace) const;
 
+	/** Convert a time value (in the range [0, Duration]) to an input key. */
+	UFUNCTION(BlueprintCallable, Category = Spline)
+	ENGINE_API float GetInputKeyValueAtTime(float Time) const;
+
 	/** Given a distance along the length of this spline, return the corresponding input key at that point */
 	/** This method has been deprecated because it was incorrectly returning the input key at time. To maintain the same behavior,
 	 *  replace it with GetTimeAtDistanceAlongSpline. To actually get the input key, instead call GetInputKeyValueAtDistanceAlongSpline. */
@@ -1016,6 +1021,9 @@ protected:
 	
 private:
 	
+	/** Returns an interval representing the full range of valid input keys, or an empty interval if the spline is empty. */
+	UE::Geometry::FInterval1f GetInputKeyRange() const;
+	
 	/**
 	 * This function takes on the roles of the previously existing functions GetPositionPointSafe, GetRotationPointSafe, and GetScalePointSafe.
 	 * Specifically, it makes sure that a provided index is valid by
diff --git a/Engine/Source/Runtime/Engine/Private/Components/SplineComponent.cpp b/Engine/Source/Runtime/Engine/Private/Components/SplineComponent.cpp
--- a/Engine/Source/Runtime/Engine/Private/Components/SplineComponent.cpp
+++ b/Engine/Source/Runtime/Engine/Private/Components/SplineComponent.cpp
@@ -1646,6 +1646,12 @@ float USplineComponent::GetInputKeyAtDistanceAlongSpline(float Distance) const
 	return GetTimeAtDistanceAlongSpline(Distance);
 }
 
+float USplineComponent::GetInputKeyValueAtTime(float Time) const
+{
+	float T = Duration > 0.f ? FMath::Clamp(Time, 0.f, Duration) / Duration : 0.f;
+	return GetInputKeyRange().Interpolate(T);
+}
+
 float USplineComponent::GetInputKeyValueAtDistanceAlongSpline(float Distance) const
 {
 	const int32 NumPoints = GetNumberOfSplinePoints();
@@ -1669,10 +1675,12 @@ float USplineComponent::GetTimeAtDistanceAlongSpline(float Distance) const
 		return 0.0f;
 	}
 
-	const float TimeMultiplier = Duration / (bClosedLoop ? NumPoints : (NumPoints - 1.0f));
-	return (UE::SplineComponent::ShouldUseSplineCurves()
+	const float Param = UE::SplineComponent::ShouldUseSplineCurves() 
 		? WarninglessSplineCurves().ReparamTable.Eval(Distance, 0.0f)
-		: Spline.GetParameterAtDistance(Distance)) * TimeMultiplier;
+		: Spline.GetParameterAtDistance(Distance);
+
+	// Convert Param to T and interpolate the duration interval (which yields Time).
+	return UE::Geometry::FInterval1f(0.f, Duration).Interpolate(GetInputKeyRange().GetT(Param));
 }
 
 FVector USplineComponent::GetLocationAtDistanceAlongSpline(float Distance, ESplineCoordinateSpace::Type CoordinateSpace) const
@@ -1769,10 +1777,7 @@ FVector USplineComponent::GetLocationAtTime(float Time, ESplineCoordinateSpace::
 	}
 	else
 	{
-		const int32 NumPoints = GetNumberOfSplinePoints();
-		const int32 NumSegments = bClosedLoop ? NumPoints : NumPoints - 1;
-		const float TimeMultiplier = NumSegments / Duration;
-		return GetLocationAtSplineInputKey(Time * TimeMultiplier, CoordinateSpace);
+		return GetLocationAtSplineInputKey(GetInputKeyValueAtTime(Time), CoordinateSpace);
 	}
 }
 
@@ -1789,10 +1794,7 @@ FVector USplineComponent::GetDirectionAtTime(float Time, ESplineCoordinateSpace:
 	}
 	else
 	{
-		const int32 NumPoints = GetNumberOfSplinePoints();
-		const int32 NumSegments = bClosedLoop ? NumPoints : NumPoints - 1;
-		const float TimeMultiplier = NumSegments / Duration;
-		return GetDirectionAtSplineInputKey(Time * TimeMultiplier, CoordinateSpace);
+		return GetDirectionAtSplineInputKey(GetInputKeyValueAtTime(Time), CoordinateSpace);
 	}
 }
 
@@ -1809,10 +1811,7 @@ FVector USplineComponent::GetTangentAtTime(float Time, ESplineCoordinateSpace::T
 	}
 	else
 	{
-		const int32 NumPoints = GetNumberOfSplinePoints();
-		const int32 NumSegments = bClosedLoop ? NumPoints : NumPoints - 1;
-		const float TimeMultiplier = NumSegments / Duration;
-		return GetTangentAtSplineInputKey(Time * TimeMultiplier, CoordinateSpace);
+		return GetTangentAtSplineInputKey(GetInputKeyValueAtTime(Time), CoordinateSpace);
 	}
 }
 
@@ -1829,10 +1828,7 @@ FRotator USplineComponent::GetRotationAtTime(float Time, ESplineCoordinateSpace:
 	}
 	else
 	{
-		const int32 NumPoints = GetNumberOfSplinePoints();
-		const int32 NumSegments = bClosedLoop ? NumPoints : NumPoints - 1;
-		const float TimeMultiplier = NumSegments / Duration;
-		return GetRotationAtSplineInputKey(Time * TimeMultiplier, CoordinateSpace);
+		return GetRotationAtSplineInputKey(GetInputKeyValueAtTime(Time), CoordinateSpace);
 	}
 }
 
@@ -1849,10 +1845,7 @@ FQuat USplineComponent::GetQuaternionAtTime(float Time, ESplineCoordinateSpace::
 	}
 	else
 	{
-		const int32 NumPoints = GetNumberOfSplinePoints();
-		const int32 NumSegments = bClosedLoop ? NumPoints : NumPoints - 1;
-		const float TimeMultiplier = NumSegments / Duration;
-		return GetQuaternionAtSplineInputKey(Time * TimeMultiplier, CoordinateSpace);
+		return GetQuaternionAtSplineInputKey(GetInputKeyValueAtTime(Time), CoordinateSpace);
 	}
 }
 
@@ -1869,10 +1862,7 @@ FVector USplineComponent::GetUpVectorAtTime(float Time, ESplineCoordinateSpace::
 	}
 	else
 	{
-		const int32 NumPoints = GetNumberOfSplinePoints();
-		const int32 NumSegments = bClosedLoop ? NumPoints : NumPoints - 1;
-		const float TimeMultiplier = NumSegments / Duration;
-		return GetUpVectorAtSplineInputKey(Time * TimeMultiplier, CoordinateSpace);
+		return GetUpVectorAtSplineInputKey(GetInputKeyValueAtTime(Time), CoordinateSpace);
 	}
 }
 
@@ -1889,10 +1879,7 @@ FVector USplineComponent::GetRightVectorAtTime(float Time, ESplineCoordinateSpac
 	}
 	else
 	{
-		const int32 NumPoints = GetNumberOfSplinePoints();
-		const int32 NumSegments = bClosedLoop ? NumPoints : NumPoints - 1;
-		const float TimeMultiplier = NumSegments / Duration;
-		return GetRightVectorAtSplineInputKey(Time * TimeMultiplier, CoordinateSpace);
+		return GetRightVectorAtSplineInputKey(GetInputKeyValueAtTime(Time), CoordinateSpace);
 	}
 }
 
@@ -1909,10 +1896,7 @@ float USplineComponent::GetRollAtTime(float Time, ESplineCoordinateSpace::Type C
 	}
 	else
 	{
-		const int32 NumPoints = GetNumberOfSplinePoints();
-		const int32 NumSegments = bClosedLoop ? NumPoints : NumPoints - 1;
-		const float TimeMultiplier = NumSegments / Duration;
-		return GetRollAtSplineInputKey(Time * TimeMultiplier, CoordinateSpace);
+		return GetRollAtSplineInputKey(GetInputKeyValueAtTime(Time), CoordinateSpace);
 	}
 }
 
@@ -1929,10 +1913,7 @@ FTransform USplineComponent::GetTransformAtTime(float Time, ESplineCoordinateSpa
 	}
 	else
 	{
-		const int32 NumPoints = GetNumberOfSplinePoints();
-		const int32 NumSegments = bClosedLoop ? NumPoints : NumPoints - 1;
-		const float TimeMultiplier = NumSegments / Duration;
-		return GetTransformAtSplineInputKey(Time * TimeMultiplier, CoordinateSpace, bUseScale);
+		return GetTransformAtSplineInputKey(GetInputKeyValueAtTime(Time), CoordinateSpace, bUseScale);
 	}
 }
 
@@ -1949,10 +1930,7 @@ FVector USplineComponent::GetScaleAtTime(float Time, bool bUseConstantVelocity)
 	}
 	else
 	{
-		const int32 NumPoints = GetNumberOfSplinePoints();
-		const int32 NumSegments = bClosedLoop ? NumPoints : NumPoints - 1;
-		const float TimeMultiplier = NumSegments / Duration;
-		return GetScaleAtSplineInputKey(Time * TimeMultiplier);
+		return GetScaleAtSplineInputKey(GetInputKeyValueAtTime(Time));
 	}
 }
 
@@ -3305,6 +3283,25 @@ void USplineComponent::PostTransacted(const FTransactionObjectEvent& Transaction
 }
 #endif
 
+UE::Geometry::FInterval1f USplineComponent::GetInputKeyRange() const
+{
+	UE::Geometry::FInterval1f Range;
+	
+	if (WarninglessSplineCurves().Position.Points.Num() > 0)
+	{
+		Range.Min = WarninglessSplineCurves().Position.Points[0].InVal;
+		Range.Max = WarninglessSplineCurves().Position.Points.Last().InVal;
+		
+		// The valid range is extended by LoopKeyOffset for closed loops (this offset represents a virtual segment).
+		if (WarninglessSplineCurves().Position.bIsLooped)
+		{
+			Range.Max += WarninglessSplineCurves().Position.LoopKeyOffset;
+		}
+	}
+	
+	return Range;
+}
+
 void FSplinePositionLinearApproximation::Build(const FSplineCurves& InCurves, TArray<FSplinePositionLinearApproximation>& OutPoints, float InDensity)
 {
 	OutPoints.Reset();
