SplineComponent: Prevent AddPoint from generating an invalid knot vector

AddPoint did not previously handle colliding input key values explicitly, we just inserted as if the input were valid which would cause an ensure in UpdateSpline.

Now, if a colliding point is inserted, we increment the input keys of all points starting at the colliding point before we insert the point, resulting in a valid knot vector. The insertion index is decremented out of necessity in order to respect the user supplied input key while maintaining the strictly ascending input key rule. One other approach would be to leave the index as-is, but the consequence of this would be that the point being added would need its input key incremented as well. We err on the side of respecting the function parameters.

#FailedValidations Preflight

diff --git a/Engine/Source/Runtime/Engine/Private/Components/SplineComponent.cpp b/Engine/Source/Runtime/Engine/Private/Components/SplineComponent.cpp
--- a/Engine/Source/Runtime/Engine/Private/Components/SplineComponent.cpp
+++ b/Engine/Source/Runtime/Engine/Private/Components/SplineComponent.cpp
@@ -979,35 +979,47 @@ static int32 UpperBound(const TArray<FInterpCurvePoint<FVector>>& SplinePoints,
 
 void USplineComponent::AddPoint(const FSplinePoint& InSplinePoint, bool bUpdateSpline)
 {
-	Spline.AddPoint(InSplinePoint);
-
-	{
-		int32 Index = UpperBound(WarninglessSplineCurves().Position.Points, InSplinePoint.InputKey);
+	int32 Index = UpperBound(WarninglessSplineCurves().Position.Points, InSplinePoint.InputKey);
 		
-        WarninglessSplineCurves().Position.Points.Insert(FInterpCurvePoint<FVector>(
-        	InSplinePoint.InputKey,
-        	InSplinePoint.Position,
-        	InSplinePoint.ArriveTangent,
-        	InSplinePoint.LeaveTangent,
-        	ConvertSplinePointTypeToInterpCurveMode(InSplinePoint.Type)
-        ), Index);
-    
-        WarninglessSplineCurves().Rotation.Points.Insert(FInterpCurvePoint<FQuat>(
-        	InSplinePoint.InputKey,
-        	InSplinePoint.Rotation.Quaternion(),
-        	FQuat::Identity,
-        	FQuat::Identity,
-        	CIM_CurveAuto
-        ), Index);
-    
-        WarninglessSplineCurves().Scale.Points.Insert(FInterpCurvePoint<FVector>(
-        	InSplinePoint.InputKey,
-        	InSplinePoint.Scale,
-        	FVector::ZeroVector,
-        	FVector::ZeroVector,
-        	CIM_CurveAuto
-        ), Index);
+	if (Index > 0 && WarninglessSplineCurves().Position.Points[Index - 1].InVal == InSplinePoint.InputKey)
+	{
+		// Decrement in the case of a collision because we insert before the colliding element and increment that element's input key.
+		Index--;
+
+		// Increment keys only in the case of a collision to prevent invalid input keys.
+		for (int I = Index; I < WarninglessSplineCurves().Position.Points.Num(); ++I)
+		{
+			WarninglessSplineCurves().Position.Points[I].InVal += 1.0f;
+			WarninglessSplineCurves().Rotation.Points[I].InVal += 1.0f;
+			WarninglessSplineCurves().Scale.Points[I].InVal += 1.0f;
+		}
 	}
+		
+	WarninglessSplineCurves().Position.Points.Insert(FInterpCurvePoint<FVector>(
+		InSplinePoint.InputKey,
+		InSplinePoint.Position,
+		InSplinePoint.ArriveTangent,
+		InSplinePoint.LeaveTangent,
+		ConvertSplinePointTypeToInterpCurveMode(InSplinePoint.Type)
+	), Index);
+
+	WarninglessSplineCurves().Rotation.Points.Insert(FInterpCurvePoint<FQuat>(
+		InSplinePoint.InputKey,
+		InSplinePoint.Rotation.Quaternion(),
+		FQuat::Identity,
+		FQuat::Identity,
+		CIM_CurveAuto
+	), Index);
+
+	WarninglessSplineCurves().Scale.Points.Insert(FInterpCurvePoint<FVector>(
+		InSplinePoint.InputKey,
+		InSplinePoint.Scale,
+		FVector::ZeroVector,
+		FVector::ZeroVector,
+		CIM_CurveAuto
+	), Index);
+	
+	Spline.AddPoint(InSplinePoint);
 	
 	if (USplineMetadata* Metadata = GetSplinePointsMetadata())
 	{
