SplineComponent: Prevent knot vector corruption when knot value deltas are less than 1.

#FailedValidations Preflight

diff --git a/Engine/Source/Runtime/Engine/Private/Components/SplineComponent.cpp b/Engine/Source/Runtime/Engine/Private/Components/SplineComponent.cpp
--- a/Engine/Source/Runtime/Engine/Private/Components/SplineComponent.cpp
+++ b/Engine/Source/Runtime/Engine/Private/Components/SplineComponent.cpp
@@ -1168,6 +1168,10 @@ void USplineComponent::RemoveSplinePoint(int32 Index, bool bUpdateSpline)
 	{
 		Spline.RemovePoint(Index);
 
+		const float KnotDelta = Index > 0
+			? WarninglessSplineCurves().Position.Points[Index].InVal - WarninglessSplineCurves().Position.Points[Index - 1].InVal
+			: WarninglessSplineCurves().Position.Points[Index].InVal;
+		
 		WarninglessSplineCurves().Position.Points.RemoveAt(Index, EAllowShrinking::No);
 		WarninglessSplineCurves().Rotation.Points.RemoveAt(Index, EAllowShrinking::No);
 		WarninglessSplineCurves().Scale.Points.RemoveAt(Index, EAllowShrinking::No);
@@ -1182,9 +1186,9 @@ void USplineComponent::RemoveSplinePoint(int32 Index, bool bUpdateSpline)
 		// Adjust all following spline point input keys to close the gap left by the removed point
 		while (Index < NumPoints)
 		{
-			WarninglessSplineCurves().Position.Points[Index].InVal -= 1.0f;
-			WarninglessSplineCurves().Rotation.Points[Index].InVal -= 1.0f;
-			WarninglessSplineCurves().Scale.Points[Index].InVal -= 1.0f;
+			WarninglessSplineCurves().Position.Points[Index].InVal -= KnotDelta;
+			WarninglessSplineCurves().Rotation.Points[Index].InVal -= KnotDelta;
+			WarninglessSplineCurves().Scale.Points[Index].InVal -= KnotDelta;
 			Index++;
 		}
 
