Replaced check with ensure in FFolder::GetOptionalFolderRootObject to avoid crash for licensees.

diff --git a/Engine/Source/Runtime/Engine/Private/Folder.cpp b/Engine/Source/Runtime/Engine/Private/Folder.cpp
--- a/Engine/Source/Runtime/Engine/Private/Folder.cpp
+++ b/Engine/Source/Runtime/Engine/Private/Folder.cpp
@@ -35,8 +35,7 @@ bool FFolder::IsRootObjectPersistentLevel(const FRootObject& Key)
 TOptional<FFolder::FRootObject> FFolder::GetOptionalFolderRootObject(const ULevel* InLevel)
 {
 	TOptional<FFolder::FRootObject> Result;
-	check(InLevel);
-	if (InLevel)
+	if (ensure(InLevel))
 	{
 		if (InLevel->IsPersistentLevel())
 		{
