Fixed crash when starting PIE on a World Partition map containing dirty/modified actors that have bAsyncPhysicsTickEnabled = true.
In this case, the dirty actor will be duplicated and its outer will temporarily be a UActorContainer. When moving the actor back to the runtime level, AActor::Rename will call RegisterAllActorTickFunctions to unregister the actor/components tick functions. When bAsyncPhysicsTickEnabled = true, the code tries to access the PhysicsScene through the actor world (which is null in this particular case).

diff --git a/Engine/Source/Runtime/Engine/Private/Actor.cpp b/Engine/Source/Runtime/Engine/Private/Actor.cpp
--- a/Engine/Source/Runtime/Engine/Private/Actor.cpp
+++ b/Engine/Source/Runtime/Engine/Private/Actor.cpp
@@ -1697,15 +1697,18 @@ void AActor::RegisterAllActorTickFunctions(bool bRegister, bool bDoComponents)
 
 		if (bAsyncPhysicsTickEnabled)
 		{
-			if (FPhysScene_Chaos* Scene = static_cast<FPhysScene_Chaos*>(GetWorld()->GetPhysicsScene()))
+			if (UWorld* World = GetWorld())
 			{
-				if (bRegister)
+				if (FPhysScene_Chaos* Scene = static_cast<FPhysScene_Chaos*>(World->GetPhysicsScene()))
 				{
-					Scene->RegisterAsyncPhysicsTickActor(this);
-				}
-				else
-				{
-					Scene->UnregisterAsyncPhysicsTickActor(this);
+					if (bRegister)
+					{
+						Scene->RegisterAsyncPhysicsTickActor(this);
+					}
+					else
+					{
+						Scene->UnregisterAsyncPhysicsTickActor(this);
+					}
 				}
 			}
 		}
