Remove vestigial audio decompression code paths causing behavior to diverge on different platforms.

diff --git a/Engine/Source/Runtime/Apple/AudioMixerAudioUnit/Private/AudioMixerPlatformAudioUnit.h b/Engine/Source/Runtime/Apple/AudioMixerAudioUnit/Private/AudioMixerPlatformAudioUnit.h
--- a/Engine/Source/Runtime/Apple/AudioMixerAudioUnit/Private/AudioMixerPlatformAudioUnit.h
+++ b/Engine/Source/Runtime/Apple/AudioMixerAudioUnit/Private/AudioMixerPlatformAudioUnit.h
@@ -40,9 +40,6 @@ namespace Audio
 		virtual int32 GetNumFrames(const int32 InNumReqestedFrames) override;
 		virtual void ResumeContext() override;
 		virtual void SuspendContext() override;
-			
-		/** Whether or not the platform disables caching of decompressed PCM data (i.e. to save memory on fixed memory platforms) */
-		virtual bool DisablePCMAudioCaching() const override { return true; }
 		
 		//~ End IAudioMixerPlatformInterface
         
diff --git a/Engine/Source/Runtime/AudioMixer/Private/AudioMixerDevice.cpp b/Engine/Source/Runtime/AudioMixer/Private/AudioMixerDevice.cpp
--- a/Engine/Source/Runtime/AudioMixer/Private/AudioMixerDevice.cpp
+++ b/Engine/Source/Runtime/AudioMixer/Private/AudioMixerDevice.cpp
@@ -1094,17 +1094,6 @@ namespace Audio
 		// Every platform has compressed audio. 
 		return true;
 	}
-
-	bool FMixerDevice::SupportsRealtimeDecompression() const
-	{
-		// Every platform supports realtime decompression.
-		return true;
-	}
-
-	bool FMixerDevice::DisablePCMAudioCaching() const
-	{
-		return AudioMixerPlatform->DisablePCMAudioCaching();
-	}
 	
 	bool FMixerDevice::ValidateAPICall(const TCHAR* Function, uint32 ErrorCode)
 	{
diff --git a/Engine/Source/Runtime/AudioMixer/Public/AudioMixerDevice.h b/Engine/Source/Runtime/AudioMixer/Public/AudioMixerDevice.h
--- a/Engine/Source/Runtime/AudioMixer/Public/AudioMixerDevice.h
+++ b/Engine/Source/Runtime/AudioMixer/Public/AudioMixerDevice.h
@@ -134,8 +134,6 @@ namespace Audio
 		AUDIOMIXER_API virtual FAudioEffectsManager* CreateEffectsManager() override;
 		AUDIOMIXER_API virtual FSoundSource* CreateSoundSource() override;
 		AUDIOMIXER_API virtual bool HasCompressedAudioInfoClass(USoundWave* SoundWave) override;
-		AUDIOMIXER_API virtual bool SupportsRealtimeDecompression() const override;
-		AUDIOMIXER_API virtual bool DisablePCMAudioCaching() const override;
 		AUDIOMIXER_API virtual bool ValidateAPICall(const TCHAR* Function, uint32 ErrorCode) override;
 #if UE_ALLOW_EXEC_COMMANDS
 		AUDIOMIXER_API virtual bool Exec(UWorld* InWorld, const TCHAR* Cmd, FOutputDevice& Ar) override;
diff --git a/Engine/Source/Runtime/AudioMixerCore/Public/AudioMixer.h b/Engine/Source/Runtime/AudioMixerCore/Public/AudioMixer.h
--- a/Engine/Source/Runtime/AudioMixerCore/Public/AudioMixer.h
+++ b/Engine/Source/Runtime/AudioMixerCore/Public/AudioMixer.h
@@ -568,7 +568,8 @@ namespace Audio
 		virtual int32 GetNumFrames(const int32 InNumReqestedFrames) { return InNumReqestedFrames; }
 
 		/** Whether or not the platform disables caching of decompressed PCM data (i.e. to save memory on fixed memory platforms) */
-		virtual bool DisablePCMAudioCaching() const { return false; }
+		UE_DEPRECATED(5.7, "PCM Audio Caching is a leagacy feature that is now disabled for all platforms")
+		virtual bool DisablePCMAudioCaching() const { return true; }
 
 		/** Whether or not this platform has hardware decompression. */
 		virtual bool SupportsHardwareDecompression() const { return false; }
diff --git a/Engine/Source/Runtime/Engine/Private/AudioDevice.cpp b/Engine/Source/Runtime/Engine/Private/AudioDevice.cpp
--- a/Engine/Source/Runtime/Engine/Private/AudioDevice.cpp
+++ b/Engine/Source/Runtime/Engine/Private/AudioDevice.cpp
@@ -62,14 +62,6 @@ FAutoConsoleVariableRef CVarDisableStoppingVoices(
 	TEXT("0: Not Disabled, 1: Disabled"),
 	ECVF_Default);
 
-static int32 ForceRealtimeDecompressionCvar = 0;
-FAutoConsoleVariableRef CVarForceRealtimeDecompression(
-	TEXT("au.ForceRealtimeDecompression"),
-	ForceRealtimeDecompressionCvar,
-	TEXT("When set to 1, this deliberately ensures that all audio assets are decompressed as they play, rather than fully on load.\n")
-	TEXT("0: Allow full decompression on load, 1: force realtime decompression."),
-	ECVF_Default);
-
 static int32 DisableAppVolumeCvar = 0;
 FAutoConsoleVariableRef CVarDisableAppVolume(
 	TEXT("au.DisableAppVolume"),
@@ -86,22 +78,6 @@ FAutoConsoleVariableRef CVarDisableAutomaticPrecache(
 	TEXT("0: Use normal precaching logic, 1: disables all precaching except for synchronous calls."),
 	ECVF_Default);
 
-static float DecompressionThresholdCvar = 0.0f;
-FAutoConsoleVariableRef CVarDecompressionThreshold(
-	TEXT("au.DecompressionThreshold"),
-	DecompressionThresholdCvar,
-	TEXT("If non-zero, overrides the decompression threshold set in either the sound group or the platform's runtime settings.\n")
-	TEXT("Value: Maximum duration we should fully decompress, in seconds."),
-	ECVF_Default);
-
-static int32 RealtimeDecompressZeroDurationSoundsCvar = 0;
-FAutoConsoleVariableRef CVarForceRealtimeDecompressOnZeroDuration(
-	TEXT("au.RealtimeDecompressZeroDurationSounds"),
-	RealtimeDecompressZeroDurationSoundsCvar,
-	TEXT("When set to 1, we will fallback to realtime decoding any sound waves with an invalid duration..\n")
-	TEXT("0: Fully decompress sounds with a duration of 0, 1: realtime decompress sounds with a duration of 0."),
-	ECVF_Default);
-
 static int32 WaitForSoundWaveToLoadCvar = 1;
 FAutoConsoleVariableRef CVarWaitForSoundWaveToLoad(
 	TEXT("au.WaitForSoundWaveToLoad"),
@@ -374,7 +350,6 @@ FAudioDevice::FAudioDevice()
 	, bHRTFEnabledForAll_OnGameThread(false)
 	, bHRTFDisabled_OnGameThread(false)
 	, bGameWasTicking(true)
-	, bDisableAudioCaching(false)
 	, bIsAudioDeviceHardwareInitialized(false)
 	, bIsStoppingVoicesEnabled(false)
 	, bIsBakedAnalysisEnabled(false)
@@ -1904,20 +1879,9 @@ bool FAudioDevice::HandleAudioMemoryInfo(const TCHAR* Cmd, FOutputDevice& Ar)
 
 			const FSoundGroup& SoundGroup = GetDefault<USoundGroups>()->GetSoundGroup(SoundWave->SoundGroup);
 
-			float CompressionDurationThreshold = GetCompressionDurationThreshold(SoundGroup);
-
 			// Determine whether this asset is streaming compressed data from disk, decompressed in realtime, or fully decompressed on load.
-			ELoadingType LoadType;
-
-			if (SoundWave->IsStreaming())
-			{
-				LoadType = ELoadingType::Streaming;
-			}
-			else
-			{
-				LoadType = ShouldUseRealtimeDecompression(false, SoundGroup, SoundWave, CompressionDurationThreshold) ? ELoadingType::CompressedInMemory : ELoadingType::DecompressedInMemory;
-			}
-
+			ELoadingType LoadType = SoundWave->IsStreaming()? ELoadingType::Streaming : ELoadingType::CompressedInMemory;
+			
 			FString SoundGroupName;
 			switch (SoundWave->SoundGroup)
 			{
@@ -6981,9 +6945,6 @@ void FAudioDevice::Precache(USoundWave* SoundWave, bool bSynchronous, bool bTrac
 
 			return;
 		}
-
-
-		
 	}
 
 	// calculate the decompression type
@@ -7012,10 +6973,7 @@ void FAudioDevice::Precache(USoundWave* SoundWave, bool bSynchronous, bool bTrac
 		{
 			UE_LOG(LogAudio, Warning, TEXT("Sound Wave reported a duration of zero. This will likely result in incorrect decoding."));
 		}
-
-
-		float CompressedDurationThreshold = GetCompressionDurationThreshold(SoundGroup);
-
+		
 		static FName NAME_OGG(TEXT("OGG"));
 		SoundWave->bDecompressedFromOgg = SoundWave->GetRuntimeFormat() == NAME_OGG;
 		bool bDecompressToPCM = false;
@@ -7026,7 +6984,7 @@ void FAudioDevice::Precache(USoundWave* SoundWave, bool bSynchronous, bool bTrac
 			SoundWave->DecompressionType = DTYPE_Streaming;
 			SoundWave->bCanProcessAsync = false;
 		}
-		else if (ShouldUseRealtimeDecompression(bForceFullDecompression, SoundGroup, SoundWave, CompressedDurationThreshold)) 
+		else if (!bForceFullDecompression) 
 		{
 			// Store as compressed data and decompress in realtime
 			SoundWave->DecompressionType = DTYPE_RealTime;
@@ -7087,7 +7045,8 @@ void FAudioDevice::Precache(USoundWave* SoundWave, bool bSynchronous, bool bTrac
 float FAudioDevice::GetCompressionDurationThreshold(const FSoundGroup &SoundGroup)
 {
 	// Check to see if the compression duration threshold is overridden via CVar:
-	float CompressedDurationThreshold = DecompressionThresholdCvar;
+	float CompressedDurationThreshold = 0.f; // related cvar has been removed
+	
 	// If not, check to see if there is an override for the compression duration on this platform in the project settings:
 	if (CompressedDurationThreshold <= 0.0f)
 	{
@@ -7103,13 +7062,10 @@ float FAudioDevice::GetCompressionDurationThreshold(const FSoundGroup &SoundGrou
 	return CompressedDurationThreshold;
 }
 
+// (deprecated)
 bool FAudioDevice::ShouldUseRealtimeDecompression(bool bForceFullDecompression, const FSoundGroup &SoundGroup, USoundWave* SoundWave, float CompressedDurationThreshold) const
 {
-	return !bForceFullDecompression &&
-		SupportsRealtimeDecompression() &&
-		((bDisableAudioCaching || DisablePCMAudioCaching()) ||
-		(!SoundGroup.bAlwaysDecompressOnLoad &&
-			(ForceRealtimeDecompressionCvar || SoundWave->Duration > CompressedDurationThreshold || (RealtimeDecompressZeroDurationSoundsCvar && SoundWave->Duration <= 0.0f))));
+	return !bForceFullDecompression;
 }
 
 // deprecated
diff --git a/Engine/Source/Runtime/Engine/Private/AudioDeviceManager.cpp b/Engine/Source/Runtime/Engine/Private/AudioDeviceManager.cpp
--- a/Engine/Source/Runtime/Engine/Private/AudioDeviceManager.cpp
+++ b/Engine/Source/Runtime/Engine/Private/AudioDeviceManager.cpp
@@ -1165,15 +1165,12 @@ void FAudioDeviceManager::FreeBufferResource(FSoundBuffer* SoundBuffer)
 // (deprecated)
 FSoundBuffer* FAudioDeviceManager::GetSoundBufferForResourceID(uint32 ResourceID)
 {
-	// maxtodo: warn
 	return {};
 }
 
 // deprecated
 void FAudioDeviceManager::RemoveSoundBufferForResourceID(uint32)
 {
-	// maxtodo: warn
-	// WaveBufferMap.Remove(ResourceID);
 }
 
 void FAudioDeviceManager::RemoveSoundMix(USoundMix* SoundMix)
diff --git a/Engine/Source/Runtime/Engine/Public/AudioDevice.h b/Engine/Source/Runtime/Engine/Public/AudioDevice.h
--- a/Engine/Source/Runtime/Engine/Public/AudioDevice.h
+++ b/Engine/Source/Runtime/Engine/Public/AudioDevice.h
@@ -623,11 +623,13 @@ public:
 	 */
 	ENGINE_API virtual void Precache(USoundWave* SoundWave, bool bSynchronous = false, bool bTrackMemory = true, bool bForceFullDecompression = false);
 
+	UE_DEPRECATED(5.7, "With PCM audio caching removed, this hueristic is no longer used.")
 	ENGINE_API float GetCompressionDurationThreshold(const FSoundGroup &SoundGroup);
 
 	/**
 	 * Returns true if a sound wave should be decompressed.
 	 */
+	UE_DEPRECATED(5.7, "Most of these arguments are no longer used, replace calls with the caller's value for bForceFullDecompression")
 	ENGINE_API bool ShouldUseRealtimeDecompression(bool bForceFullDecompression, const FSoundGroup &SoundGroup, USoundWave* SoundWave, float CompressedDurationThreshold) const;
 
 	/**
@@ -1138,15 +1140,17 @@ public:
 	virtual bool HasCompressedAudioInfoClass(USoundWave* SoundWave) { return false; }
 
 	/** Whether this device supports realtime decompression of sound waves (i.e. DTYPE_RealTime) */
+	UE_DEPRECATED(5.7, "Realtime decpmpression is available on all platforms")
 	virtual bool SupportsRealtimeDecompression() const
 	{
-		return false;
+		return true;
 	}
 
 	/** Whether or not the platform disables caching of decompressed PCM data (i.e. to save memory on fixed memory platforms */
+	UE_DEPRECATED(5.7, "PCM Audio Caching is a leagacy feature that is now disabled for all platforms")
 	virtual bool DisablePCMAudioCaching() const
 	{
-		return false;
+		return true;
 	}
 
 	/**
@@ -2149,7 +2153,8 @@ private:
 
 public:
 	/** HACK: Temporarily disable audio caching.  This will be done better by changing the decompression pool size in the future */
-	uint8 bDisableAudioCaching:1;
+	UE_DEPRECATED(5.7, "bDisableAudioCaching has no effect and will be removed in 5.8")
+	uint8 bDisableAudioCaching:1 = false;
 
 	/** Whether or not the lower-level audio device hardware initialized. */
 	uint8 bIsAudioDeviceHardwareInitialized : 1;
diff --git a/Engine/Source/Runtime/IOS/IOSAudio/Private/IOSAudioDevice.cpp b/Engine/Source/Runtime/IOS/IOSAudio/Private/IOSAudioDevice.cpp
--- a/Engine/Source/Runtime/IOS/IOSAudio/Private/IOSAudioDevice.cpp
+++ b/Engine/Source/Runtime/IOS/IOSAudio/Private/IOSAudioDevice.cpp
@@ -104,7 +104,6 @@ FIOSAudioDevice::FIOSAudioDevice() :
 	MixerUnit(NULL),
 	NextBusNumber(0)
 {
-	bDisableAudioCaching = true;	// Only use DTYPE_RealTime or DTYPE_Streaming since on the fly decompression is so cheap, it saves memory, and requires fewer code paths
 }
 
 bool FIOSAudioDevice::Exec(UWorld* InWorld, const TCHAR* Cmd, FOutputDevice& Ar)
diff --git a/Engine/Source/Runtime/IOS/IOSAudio/Public/IOSAudioDevice.h b/Engine/Source/Runtime/IOS/IOSAudio/Public/IOSAudioDevice.h
--- a/Engine/Source/Runtime/IOS/IOSAudio/Public/IOSAudioDevice.h
+++ b/Engine/Source/Runtime/IOS/IOSAudio/Public/IOSAudioDevice.h
@@ -242,9 +242,7 @@ protected:
 private:
 
 	// Used to support adpcm streaming
-	virtual bool HasCompressedAudioInfoClass(USoundWave* SoundWave) override { return true; }
-
-	virtual bool SupportsRealtimeDecompression() const override { return true; }	
+	virtual bool HasCompressedAudioInfoClass(USoundWave* SoundWave) override { return true; }	
 
 	void HandleError(const TCHAR* InLogOutput, bool bTeardown = false);
 
diff --git a/Engine/Source/Runtime/NonRealtimeAudioRenderer/Private/AudioMixerPlatformNonRealtime.cpp b/Engine/Source/Runtime/NonRealtimeAudioRenderer/Private/AudioMixerPlatformNonRealtime.cpp
--- a/Engine/Source/Runtime/NonRealtimeAudioRenderer/Private/AudioMixerPlatformNonRealtime.cpp
+++ b/Engine/Source/Runtime/NonRealtimeAudioRenderer/Private/AudioMixerPlatformNonRealtime.cpp
@@ -319,11 +319,6 @@ namespace Audio
 		return 0;
 	}
 
-	bool FMixerPlatformNonRealtime::DisablePCMAudioCaching() const
-	{
-		return true;
-	}
-
 	void FMixerPlatformNonRealtime::FadeIn()
 	{
 		bFadedOut = false;
diff --git a/Engine/Source/Runtime/NonRealtimeAudioRenderer/Public/AudioMixerPlatformNonRealtime.h b/Engine/Source/Runtime/NonRealtimeAudioRenderer/Public/AudioMixerPlatformNonRealtime.h
--- a/Engine/Source/Runtime/NonRealtimeAudioRenderer/Public/AudioMixerPlatformNonRealtime.h
+++ b/Engine/Source/Runtime/NonRealtimeAudioRenderer/Public/AudioMixerPlatformNonRealtime.h
@@ -37,7 +37,6 @@ namespace Audio
 		virtual void ResumePlaybackOnNewDevice() override;
 		virtual FAudioPlatformDeviceInfo GetPlatformDeviceInfo() const override;
 		virtual void SubmitBuffer(const uint8* Buffer) override;
-		virtual bool DisablePCMAudioCaching() const override;		
 		virtual FString GetDefaultDeviceName() override;
 		virtual FAudioPlatformSettings GetPlatformSettings() const override;
 		virtual void OnHardwareUpdate() override;
diff --git a/Engine/Source/Runtime/Windows/AudioMixerWasapi/Private/AudioMixerWasapi.h b/Engine/Source/Runtime/Windows/AudioMixerWasapi/Private/AudioMixerWasapi.h
--- a/Engine/Source/Runtime/Windows/AudioMixerWasapi/Private/AudioMixerWasapi.h
+++ b/Engine/Source/Runtime/Windows/AudioMixerWasapi/Private/AudioMixerWasapi.h
@@ -76,7 +76,6 @@ namespace Audio
 		virtual FAudioPlatformDeviceInfo GetPlatformDeviceInfo() const override;
 		virtual void SubmitBuffer(const uint8* Buffer) override;
 		virtual void SubmitDirectOutBuffer(const int32 InDirectOutIndex, const FAlignedFloatBuffer& InBuffer) override;
-		virtual bool DisablePCMAudioCaching() const override { return true; }
 		virtual FString GetDefaultDeviceName() override;
 		virtual FAudioPlatformSettings GetPlatformSettings() const override;
 		virtual IAudioPlatformDeviceInfoCache* GetDeviceInfoCache() const override;
diff --git a/Engine/Source/Runtime/Windows/AudioMixerXAudio2/Private/AudioMixerPlatformXAudio2.cpp b/Engine/Source/Runtime/Windows/AudioMixerXAudio2/Private/AudioMixerPlatformXAudio2.cpp
--- a/Engine/Source/Runtime/Windows/AudioMixerXAudio2/Private/AudioMixerPlatformXAudio2.cpp
+++ b/Engine/Source/Runtime/Windows/AudioMixerXAudio2/Private/AudioMixerPlatformXAudio2.cpp
@@ -1256,11 +1256,6 @@ namespace Audio
 		}
 	}
 
-	bool FMixerPlatformXAudio2::DisablePCMAudioCaching() const
-	{
-		return true;
-	}
-
 	IXAudio2MasteringVoice* FMixerPlatformXAudio2::CreateMasteringVoice(IXAudio2& InXAudio2System, const FAudioPlatformDeviceInfo& NewDevice, bool bUseDefaultDevice)
 	{
 		IXAudio2MasteringVoice* MasteringVoice = nullptr;
diff --git a/Engine/Source/Runtime/Windows/AudioMixerXAudio2/Private/AudioMixerPlatformXAudio2.h b/Engine/Source/Runtime/Windows/AudioMixerXAudio2/Private/AudioMixerPlatformXAudio2.h
--- a/Engine/Source/Runtime/Windows/AudioMixerXAudio2/Private/AudioMixerPlatformXAudio2.h
+++ b/Engine/Source/Runtime/Windows/AudioMixerXAudio2/Private/AudioMixerPlatformXAudio2.h
@@ -116,7 +116,6 @@ namespace Audio
 		
 		virtual FAudioPlatformDeviceInfo GetPlatformDeviceInfo() const override;
 		virtual void SubmitBuffer(const uint8* Buffer) override;
-		virtual bool DisablePCMAudioCaching() const override;		
 		virtual FString GetDefaultDeviceName() override;
 		virtual FAudioPlatformSettings GetPlatformSettings() const override;
 		virtual void OnHardwareUpdate() override;
