Workaround for crash when regenerating skeletal mesh. There are two problems.
1) The skeletal mesh is being destroyed while the component is active. ReleaseResources therefore cannot reference the skeletal mesh UAsset data.
2) The weight buffer SRV is null. This is causing the shader to crash. Added a workaround to use a null buffer instead. This needs further investigation.

diff --git a/Engine/Source/Runtime/Engine/Private/GPUSkinCache.cpp b/Engine/Source/Runtime/Engine/Private/GPUSkinCache.cpp
--- a/Engine/Source/Runtime/Engine/Private/GPUSkinCache.cpp
+++ b/Engine/Source/Runtime/Engine/Private/GPUSkinCache.cpp
@@ -1932,7 +1932,7 @@ void FGPUSkinCache::DispatchPassExecute(FTaskData& TaskData)
 			Parameters.InputWeightIndexSize = Entry->InputWeightIndexSize;
 			Parameters.InputWeightStart = DispatchData.InputWeightStart;
 			Parameters.InputWeightStride = Entry->InputWeightStride;
-			Parameters.InputWeightStream = Entry->InputWeightStreamSRV;
+			Parameters.InputWeightStream = Entry->InputWeightStreamSRV ? Entry->InputWeightStreamSRV : GNullVertexBuffer.VertexBufferSRV;
 			Parameters.InputWeightLookupStream = Entry->InputWeightLookupStreamSRV;
 			Parameters.PositionInputBuffer = DispatchData.PositionBufferSRV;
 			Parameters.TangentInputBuffer = DispatchData.TangentBufferSRV;
diff --git a/Engine/Source/Runtime/Engine/Private/SkeletalRenderGPUSkin.cpp b/Engine/Source/Runtime/Engine/Private/SkeletalRenderGPUSkin.cpp
--- a/Engine/Source/Runtime/Engine/Private/SkeletalRenderGPUSkin.cpp
+++ b/Engine/Source/Runtime/Engine/Private/SkeletalRenderGPUSkin.cpp
@@ -2480,7 +2480,8 @@ void FSkeletalMeshObjectGPUSkin::FSkeletalMeshObjectLOD::InitResources(
 
 		if (DynamicBoundsStartOffset == INDEX_NONE && GPUSkinTechnique == ESkeletalMeshGPUSkinTechnique::GPUSkinCache)
 		{
-			DynamicBoundsStartOffset = InGPUSkinCache->AllocateDynamicMeshBoundsSlot(LODData.RenderSections.Num());
+			DynamicBoundsNumSections = LODData.RenderSections.Num();
+			DynamicBoundsStartOffset = InGPUSkinCache->AllocateDynamicMeshBoundsSlot(DynamicBoundsNumSections);
 		}
 	});
 }
@@ -2494,7 +2495,7 @@ void FSkeletalMeshObjectGPUSkin::FSkeletalMeshObjectLOD::ReleaseResources(FGPUSk
 
 		if (DynamicBoundsStartOffset != INDEX_NONE)
 		{
-			InGPUSkinCache->ReleaseDynamicMeshBoundsSlot(DynamicBoundsStartOffset, SkelMeshRenderData->LODRenderData[LODIndex].RenderSections.Num());
+			InGPUSkinCache->ReleaseDynamicMeshBoundsSlot(DynamicBoundsStartOffset, DynamicBoundsNumSections);
 			DynamicBoundsStartOffset = INDEX_NONE;
 		}
 	});
diff --git a/Engine/Source/Runtime/Engine/Private/SkeletalRenderGPUSkin.h b/Engine/Source/Runtime/Engine/Private/SkeletalRenderGPUSkin.h
--- a/Engine/Source/Runtime/Engine/Private/SkeletalRenderGPUSkin.h
+++ b/Engine/Source/Runtime/Engine/Private/SkeletalRenderGPUSkin.h
@@ -651,6 +651,7 @@ private:
 
 		/** The handle to a dynamic mesh bounds range for each section in the LOD when using skin cache. */
 		int32 DynamicBoundsStartOffset = INDEX_NONE;
+		int32 DynamicBoundsNumSections = 0;
 
 		/**
 		 * Update the contents of the morphtarget vertex buffer by accumulating all 
