[AutoViz] Adressing 2 bug hawk issues
- Missleadig comment
- Max value for Zebra stripe count now in sync with API

diff --git a/Engine/Source/Editor/UnrealEd/Private/ViewportToolbar/UnrealEdViewportToolbar.cpp b/Engine/Source/Editor/UnrealEd/Private/ViewportToolbar/UnrealEdViewportToolbar.cpp
--- a/Engine/Source/Editor/UnrealEd/Private/ViewportToolbar/UnrealEdViewportToolbar.cpp
+++ b/Engine/Source/Editor/UnrealEd/Private/ViewportToolbar/UnrealEdViewportToolbar.cpp
@@ -2231,7 +2231,7 @@ void PopulateViewModesMenu(UToolMenu* InMenu)
 								.Style(&FAppStyle::Get(), "Menu.SpinBox")
 								.Font( FAppStyle::GetFontStyle( TEXT( "MenuItem.Font" ) ) )
 								.MinValue(1.f)
-								.MaxValue(16.f)
+								.MaxValue(32.f)
 								.SupportDynamicSliderMaxValue(false)
 								.Value_Static ( &FEditorViewportClient::GetZebraViewModeStripeCount )
 								.OnValueChanged_Lambda( [EditorViewportWeakPtr](uint16 NewValue)
diff --git a/Engine/Source/Runtime/Engine/Private/UnrealEngine.cpp b/Engine/Source/Runtime/Engine/Private/UnrealEngine.cpp
--- a/Engine/Source/Runtime/Engine/Private/UnrealEngine.cpp
+++ b/Engine/Source/Runtime/Engine/Private/UnrealEngine.cpp
@@ -3062,7 +3062,8 @@ void LoadSpecialMaterial(const FString& MaterialName, const FString& MaterialNam
 }
 
 /**
-* Loads a special material interface and verifies that it can be loaded
+* Loads a special material interface and, in editor builds - verifies that it can be loaded,
+* while reporting only a log message only in non-editor builds.
 *
 * @param MaterialName Variable name of the material in the engine to identify which material can't be loaded
 * @param MaterialNamePath Fully qualified name of a material to load/find
@@ -16987,8 +16988,8 @@ void UEngine::CheckAndHandleStaleWorldObjectReferences(FWorldContext* WorldConte
 	for( TObjectIterator<UWorld> It; It; ++It )
 	{
 		UWorld* World = *It;
-		const bool bPersistentWorldType = (World->WorldType == EWorldType::Inactive) || (World->WorldType == EWorldType::EditorPreview);
-		if (!bPersistentWorldType && !WorldHasValidContext(World))
+		const bool bIsPersistentWorldType = (World->WorldType == EWorldType::Inactive) || (World->WorldType == EWorldType::EditorPreview);
+		if (!bIsPersistentWorldType && !WorldHasValidContext(World))
 		{
 			if ((World->PersistentLevel == nullptr || !WorldHasValidContext(World->PersistentLevel->OwningWorld)) && !IsWorldDuplicate(World))
 			{
