Fix for inertial request on re-entry even when target state doesn't reset.
#rnx

diff --git a/Engine/Source/Editor/AnimGraph/Public/AnimStateTransitionNode.h b/Engine/Source/Editor/AnimGraph/Public/AnimStateTransitionNode.h
--- a/Engine/Source/Editor/AnimGraph/Public/AnimStateTransitionNode.h
+++ b/Engine/Source/Editor/AnimGraph/Public/AnimStateTransitionNode.h
@@ -90,8 +90,9 @@ class UAnimStateTransitionNode : public UAnimStateNodeBase
 	UPROPERTY(EditAnywhere, Category=Events)
 	FAnimNotifyEvent TransitionInterrupt;
 
-	/** Whether to fall back to inertialization/dead blending when reentering an already-active state. This can avoid pops. */
-	UPROPERTY(EditAnywhere, Category=Transition)
+	/** Whether to fall back to inertialization/dead blending when reentering an already-active state. This can avoid pops. 
+		The target state must enable bAlwaysResetOnEntry for the inertial blend to trigger. */
+	UPROPERTY(EditAnywhere, Category="Transition|Experimental")
 	bool bAllowInertializationForSelfTransitions;
 
 	/** This transition can go both directions */
diff --git a/Engine/Source/Runtime/Engine/Private/Animation/AnimNode_StateMachine.cpp b/Engine/Source/Runtime/Engine/Private/Animation/AnimNode_StateMachine.cpp
--- a/Engine/Source/Runtime/Engine/Private/Animation/AnimNode_StateMachine.cpp
+++ b/Engine/Source/Runtime/Engine/Private/Animation/AnimNode_StateMachine.cpp
@@ -1348,7 +1348,7 @@ void FAnimNode_StateMachine::TransitionToState(const FAnimationUpdateContext& Co
 			NewTransition.StateBlendData[1].PerBoneBlendData.AddZeroed(NewTransition.BlendProfile->GetNumBlendEntries());
 		}
 
-		auto ShouldInertialize = [&TransitionInfo, &ExistingWeightOfNextState]()
+		auto ShouldInertialize = [MachineDescription = PRIVATE_MachineDescription, NextState, &TransitionInfo, &ExistingWeightOfNextState]()
 		{
 			// We are set to inertialize
 			if (TransitionInfo.LogicType == ETransitionLogicType::TLT_Inertialization)
@@ -1358,7 +1358,12 @@ void FAnimNode_StateMachine::TransitionToState(const FAnimationUpdateContext& Co
 			// Fall back if we are allowed to and the next state has weight, to hide a pop if the state resets
 			else if (TransitionInfo.bAllowInertializationForSelfTransitions && ExistingWeightOfNextState > 0.0f)
 			{
-				return true;
+				const FBakedAnimationState& BakedCurrentState = MachineDescription->States[NextState];
+				// If the target state will not reset, then there's no need to inertialize.
+				if (BakedCurrentState.bAlwaysResetOnEntry)
+				{
+					return true;
+				}
 			}
 			return false;
 		};
