Fix Substrate Thickness input value not taken into account if material use RootNode's constant instead of an input expression.

diff --git a/Engine/Source/Runtime/Engine/Private/Materials/HLSLMaterialTranslator.cpp b/Engine/Source/Runtime/Engine/Private/Materials/HLSLMaterialTranslator.cpp
--- a/Engine/Source/Runtime/Engine/Private/Materials/HLSLMaterialTranslator.cpp
+++ b/Engine/Source/Runtime/Engine/Private/Materials/HLSLMaterialTranslator.cpp
@@ -1683,7 +1683,7 @@ void FHLSLMaterialTranslator::TranslateMaterial()
 			FSubstrateCompilationContext& SubstrateCtx = SubstrateCompilationContext[SubstrateCompilationContextIndex];
 			SubstrateCtx = FSubstrateCompilationContext(ESubstrateCompilationContext(SubstrateCompilationContextIndex));
 
-			FExpressionInput* SurfaceThickness = Material->IsThinSurface() && Material->GetMaterialInterface() ? &Material->GetMaterialInterface()->GetMaterial()->GetEditorOnlyData()->SurfaceThickness : nullptr;
+			FScalarMaterialInput* SurfaceThickness = Material->IsThinSurface() && Material->GetMaterialInterface() ? &Material->GetMaterialInterface()->GetMaterial()->GetEditorOnlyData()->SurfaceThickness : nullptr;
 			SubstrateThicknessStackPush(nullptr, SurfaceThickness);
 			FrontMaterialExpr->SubstrateGenerateMaterialTopologyTree(this, nullptr, FrontMaterialOutputIndex);
 			SubstrateThicknessStackPop();
@@ -15360,11 +15360,21 @@ int32 FHLSLMaterialTranslator::SubstrateThicknessStackGetThicknessCode(int32 Ind
 	{
 		UE_LOG(LogMaterial, Error, TEXT(" SubstrateThichkness: %i could not be found)"), Index);
 	}
-	else if (FExpressionInput* Input = SubstrateCtx.SubstrateThicknessIndexToExpressionInput[Index])
+	else if (FExpressionInput* Input = SubstrateCtx.SubstrateThicknessIndexToExpressionInput[Index].ExpressionInput)
 	{
 		OutCode = Input->GetTracedInput().Expression ? Input->Compile(this) : Constant(SUBSTRATE_LAYER_DEFAULT_THICKNESS_CM);
 		EMaterialValueType Type = GetType(OutCode);
 		if (IsLWCType(Type))
+		{
+			Type = MakeNonLWCType(Type);
+			OutCode = ValidCast(OutCode, Type);
+		}
+	}
+	else if (FScalarMaterialInput* MaterialInput = SubstrateCtx.SubstrateThicknessIndexToExpressionInput[Index].MaterialInput)
+	{
+		OutCode = MaterialInput->UseConstant ? Constant(MaterialInput->Constant) : Constant(SUBSTRATE_LAYER_DEFAULT_THICKNESS_CM);
+		EMaterialValueType Type = GetType(OutCode);
+		if (IsLWCType(Type))
 		{
 			Type = MakeNonLWCType(Type);
 			OutCode = ValidCast(OutCode, Type);
@@ -15377,11 +15387,28 @@ int32 FHLSLMaterialTranslator::SubstrateThicknessStackGetThicknessCode(int32 Ind
 	return OutCode;
 }
 
+int32 FHLSLMaterialTranslator::SubstrateThicknessStackPush(UMaterialExpression* Expression, FScalarMaterialInput* Input)
+{
+	// If the input is actually not a constant, reroute it to the Expression overload to that the expression can be traced.
+	if (Input && !Input->UseConstant)
+	{
+		return SubstrateThicknessStackPush(Expression, (FExpressionInput*)Input);
+	}
+
+	FSubstrateCompilationContext& SubstrateCtx = SubstrateCompilationContext[CurrentSubstrateCompilationContext];
+	int32 Index = SubstrateCtx.SubstrateThicknessIndexToExpressionInput.Num();
+	FSubstrateCompilationContext::FSubstrateThicknessExpression& Entry = SubstrateCtx.SubstrateThicknessIndexToExpressionInput.AddDefaulted_GetRef();
+	Entry.MaterialInput = Input;
+	SubstrateCtx.SubstrateThicknessStack.Push(Index);
+	return Index;
+}
+
 int32 FHLSLMaterialTranslator::SubstrateThicknessStackPush(UMaterialExpression* Expression, FExpressionInput* Input)
 {
 	FSubstrateCompilationContext& SubstrateCtx = SubstrateCompilationContext[CurrentSubstrateCompilationContext];
 	int32 Index = SubstrateCtx.SubstrateThicknessIndexToExpressionInput.Num();
-	SubstrateCtx.SubstrateThicknessIndexToExpressionInput.Add(Input);
+	FSubstrateCompilationContext::FSubstrateThicknessExpression& Entry = SubstrateCtx.SubstrateThicknessIndexToExpressionInput.AddDefaulted_GetRef();
+	Entry.ExpressionInput = Input;
 	SubstrateCtx.SubstrateThicknessStack.Push(Index);
 	return Index;	
 }
diff --git a/Engine/Source/Runtime/Engine/Private/Materials/HLSLMaterialTranslator.h b/Engine/Source/Runtime/Engine/Private/Materials/HLSLMaterialTranslator.h
--- a/Engine/Source/Runtime/Engine/Private/Materials/HLSLMaterialTranslator.h
+++ b/Engine/Source/Runtime/Engine/Private/Materials/HLSLMaterialTranslator.h
@@ -598,8 +598,13 @@ protected:
 		/** Stack of thickness input used for propagating thickness information from root node and vertical operation
 		* This is transient and updated when calling SubstrateGenerateMaterialTopologyTree. The information is then stored into the FStratOperator
 		*/
+		struct FSubstrateThicknessExpression
+		{
+			FExpressionInput* ExpressionInput = nullptr;
+			FScalarMaterialInput* MaterialInput = nullptr;
+		};
 		TArray<int32> SubstrateThicknessStack;
-		TArray<FExpressionInput*> SubstrateThicknessIndexToExpressionInput;
+		TArray<FSubstrateThicknessExpression> SubstrateThicknessIndexToExpressionInput;
 
 		FSubstrateSimplificationStatus SubstrateSimplificationStatus;
 
@@ -1429,6 +1434,7 @@ protected:
 	
 	virtual int32 SubstrateThicknessStackGetThicknessIndex() override;
 	virtual int32 SubstrateThicknessStackGetThicknessCode(int32 Index) override;
+	virtual int32 SubstrateThicknessStackPush(UMaterialExpression* Expression, FScalarMaterialInput* Input) override;
 	virtual int32 SubstrateThicknessStackPush(UMaterialExpression* Expression, FExpressionInput* Input) override;
 	virtual void SubstrateThicknessStackPop() override;
 
diff --git a/Engine/Source/Runtime/Engine/Public/MaterialCompiler.h b/Engine/Source/Runtime/Engine/Public/MaterialCompiler.h
--- a/Engine/Source/Runtime/Engine/Public/MaterialCompiler.h
+++ b/Engine/Source/Runtime/Engine/Public/MaterialCompiler.h
@@ -755,6 +755,7 @@ public:
 
 	virtual int32 SubstrateThicknessStackGetThicknessIndex() = 0;
 	virtual int32 SubstrateThicknessStackGetThicknessCode(int32 Index) = 0;
+	virtual int32 SubstrateThicknessStackPush(UMaterialExpression* Expression, FScalarMaterialInput* Input) = 0;
 	virtual int32 SubstrateThicknessStackPush(UMaterialExpression* Expression, FExpressionInput* Input) = 0;
 	virtual void SubstrateThicknessStackPop() = 0;
 
@@ -1606,6 +1607,10 @@ public:
 	{
 		return Compiler->SubstrateThicknessStackGetThicknessCode(Index);
 	}
+	virtual int32 SubstrateThicknessStackPush(UMaterialExpression* Expression, FScalarMaterialInput* Input) override
+	{
+		return Compiler->SubstrateThicknessStackPush(Expression, Input);
+	}
 	virtual int32 SubstrateThicknessStackPush(UMaterialExpression* Expression, FExpressionInput* Input) override
 	{
 		return Compiler->SubstrateThicknessStackPush(Expression, Input);
