Tentative fix for shader translation crash.

diff --git a/Engine/Source/Runtime/Engine/Private/Materials/MaterialExpressionSubstrate.cpp b/Engine/Source/Runtime/Engine/Private/Materials/MaterialExpressionSubstrate.cpp
--- a/Engine/Source/Runtime/Engine/Private/Materials/MaterialExpressionSubstrate.cpp
+++ b/Engine/Source/Runtime/Engine/Private/Materials/MaterialExpressionSubstrate.cpp
@@ -130,6 +130,8 @@ static int32 CompileWithDefaultNormalWS(class FMaterialCompiler* Compiler, FExpr
 
 		// Ensure the normal has always a valid float3 type
 		NormalCodeChunk = Compiler->ForceCast(NormalCodeChunk, MCT_Float3, MFCF_ExactMatch | MFCF_ReplicateValue);
+		NormalCodeChunk = CastToNonLWCType(Compiler, NormalCodeChunk);
+		NormalCodeChunk = NormalCodeChunk == INDEX_NONE ? Compiler->VertexNormal() : NormalCodeChunk;
 
 		// Transform into world space normal if needed. BSDF nodes always expects world space normal as input.
 		return bConvertToRequestedSpace ? Compiler->TransformNormalFromRequestedBasisToWorld(NormalCodeChunk) : NormalCodeChunk;
@@ -1619,7 +1621,7 @@ int32 UMaterialExpressionSubstrateLightFunction::CompileCommon(class FMaterialCo
 	FSubstrateOperator& SubstrateOperator = Compiler->SubstrateCompilationGetOperator(Compiler->SubstrateTreeStackGetPathUniqueId());
 	int32 OutputCodeChunk = Compiler->SubstrateUnlitBSDF(
 		CompileWithDefaultFloat3(Compiler, Color, 0.0f, 0.0f, 0.0f, EditorOnlyData ? &EditorOnlyData->EmissiveColor : nullptr),
-		Compiler->Constant(1.0f),				// Opacity / Transmittance is ignored by light functions.
+		Compiler->Constant3(1.0f, 1.0f, 1.0f),	// Opacity / Transmittance is ignored by light functions.
 		Compiler->Constant3(0.0f, 0.0f, 1.0f),	// place holder normal
 		&SubstrateOperator);
 	return OutputCodeChunk;
