Process Tree Support for SCW memory estimation.

Measure SCW memory usage by query process memory, and child process memory.
This replaces the Job Objects on Windows and provides a consistent code path on Wine and Native Windows.
Testing show both methods returning the same memory usage values, so there should be no change in behavior.
This is intended to be used as part of subsequent work to dynamically manage the amount of memory the Shader Compiler Workers have available.

#rnx
#FailedValidations Preflight

diff --git a/Engine/Source/Runtime/Engine/Private/ShaderCompiler/ShaderCompilerMemoryLimit.cpp b/Engine/Source/Runtime/Engine/Private/ShaderCompiler/ShaderCompilerMemoryLimit.cpp
--- a/Engine/Source/Runtime/Engine/Private/ShaderCompiler/ShaderCompilerMemoryLimit.cpp
+++ b/Engine/Source/Runtime/Engine/Private/ShaderCompiler/ShaderCompilerMemoryLimit.cpp
@@ -1,188 +1,133 @@
 // Copyright Epic Games, Inc. All Rights Reserved.
 
 /*=============================================================================
-	ShaderCompilerMemoryLimit.cpp: Wrapper for Windows specific Job Object functionality.
+	ShaderCompilerMemoryLimit.cpp: Wrapper for Windows specific Process Tree functionality.
 =============================================================================*/
 
 #include "ShaderCompilerMemoryLimit.h"
 #include "HAL/ConsoleManager.h"
 
-
 #if PLATFORM_WINDOWS
 
-FWindowsResourceRestrictedJobObject::FWindowsResourceRestrictedJobObject(const FString& InJobName, int32 InitialJobMemoryLimitMiB) :
-	JobName(FString::Printf(TEXT("%s-%u"), *InJobName, FPlatformProcess::GetCurrentProcessId()))
-{
-	JobObject = ::CreateJobObject(nullptr, *JobName);
-	if (JobObject == nullptr)
-	{
-		UE_LOG(LogWindows, Warning, TEXT("Failed to create job object \"%s\": %s"), *JobName, *GetErrorMessage());
-		return;
-	}
-	CreateAndLinkCompletionPort();
-	if (InitialJobMemoryLimitMiB > 0)
-	{
-		SetMemoryLimit(InitialJobMemoryLimitMiB);
-	}
-}
+#include "Windows/MinimalWindowsApi.h"
+#include "Windows/WindowsHWrapper.h"
+#include "Windows/AllowWindowsPlatformTypes.h"
+	#include <Psapi.h>
+#include "Windows/HideWindowsPlatformTypes.h"
 
-FWindowsResourceRestrictedJobObject::~FWindowsResourceRestrictedJobObject()
+void FWindowsResourceProcessTreeMemory::AddRootProcessId(uint32 ProcessId)
 {
-	if (JobObject != nullptr)
-	{
-		::CloseHandle(JobObject);
-	}
-	if (CompletionPort != nullptr)
-	{
-		::CloseHandle(CompletionPort);
-	}
+	RootProcessIds.AddUnique(ProcessId);
 }
 
-void FWindowsResourceRestrictedJobObject::AssignProcess(const FProcHandle& Process)
+void FWindowsResourceProcessTreeMemory::Reset()
 {
-	if (IsValid())
-	{
-		// Assign input process to job object
-		if (::AssignProcessToJobObject(JobObject, Process.Get()) == 0)
-		{
-			UE_LOG(LogWindows, Warning, TEXT("Failed to assign process to job object \"%s\": %s"), *JobName, *GetErrorMessage());
-		}
-	}
+	RootProcessIds.Empty();
+
+	AllProcesses.Empty();
+	TreeProcessIds.Empty();
 }
 
-void FWindowsResourceRestrictedJobObject::SetMemoryLimit(int32 InJobMemoryLimitMiB)
+bool FWindowsResourceProcessTreeMemory::TryGetMemoryUsage(FPlatformProcessMemoryStats& OutStats)
 {
-	checkf(InJobMemoryLimitMiB >= 1024, TEXT("Cannot launch ShaderCompileWorker processes with memory restriction of less than 1024 MiB (%u MiB was specified)"), InJobMemoryLimitMiB);
-	if (MemoryLimit != InJobMemoryLimitMiB)
-	{
-		if (IsValid())
-		{
-			JOBOBJECT_NOTIFICATION_LIMIT_INFORMATION JobLimitInfo = {};
+	CollectTreeProcessIds();
 
-			JobLimitInfo.JobMemoryLimit = 1024 * 1024 * static_cast<SIZE_T>(InJobMemoryLimitMiB);
-			JobLimitInfo.LimitFlags = JOB_OBJECT_LIMIT_JOB_MEMORY;
-
-			if (::SetInformationJobObject(JobObject, JobObjectNotificationLimitInformation, &JobLimitInfo, sizeof(JobLimitInfo)) == 0)
-			{
-				UE_LOG(LogWindows, Warning, TEXT("Failed to set restriction information for job object \"%s\": %s"), *JobName, *GetErrorMessage());
-			}
-		}
-		MemoryLimit = InJobMemoryLimitMiB;
-	}
-}
+	OutStats = { };
 
-bool FWindowsResourceRestrictedJobObject::QueryStatus(FJobObjectLimitationInfo& OutInfo)
-{
-	JOBOBJECT_LIMIT_VIOLATION_INFORMATION LimitViolationInfo = {};
-	DWORD ReturnLength = 0;
-	if (::QueryInformationJobObject(JobObject, JobObjectLimitViolationInformation, &LimitViolationInfo, sizeof(LimitViolationInfo), &ReturnLength) != 0 &&
-		ReturnLength == sizeof(LimitViolationInfo) &&
-		LimitViolationInfo.JobMemoryLimit > 0)
+	for (const uint32 ProcessId : TreeProcessIds)
 	{
-		OutInfo.MemoryLimit = LimitViolationInfo.JobMemoryLimit;
-		OutInfo.MemoryUsed = LimitViolationInfo.JobMemory;
-		return true;
-	}
-	return false;
-}
+		HANDLE ProcessHandle = ::OpenProcess(PROCESS_QUERY_INFORMATION | PROCESS_VM_READ, false, ProcessId);
 
-bool FWindowsResourceRestrictedJobObject::QueryLimitViolationStatus(FJobObjectLimitationInfo& OutInfo)
-{
-	if (IsValid())
-	{
-		// Check if a packet has been dequeued from the IO completion port
-		constexpr DWORD WaitDurationInMilliseconds = 0;
-		DWORD NumBytesTransferred = 0;
-		ULONG_PTR CompletionKey = 0;
-		LPOVERLAPPED Overlapped = nullptr;
-		while (::GetQueuedCompletionStatus(CompletionPort, &NumBytesTransferred, &CompletionKey, &Overlapped, WaitDurationInMilliseconds) != 0)
+		if (ProcessHandle != NULL)
 		{
-			// Check if this is the message we are looking for, i.e. notification about limiation violation
-			if (NumBytesTransferred == JOB_OBJECT_MSG_NOTIFICATION_LIMIT)
+			PROCESS_MEMORY_COUNTERS ProcessMemoryCounters;
+			FPlatformMemory::Memzero(&ProcessMemoryCounters, sizeof(ProcessMemoryCounters));
+
+			if (!::GetProcessMemoryInfo(ProcessHandle, &ProcessMemoryCounters, sizeof(ProcessMemoryCounters)))
 			{
-				// Query violation information from job object
-				JOBOBJECT_LIMIT_VIOLATION_INFORMATION LimitViolationInfo = {};
-				DWORD ReturnLength = 0;
-				if (::QueryInformationJobObject(JobObject, JobObjectLimitViolationInformation, &LimitViolationInfo, sizeof(LimitViolationInfo), &ReturnLength) != 0 &&
-					ReturnLength == sizeof(LimitViolationInfo))
+				if (FWindowsPlatformMisc::IsWine() && ::GetLastError() == ERROR_ACCESS_DENIED)
+				{
+					// Under Wine, GetProcessMemoryInfo() incorrectly returns ERROR_ACCESS_DENIED if a process is no longer running.
+				}
+				else
 				{
-					// Even though the limit has been exceeded as indicated by the JOB_OBJECT_MSG_NOTIFICATION_LIMIT message,
-					// by the time we query the information it might have already gone down again,
-					// so check if the memory information is actually exceeding our limit here.
-					if (LimitViolationInfo.JobMemory >= LimitViolationInfo.JobMemoryLimit)
-					{
-						OutInfo.MemoryLimit = LimitViolationInfo.JobMemoryLimit;
-						OutInfo.MemoryUsed = LimitViolationInfo.JobMemory;
-						return true;
-					}
+					UE_LOG(LogWindows, Warning, TEXT("Failure in call to GetProcessMemoryInfo (GetLastError=%d)"), ::GetLastError());
 				}
+
+				::CloseHandle(ProcessHandle);
+				return false;
+			}
+			else
+			{
+				OutStats.UsedPhysical += ProcessMemoryCounters.WorkingSetSize;
+				OutStats.PeakUsedPhysical += ProcessMemoryCounters.PeakWorkingSetSize;
+				OutStats.UsedVirtual += ProcessMemoryCounters.PagefileUsage;
+				OutStats.PeakUsedVirtual += ProcessMemoryCounters.PeakPagefileUsage;
 			}
+
+			::CloseHandle(ProcessHandle);
 		}
 	}
-	return false;
-}
 
-FString FWindowsResourceRestrictedJobObject::GetErrorMessage()
-{
-	TCHAR ErrorMessage[4096];
-	return FString(FPlatformMisc::GetSystemErrorMessage(ErrorMessage, UE_ARRAY_COUNT(ErrorMessage), 0));
+	return true;
 }
 
-bool FWindowsResourceRestrictedJobObject::IsValid() const
+void FWindowsResourceProcessTreeMemory::CollectTreeProcessIds()
 {
-	return JobObject != nullptr && CompletionPort != nullptr;
-}
+	AllProcesses.Empty();
+	TreeProcessIds.Empty();
 
-void FWindowsResourceRestrictedJobObject::CreateAndLinkCompletionPort()
-{
-	// Create completion port
-	CompletionPort = ::CreateIoCompletionPort(INVALID_HANDLE_VALUE, nullptr, 0, 1);
-	if (CompletionPort == nullptr)
+	if (!RootProcessIds.IsEmpty())
 	{
-		UE_LOG(LogWindows, Warning, TEXT("Failed to create I/O completion port for job object \"%s\": %s"), *JobName, *GetErrorMessage());
-		return;
-	}
+		FPlatformProcess::FProcEnumerator ProcEnumerator;
 
-	// Link completion port to job object
-	JOBOBJECT_ASSOCIATE_COMPLETION_PORT CompletionPortInfo = {};
+		while (ProcEnumerator.MoveNext())
+		{
+			FPlatformProcess::FProcEnumInfo ProcEnumInfo = ProcEnumerator.GetCurrent();
 
-	CompletionPortInfo.CompletionKey = JobObject;
-	CompletionPortInfo.CompletionPort = CompletionPort;
+			FProcessInfo ProcessInfo;
+			ProcessInfo.ProcessId = ProcEnumInfo.GetPID();
+			ProcessInfo.ParentProcessId = ProcEnumInfo.GetParentPID();
 
-	if (::SetInformationJobObject(JobObject, JobObjectAssociateCompletionPortInformation, &CompletionPortInfo, sizeof(CompletionPortInfo)) == 0)
-	{
-		UE_LOG(LogWindows, Warning, TEXT("Failed to set completion port for job object \"%s\": %s"), *JobName, *GetErrorMessage());
+			AllProcesses.Push(ProcessInfo);
+		}
+
+		for (const uint32 RootProcessId : RootProcessIds)
+		{
+			CollectTreeProcessIdsRecurse(RootProcessId);
+		}
 	}
 }
 
-#else
-
-FGenericResourceRestrictedJobObject::FGenericResourceRestrictedJobObject(const FString& JobName, int32 InitialJobMemoryLimitMiB)
+void FWindowsResourceProcessTreeMemory::CollectTreeProcessIdsRecurse(uint32 RootProcessId)
 {
-	// dummy
-}
+	TreeProcessIds.Push(RootProcessId);
 
-void FGenericResourceRestrictedJobObject::AssignProcess(const FProcHandle& Process)
-{
-	// dummy
+	for (const FProcessInfo& JobProcessInfo : AllProcesses)
+	{
+		if (JobProcessInfo.ParentProcessId == RootProcessId)
+		{
+			CollectTreeProcessIdsRecurse(JobProcessInfo.ProcessId);
+		}
+	}
 }
 
-void FGenericResourceRestrictedJobObject::SetMemoryLimit(int32 InJobMemoryLimitMiB)
+#else
+
+void FGenericResourceProcessTreeMemory::AddRootProcessId(uint32 ProcessId)
 {
 	// dummy
 }
 
-bool FGenericResourceRestrictedJobObject::QueryStatus(FJobObjectLimitationInfo& OutInfo)
+void FGenericResourceProcessTreeMemory::Reset()
 {
-	return false; // dummy
+	// dummy
 }
 
-bool FGenericResourceRestrictedJobObject::QueryLimitViolationStatus(FJobObjectLimitationInfo& OutInfo)
+bool FGenericResourceProcessTreeMemory::TryGetMemoryUsage(FPlatformProcessMemoryStats& OutStats)
 {
 	return false; // dummy
 }
 
 #endif // PLATFORM_WINDOWS
 
-
-
diff --git a/Engine/Source/Runtime/Engine/Private/ShaderCompiler/ShaderCompilerMemoryLimit.h b/Engine/Source/Runtime/Engine/Private/ShaderCompiler/ShaderCompilerMemoryLimit.h
--- a/Engine/Source/Runtime/Engine/Private/ShaderCompiler/ShaderCompilerMemoryLimit.h
+++ b/Engine/Source/Runtime/Engine/Private/ShaderCompiler/ShaderCompilerMemoryLimit.h
@@ -1,7 +1,7 @@
 // Copyright Epic Games, Inc. All Rights Reserved.
 
 /*=============================================================================
-	ShaderCompilerMemoryLimit.h: Wrapper for Windows specific Job Object functionality.
+	ShaderCompilerMemoryLimit.h: Wrapper for Windows specific Process Tree functionality.
 =============================================================================*/
 
 #pragma once
@@ -9,78 +9,68 @@
 #include "CoreMinimal.h"
 #include "HAL/PlatformProcess.h"
 
-
-/** Output structure for polling job object limitation violation status. See FWindowsResourceRestrictedJobObject::QueryLimitViolationStatus(). */
-struct FJobObjectLimitationInfo
-{
-	/** Size (in bytes) of the job memory limitation. */
-	int64 MemoryLimit = 0;
-
-	/** Size (in bytes) of the job memory usage. When QueryLimitViolationStatus() returns true, this will be greater than JobMemoryLimit, since the job object violated the limitation requirements. */
-	int64 MemoryUsed = 0;
-};
-
-
 #if PLATFORM_WINDOWS
 
-#include "Windows/MinimalWindowsApi.h"
-#include "Windows/WindowsHWrapper.h"
-
-/** Wrapper lass for a resource restricted job object. This is only available on Windows and acts as a placeholder on other platforms. */
-class FWindowsResourceRestrictedJobObject
+/** Wrapper for measuring memory of a process tree. This is only available on Windows and acts as a placeholder on other platforms. */
+class FWindowsResourceProcessTreeMemory
 {
 public:
-	FWindowsResourceRestrictedJobObject(const FWindowsResourceRestrictedJobObject&) = delete;
-	FWindowsResourceRestrictedJobObject& operator = (const FWindowsResourceRestrictedJobObject&) = delete;
-
-	FWindowsResourceRestrictedJobObject(const FString& InJobName, int32 InitialJobMemoryLimitMiB = 0);
-	~FWindowsResourceRestrictedJobObject();
+	FWindowsResourceProcessTreeMemory(const FWindowsResourceProcessTreeMemory&) = delete;
+	FWindowsResourceProcessTreeMemory& operator = (const FWindowsResourceProcessTreeMemory&) = delete;
 
-	/** Assigns the specified process to this job object. */
-	void AssignProcess(const FProcHandle& Process);
+	FWindowsResourceProcessTreeMemory() = default;
+	~FWindowsResourceProcessTreeMemory() = default;
 
-	/** Sets the memory limitation for this job object in MiB. This must be greater than or equal to 1024 MiB. */
-	void SetMemoryLimit(int32 InJobMemoryLimitMiB);
+	/** Add a ProcessId as the root of a set of processess to measure. */
+	void AddRootProcessId(uint32 ProcessId);
 
-	/** Queryies the status of the current memory usage. */
-	bool QueryStatus(FJobObjectLimitationInfo& OutInfo);
+	/** Empty the list of ProcessIds to measure */
+	void Reset();
 
-	/** Queries the status of limitation violation notifications. */
-	bool QueryLimitViolationStatus(FJobObjectLimitationInfo& OutInfo);
+	/** Try and get the total memory used by all process and their children from AddRootProcessId. */
+	bool TryGetMemoryUsage(FPlatformProcessMemoryStats& OutStats);
 
 private:
-	static FString GetErrorMessage();
+	struct FProcessInfo {
+		uint32 ProcessId = 0;
+		uint32 ParentProcessId = 0;
+	};
 
-	bool IsValid() const;
+	/* The Process IDs we will measure */
+	TArray<uint32> RootProcessIds;
 
-	void CreateAndLinkCompletionPort();
+	/* All Processes in the system */
+	TArray<FProcessInfo> AllProcesses;
 
-private:
-	FString JobName;
-	HANDLE JobObject = nullptr;
-	HANDLE CompletionPort = nullptr;
-	int32 MemoryLimit = 0;
+	/* The root Process IDs and all child Process IDs */
+	TArray<uint32> TreeProcessIds;
+
+	/* Populate the TreeProcessIds with RootProcessIds and their children. */
+	void CollectTreeProcessIds();
+
+	/* Recursively add the RootProcessId and children to the TreeProcessIds */
+	void CollectTreeProcessIdsRecurse(uint32 RootProcessId);
 };
 
-using FResourceRestrictedJobObject = FWindowsResourceRestrictedJobObject;
+using FResourceProcessTreeMemory = FWindowsResourceProcessTreeMemory;
 
 #else
 
-/** Wrapper lass for a resource restricted job object. This is only available on Windows and acts as a placeholder on other platforms. */
-class FGenericResourceRestrictedJobObject
+/** Wrapper for measuring memory of a process tree. This is only available on Windows and acts as a placeholder on other platforms. */
+class FGenericResourceProcessTreeMemory
 {
 public:
-	FGenericResourceRestrictedJobObject(const FGenericResourceRestrictedJobObject&) = delete;
-	FGenericResourceRestrictedJobObject& operator = (const FGenericResourceRestrictedJobObject&) = delete;
+	FGenericResourceProcessTreeMemory(const FGenericResourceProcessTreeMemory&) = delete;
+	FGenericResourceProcessTreeMemory& operator = (const FGenericResourceProcessTreeMemory&) = delete;
 
-	FGenericResourceRestrictedJobObject(const FString& JobName, int32 InitialJobMemoryLimitMiB = 0);
+	FGenericResourceProcessTreeMemory() = default;
+	~FGenericResourceProcessTreeMemory() = default;
 
-	void AssignProcess(const FProcHandle& Process);
-	void SetMemoryLimit(int32 InJobMemoryLimitMiB);
-	bool QueryStatus(FJobObjectLimitationInfo& OutInfo);
-	bool QueryLimitViolationStatus(FJobObjectLimitationInfo& OutInfo);
+	void AddRootProcessId(uint32 ProcessId);
+	void Reset();
+	bool TryGetMemoryUsage(FPlatformProcessMemoryStats& OutStats);
 };
 
-using FResourceRestrictedJobObject = FGenericResourceRestrictedJobObject;
+using FResourceProcessTreeMemory = FGenericResourceProcessTreeMemory;
 
 #endif // PLATFORM_WINDOWS
diff --git a/Engine/Source/Runtime/Engine/Private/ShaderCompiler/ShaderCompilerThreadRunnable.cpp b/Engine/Source/Runtime/Engine/Private/ShaderCompiler/ShaderCompilerThreadRunnable.cpp
--- a/Engine/Source/Runtime/Engine/Private/ShaderCompiler/ShaderCompilerThreadRunnable.cpp
+++ b/Engine/Source/Runtime/Engine/Private/ShaderCompiler/ShaderCompilerThreadRunnable.cpp
@@ -55,25 +55,15 @@ static constexpr int32 GSingleThreadedRunsMaxCount = (1 << 24);
 static const TCHAR* GWorkerInputFilename = TEXT("WorkerInputOnly.in");
 static const TCHAR* GWorkerOutputFilename = TEXT("WorkerOutputOnly.out");
 
-
-static FResourceRestrictedJobObject GSCWResourceRestrictedJobObject(TEXT("UE.ShaderCompileWorker.JobGroup"));
-
-// Apply memory limits (see CVar r.ShaderCompiler.MemoryLimit) to by assigning input process to resource restricted job object
-// and initialize this job object here, since we can't guarantee execution order of static global object (i.e. global cvar and the job object).
-static void ApplyWorkerProcessMemoryLimits(const FProcHandle& Process)
+/** Output structure for polling shader process memory violation status. See FShaderCompileThreadRunnable::QueryMemoryLimitViolationStatus(). */
+struct FShaderCompilerJobMemoryLimitInfo
 {
-	if (GShaderCompilerMemoryLimit > 0)
-	{
-		static bool bIsJobObjectLimitInitialized;
-		if (!bIsJobObjectLimitInitialized)
-		{
-			GSCWResourceRestrictedJobObject.SetMemoryLimit(GShaderCompilerMemoryLimit);
-			bIsJobObjectLimitInitialized = true;
-		}
-		GSCWResourceRestrictedJobObject.AssignProcess(Process);
-	}
-}
+	/** Size (in bytes) of the shader process memory limitation. */
+	int64 MemoryLimit = 0;
 
+	/** Size (in bytes) of the shader process memory usage. When QueryMemoryLimitViolationStatus() returns true, this will be greater than JobMemoryLimit, since the process violated the limitation requirements. */
+	int64 MemoryUsed = 0;
+};
 
 /** Information tracked for each shader compile worker process instance. */
 struct FShaderCompileWorkerInfo
@@ -206,11 +196,6 @@ void FShaderCompileThreadRunnableBase::SetPriorityRange(EShaderCompileJobPriorit
 
 FShaderCompileThreadRunnable::FShaderCompileThreadRunnable(FShaderCompilingManager* InManager)
 	: FShaderCompileThreadRunnableBase(InManager)
-#if PLATFORM_WINDOWS
-	, bEstimateCommittedMemory(FPlatformMisc::IsWine()) // Use alternative code path to estimate memory when we're running on POSIX/Wine instead of a real Windows host system
-#else
-	, bEstimateCommittedMemory(true) // Use alternative code path to estimate memory when we're running on POSIX
-#endif
 {
 	for (uint32 WorkerIndex = 0; WorkerIndex < Manager->NumShaderCompilingThreads; WorkerIndex++)
 	{
@@ -717,12 +702,6 @@ bool FShaderCompileThreadRunnable::LaunchWorkersIfNeeded()
 				CurrentWorkerInfo.WorkerProcess = Manager->LaunchWorker(WorkingDirectory, Manager->ProcessId, WorkerIndex, GWorkerInputFilename, GWorkerOutputFilename, &CurrentWorkerInfo.WorkerProcessId);
 				CurrentWorkerInfo.bLaunchedWorker = true;
 
-				// Assign process to job object to monitor the total memory consumption of all SCW processes. Ignore if we only estimate the memory due to partial support in POSIX/Wine.
-				if (!bEstimateCommittedMemory)
-				{
-					ApplyWorkerProcessMemoryLimits(CurrentWorkerInfo.WorkerProcess);
-				}
-
 				NumberLaunched++;
 			}
 		}
@@ -1066,7 +1045,7 @@ void FShaderCompileThreadRunnable::CheckMemoryLimitViolation()
 		MemoryMonitoringState.LastTimeOfMemoryLimitPoll = CurrentTime;
 
 		// Check if memory limit has been exceeded
-		FJobObjectLimitationInfo LimitInfo;
+		FShaderCompilerJobMemoryLimitInfo LimitInfo;
 		if (QueryMemoryLimitViolationStatus(LimitInfo))
 		{
 			MemoryMonitoringState.LastTimeOfSuspeningOrResumingWorkers = CurrentTime;
@@ -1112,7 +1091,7 @@ void FShaderCompileThreadRunnable::CheckMemoryLimitViolation()
 	{
 		MemoryMonitoringState.LastTimeOfSuspeningOrResumingWorkers = CurrentTime;
 
-		FJobObjectLimitationInfo LimitInfo;
+		FShaderCompilerJobMemoryLimitInfo LimitInfo;
 		if (QueryMemoryStatus(LimitInfo))
 		{
 			// If we are below half of our memory limit, resume 50% of available workers.
@@ -1146,43 +1125,29 @@ void FShaderCompileThreadRunnable::CheckMemoryLimitViolation()
 	}
 }
 
-bool FShaderCompileThreadRunnable::QueryMemoryStatus(FJobObjectLimitationInfo& OutInfo)
+bool FShaderCompileThreadRunnable::QueryMemoryStatus(FShaderCompilerJobMemoryLimitInfo& OutInfo)
 {
-	if (bEstimateCommittedMemory)
+	const FShaderCompileMemoryUsage MemoryUsage = GetExternalWorkerMemoryUsage();
+	if (MemoryUsage.VirtualMemory > 0)
 	{
-		const FShaderCompileMemoryUsage MemoryUsage = GetExternalWorkerMemoryUsage();
-		if (MemoryUsage.VirtualMemory > 0)
-		{
-			OutInfo.MemoryLimit = static_cast<int64>(GShaderCompilerMemoryLimit) * 1024 * 1024;
-			OutInfo.MemoryUsed = MemoryUsage.VirtualMemory;
-			return true;
-		}
-		return false;
-	}
-	else
-	{
-		return GSCWResourceRestrictedJobObject.QueryStatus(OutInfo);
+		OutInfo.MemoryLimit = static_cast<int64>(GShaderCompilerMemoryLimit) * 1024 * 1024;
+		OutInfo.MemoryUsed = MemoryUsage.VirtualMemory;
+		return true;
 	}
+	return false;
 }
 
-bool FShaderCompileThreadRunnable::QueryMemoryLimitViolationStatus(FJobObjectLimitationInfo& OutInfo)
+bool FShaderCompileThreadRunnable::QueryMemoryLimitViolationStatus(FShaderCompilerJobMemoryLimitInfo& OutInfo)
 {
-	if (bEstimateCommittedMemory)
-	{
-		const FShaderCompileMemoryUsage MemoryUsage = GetExternalWorkerMemoryUsage();
-		const int64 MemoryLimitInBytes = static_cast<int64>(GShaderCompilerMemoryLimit) * 1024 * 1024;
-		if (MemoryUsage.VirtualMemory >= static_cast<uint64>(MemoryLimitInBytes))
-		{
-			OutInfo.MemoryLimit = MemoryLimitInBytes;
-			OutInfo.MemoryUsed = MemoryUsage.VirtualMemory;
-			return true;
-		}
-		return false;
-	}
-	else
+	const FShaderCompileMemoryUsage MemoryUsage = GetExternalWorkerMemoryUsage();
+	const int64 MemoryLimitInBytes = static_cast<int64>(GShaderCompilerMemoryLimit) * 1024 * 1024;
+	if (MemoryUsage.VirtualMemory >= static_cast<uint64>(MemoryLimitInBytes))
 	{
-		return GSCWResourceRestrictedJobObject.QueryLimitViolationStatus(OutInfo);
+		OutInfo.MemoryLimit = MemoryLimitInBytes;
+		OutInfo.MemoryUsed = MemoryUsage.VirtualMemory;
+		return true;
 	}
+	return false;
 }
 
 bool FShaderCompileThreadRunnable::PrintWorkerMemoryUsage(bool bAllowToWaitForLock)
@@ -1207,23 +1172,27 @@ bool FShaderCompileThreadRunnable::PrintWorkerMemoryUsage(bool bAllowToWaitForLo
 
 FShaderCompileMemoryUsage FShaderCompileThreadRunnable::GetExternalWorkerMemoryUsage()
 {
-	FScopeLock WorkerScopeLock(&WorkerInfosLock);
 	FShaderCompileMemoryUsage MemoryUsage{};
+	FScopeLock WorkerScopeLock(&WorkerInfosLock);
+
+	FResourceProcessTreeMemory ProcessTreeMemory;
+
 	for (const TUniquePtr<FShaderCompileWorkerInfo>& WorkerInfo : WorkerInfos)
 	{
-		FProcHandle ProcHandle = WorkerInfo->WorkerProcess;
-		if (!ProcHandle.IsValid())
+		if (WorkerInfo->WorkerProcess.IsValid() && WorkerInfo->WorkerProcessId != 0)
 		{
-			continue;
-		}
-		FPlatformProcessMemoryStats MemoryStats;
-		if (FPlatformProcess::TryGetMemoryUsage(ProcHandle, MemoryStats))
-		{
-			// Virtual memory is committed memory on Windows.
-			MemoryUsage.VirtualMemory += MemoryStats.UsedVirtual;
-			MemoryUsage.PhysicalMemory += MemoryStats.UsedPhysical;
+			ProcessTreeMemory.AddRootProcessId(WorkerInfo->WorkerProcessId);
 		}
 	}
+
+	FPlatformProcessMemoryStats MemoryStats;
+	if (ProcessTreeMemory.TryGetMemoryUsage(MemoryStats))
+	{
+		// Virtual memory is committed memory on Windows.
+		MemoryUsage.VirtualMemory = MemoryStats.UsedVirtual;
+		MemoryUsage.PhysicalMemory = MemoryStats.UsedPhysical;
+	}
+
 	return MemoryUsage;
 }
 
diff --git a/Engine/Source/Runtime/Engine/Public/ShaderCompiler.h b/Engine/Source/Runtime/Engine/Public/ShaderCompiler.h
--- a/Engine/Source/Runtime/Engine/Public/ShaderCompiler.h
+++ b/Engine/Source/Runtime/Engine/Public/ShaderCompiler.h
@@ -60,7 +60,7 @@ extern ENGINE_API bool IsMaterialMapDDCEnabled();
  */
 extern ENGINE_API bool ShouldCompileODSCOnlyShaders();
 
-struct FJobObjectLimitationInfo;
+struct FShaderCompilerJobMemoryLimitInfo;
 
 struct FShaderJobCacheStoredOutput;
 class FShaderJobCache;
@@ -232,8 +232,6 @@ class FShaderCompileThreadRunnable : public FShaderCompileThreadRunnableBase
 	friend class FShaderCompilingManager;
 private:
 
-	const bool bEstimateCommittedMemory = false; // Must be true on POSIX/Wine where only a small subset of the Job Object functionality is implemented
-
 	/** Information about the active workers that this thread is tracking. */
 	TArray<TUniquePtr<struct FShaderCompileWorkerInfo>> WorkerInfos;
 	mutable FCriticalSection WorkerInfosLock;
@@ -334,11 +332,11 @@ private:
 	/** Checks it the memory limit for shader compile workers has been exceeded and suspend workers as needed. */
 	void CheckMemoryLimitViolation();
 
-	/** Queries the memory status of all worker processes. Either uses FResourceRestrictedJobObject or QueryEstimatedCommittedMemory() when running on POSIX/Wine. */
-	bool QueryMemoryStatus(FJobObjectLimitationInfo& OutInfo);
+	/** Queries the memory status of all worker processes. Uses FResourceProcessTreeMemory. */
+	bool QueryMemoryStatus(FShaderCompilerJobMemoryLimitInfo& OutInfo);
 
-	/** Queries the status if the job object for all worker processes has violated the memory limitation. Either uses FResourceRestrictedJobObject or QueryEstimatedCommittedMemory() when running on POSIX/Wine. */
-	bool QueryMemoryLimitViolationStatus(FJobObjectLimitationInfo& OutInfo);
+	/** Queries the status if the worker processes have violated the memory limitation. Either uses FResourceProcessTreeMemory. */
+	bool QueryMemoryLimitViolationStatus(FShaderCompilerJobMemoryLimitInfo& OutInfo);
 };
 
 class FShaderCompileUtilities
