SkeletalMeshComponent: Fix crash when calling SetAnimInstanceClass during a notify callback

Early out with a warning instead.
Updated comment/tooltip to describe the new behavior.

diff --git a/Engine/Source/Runtime/Engine/Classes/Components/SkeletalMeshComponent.h b/Engine/Source/Runtime/Engine/Classes/Components/SkeletalMeshComponent.h
--- a/Engine/Source/Runtime/Engine/Classes/Components/SkeletalMeshComponent.h
+++ b/Engine/Source/Runtime/Engine/Classes/Components/SkeletalMeshComponent.h
@@ -1002,7 +1002,10 @@ public:
 	UFUNCTION(BlueprintInternalUseOnly)
 	ENGINE_API class UClass*  GetAnimClass();
 
-	/** Set the anim instance class. Clears and re-initializes the anim instance with the new class and sets animation mode to 'AnimationBlueprint' */
+	/**
+	 * Set the anim instance class. Clears and re-initializes the anim instance with the new class and sets animation mode to 'AnimationBlueprint'
+	 * This function will fail with a warning if it is called while the current anim instance is evaluating (e.g. from an anim notify).
+	 */
 	UFUNCTION(BlueprintCallable, Category = "Components|SkeletalMesh", meta = (Keywords = "AnimBlueprint", DisplayName = "Set Anim Instance Class"))
 	ENGINE_API virtual void SetAnimInstanceClass(class UClass* NewClass);
 
diff --git a/Engine/Source/Runtime/Engine/Private/Components/SkeletalMeshComponent.cpp b/Engine/Source/Runtime/Engine/Private/Components/SkeletalMeshComponent.cpp
--- a/Engine/Source/Runtime/Engine/Private/Components/SkeletalMeshComponent.cpp
+++ b/Engine/Source/Runtime/Engine/Private/Components/SkeletalMeshComponent.cpp
@@ -3440,6 +3440,15 @@ class UClass* USkeletalMeshComponent::GetAnimClass()
 
 void USkeletalMeshComponent::SetAnimInstanceClass(class UClass* NewClass)
 {
+	if (bPostEvaluatingAnimation
+#if DO_CHECK
+		|| (AnimScriptInstance && (AnimScriptInstance->IsUpdatingAnimation() || AnimScriptInstance->IsPostUpdatingAnimation()))
+#endif
+		)
+	{
+		UE_LOG(LogAnimation, Warning, TEXT("USkeletalMeshComponent::SetAnimInstanceClass: Cannot set anim instance class while animation is still updating/evaluating. Apply a delay or perform the switch elsewhere."));
+		return;
+	}
 	if (NewClass != nullptr)
 	{
 		// set the animation mode
