PropertyAccess: Hardened GetAccessAddress in case of unexpected data

Unable to reproduce the case, but I can envisage potential issues when deleting/editing functions and skipping compile-on-load, so hardening GetAccessAddress to better deal with unexpected data inconsistencies.
Additionally address some logical errors where missing break statements meant we fell through to subsequent states.

diff --git a/Engine/Source/Runtime/Engine/Private/PropertyAccess.cpp b/Engine/Source/Runtime/Engine/Private/PropertyAccess.cpp
--- a/Engine/Source/Runtime/Engine/Private/PropertyAccess.cpp
+++ b/Engine/Source/Runtime/Engine/Private/PropertyAccess.cpp
@@ -490,7 +490,9 @@ struct FPropertyAccessSystem
 					break;
 				}
 				default:
-					check(false);
+					Address = nullptr;
+					checkNoEntry();
+					break;
 				}
 				break;
 			}
@@ -517,7 +519,7 @@ struct FPropertyAccessSystem
 			case EPropertyAccessIndirectionType::ScriptFunction:
 			case EPropertyAccessIndirectionType::NativeFunction:
 			{
-				if(Indirection.Function != nullptr)
+				if(Indirection.Function != nullptr && Indirection.Function->GetReturnProperty() != nullptr)
 				{
 					UObject* CalleeObject = static_cast<UObject*>(Address);
  
@@ -526,7 +528,6 @@ struct FPropertyAccessSystem
 					Address = Align(ArrayBuffer.GetData(), Indirection.ReturnBufferAlignment);
  
 					// Init value
-					check(Indirection.Function->GetReturnProperty());
 					check(Indirection.Function->GetReturnProperty() == Indirection.Property.Get());
 					Indirection.Property.Get()->InitializeValue(Address);
 					BufferProperty = Indirection.Property.Get();
@@ -564,6 +565,8 @@ struct FPropertyAccessSystem
 						break;
 					}
 					default:
+						Address = nullptr;
+						checkNoEntry();
 						break;
 					}
 					break;
@@ -572,9 +575,12 @@ struct FPropertyAccessSystem
 				{
 					Address = nullptr;
 				}
+				break;
 			}
 			default:
-				check(false);
+				Address = nullptr;
+				checkNoEntry();
+				break;
 			}
 		}
  
