Fixed a physics bug where unwelding via socket name failed to unweld.

diff --git a/Engine/Source/Runtime/Engine/Private/PrimitiveComponentPhysics.cpp b/Engine/Source/Runtime/Engine/Private/PrimitiveComponentPhysics.cpp
--- a/Engine/Source/Runtime/Engine/Private/PrimitiveComponentPhysics.cpp
+++ b/Engine/Source/Runtime/Engine/Private/PrimitiveComponentPhysics.cpp
@@ -1015,6 +1015,26 @@ void UPrimitiveComponent::GetWeldedBodies(TArray<FBodyInstance*> & OutWeldedBodi
 	}
 }
 
+void GetRootBodyAndBoneName(const UPrimitiveComponent* RootComponent, const FName& SocketName, FBodyInstance*& OutRootBI, FName& OutBoneName)
+{
+	FBodyInstance* RootBI = RootComponent->GetBodyInstance(SocketName, false);
+	FName BoneName = SocketName;
+
+	//If initial lookup fails, check if we need to lookup the bone from the socket instead
+	if (!RootBI && RootComponent->IsA<USkeletalMeshComponent>())
+	{
+		const USkeletalMeshComponent* RootSkeletal = Cast<USkeletalMeshComponent>(RootComponent);
+		if (const USkeletalMeshSocket* Socket = RootSkeletal->GetSocketByName(SocketName))
+		{
+			BoneName = Socket->BoneName;
+			RootBI = RootComponent->GetBodyInstance(BoneName, false);
+		}
+	}
+
+	OutRootBI = RootBI;
+	OutBoneName = BoneName;
+}
+
 bool UPrimitiveComponent::WeldToImplementation(USceneComponent * InParent, FName ParentSocketName /* = Name_None */, bool bWeldSimulatedChild /* = true */, bool bWeldToKinematicParent /* = false */)
 {
 	SCOPE_CYCLE_COUNTER(STAT_WeldPhysics);
@@ -1051,19 +1071,9 @@ bool UPrimitiveComponent::WeldToImplementation(USceneComponent * InParent, FName
 
 	if (RootComponent)
 	{
-		FBodyInstance* RootBI = RootComponent->GetBodyInstance(SocketName, false);
-		FName BoneName = SocketName;
-
-		//If initial lookup fails, check if we need to lookup the bone from the socket instead
-		if (!RootBI && RootComponent->IsA<USkeletalMeshComponent>())
-		{
-			USkeletalMeshComponent* RootSkeletal = Cast<USkeletalMeshComponent>(RootComponent);
-			if (const USkeletalMeshSocket* Socket = RootSkeletal->GetSocketByName(SocketName))
-			{
-				BoneName = Socket->BoneName;
-				RootBI = RootComponent->GetBodyInstance(BoneName, false);
-			}
-		}
+		FBodyInstance* RootBI = nullptr;
+		FName BoneName;
+		GetRootBodyAndBoneName(RootComponent, SocketName, RootBI, BoneName);
 
 		if (RootBI)
 		{
@@ -1142,7 +1152,11 @@ void UPrimitiveComponent::UnWeldFromParent()
 
 	if (RootComponent)
 	{
-		if (FBodyInstance* RootBI = RootComponent->GetBodyInstance(SocketName, false))
+		FBodyInstance* RootBI = nullptr;
+		FName BoneName;
+		GetRootBodyAndBoneName(RootComponent, SocketName, RootBI, BoneName);
+
+		if (RootBI != nullptr)
 		{
 			bool bRootIsBeingDeleted = !IsValidChecked(RootComponent) || RootComponent->IsUnreachable();
 			const FBodyInstance* PrevWeldParent = NewRootBI->WeldParent;
