PR #13683: [????] Add more inline crash details when skeletons crash

diff --git a/Engine/Source/Runtime/Engine/Public/BonePose.h b/Engine/Source/Runtime/Engine/Public/BonePose.h
--- a/Engine/Source/Runtime/Engine/Public/BonePose.h
+++ b/Engine/Source/Runtime/Engine/Public/BonePose.h
@@ -631,8 +631,18 @@ void FCSPose<PoseType>::CalculateComponentSpaceTransform(BoneIndexType BoneIndex
 		// current Bones(Index) should contain LocalPoses.
 		FTransform& Bone = Pose[BoneIndex];
 		FTransform& ParentBone = Pose[ParentIndex];
-		check(!Bone.ContainsNaN());
-		check(!ParentBone.ContainsNaN());
+		
+		checkf(!Bone.ContainsNaN(), TEXT("Pose[%d] contains NaN: %s\nBase64: %s"),
+			BoneIndex.GetInt(),
+			*Bone.ToString(),
+			*FBase64::Encode(reinterpret_cast<const uint8*>(&Bone), sizeof(FTransform))
+		);
+
+		checkf(!ParentBone.ContainsNaN(), TEXT("Pose[%d] contains NaN: %s\nBase64: %s"),
+			ParentIndex.GetInt(),
+			*ParentBone.ToString(),
+			*FBase64::Encode(reinterpret_cast<const uint8*>(&ParentBone), sizeof(FTransform))
+		);
 
 		FTransform ComponentTransform = Bone * ParentBone;
 		if (ComponentTransform.ContainsNaN())
@@ -655,7 +665,13 @@ void FCSPose<PoseType>::CalculateComponentSpaceTransform(BoneIndexType BoneIndex
 		}
 		Bone = ComponentTransform;
 		Bone.NormalizeRotation();
-		check(!Bone.ContainsNaN());
+		
+		checkf(!Bone.ContainsNaN(), TEXT("Normalized ComponentTransform for Pose[%d] contains NaN: %s\nBase64: %s"),
+			BoneIndex.GetInt(),
+			*Bone.ToString(),
+			*FBase64::Encode(reinterpret_cast<const uint8*>(&Bone), sizeof(FTransform))
+		);
+
 		ComponentSpaceFlags[BoneIndex] = 1;
 
 		BoneIndex = BoneIndexStack.Pop(EAllowShrinking::No);
@@ -985,4 +1001,4 @@ extern ENGINE_API void BuildPoseFromRawData(
 	FName SourceName, 
 	const TArray<FTransform>& RetargetTransforms,
 	const TMap<int32, const FTransformCurve*>* AdditiveBoneTransformCurves = nullptr
-	);
+	);
\ No newline at end of file
