Substrate: When subsurface mode is set to none, force MeanFreePath to 0

This avoids needing additional checks at runtime

#rnx

diff --git a/Engine/Source/Runtime/Engine/Private/Materials/MaterialExpressionSubstrate.cpp b/Engine/Source/Runtime/Engine/Private/Materials/MaterialExpressionSubstrate.cpp
--- a/Engine/Source/Runtime/Engine/Private/Materials/MaterialExpressionSubstrate.cpp
+++ b/Engine/Source/Runtime/Engine/Private/Materials/MaterialExpressionSubstrate.cpp
@@ -877,7 +877,7 @@ int32 UMaterialExpressionSubstrateSlabBSDF::Compile(class FMaterialCompiler* Com
 	int32 GlintUVCodeChunk = CompileWithDefaultFloat2(Compiler, GlintUV, 0.0f, 0.0f);
 
 	// Disable some features if requested by the simplification process
-	if (!SubstrateOperator.Has(ESubstrateBsdfFeature::MFPPluggedIn))
+	if (!SubstrateOperator.Has(ESubstrateBsdfFeature::MFPPluggedIn) || SubSurfaceType == EMaterialSubSurfaceType::MSS_None)
 	{
 		SSSMFPCodeChunk = Compiler->Constant3(0.0f, 0.0f, 0.0f);
 	}
