Localize the AssetRegistry keys for Materials.

- Has Per Instance Custom Data
- Has Vertex Interpolator
- Has Per Instance Random
- Has Scene Color

#rnx

diff --git a/Engine/Source/Runtime/Engine/Private/Materials/MaterialInterface.cpp b/Engine/Source/Runtime/Engine/Private/Materials/MaterialInterface.cpp
--- a/Engine/Source/Runtime/Engine/Private/Materials/MaterialInterface.cpp
+++ b/Engine/Source/Runtime/Engine/Private/Materials/MaterialInterface.cpp
@@ -858,6 +858,43 @@ void UMaterialInterface::PostEditChangeProperty(FPropertyChangedEvent& PropertyC
 	Super::PostEditChangeProperty(PropertyChangedEvent);
 }
 
+namespace MaterialPrivate
+{
+	static const FName NAME_HasSceneColor("HasSceneColor");
+	static const FName NAME_HasPerInstanceRandom("HasPerInstanceRandom");
+	static const FName NAME_HasPerInstanceCustomData("HasPerInstanceCustomData");
+	static const FName NAME_HasVertexInterpolator("HasVertexInterpolator");
+}
+
+void UMaterialInterface::GetAssetRegistryTagMetadata(TMap<FName, FAssetRegistryTagMetadata>& OutMetadata) const
+{
+	Super::GetAssetRegistryTagMetadata(OutMetadata);
+
+	OutMetadata.Add(
+		MaterialPrivate::NAME_HasSceneColor,
+		FAssetRegistryTagMetadata()
+		.SetDisplayName(LOCTEXT("HasSceneColor", "Has Scene Color"))
+	);
+
+	OutMetadata.Add(
+		MaterialPrivate::NAME_HasPerInstanceRandom,
+		FAssetRegistryTagMetadata()
+		.SetDisplayName(LOCTEXT("HasPerInstanceRandom", "Has Per Instance Random"))
+	);
+
+	OutMetadata.Add(
+		MaterialPrivate::NAME_HasPerInstanceCustomData,
+		FAssetRegistryTagMetadata()
+		.SetDisplayName(LOCTEXT("HasPerInstanceCustomData", "Has Per Instance Custom Data"))
+	);
+
+	OutMetadata.Add(
+		MaterialPrivate::NAME_HasVertexInterpolator,
+		FAssetRegistryTagMetadata()
+		.SetDisplayName(LOCTEXT("HasVertexInterpolator", "Has Vertex Interpolator"))
+	);
+}
+
 void UMaterialInterface::GetAssetRegistryTags(FAssetRegistryTagsContext Context) const
 {
 	if (AssetImportData)
@@ -867,10 +904,10 @@ void UMaterialInterface::GetAssetRegistryTags(FAssetRegistryTagsContext Context)
 
 	{
 		const FMaterialCachedExpressionData& CachedData = GetCachedExpressionData();
-		Context.AddTag(FAssetRegistryTag("HasSceneColor", CachedData.bHasSceneColor ? TEXT("True") : TEXT("False"), FAssetRegistryTag::TT_Alphabetical));
-		Context.AddTag(FAssetRegistryTag("HasPerInstanceRandom", CachedData.bHasPerInstanceRandom ? TEXT("True") : TEXT("False"), FAssetRegistryTag::TT_Alphabetical));
-		Context.AddTag(FAssetRegistryTag("HasPerInstanceCustomData", CachedData.bHasPerInstanceCustomData ? TEXT("True") : TEXT("False"), FAssetRegistryTag::TT_Alphabetical));
-		Context.AddTag(FAssetRegistryTag("HasVertexInterpolator", CachedData.bHasVertexInterpolator ? TEXT("True") : TEXT("False"), FAssetRegistryTag::TT_Alphabetical));
+		Context.AddTag(FAssetRegistryTag(MaterialPrivate::NAME_HasSceneColor, CachedData.bHasSceneColor ? TEXT("True") : TEXT("False"), FAssetRegistryTag::TT_Alphabetical));
+		Context.AddTag(FAssetRegistryTag(MaterialPrivate::NAME_HasPerInstanceRandom, CachedData.bHasPerInstanceRandom ? TEXT("True") : TEXT("False"), FAssetRegistryTag::TT_Alphabetical));
+		Context.AddTag(FAssetRegistryTag(MaterialPrivate::NAME_HasPerInstanceCustomData, CachedData.bHasPerInstanceCustomData ? TEXT("True") : TEXT("False"), FAssetRegistryTag::TT_Alphabetical));
+		Context.AddTag(FAssetRegistryTag(MaterialPrivate::NAME_HasVertexInterpolator, CachedData.bHasVertexInterpolator ? TEXT("True") : TEXT("False"), FAssetRegistryTag::TT_Alphabetical));
 	}
 
 	Super::GetAssetRegistryTags(Context);
diff --git a/Engine/Source/Runtime/Engine/Public/Materials/MaterialInterface.h b/Engine/Source/Runtime/Engine/Public/Materials/MaterialInterface.h
--- a/Engine/Source/Runtime/Engine/Public/Materials/MaterialInterface.h
+++ b/Engine/Source/Runtime/Engine/Public/Materials/MaterialInterface.h
@@ -454,6 +454,7 @@ public:
 
 #if WITH_EDITOR
 	ENGINE_API virtual void PostEditChangeProperty(FPropertyChangedEvent& PropertyChangedEvent) override;
+	ENGINE_API virtual void GetAssetRegistryTagMetadata(TMap<FName, FAssetRegistryTagMetadata>& OutMetadata) const override;
 	ENGINE_API virtual void GetAssetRegistryTags(FAssetRegistryTagsContext Context) const override;
 #endif // WITH_EDITOR
 #if WITH_EDITOR
