Fixed HLOD collision data not being stripped out at cook
* Issue introduced in CL45877811
* IsProceduralSave() was used rather than IsFromAutoSave()

diff --git a/Engine/Source/Runtime/Engine/Private/WorldPartition/HLOD/HLODActor.cpp b/Engine/Source/Runtime/Engine/Private/WorldPartition/HLOD/HLODActor.cpp
--- a/Engine/Source/Runtime/Engine/Private/WorldPartition/HLOD/HLODActor.cpp
+++ b/Engine/Source/Runtime/Engine/Private/WorldPartition/HLOD/HLODActor.cpp
@@ -322,7 +322,7 @@ void AWorldPartitionHLOD::PreSave(FObjectPreSaveContext ObjectSaveContext)
 {
 	Super::PreSave(ObjectSaveContext);
 
-	if (ObjectSaveContext.IsProceduralSave())
+	if (ObjectSaveContext.IsFromAutoSave())
 	{
 		return;
 	}
