Atom - Ensure Cook/HLOD builds determinism by clearing CustomPrimitiveData fields of components
* We can't guarantee that LUTs added to the atlas will be inserted in the same order. If the order changes, this means that a different index will be stored in the components.
* Cook: Given that we always register the LUT on BeginPlay(), make sure to clear out the values (0) when cooking.
* HLOD Build: The HLOD hash (hash of the input components) can be affected by this non-determinism too as we include the component's CustomPrimitiveData fields in it. Added a Pre/Post HLOD build event to clear the atlas before a build. This ensure that the HLOD source components will get the same LUT indices for each build. The atlas is restored after building each HLOD.

diff --git a/Engine/Source/Runtime/Engine/Private/WorldPartition/HLOD/HLODActor.cpp b/Engine/Source/Runtime/Engine/Private/WorldPartition/HLOD/HLODActor.cpp
--- a/Engine/Source/Runtime/Engine/Private/WorldPartition/HLOD/HLODActor.cpp
+++ b/Engine/Source/Runtime/Engine/Private/WorldPartition/HLOD/HLODActor.cpp
@@ -45,6 +45,7 @@ static FAutoConsoleVariableRef CVarWorldPartitionHLODForceDisableShadows(
 
 #if WITH_EDITOR
 const FName AWorldPartitionHLOD::NAME_HLODHash_AssetTag(TEXT("HLODActor_HLODHash"));
+FWorldPartitionHLODBuildEventDelegate AWorldPartitionHLOD::HLODBuildEventDelegate;
 #endif
 
 AWorldPartitionHLOD::AWorldPartitionHLOD(const FObjectInitializer& ObjectInitializer)
@@ -584,6 +585,8 @@ void AWorldPartitionHLOD::SetHLODHash(uint32 InHLODHash, const FString& InHLODBu
 
 void AWorldPartitionHLOD::BuildHLOD(bool bForceBuild)
 {
+	GetHLODBuildEventDelegate().Broadcast({ FWorldPartitionHLODBuildEvent::EEventType::BeginBuild, this });
+
 	IWorldPartitionHLODUtilitiesModule* WPHLODUtilitiesModule = FModuleManager::Get().LoadModulePtr<IWorldPartitionHLODUtilitiesModule>("WorldPartitionHLODUtilities");
 	if (IWorldPartitionHLODUtilities* WPHLODUtilities = WPHLODUtilitiesModule != nullptr ? WPHLODUtilitiesModule->GetUtilities() : nullptr)
 	{
@@ -600,6 +603,13 @@ void AWorldPartitionHLOD::BuildHLOD(bool bForceBuild)
 	// This built data will be used by the cooking (it will convert it to level texture streaming built data).
 	// Use same quality level and feature level as FEditorBuildUtils::EditorBuildTextureStreaming
 	BuildActorTextureStreamingData(this, EMaterialQualityLevel::High, GMaxRHIFeatureLevel);
+
+	GetHLODBuildEventDelegate().Broadcast({ FWorldPartitionHLODBuildEvent::EEventType::EndBuild, this });
+}
+
+FWorldPartitionHLODBuildEventDelegate& AWorldPartitionHLOD::GetHLODBuildEventDelegate()
+{
+	return HLODBuildEventDelegate;
 }
 
 void AWorldPartitionHLOD::GetAssetRegistryTags(FAssetRegistryTagsContext Context) const
diff --git a/Engine/Source/Runtime/Engine/Public/WorldPartition/HLOD/HLODActor.h b/Engine/Source/Runtime/Engine/Public/WorldPartition/HLOD/HLODActor.h
--- a/Engine/Source/Runtime/Engine/Public/WorldPartition/HLOD/HLODActor.h
+++ b/Engine/Source/Runtime/Engine/Public/WorldPartition/HLOD/HLODActor.h
@@ -23,6 +23,22 @@ class UWorldPartitionHLODSourceActors;
 
 ENGINE_API DECLARE_LOG_CATEGORY_EXTERN(LogHLODHash, Log, All);
 
+#if WITH_EDITOR
+struct FWorldPartitionHLODBuildEvent
+{
+	enum class EEventType
+	{
+		BeginBuild,
+		EndBuild
+	};
+
+	EEventType EventType;
+	TObjectPtr<AWorldPartitionHLOD> TargetHLODActor;
+};
+
+DECLARE_MULTICAST_DELEGATE_OneParam(FWorldPartitionHLODBuildEventDelegate, const FWorldPartitionHLODBuildEvent&)
+#endif
+
 UCLASS(NotPlaceable, MinimalAPI, HideCategories=(Rendering, Replication, Collision, Physics, Navigation, Networking, Input, Actor, LevelInstance, Cooking))
 class AWorldPartitionHLOD : public AActor, public IWorldPartitionHLODObject
 {
@@ -84,6 +100,8 @@ public:
 	void SetStat(FName InStatName, int64 InStatValue) { HLODStats.Add(InStatName, InStatValue); }
 	void ResetStats() { HLODStats.Reset(); }
 
+	ENGINE_API static FWorldPartitionHLODBuildEventDelegate& GetHLODBuildEventDelegate();
+
 	ENGINE_API void UpdateHLODBuildReportHeader();
 	ENGINE_API void UpdateHLODBuildReportContent(const FString& InHLODBuildReportContent);
 
@@ -188,6 +206,10 @@ private:
 	UPROPERTY()
 	TObjectPtr<const UHLODLayer> SubActorsHLODLayer_DEPRECATED;
 #endif
+
+#if WITH_EDITOR
+	static FWorldPartitionHLODBuildEventDelegate HLODBuildEventDelegate;
+#endif
 };
 
 DEFINE_ACTORDESC_TYPE(AWorldPartitionHLOD, FHLODActorDesc);
