Added r.Velocity.TemporalResponsiveness.Supported and r.Velocity.PixelShaderMotionVectorWorldOffset.Supported to shader DDC key. Also added proper api to access r.Velocity shader cvar

#rnx
#FailedValidations Preflight

diff --git a/Engine/Source/Runtime/Engine/Private/ShaderCompiler/ShaderCompiler.cpp b/Engine/Source/Runtime/Engine/Private/ShaderCompiler/ShaderCompiler.cpp
--- a/Engine/Source/Runtime/Engine/Private/ShaderCompiler/ShaderCompiler.cpp
+++ b/Engine/Source/Runtime/Engine/Private/ShaderCompiler/ShaderCompiler.cpp
@@ -3827,12 +3827,9 @@ void GlobalBeginCompileShader(
 	{
 		if (VelocityEncodeDepth((EShaderPlatform)Target.Platform))
 		{
-			static FShaderPlatformCachedIniValue<int32> CVarTemporalResponsivenessSupported(TEXT("r.Velocity.TemporalResponsiveness.Supported"));
-			static FShaderPlatformCachedIniValue<int32> CVarPixelShaderMotionVectorWorldOffsetSupported(TEXT("r.Velocity.PixelShaderMotionVectorWorldOffset.Supported"));
-
 			SET_SHADER_DEFINE(Input.Environment, VELOCITY_ENCODE_DEPTH, 1);
-			SET_SHADER_DEFINE(Input.Environment, VELOCITY_SUPPORT_TEMPORAL_RESPONSIVENESS, CVarTemporalResponsivenessSupported.Get((EShaderPlatform)Target.Platform) !=0 ? 1 : 0);
-			SET_SHADER_DEFINE(Input.Environment, VELOCITY_SUPPORT_PS_MOTION_VECTOR_WORLD_OFFSET, CVarPixelShaderMotionVectorWorldOffsetSupported.Get((EShaderPlatform)Target.Platform) != 0 ? 1 : 0);
+			SET_SHADER_DEFINE(Input.Environment, VELOCITY_SUPPORT_TEMPORAL_RESPONSIVENESS, VelocitySupportsTemporalResponsiveness((EShaderPlatform)Target.Platform) ? 1 : 0);
+			SET_SHADER_DEFINE(Input.Environment, VELOCITY_SUPPORT_PS_MOTION_VECTOR_WORLD_OFFSET, VelocitySupportsPixelShaderMotionVectorWorldOffset((EShaderPlatform)Target.Platform) ? 1 : 0);
 		}	
 		else
 		{
diff --git a/Engine/Source/Runtime/RenderCore/Private/RenderUtils.cpp b/Engine/Source/Runtime/RenderCore/Private/RenderUtils.cpp
--- a/Engine/Source/Runtime/RenderCore/Private/RenderUtils.cpp
+++ b/Engine/Source/Runtime/RenderCore/Private/RenderUtils.cpp
@@ -1676,6 +1676,18 @@ bool VelocityEncodeHasPixelAnimation(const FStaticShaderPlatform Platform)
 	return !IsMobilePlatform(Platform);
 }
 
+bool VelocitySupportsTemporalResponsiveness(const FStaticShaderPlatform Platform)
+{
+	static FShaderPlatformCachedIniValue<int32> PerPlatformCVar(TEXT("r.Velocity.TemporalResponsiveness.Supported"));
+	return (PerPlatformCVar.Get(Platform) != 0);
+}
+
+bool VelocitySupportsPixelShaderMotionVectorWorldOffset(const FStaticShaderPlatform Platform)
+{
+	static FShaderPlatformCachedIniValue<int32> PerPlatformCVar(TEXT("r.Velocity.PixelShaderMotionVectorWorldOffset.Supported"));
+	return (PerPlatformCVar.Get(Platform) != 0);
+}
+
 bool PlatformSupportsOpenXRMotionVectors(const FStaticShaderPlatform Platform)
 {
 	if (IsVulkanMobilePlatform(Platform) && !IsMobileDeferredShadingEnabled(Platform))
diff --git a/Engine/Source/Runtime/RenderCore/Private/Shader.cpp b/Engine/Source/Runtime/RenderCore/Private/Shader.cpp
--- a/Engine/Source/Runtime/RenderCore/Private/Shader.cpp
+++ b/Engine/Source/Runtime/RenderCore/Private/Shader.cpp
@@ -2839,6 +2839,18 @@ void ShaderMapAppendKey(EShaderPlatform Platform, FShaderKeyGenerator& KeyGen)
 		KeyGen.Append(TEXT("VED"));
 	}
 
+	if (VelocitySupportsTemporalResponsiveness(Platform))
+	{
+		KeyGen.AppendSeparator();
+		KeyGen.Append(TEXT("VTR"));
+	}
+	
+	if (VelocitySupportsPixelShaderMotionVectorWorldOffset(Platform))
+	{
+		KeyGen.AppendSeparator();
+		KeyGen.Append(TEXT("VPSMVWO"));
+	}
+
 	{
 		const bool bSupportsAnisotropicMaterials = FDataDrivenShaderPlatformInfo::GetSupportsAnisotropicMaterials(Platform);
 		KeyGen.AppendSeparator();
diff --git a/Engine/Source/Runtime/RenderCore/Public/RenderUtils.h b/Engine/Source/Runtime/RenderCore/Public/RenderUtils.h
--- a/Engine/Source/Runtime/RenderCore/Public/RenderUtils.h
+++ b/Engine/Source/Runtime/RenderCore/Public/RenderUtils.h
@@ -383,6 +383,10 @@ RENDERCORE_API bool VelocityEncodeDepth(const FStaticShaderPlatform Platform);
 
 RENDERCORE_API bool VelocityEncodeHasPixelAnimation(const FStaticShaderPlatform Platform);
 
+RENDERCORE_API bool VelocitySupportsTemporalResponsiveness(const FStaticShaderPlatform Platform);
+
+RENDERCORE_API bool VelocitySupportsPixelShaderMotionVectorWorldOffset(const FStaticShaderPlatform Platform);
+
 RENDERCORE_API bool PlatformSupportsOpenXRMotionVectors(const FStaticShaderPlatform Platform);
 
 RENDERCORE_API bool VelocityIncludeStationaryPrimitives(const FStaticShaderPlatform Platform);
diff --git a/Engine/Source/Runtime/Renderer/Private/VelocityRendering.cpp b/Engine/Source/Runtime/Renderer/Private/VelocityRendering.cpp
--- a/Engine/Source/Runtime/Renderer/Private/VelocityRendering.cpp
+++ b/Engine/Source/Runtime/Renderer/Private/VelocityRendering.cpp
@@ -126,14 +126,12 @@ bool NeedVelocityDepth(EShaderPlatform ShaderPlatform)
 
 bool SupportsTemporalResponsiveness(EShaderPlatform ShaderPlatform)
 {
-	static FShaderPlatformCachedIniValue<int32> PerPlatformCVar(TEXT("r.Velocity.TemporalResponsiveness.Supported"));
-	return NeedVelocityDepth(ShaderPlatform) && (PerPlatformCVar.Get(ShaderPlatform) != 0);
+	return NeedVelocityDepth(ShaderPlatform) && VelocitySupportsTemporalResponsiveness(ShaderPlatform);
 }
 
 bool SupportsPixelShaderMotionVectorWorldOffset(EShaderPlatform ShaderPlatform)
 {
-	static FShaderPlatformCachedIniValue<int32> PerPlatformCVar(TEXT("r.Velocity.PixelShaderMotionVectorWorldOffset.Supported"));
-	return NeedVelocityDepth(ShaderPlatform) && (PerPlatformCVar.Get(ShaderPlatform) != 0);
+	return NeedVelocityDepth(ShaderPlatform) && VelocitySupportsPixelShaderMotionVectorWorldOffset(ShaderPlatform);
 }
 
 class FVelocityVS : public FMeshMaterialShader
