[Animation] Export FRawCurveTracks::EvaluateCurveData.

glTF parser uses this; need support in non-monolithic non-editor builds.

#rnx

diff --git a/Engine/Source/Runtime/Engine/Public/Animation/AnimCurveTypes.h b/Engine/Source/Runtime/Engine/Public/Animation/AnimCurveTypes.h
--- a/Engine/Source/Runtime/Engine/Public/Animation/AnimCurveTypes.h
+++ b/Engine/Source/Runtime/Engine/Public/Animation/AnimCurveTypes.h
@@ -1093,7 +1093,7 @@ struct FRawCurveTracks
 	 *
 	 * return true if curve exists, false otherwise
 	 */
-	void EvaluateCurveData( FBlendedCurve& Curves, float CurrentTime ) const;
+	ENGINE_API void EvaluateCurveData( FBlendedCurve& Curves, float CurrentTime ) const;
 
 #if WITH_EDITOR
 	/**
