Tweaks to SpringMath API:
* Make concept of 'damping' variable private, since it is easily confused with spring damping ratio
* Add documentation to spring math blueprint library
* Tweak some names to differentiate from existing more general purpose SpringInterp functions. SpringDamper -> CriticalSpringDamper

#rnx
#virtualized

diff --git a/Engine/Plugins/Animation/AnimationWarping/Source/Runtime/Private/BoneControllers/AnimNode_Steering.cpp b/Engine/Plugins/Animation/AnimationWarping/Source/Runtime/Private/BoneControllers/AnimNode_Steering.cpp
--- a/Engine/Plugins/Animation/AnimationWarping/Source/Runtime/Private/BoneControllers/AnimNode_Steering.cpp
+++ b/Engine/Plugins/Animation/AnimationWarping/Source/Runtime/Private/BoneControllers/AnimNode_Steering.cpp
@@ -148,7 +148,7 @@ void FAnimNode_Steering::EvaluateSkeletalControl_AnyThread(FComponentSpacePoseCo
 						{
 							// Apply linear correction
 							FQuat LinearCorrection = FQuat::Identity;
-							SpringMath::SimpleSpringDamperQuat(IN OUT LinearCorrection, IN OUT AngularVelocity, DeltaToTargetOrientation, SpringMath::HalfLifeToSmoothingTime(ProceduralTargetTime), DeltaSeconds);
+							SpringMath::CriticalSpringDamperQuat(IN OUT LinearCorrection, IN OUT AngularVelocity, DeltaToTargetOrientation, SpringMath::HalfLifeToSmoothingTime(ProceduralTargetTime), DeltaSeconds);
 
 							UE_VLOG_ARROW(Output.AnimInstanceProxy->GetAnimInstanceObject(), "Steering", Display,
 								RootBoneTransform.GetLocation(),
diff --git a/Engine/Plugins/Experimental/Mover/Source/Mover/Private/DefaultMovementSet/Modes/SimpleSpringWalkingMode.cpp b/Engine/Plugins/Experimental/Mover/Source/Mover/Private/DefaultMovementSet/Modes/SimpleSpringWalkingMode.cpp
--- a/Engine/Plugins/Experimental/Mover/Source/Mover/Private/DefaultMovementSet/Modes/SimpleSpringWalkingMode.cpp
+++ b/Engine/Plugins/Experimental/Mover/Source/Mover/Private/DefaultMovementSet/Modes/SimpleSpringWalkingMode.cpp
@@ -36,6 +36,6 @@ void USimpleSpringWalkingMode::GenerateWalkMove_Implementation(FMoverTickStartDa
 
 	// Angular //
 	OutFacing = StartingSyncState->GetOrientation_WorldSpace().Quaternion();
-	SpringMath::SimpleSpringDamperQuat(OutFacing, SpringState.CurrentAngularVelocity, DesiredFacing, FacingSmoothingTime, DeltaSeconds);
+	SpringMath::CriticalSpringDamperQuat(OutFacing, SpringState.CurrentAngularVelocity, DesiredFacing, FacingSmoothingTime, DeltaSeconds);
 }
 
diff --git a/Engine/Plugins/Experimental/Mover/Source/Mover/Private/DefaultMovementSet/Modes/SmoothWalkingMode.cpp b/Engine/Plugins/Experimental/Mover/Source/Mover/Private/DefaultMovementSet/Modes/SmoothWalkingMode.cpp
--- a/Engine/Plugins/Experimental/Mover/Source/Mover/Private/DefaultMovementSet/Modes/SmoothWalkingMode.cpp
+++ b/Engine/Plugins/Experimental/Mover/Source/Mover/Private/DefaultMovementSet/Modes/SmoothWalkingMode.cpp
@@ -125,7 +125,7 @@ void USmoothWalkingMode::GenerateWalkMove_Implementation(FMoverTickStartData& St
 	TrackVelocity = TrackVelocity.GetClampedToMaxSize(FMath::Max(PreviousVelocityLength, DesiredVelocity.Length()));
 
 	// Apply the smoothing to the track velocity - effectively tracking the intermediate velocity at the appropriate time in the future
-	SpringMath::SimpleSpringDamper(SpringState.SpringVelocity, SpringState.SpringAcceleration, TrackVelocity, VelocitySmoothingTime, DeltaSeconds);
+	SpringMath::CriticalSpringDamper(SpringState.SpringVelocity, SpringState.SpringAcceleration, TrackVelocity, VelocitySmoothingTime, DeltaSeconds);
 
 	// Snap the velocity to the desired velocity based on the dead-zone
 	if ((DesiredVelocity - SpringState.SpringVelocity).SquaredLength() < FMath::Square(VelocityDeadzoneThreshold))
@@ -149,14 +149,14 @@ void USmoothWalkingMode::GenerateWalkMove_Implementation(FMoverTickStartData& St
 	// For controlling the facing direction we use either a spring or a double spring
 	if (bSmoothFacingWithDoubleSpring)
 	{
-		SpringMath::SimpleSpringDamperQuat(SpringState.IntermediateFacing, SpringState.IntermediateAngularVelocity, DesiredFacing, FacingSmoothingTime / 2.0f, DeltaSeconds);
-		SpringMath::SimpleSpringDamperQuat(InOutFacing, SpringState.CurrentAngularVelocity, SpringState.IntermediateFacing, FacingSmoothingTime / 2.0f, DeltaSeconds);
+		SpringMath::CriticalSpringDamperQuat(SpringState.IntermediateFacing, SpringState.IntermediateAngularVelocity, DesiredFacing, FacingSmoothingTime / 2.0f, DeltaSeconds);
+		SpringMath::CriticalSpringDamperQuat(InOutFacing, SpringState.CurrentAngularVelocity, SpringState.IntermediateFacing, FacingSmoothingTime / 2.0f, DeltaSeconds);
 	}
 	else
 	{
 		SpringState.IntermediateFacing = DesiredFacing;
 		SpringState.IntermediateAngularVelocity = SpringState.CurrentAngularVelocity;
-		SpringMath::SimpleSpringDamperQuat(InOutFacing, SpringState.CurrentAngularVelocity, DesiredFacing, FacingSmoothingTime, DeltaSeconds);
+		SpringMath::CriticalSpringDamperQuat(InOutFacing, SpringState.CurrentAngularVelocity, DesiredFacing, FacingSmoothingTime, DeltaSeconds);
 	}
 
 	// Snap the facing to the desired facing based on the dead-zone
diff --git a/Engine/Plugins/Experimental/UAF/UAFWarping/Source/UAFWarping/Private/SteeringTrait.cpp b/Engine/Plugins/Experimental/UAF/UAFWarping/Source/UAFWarping/Private/SteeringTrait.cpp
--- a/Engine/Plugins/Experimental/UAF/UAFWarping/Source/UAFWarping/Private/SteeringTrait.cpp
+++ b/Engine/Plugins/Experimental/UAF/UAFWarping/Source/UAFWarping/Private/SteeringTrait.cpp
@@ -198,7 +198,7 @@ void FAnimNextSteeringTask::Execute(UE::UAF::FEvaluationVM& VM) const
 						{
 							// Apply linear correction
 							FQuat LinearCorrection = FQuat::Identity;
-							SpringMath::SimpleSpringDamperQuat(LinearCorrection, InstanceData->AngularVelocity, DeltaToTargetOrientation, SpringMath::HalfLifeToSmoothingTime(SharedData->ProceduralTargetTime), InstanceData->DeltaTime);
+							SpringMath::CriticalSpringDamperQuat(LinearCorrection, InstanceData->AngularVelocity, DeltaToTargetOrientation, SpringMath::HalfLifeToSmoothingTime(SharedData->ProceduralTargetTime), InstanceData->DeltaTime);
 
 #if ENABLE_ANIM_DEBUG
 							UE_VLOG_ARROW(InstanceData->HostObject, "Steering", Display,
diff --git a/Engine/Source/Runtime/Engine/Classes/Kismet/BlueprintSpringMathLibrary.h b/Engine/Source/Runtime/Engine/Classes/Kismet/BlueprintSpringMathLibrary.h
--- a/Engine/Source/Runtime/Engine/Classes/Kismet/BlueprintSpringMathLibrary.h
+++ b/Engine/Source/Runtime/Engine/Classes/Kismet/BlueprintSpringMathLibrary.h
@@ -14,82 +14,239 @@ class UBlueprintSpringMathLibrary : public UBlueprintFunctionLibrary
 public:
 	// Simple Spring //
 
+	/** Interpolates the value InOutX towards TargetX with the motion of a critically damped spring. The velocity of X is stored in InOutV.
+	 * 
+	 * @param InOutX The value to be damped
+	 * @param InOutV The speed of the value to be damped
+	 * @param TargetX The goal to damp towards
+	 * @param SmoothingTime The smoothing time to use for the spring. Longer times lead to more damped behaviour. A time of 0 will snap InOutX to TargetX
+	 * @param DeltaTime Timestep in seconds
+	 */
 	UFUNCTION(BlueprintCallable, meta=(Category="Math|Spring|Experimental"))
 	UE_EXPERIMENTAL(5.7, "SpringMath is experimental")
-	static ENGINE_API void SpringDampFloat(UPARAM(ref) float& InOutX, UPARAM(ref) float& InOutV, const float& TargetX, float DeltaTime, float SmoothingTime = 0.2f);
-
+	static ENGINE_API void CriticalSpringDampFloat(UPARAM(ref) float& InOutX, UPARAM(ref) float& InOutV, const float& TargetX, float DeltaTime, float SmoothingTime = 0.2f);
+
+	/** Interpolates the value InOutX towards TargetX with the motion of a critically damped spring. The velocity of X is stored in InOutV.
+	 * 
+	 * @param InOutX The value to be damped
+	 * @param InOutV The speed of the value to be damped
+	 * @param TargetX The goal to damp towards
+	 * @param SmoothingTime The smoothing time to use for the spring. Longer times lead to more damped behaviour. A time of 0 will snap InOutX to TargetX
+	 * @param DeltaTime Timestep in seconds
+	 */
 	UE_EXPERIMENTAL(5.7, "SpringMath is experimental")
 	UFUNCTION(BlueprintCallable, meta=(Category="Math|Spring|Experimental"))
-	static ENGINE_API void SpringDampVector(UPARAM(ref) FVector& InOutX, UPARAM(ref) FVector& InOutV, const FVector& TargetX, float DeltaTime, float SmoothingTime = 0.2f);
-
+	static ENGINE_API void CriticalSpringDampVector(UPARAM(ref) FVector& InOutX, UPARAM(ref) FVector& InOutV, const FVector& TargetX, float DeltaTime, float SmoothingTime = 0.2f);
+
+	/** Interpolates the value InOutX towards TargetX with the motion of a critically damped spring. The velocity of X is stored in InOutV.
+	 * 
+	 * @param InOutX The value to be damped
+	 * @param InOutV The speed of the value to be damped
+	 * @param TargetX The goal to damp towards
+	 * @param SmoothingTime The smoothing time to use for the spring. Longer times lead to more damped behaviour. A time of 0 will snap InOutX to TargetX
+	 * @param DeltaTime Timestep in seconds
+	 */
 	UE_EXPERIMENTAL(5.7, "SpringMath is experimental")
 	UFUNCTION(BlueprintCallable, meta=(Category="Math|Spring|Experimental"))
-	static ENGINE_API void SpringDampVector2D(UPARAM(ref) FVector2D& InOutX, UPARAM(ref) FVector2D& InOutV, const FVector2D& TargetX, float DeltaTime, float SmoothingTime = 0.2f);
-
+	static ENGINE_API void CriticalSpringDampVector2D(UPARAM(ref) FVector2D& InOutX, UPARAM(ref) FVector2D& InOutV, const FVector2D& TargetX, float DeltaTime, float SmoothingTime = 0.2f);
+
+	/** Interpolates the value InOutX towards TargetX with the motion of a critically damped spring. The velocity of X is stored in InOutV.
+	 * 
+	 * @param InOutX The value to be damped
+	 * @param InOutV The speed of the value to be damped
+	 * @param TargetX The goal to damp towards
+	 * @param SmoothingTime The smoothing time to use for the spring. Longer times lead to more damped behaviour. A time of 0 will snap InOutX to TargetX
+	 * @param DeltaTime Timestep in seconds
+	 */
 	UE_EXPERIMENTAL(5.7, "SpringMath is experimental")
 	UFUNCTION(BlueprintCallable, meta=(Category="Math|Spring|Experimental"))
-	static ENGINE_API void SpringDampQuat(UPARAM(ref) FQuat& InOutRotation, UPARAM(ref) FVector& InOutAngularVelocity, const FQuat& TargetRotation, float DeltaTime, float SmoothingTime = 0.2f);
-
+	static ENGINE_API void CriticalSpringDampQuat(UPARAM(ref) FQuat& InOutRotation, UPARAM(ref) FVector& InOutAngularVelocity, const FQuat& TargetRotation, float DeltaTime, float SmoothingTime = 0.2f);
+
+	/** Interpolates the value InOutX towards TargetX with the motion of a critically damped spring. The velocity of X is stored in InOutV.
+	 * 
+	 * @param InOutX The value to be damped
+	 * @param InOutV The speed of the value to be damped
+	 * @param TargetX The goal to damp towards
+	 * @param SmoothingTime The smoothing time to use for the spring. Longer times lead to more damped behaviour. A time of 0 will snap InOutX to TargetX
+	 * @param DeltaTime Timestep in seconds
+	 */
 	UE_EXPERIMENTAL(5.7, "SpringMath is experimental")
 	UFUNCTION(BlueprintCallable, meta=(Category="Math|Spring|Experimental"))
-	static ENGINE_API void SpringDampRotator(UPARAM(ref) FRotator& InOutRotation, UPARAM(ref) FVector& InOutAngularVelocity, const FRotator& TargetRotation, float DeltaTime, float SmoothingTime = 0.2f);
-
-	// Velocity Spring //
-
+	static ENGINE_API void CriticalSpringDampRotator(UPARAM(ref) FRotator& InOutRotation, UPARAM(ref) FVector& InOutAngularVelocity, const FRotator& TargetRotation, float DeltaTime, float SmoothingTime = 0.2f);
+
+	/** A velocity spring will damp towards a target that follows a fixed linear target velocity, allowing control of the interpolation speed
+	 * while still giving a smoothed behavior. A SmoothingTime of 0 will give a linear interpolation between X and TargetX
+	 * 
+	 * @param InOutX The value to be damped
+	 * @param InOutV The velocity of the value to be damped
+	 * @param InOutVi The intermediate target velocity of the value to be damped
+	 * @param TargetX The target value of X to damp towards
+	 * @param MaxSpeed The desired speed to achieve while damping towards X
+	 * @param SmoothingTime The smoothing time to use while damping towards X. Higher values will give more smoothed behaviour. A value of 0 will give a linear interpolation of X to Target
+	 * @param DeltaTime The timestep in seconds
+	 */
 	UFUNCTION(BlueprintCallable, meta=(Category="Math|Spring|Experimental"))
 	UE_EXPERIMENTAL(5.7, "SpringMath is experimental")
 	static ENGINE_API void VelocitySpringDampFloat(UPARAM(ref) float& InOutX, UPARAM(ref) float& InOutV, UPARAM(ref) float& InOutVi, float TargetX, float MaxSpeed, float DeltaTime, float SmoothingTime = 0.2f);
 
+	/** A velocity spring will damp towards a target that follows a fixed linear target velocity, allowing control of the interpolation speed
+	 * while still giving a smoothed behavior. A SmoothingTime of 0 will give a linear interpolation between X and TargetX
+	 * 
+	 * @param InOutX The value to be damped
+	 * @param InOutV The velocity of the value to be damped
+	 * @param InOutVi The intermediate target velocity of the value to be damped
+	 * @param TargetX The target value of X to damp towards
+	 * @param MaxSpeed The desired speed to achieve while damping towards X
+	 * @param SmoothingTime The smoothing time to use while damping towards X. Higher values will give more smoothed behaviour. A value of 0 will give a linear interpolation of X to Target
+	 * @param DeltaTime The timestep in seconds
+	 */
 	UE_EXPERIMENTAL(5.7, "SpringMath is experimental")
 	UFUNCTION(BlueprintCallable, meta=(Category="Math|Spring|Experimental"))
 	static ENGINE_API void VelocitySpringDampVector(UPARAM(ref) FVector& InOutX, UPARAM(ref) FVector& InOutV, UPARAM(ref) FVector& InOutVi, const FVector& TargetX, float MaxSpeed, float DeltaTime,  float SmoothingTime = 0.2f);
 
+	/** A velocity spring will damp towards a target that follows a fixed linear target velocity, allowing control of the interpolation speed
+	 * while still giving a smoothed behavior. A SmoothingTime of 0 will give a linear interpolation between X and TargetX
+	 * 
+	 * @param InOutX The value to be damped
+	 * @param InOutV The velocity of the value to be damped
+	 * @param InOutVi The intermediate target velocity of the value to be damped
+	 * @param TargetX The target value of X to damp towards
+	 * @param MaxSpeed The desired speed to achieve while damping towards X
+	 * @param SmoothingTime The smoothing time to use while damping towards X. Higher values will give more smoothed behaviour. A value of 0 will give a linear interpolation of X to Target
+	 * @param DeltaTime The timestep in seconds
+	 */
 	UE_EXPERIMENTAL(5.7, "SpringMath is experimental")
 	UFUNCTION(BlueprintCallable, meta=(Category="Math|Spring|Experimental"))
 	static ENGINE_API void VelocitySpringDampVector2D(UPARAM(ref) FVector2D& InOutX, UPARAM(ref) FVector2D& InOutV, UPARAM(ref) FVector2D& InOutVi, const FVector2D& TargetX, float MaxSpeed, float DeltaTime,  float SmoothingTime = 0.2f);
-	
-	// Damper, aka ExponentialSmoothingApprox // 
 
+	// Damper //
+
+	/**
+	 * Smooths a value using exponential damping towards a target.
+	 * 
+	 * 
+	 * @param  Value The value to be smoothed
+	 * @param  Target The target to smooth towards
+	 * @param  DeltaTime Time interval
+	 * @param  SmoothingTime Timescale over which to smooth. Larger values result in more smoothed behaviour. Can be zero.
+	 */
 	UE_EXPERIMENTAL(5.7, "SpringMath is experimental")
 	UFUNCTION(BlueprintPure, meta=(Category="Math|Spring|Experimental"))
 	static ENGINE_API float DampFloat(float Value, float Target, float DeltaTime, float SmoothingTime = 0.2f);
 
+	/**
+	 * Smooths a value using exponential damping towards a target.
+	 * 
+	 * 
+	 * @param  Value The value to be smoothed
+	 * @param  Target The target to smooth towards
+	 * @param  DeltaTime Time interval
+	 * @param  SmoothingTime Timescale over which to smooth. Larger values result in more smoothed behaviour. Can be zero.
+	 */
 	UE_EXPERIMENTAL(5.7, "SpringMath is experimental")
 	UFUNCTION(BlueprintPure, meta=(Category="Math|Spring|Experimental"))
 	static ENGINE_API FVector DampVector(const FVector& Value, const FVector& Target, float DeltaTime, float SmoothingTime = 0.2f);
 
+	/**
+	 * Smooths a value using exponential damping towards a target.
+	 * 
+	 * 
+	 * @param  Value The value to be smoothed
+	 * @param  Target The target to smooth towards
+	 * @param  DeltaTime Time interval
+	 * @param  SmoothingTime Timescale over which to smooth. Larger values result in more smoothed behaviour. Can be zero.
+	 */
 	UE_EXPERIMENTAL(5.7, "SpringMath is experimental")
 	UFUNCTION(BlueprintPure, meta=(Category="Math|Spring|Experimental"))
 	static ENGINE_API FVector2D DampVector2D(const FVector2D& Value, const FVector2D& Target, float DeltaTime, float SmoothingTime = 0.2f);
 
+	/**
+	 * Smooths a value using exponential damping towards a target.
+	 * 
+	 * 
+	 * @param  Value The value to be smoothed
+	 * @param  Target The target to smooth towards
+	 * @param  DeltaTime Time interval
+	 * @param  SmoothingTime Timescale over which to smooth. Larger values result in more smoothed behaviour. Can be zero.
+	 */
 	UE_EXPERIMENTAL(5.7, "SpringMath is experimental")
 	UFUNCTION(BlueprintPure, meta=(Category="Math|Spring|Experimental"))
 	static ENGINE_API FQuat DampQuat(const FQuat& Value, const FQuat& Target, float DeltaTime, float SmoothingTime = 0.2f);
 
+	/**
+	 * Smooths a value using exponential damping towards a target.
+	 * 
+	 * 
+	 * @param  Value The value to be smoothed
+	 * @param  Target The target to smooth towards
+	 * @param  DeltaTime Time interval
+	 * @param  SmoothingTime Timescale over which to smooth. Larger values result in more smoothed behaviour. Can be zero.
+	 */
 	UE_EXPERIMENTAL(5.7, "SpringMath is experimental")
 	UFUNCTION(BlueprintPure, meta=(Category="Math|Spring|Experimental"))
 	static ENGINE_API FRotator DampRotator(const FRotator& Value, const FRotator& Target, float DeltaTime, float SmoothingTime = 0.2f);
 
-	// Spring Character
+	/** Update the position of a character given a target velocity using a simple damped spring
+	 * 
+	 * @param InOutPosition The position of the character
+	 * @param InOutVelocity The velocity of the character. Needs to be stored and persisted by the caller. Usually initialized to zero and not modified by the caller.
+	 * @param InOutAcceleration The acceleration of the character. Needs to be stored and persisted by the caller. Usually initialized to zero and not modified by the caller.
+	 * @param TargetVelocity The target velocity of the character.
+	 * @param DeltaTime The delta time to tick the character
+	 * @param SmoothingTime The time over which to smooth velocity. It takes roughly the smoothing time in order for the character to reach the target velocity.
+	 */
 	UE_EXPERIMENTAL(5.7, "SpringMath is experimental")
 	UFUNCTION(BlueprintCallable, meta=(Category="Math|Spring|Experimental"))
 	static ENGINE_API void SpringCharacterUpdate(UPARAM(ref) FVector& InOutPosition, UPARAM(ref) FVector& InOutVelocity, UPARAM(ref) FVector& InOutAcceleration, const FVector& TargetVelocity, float DeltaTime, float SmoothingTime = 0.2f);
 
-	// Velocity Spring Character
+	/** Update a position representing a character given a target velocity using a velocity spring.
+	 * A velocity spring tracks an intermediate velocity which moves at a maximum acceleration linearly towards a target.
+	 * This means unlike the "SpringCharacterUpdate", it will take longer to reach a target velocity that is further away from the current velocity.
+	 * 
+	 * @param InOutPosition The position of the character
+	 * @param InOutVelocity The velocity of the character. Needs to be stored and persisted by the caller. Usually initialized to zero and not modified by the caller.
+	 * @param InOutVelocityIntermediate The intermediate velocity of the character. Needs to be stored and persisted by the caller. Usually initialized to zero and not modified by the caller.
+	 * @param InOutAcceleration The acceleration of the character. Needs to be stored and persisted by the caller. Usually initialized to zero and not modified by the caller.
+	 * @param TargetVelocity The target velocity of the character.
+	 * @param DeltaTime The delta time to tick the character
+	 * @param SmoothingTime The time over which to smooth velocity. It takes roughly the smoothing time in order for the character to reach the target velocity.
+	 * @param MaxAcceleration Puts a limit on the maximum acceleration that the intermediate velocity can do each frame. If MaxAccel is very large, the behaviour wil lbe the same as SpringCharacterUpdate
+	 */
 	UE_EXPERIMENTAL(5.7, "SpringMath is experimental")
 	UFUNCTION(BlueprintCallable, meta=(Category="Math|Spring|Experimental"))
 	static ENGINE_API void VelocitySpringCharacterUpdate(UPARAM(ref) FVector& InOutPosition, UPARAM(ref) FVector& InOutVelocity, UPARAM(ref) FVector& InOutVelocityIntermediate, UPARAM(ref) FVector& InOutAcceleration, const FVector& TargetVelocity, float DeltaTime, float SmoothingTime = 0.2f, float MaxAcceleration = 1.0f);
 
 	// Conversion Methods
+
+	/** Convert from smoothing time to spring strength.
+	* 
+	* @param SmoothingTime The smoothing time of the spring in seconds. It is the time by which the output lags the input when critically damped, when the input is changing at a constant rate.
+	* @return The spring strength. This corresponds to the undamped frequency of the spring in hz.
+	*/
 	UFUNCTION(BlueprintPure, meta=(Category="Math|Spring|Experimental"))
 	static ENGINE_API float ConvertSmoothingTimeToStrength(float SmoothingTime);
 
+	/** Convert from spring strength to smoothing time.
+ 	* 
+ 	* @param Strength The spring strength. This corresponds to the undamped frequency of the spring in hz.
+ 	* @return The smoothing time of the spring in seconds. It is the time by which the output lags the input when critically damped, when the input is changing at a constant rate.
+ 	*/
 	UFUNCTION(BlueprintPure, meta=(Category="Math|Spring|Experimental"))
 	static ENGINE_API float ConvertStrengthToSmoothingTime(float Strength);
 
+	/** Convert a halflife to a smoothing time
+ 	* 
+ 	* @param HalfLife The half life of the spring. How long it takes the value to get halfway towards the target.
+ 	* @return The smoothing time of the spring in seconds. It is the time by which the output lags the input when critically damped, when the input is changing at a constant rate.
+ 	*/
 	UFUNCTION(BlueprintPure, meta=(Category="Math|Spring|Experimental"))
 	static ENGINE_API float ConvertHalfLifeToSmoothingTime(float HalfLife);
 
+	/** Convert a smoothing time to a half life
+ 	* 
+ 	* @param SmoothingTime The smoothing time of the spring in seconds. It is the time by which the output lags the input when critically damped, when the input is changing at a constant rate. 
+ 	* @return The half life of the spring. How long it takes the value to get halfway towards the target. 
+ 	*/
 	UFUNCTION(BlueprintPure, meta=(Category="Math|Spring|Experimental"))
 	static ENGINE_API float ConvertSmoothingTimeToHalfLife(float SmoothingTime);
 };
diff --git a/Engine/Source/Runtime/Engine/Private/BlueprintSpringMathLibrary.cpp b/Engine/Source/Runtime/Engine/Private/BlueprintSpringMathLibrary.cpp
--- a/Engine/Source/Runtime/Engine/Private/BlueprintSpringMathLibrary.cpp
+++ b/Engine/Source/Runtime/Engine/Private/BlueprintSpringMathLibrary.cpp
@@ -6,32 +6,32 @@
 
 DEFINE_LOG_CATEGORY_STATIC(LogBlueprintSpingMathLibrary, Log, All);
 
-void UBlueprintSpringMathLibrary::SpringDampVector(FVector& InOutX, FVector& InOutV, const FVector& TargetX, float DeltaTime, float SmoothingTime)
+void UBlueprintSpringMathLibrary::CriticalSpringDampVector(FVector& InOutX, FVector& InOutV, const FVector& TargetX, float DeltaTime, float SmoothingTime)
 {
-	SpringMath::SimpleSpringDamper(InOutX, InOutV, TargetX, SmoothingTime, DeltaTime);
+	SpringMath::CriticalSpringDamper(InOutX, InOutV, TargetX, SmoothingTime, DeltaTime);
 }
 
-void UBlueprintSpringMathLibrary::SpringDampVector2D(FVector2D& InOutX, FVector2D& InOutV, const FVector2D& TargetX, float DeltaTime, float SmoothingTime)
+void UBlueprintSpringMathLibrary::CriticalSpringDampVector2D(FVector2D& InOutX, FVector2D& InOutV, const FVector2D& TargetX, float DeltaTime, float SmoothingTime)
 {
-	SpringMath::SimpleSpringDamper(InOutX, InOutV, TargetX, SmoothingTime, DeltaTime);
+	SpringMath::CriticalSpringDamper(InOutX, InOutV, TargetX, SmoothingTime, DeltaTime);
 }
 
-void UBlueprintSpringMathLibrary::SpringDampFloat(float& InOutX, float& InOutV, const float& TargetX, float DeltaTime, float SmoothingTime)
+void UBlueprintSpringMathLibrary::CriticalSpringDampFloat(float& InOutX, float& InOutV, const float& TargetX, float DeltaTime, float SmoothingTime)
 {
-	SpringMath::SimpleSpringDamper(InOutX, InOutV, TargetX, SmoothingTime, DeltaTime);
+	SpringMath::CriticalSpringDamper(InOutX, InOutV, TargetX, SmoothingTime, DeltaTime);
 }
 
-void UBlueprintSpringMathLibrary::SpringDampQuat(FQuat& InOutRotation, FVector& InOutAngularVelocity, const FQuat& TargetRotation,
+void UBlueprintSpringMathLibrary::CriticalSpringDampQuat(FQuat& InOutRotation, FVector& InOutAngularVelocity, const FQuat& TargetRotation,
 	float DeltaTime, float SmoothingTime)
 {
-	SpringMath::SimpleSpringDamperQuat(InOutRotation, InOutAngularVelocity, TargetRotation, SmoothingTime, DeltaTime);
+	SpringMath::CriticalSpringDamperQuat(InOutRotation, InOutAngularVelocity, TargetRotation, SmoothingTime, DeltaTime);
 }
 
-void UBlueprintSpringMathLibrary::SpringDampRotator(FRotator& InOutRotation, FVector& InOutAngularVelocity, const FRotator& TargetRotation,
+void UBlueprintSpringMathLibrary::CriticalSpringDampRotator(FRotator& InOutRotation, FVector& InOutAngularVelocity, const FRotator& TargetRotation,
 	float DeltaTime, float SmoothingTime)
 {
 	FQuat InOutRotationQuat = InOutRotation.Quaternion();
-	SpringMath::SimpleSpringDamperQuat(InOutRotationQuat, InOutAngularVelocity, TargetRotation.Quaternion(), SmoothingTime, DeltaTime);
+	SpringMath::CriticalSpringDamperQuat(InOutRotationQuat, InOutAngularVelocity, TargetRotation.Quaternion(), SmoothingTime, DeltaTime);
 	InOutRotation = InOutRotationQuat.Rotator();
 }
 
diff --git a/Engine/Source/Runtime/Engine/Public/Animation/SpringMath.h b/Engine/Source/Runtime/Engine/Public/Animation/SpringMath.h
--- a/Engine/Source/Runtime/Engine/Public/Animation/SpringMath.h
+++ b/Engine/Source/Runtime/Engine/Public/Animation/SpringMath.h
@@ -9,71 +9,44 @@
 // Collection of useful spring methods which can be used for damping, simulating characters etc.
 // Reference https://theorangeduck.com/page/spring-roll-call
 
-namespace SpringMath
+struct SpringMath
 {
-
-	/** Convert a halflife to a damping ratio
-	 * 
-	 * @param HalfLife The half life of the spring. How long it takes the value to get halfway towards the target.
-	 * @return The damping ratio
-	 */
-	UE_EXPERIMENTAL(5.7, "SpringMath is experimental")
-	constexpr float HalfLifeToDampingRatio(float HalfLife)
-	{
-		return (4.0f * UE_LN2) / FMath::Max(HalfLife, UE_KINDA_SMALL_NUMBER);
-	}
-
-	/** Convert a damping ratio to a halflife
-	 * 
-	 * @param DampingRatio
-	 * @return The half life of the spring. How long it takes the value to get halfway towards the target.
-	 */
-	UE_EXPERIMENTAL(5.7, "SpringMath is experimental")
-	constexpr float DampingRatioToHalfLife(float DampingRatio)
+private:
+	static constexpr float SmoothingTimeToDamping(float SmoothingTime)
 	{
-		return DampingRatio / (4.0f * UE_LN2);
+		return 4.0f / FMath::Max(SmoothingTime, UE_KINDA_SMALL_NUMBER);
 	}
-
+	
+public:
 	/** Convert a smoothing time to a half life
 	 * 
 	 * @param SmoothingTime The smoothing time of the spring in seconds. It is the time by which the output lags the input when critically damped, when the input is changing at a constant rate. 
 	 * @return The half life of the spring. How long it takes the value to get halfway towards the target. 
 	 */
 	UE_EXPERIMENTAL(5.7, "SpringMath is experimental")
-	constexpr float SmoothingTimeToHalfLife(float SmoothingTime)
+	static constexpr float SmoothingTimeToHalfLife(float SmoothingTime)
 	{
 		return SmoothingTime * UE_LN2;
 	}
 
 	/** Convert a halflife to a smoothing time
 	 * 
-	 * @param HalfLife 
-	 * @return 
+	 * @param HalfLife The half life of the spring. How long it takes the value to get halfway towards the target.
+	 * @return The smoothing time of the spring in seconds. It is the time by which the output lags the input when critically damped, when the input is changing at a constant rate.
 	 */
 	UE_EXPERIMENTAL(5.7, "SpringMath is experimental")
-	constexpr float HalfLifeToSmoothingTime(float HalfLife)
+	static constexpr float HalfLifeToSmoothingTime(float HalfLife)
 	{
 		return HalfLife / UE_LN2;
 	}
 
-	/** Convert from smoothing time to a damping ratio.
-	 * 
-	 * @param SmoothingTime 
-	 * @return Damping Ratio
-	 */
-	UE_EXPERIMENTAL(5.7, "SpringMath is experimental")
-	constexpr float SmoothingTimeToDampingRatio(float SmoothingTime)
-	{
-		return 4.0f / FMath::Max(SmoothingTime, UE_KINDA_SMALL_NUMBER);
-	}
-
 	/** Convert from smoothing time to spring strength.
 	  * 
 	 * @param SmoothingTime The smoothing time of the spring in seconds. It is the time by which the output lags the input when critically damped, when the input is changing at a constant rate.
 	 * @return The spring strength. This corresponds to the undamped frequency of the spring in hz.
 	 */
 	UE_EXPERIMENTAL(5.7, "SpringMath is experimental")
-	constexpr float SmoothingTimeToStrength(float SmoothingTime)
+	static constexpr float SmoothingTimeToStrength(float SmoothingTime)
 	{
 		return 2.0f / FMath::Max(SmoothingTime, UE_KINDA_SMALL_NUMBER);
 	}
@@ -84,12 +57,14 @@ namespace SpringMath
 	 * @return The smoothing time of the spring in seconds. It is the time by which the output lags the input when critically damped, when the input is changing at a constant rate.
 	 */
 	UE_EXPERIMENTAL(5.7, "SpringMath is experimental")
-	constexpr float StrengthToSmoothingTime(float Strength)
+
+	static constexpr float StrengthToSmoothingTime(float Strength)
 	{
 		return 2.0f / FMath::Max(Strength, UE_KINDA_SMALL_NUMBER);
 	}
 
-	/** Simplified version of FMath::CriticallyDampedSmoothing where v_goal is assumed to be 0
+	/** Simplified version of FMath::CriticallyDampedSmoothing where v_goal is assumed to be 0. This interpolates the value InOutX towards TargetX
+	 * with the motion of a critically damped spring. The velocity of X is stored in InOutV.
 	 * 
 	 * @tparam T The type to be damped 
 	 * @param InOutX The value to be damped
@@ -100,14 +75,14 @@ namespace SpringMath
 	 */
 	template <typename T>
 	UE_EXPERIMENTAL(5.7, "SpringMath is experimental")
-	void SimpleSpringDamper(
+	static void CriticalSpringDamper(
 		T& InOutX,
 		T& InOutV,
 		T TargetX,
 		float SmoothingTime,
 		float DeltaTime)
 	{
-		float Y = SmoothingTimeToDampingRatio(SmoothingTime) / 2.0f;
+		float Y = SmoothingTimeToDamping(SmoothingTime) / 2.0f;
 		T J0 = InOutX - TargetX;
 		T J1 = InOutV + J0 * Y;
 		float EyDt = FMath::InvExpApprox(Y * DeltaTime);
@@ -116,7 +91,7 @@ namespace SpringMath
 		InOutV = EyDt * (InOutV - J1 * Y * DeltaTime);
 	}
 
-	/** Specialized quaternion version, uses FVector for angular velocity
+	/** Specialized quaternion version of CriticalSpringDamper, uses FVector for angular velocity
 	 * 
 	 * @param InOutRotation The value to be damped
 	 * @param InOutAngularVelocityRadians The angular velocity of the rotation in radians
@@ -125,7 +100,7 @@ namespace SpringMath
 	 * @param DeltaTime 
 	 */
 	UE_EXPERIMENTAL(5.7, "SpringMath is experimental")
-	inline void SimpleSpringDamperQuat(
+	static void CriticalSpringDamperQuat(
 		FQuat& InOutRotation,
 		FVector& InOutAngularVelocityRadians,
 		const FQuat& TargetRotation,
@@ -133,7 +108,7 @@ namespace SpringMath
 		float DeltaTime)
 	{
 		PRAGMA_DISABLE_EXPERIMENTAL_WARNINGS
-		float Y = SmoothingTimeToDampingRatio(SmoothingTime) / 2.0f;
+		float Y = SmoothingTimeToDamping(SmoothingTime) / 2.0f;
 		PRAGMA_ENABLE_EXPERIMENTAL_WARNINGS
 
 		FQuat Diff = InOutRotation * TargetRotation.Inverse();
@@ -149,9 +124,8 @@ namespace SpringMath
 	}
 
 	/** A velocity spring will damp towards a target that follows a fixed linear target velocity, allowing control of the interpolation speed
-	 * while still giving a smoothed beheaviour. A SmoothingTime of 0 will give a linear interpolation between X and TargetX
+	 * while still giving a smoothed behavior. A SmoothingTime of 0 will give a linear interpolation between X and TargetX
 	 * 
-	 * @tparam TFloat The type of floating point to damp (e.g. float or double)
 	 * @param InOutX The value to be damped
 	 * @param InOutV The velocity of the value to be damped
 	 * @param InOutXi The intermediate target of the value to be damped
@@ -162,7 +136,7 @@ namespace SpringMath
 	 */
 	template <typename TFloat>
 	UE_EXPERIMENTAL(5.7, "SpringMath is experimental")
-	void VelocitySpringDamperF(
+	static void VelocitySpringDamperF(
 		TFloat& InOutX,
 		TFloat& InOutV,
 		TFloat& InOutXi,
@@ -184,7 +158,7 @@ namespace SpringMath
 			? InOutXi + XDiff * TGoalFuture
 			: TargetX;
 
-		SimpleSpringDamper(InOutX, InOutV, XGoalFuture, SmoothingTime, DeltaTime);
+		CriticalSpringDamper(InOutX, InOutV, XGoalFuture, SmoothingTime, DeltaTime);
 
 		InOutXi = FMath::Abs(TargetX - InOutXi) > DeltaTime * MaxSpeed
 			? InOutXi + XDiff * DeltaTime
@@ -205,7 +179,7 @@ namespace SpringMath
 	 */
 	template <typename TVector>
 	UE_EXPERIMENTAL(5.7, "SpringMath is experimental")
-	void VelocitySpringDamper(
+	static void VelocitySpringDamper(
 		TVector& InOutX,
 		TVector& InOutV,
 		TVector& InOutXi,
@@ -225,7 +199,7 @@ namespace SpringMath
 			? InOutXi + (XDiffDir * MaxSpeed) * TGoalFuture
 			: TargetX;
 
-		SimpleSpringDamper(InOutX, InOutV, XGoalFuture, SmoothingTime, DeltaTime);
+		CriticalSpringDamper(InOutX, InOutV, XGoalFuture, SmoothingTime, DeltaTime);
 
 		InOutXi = XDiffLength > DeltaTime * MaxSpeed
 			? InOutXi + XDiffDir * MaxSpeed * DeltaTime
@@ -246,7 +220,7 @@ namespace SpringMath
 	 */
 	template <typename TVector>
 	UE_EXPERIMENTAL(5.7, "SpringMath is experimental")
-	void SpringCharacterUpdate(
+	static void SpringCharacterUpdate(
 		TVector& InOutPosition,
 		TVector& InOutVelocity,
 		TVector& InOutAcceleration,
@@ -256,7 +230,7 @@ namespace SpringMath
 		float VDeadzone = 1e-2,
 		float ADeadzone = 1e-4)
 	{
-		float Y = SmoothingTimeToDampingRatio(SmoothingTime) / 2.0f;
+		float Y = SmoothingTimeToDamping(SmoothingTime) / 2.0f;
 		TVector J0 = InOutVelocity - TargetVelocity;
 		TVector J1 = InOutAcceleration + J0 * Y;
 		float EyDt = FMath::InvExpApprox(Y * DeltaTime);
@@ -296,7 +270,7 @@ namespace SpringMath
 	 */
 	template <typename TVector>
 	UE_EXPERIMENTAL(5.7, "SpringMath is experimental")
-	void SpringCharacterPredict(
+	static void SpringCharacterPredict(
 		TArrayView<TVector> OutPredictedPositions,
 		TArrayView<TVector> OutPredictedVelocities,
 		TArrayView<TVector> OutPredictedAccelerations,
@@ -349,7 +323,7 @@ namespace SpringMath
 	 */
 	template <typename TVector>
 	UE_EXPERIMENTAL(5.7, "SpringMath is experimental")
-	void VelocitySpringCharacterUpdate(
+	static void VelocitySpringCharacterUpdate(
 		TVector& InOutPosition,
 		TVector& InOutVelocity,
 		TVector& InOutVelocityIntermediate,
@@ -372,7 +346,7 @@ namespace SpringMath
 			? InOutVelocityIntermediate + (VDiffDir * MaxAcceleration) * TGoalFuture
 			: TargetVelocity;
 
-		float Y = SmoothingTimeToDampingRatio(SmoothingTime) / 2.0f;
+		float Y = SmoothingTimeToDamping(SmoothingTime) / 2.0f;
 		TVector J0 = InOutVelocity - MaxVFuture;
 		TVector J1 = InOutAcceleration + J0 * Y;
 		float EyDt = FMath::InvExpApprox(Y * DeltaTime);
@@ -418,7 +392,7 @@ namespace SpringMath
 	 */
 	template <typename TVector>
 	UE_EXPERIMENTAL(5.7, "SpringMath is experimental")
-	void VelocitySpringCharacterPredict(
+	static void VelocitySpringCharacterPredict(
 		TArrayView<TVector> OutPredictedPositions,
 		TArrayView<TVector> OutPredictedVelocities,
 		TArrayView<TVector> OutPredictedIntermediateVelocities,
@@ -464,7 +438,7 @@ namespace SpringMath
 		}
 	}
 
-	/** Prediction of SimpleSpringDamperQuat
+	/** Prediction of CriticalSpringDamperQuat
 	 * 
 	 * @param OutPredictedRotations ArrayView of output buffer to put the predicted rotations. ArrayView should be the same size as PredictCount
 	 * @param OutPredictedAngularVelocities ArrayView of output buffer to put the predicted angular velocities. ArrayView should be the same size as PredictCount
@@ -476,7 +450,7 @@ namespace SpringMath
 	 * @param SecondsPerPredictionStep How many seconds per prediction step
 	 */
 	UE_EXPERIMENTAL(5.7, "SpringMath is experimental")
-	inline void SimpleSpringDamperQuatPredict(TArrayView<FQuat> OutPredictedRotations, TArrayView<FVector> OutPredictedAngularVelocities,
+	static void CriticalSpringDamperQuatPredict(TArrayView<FQuat> OutPredictedRotations, TArrayView<FVector> OutPredictedAngularVelocities,
 	                                          int32 PredictCount, const FQuat& CurrentRotation, const FVector& CurrentAngularVelocity,
 	                                          const FQuat& TargetRotation, float SmoothingTime, float SecondsPerPredictionStep)
 	{
@@ -493,7 +467,7 @@ namespace SpringMath
 		{
 			const float PredictTime = (i + 1) * SecondsPerPredictionStep; // Note i+1 since we want index 0 to be the first prediction step
 			PRAGMA_DISABLE_EXPERIMENTAL_WARNINGS
-			SimpleSpringDamperQuat(OutPredictedRotations[i], OutPredictedAngularVelocities[i], TargetRotation, SmoothingTime, PredictTime);
+			CriticalSpringDamperQuat(OutPredictedRotations[i], OutPredictedAngularVelocities[i], TargetRotation, SmoothingTime, PredictTime);
 			PRAGMA_ENABLE_EXPERIMENTAL_WARNINGS
 		}
 	}
@@ -518,4 +492,4 @@ namespace SpringMath
 			InOutRotation = InTargetRotation;
 		}
 	}
-}
+};
