Splines: Fix incorrect assumptions about spline input keys when sampling splines for rendering.

diff --git a/Engine/Source/Editor/ComponentVisualizers/Private/SplineComponentVisualizer.cpp b/Engine/Source/Editor/ComponentVisualizers/Private/SplineComponentVisualizer.cpp
--- a/Engine/Source/Editor/ComponentVisualizers/Private/SplineComponentVisualizer.cpp
+++ b/Engine/Source/Editor/ComponentVisualizers/Private/SplineComponentVisualizer.cpp
@@ -824,6 +824,8 @@ void FSplineComponentVisualizer::UpdatePDICache(const USplineComponent* SplineCo
 		// If not the first keypoint, draw a line to the previous keypoint.
 		if (KeyIdx > 0)
 		{
+			UE::Geometry::FInterval1f SegmentParameterRange(SplineComp->GetInputKeyValueAtSplinePoint(KeyIdx - 1), SplineComp->GetInputKeyValueAtSplinePoint(KeyIdx));
+			
 			// how many?
 			LineBatch.Elements.Reset();
 			
@@ -883,8 +885,8 @@ void FSplineComponentVisualizer::UpdatePDICache(const USplineComponent* SplineCo
 
 				for (int32 StepIdx = 1; StepIdx <= NumStepsPerSegment; StepIdx++)
 				{
-					const float StepRatio = static_cast<float>(StepIdx) / static_cast<float>(NumStepsPerSegment);
-					const float Key = static_cast<float>(EndIdx) + StepRatio;
+					const float StepRatio = static_cast<float>(StepIdx) / NumStepsPerSegment;
+					const float Key = SegmentParameterRange.Interpolate(StepRatio);
 					const FVector NewPos = SplineComp->GetLocationAtSplineInputKey(Key, ESplineCoordinateSpace::World);
 					const FVector NewRightVector = SplineComp->GetRightVectorAtSplineInputKey(Key, ESplineCoordinateSpace::World);
 					const FVector NewScale = SplineComp->GetScaleAtSplineInputKey(Key) * DefaultScale;
diff --git a/Engine/Source/Runtime/Engine/Private/Components/SplineComponent.cpp b/Engine/Source/Runtime/Engine/Private/Components/SplineComponent.cpp
--- a/Engine/Source/Runtime/Engine/Private/Components/SplineComponent.cpp
+++ b/Engine/Source/Runtime/Engine/Private/Components/SplineComponent.cpp
@@ -1450,11 +1450,28 @@ float USplineComponent::GetInputKeyValueAtSplinePoint(int32 PointIndex) const
 		return 0.f;
 	}
 
-	PointIndex = GetClampedIndex(PointIndex);
-	
-	return UE::SplineComponent::ShouldUseSplineCurves()
-		? WarninglessSplineCurves().Position.Points[PointIndex].InVal
-		: Spline.GetParameterAtIndex(PointIndex);
+	// Special case if we are closed and PointIndex refers to the closing point. For this function only we
+	// allow users to treat it as a separate point in order to expose the parameterization of the closing segment.
+
+	if (bClosedLoop && PointIndex == GetNumberOfSplinePoints())
+	{
+		if (UE::SplineComponent::ShouldUseSplineCurves())
+		{
+			return WarninglessSplineCurves().Position.Points[PointIndex - 1].InVal + WarninglessSplineCurves().Position.LoopKeyOffset;
+		}
+		else
+		{
+			return Spline.GetParameterAtIndex(PointIndex);
+		}
+	}
+	else
+	{
+		PointIndex = GetClampedIndex(PointIndex);
+
+		return UE::SplineComponent::ShouldUseSplineCurves()
+			? WarninglessSplineCurves().Position.Points[PointIndex].InVal
+			: Spline.GetParameterAtIndex(PointIndex);
+	}
 }
 
 FSplinePoint USplineComponent::GetSplinePointAt(int32 PointIndex, ESplineCoordinateSpace::Type CoordinateSpace) const
@@ -2543,8 +2560,6 @@ FPrimitiveSceneProxy* USplineComponent::CreateSceneProxy()
 			PointBatchData.NumPrimitives = NumPoints * 2;
 			PointBatchData.MaterialProxy = InComponent->PointMaterialLifetimePtr->GetRenderProxy();
 			
-			FVector OldKeyPos(0);
-			
 			auto AppendLineToBuffers = [this, &LineVertexBufferIndex, &LineIndexBufferIndex](FVector Start, FVector End, FColor Color, float Thickness, float DepthBias = 0.f) -> void
 			{
 				const FVector LineDirection = (End - Start).GetSafeNormal();
@@ -2628,7 +2643,8 @@ FPrimitiveSceneProxy* USplineComponent::CreateSceneProxy()
 			
 			for (int32 KeyIdx = 0; KeyIdx < NumSegments + 1; KeyIdx++)
 			{
-				const FVector NewKeyPos = SplineInfo.Eval(static_cast<float>(KeyIdx), FVector(0));
+				const FInterpCurvePoint<FVector> NewKey = SplineInfo.Points[KeyIdx];
+				const FVector NewKeyPos = NewKey.OutVal;
 
 				// Draw the keypoint
 				if (KeyIdx < NumPoints)
@@ -2640,6 +2656,9 @@ FPrimitiveSceneProxy* USplineComponent::CreateSceneProxy()
 				// If not the first keypoint, draw a line to the previous keypoint.
 				if (KeyIdx > 0)
 				{
+					const FInterpCurvePoint<FVector> OldKey = SplineInfo.Points[KeyIdx - 1];
+					FVector OldKeyPos = OldKey.OutVal;
+					
 					// This initial value of 1.4f is an experimentally determined value which gives the same effective segment thickness as the
 					// previous PDI rendering when the level editor thickness adjustment setting was at its default value (0). PDI special cases
 					// lines of 0 thickness while our approach does not, so if we are not artificially increasing it here we will get invisible
@@ -2654,7 +2673,7 @@ FPrimitiveSceneProxy* USplineComponent::CreateSceneProxy()
 					{
 						FVector OldPos = OldKeyPos;
 						
-						const FVector SegmentVector = SplineInfo.Points[KeyIdx].OutVal - OldKeyPos;
+						const FVector SegmentVector = NewKeyPos - OldKeyPos;
 						const float SegmentLength = SegmentVector.Length();
 						const FVector SegmentDirection = SegmentVector.GetSafeNormal();
 						const float StepLength = SegmentLength / NumStepsPerSegment;
@@ -2677,11 +2696,13 @@ FPrimitiveSceneProxy* USplineComponent::CreateSceneProxy()
 					}
 					else
 					{
+						UE::Geometry::FInterval1f SegmentParameterRange(OldKey.InVal, NewKey.InVal);
+						
 						FVector OldPos = OldKeyPos;
 
 						for (int32 StepIdx = 1; StepIdx <= NumStepsPerSegment; StepIdx++)
 						{
-							const float Key = (KeyIdx - 1) + (StepIdx / static_cast<float>(NumStepsPerSegment));
+							const float Key = SegmentParameterRange.Interpolate(static_cast<float>(StepIdx) / NumStepsPerSegment);
 							const FVector NewPos = SplineInfo.Eval(Key, FVector(0));
 							
 							AppendLineToBuffers(OldPos, NewPos, LineColor.ToFColorSRGB(), SegmentLineThickness);
@@ -2690,8 +2711,6 @@ FPrimitiveSceneProxy* USplineComponent::CreateSceneProxy()
 						}
 					}
 				}
-
-				OldKeyPos = NewKeyPos;
 			}
 
 			if (NumLines > 0)
@@ -2912,7 +2931,6 @@ void USplineComponent::PushSelectionToProxy()
 void USplineComponent::Draw(FPrimitiveDrawInterface* PDI, const FSceneView* View, const FInterpCurveVector& SplineInfo, const FMatrix& LocalToWorld, const FLinearColor& LineColor, uint8 DepthPriorityGroup)
 {
 	const int32 GrabHandleSize = 6;
-	FVector OldKeyPos(0);
 
 	auto GetLineColor = [LineColor](int32 KeyIdx, int32 NumSegments, int32 NumSteps) -> FLinearColor
 	{
@@ -2928,7 +2946,8 @@ void USplineComponent::Draw(FPrimitiveDrawInterface* PDI, const FSceneView* View
 	const int32 NumSegments = SplineInfo.bIsLooped ? NumPoints : NumPoints - 1;
 	for (int32 KeyIdx = 0; KeyIdx < NumSegments + 1; KeyIdx++)
 	{
-		const FVector NewKeyPos = LocalToWorld.TransformPosition(SplineInfo.Eval(static_cast<float>(KeyIdx), FVector(0)));
+		const FInterpCurvePoint<FVector>& NewKey = SplineInfo.Points[KeyIdx];
+		const FVector NewKeyPos = LocalToWorld.TransformPosition(NewKey.OutVal);
 
 		// Draw the keypoint
 		if (KeyIdx < NumPoints)
@@ -2939,35 +2958,40 @@ void USplineComponent::Draw(FPrimitiveDrawInterface* PDI, const FSceneView* View
 		// If not the first keypoint, draw a line to the previous keypoint.
 		if (KeyIdx > 0)
 		{
+			const FInterpCurvePoint<FVector>& OldKey = SplineInfo.Points[KeyIdx - 1];
+			const FVector OldKeyPos = LocalToWorld.TransformPosition(OldKey.OutVal);
+			
 			// For constant interpolation - don't draw ticks - just draw dotted line.
-			if (SplineInfo.Points[KeyIdx - 1].InterpMode == CIM_Constant)
+			if (OldKey.InterpMode == CIM_Constant)
 			{
 				// Calculate dash length according to size on screen
 				const float StartW = View->WorldToScreen(OldKeyPos).W;
 				const float EndW = View->WorldToScreen(NewKeyPos).W;
 
-				const float WLimit = 10.0f;
+				constexpr float WLimit = 10.0f;
 				if (StartW > WLimit || EndW > WLimit)
 				{
-					const float Scale = 0.03f;
+					constexpr float Scale = 0.03f;
 					DrawDashedLine(PDI, OldKeyPos, NewKeyPos, LineColor, FMath::Max(StartW, EndW) * Scale, DepthPriorityGroup);
 				}
 			}
 			else
 			{
+				UE::Geometry::FInterval1f SegmentParameterRange(OldKey.InVal, NewKey.InVal);
+				
 				// Find position on first keyframe.
 				FVector OldPos = OldKeyPos;
 
 				// Then draw a line for each substep.
-				const int32 NumSteps = 20;
+				constexpr int32 NumSteps = 20;
 #if WITH_EDITOR
 				const float SegmentLineThickness = GetDefault<ULevelEditorViewportSettings>()->SplineLineThicknessAdjustment;
 #endif
 
 				for (int32 StepIdx = 1; StepIdx <= NumSteps; StepIdx++)
 				{
-					const float Key = (KeyIdx - 1) + (StepIdx / static_cast<float>(NumSteps));
-					const FVector NewPos = LocalToWorld.TransformPosition(SplineInfo.Eval(Key, FVector(0)));
+					const float SampleKey = SegmentParameterRange.Interpolate(static_cast<float>(StepIdx) / NumSteps);
+					const FVector NewPos = LocalToWorld.TransformPosition(SplineInfo.Eval(SampleKey, FVector(0)));
 #if WITH_EDITOR
 					PDI->DrawTranslucentLine(OldPos, NewPos, GetLineColor(KeyIdx * NumSteps + StepIdx, NumSegments, NumSteps), DepthPriorityGroup, SegmentLineThickness);
 #else
@@ -2977,8 +3001,6 @@ void USplineComponent::Draw(FPrimitiveDrawInterface* PDI, const FSceneView* View
 				}
 			}
 		}
-
-		OldKeyPos = NewKeyPos;
 	}
 }
 #endif
