When interactively editing a material instance, only flush any associated RVT at the end of the edit.
This is better for performance, and feels better to not have the low res RVT always showing during the edit.
#rnx

diff --git a/Engine/Source/Runtime/Engine/Private/Materials/MaterialInstance.cpp b/Engine/Source/Runtime/Engine/Private/Materials/MaterialInstance.cpp
--- a/Engine/Source/Runtime/Engine/Private/Materials/MaterialInstance.cpp
+++ b/Engine/Source/Runtime/Engine/Private/Materials/MaterialInstance.cpp
@@ -4716,24 +4716,25 @@ void UMaterialInstance::PostEditChangeProperty(FPropertyChangedEvent& PropertyCh
 	}
 	UpdateStaticPermutation(GetStaticParameters(), BasePropertyOverrides, bForceStaticPermutationUpdate);
 
-	if (PropertyChangedEvent.ChangeType & (EPropertyChangeType::ValueSet | EPropertyChangeType::ArrayClear | EPropertyChangeType::ArrayRemove | EPropertyChangeType::ArrayMove | EPropertyChangeType::Unspecified | EPropertyChangeType::Duplicate))
+	// To avoid unnecessary work, don't apply to child instances and RVT during interactive edits.
+	if ((PropertyChangedEvent.ChangeType & EPropertyChangeType::Interactive) == 0)
 	{
 		RecacheMaterialInstanceUniformExpressions(this, false);
-	}
-
-	UpdateCachedData();
 
-	if (GIsEditor)
-	{
-		// Brute force all flush virtual textures if this material writes to any runtime virtual texture.
-		if (WritesToRuntimeVirtualTexture())
+		if (GIsEditor)
 		{
-			ENQUEUE_RENDER_COMMAND(FlushVTCommand)([ResourcePtr = Resource](FRHICommandListImmediate& RHICmdList) 
+			// Brute force all flush virtual textures if this material writes to any runtime virtual texture.
+			if (WritesToRuntimeVirtualTexture())
 			{
-				GetRendererModule().FlushVirtualTextureCache();	
-			});
+				ENQUEUE_RENDER_COMMAND(FlushVTCommand)([ResourcePtr = Resource](FRHICommandListImmediate& RHICmdList)
+				{
+					GetRendererModule().FlushVirtualTextureCache();
+				});
+			}
 		}
 	}
+
+	UpdateCachedData();
 }
 
 void UMaterialInstance::PostEditUndo()
