Add a check to verify ContentBundle streaming object package exsists before trying to load it.
This will prevent the log from containing missing packages warning for package that would never exist due to the content bundle being empty of Actors.
Same thing we do for external data layers

diff --git a/Engine/Source/Runtime/Engine/Private/WorldPartition/ContentBundle/ContentBundle.cpp b/Engine/Source/Runtime/Engine/Private/WorldPartition/ContentBundle/ContentBundle.cpp
--- a/Engine/Source/Runtime/Engine/Private/WorldPartition/ContentBundle/ContentBundle.cpp
+++ b/Engine/Source/Runtime/Engine/Private/WorldPartition/ContentBundle/ContentBundle.cpp
@@ -26,23 +26,31 @@ void FContentBundle::DoInitialize()
 #if WITH_EDITOR
 	InitializeForPIE();
 #else
-	ExternalStreamingObjectPackage = LoadPackage(nullptr, *GetExternalStreamingObjectPackagePath(), LOAD_None);
-	if (ExternalStreamingObjectPackage != nullptr)
+	const FString ExternalStreamingObjectPackagePath = GetExternalStreamingObjectPackagePath();
+	if (FPackageName::DoesPackageExist(ExternalStreamingObjectPackagePath))
 	{
-		ExternalStreamingObject = Cast<URuntimeHashExternalStreamingObjectBase>((UObject*)FindObjectWithOuter(ExternalStreamingObjectPackage, URuntimeHashExternalStreamingObjectBase::StaticClass()));
-
-		if (ExternalStreamingObject)
+		ExternalStreamingObjectPackage = LoadPackage(nullptr, *ExternalStreamingObjectPackagePath, LOAD_None);
+		if (ExternalStreamingObjectPackage != nullptr)
 		{
-			ExternalStreamingObject->OnStreamingObjectLoaded(GetInjectedWorld());
+			ExternalStreamingObject = Cast<URuntimeHashExternalStreamingObjectBase>((UObject*)FindObjectWithOuter(ExternalStreamingObjectPackage, URuntimeHashExternalStreamingObjectBase::StaticClass()));
+
+			if (ExternalStreamingObject)
+			{
+				ExternalStreamingObject->OnStreamingObjectLoaded(GetInjectedWorld());
+			}
+			else
+			{
+				UE_LOG(LogContentBundle, Error, TEXT("%s No streaming object found in package %s. No content will be injected."), *ContentBundle::Log::MakeDebugInfoString(*this), *GetExternalStreamingObjectPackagePath());
+			}
 		}
 		else
 		{
-			UE_LOG(LogContentBundle, Error, TEXT("%s No streaming object found in package %s. No content will be injected."), *ContentBundle::Log::MakeDebugInfoString(*this), *GetExternalStreamingObjectPackagePath());
+			UE_LOG(LogContentBundle, Error, TEXT("%s Streaming package %s failed to load. No content will be injected."), *ContentBundle::Log::MakeDebugInfoString(*this), *GetExternalStreamingObjectPackagePath());
 		}
 	}
 	else
 	{
-		UE_LOG(LogContentBundle, Log, TEXT("%s Streaming package %s not found. No content will be injected."), *ContentBundle::Log::MakeDebugInfoString(*this), *GetExternalStreamingObjectPackagePath());
+		UE_LOG(LogContentBundle, Log, TEXT("%s Has no content for the current World (this can happen if the ContentBundle has no Actors for the current World)."), *ContentBundle::Log::MakeDebugInfoString(*this));
 	}
 #endif
 
