Fix non-unity build

diff --git a/Engine/Source/Runtime/Engine/Private/WorldPartition/ContentBundle/ContentBundle.cpp b/Engine/Source/Runtime/Engine/Private/WorldPartition/ContentBundle/ContentBundle.cpp
--- a/Engine/Source/Runtime/Engine/Private/WorldPartition/ContentBundle/ContentBundle.cpp
+++ b/Engine/Source/Runtime/Engine/Private/WorldPartition/ContentBundle/ContentBundle.cpp
@@ -8,6 +8,7 @@
 #include "WorldPartition/ContentBundle/ContentBundleLog.h"
 #include "Engine/World.h"
 #include "UObject/Package.h"
+#include "Misc/PackageName.h"
 
 #if WITH_EDITOR
 #include "WorldPartition/ContentBundle/ContentBundleWorldSubsystem.h"
