Fix for some platform specific compile warnings

#rnx

diff --git a/Engine/Source/Runtime/Engine/Public/Animation/SpringMath.h b/Engine/Source/Runtime/Engine/Public/Animation/SpringMath.h
--- a/Engine/Source/Runtime/Engine/Public/Animation/SpringMath.h
+++ b/Engine/Source/Runtime/Engine/Public/Animation/SpringMath.h
@@ -227,8 +227,8 @@ public:
 		const TVector& TargetVelocity,
 		float SmoothingTime,
 		float DeltaTime,
-		float VDeadzone = 1e-2,
-		float ADeadzone = 1e-4)
+		float VDeadzone = 1e-2f,
+		float ADeadzone = 1e-4f)
 	{
 		float Y = SmoothingTimeToDamping(SmoothingTime) / 2.0f;
 		TVector J0 = InOutVelocity - TargetVelocity;
@@ -280,8 +280,8 @@ public:
 		const TVector& TargetVelocity,
 		float SmoothingTime,
 		float SecondsPerPredictionStep,
-		float VDeadzone = 1e-2,
-		float ADeadzone = 1e-4)
+		float VDeadzone = 1e-2f,
+		float ADeadzone = 1e-4f)
 	{
 		int32 PredictCount = OutPredictedPositions.Num();
 		check(PredictCount > 0);
@@ -297,7 +297,7 @@ public:
 
 		for (int32 i = 0; i < PredictCount; i++)
 		{
-			const float PredictTime = (i + 1) * SecondsPerPredictionStep; // Note i+1 since we want index 0 to be the first prediction step
+			const float PredictTime = (float)(i + 1) * SecondsPerPredictionStep; // Note i+1 since we want index 0 to be the first prediction step
 			PRAGMA_DISABLE_EXPERIMENTAL_WARNINGS
 			SpringCharacterUpdate(OutPredictedPositions[i], OutPredictedVelocities[i], OutPredictedAccelerations[i], TargetVelocity, SmoothingTime,
 				PredictTime, VDeadzone, ADeadzone);
@@ -332,8 +332,8 @@ public:
 		float SmoothingTime,
 		float MaxAcceleration,
 		float DeltaTime,
-		float VDeadzone = 1e-2,
-		float ADeadzone = 1e-4)
+		float VDeadzone = 1e-2f,
+		float ADeadzone = 1e-4f)
 	{
 		TVector VDiff = TargetVelocity - InOutVelocityIntermediate;
 		float VDiffLength = VDiff.Length();
@@ -405,8 +405,8 @@ public:
 		float SmoothingTime,
 		float MaxAcceleration,
 		float SecondsPerPredictionStep,
-		float VDeadzone = 1e-2,
-		float ADeadzone = 1e-4)
+		float VDeadzone = 1e-2f,
+		float ADeadzone = 1e-4f)
 	{
 		int32 PredictCount = OutPredictedPositions.Num();
 		check(PredictCount > 0);
@@ -424,7 +424,7 @@ public:
 
 		for (int32 i = 0; i < PredictCount; i++)
 		{
-			const float PredictTime = (i + 1) * SecondsPerPredictionStep; // Note i+1 since we want index 0 to be the first prediction step
+			const float PredictTime = (float)(i + 1) * SecondsPerPredictionStep; // Note i+1 since we want index 0 to be the first prediction step
 			VelocitySpringCharacterUpdate(OutPredictedPositions[i],
 				OutPredictedVelocities[i],
 				OutPredictedIntermediateVelocities[i],
@@ -465,7 +465,7 @@ public:
 
 		for (int32 i = 0; i < PredictCount; i++)
 		{
-			const float PredictTime = (i + 1) * SecondsPerPredictionStep; // Note i+1 since we want index 0 to be the first prediction step
+			const float PredictTime = (float)(i + 1) * SecondsPerPredictionStep; // Note i+1 since we want index 0 to be the first prediction step
 			PRAGMA_DISABLE_EXPERIMENTAL_WARNINGS
 			CriticalSpringDamperQuat(OutPredictedRotations[i], OutPredictedAngularVelocities[i], TargetRotation, SmoothingTime, PredictTime);
 			PRAGMA_ENABLE_EXPERIMENTAL_WARNINGS
