[VisualLogger] replaced MT access detector and Lock by their transactionally safe versions

#FailedValidations Preflight

diff --git a/Engine/Source/Runtime/Engine/Private/VisualLogger/VisualLogger.cpp b/Engine/Source/Runtime/Engine/Private/VisualLogger/VisualLogger.cpp
--- a/Engine/Source/Runtime/Engine/Private/VisualLogger/VisualLogger.cpp
+++ b/Engine/Source/Runtime/Engine/Private/VisualLogger/VisualLogger.cpp
@@ -325,7 +325,7 @@ bool FVisualLogger::IsClassAllowed(const UClass& InClass) const
 
 void FVisualLogger::AddObjectToAllowList(const UObject& InObject)
 {
-	FWriteScopeLock Lock(EntryRWLock);
+	UE::TWriteScopeLock Lock(EntryRWLock);
 	bool bIsAlreadyInSet = false;
 	ObjectAllowList.Add(&InObject, &bIsAlreadyInSet);
 
@@ -355,7 +355,7 @@ void FVisualLogger::AddObjectToAllowList(const UObject& InObject)
 
 void FVisualLogger::ClearObjectAllowList()
 {
-	FWriteScopeLock Lock(EntryRWLock);
+	UE::TWriteScopeLock Lock(EntryRWLock);
 	for (FObjectKey It : ObjectAllowList)
 	{
 		{
@@ -443,7 +443,7 @@ FVisualLogEntry* FVisualLogger::GetEntryToWrite(const UObject* Object, const dou
 		{
 			if (ThreadCurrentEntry->ShouldFlush(TimeStamp))
 			{
-				FWriteScopeLock Lock(EntryRWLock);
+				UE::TWriteScopeLock Lock(EntryRWLock);
 				UE_MT_SCOPED_WRITE_ACCESS(CurrentEntryPerObject.EntriesAccessDetector);
 				if (FVisualLogEntry* GlobalCurrentEntry = GetEntryToWriteInternal(LogOwner, ThreadCurrentEntry->TimeStamp, ShouldCreate))
 				{
@@ -460,7 +460,7 @@ FVisualLogEntry* FVisualLogger::GetEntryToWrite(const UObject* Object, const dou
 		// if so then force creation for the current thread
 		if (ShouldCreate != ECreateIfNeeded::Create)
 		{
-			FReadScopeLock Lock(EntryRWLock);
+			UE::TReadScopeLock Lock(EntryRWLock);
 			UE_MT_SCOPED_READ_ACCESS(CurrentEntryPerObject.EntriesAccessDetector);
 			if (CurrentEntryPerObject.Find(LogOwner))
 			{
@@ -510,8 +510,8 @@ FVisualLogEntry* FVisualLogger::GetEntryToWriteInternal(const UObject* Object, c
 	// Entry can be created or reused (after being flushed) and will need to be initialized
 	bool bInitializeEntry = false;
 
-	// Using the Write access since it might flush and reset an entry and/or initialize one.
-	UE_MT_SCOPED_WRITE_ACCESS(CurrentEntryPerObject.EntriesAccessDetector);
+	// We were previously using UE_MT_SCOPED_WRITE_ACCESS(CurrentEntryPerObject.EntriesAccessDetector) here
+	// but Transactionally safe detectors are not reentrant so it is now the responsibility of the caller.
 	FVisualLogEntry* CurrentEntry = CurrentEntryPerObject.Find(LogOwner);
 	if (CurrentEntry)
 	{
@@ -673,7 +673,7 @@ void FVisualLogger::FlushThreadsEntries()
 
 void FVisualLogger::Flush()
 {
-	FWriteScopeLock Lock(EntryRWLock);
+	UE::TWriteScopeLock Lock(EntryRWLock);
 
 	FlushThreadsEntries();
 
@@ -706,7 +706,7 @@ FVisualLogger::FVisualLoggerObjectEntryMap& FVisualLogger::GetThreadCurrentEntry
 	if (GThreadCurrentEntryPerObject == nullptr)
 	{
 		GThreadCurrentEntryPerObject = new FVisualLoggerObjectEntryMap();
-		FWriteScopeLock Lock(EntryRWLock);
+		UE::TWriteScopeLock Lock(EntryRWLock);
 		UE_MT_SCOPED_WRITE_ACCESS(ThreadLocalMapListAccessDetector);
 		ThreadCurrentEntryMaps.Add(GThreadCurrentEntryPerObject);
 	}
diff --git a/Engine/Source/Runtime/Engine/Public/VisualLogger/VisualLogger.h b/Engine/Source/Runtime/Engine/Public/VisualLogger/VisualLogger.h
--- a/Engine/Source/Runtime/Engine/Public/VisualLogger/VisualLogger.h
+++ b/Engine/Source/Runtime/Engine/Public/VisualLogger/VisualLogger.h
@@ -1,19 +1,22 @@
 // Copyright Epic Games, Inc. All Rights Reserved.
 #pragma once
 
+#if UE_ENABLE_INCLUDE_ORDER_DEPRECATED_IN_5_7
 #include "CoreMinimal.h"
-#include "Stats/Stats.h"
-#include "UObject/Object.h"
+#endif // UE_ENABLE_INCLUDE_ORDER_DEPRECATED_IN_5_7 (switch to 5.8 in main)
+#include "Containers/Ticker.h"
 #include "Engine/EngineTypes.h"
 #include "EngineDefines.h"
-#include "VisualLogger/VisualLoggerTypes.h"
-#include "Templates/IsValidVariadicFunctionArg.h"
+#include "Misc/MTTransactionallySafeAccessDetector.h"
+#include "Misc/ScopeRWLock.h"
+#include "Misc/TransactionallySafeRWLock.h"
+#include "Stats/Stats.h"
 #include "Templates/IsArrayOrRefOfTypeByPredicate.h"
+#include "Templates/IsValidVariadicFunctionArg.h"
 #include "Traits/IsCharEncodingCompatibleWith.h"
+#include "UObject/Object.h"
 #include "UObject/ObjectKey.h"
-#include "Containers/Ticker.h"
-#include "Misc/ScopeRWLock.h"
-#include "Misc/TransactionallySafeRWLock.h"
+#include "VisualLogger/VisualLoggerTypes.h"
 
 #if ENABLE_VISUAL_LOG
 
@@ -893,7 +896,7 @@ public:
 	struct FVisualLoggerObjectEntryMap : TMap<FObjectKey, FVisualLogEntry>
 	{
 		// Multithread access detector for map entries (adding/removing to/from the map or modifying an entry)
-		UE_MT_DECLARE_RW_RECURSIVE_ACCESS_DETECTOR(EntriesAccessDetector);
+		UE_MT_DECLARE_TS_RW_ACCESS_DETECTOR(EntriesAccessDetector);
 	};
 
 private:
@@ -960,7 +963,7 @@ protected:
 	UE_MT_DECLARE_RW_ACCESS_DETECTOR(ThreadLocalMapListAccessDetector);
 
 	// Read Write lock protecting object entries
-	mutable FRWLock EntryRWLock;
+	mutable FTransactionallySafeRWLock EntryRWLock;
 	struct FObjectNames
 	{
 		FName Name;
