Add Insights counter for tracing shader queue, so can see pending and compiled in Insights
#rnx

diff --git a/Engine/Source/Runtime/Engine/Private/ShaderCompiler/ShaderCompilerStats.cpp b/Engine/Source/Runtime/Engine/Private/ShaderCompiler/ShaderCompilerStats.cpp
--- a/Engine/Source/Runtime/Engine/Private/ShaderCompiler/ShaderCompilerStats.cpp
+++ b/Engine/Source/Runtime/Engine/Private/ShaderCompiler/ShaderCompilerStats.cpp
@@ -890,6 +890,7 @@ void FShaderCompilerStats::RegisterLocalWorkerIdleTime(double IdleTime)
 	Counters.TimesLocalWorkersWereIdle++;
 }
 
+TRACE_DECLARE_INT_COUNTER(Shaders_Pending, TEXT("Shaders/Pending"));
 void FShaderCompilerStats::RegisterNewPendingJob(FShaderCommonCompileJob& Job)
 {
 	// accessing job timestamps isn't arbitrated by any lock. It is assumed that the registration of a job at one of the stages
@@ -897,6 +898,7 @@ void FShaderCompilerStats::RegisterNewPendingJob(FShaderCommonCompileJob& Job)
 	// so it cannot be given away to a worker while it's still being registered, and an assigned job is registered before it is actually given to the worker,
 	// so it cannot end up being registered as finished at the same time on some other thread).
 	Job.TimeAddedToPendingQueue = FPlatformTime::Seconds();
+	TRACE_COUNTER_ADD(Shaders_Pending, 1);
 }
 
 void FShaderCompilerStats::RegisterAssignedJob(FShaderCommonCompileJob& Job)
