Fix for DPI Scaling not impacting GlobalResolutionFraction

diff --git a/Engine/Source/Runtime/Engine/Private/LegacyScreenPercentageDriver.cpp b/Engine/Source/Runtime/Engine/Private/LegacyScreenPercentageDriver.cpp
--- a/Engine/Source/Runtime/Engine/Private/LegacyScreenPercentageDriver.cpp
+++ b/Engine/Source/Runtime/Engine/Private/LegacyScreenPercentageDriver.cpp
@@ -421,7 +421,7 @@ float FStaticResolutionFractionHeuristic::ResolveResolutionFraction() const
 	else
 	{
 		ensure(EffectiveScreenPercentageMode == EScreenPercentageMode::Manual);
-		GlobalResolutionFraction = Settings.GlobalResolutionFraction;
+		GlobalResolutionFraction = Settings.GlobalResolutionFraction / DPIScale;
 	}
 
 	// Max the rendering resolution to average target resolution used on platform that have dynamic resolution.
