[Waveform Editor] - Crash on importing wave file with Bad PCM Data Chunks

diff --git a/Engine/Source/Runtime/Engine/Private/SoundWave.cpp b/Engine/Source/Runtime/Engine/Private/SoundWave.cpp
--- a/Engine/Source/Runtime/Engine/Private/SoundWave.cpp
+++ b/Engine/Source/Runtime/Engine/Private/SoundWave.cpp
@@ -1726,18 +1726,86 @@ bool USoundWave::GetImportedSoundWaveData(TArray<uint8>& OutRawPCMData, uint32&
 				SPEAKER_RightBack
 			};
 
+			TFunction<bool(const FWaveModInfo& WaveModInfo)> ValidateWaveModInfo = [](const FWaveModInfo& WaveModInfo)
+			{
+				// Check for valid channel count
+				if (1 != *WaveModInfo.pChannels)
+				{
+					UE_LOG(LogAudio, Warning, TEXT("Cannot handle audio format. Expected single channel audio but read %d channels"), *WaveModInfo.pChannels);
+					
+					return false;
+				}
+
+				// Check for valid bit depth
+				if (16 != *WaveModInfo.pBitsPerSample)
+				{
+					UE_LOG(LogAudio, Warning, TEXT("Cannot handle audio format. Expected 16bit audio but found %d bit audio"), *WaveModInfo.pBitsPerSample);
+					
+					return false;
+				}
+
+				return true;
+			};
+
 			check(ChannelOffsets.Num() == ChannelSizes.Num());
 			check(ChannelOffsets.Num() == SPEAKER_Count);
 
 			// Multichannel audio with more than 2 channels must be accessed by
 			// inspecting the channel offsets and channel sizes of the USoundWave.
 
-			bool bIsOutputInitialized = false;
+			// Determine which channel is the smallest and read the wave for the number of frames so we can truncate larger channels if needed
+			int32 SmallestValidChannelIndex = INDEX_NONE;
+			int32 SmallestChannelSize = INT32_MAX;
 
-			int32 NumFrames = 0;
-			int32 NumSamples = 0;
-			OutSampleRate = 0;
+			for (int32 Index = 0; Index < ChannelOffsets.Num(); Index++)
+			{
+				if (ChannelSizes[Index] > 0 && ChannelSizes[Index] < SmallestChannelSize)
+				{
+					SmallestValidChannelIndex = Index;
+					SmallestChannelSize = ChannelSizes[Index];
+				}
+			}
+			
+			if (SmallestValidChannelIndex == INDEX_NONE)
+			{
+				UE_LOG(LogAudio, Warning, TEXT("Failed to read wave data: %s."), *GetFullName());
+				
+				return false;
+			}
+
+			FWaveModInfo SmallestChannelWaveInfo;
+			if (!SmallestChannelWaveInfo.ReadWaveHeader(Data, ChannelSizes[SmallestValidChannelIndex], ChannelOffsets[SmallestValidChannelIndex]))
+			{
+				UE_LOG(LogAudio, Warning, TEXT("Failed to read wave data: %s."), *GetFullName());
+				
+				return false;
+			}
 
+			if (!ValidateWaveModInfo(SmallestChannelWaveInfo))
+			{
+				return false;
+			}
+
+			check(SmallestChannelWaveInfo.SampleDataSize)
+
+			// The largest number of frames a channel should have in this soundwave
+			// Bad chunk data can lead to NumFramesPerChannel being a decimal value but the data truncation to int32 removes the bad sample.
+			int32 NumFramesPerChannel = SmallestChannelWaveInfo.SampleDataSize / sizeof(int16);
+			int32 NumSamples = NumFramesPerChannel * NumChannels; // Number of total samples across all channels
+			OutSampleRate = *SmallestChannelWaveInfo.pSamplesPerSec;
+
+			if (NumSamples > 0)
+			{
+				// Translate NumSamples to bytes because OutRawPCMData is in bytes. 
+				const int32 NumBytes = NumSamples * sizeof(int16);
+				OutRawPCMData.AddZeroed(NumBytes);
+			}
+			else
+			{
+				UE_LOG(LogAudio, Warning, TEXT("Zero samples found in wave data: %s."), *GetFullName());
+
+				return false;
+			}
 
 			// Determine which channels have data and Interleave channel data
 			for (int32 ChannelIndex = 0; ChannelIndex < ChannelOffsets.Num(); ChannelIndex++)
@@ -1756,66 +1824,52 @@ bool USoundWave::GetImportedSoundWaveData(TArray<uint8>& OutRawPCMData, uint32&
 					return false;
 				}
 
-				// Check for valid channel count
-				if (1 != *WaveInfo.pChannels)
-				{
-					UE_LOG(LogAudio, Warning, TEXT("Cannot audio handle format. Expected single channel audio but read %d channels"), *WaveInfo.pChannels);
-					return false;
-				}
-
-				// Check for valid bit depth
-				if (16 != *WaveInfo.pBitsPerSample)
+				if (!ValidateWaveModInfo(WaveInfo))
 				{
-					UE_LOG(LogAudio, Warning, TEXT("Cannot audio handle format. Expected 16bit audio but found %d bit audio"), *WaveInfo.pBitsPerSample);
 					return false;
 				}
 
 				// Set output channel type
 				OutChannelOrder.Add(DefaultChannelOrder[ChannelIndex]);
 
-				// The output info needs to be initialized from the first channel's wave info.
-				if (!bIsOutputInitialized)
-				{
-					OutSampleRate = *WaveInfo.pSamplesPerSec;
-
-					NumFrames = WaveInfo.SampleDataSize / sizeof(int16);
-					NumSamples = NumFrames * NumChannels;
+				check(OutSampleRate == *WaveInfo.pSamplesPerSec);
+				const int32 ThisChannelNumFrames = WaveInfo.SampleDataSize / sizeof(int16);
 
-					check(NumSamples % NumChannels == 0);
+				TArrayView<const int16> ChannelArrayView;
+				// If not all channels have the same number of frames we will have to copy the channel data and alter it
+				TArray<int16> ChannelArrayCopy;
 
-					if (NumSamples > 0)
-					{
-						// Translate NumSamples to bytes because OutRawPCMData is in bytes. 
-						const int32 NumBytes = NumSamples * sizeof(int16);
-						OutRawPCMData.AddZeroed(NumBytes);
-					}
-					
-					bIsOutputInitialized = true;
+				if (ThisChannelNumFrames == NumFramesPerChannel)
+				{
+					ChannelArrayView = TArrayView<const int16>(reinterpret_cast<const int16*>(WaveInfo.SampleDataStart), NumFramesPerChannel);
 				}
+				else // Adjust if not all channels have the same number of frames
+				{
+					UE_LOG(LogAudio, Warning, TEXT("Audio channels contain varying number of frames (%d vs %d). Truncating frames."), NumFramesPerChannel, ThisChannelNumFrames);
+					check(NumFramesPerChannel < ThisChannelNumFrames);
 
+					// Copy the data instead of referencing it directly so we don't modify the imported WaveInfo
+					ChannelArrayCopy.Append(reinterpret_cast<const int16*>(WaveInfo.SampleDataStart), ThisChannelNumFrames);
 
-				check(OutSampleRate == *WaveInfo.pSamplesPerSec);
-				const int32 ThisChannelNumFrames = WaveInfo.SampleDataSize / sizeof(int16);
+					constexpr EAllowShrinking AllowShrinking = EAllowShrinking::Yes;
+					ChannelArrayCopy.SetNum(NumFramesPerChannel, AllowShrinking); //Truncate 
 
-				if (ensureMsgf(ThisChannelNumFrames == NumFrames, TEXT("Audio channels contain varying number of frames (%d vs %d)"), NumFrames, ThisChannelNumFrames))
-				{
-					TArrayView<int16> OutPCM(reinterpret_cast<int16*>(OutRawPCMData.GetData()), NumSamples);
-					TArrayView<const int16> ChannelArrayView(reinterpret_cast<const int16*>(WaveInfo.SampleDataStart), NumFrames);
+					ChannelArrayView = TArrayView<const int16>(ChannelArrayCopy);
+				}
 
-					int32 DestSamplePos = OutChannelOrder.Num() - 1;
-					int32 SourceSamplePos = 0;
+				// Interleave each channel's data into one multichannel audio buffer
+				TArrayView<int16> OutPCM(reinterpret_cast<int16*>(OutRawPCMData.GetData()), NumSamples);
 
-					while (DestSamplePos < NumSamples)
-					{
-						OutPCM[DestSamplePos] = ChannelArrayView[SourceSamplePos];
+				check(OutChannelOrder.Num() >= 1);
+				int32 DestSamplePos = OutChannelOrder.Num() - 1;
+				int32 SourceSamplePos = 0;
 
-						SourceSamplePos++;
-						DestSamplePos += NumChannels;
-					}
-				}
-				else
+				while (DestSamplePos < NumSamples)
 				{
-					return false;
+					OutPCM[DestSamplePos] = ChannelArrayView[SourceSamplePos];
+
+					SourceSamplePos++;
+					DestSamplePos += NumChannels;
 				}
 			}
 		}
@@ -1831,11 +1885,22 @@ bool USoundWave::GetImportedSoundWaveData(TArray<uint8>& OutRawPCMData, uint32&
 			}
 
 			check(NumChannels > 0);
-			check(WaveInfo.SampleDataSize% NumChannels == 0);
+			const bool InvalidDataSize = (WaveInfo.SampleDataSize % sizeof(int16) != 0);
+			uint32 NumSamples = WaveInfo.SampleDataSize / sizeof(int16);
+
+			if (InvalidDataSize || NumSamples % NumChannels != 0)
+			{
+				UE_LOG(LogAudio, Warning, TEXT("Invalid Sample Data Size compared to NumChannels for: %s. Adjusting sample count"), *GetFullName());
+
+				NumSamples -= NumSamples % NumChannels;
+			}
+
+			// NumSamples * sizeof(int16) can be different from WaveInfo.DataSize if the file contains bad data
+			const uint32 SampleDataSize = NumSamples * sizeof(int16);
 
 			// Copy the raw PCM data and the header info that was parsed
-			OutRawPCMData.AddUninitialized(WaveInfo.SampleDataSize);
-			FMemory::Memcpy(OutRawPCMData.GetData(), WaveInfo.SampleDataStart, WaveInfo.SampleDataSize);
+			OutRawPCMData.AddUninitialized(SampleDataSize);
+			FMemory::Memcpy(OutRawPCMData.GetData(), WaveInfo.SampleDataStart, SampleDataSize);
 
 			OutSampleRate = *WaveInfo.pSamplesPerSec;
 			
