World Partition - HLOD: Minor tweaks to the HLOD build report

diff --git a/Engine/Plugins/Editor/WorldPartitionHLODUtilities/Source/Private/WorldPartition/HLOD/Utilities/WorldPartitionHLODUtilities.cpp b/Engine/Plugins/Editor/WorldPartitionHLODUtilities/Source/Private/WorldPartition/HLOD/Utilities/WorldPartitionHLODUtilities.cpp
--- a/Engine/Plugins/Editor/WorldPartitionHLODUtilities/Source/Private/WorldPartition/HLOD/Utilities/WorldPartitionHLODUtilities.cpp
+++ b/Engine/Plugins/Editor/WorldPartitionHLODUtilities/Source/Private/WorldPartition/HLOD/Utilities/WorldPartitionHLODUtilities.cpp
@@ -76,9 +76,6 @@ static void ComputeHLODHash(FHLODHashBuilder& InHashBuilder, AWorldPartitionHLOD
 
 	// Append all components hashes
 	UHLODBuilder::ComputeHLODHash(InHashBuilder, InSourceComponents);
-	
-	FString HashReport = InHashBuilder.BuildHashReport();
-	UE_LOG(LogHLODHash, Warning, TEXT("%s"), *HashReport);
 }
 
 void AddSourceActor(const FStreamingGenerationActorDescView& ActorDescView, const FName WorldPackageName, TMap<FGuid, FWorldPartitionRuntimeCellObjectMapping>& SourceActors)
diff --git a/Engine/Source/Runtime/Core/Private/Misc/App.cpp b/Engine/Source/Runtime/Core/Private/Misc/App.cpp
--- a/Engine/Source/Runtime/Core/Private/Misc/App.cpp
+++ b/Engine/Source/Runtime/Core/Private/Misc/App.cpp
@@ -97,6 +97,10 @@ const TCHAR* FApp::GetExecutingJobURL()
 		{
 			return FString();
 		}
+		else
+		{
+			HordeUrl.RemoveFromEnd(TEXT("/"));
+		}
 		FString HordeJobId = FPlatformMisc::GetEnvironmentVariable(TEXT("UE_HORDE_JOBID"));
 		if (HordeJobId.IsEmpty())
 		{
diff --git a/Engine/Source/Runtime/Engine/Private/WorldPartition/HLOD/HLODActor.cpp b/Engine/Source/Runtime/Engine/Private/WorldPartition/HLOD/HLODActor.cpp
--- a/Engine/Source/Runtime/Engine/Private/WorldPartition/HLOD/HLODActor.cpp
+++ b/Engine/Source/Runtime/Engine/Private/WorldPartition/HLOD/HLODActor.cpp
@@ -639,7 +639,7 @@ FString AWorldPartitionHLOD::GenerateHLODBuildReportHeaderString() const
 	CommandLine.ReplaceInline(TEXT("\n"), TEXT(""));
 	CommandLine.ReplaceInline(TEXT("\r"), TEXT(""));
 	FString EngineMode = FGenericPlatformMisc::GetEngineMode();
-	FString GraphicsRHI = FApp::GetGraphicsRHI();
+	FString GraphicsRHI = !FApp::GetGraphicsRHI().IsEmpty() ? FApp::GetGraphicsRHI() : TEXT("NullRHI");
 	FString DateTimeUTC = FDateTime::UtcNow().ToString(TEXT("%Y-%m-%d %H:%M:%S"));
 
 	BuildReportHeader << HLODLastBuildInfoSectionHeader << LINE_TERMINATOR;
@@ -650,7 +650,7 @@ FString AWorldPartitionHLOD::GenerateHLODBuildReportHeaderString() const
 	BuildReportHeader << TEXT(" * BuildConfiguration: ") << BuildConfiguration << LINE_TERMINATOR;
 	BuildReportHeader << TEXT(" * CommandLine:        ") << CommandLine        << LINE_TERMINATOR;
 	BuildReportHeader << TEXT(" * EngineMode:         ") << EngineMode         << LINE_TERMINATOR;
-	BuildReportHeader << TEXT(" * GraphicsRHI:	      ") << GraphicsRHI        << LINE_TERMINATOR;
+	BuildReportHeader << TEXT(" * GraphicsRHI:        ") << GraphicsRHI        << LINE_TERMINATOR;
 	BuildReportHeader << TEXT(" * DateTimeUTC:        ") << DateTimeUTC        << LINE_TERMINATOR;
 	BuildReportHeader << LINE_TERMINATOR;
 	// Actor Details
@@ -703,6 +703,7 @@ void AWorldPartitionHLOD::UpdateHLODBuildReportHeader()
 	HLODBuildReport  = HLOD_REPORT_BEGIN;
 	HLODBuildReport += LINE_TERMINATOR;
 	HLODBuildReport += GenerateHLODBuildReportHeaderString();
+	HLODBuildReport += LINE_TERMINATOR;
 	HLODBuildReport += BuildReportContent;
 }
 
diff --git a/Engine/Source/Runtime/Engine/Private/WorldPartition/HLOD/HLODHashBuilder.cpp b/Engine/Source/Runtime/Engine/Private/WorldPartition/HLOD/HLODHashBuilder.cpp
--- a/Engine/Source/Runtime/Engine/Private/WorldPartition/HLOD/HLODHashBuilder.cpp
+++ b/Engine/Source/Runtime/Engine/Private/WorldPartition/HLOD/HLODHashBuilder.cpp
@@ -22,6 +22,7 @@ void FHLODHashBuilder::PushObjectContext(const UObject* InObjectContext)
 void FHLODHashBuilder::PopObjectContext()
 {
 	check(!ObjectContextStack.IsEmpty());
+	ObjectsHashes.FindOrAdd(ObjectContextStack.Top()).Hash = GetCrc();
 	ObjectContextStack.Pop();
 }
 
@@ -46,7 +47,7 @@ uint32 FHLODHashBuilder::AddAssetReference(UObject* Asset, TFunctionRef<uint32()
 
 	if (!ObjectContextStack.IsEmpty())
 	{		
-		ObjectsReferences.FindOrAdd(ObjectContextStack.Top()).ReferencedAssets.Add(AssetName);
+		ObjectsHashes.FindOrAdd(ObjectContextStack.Top()).ReferencedAssets.Add(AssetName);
 	}
 
 	return AssetHash.Hash;
@@ -237,7 +238,7 @@ FString FHLODHashBuilder::BuildHashReport() const
 
 		FString RefsStr;
 		const FString ContextKey = Type.ToString() + TEXT(" ") + AssetName.ToString();
-		if (const FObjectHash* ContextObj = ObjectsReferences.Find(ContextKey))
+		if (const FObjectHash* ContextObj = ObjectsHashes.Find(ContextKey))
 		{
 			if (!ContextObj->ReferencedAssets.IsEmpty())
 			{
@@ -256,8 +257,8 @@ FString FHLODHashBuilder::BuildHashReport() const
 
 	// Objects section (sorted), skipping any keys consumed above
 	TArray<FString> SortedObjectPaths;
-	SortedObjectPaths.Reserve(ObjectsReferences.Num());
-	ObjectsReferences.GenerateKeyArray(SortedObjectPaths);
+	SortedObjectPaths.Reserve(ObjectsHashes.Num());
+	ObjectsHashes.GenerateKeyArray(SortedObjectPaths);
 	SortedObjectPaths.Sort([](const FString& A, const FString& B)
 	{
 		return A.Compare(B, ESearchCase::IgnoreCase) < 0;
@@ -272,19 +273,19 @@ FString FHLODHashBuilder::BuildHashReport() const
 			continue; // Object already shown in the "Referenced Assets" section
 		}
 
-		const FObjectHash& ObjectHashes = ObjectsReferences.FindChecked(ObjPath);
+		const FObjectHash& ObjectHash = ObjectsHashes.FindChecked(ObjPath);
 
-		Out += ObjPath;
+		Out += FString::Printf(TEXT("%s (Hash=%08X)\n"), *ObjPath, ObjectHash.Hash);
 
-		for (const TTuple<FName, FString>& Field : ObjectHashes.Fields)
+		for (const TTuple<FName, FString>& Field : ObjectHash.Fields)
 		{
 			Out += FString::Printf(TEXT("    * %s: %s\n"), *Field.Key.ToString(), *Field.Value);
 		}
 
-		if (!ObjectHashes.ReferencedAssets.IsEmpty())
+		if (!ObjectHash.ReferencedAssets.IsEmpty())
 		{
 			Out += TEXT("    * References: ");
-			Out += GetReferencesString(ObjectHashes);
+			Out += GetReferencesString(ObjectHash);
 			Out += TEXT("\n");
 		}
 
diff --git a/Engine/Source/Runtime/Engine/Public/WorldPartition/HLOD/HLODHashBuilder.h b/Engine/Source/Runtime/Engine/Public/WorldPartition/HLOD/HLODHashBuilder.h
--- a/Engine/Source/Runtime/Engine/Public/WorldPartition/HLOD/HLODHashBuilder.h
+++ b/Engine/Source/Runtime/Engine/Public/WorldPartition/HLOD/HLODHashBuilder.h
@@ -122,7 +122,7 @@ private:
 	// Add a field to the current object context
 	void AddField(const FName& InFieldName, const FString& InFieldValue)
 	{
-		FFieldArray& Fields = ObjectContextStack.IsEmpty() ? GlobalFields : ObjectsReferences.FindOrAdd(ObjectContextStack.Top()).Fields;
+		FFieldArray& Fields = ObjectContextStack.IsEmpty() ? GlobalFields : ObjectsHashes.FindOrAdd(ObjectContextStack.Top()).Fields;
 		Fields.Emplace(InFieldName, InFieldValue);
 	}
 
@@ -200,12 +200,13 @@ private:
 	// Track assets references & fields values for a given object.
 	struct FObjectHash
 	{
+		uint32		Hash = 0;
 		TSet<FName>	ReferencedAssets;
 		FFieldArray Fields;
 	};
 
 	// Map of objects names to their data.
-	TMap<FString, FObjectHash> ObjectsReferences;
+	TMap<FString, FObjectHash> ObjectsHashes;
 
 	// Stack of objects.
 	TArray<FString> ObjectContextStack;
