[UE] BugHawk catch, followup fix.

#rnx

diff --git a/Engine/Source/Runtime/Engine/Private/InheritableComponentHandler.cpp b/Engine/Source/Runtime/Engine/Private/InheritableComponentHandler.cpp
--- a/Engine/Source/Runtime/Engine/Private/InheritableComponentHandler.cpp
+++ b/Engine/Source/Runtime/Engine/Private/InheritableComponentHandler.cpp
@@ -274,14 +274,17 @@ void UInheritableComponentHandler::ValidateTemplates()
 		else
 		{
 #if WITH_EDITOR
-			// Mark the export as invalid to ensure the loader won't reload the object, 
-			// should the package owning this object be reloaded.
-			// Note, Rename will remove the renamed object's linker when moving to a new package so invalidate the export beforehand
-			FLinkerLoad::InvalidateExport(Record.ComponentTemplate);
-			// in editor, move the component template aside so its name is free:
-			Record.ComponentTemplate->Rename(nullptr, GetTransientPackage(), REN_DoNotDirty | REN_DontCreateRedirectors | REN_NonTransactional);
-			Record.ComponentTemplate->ClearFlags(RF_Standalone);
-			Record.ComponentTemplate->RemoveFromRoot();
+			if (Record.ComponentTemplate)
+			{
+				// Mark the export as invalid to ensure the loader won't reload the object, 
+				// should the package owning this object be reloaded.
+				// Note, Rename will remove the renamed object's linker when moving to a new package so invalidate the export beforehand
+				FLinkerLoad::InvalidateExport(Record.ComponentTemplate);
+				// in editor, move the component template aside so its name is free:
+				Record.ComponentTemplate->Rename(nullptr, GetTransientPackage(), REN_DoNotDirty | REN_DontCreateRedirectors | REN_NonTransactional);
+				Record.ComponentTemplate->ClearFlags(RF_Standalone);
+				Record.ComponentTemplate->RemoveFromRoot();
+			}
 #endif // WITH_EDITOR
 			Records.RemoveAtSwap(Index);
 		}
