Fix to correctly copy notifies when duplicating an anim sequence so the new notifies no longer reference the old sequence

diff --git a/Engine/Source/Runtime/Engine/Private/Animation/AnimSequence.cpp b/Engine/Source/Runtime/Engine/Private/Animation/AnimSequence.cpp
--- a/Engine/Source/Runtime/Engine/Private/Animation/AnimSequence.cpp
+++ b/Engine/Source/Runtime/Engine/Private/Animation/AnimSequence.cpp
@@ -3156,11 +3156,10 @@ bool UAnimSequence::CreateAnimation(UAnimSequence* Sequence)
 
 		CopyDataModel(Sequence->GetDataModelInterface());
 		Controller->SetModel(DataModelInterface);
-		Notifies = Sequence->Notifies;
-		AnimNotifyTracks = Sequence->AnimNotifyTracks;
-
 		Controller->NotifyPopulated();
 
+		UE::Anim::CopyNotifies(Sequence, this, false);
+
 		return true;
 	}
 	
