SplineComponent: Fix crash in scene proxy constructors for 0 point splines

diff --git a/Engine/Source/Runtime/Engine/Private/Components/SplineComponent.cpp b/Engine/Source/Runtime/Engine/Private/Components/SplineComponent.cpp
--- a/Engine/Source/Runtime/Engine/Private/Components/SplineComponent.cpp
+++ b/Engine/Source/Runtime/Engine/Private/Components/SplineComponent.cpp
@@ -2641,7 +2641,7 @@ FPrimitiveSceneProxy* USplineComponent::CreateSceneProxy()
 				PointIndexBufferIndex += 6;
 			};
 			
-			for (int32 KeyIdx = 0; KeyIdx < NumSegments + 1; KeyIdx++)
+			for (int32 KeyIdx = 0; KeyIdx < NumSegments + 1 && NumPoints > 0; KeyIdx++)
 			{
 				const float InputKey = KeyIdx < SplineInfo.Points.Num() ? SplineInfo.Points[KeyIdx].InVal : SplineInfo.Points.Last().InVal + SplineInfo.LoopKeyOffset;
 				const FVector NewKeyPos = SplineInfo.Eval(InputKey);
@@ -2944,7 +2944,7 @@ void USplineComponent::Draw(FPrimitiveDrawInterface* PDI, const FSceneView* View
 	
 	const int32 NumPoints = SplineInfo.Points.Num();
 	const int32 NumSegments = SplineInfo.bIsLooped ? NumPoints : NumPoints - 1;
-	for (int32 KeyIdx = 0; KeyIdx < NumSegments + 1; KeyIdx++)
+	for (int32 KeyIdx = 0; KeyIdx < NumSegments + 1 && NumPoints > 0; KeyIdx++)
 	{
 		const float InputKey = KeyIdx < SplineInfo.Points.Num() ? SplineInfo.Points[KeyIdx].InVal : SplineInfo.Points.Last().InVal + SplineInfo.LoopKeyOffset;
 		const FVector NewKeyPos = LocalToWorld.TransformPosition(SplineInfo.Eval(InputKey));
