World Partition - HLOD: Minor tweak to the HLOD build report

diff --git a/Engine/Source/Runtime/Engine/Private/WorldPartition/HLOD/HLODActor.cpp b/Engine/Source/Runtime/Engine/Private/WorldPartition/HLOD/HLODActor.cpp
--- a/Engine/Source/Runtime/Engine/Private/WorldPartition/HLOD/HLODActor.cpp
+++ b/Engine/Source/Runtime/Engine/Private/WorldPartition/HLOD/HLODActor.cpp
@@ -698,9 +698,13 @@ void AWorldPartitionHLOD::UpdateHLODBuildReportHeader()
 		BuildReportContent += HLODBuildDetailsSectionHeader;
 		BuildReportContent += LINE_TERMINATOR;
 		BuildReportContent += HLOD_REPORT_END;
+		BuildReportContent += LINE_TERMINATOR;
 	}
 
-	HLODBuildReport  = HLOD_REPORT_BEGIN;
+	HLODBuildReport.Reset();
+
+	HLODBuildReport += LINE_TERMINATOR;
+	HLODBuildReport += HLOD_REPORT_BEGIN;
 	HLODBuildReport += LINE_TERMINATOR;
 	HLODBuildReport += GenerateHLODBuildReportHeaderString();
 	HLODBuildReport += LINE_TERMINATOR;
