Subtitles: Prevent LogEditCondition errors when undoing subtitle creation by adding an editor-only helper variable.

diff --git a/Engine/Source/Runtime/Engine/Private/Subtitles/SubtitlesAndClosedCaptionsDelegates.cpp b/Engine/Source/Runtime/Engine/Private/Subtitles/SubtitlesAndClosedCaptionsDelegates.cpp
--- a/Engine/Source/Runtime/Engine/Private/Subtitles/SubtitlesAndClosedCaptionsDelegates.cpp
+++ b/Engine/Source/Runtime/Engine/Private/Subtitles/SubtitlesAndClosedCaptionsDelegates.cpp
@@ -8,3 +8,13 @@ TDelegate<void(const FQueueSubtitleParameters&, const ESubtitleTiming)> FSubtitl
 TDelegate<bool(const FSubtitleAssetData&)> FSubtitlesAndClosedCaptionsDelegates::IsSubtitleActive;
 TDelegate<void(const FSubtitleAssetData&)> FSubtitlesAndClosedCaptionsDelegates::StopSubtitle;
 TDelegate<void()> FSubtitlesAndClosedCaptionsDelegates::StopAllSubtitles;
+
+#if WITH_EDITORONLY_DATA
+void USubtitleAssetUserData::PostEditChangeProperty(FPropertyChangedEvent& PropertyChangedEvent)
+{
+	for (FSubtitleAssetData& Subtitle : Subtitles)
+	{
+		Subtitle.bCanEditDuration = (Subtitle.SubtitleDurationType == ESubtitleDurationType::UseDurationProperty);
+	}
+}
+#endif // WITH_EDITORONLY_DATA
diff --git a/Engine/Source/Runtime/Engine/Public/Subtitles/SubtitlesAndClosedCaptionsDelegates.h b/Engine/Source/Runtime/Engine/Public/Subtitles/SubtitlesAndClosedCaptionsDelegates.h
--- a/Engine/Source/Runtime/Engine/Public/Subtitles/SubtitlesAndClosedCaptionsDelegates.h
+++ b/Engine/Source/Runtime/Engine/Public/Subtitles/SubtitlesAndClosedCaptionsDelegates.h
@@ -60,10 +60,16 @@ struct FSubtitleAssetData
 	UPROPERTY(EditAnywhere, BlueprintReadWrite, Category = "Subtitles|Duration")
 	ESubtitleDurationType SubtitleDurationType = ESubtitleDurationType::UseSoundDuration;
 
+#if WITH_EDITORONLY_DATA
+	// Editor-only helper value for SubtitleDefaultDuration's EditCondition.
+	UPROPERTY(Transient)
+	bool bCanEditDuration = (SubtitleDurationType == ESubtitleDurationType::UseDurationProperty);
+#endif
+
 	// Time to display in seconds.
 	// Defaulted to 3 seconds so when adding new subtitles it's not required to enter a placeholder Duration.
 	// Duration can be be set by ingestion pipelines when importing Subtitles in bulk
-	UPROPERTY(EditAnywhere, BlueprintReadWrite, Category = "Subtitles|Duration", meta = (ClampMin = 0.05f, EditCondition = "SubtitleDurationType == ESubtitleDurationType::UseDurationProperty"))
+	UPROPERTY(EditAnywhere, BlueprintReadWrite, Category = "Subtitles|Duration", meta = (ClampMin = 0.05f, EditCondition = "bCanEditDuration"))
 	float Duration = SubtitleDefaultDuration;
 
 	// Some subtitles have a delay before they're allowed to be displayed (primarily from the legacy system).
@@ -122,6 +128,12 @@ public:
 
 		return Duration;
 	};
+
+#if WITH_EDITORONLY_DATA
+	// UObject::PostEditChangeProperty is in a #if WITH_EDITOR wrapper.
+	// EDITORONLY_DATA seems equivalent here and ensures that it matches the #if for bCanEditDuration above (which is indeed editor-only data).
+	virtual void PostEditChangeProperty(FPropertyChangedEvent& PropertyChangedEvent) override;
+#endif
 };
 
 
