//UE5/Release-5.7 - Build MLDeformerSample Win64 LogAnimationCompression: Error: Failed to generate compressed animation data for AnimSequence
#fix Added slightly more detail to "Failed to generate compressed animation data" log. And handle nullptr AnimSequence separately from compression failure for logging

diff --git a/Engine/Source/Runtime/Engine/Private/Animation/AnimationCompressionDerivedData.cpp b/Engine/Source/Runtime/Engine/Private/Animation/AnimationCompressionDerivedData.cpp
--- a/Engine/Source/Runtime/Engine/Private/Animation/AnimationCompressionDerivedData.cpp
+++ b/Engine/Source/Runtime/Engine/Private/Animation/AnimationCompressionDerivedData.cpp
@@ -225,7 +225,8 @@ bool FAnimationSequenceAsyncCacheTask::BuildData() const
 	const bool bIsCanceled = Owner.IsCanceled();
 
 	const bool bCompressionSuccessful = bBoneCompressionOk && bCurveCompressionOk;
-	const FString CompressionName = DataToCompress.BoneCompressionSettings->GetFullName();
+	const FString BoneCompressionName = DataToCompress.BoneCompressionSettings ? DataToCompress.BoneCompressionSettings->GetFullName() : TEXT("NONE");
+	const FString CurveCompressionName = DataToCompress.CurveCompressionSettings ? DataToCompress.CurveCompressionSettings->GetFullName() : TEXT("NONE");
 
 	const UAnimSequence* AnimSequence = WeakAnimSequence.Get();
 	if (bCompressionSuccessful && !bIsCanceled && AnimSequence)
@@ -239,9 +240,13 @@ bool FAnimationSequenceAsyncCacheTask::BuildData() const
 		
 		return true;
 	}
-	else if(!bIsCanceled)
+	else if(!bCompressionSuccessful && !bIsCanceled)
 	{
-		UE_LOG(LogAnimationCompression, Error, TEXT("Failed to generate compressed animation data for %s with compression scheme %s for target platform %s"), *CompressibleAnimPtr->FullName, *CompressionName, *TargetPlatform->DisplayName().ToString());
+		UE_LOG(LogAnimationCompression, Error, TEXT("Failed to generate compressed animation data for %s with bone compression scheme %s [%i] and curve compression scheme %s [%i] for target platform %s"), *CompressibleAnimPtr->FullName, *BoneCompressionName, bBoneCompressionOk ? 1 : 0, *CurveCompressionName, bCurveCompressionOk ? 1 : 0, *TargetPlatform->DisplayName().ToString());
+	}
+	else if (!AnimSequence)
+	{
+		UE_LOG(LogAnimationCompression, Warning, TEXT("Failed to apply compressed animation data for %s as it has previously been garbage collected"), *CompressibleAnimPtr->FullName);
 	}
 	
 	return false;
