Animation attribute evaluation should remap bone indices to target skeleton (if applicable)
#fix use skeleton remapping to retrieve target skeleton bone index for evaluated attribute value insertion

diff --git a/Engine/Source/Runtime/Engine/Private/Animation/AnimSequence.cpp b/Engine/Source/Runtime/Engine/Private/Animation/AnimSequence.cpp
--- a/Engine/Source/Runtime/Engine/Private/Animation/AnimSequence.cpp
+++ b/Engine/Source/Runtime/Engine/Private/Animation/AnimSequence.cpp
@@ -3706,6 +3706,7 @@ void UAnimSequence::EvaluateAttributes(FAnimationPoseData& OutAnimationPoseData,
 
 	const FBoneContainer& RequiredBones = OutAnimationPoseData.GetPose().GetBoneContainer();
 	UE::Anim::FStackAttributeContainer& OutAttributes = OutAnimationPoseData.GetAttributes();
+	const FSkeletonRemapping& SkeletonRemapping = UE::Anim::FSkeletonRemappingRegistry::Get().GetRemapping(GetSkeleton(), RequiredBones.GetSkeletonAsset());
 
 #if WITH_EDITOR
 	if (bUseRawData)
@@ -3713,8 +3714,10 @@ void UAnimSequence::EvaluateAttributes(FAnimationPoseData& OutAnimationPoseData,
 		ValidateModel();
 		for (const FAnimatedBoneAttribute& Attribute : DataModelInterface->GetAttributes())
 		{
-			const FCompactPoseBoneIndex PoseBoneIndex = RequiredBones.GetCompactPoseIndexFromSkeletonIndex(Attribute.Identifier.GetBoneIndex());
-			// Only add attribute if the bone its tied to exists in the currently evaluated set of bones
+			const int32 AttributeBoneIndex = Attribute.Identifier.GetBoneIndex();
+			const int32 RemappedSkeletonBoneIndex = SkeletonRemapping.IsValid() ? SkeletonRemapping.GetTargetSkeletonBoneIndex(AttributeBoneIndex) : AttributeBoneIndex;
+			const FCompactPoseBoneIndex PoseBoneIndex = RequiredBones.GetCompactPoseIndexFromSkeletonIndex(RemappedSkeletonBoneIndex);
+			// Only add attribute if the, optionally remapped, bone it is targeting exists in the currently evaluated set of bones
 			if(PoseBoneIndex.IsValid())
 			{
 				UE::Anim::Attributes::GetAttributeValue(OutAttributes, PoseBoneIndex, Attribute, ExtractionContext.CurrentTime);
@@ -3726,8 +3729,10 @@ void UAnimSequence::EvaluateAttributes(FAnimationPoseData& OutAnimationPoseData,
 	{
 		for (const TPair<FAnimationAttributeIdentifier, FAttributeCurve>& BakedAttribute : AttributeCurves)
 		{
-			const FCompactPoseBoneIndex PoseBoneIndex = RequiredBones.GetCompactPoseIndexFromSkeletonIndex(BakedAttribute.Key.GetBoneIndex());
-			// Only add attribute if the bone its tied to exists in the currently evaluated set of bones
+			const int32 AttributeBoneIndex = BakedAttribute.Key.GetBoneIndex();
+			const int32 RemappedSkeletonBoneIndex = SkeletonRemapping.IsValid() ? SkeletonRemapping.GetTargetSkeletonBoneIndex(AttributeBoneIndex) : AttributeBoneIndex;
+			const FCompactPoseBoneIndex PoseBoneIndex = RequiredBones.GetCompactPoseIndexFromSkeletonIndex(RemappedSkeletonBoneIndex);
+			// Only add attribute if , optionally remapped, bone it is targeting exists in the currently evaluated set of bones
 			if(PoseBoneIndex.IsValid())
 			{
 				UE::Anim::FAttributeId Info(BakedAttribute.Key.GetName(), PoseBoneIndex);
