Investigate one spot in code regarding interpolation overrides
#fix ensure interpolation type in AnimExtract context is set according to DecompressionContext its value

diff --git a/Engine/Source/Runtime/Engine/Private/Animation/AnimSequence.cpp b/Engine/Source/Runtime/Engine/Private/Animation/AnimSequence.cpp
--- a/Engine/Source/Runtime/Engine/Private/Animation/AnimSequence.cpp
+++ b/Engine/Source/Runtime/Engine/Private/Animation/AnimSequence.cpp
@@ -1518,7 +1518,8 @@ void UAnimSequence::GetBoneTransform(FTransform& OutAtom, FSkeletonPoseBoneIndex
 {
 	UE_COMPRESSED_DATA_READ_SCOPE(this);
 	FAnimExtractContext ExtractContext(DecompContext.GetEvaluationTime());
-	GetBoneTransform(OutAtom, BoneIndex, ExtractContext, bUseRawData);
+	ExtractContext.InterpolationOverride = DecompContext.Interpolation;
+	GetBoneTransform_Lockless(OutAtom, BoneIndex, ExtractContext, bUseRawData);
 }
 
 void UAnimSequence::HandleAssetPlayerTickedInternal(FAnimAssetTickContext &Context, const float PreviousTime, const float MoveDelta, const FAnimTickRecord &Instance, struct FAnimNotifyQueue& NotifyQueue) const
