Defend against infinte recursion in process tree traversal.

diff --git a/Engine/Source/Runtime/Engine/Private/ShaderCompiler/ShaderCompilerMemoryLimit.cpp b/Engine/Source/Runtime/Engine/Private/ShaderCompiler/ShaderCompilerMemoryLimit.cpp
--- a/Engine/Source/Runtime/Engine/Private/ShaderCompiler/ShaderCompilerMemoryLimit.cpp
+++ b/Engine/Source/Runtime/Engine/Private/ShaderCompiler/ShaderCompilerMemoryLimit.cpp
@@ -17,7 +17,10 @@
 
 void FWindowsResourceProcessTreeMemory::AddRootProcessId(uint32 ProcessId)
 {
-	RootProcessIds.AddUnique(ProcessId);
+	if (ProcessId != 0 )
+	{
+		RootProcessIds.Add(ProcessId);
+	}
 }
 
 void FWindowsResourceProcessTreeMemory::Reset()
@@ -101,7 +104,18 @@ void FWindowsResourceProcessTreeMemory::CollectTreeProcessIds()
 
 void FWindowsResourceProcessTreeMemory::CollectTreeProcessIdsRecurse(uint32 RootProcessId)
 {
-	TreeProcessIds.Push(RootProcessId);
+	if (RootProcessId == 0)
+	{
+		return;
+	}
+
+	bool AlreadyVisited = false;
+	TreeProcessIds.Add(RootProcessId, &AlreadyVisited);
+
+	if (AlreadyVisited)
+	{
+		return;
+	}
 
 	for (const FProcessInfo& JobProcessInfo : AllProcesses)
 	{
diff --git a/Engine/Source/Runtime/Engine/Private/ShaderCompiler/ShaderCompilerMemoryLimit.h b/Engine/Source/Runtime/Engine/Private/ShaderCompiler/ShaderCompilerMemoryLimit.h
--- a/Engine/Source/Runtime/Engine/Private/ShaderCompiler/ShaderCompilerMemoryLimit.h
+++ b/Engine/Source/Runtime/Engine/Private/ShaderCompiler/ShaderCompilerMemoryLimit.h
@@ -37,13 +37,13 @@ private:
 	};
 
 	/* The Process IDs we will measure */
-	TArray<uint32> RootProcessIds;
+	TSet<uint32> RootProcessIds;
 
 	/* All Processes in the system */
 	TArray<FProcessInfo> AllProcesses;
 
 	/* The root Process IDs and all child Process IDs */
-	TArray<uint32> TreeProcessIds;
+	TSet<uint32> TreeProcessIds;
 
 	/* Populate the TreeProcessIds with RootProcessIds and their children. */
 	void CollectTreeProcessIds();
