Substrate - clarify requirements for opaque rough refractions.

#tests pc editor dx12
#rnx

diff --git a/Engine/Source/Runtime/Engine/Classes/Engine/RendererSettings.h b/Engine/Source/Runtime/Engine/Classes/Engine/RendererSettings.h
--- a/Engine/Source/Runtime/Engine/Classes/Engine/RendererSettings.h
+++ b/Engine/Source/Runtime/Engine/Classes/Engine/RendererSettings.h
@@ -1260,7 +1260,8 @@ class URendererSettings : public UDeveloperSettings
 	*/
 	UPROPERTY(config, EditAnywhere, Category = Substrate, meta = (
 		ConsoleVariable = "r.Substrate.OpaqueMaterialRoughRefraction", DisplayName = "Substrate opaque material rough refraction (Experimental)",
-		ToolTip = "Enable Substrate opaque material rough refractions effect from top layers over layers below. (Experimental for now since this is not a path validated in production)",
+		ToolTip = "Enable Substrate opaque material rough refractions effect from top layers over layers below. (Experimental for now since this is not a path validated in production). Requires adaptive GBuffer format and at least 2 closures per pixel.",
+		EditCondition = "SubstrateGBufferFormat == 1 && SubstrateProjectClosuresPerPixel > 1",
 		ConfigRestartRequired = true))
 		uint32 SubstrateOpaqueMaterialRoughRefraction : 1;
 		
diff --git a/Engine/Source/Runtime/RenderCore/Private/RenderUtils.cpp b/Engine/Source/Runtime/RenderCore/Private/RenderUtils.cpp
--- a/Engine/Source/Runtime/RenderCore/Private/RenderUtils.cpp
+++ b/Engine/Source/Runtime/RenderCore/Private/RenderUtils.cpp
@@ -2311,7 +2311,8 @@ namespace Substrate
 	bool IsOpaqueRoughRefractionEnabled(EShaderPlatform InPlatform)
 	{
 		static FShaderPlatformCachedIniValue<int32> CVar(TEXT("r.Substrate.OpaqueMaterialRoughRefraction"));
-		return IsSubstrateEnabled() && CVar.Get(InPlatform) > 0 && !IsSubstrateBlendableGBufferEnabled(InPlatform);
+		const bool bManyClosures = GetClosurePerPixel(InPlatform) > 1; // Opaque Rough refraction can only happen when there is at least 2 closures.
+		return IsSubstrateEnabled() && CVar.Get(InPlatform) > 0 && !IsSubstrateBlendableGBufferEnabled(InPlatform) && bManyClosures;
 	}
 
 	bool IsAdvancedVisualizationEnabled()
