Do not virtualize active sound if referenced sound is invalid

diff --git a/Engine/Source/Runtime/Engine/Private/AudioDevice.cpp b/Engine/Source/Runtime/Engine/Private/AudioDevice.cpp
--- a/Engine/Source/Runtime/Engine/Private/AudioDevice.cpp
+++ b/Engine/Source/Runtime/Engine/Private/AudioDevice.cpp
@@ -3273,6 +3273,11 @@ void FAudioDevice::VirtualizeInactiveLoops()
 			continue;
 		}
 
+		if (!ActiveSound->Sound)
+		{
+			continue;
+		}
+
 		ActiveSoundsToVirtualize.Add(ActiveSound);
 	}
 
