[Player Input] Fix a bug where if you had an input component whose outer was a subobject of the currently controlled pawn, the input component would be processed twice.

One example where this could be a common issue is within the "Modular Gameplay Framework". Users often times would add a UInputComponent to some UActorComponent, and then call "PushInputComponent" on the owning player controller. This would result in the UInputComponent being processed twice. This would happen because in "APlayerController::BuildInputStack" it calls "GetComponents" on the controlled pawn, which would return the UInputComponent which had already been pushed to the controller's input stack, thus adding it twice to the stack for processing.

diff --git a/Engine/Source/Runtime/Engine/Private/PlayerController.cpp b/Engine/Source/Runtime/Engine/Private/PlayerController.cpp
--- a/Engine/Source/Runtime/Engine/Private/PlayerController.cpp
+++ b/Engine/Source/Runtime/Engine/Private/PlayerController.cpp
@@ -2702,7 +2702,11 @@ void APlayerController::BuildInputStack(TArray<UInputComponent*>& InputStack)
 			for (UActorComponent* ActorComponent : ControlledPawn->GetComponents())
 			{
 				UInputComponent* PawnInputComponent = Cast<UInputComponent>(ActorComponent);
-				if (PawnInputComponent && PawnInputComponent != ControlledPawn->InputComponent)
+
+				// Do not add this input component it is already contained on the "CurrentInputStack". Input components which are on the CurrentInputStack
+				// already, will have been specifically added via the PushInputComponent function already. Adding them again here would result in them being
+				// processed twice
+				if (PawnInputComponent && PawnInputComponent != ControlledPawn->InputComponent && !CurrentInputStack.Contains(PawnInputComponent))
 				{
 					InputStack.Push(PawnInputComponent);
 				}
