[Player Input] Sort the input component stack based on the priority while building it. This resolves an issue where only the priorities of the "CurrentInputStack" would be respected, and not any components which changed their priority on a pawn subobject or level script actor.

This can be turned off with the "PlayerController.Input.SortInputStackByPriority" CVar, in case your project was relying on the previous ordering behavior. This will be removed in a future release and is intended for use as users update their projects to have proper priority set on the components

diff --git a/Engine/Source/Runtime/Engine/Private/PlayerController.cpp b/Engine/Source/Runtime/Engine/Private/PlayerController.cpp
--- a/Engine/Source/Runtime/Engine/Private/PlayerController.cpp
+++ b/Engine/Source/Runtime/Engine/Private/PlayerController.cpp
@@ -115,6 +115,13 @@ namespace UE::Gameplay::CVars
 		TEXT("If true, APlayerController's replicated properties will use push-based networking, and will therefore need to be marked dirty when changed."),
 		ECVF_Default);
 
+	bool bSortInputStackByPriority = true;
+	static FAutoConsoleVariableRef CVarSortInputStackByPriority(
+		TEXT("PlayerController.Input.SortInputStackByPriority"), 
+		bSortInputStackByPriority,
+		TEXT("If true, APlayerController's input component stack will be sorted by priority after being built. False will maintain 'legacy' ordering."),
+		ECVF_Default);
+
 	extern bool bAlwaysNotifyClientOnControllerChange;
 }
 
@@ -2745,6 +2752,17 @@ void APlayerController::BuildInputStack(TArray<UInputComponent*>& InputStack)
 			CurrentInputStack.RemoveAt(Idx--);
 		}
 	}
+
+	// We need to sort the input stack in order to respect the priority of each input component.
+	// This way, all input components and their priorities will have their order respected. Not just those
+	// manually added via "PushInputComponent"
+	if (UE::Gameplay::CVars::bSortInputStackByPriority)
+	{
+		Algo::StableSort(InputStack, [](const UInputComponent* A, const UInputComponent* B)
+			{
+				return A->Priority < B->Priority;
+			});	
+	}
 }
 
 void APlayerController::ProcessPlayerInput(const float DeltaTime, const bool bGamePaused)
