Added Angle Spring Damper functions.

#virtualized

diff --git a/Engine/Source/Runtime/Engine/Classes/Kismet/BlueprintSpringMathLibrary.h b/Engine/Source/Runtime/Engine/Classes/Kismet/BlueprintSpringMathLibrary.h
--- a/Engine/Source/Runtime/Engine/Classes/Kismet/BlueprintSpringMathLibrary.h
+++ b/Engine/Source/Runtime/Engine/Classes/Kismet/BlueprintSpringMathLibrary.h
@@ -26,6 +26,18 @@ public:
 	UE_EXPERIMENTAL(5.7, "SpringMath is experimental")
 	static ENGINE_API void CriticalSpringDampFloat(UPARAM(ref) float& InOutX, UPARAM(ref) float& InOutV, const float& TargetX, float DeltaTime, float SmoothingTime = 0.2f);
 
+	/** Interpolates the angle InOutAngle towards TargetAngle with the motion of a critically damped spring. The velocity of InOutAngle is stored in InOutAngularVelocity in deg/s.
+	 *
+	 * @param InOutAngle The value to be damped in degrees
+	 * @param InOutAngularVelocity The speed of the value to be damped in deg/s
+	 * @param TargetAngle The goal to damp towards in degrees
+	 * @param SmoothingTime The smoothing time to use for the spring. Longer times lead to more damped behaviour. A time of 0 will snap InOutAngle to TargetX
+	 * @param DeltaTime Timestep in seconds
+	 */
+	UFUNCTION(BlueprintCallable, meta = (Category = "Math|Spring|Experimental"))
+	UE_EXPERIMENTAL(5.7, "SpringMath is experimental")
+	static ENGINE_API void CriticalSpringDampAngle(UPARAM(ref) float& InOutAngle, UPARAM(ref) float& InOutAngularVelocity, const float& TargetAngle, float DeltaTime, float SmoothingTime = 0.2f);
+
 	/** Interpolates the value InOutX towards TargetX with the motion of a critically damped spring. The velocity of X is stored in InOutV.
 	 * 
 	 * @param InOutX The value to be damped
@@ -50,24 +62,24 @@ public:
 	UFUNCTION(BlueprintCallable, meta=(Category="Math|Spring|Experimental"))
 	static ENGINE_API void CriticalSpringDampVector2D(UPARAM(ref) FVector2D& InOutX, UPARAM(ref) FVector2D& InOutV, const FVector2D& TargetX, float DeltaTime, float SmoothingTime = 0.2f);
 
-	/** Interpolates the value InOutX towards TargetX with the motion of a critically damped spring. The velocity of X is stored in InOutV.
-	 * 
-	 * @param InOutX The value to be damped
-	 * @param InOutV The speed of the value to be damped
-	 * @param TargetX The goal to damp towards
-	 * @param SmoothingTime The smoothing time to use for the spring. Longer times lead to more damped behaviour. A time of 0 will snap InOutX to TargetX
+	/** Interpolates the value InOutRotation towards TargetRotation with the motion of a critically damped spring. The velocity of InOutRotation is stored in InOutAngularVelocity in deg/s.
+	 *
+	 * @param InOutRotation The value to be damped
+	 * @param InOutAngularVelocity The speed of the value to be damped in deg/s
+	 * @param TargetRotation The goal to damp towards
+	 * @param SmoothingTime The smoothing time to use for the spring. Longer times lead to more damped behaviour. A time of 0 will snap InOutRotation to TargetRotation
 	 * @param DeltaTime Timestep in seconds
 	 */
 	UE_EXPERIMENTAL(5.7, "SpringMath is experimental")
 	UFUNCTION(BlueprintCallable, meta=(Category="Math|Spring|Experimental"))
 	static ENGINE_API void CriticalSpringDampQuat(UPARAM(ref) FQuat& InOutRotation, UPARAM(ref) FVector& InOutAngularVelocity, const FQuat& TargetRotation, float DeltaTime, float SmoothingTime = 0.2f);
 
-	/** Interpolates the value InOutX towards TargetX with the motion of a critically damped spring. The velocity of X is stored in InOutV.
-	 * 
-	 * @param InOutX The value to be damped
-	 * @param InOutV The speed of the value to be damped
-	 * @param TargetX The goal to damp towards
-	 * @param SmoothingTime The smoothing time to use for the spring. Longer times lead to more damped behaviour. A time of 0 will snap InOutX to TargetX
+	/** Interpolates the value InOutRotation towards TargetRotation with the motion of a critically damped spring. The velocity of InOutRotation is stored in InOutAngularVelocity in deg/s.
+	 *
+	 * @param InOutRotation The value to be damped
+	 * @param InOutAngularVelocity The speed of the value to be damped in deg/s
+	 * @param TargetRotation The goal to damp towards
+	 * @param SmoothingTime The smoothing time to use for the spring. Longer times lead to more damped behaviour. A time of 0 will snap InOutRotation to TargetRotation
 	 * @param DeltaTime Timestep in seconds
 	 */
 	UE_EXPERIMENTAL(5.7, "SpringMath is experimental")
@@ -134,6 +146,19 @@ public:
 	UFUNCTION(BlueprintPure, meta=(Category="Math|Spring|Experimental"))
 	static ENGINE_API float DampFloat(float Value, float Target, float DeltaTime, float SmoothingTime = 0.2f);
 
+	/**
+	 * Smooths an angle in degrees using exponential damping towards a target.
+	 *
+	 *
+	 * @param  Angle The angle to be smoothed in degrees
+	 * @param  TargetAngle The target to smooth towards in degrees
+	 * @param  DeltaTime Time interval
+	 * @param  SmoothingTime Timescale over which to smooth. Larger values result in more smoothed behaviour. Can be zero.
+	 */
+	UE_EXPERIMENTAL(5.7, "SpringMath is experimental")
+	UFUNCTION(BlueprintPure, meta = (Category = "Math|Spring|Experimental"))
+	static ENGINE_API float DampAngle(float Angle, float TargetAngle, float DeltaTime, float SmoothingTime = 0.2f);
+
 	/**
 	 * Smooths a value using exponential damping towards a target.
 	 * 
@@ -161,30 +186,30 @@ public:
 	static ENGINE_API FVector2D DampVector2D(const FVector2D& Value, const FVector2D& Target, float DeltaTime, float SmoothingTime = 0.2f);
 
 	/**
-	 * Smooths a value using exponential damping towards a target.
+	 * Smooths a rotation using exponential damping towards a target.
 	 * 
 	 * 
-	 * @param  Value The value to be smoothed
-	 * @param  Target The target to smooth towards
+	 * @param  Rotation The rotation to be smoothed
+	 * @param  TargetRotation The target to smooth towards
 	 * @param  DeltaTime Time interval
 	 * @param  SmoothingTime Timescale over which to smooth. Larger values result in more smoothed behaviour. Can be zero.
 	 */
 	UE_EXPERIMENTAL(5.7, "SpringMath is experimental")
 	UFUNCTION(BlueprintPure, meta=(Category="Math|Spring|Experimental"))
-	static ENGINE_API FQuat DampQuat(const FQuat& Value, const FQuat& Target, float DeltaTime, float SmoothingTime = 0.2f);
+	static ENGINE_API FQuat DampQuat(const FQuat& Rotation, const FQuat& TargetRotation, float DeltaTime, float SmoothingTime = 0.2f);
 
 	/**
 	 * Smooths a value using exponential damping towards a target.
 	 * 
 	 * 
-	 * @param  Value The value to be smoothed
-	 * @param  Target The target to smooth towards
+	 * @param  Rotation The rotation to be smoothed
+	 * @param  TargetRotation The target to smooth towards
 	 * @param  DeltaTime Time interval
 	 * @param  SmoothingTime Timescale over which to smooth. Larger values result in more smoothed behaviour. Can be zero.
 	 */
 	UE_EXPERIMENTAL(5.7, "SpringMath is experimental")
 	UFUNCTION(BlueprintPure, meta=(Category="Math|Spring|Experimental"))
-	static ENGINE_API FRotator DampRotator(const FRotator& Value, const FRotator& Target, float DeltaTime, float SmoothingTime = 0.2f);
+	static ENGINE_API FRotator DampRotator(const FRotator& Rotation, const FRotator& TargetRotation, float DeltaTime, float SmoothingTime = 0.2f);
 
 	/** Update the position of a character given a target velocity using a simple damped spring
 	 * 
diff --git a/Engine/Source/Runtime/Engine/Private/BlueprintSpringMathLibrary.cpp b/Engine/Source/Runtime/Engine/Private/BlueprintSpringMathLibrary.cpp
--- a/Engine/Source/Runtime/Engine/Private/BlueprintSpringMathLibrary.cpp
+++ b/Engine/Source/Runtime/Engine/Private/BlueprintSpringMathLibrary.cpp
@@ -21,18 +21,31 @@ void UBlueprintSpringMathLibrary::CriticalSpringDampFloat(float& InOutX, float&
 	SpringMath::CriticalSpringDamper(InOutX, InOutV, TargetX, SmoothingTime, DeltaTime);
 }
 
+void UBlueprintSpringMathLibrary::CriticalSpringDampAngle(float& InOutAngle, float& InOutAngularVelocity, const float& TargetAngle, float DeltaTime, float SmoothingTime)
+{
+	float InOutAngleRadians = FMath::DegreesToRadians(InOutAngle);
+	float InOutAngularVelocityRadians = FMath::DegreesToRadians(InOutAngularVelocity);
+	SpringMath::CriticalSpringDamperAngle(InOutAngleRadians, InOutAngularVelocityRadians, FMath::DegreesToRadians(TargetAngle), SmoothingTime, DeltaTime);
+	InOutAngle = FMath::RadiansToDegrees(InOutAngleRadians);
+	InOutAngularVelocity = FMath::RadiansToDegrees(InOutAngularVelocityRadians);
+}
+
 void UBlueprintSpringMathLibrary::CriticalSpringDampQuat(FQuat& InOutRotation, FVector& InOutAngularVelocity, const FQuat& TargetRotation,
 	float DeltaTime, float SmoothingTime)
 {
-	SpringMath::CriticalSpringDamperQuat(InOutRotation, InOutAngularVelocity, TargetRotation, SmoothingTime, DeltaTime);
+	FVector InOutAngularVelocityRadians = FMath::DegreesToRadians(InOutAngularVelocity);
+	SpringMath::CriticalSpringDamperQuat(InOutRotation, InOutAngularVelocityRadians, TargetRotation, SmoothingTime, DeltaTime);
+	InOutAngularVelocity = FMath::RadiansToDegrees(InOutAngularVelocityRadians);
 }
 
 void UBlueprintSpringMathLibrary::CriticalSpringDampRotator(FRotator& InOutRotation, FVector& InOutAngularVelocity, const FRotator& TargetRotation,
 	float DeltaTime, float SmoothingTime)
 {
 	FQuat InOutRotationQuat = InOutRotation.Quaternion();
-	SpringMath::CriticalSpringDamperQuat(InOutRotationQuat, InOutAngularVelocity, TargetRotation.Quaternion(), SmoothingTime, DeltaTime);
+	FVector InOutAngularVelocityRadians = FMath::DegreesToRadians(InOutAngularVelocity);
+	SpringMath::CriticalSpringDamperQuat(InOutRotationQuat, InOutAngularVelocityRadians, TargetRotation.Quaternion(), SmoothingTime, DeltaTime);
 	InOutRotation = InOutRotationQuat.Rotator();
+	InOutAngularVelocity = FMath::RadiansToDegrees(InOutAngularVelocityRadians);
 }
 
 void UBlueprintSpringMathLibrary::VelocitySpringDampFloat(float& InOutX, float& InOutV, float& InOutVi, float TargetX, float MaxSpeed,
@@ -78,6 +91,13 @@ float UBlueprintSpringMathLibrary::DampFloat(float Value, float Target, float De
 	return DampedValue;
 }
 
+float UBlueprintSpringMathLibrary::DampAngle(float Angle, float TargetAngle, float DeltaTime, float SmoothingTime)
+{
+	float DampedAngle = FMath::DegreesToRadians(Angle);
+	SpringMath::ExponentialSmoothingApproxAngle(DampedAngle, FMath::DegreesToRadians(TargetAngle), DeltaTime, SmoothingTime);
+	return FMath::RadiansToDegrees(DampedAngle);
+}
+
 FVector UBlueprintSpringMathLibrary::DampVector(const FVector& Value, const FVector& Target, float DeltaTime, float SmoothingTime)
 {
 	FVector DampedValue = Value;
@@ -92,18 +112,18 @@ FVector2D UBlueprintSpringMathLibrary::DampVector2D(const FVector2D& Value, cons
 	return DampedValue;
 }
 
-FQuat UBlueprintSpringMathLibrary::DampQuat(const FQuat& Value, const FQuat& Target, float DeltaTime, float SmoothingTime)
+FQuat UBlueprintSpringMathLibrary::DampQuat(const FQuat& Rotation, const FQuat& TargetRotation, float DeltaTime, float SmoothingTime)
 {
-	FQuat DampedValue = Value;
-	SpringMath::ExponentialSmoothingApproxQuat(DampedValue, Target, DeltaTime, SmoothingTime);
-	return DampedValue;
+	FQuat DampedRotation = Rotation;
+	SpringMath::ExponentialSmoothingApproxQuat(DampedRotation, TargetRotation, DeltaTime, SmoothingTime);
+	return DampedRotation;
 }
 
-FRotator UBlueprintSpringMathLibrary::DampRotator(const FRotator& Value, const FRotator& Target, float DeltaTime, float SmoothingTime)
+FRotator UBlueprintSpringMathLibrary::DampRotator(const FRotator& Rotation, const FRotator& TargetRotation, float DeltaTime, float SmoothingTime)
 {
-	FQuat DampedValue = Value.Quaternion();
-	SpringMath::ExponentialSmoothingApproxQuat(DampedValue, Target.Quaternion(), DeltaTime, SmoothingTime);
-	return DampedValue.Rotator();
+	FQuat DampedRotation = Rotation.Quaternion();
+	SpringMath::ExponentialSmoothingApproxQuat(DampedRotation, TargetRotation.Quaternion(), DeltaTime, SmoothingTime);
+	return DampedRotation.Rotator();
 }
 
 void UBlueprintSpringMathLibrary::SpringCharacterUpdate(FVector& InOutPosition, FVector& InOutVelocity, FVector& InOutAcceleration,
diff --git a/Engine/Source/Runtime/Engine/Public/Animation/SpringMath.h b/Engine/Source/Runtime/Engine/Public/Animation/SpringMath.h
--- a/Engine/Source/Runtime/Engine/Public/Animation/SpringMath.h
+++ b/Engine/Source/Runtime/Engine/Public/Animation/SpringMath.h
@@ -57,20 +57,19 @@ public:
 	 * @return The smoothing time of the spring in seconds. It is the time by which the output lags the input when critically damped, when the input is changing at a constant rate.
 	 */
 	UE_EXPERIMENTAL(5.7, "SpringMath is experimental")
-
 	static constexpr float StrengthToSmoothingTime(float Strength)
 	{
 		return 2.0f / FMath::Max(Strength, UE_KINDA_SMALL_NUMBER);
 	}
 
 	/** Simplified version of FMath::CriticallyDampedSmoothing where v_goal is assumed to be 0. This interpolates the value InOutX towards TargetX
-	 * with the motion of a critically damped spring. The velocity of X is stored in InOutV.
+	 * with the motion of a critically damped spring. The velocity of InOutX is stored in InOutV.
 	 * 
 	 * @tparam T The type to be damped 
 	 * @param InOutX The value to be damped
 	 * @param InOutV The speed of the value to be damped
 	 * @param TargetX The goal to damp towards
-	 * @param SmoothingTime The smoothing time to use for the spring. Longer times lead to more damped behaviour. A time of 0 will snap X to XGoal
+	 * @param SmoothingTime The smoothing time to use for the spring. Longer times lead to more damped behaviour. A time of 0 will snap InOutX to TargetX
 	 * @param DeltaTime Timestep in seconds
 	 */
 	template <typename T>
@@ -91,13 +90,41 @@ public:
 		InOutV = EyDt * (InOutV - J1 * Y * DeltaTime);
 	}
 
+	/** Specialized angle version of CriticalSpringDamper that handles angle wrapping.
+	 *
+	 * @param InOutAngleRadians The value to be damped
+	 * @param InOutAngularVelocityRadians The speed of the value to be damped
+	 * @param TargetAngleRadians The goal to damp towards
+	 * @param SmoothingTime The smoothing time to use for the spring. Longer times lead to more damped behaviour. A time of 0 will snap InOutAngleRadians to TargetAngleRadians
+	 * @param DeltaTime Timestep in seconds
+	 */
+	UE_EXPERIMENTAL(5.7, "SpringMath is experimental")
+	static void CriticalSpringDamperAngle(
+		float& InOutAngleRadians,
+		float& InOutAngularVelocityRadians,
+		float TargetAngleRadians,
+		float SmoothingTime,
+		float DeltaTime)
+	{
+		PRAGMA_DISABLE_EXPERIMENTAL_WARNINGS
+		float Y = SmoothingTimeToDamping(SmoothingTime) / 2.0f;
+		PRAGMA_ENABLE_EXPERIMENTAL_WARNINGS
+
+		float J0 = FMath::FindDeltaAngleRadians(TargetAngleRadians, InOutAngleRadians);
+		float J1 = InOutAngularVelocityRadians + J0 * Y;
+		float EyDt = FMath::InvExpApprox(Y * DeltaTime);
+
+		InOutAngleRadians = EyDt * (J0 + J1 * DeltaTime) + TargetAngleRadians;
+		InOutAngularVelocityRadians = EyDt * (InOutAngularVelocityRadians - J1 * Y * DeltaTime);
+	}
+
 	/** Specialized quaternion version of CriticalSpringDamper, uses FVector for angular velocity
-	 * 
+	 *
 	 * @param InOutRotation The value to be damped
 	 * @param InOutAngularVelocityRadians The angular velocity of the rotation in radians
 	 * @param TargetRotation The target rotation to damp towards
 	 * @param SmoothingTime he smoothing time to use for the spring. Longer times lead to more damped behaviour. A time of 0 will snap InOutRotation to TargetRotation
-	 * @param DeltaTime 
+	 * @param DeltaTime
 	 */
 	UE_EXPERIMENTAL(5.7, "SpringMath is experimental")
 	static void CriticalSpringDamperQuat(
@@ -166,7 +193,7 @@ public:
 	}
 
 	/** A velocity spring will damp towards a target that follows a fixed linear target velocity, allowing control of the interpolation speed
-	* while still giving a smoothed beheaviour. A SmoothingTime of 0 will give a linear interpolation between X and TargetX
+	* while still giving a smoothed behaviour. A SmoothingTime of 0 will give a linear interpolation between X and TargetX
 	 * 
 	 * @tparam TVector The type of vector to use (e.g. FVector2D, FVector4, FVector etc)
 	 * @param InOutX The value to be damped
@@ -492,4 +519,25 @@ public:
 			InOutRotation = InTargetRotation;
 		}
 	}
+
+	/** Specialized angle damper, similar to FMath::ExponentialSmoothingApprox but deals correctly with angle wrap-around.
+	 *  Smooths an angle using exponential damping towards a target.
+	 *
+	 * @param InOutAngleRadians		The angle to be smoothed
+	 * @param InTargetAngleRadians	The target to smooth towards
+	 * @param InDeltaTime			Time interval
+	 * @param InSmoothingTime		Timescale over which to smooth. Larger values result in more smoothed behaviour. Can be zero.
+	 */
+	UE_EXPERIMENTAL(5.7, "SpringMath is experimental")
+	static constexpr void ExponentialSmoothingApproxAngle(float& InOutAngleRadians, const float& InTargetAngleRadians, const float InDeltaTime, const float InSmoothingTime)
+	{
+		if (InSmoothingTime > UE_KINDA_SMALL_NUMBER)
+		{
+			InOutAngleRadians += FMath::FindDeltaAngleRadians(InOutAngleRadians, InTargetAngleRadians) * (1.0f - FMath::InvExpApprox(InDeltaTime / InSmoothingTime));
+		}
+		else
+		{
+			InOutAngleRadians = InTargetAngleRadians;
+		}
+	}
 };
