[HWRT] Fix HISM missing from ray tracing scene.
- regression caused by CL 44629078 since FHierarchicalStaticMeshSceneProxy generates dynamic ray tracing instances even when using static ray tracing geometry (not using WPO).
- added logic to properly support this use case.

diff --git a/Engine/Source/Runtime/Engine/Private/HierarchicalInstancedStaticMesh.cpp b/Engine/Source/Runtime/Engine/Private/HierarchicalInstancedStaticMesh.cpp
--- a/Engine/Source/Runtime/Engine/Private/HierarchicalInstancedStaticMesh.cpp
+++ b/Engine/Source/Runtime/Engine/Private/HierarchicalInstancedStaticMesh.cpp
@@ -809,6 +809,10 @@ FHierarchicalStaticMeshSceneProxy::FHierarchicalStaticMeshSceneProxy(UHierarchic
 	bIsHierarchicalInstancedStaticMesh = true;
 	bIsLandscapeGrass = (ViewRelevance == EHISMViewRelevanceType::Grass);
 
+#if RHI_RAYTRACING
+	bDynamicRayTracingInstances = true;
+#endif
+
 	// Store LODDistanceScale so it can be used in FInstancedStaticMeshVertexFactoryShaderParameters::GetElementShaderBindings when dither LOD transitions are enabled
 	float LODDistanceScale = InComponent->InstanceLODDistanceScale * CVarFoliageLODDistanceScale.GetValueOnGameThread();
 	UserData_AllInstances.LODDistanceScale = LODDistanceScale;
diff --git a/Engine/Source/Runtime/Engine/Private/InstancedStaticMesh.cpp b/Engine/Source/Runtime/Engine/Private/InstancedStaticMesh.cpp
--- a/Engine/Source/Runtime/Engine/Private/InstancedStaticMesh.cpp
+++ b/Engine/Source/Runtime/Engine/Private/InstancedStaticMesh.cpp
@@ -1750,7 +1750,7 @@ void FInstancedStaticMeshSceneProxy::GetDynamicRayTracingInstances(FRayTracingIn
 	// Workaround: SetEvaluateWorldPositionOffsetInRayTracing(...) calls UpdateCachedRayTracingState(...)
 	// however the update only happens after gathering relevant ray tracing primitives
 	// so ERayTracingPrimitiveFlags::Dynamic is set for one frame after the WPO evaluation is disabled.
-	if (!bDynamicRayTracingGeometry)
+	if (!bDynamicRayTracingInstances && !bDynamicRayTracingGeometry)
 	{
 		return;
 	}
@@ -1849,6 +1849,7 @@ void FInstancedStaticMeshSceneProxy::GetDynamicRayTracingInstances(FRayTracingIn
 	const uint32 RequestedSimulatedInstances = CVarRayTracingSimulatedInstanceCount.GetValueOnRenderThread();
 	const uint32 SimulatedInstances = FMath::Min(RequestedSimulatedInstances == -1 ? InstanceCount : FMath::Clamp(RequestedSimulatedInstances, 1u, InstanceCount), MaxSimulatedInstances);
 
+	if (bDynamicRayTracingGeometry)
 	{
 		int32 SectionCount = InstancedRenderData.LODModels[LODIndex].Sections.Num();
 
@@ -1889,7 +1890,7 @@ void FInstancedStaticMeshSceneProxy::GetDynamicRayTracingInstances(FRayTracingIn
 
 		ActiveInstances.AddZeroed(SimulatedInstances);
 
-		for (auto &Instance : ActiveInstances)
+		for (auto& Instance : ActiveInstances)
 		{
 			Instance = INDEX_NONE;
 		}
@@ -2001,7 +2002,7 @@ void FInstancedStaticMeshSceneProxy::GetDynamicRayTracingInstances(FRayTracingIn
 	}
 
 	// Bucket which visible instances to simulate based on camera distance
-	if (VisibleInstances.Num() > 0)
+	if (bDynamicRayTracingGeometry && VisibleInstances.Num() > 0)
 	{
 		float SimulationClusterRadius = CVarRayTracingInstancesSimulationClusterRadius.GetValueOnRenderThread();
 
@@ -2056,7 +2057,7 @@ void FInstancedStaticMeshSceneProxy::GetDynamicRayTracingInstances(FRayTracingIn
 		FMatrix InstanceToWorld = InstanceSceneDataBuffers->GetInstanceToWorld(InstanceIndex);
 		const uint32 DynamicInstanceIdx = InstanceIndex % SimulatedInstances;
 
-		if (InstancedRenderData.VertexFactories[LODIndex].GetType()->SupportsRayTracingDynamicGeometry())
+		if (bDynamicRayTracingGeometry && InstancedRenderData.VertexFactories[LODIndex].GetType()->SupportsRayTracingDynamicGeometry())
 		{
 			FRayTracingInstance* DynamicInstance = nullptr;
 
diff --git a/Engine/Source/Runtime/Engine/Private/StaticMeshSceneProxy.cpp b/Engine/Source/Runtime/Engine/Private/StaticMeshSceneProxy.cpp
--- a/Engine/Source/Runtime/Engine/Private/StaticMeshSceneProxy.cpp
+++ b/Engine/Source/Runtime/Engine/Private/StaticMeshSceneProxy.cpp
@@ -373,6 +373,7 @@ FStaticMeshSceneProxy::FStaticMeshSceneProxy(const FStaticMeshSceneProxyDesc& In
 #if RHI_RAYTRACING
 	bSupportRayTracing = IsRayTracingAllowed() && InProxyDesc.GetStaticMesh()->bSupportRayTracing;
 	bHasRayTracingRepresentation = false;
+	bDynamicRayTracingInstances = false;
 	bDynamicRayTracingGeometry = false;
 	bCreateDynamicRayTracingGeometries = true;
 
diff --git a/Engine/Source/Runtime/Engine/Public/HierarchicalStaticMeshSceneProxy.h b/Engine/Source/Runtime/Engine/Public/HierarchicalStaticMeshSceneProxy.h
--- a/Engine/Source/Runtime/Engine/Public/HierarchicalStaticMeshSceneProxy.h
+++ b/Engine/Source/Runtime/Engine/Public/HierarchicalStaticMeshSceneProxy.h
@@ -71,13 +71,6 @@ public:
 	
 	ENGINE_API virtual FPrimitiveViewRelevance GetViewRelevance(const FSceneView* View) const override;
 
-#if RHI_RAYTRACING
-	virtual bool IsRayTracingStaticRelevant() const override
-	{
-		return false;
-	}
-#endif
-
 	ENGINE_API virtual void GetDynamicMeshElements(const TArray<const FSceneView*>& Views, const FSceneViewFamily& ViewFamily, uint32 VisibilityMap, FMeshElementCollector& Collector) const override;
 
 	ENGINE_API virtual const TArray<FBoxSphereBounds>* GetOcclusionQueries(const FSceneView* View) const override;
diff --git a/Engine/Source/Runtime/Engine/Public/StaticMeshSceneProxy.h b/Engine/Source/Runtime/Engine/Public/StaticMeshSceneProxy.h
--- a/Engine/Source/Runtime/Engine/Public/StaticMeshSceneProxy.h
+++ b/Engine/Source/Runtime/Engine/Public/StaticMeshSceneProxy.h
@@ -177,7 +177,7 @@ public:
 	virtual bool IsRayTracingStaticRelevant() const override 
 	{ 
 		const bool bAllowStaticLighting = IsStaticLightingAllowed();
-		const bool bIsStaticInstance = !bDynamicRayTracingGeometry;
+		const bool bIsStaticInstance = !bDynamicRayTracingInstances && !bDynamicRayTracingGeometry;
 		return bIsStaticInstance && !HasViewDependentDPG() && !(bAllowStaticLighting && HasStaticLighting() && !HasValidSettingsForStaticLighting());
 	}
 	ENGINE_API virtual RayTracing::FGeometryGroupHandle GetRayTracingGeometryGroupHandle() const override;
@@ -283,6 +283,7 @@ protected:
 
 	bool bSupportRayTracing : 1;
 	bool bHasRayTracingRepresentation : 1;
+	bool bDynamicRayTracingInstances : 1;
 	bool bDynamicRayTracingGeometry : 1;
 	bool bCreateDynamicRayTracingGeometries : 1;
 
