World Partition - HLOD: Deprecated CellBounds member of FHLODCreationParams (now unused)

#DisabledValidations Preflight

diff --git a/Engine/Source/Runtime/Engine/Private/WorldPartition/RuntimeHashSet/WorldPartitionRuntimeHashSetHLODGeneration.cpp b/Engine/Source/Runtime/Engine/Private/WorldPartition/RuntimeHashSet/WorldPartitionRuntimeHashSetHLODGeneration.cpp
--- a/Engine/Source/Runtime/Engine/Private/WorldPartition/RuntimeHashSet/WorldPartitionRuntimeHashSetHLODGeneration.cpp
+++ b/Engine/Source/Runtime/Engine/Private/WorldPartition/RuntimeHashSet/WorldPartitionRuntimeHashSetHLODGeneration.cpp
@@ -381,7 +381,6 @@ bool UWorldPartitionRuntimeHashSet::SetupHLODActors(const IStreamingGenerationCo
 				HLODCreationParams.TargetWorld = WorldPartition->GetWorld();
 				HLODCreationParams.CellName = CellUniqueId.Name;
 				HLODCreationParams.CellGuid = CellUniqueId.Guid;
-				HLODCreationParams.CellBounds = CellDescInstance.CellBounds.GetValue();
 				HLODCreationParams.GetRuntimeGrid = GetHLODRuntimeGrid;
 				HLODCreationParams.HLODLevel = HLODLevel;
 				HLODCreationParams.MinVisibleDistance = RuntimePartition->LoadingRange;
diff --git a/Engine/Source/Runtime/Engine/Private/WorldPartition/RuntimeSpatialHash/RuntimeSpatialHashHLOD.cpp b/Engine/Source/Runtime/Engine/Private/WorldPartition/RuntimeSpatialHash/RuntimeSpatialHashHLOD.cpp
--- a/Engine/Source/Runtime/Engine/Private/WorldPartition/RuntimeSpatialHash/RuntimeSpatialHashHLOD.cpp
+++ b/Engine/Source/Runtime/Engine/Private/WorldPartition/RuntimeSpatialHash/RuntimeSpatialHashHLOD.cpp
@@ -214,7 +214,6 @@ static TArray<FGuid> GenerateHLODActorsForGrid(UWorldPartition* WorldPartition,
 				CreationParams.TargetWorld = WorldPartition->GetWorld();
 				CreationParams.CellGuid = CellGuid;
 				CreationParams.CellName = CellName;
-				CreationParams.CellBounds = CellBounds;
 				CreationParams.GetRuntimeGrid = [HLODLevel](const UHLODLayer* HLODLayer) 
 				{ 
 					PRAGMA_DISABLE_DEPRECATION_WARNINGS
diff --git a/Engine/Source/Runtime/Engine/Public/WorldPartition/HLOD/IWorldPartitionHLODUtilities.h b/Engine/Source/Runtime/Engine/Public/WorldPartition/HLOD/IWorldPartitionHLODUtilities.h
--- a/Engine/Source/Runtime/Engine/Public/WorldPartition/HLOD/IWorldPartitionHLODUtilities.h
+++ b/Engine/Source/Runtime/Engine/Public/WorldPartition/HLOD/IWorldPartitionHLODUtilities.h
@@ -33,7 +33,6 @@ struct FHLODCreationParams
 
 	FGuid CellGuid;
 	FString CellName;
-	FBox CellBounds;
 	TUniqueFunction<FName(const UHLODLayer*)> GetRuntimeGrid;
 	uint32 HLODLevel;
 	FGuid ContentBundleGuid;
@@ -51,6 +50,9 @@ struct FHLODCreationParams
 	}
 
 	double MinVisibleDistance;
+	
+	UE_DEPRECATED(5.7, "CellBounds member is not used anymore.")
+	FBox CellBounds;
 };
 
 /**
