[CustomHLOD] Add support for Custom HLOD Actors to replace HLOD actors in any HLOD partition.

* Enables use cases where users want to use generated "HLOD(N)" actors but substitute Custom HLOD in place of generated "HLOD(N+1)" actors.
* Adds a new LinkedLayer property to UHLODLayer, visible only when LayerType is set to Custom HLOD Actor
* LinkedLayer is used to control the visibility of Custom HLOD Actors. They become visible when actors from the LinkedLayer are unloaded.
* If LinkedLayer is not specified, Main Partition is used to control the visibility.

diff --git a/Engine/Source/Runtime/Engine/Private/WorldPartition/WorldPartitionStreamingGeneration.cpp b/Engine/Source/Runtime/Engine/Private/WorldPartition/WorldPartitionStreamingGeneration.cpp
--- a/Engine/Source/Runtime/Engine/Private/WorldPartition/WorldPartitionStreamingGeneration.cpp
+++ b/Engine/Source/Runtime/Engine/Private/WorldPartition/WorldPartitionStreamingGeneration.cpp
@@ -761,6 +761,22 @@ class FWorldPartitionStreamingGenerator
 		}
 	}
 
+	FName GetResolvedRuntimeGridForHLODLayer(FName RuntimeGrid, const FSoftObjectPath& HLODLayerPath) const
+	{
+		if (WorldPartitionContext)
+		{
+			if (const UWorldPartitionRuntimeHashSet* HashSet = Cast<UWorldPartitionRuntimeHashSet>(WorldPartitionContext->RuntimeHash))
+			{
+				if (const URuntimePartition* RuntimePartition = HashSet->ResolveRuntimePartitionForHLODLayer(RuntimeGrid, HLODLayerPath))
+				{
+					FName GridName = RuntimeGrid.IsNone() ? DefaultGrid : RuntimeGrid;
+					return FName(GridName.ToString() + TEXT(":") + RuntimePartition->Name.ToString());
+				}
+			}
+		}
+		return NAME_None;
+	}
+
 	void ResolveRuntimeDataLayers(FStreamingGenerationActorDescView& ActorDescView, const TArray<const FWorldDataLayersActorDesc*>& InDataLayerResolvers)
 	{
 		// Resolve DataLayerInstanceNames of ActorDescView only when necessary (i.e. when container is a template)
@@ -912,7 +928,23 @@ class FWorldPartitionStreamingGenerator
 
 				// Spawn placeholder actor and initialize properties based on Custom HLOD actor
 				AWorldPartitionCustomHLODPlaceholder* CustomHLODPlaceholderActor = WorldPartitionContext->GetWorld()->SpawnActor<AWorldPartitionCustomHLODPlaceholder>(SpawnParameters);
+				check(CustomHLODPlaceholderActor);
 				CustomHLODPlaceholderActor->InitFrom(*Iterator);
+
+				if (UHLODLayer* HLODLayer = Cast<UHLODLayer>(Iterator->GetHLODLayer().TryLoad()))
+				{
+					if (HLODLayer->GetLayerType() == EHLODLayerType::CustomHLODActor)
+					{
+						if (UHLODLayer* LinkedHLODLayer = HLODLayer->GetLinkedLayer())
+						{
+							FName ResolvedRuntimeGrid = GetResolvedRuntimeGridForHLODLayer(CustomHLODPlaceholderActor->GetRuntimeGrid(), LinkedHLODLayer);
+							if (!ResolvedRuntimeGrid.IsNone())
+							{
+								CustomHLODPlaceholderActor->SetRuntimeGrid(ResolvedRuntimeGrid);
+							}
+						}
+					}
+				}
 				
 				// Register placeholder actor desc view
 				TUniquePtr<FStreamingGenerationTempActorDescInstance>& CustomHLODPlaceholderRef = OutCustomHLODActorPlaceholderInstances.Add_GetRef(MakeUnique<FStreamingGenerationTempActorDescInstance>(Iterator->GetContainerInstance(), CustomHLODPlaceholderActor));
@@ -1076,16 +1108,10 @@ class FWorldPartitionStreamingGenerator
 				{
 					if (HLODLayer->GetLayerType() == EHLODLayerType::CustomHLODActor)
 					{
-						if (WorldPartitionContext)
+						FName ResolvedRuntimeGrid = GetResolvedRuntimeGridForHLODLayer(ResultPerInstanceData.RuntimeGrid, InActorDescView.GetHLODLayer());
+						if (!ResolvedRuntimeGrid.IsNone())
 						{
-							if (const UWorldPartitionRuntimeHashSet* HashSet = Cast<UWorldPartitionRuntimeHashSet>(WorldPartitionContext->RuntimeHash))
-							{
-								if (const URuntimePartition* RuntimePartition = HashSet->ResolveRuntimePartitionForHLODLayer(ResultPerInstanceData.RuntimeGrid, InActorDescView.GetHLODLayer()))
-								{
-									FName GridName = ResultPerInstanceData.RuntimeGrid.IsNone() ? DefaultGrid : ResultPerInstanceData.RuntimeGrid;
-									ResultPerInstanceData.RuntimeGrid = FName(GridName.ToString() + TEXT(":") + RuntimePartition->Name.ToString());
-								}
-							}
+							ResultPerInstanceData.RuntimeGrid = ResolvedRuntimeGrid;
 						}
 					}
 				}
diff --git a/Engine/Source/Runtime/Engine/Public/WorldPartition/HLOD/HLODLayer.h b/Engine/Source/Runtime/Engine/Public/WorldPartition/HLOD/HLODLayer.h
--- a/Engine/Source/Runtime/Engine/Public/WorldPartition/HLOD/HLODLayer.h
+++ b/Engine/Source/Runtime/Engine/Public/WorldPartition/HLOD/HLODLayer.h
@@ -59,6 +59,7 @@ public:
 	const TSubclassOf<UWorldPartitionHLODModifier> GetHLODModifierClass() const { return HLODModifierClass; }
 	UHLODLayer* GetParentLayer() const { return ParentLayer; }
 	void SetParentLayer(UHLODLayer* InParentLayer) { ParentLayer = InParentLayer; }
+	UHLODLayer* GetLinkedLayer() const { return LinkedLayer; }
 	ENGINE_API bool DoesRequireWarmup() const;
 	ENGINE_API void ComputeHLODHash(FHLODHashBuilder& InHLODHashBuilder) const;
 
@@ -118,6 +119,10 @@ private:
 	/** HLOD Layer to assign to the generated HLOD actors */
 	UPROPERTY(EditAnywhere, Category=General)
 	TObjectPtr<UHLODLayer> ParentLayer;
+	
+	/** HLOD Layer used to control visiblity of Custom HLOD Actors. Custom HLOD Actors become visible when actors from the linked HLOD Layer are unloaded. */
+	UPROPERTY(EditAnywhere, Category=General, AdvancedDisplay, meta = (EditConditionHides, EditCondition = "LayerType == EHLODLayerType::CustomHLODActor"))
+	TObjectPtr<UHLODLayer> LinkedLayer;
 
 	/** Specify a custom HLOD Actor class, the default is AWorldPartitionHLOD */
 	UPROPERTY(EditAnywhere, Category = General, AdvancedDisplay, meta = (DisplayName = "HLOD Actor Class"))
