Fix incorrect SSS profile being used for material using MaterialAttributes during implicit conversion.

diff --git a/Engine/Source/Runtime/Engine/Private/Materials/HLSLMaterialTranslator.cpp b/Engine/Source/Runtime/Engine/Private/Materials/HLSLMaterialTranslator.cpp
--- a/Engine/Source/Runtime/Engine/Private/Materials/HLSLMaterialTranslator.cpp
+++ b/Engine/Source/Runtime/Engine/Private/Materials/HLSLMaterialTranslator.cpp
@@ -1868,18 +1868,20 @@ void FHLSLMaterialTranslator::TranslateMaterial()
 
 				if (BaseMaterial && BaseMaterial->bUseMaterialAttributes)
 				{
-					const bool bHasSSS = MaterialInterface && MaterialInterface->SubsurfaceProfile;
-
+					USubsurfaceProfile* SSSProfile = MaterialInterface ? MaterialInterface->GetSubsurfaceProfile_Internal() : nullptr;
+					const bool bHasSSS = SSSProfile != nullptr;
+					
 					return UMaterialExpressionSubstrateConvertMaterialAttributes::CompileCommon(this, 0,
 						FMaterialAttributeDefinitionMap::GetConnectedMaterialAttributesBitmask(BaseMaterial->GetExpressions()), EditorOnlyData->MaterialAttributes, MaterialInterface->GetShadingModels().GetFirstShadingModel(),
 						*WaterScatteringCoefficients, *WaterAbsorptionCoefficients, *WaterPhaseG, *ColorScaleBehindWater,
-						bHasSSS, bHasSSS ? MaterialInterface->SubsurfaceProfile : nullptr,
+						bHasSSS, SSSProfile,
 						*ClearCoatNormal, *CustomTangent);
 
 				}
 				else if ((Material->GetMaterialDomain() == MD_Surface || Material->GetMaterialDomain() == MD_RuntimeVirtualTexture || Material->GetMaterialDomain() == MD_DeferredDecal) && EditorOnlyData)
 				{
-					const bool bHasSSS = MaterialInterface && MaterialInterface->SubsurfaceProfile;
+					USubsurfaceProfile* SSSProfile = MaterialInterface ? MaterialInterface->GetSubsurfaceProfile_Internal() : nullptr;
+					const bool bHasSSS = SSSProfile != nullptr;
 					const bool bHasAnisotropy = EditorOnlyData->Anisotropy.IsConnected();
 
 					return UMaterialExpressionSubstrateShadingModels::CompileCommon(this,
@@ -1891,7 +1893,7 @@ void FHLSLMaterialTranslator::TranslateMaterial()
 						bHasAnisotropy, EditorOnlyData->Anisotropy,
 						EditorOnlyData->Normal, EditorOnlyData->Tangent,
 						*ClearCoatNormal, *CustomTangent,
-						bHasSSS, MaterialInterface->SubsurfaceProfile,
+						bHasSSS, SSSProfile,
 						EditorOnlyData);
 				}
 				else if (Material->GetMaterialDomain() == MD_Volume && EditorOnlyData)
