Fix possible nullptr access

#rnx

diff --git a/Engine/Source/Runtime/Engine/Private/BlueprintGeneratedClass.cpp b/Engine/Source/Runtime/Engine/Private/BlueprintGeneratedClass.cpp
--- a/Engine/Source/Runtime/Engine/Private/BlueprintGeneratedClass.cpp
+++ b/Engine/Source/Runtime/Engine/Private/BlueprintGeneratedClass.cpp
@@ -1972,7 +1972,11 @@ void UBlueprintGeneratedClass::DestroyPersistentUberGraphFrame(UObject* Obj, boo
 	{
 		UClass* ParentClass = GetSuperClass();
 		checkSlow(ParentClass);
-		ParentClass->DestroyPersistentUberGraphFrame(Obj);
+
+		if (IsValid(ParentClass))
+		{
+			ParentClass->DestroyPersistentUberGraphFrame(Obj);
+		}
 	}
 }
 
