Added Experimental to TextureCollection to match Product Board

diff --git a/Engine/Source/Runtime/Engine/Classes/Engine/TextureCollection.h b/Engine/Source/Runtime/Engine/Classes/Engine/TextureCollection.h
--- a/Engine/Source/Runtime/Engine/Classes/Engine/TextureCollection.h
+++ b/Engine/Source/Runtime/Engine/Classes/Engine/TextureCollection.h
@@ -42,7 +42,7 @@ private:
 	FRHIResourceCollectionRef ResourceCollectionRHI;
 };
 
-UCLASS(MinimalAPI)
+UCLASS(MinimalAPI, Experimental)
 class UTextureCollection : public UDataAsset
 {
 	GENERATED_BODY()
