Marked VirtualTextureCollection, MaterialCacheVirtualTexture, MaterialCacheVirtualTextureTag as Experimental

diff --git a/Engine/Source/Runtime/Engine/Classes/Engine/VirtualTextureCollection.h b/Engine/Source/Runtime/Engine/Classes/Engine/VirtualTextureCollection.h
--- a/Engine/Source/Runtime/Engine/Classes/Engine/VirtualTextureCollection.h
+++ b/Engine/Source/Runtime/Engine/Classes/Engine/VirtualTextureCollection.h
@@ -112,7 +112,7 @@ private:
 	bool bAllowFormatConversion = false;
 };
 
-UCLASS(MinimalAPI)
+UCLASS(MinimalAPI, Experimental)
 class UVirtualTextureCollection : public UTextureCollection
 {
 	GENERATED_BODY()
diff --git a/Engine/Source/Runtime/Engine/Internal/MaterialCache/MaterialCacheVirtualTexture.h b/Engine/Source/Runtime/Engine/Internal/MaterialCache/MaterialCacheVirtualTexture.h
--- a/Engine/Source/Runtime/Engine/Internal/MaterialCache/MaterialCacheVirtualTexture.h
+++ b/Engine/Source/Runtime/Engine/Internal/MaterialCache/MaterialCacheVirtualTexture.h
@@ -14,7 +14,7 @@ class FVirtualTexture2DResource;
 class FSceneInterface;
 struct FMaterialCacheVirtualBaton;
 
-UCLASS(MinimalAPI)
+UCLASS(MinimalAPI, Experimental)
 class UMaterialCacheVirtualTexture : public UTexture
 {
 	GENERATED_UCLASS_BODY()
diff --git a/Engine/Source/Runtime/Engine/Internal/MaterialCache/MaterialCacheVirtualTextureTag.h b/Engine/Source/Runtime/Engine/Internal/MaterialCache/MaterialCacheVirtualTextureTag.h
--- a/Engine/Source/Runtime/Engine/Internal/MaterialCache/MaterialCacheVirtualTextureTag.h
+++ b/Engine/Source/Runtime/Engine/Internal/MaterialCache/MaterialCacheVirtualTextureTag.h
@@ -8,7 +8,7 @@
 
 #define UE_API ENGINE_API
 
-UCLASS(MinimalAPI)
+UCLASS(MinimalAPI, Experimental)
 class UMaterialCacheVirtualTextureTag : public UObject
 {
 	GENERATED_UCLASS_BODY()
