Material Cache, hide expression by default for now

diff --git a/Engine/Source/Runtime/Engine/Public/Materials/MaterialExpressionMaterialCache.h b/Engine/Source/Runtime/Engine/Public/Materials/MaterialExpressionMaterialCache.h
--- a/Engine/Source/Runtime/Engine/Public/Materials/MaterialExpressionMaterialCache.h
+++ b/Engine/Source/Runtime/Engine/Public/Materials/MaterialExpressionMaterialCache.h
@@ -29,7 +29,7 @@ struct FMaterialExpressionMaterialCacheAttribute
 	FExpressionInput Input;
 };
 
-UCLASS(collapsecategories, hidecategories = Object, meta = (DisplayName = "MaterialCache"))
+UCLASS(collapsecategories, hidecategories = Object, meta = (Private, DisplayName = "MaterialCache"))
 class UMaterialExpressionMaterialCache : public UMaterialExpressionCustomOutput
 {
 	GENERATED_BODY()
