Fix virtual texture upload issue where we could overrun the number of staging buffers addressable by the tile handle.
When it happened it would corrupt the staging buffer free tile array and lead to memory exceptions when copying the tile data.
Changed allocation of bits in tile handle to avoid this happening. Added check to catch unlikely occurrence in dev builds.
#rnx

diff --git a/Engine/Source/Runtime/Engine/Private/VT/VirtualTextureUploadCache.cpp b/Engine/Source/Runtime/Engine/Private/VT/VirtualTextureUploadCache.cpp
--- a/Engine/Source/Runtime/Engine/Private/VT/VirtualTextureUploadCache.cpp
+++ b/Engine/Source/Runtime/Engine/Private/VT/VirtualTextureUploadCache.cpp
@@ -65,6 +65,7 @@ uint32 FVTUploadTileAllocator::Allocate(FRHICommandListBase& RHICmdList, EVTUplo
 	if (FormatIndex == FormatDescs.Num())
 	{
 		// Add newly found format.
+		check(FormatDescs.Num() < FHandle::MaxFormats);
 		FormatDescs.Add(Desc);
 		FormatBuffers.AddDefaulted();
 	}
@@ -91,6 +92,7 @@ uint32 FVTUploadTileAllocator::Allocate(FRHICommandListBase& RHICmdList, EVTUplo
 	if (StagingBufferIndex == FormatBuffer.StagingBuffers.Num())
 	{
 		// Current staging buffers are full. Need to allocate a new staging buffer.
+		check(FormatBuffer.StagingBuffers.Num() < FHandle::MaxStagingBuffers);
 		FormatBuffer.StagingBuffers.AddDefaulted();
 	}
 
@@ -176,6 +178,7 @@ void FVTUploadTileAllocator::FStagingBuffer::Init(FRHICommandListBase& RHICmdLis
 
 	const uint32 RequestedBufferSize = CVarVTUploadMemoryPageSize.GetValueOnRenderThread() * 1024 * 1024;
 	NumTiles = FMath::DivideAndRoundUp(RequestedBufferSize, TileSizeAligned);
+	NumTiles = FMath::Min(NumTiles, FVTUploadTileAllocator::FHandle::MaxTiles);
 	const uint32 BufferSize = TileSizeAligned * NumTiles;
 
 	check(TileFreeList.Num() == 0);
diff --git a/Engine/Source/Runtime/Engine/Private/VT/VirtualTextureUploadCache.h b/Engine/Source/Runtime/Engine/Private/VT/VirtualTextureUploadCache.h
--- a/Engine/Source/Runtime/Engine/Private/VT/VirtualTextureUploadCache.h
+++ b/Engine/Source/Runtime/Engine/Private/VT/VirtualTextureUploadCache.h
@@ -87,10 +87,14 @@ private:
 		uint32 PackedValue = 0u;
 		struct
 		{
-			uint32 FormatIndex : 8;
-			uint32 StagingBufferIndex : 8;
-			uint32 TileIndex : 16;
+			uint32 FormatIndex : 6;
+			uint32 StagingBufferIndex : 12;
+			uint32 TileIndex : 14;
 		};
+
+		static const uint32 MaxFormats = 1 << 6;
+		static const uint32 MaxStagingBuffers = 1 << 12;
+		static const uint32 MaxTiles = 1 << 14;
 	};
 
 	/**
