Removed AnimInstanceName from the FAnimInstanceProxy

diff --git a/Engine/Source/Runtime/Engine/Private/Animation/AnimInstanceProxy.cpp b/Engine/Source/Runtime/Engine/Private/Animation/AnimInstanceProxy.cpp
--- a/Engine/Source/Runtime/Engine/Private/Animation/AnimInstanceProxy.cpp
+++ b/Engine/Source/Runtime/Engine/Private/Animation/AnimInstanceProxy.cpp
@@ -247,8 +247,6 @@ void FAnimInstanceProxy::Initialize(UAnimInstance* InAnimInstance)
 	ActorName = GetNameSafe(InAnimInstance->GetOwningActor());
 #endif
 
-	AnimInstanceName = *InAnimInstance->GetFullName();
-
 	UpdateCounter.Reset();
 	ReinitializeSlotNodes();
 
diff --git a/Engine/Source/Runtime/Engine/Public/Animation/AnimInstanceProxy.h b/Engine/Source/Runtime/Engine/Public/Animation/AnimInstanceProxy.h
--- a/Engine/Source/Runtime/Engine/Public/Animation/AnimInstanceProxy.h
+++ b/Engine/Source/Runtime/Engine/Public/Animation/AnimInstanceProxy.h
@@ -496,9 +496,9 @@ public:
 	}
 #endif
 
-	const FString& GetAnimInstanceName() const 
+	FString GetAnimInstanceName() const
 	{ 
-		return AnimInstanceName;
+		return AnimInstanceObject->GetFullName();
 	}
 
 	// Get the root motion mode assigned to this anim instance proxy
@@ -1054,9 +1054,6 @@ private:
 	FString ActorName;
 #endif
 
-	/** Anim instance name for debug purposes */
-	FString AnimInstanceName;
-
 	/** Anim graph */
 	FAnimNode_Base* RootNode;
 
