Update Heterogeneous Volume render settings. Enable shadow-casting by default.

1. Heterogeneous Volumes are now in beta
2. Enabling shadow-casting no longer requires an editor restart

#DisabledValidations Preflight

diff --git a/Engine/Source/Runtime/Engine/Classes/Engine/RendererSettings.h b/Engine/Source/Runtime/Engine/Classes/Engine/RendererSettings.h
--- a/Engine/Source/Runtime/Engine/Classes/Engine/RendererSettings.h
+++ b/Engine/Source/Runtime/Engine/Classes/Engine/RendererSettings.h
@@ -1136,7 +1136,7 @@ class URendererSettings : public UDeveloperSettings
 		uint32 bMeshStreaming : 1;
 
 	UPROPERTY(config, EditAnywhere, Category = "Heterogeneous Volumes", meta = (
-		ConsoleVariable = "r.HeterogeneousVolumes", DisplayName = "Heterogeneous Volumes (Experimental)",
+		ConsoleVariable = "r.HeterogeneousVolumes", DisplayName = "Heterogeneous Volumes (Beta)",
 		ToolTip = "Enable rendering with the heterogeneous volumes subsystem.",
 		ConfigRestartRequired = false))
 	uint32 bEnableHeterogeneousVolumes : 1;
@@ -1144,7 +1144,7 @@ class URendererSettings : public UDeveloperSettings
 	UPROPERTY(config, EditAnywhere, Category = "Heterogeneous Volumes", meta = (
 		ConsoleVariable = "r.HeterogeneousVolumes.Shadows", DisplayName = "Shadow Casting",
 		ToolTip = "Enable heterogeneous volumes to cast shadows onto the environment.",
-		ConfigRestartRequired = true))
+		ConfigRestartRequired = false))
 		uint32 bShouldHeterogeneousVolumesCastShadows : 1;
 
 	UPROPERTY(config, EditAnywhere, Category = "Heterogeneous Volumes", meta = (
diff --git a/Engine/Source/Runtime/Renderer/Private/HeterogeneousVolumes/HeterogeneousVolumes.cpp b/Engine/Source/Runtime/Renderer/Private/HeterogeneousVolumes/HeterogeneousVolumes.cpp
--- a/Engine/Source/Runtime/Renderer/Private/HeterogeneousVolumes/HeterogeneousVolumes.cpp
+++ b/Engine/Source/Runtime/Renderer/Private/HeterogeneousVolumes/HeterogeneousVolumes.cpp
@@ -21,8 +21,8 @@ static TAutoConsoleVariable<int32> CVarHeterogeneousVolumes(
 
 static TAutoConsoleVariable<int32> CVarHeterogeneousVolumesShadows(
 	TEXT("r.HeterogeneousVolumes.Shadows"),
-	0,
-	TEXT("Enables heterogeneous volume-casting shadows (default = 0)"),
+	1,
+	TEXT("Enables heterogeneous volume-casting shadows (default = 1)"),
 	ECVF_RenderThreadSafe
 );
 
