Adjust distribution of bits in the virtual texture upload handle to deal with extreme case of uploading many pages in the same render frame.
Make the memory byte counter 64bit to deal with the same case.
#rnx

diff --git a/Engine/Source/Runtime/Engine/Private/VT/VirtualTextureUploadCache.h b/Engine/Source/Runtime/Engine/Private/VT/VirtualTextureUploadCache.h
--- a/Engine/Source/Runtime/Engine/Private/VT/VirtualTextureUploadCache.h
+++ b/Engine/Source/Runtime/Engine/Private/VT/VirtualTextureUploadCache.h
@@ -87,14 +87,14 @@ private:
 		uint32 PackedValue = 0u;
 		struct
 		{
-			uint32 FormatIndex : 6;
-			uint32 StagingBufferIndex : 12;
-			uint32 TileIndex : 14;
+			uint32 FormatIndex : 5;
+			uint32 StagingBufferIndex : 14;
+			uint32 TileIndex : 13;
 		};
 
-		static const uint32 MaxFormats = 1 << 6;
-		static const uint32 MaxStagingBuffers = 1 << 12;
-		static const uint32 MaxTiles = 1 << 14;
+		static const uint32 MaxFormats = 1 << 5;
+		static const uint32 MaxStagingBuffers = 1 << 14;
+		static const uint32 MaxTiles = 1 << 13;
 	};
 
 	/**
@@ -143,7 +143,7 @@ private:
 	/** Array of staging buffers. Kept in sync with associated formats from FormatDescs. */
 	TArray<FSharedFormatBuffers> FormatBuffers;
 	/** Allocated memory counter. */
-	uint32 NumAllocatedBytes = 0;
+	uint64 NumAllocatedBytes = 0;
 };
 
 /** 
