Fix 64-32bit downcast when reporting virtual texture upload staging memory.
#rnx

diff --git a/Engine/Source/Runtime/Engine/Private/VT/VirtualTextureUploadCache.cpp b/Engine/Source/Runtime/Engine/Private/VT/VirtualTextureUploadCache.cpp
--- a/Engine/Source/Runtime/Engine/Private/VT/VirtualTextureUploadCache.cpp
+++ b/Engine/Source/Runtime/Engine/Private/VT/VirtualTextureUploadCache.cpp
@@ -611,5 +611,5 @@ uint32 FVirtualTextureUploadCache::IsInMemoryBudget() const
 {
 	return 
 		PendingUpload.Num() + PendingRelease.Num() <= CVarMaxUploadRequests.GetValueOnRenderThread() &&
-		TileAllocator.TotalAllocatedBytes() <= (uint32)CVarVTMaxUploadMemory.GetValueOnRenderThread() * 1024u * 1024u;
+		TileAllocator.TotalAllocatedBytes() <= (uint64)CVarVTMaxUploadMemory.GetValueOnRenderThread() * 1024u * 1024u;
 }
diff --git a/Engine/Source/Runtime/Engine/Private/VT/VirtualTextureUploadCache.h b/Engine/Source/Runtime/Engine/Private/VT/VirtualTextureUploadCache.h
--- a/Engine/Source/Runtime/Engine/Private/VT/VirtualTextureUploadCache.h
+++ b/Engine/Source/Runtime/Engine/Private/VT/VirtualTextureUploadCache.h
@@ -75,7 +75,7 @@ public:
 	FVTUploadTileBufferExt GetBufferFromHandleExt(uint32 InHandle) const;
 	
 	/** Get allocated memory in bytes. */
-	uint32 TotalAllocatedBytes() const { return NumAllocatedBytes; }
+	uint64 TotalAllocatedBytes() const { return NumAllocatedBytes; }
 
 private:
 	/** Mutex for Allocate/Free. */
