Add output tooltip for HairAttributes node.

diff --git a/Engine/Source/Runtime/Engine/Private/Materials/MaterialExpressions.cpp b/Engine/Source/Runtime/Engine/Private/Materials/MaterialExpressions.cpp
--- a/Engine/Source/Runtime/Engine/Private/Materials/MaterialExpressions.cpp
+++ b/Engine/Source/Runtime/Engine/Private/Materials/MaterialExpressions.cpp
@@ -20705,6 +20705,92 @@ void UMaterialExpressionHairAttributes::GetCaption(TArray<FString>& OutCaptions)
 {
 	OutCaptions.Add(TEXT("Hair Attributes"));
 }
+
+void UMaterialExpressionHairAttributes::GetConnectorToolTip(int32 InputIndex, int32 OutputIndex, TArray<FString>& OutToolTip)
+{
+	if (OutputIndex == 0)
+	{
+		OutToolTip.Add(TEXT("The U coordinates of hair. Coordinate along the hair, with 0 being the root and 1 the tip."));
+		return;
+	}
+	else if (OutputIndex == 1)
+	{
+		OutToolTip.Add(TEXT("The V coordinates of hair. Coordinate across the hair, with 0 being the left and 1 the right."));
+		return;
+	}
+	else if (OutputIndex == 2)
+	{
+		OutToolTip.Add(TEXT("The length of the current curve. Length from the root to the shading point, in cm."));
+		return;
+	}
+	else if (OutputIndex == 3)
+	{
+		OutToolTip.Add(TEXT("Radius of the curve at the current position."));
+		return;
+	}
+	else if (OutputIndex == 4)
+	{
+		OutToolTip.Add(TEXT("Random value in 0 to 1, and uniform along the curve."));
+		return;
+	}
+	else if (OutputIndex == 5)
+	{
+		OutToolTip.Add(TEXT("Tangent vector aligned in the direction of the curve."));
+		return;
+	}
+	else if (OutputIndex == 6)
+	{
+		OutToolTip.Add(TEXT("UV of the underlying mesh at the curve's root position."));
+		return;
+	}
+	else if (OutputIndex == 7)
+	{
+		OutToolTip.Add(TEXT("Per curve's point color."));
+		return;
+	}
+	else if (OutputIndex == 8)
+	{
+		OutToolTip.Add(TEXT("Per curve's point roughness."));
+		return;
+	}
+	else if (OutputIndex == 9)
+	{
+		OutToolTip.Add(TEXT("Depth offset. Only used for cards and mesh geometry."));
+		return;
+	}
+	else if (OutputIndex == 10)
+	{
+		OutToolTip.Add(TEXT("Coverage mask value. Only used for cards and mesh geometry."));
+		return;
+	}
+	else if (OutputIndex == 11)
+	{
+		OutToolTip.Add(TEXT("Auxiliary data that is only used for cards and mesh geometry."));
+		return;
+	}
+	else if (OutputIndex == 12)
+	{
+		OutToolTip.Add(TEXT("Cards UVs that are only used for cards and mesh geometry."));
+		return;
+	}
+	else if (OutputIndex == 13)
+	{
+		OutToolTip.Add(TEXT("The group index of the curve."));
+		return;
+	}
+	else if (OutputIndex == 14)
+	{
+		OutToolTip.Add(TEXT("Per curve's ambient occlusion."));
+		return;
+	}
+	else if (OutputIndex == 15)
+	{
+		OutToolTip.Add(TEXT("The clump ID of the curve."));
+		return;
+	}
+
+	Super::GetConnectorToolTip(InputIndex, INDEX_NONE, OutToolTip);
+}
 #endif // WITH_EDITOR
 
 //
diff --git a/Engine/Source/Runtime/Engine/Public/Materials/MaterialExpressionHairAttributes.h b/Engine/Source/Runtime/Engine/Public/Materials/MaterialExpressionHairAttributes.h
--- a/Engine/Source/Runtime/Engine/Public/Materials/MaterialExpressionHairAttributes.h
+++ b/Engine/Source/Runtime/Engine/Public/Materials/MaterialExpressionHairAttributes.h
@@ -23,6 +23,7 @@ class UMaterialExpressionHairAttributes : public UMaterialExpression
 	virtual void Build(MIR::FEmitter& Emitter) override;
 	virtual int32 Compile(class FMaterialCompiler* Compiler, int32 OutputIndex) override;
 	virtual void GetCaption(TArray<FString>& OutCaptions) const override;
+	virtual void GetConnectorToolTip(int32 InputIndex, int32 OutputIndex, TArray<FString>& OutToolTip);
 #endif
 	//~ End UMaterialExpression Interface
 };
