[CrashReport][Crash] UnrealEditor-Engine!UAnimSequence::GetApproxCompressedSize() [Engine/Source/Runtime/Engine/Private/Animation/AnimSequence.cpp:575]
#speculative-fix add extra conditional for whether or not an async compression task had the current platform as its target.

diff --git a/Engine/Source/Runtime/Engine/Private/Animation/AnimSequence.cpp b/Engine/Source/Runtime/Engine/Private/Animation/AnimSequence.cpp
--- a/Engine/Source/Runtime/Engine/Private/Animation/AnimSequence.cpp
+++ b/Engine/Source/Runtime/Engine/Private/Animation/AnimSequence.cpp
@@ -5476,7 +5476,7 @@ void UAnimSequence::FinishAsyncTasks()
 
 			UE_COMPRESSED_DATA_WRITE_SCOPE(this);
 			FCompressedAnimSequence* TaskData = It->Value->GetTargetCompressedData();
-			const bool bCompressedRunningPlatform = It->Key == CurrentPlatformHash;
+			const bool bCompressedRunningPlatform = It->Key == CurrentPlatformHash && It->Value->GetTargetPlatform() == RunningPlatform;
 			auto ResetData = [this, TaskData, It]()
 			{
 				if(TaskData != nullptr)
diff --git a/Engine/Source/Runtime/Engine/Private/Animation/AnimationCompressionDerivedData.h b/Engine/Source/Runtime/Engine/Private/Animation/AnimationCompressionDerivedData.h
--- a/Engine/Source/Runtime/Engine/Private/Animation/AnimationCompressionDerivedData.h
+++ b/Engine/Source/Runtime/Engine/Private/Animation/AnimationCompressionDerivedData.h
@@ -84,6 +84,7 @@ namespace UE::Anim
 		void Reschedule(FQueuedThreadPool* InThreadPool, EQueuedWorkPriority InPriority) const;
 		bool WasCancelled() const { return CompressibleAnimPtr->IsCancelled() || Owner.IsCanceled(); }
 		FCompressedAnimSequence* GetTargetCompressedData() const { return CompressedData; }
+		const ITargetPlatform* GetTargetPlatform() const { return TargetPlatform; }
 	private:
 		void BeginCache(const FIoHash& KeyHash);
 		void EndCache(DerivedData::FCacheGetValueResponse&& Response);
@@ -105,4 +106,4 @@ namespace UE::Anim
 		int64 RequiredMemory;
 	};
 #endif
-}
\ No newline at end of file
+}
