[CrashReport][Crash] Engine!UE::Anim::Decompression::DecompressPose() [AnimationDecompression.cpp
#speculative-fix for some callstacks: verify we have a valid skeleton before doing any bonepose evaluation

diff --git a/Engine/Source/Runtime/Engine/Private/Animation/AnimSequence.cpp b/Engine/Source/Runtime/Engine/Private/Animation/AnimSequence.cpp
--- a/Engine/Source/Runtime/Engine/Private/Animation/AnimSequence.cpp
+++ b/Engine/Source/Runtime/Engine/Private/Animation/AnimSequence.cpp
@@ -1734,16 +1734,12 @@ void UAnimSequence::GetBonePose(FAnimationPoseData& OutAnimationPoseData, const
 	check(!bForceUseRawData || CanEvaluateRawAnimationData());
 
 	UE_COMPRESSED_DATA_READ_SCOPE(this);
-	const bool bUseRawDataForPoseExtraction = (CanEvaluateRawAnimationData() && bForceUseRawData) || ShouldUseRawDataForPoseExtraction_Lockless(RequiredBones, ExtractionContext);
-	const FCompressedAnimSequence& PlatformCompressedData = GetPlatformCompressedData(ExtractionContext);
-#if WITH_EDITOR
-	check(!ExtractionContext.bEnforceCompressedDataSampling || (PlatformCompressedData.IsValid(this) && !bUseRawDataForPoseExtraction));
-#endif // WITH_EDITOR
+	const USkeleton* MySkeleton = GetSkeleton();
+	const bool bInvalidSkeleton = MySkeleton == nullptr;
+	const bool bUseRawDataForPoseExtraction = bInvalidSkeleton || (CanEvaluateRawAnimationData() && bForceUseRawData) || ShouldUseRawDataForPoseExtraction_Lockless(RequiredBones, ExtractionContext);
 
 	const bool bIsBakedAdditive = !bUseRawDataForPoseExtraction && IsValidAdditive();
-
-	const USkeleton* MySkeleton = GetSkeleton();
-	if (!MySkeleton)
+	if (bInvalidSkeleton)
 	{
 		if (bIsBakedAdditive)
 		{
@@ -1755,6 +1751,11 @@ void UAnimSequence::GetBonePose(FAnimationPoseData& OutAnimationPoseData, const
 		}
 		return;
 	}
+	
+	const FCompressedAnimSequence& PlatformCompressedData = GetPlatformCompressedData(ExtractionContext);
+#if WITH_EDITOR
+	check(!ExtractionContext.bEnforceCompressedDataSampling || (PlatformCompressedData.IsValid(this) && !bUseRawDataForPoseExtraction));
+#endif // WITH_EDITOR
 
 	const bool bDisableRetargeting = RequiredBones.GetDisableRetargeting();
 
@@ -3693,6 +3694,7 @@ bool UAnimSequence::ShouldUseRawDataForPoseExtraction_Lockless(const FAnimExtrac
 {
 	return CanEvaluateRawAnimationData() &&
 	(
+		GetSkeleton() == nullptr ||
 #if WITH_EDITOR
 		GForceRawData == 1 ||
 		!GetPlatformCompressedData(ExtractionContext).IsValid(this) ||
