Clarify Substrate Slab node input when using Substrate's format 0

diff --git a/Engine/Source/Runtime/Engine/Public/Materials/MaterialExpressionSubstrate.h b/Engine/Source/Runtime/Engine/Public/Materials/MaterialExpressionSubstrate.h
--- a/Engine/Source/Runtime/Engine/Public/Materials/MaterialExpressionSubstrate.h
+++ b/Engine/Source/Runtime/Engine/Public/Materials/MaterialExpressionSubstrate.h
@@ -298,8 +298,8 @@ class UMaterialExpressionSubstrateSlabBSDF : public UMaterialExpressionSubstrate
 
 	/**
 	 * Chromatic mean free path. Only used when there is not any sub-surface profile provided. (type = float3, unit = centimeters, default = 0).
-	 * For subsurface scattering, it is recommended to specify the MFP as world space centimeter directly. As it is an intuitive way to specify how far the light will scatter/bleed.
-	 * For translucent coat layer, it is recommended to use the TransmittanceToMFP node as it might be easier to specify the transmittance color.
+	 *  - For subsurface scattering, it is recommended to specify the MFP as world space centimeter directly. As it is an intuitive way to specify how far the light will scatter/bleed.
+	 *  - For translucent coat layer, it is recommended to use the TransmittanceToMFP node as it might be easier to specify the transmittance color.
 	 */
 	UPROPERTY()
 	FExpressionInput SSSMFP;
@@ -323,19 +323,19 @@ class UMaterialExpressionSubstrateSlabBSDF : public UMaterialExpressionSubstrate
 	FExpressionInput EmissiveColor;
 
 	/**
-	 * Controls the roughness of a secondary specular lobe. Roughness of 0 (smooth) is a mirror reflection and 1 (rough) is completely matte or diffuse. Does not influence diffuse roughness. (type = float, unit = unitless, defaults to 0.5)
+	* Controls the roughness of a secondary specular lobe. Roughness of 0 (smooth) is a mirror reflection and 1 (rough) is completely matte or diffuse. Does not influence diffuse roughness. (type = float, unit = unitless, defaults to 0.5) When using Blendable GBuffer, this input is ignored. 
 	 */
 	UPROPERTY()
 	FExpressionInput SecondRoughness;
 
 	/**
-	 * The weight of the second specular lobe using SecondRoughness. The first specular using Roughness will have a weight of (1 - SecondRoughnessWeight). (type = float, unitless, default = 0)
+	* The weight of the second specular lobe using SecondRoughness. The first specular using Roughness will have a weight of (1 - SecondRoughnessWeight). (type = float, unitless, default = 0) When using Blendable GBuffer, this input is ignored. 
 	 */
 	UPROPERTY()
 	FExpressionInput SecondRoughnessWeight;
 
 	/**
-	 * Controls how rough the Fuzz layer is. Roughness of 0 is smooth and 1 is rough. If FuzzRoughness Is not connected, the Roughness input will be used instead. (type = float, unit = unitless, defaults to 0.5)
+	* Controls how rough the Fuzz layer is. Roughness of 0 is smooth and 1 is rough. If FuzzRoughness Is not connected, the Roughness input will be used instead. (type = float, unit = unitless, defaults to 0.5). When using Blendable GBuffer, this input is ignored. Only the fuzz roughness is forced to 0.5. 
 	 */
 	UPROPERTY()
 	FExpressionInput FuzzRoughness;
@@ -353,13 +353,13 @@ class UMaterialExpressionSubstrateSlabBSDF : public UMaterialExpressionSubstrate
 	FExpressionInput FuzzColor;
 
 	/**
-	 * This represent the micro facet density. Only used when `r.Substrate.Glints=1`. 0 = only very sparse glints, 1 = fully covered with glints (which is equivalent to a regular specular lobe). Defaults to 1 (=no glint).
+	 * This represent the micro facet density. Only used when `r.Substrate.Glints=1`. 0 = only very sparse glints, 1 = fully covered with glints (which is equivalent to a regular specular lobe). Defaults to 1 (=no glint). When using Blendable GBuffer, this input is ignored.
 	 */
 	UPROPERTY()
 	FExpressionInput GlintValue;
 
 	/**
-	 * The parameterization of the surface required to position glints on a surface. Only used when `r.Substrate.Glints=1`. Defaults to (0,0).
+	 * The parameterization of the surface required to position glints on a surface. Only used when `r.Substrate.Glints=1`. Defaults to (0,0). When using Blendable GBuffer, this input is ignored.
 	 */
 	UPROPERTY()
 	FExpressionInput GlintUV;
