Reorganized fields in FDeferredDecalUpdateParams to avoid padding

diff --git a/Engine/Source/Runtime/Engine/Public/SceneProxies/DeferredDecalProxy.h b/Engine/Source/Runtime/Engine/Public/SceneProxies/DeferredDecalProxy.h
--- a/Engine/Source/Runtime/Engine/Public/SceneProxies/DeferredDecalProxy.h
+++ b/Engine/Source/Runtime/Engine/Public/SceneProxies/DeferredDecalProxy.h
@@ -79,16 +79,15 @@ public:
 
 struct FDeferredDecalUpdateParams
 {
-	enum class EOperationType : int
+	enum class EOperationType : int8
 	{
 		AddToSceneAndUpdate,				// Adds the decal to the scene an updates the parameters
 		Update,								// Updates the decals parameters
 		RemoveFromSceneAndDelete,			// Remove the decal from the scene and deletes the proxy
 	};
 
-	EOperationType			OperationType = EOperationType::Update;
-	FDeferredDecalProxy*	DecalProxy = nullptr;
 	FTransform				Transform;
+	FDeferredDecalProxy*	DecalProxy = nullptr;
 	FBoxSphereBounds		Bounds;
 	float					AbsSpawnTime = 0.0f;
 	float					FadeDuration = 0.0f;
@@ -98,6 +97,7 @@ struct FDeferredDecalUpdateParams
 	float					FadeScreenSize = 0.01f;
 	int32					SortOrder = 0;
 	FLinearColor			DecalColor = FLinearColor::White;
+	EOperationType			OperationType = EOperationType::Update;
 };
 
 
