[Enhanced Input]

Consolidate the processing of the input component stack by UPlayerInput and UEnhancedPlayerInput happen in order, and only once. This means means that we can :

- Only iterate the input component stack *once* per-player, instead of twice.
- Legacy Input Component's priority will be respected. This means that, a "legacy" UInputComponent with a higher priority then an UEnhancedInputComponent, will have its delegates fired before the Enhanced Input component.
- Legacy Input delegates will fire after each UInputComponent is processed, instead of all at once after iterating the whole stack and being sorted. This change is necessary to preserve the input component delegate ordering.

These changes have two large wins for input processing:

- Complicated projects which may have a lot of use of older UInputComponent's, will have an easier time upgrading to Enhnaced Input over time, and Enhanced Input will now behave in a more predictable and determinisitc way when used side-by-side with the legacy system.
- Performance improvements, because we are cutting the number of iterations of the input component stack in half.
- This will make it easier to, in the future, transition completely away from the legacy system and just use Enhanced Input. This also will allow projects to override more targeted pieces of the input system and completely disable them if they desire, or add more custom logic for specific pieces of input delegate processing.

diff --git a/Engine/Plugins/EnhancedInput/Source/EnhancedInput/Private/EnhancedPlayerInput.cpp b/Engine/Plugins/EnhancedInput/Source/EnhancedInput/Private/EnhancedPlayerInput.cpp
--- a/Engine/Plugins/EnhancedInput/Source/EnhancedInput/Private/EnhancedPlayerInput.cpp
+++ b/Engine/Plugins/EnhancedInput/Source/EnhancedInput/Private/EnhancedPlayerInput.cpp
@@ -340,6 +340,56 @@ void UEnhancedPlayerInput::EvaluateInputDelegates(const TArray<UInputComponent*>
 {
 	TRACE_CPUPROFILER_EVENT_SCOPE(EnhPIS_Main);
 
+	Super::EvaluateInputDelegates(InputComponentStack, DeltaTime, bGamePaused, KeysWithEvents);
+
+	// Calculate the current delta time as a fallback in case the Time Dilation is set to 0.0
+	// This will ensure that Timed Triggers are calculated in real time
+	const UWorld* World = GetWorld();
+	// If there is no world, then add 1/60 to the last frame time to try and recover at least some
+	// semblance of real time passing. 
+	float CurrentTime = World ? World->GetRealTimeSeconds() : LastFrameTime + (1.0f / 60.0f);
+	
+	// Reset action instance timers where necessary post delegate calls
+	for (TPair<TObjectPtr<const UInputAction>, FInputActionInstance>& ActionPair : ActionInstanceData)
+	{
+		FInputActionInstance& ActionData = ActionPair.Value;
+		switch (ActionData.TriggerEvent)
+		{
+		case ETriggerEvent::None:
+		case ETriggerEvent::Canceled:
+		case ETriggerEvent::Completed:
+			ActionData.ElapsedProcessedTime = 0.f;
+			break;
+		}
+		if (ActionData.TriggerEvent != ETriggerEvent::Triggered)
+		{
+			ActionData.ElapsedTriggeredTime = 0.f;
+		}
+
+		// Delay MappingTriggerState reset until here to allow dependent triggers (e.g. chords) access to this tick's values.
+		ActionData.TriggerStateTracker = FTriggerStateTracker();
+	}
+
+	// Remove any input actions that are no longer mapped to the player. At this point they will have been evaluated and 
+	// fired any "Canceled" input events needed
+	for (const UInputAction* IA : ActionsThatHaveBeenRemovedFromMappings)
+	{
+		ActionInstanceData.Remove(IA);
+	}
+
+	// We can clear our queue of mappings that have been removed on the next iteration directly after a key rebuild.
+	ActionsThatHaveBeenRemovedFromMappings.Reset();
+
+	LastFrameTime = CurrentTime;
+	KeysPressedThisTick.Reset();
+	bIsFlushingInputThisFrame = false;
+	ConsumedInputActions.Reset();
+}
+
+void UEnhancedPlayerInput::PrepareInputDelegatesForEvaluation(const TArray<UInputComponent*>& InputComponentStack, const float DeltaTime, const bool bGamePaused, const TArray<TPair<FKey, FKeyState*>>& KeysWithEvents)
+{
+	Super::PrepareInputDelegatesForEvaluation(InputComponentStack, DeltaTime, bGamePaused, KeysWithEvents);
+
 	// Process Action bindings
 	ActionsWithEventsThisTick.Reset();
 
@@ -555,38 +605,20 @@ void UEnhancedPlayerInput::EvaluateInputDelegates(const TArray<UInputComponent*>
 			ActionData.LastTriggeredWorldTime = CurrentTime;
 		}
 	}
+}
 
-
-	// Execute appropriate delegates
-
-	// A set of input actions whose bound delete was flagged to consume the input, preventing
-	// other listeners from receiving the event. This is how Enhanced Input respects the priority of UInputComponent's
-	// This allows other UEnhancedInputComponents (notably, not any "legacy" UInputComponents) completely consume other
-	// listeners bound to the same input action.
-	//
-	// This is useful if you have some gameplay code which you want to "override" the behavior of other listeners in some specific
-	// scenario, like if you have a specific higher priority gameplay ability or something like that. This way, you don't need to
-	// manually be tracking some additional game state constantly. 
-	TSet<const UInputAction*> ConsumedInputActions;
+bool UEnhancedPlayerInput::EvaluateInputComponentDelegates(UInputComponent* const InputComponent, const TArray<TPair<FKey, FKeyState*>>& KeysWithEvents, const float DeltaTime, const bool bGamePaused)
+{
+	bool bBlocksInput = false;
 
 #if DEV_ONLY_KEY_BINDINGS_AVAILABLE
 	// Cache modifier key states for debug key bindings
 	const bool bAlt = IsAltPressed(), bCtrl = IsCtrlPressed(), bShift = IsShiftPressed(), bCmd = IsCmdPressed();
 #endif
-
-	// TODO: Process APlayerController::InputComponent only!
-	// Walk the stack, top to bottom, grabbing actions and firing triggered delegates
-	int32 StackIndex = InputComponentStack.Num() - 1;
-	for (; StackIndex >= 0; --StackIndex)
+	
+	// We only want to process enhanced input components in this subclass
+	if (UEnhancedInputComponent* const IC = Cast<UEnhancedInputComponent>(InputComponent))
 	{
-		TRACE_CPUPROFILER_EVENT_SCOPE(EnhPIS_Delegates);
-
-		UEnhancedInputComponent* IC = Cast<UEnhancedInputComponent>(InputComponentStack[StackIndex]);
-		if (!IsValid(IC))
-		{
-			continue;
-		}
-
 		// Trigger bound event delegates
 		static TArray<TUniquePtr<FEnhancedInputActionEventBinding>> TriggeredDelegates;
 		for (const TUniquePtr<FEnhancedInputActionEventBinding>& Binding : IC->GetActionEventBindings())
@@ -766,65 +798,25 @@ void UEnhancedPlayerInput::EvaluateInputDelegates(const TArray<UInputComponent*>
 		}
 		TriggeredDebugDelegates.Reset();
 #endif
-
-
-		// Early termination if this component is blocking input.
-		// TODO: Remove support for this?
-		if (IC->bBlockInput)
-		{
-			// stop traversing the stack, all input has been consumed by this InputComponent
-			--StackIndex;
-			break;
-		}
 	}
+	
+	
+	const bool bLegacyBlocksInput = Super::EvaluateInputComponentDelegates(InputComponent, KeysWithEvents, DeltaTime, bGamePaused);
 
-	for (; StackIndex >= 0; --StackIndex)
-	{
-		if (UEnhancedInputComponent* IC = Cast<UEnhancedInputComponent>(InputComponentStack[StackIndex]))
-		{
-			for (const FEnhancedInputActionValueBinding& Binding : IC->GetActionValueBindings())
-			{
-				Binding.CurrentValue.Reset();
-			}
-		}
-	}
+	return bBlocksInput || bLegacyBlocksInput;
+}
 
-	// Reset action instance timers where necessary post delegate calls
-	for (TPair<TObjectPtr<const UInputAction>, FInputActionInstance>& ActionPair : ActionInstanceData)
+void UEnhancedPlayerInput::EvaluateBlockedInputComponent(UInputComponent* InputComponent)
+{
+	Super::EvaluateBlockedInputComponent(InputComponent);
+
+	if (UEnhancedInputComponent* IC = Cast<UEnhancedInputComponent>(InputComponent))
 	{
-		FInputActionInstance& ActionData = ActionPair.Value;
-		switch (ActionData.TriggerEvent)
-		{
-		case ETriggerEvent::None:
-		case ETriggerEvent::Canceled:
-		case ETriggerEvent::Completed:
-			ActionData.ElapsedProcessedTime = 0.f;
-			break;
-		}
-		if (ActionData.TriggerEvent != ETriggerEvent::Triggered)
+		for (const FEnhancedInputActionValueBinding& Binding : IC->GetActionValueBindings())
 		{
-			ActionData.ElapsedTriggeredTime = 0.f;
+			Binding.CurrentValue.Reset();
 		}
-
-		// Delay MappingTriggerState reset until here to allow dependent triggers (e.g. chords) access to this tick's values.
-		ActionData.TriggerStateTracker = FTriggerStateTracker();
-	}
-
-	// Remove any input actions that are no longer mapped to the player. At this point they will have been evaluated and 
-	// fired any "Canceled" input events needed
-	for (const UInputAction* IA : ActionsThatHaveBeenRemovedFromMappings)
-	{
-		ActionInstanceData.Remove(IA);
 	}
-
-	// We can clear our queue of mappings that have been removed on the next iteration directly after a key rebuild.
-	ActionsThatHaveBeenRemovedFromMappings.Reset();
-
-	LastFrameTime = CurrentTime;
-	KeysPressedThisTick.Reset();
-	bIsFlushingInputThisFrame = false;
-
-	Super::EvaluateInputDelegates(InputComponentStack, DeltaTime, bGamePaused, KeysWithEvents);
 }
 
 void UEnhancedPlayerInput::FlushPressedKeys()
diff --git a/Engine/Plugins/EnhancedInput/Source/EnhancedInput/Public/EnhancedPlayerInput.h b/Engine/Plugins/EnhancedInput/Source/EnhancedInput/Public/EnhancedPlayerInput.h
--- a/Engine/Plugins/EnhancedInput/Source/EnhancedInput/Public/EnhancedPlayerInput.h
+++ b/Engine/Plugins/EnhancedInput/Source/EnhancedInput/Public/EnhancedPlayerInput.h
@@ -154,6 +154,9 @@ protected:
 
 	UE_API virtual void EvaluateKeyMapState(const float DeltaTime, const bool bGamePaused, OUT TArray<TPair<FKey, FKeyState*>>& KeysWithEvents) override;
 	UE_API virtual void EvaluateInputDelegates(const TArray<UInputComponent*>& InputComponentStack, const float DeltaTime, const bool bGamePaused, const TArray<TPair<FKey, FKeyState*>>& KeysWithEvents) override;
+	UE_API virtual void PrepareInputDelegatesForEvaluation(const TArray<UInputComponent*>& InputComponentStack, const float DeltaTime, const bool bGamePaused, const TArray<TPair<FKey, FKeyState*>>& KeysWithEvents) override;
+	UE_API virtual bool EvaluateInputComponentDelegates(UInputComponent* const IC, const TArray<TPair<FKey, FKeyState*>>& KeysWithEvents, const float DeltaTime, const bool bGamePaused) override;
+	UE_API virtual void EvaluateBlockedInputComponent(UInputComponent* InputComponent) override;
 	
 	// Causes key to be consumed if it is affecting an action.
 	UE_API virtual bool IsKeyHandledByAction(FKey Key) const override;
@@ -184,6 +187,13 @@ protected:
 	UPROPERTY(Transient)
 	mutable TMap<TObjectPtr<const UInputAction>, FInputActionInstance> ActionInstanceData;
 
+	/** 
+	* Set of consumed input actions this frame. 
+	* Populated during the evaluation of input delegates and cleared at the end of the evaluation 
+	*/
+	UPROPERTY(Transient)
+	TSet<TObjectPtr<const UInputAction>> ConsumedInputActions;
+	
 private:
 
 	/** Add a player specific action mapping.
diff --git a/Engine/Source/Runtime/Engine/Classes/GameFramework/PlayerInput.h b/Engine/Source/Runtime/Engine/Classes/GameFramework/PlayerInput.h
--- a/Engine/Source/Runtime/Engine/Classes/GameFramework/PlayerInput.h
+++ b/Engine/Source/Runtime/Engine/Classes/GameFramework/PlayerInput.h
@@ -665,6 +665,26 @@ protected:
 	 * This should only be called from UPlayerInput::ProcessInputStack
 	 */
 	ENGINE_API virtual void EvaluateInputDelegates(const TArray<UInputComponent*>& InputComponentStack, const float DeltaTime, const bool bGamePaused, const TArray<TPair<FKey, FKeyState*>>& KeysWithEvents);
+
+	/** 
+	* Called just before iteration every input component on the stack to process their bound listener delegates.
+	* This provides an extension point to handle modifying the current input state before determining which input delegates should be fired.
+	*/
+	virtual void PrepareInputDelegatesForEvaluation(const TArray<UInputComponent*>& InputComponentStack, const float DeltaTime, const bool bGamePaused, const TArray<TPair<FKey, FKeyState*>>& KeysWithEvents) {}
+
+	/**
+	 * Evaluates a single input component on the component stack and executes its bound delegates.
+	 * 
+	 * @return True if this input component should BLOCK the rest of the input stack, preventing further components from evaluating their delegates.
+	 */
+	ENGINE_API virtual bool EvaluateInputComponentDelegates(UInputComponent* const IC, const TArray<TPair<FKey, FKeyState*>>& KeysWithEvents, const float DeltaTime, const bool bGamePaused);
+	
+	/** 
+	* Called on any input component which was not evaluated due to the input being blocked by another component earlier in the stack.
+	* These input components will not have had the EvaluateInputComponentDelegates flow executed, because they are blocked.
+	* This is a good spot to reset any delegate values if desired to reflect that they didn't execute for next frame.
+	*/
+	ENGINE_API virtual void EvaluateBlockedInputComponent(UInputComponent* InputComponent);
 	
 public:
 
diff --git a/Engine/Source/Runtime/Engine/Private/UserInterface/PlayerInput.cpp b/Engine/Source/Runtime/Engine/Private/UserInterface/PlayerInput.cpp
--- a/Engine/Source/Runtime/Engine/Private/UserInterface/PlayerInput.cpp
+++ b/Engine/Source/Runtime/Engine/Private/UserInterface/PlayerInput.cpp
@@ -57,6 +57,11 @@ namespace UE
 		static FAutoConsoleVariableRef CVarClearAxisValueIfConsumed(TEXT("Input.ClearAxisValueIfConsumed"),
 			bClearAxisValueIfConsumed,
 			TEXT("If true, we will clear the value of any FInputAxisKeyBinding whose FKey has been previously consumed.\nNote: This option will be removed in a future update."));
+			
+		static bool bMaintainLegacyDelegateExecutionOrder = false;
+		static FAutoConsoleVariableRef CVarMaintainLegacyDelegateExecutionOrder(TEXT("Input.bMaintainLegacyDelegateExecutionOrder"),
+			bMaintainLegacyDelegateExecutionOrder,
+			TEXT("If true, legacy input delegates will not fire after processing each UInputComponent, but instead wait until every UInputComponent has been processed."));
 
 		const TCHAR* LexToString(const EInputEvent Event)
 		{
@@ -1347,10 +1352,9 @@ void UPlayerInput::EvaluateKeyMapState(const float DeltaTime, const bool bGamePa
 	EventCount = 0;
 }
 
-void UPlayerInput::EvaluateInputDelegates(const TArray<UInputComponent*>& InputComponentStack, const float DeltaTime, const bool bGamePaused, const TArray<TPair<FKey, FKeyState*>>& KeysWithEvents)
+namespace UE::Input
 {
-	TRACE_CPUPROFILER_EVENT_SCOPE(UPlayerInput::EvaluateInputDelegates);
-	
+
 	// We collect axis contributions by delegate, so we can sum up
 	// contributions from multiple bindings.
 	struct FAxisDelegateDetails
@@ -1364,6 +1368,7 @@ void UPlayerInput::EvaluateInputDelegates(const TArray<UInputComponent*>& InputC
 		{
 		}
 	};
+	
 	struct FVectorAxisDelegateDetails
 	{
 		FInputVectorAxisUnifiedDelegate Delegate;
@@ -1375,7 +1380,7 @@ void UPlayerInput::EvaluateInputDelegates(const TArray<UInputComponent*>& InputC
 		{
 		}
 	};
-	
+
 	struct FDelegateDispatchDetailsSorter
 	{
 		bool operator()( const FDelegateDispatchDetails& A, const FDelegateDispatchDetails& B ) const
@@ -1384,18 +1389,69 @@ void UPlayerInput::EvaluateInputDelegates(const TArray<UInputComponent*>& InputC
 		}
 	};
 
-	
 	static TArray<FAxisDelegateDetails> AxisDelegates;
 	static TArray<FVectorAxisDelegateDetails> VectorAxisDelegates;
 	static TArray<FDelegateDispatchDetails> NonAxisDelegates;
-	static TArray<FKey> KeysToConsume;
-	static TArray<FDelegateDispatchDetails> FoundChords;	
-	static TArray<TSharedPtr<FInputActionBinding>> PotentialActions;
+	
+	/**
+	* Sorts the AxisDelegates, VectorAxisDelegates, and NonAxisDelegates arrays on the order that they occured
+	* and executes them
+	*/
+	void SortAndExecuteDelegates()
+	{
+		// Execute the delegates for this component 
+    	NonAxisDelegates.Sort(FDelegateDispatchDetailsSorter());
+    	for (const FDelegateDispatchDetails& Details : NonAxisDelegates)
+    	{
+    		if (Details.ActionDelegate.IsBound())
+    		{
+    			Details.ActionDelegate.Execute(Details.Chord.Key);
+    		}
+    		else if (Details.TouchDelegate.IsBound())
+    		{
+    			Details.TouchDelegate.Execute(Details.FingerIndex, Details.TouchLocation);
+    		}
+    		else if (Details.GestureDelegate.IsBound())
+    		{
+    			Details.GestureDelegate.Execute(Details.GestureValue);
+    		}
+    	}
+    	// Now dispatch delegates for summed axes
+    	for (const FAxisDelegateDetails& Details : AxisDelegates)
+    	{
+    		if (Details.Delegate.IsBound())
+    		{
+    			Details.Delegate.Execute(Details.Value);
+    		}
+    	}
+    	for (const FVectorAxisDelegateDetails& Details : VectorAxisDelegates)
+    	{
+    		if (Details.Delegate.IsBound())
+    		{
+    			Details.Delegate.Execute(Details.Value);
+    		}
+    	}
+    	
+    	AxisDelegates.Reset();
+    	VectorAxisDelegates.Reset();
+    	NonAxisDelegates.Reset();
+	}
+	
+	
+};
 
+void UPlayerInput::EvaluateInputDelegates(const TArray<UInputComponent*>& InputComponentStack, const float DeltaTime, const bool bGamePaused, const TArray<TPair<FKey, FKeyState*>>& KeysWithEvents)
+{
+	TRACE_CPUPROFILER_EVENT_SCOPE(UPlayerInput::EvaluateInputDelegates);
+
+	using namespace UE::Input;
+
+	PrepareInputDelegatesForEvaluation(InputComponentStack, DeltaTime, bGamePaused, KeysWithEvents);
+	
 	// must be called non-recursively and on the game thread
-	check(IsInGameThread() && !AxisDelegates.Num() && !VectorAxisDelegates.Num() && !NonAxisDelegates.Num() && !KeysToConsume.Num() && !FoundChords.Num() && !EventIndices.Num() && !PotentialActions.Num());
+	check(IsInGameThread() && !AxisDelegates.Num() && !VectorAxisDelegates.Num() && !NonAxisDelegates.Num() && !EventIndices.Num());
 
-	int32 StackIndex = InputComponentStack.Num()-1;
+	int32 StackIndex = InputComponentStack.Num() - 1;
 
 	// Walk the stack, top to bottom
 	for ( ; StackIndex >= 0; --StackIndex)
@@ -1403,266 +1459,269 @@ void UPlayerInput::EvaluateInputDelegates(const TArray<UInputComponent*>& InputC
 		UInputComponent* const IC = InputComponentStack[StackIndex];
 		if (IsValid(IC))
 		{
-			check(!KeysToConsume.Num() && !FoundChords.Num() && !EventIndices.Num() && !PotentialActions.Num());
-
-			IC->ConditionalBuildKeyMap(this);
-
-			for (const TPair<FKey,FKeyState*>& KeyWithEvent : KeysWithEvents)
-			{
-				if (!KeyWithEvent.Value->bConsumed)
-				{
-					FGetActionsBoundToKey::Get(IC, this, KeyWithEvent.Key, PotentialActions);
-				}
-			}
+			EvaluateInputComponentDelegates(IC, KeysWithEvents, DeltaTime, bGamePaused);
+		}
+	}
+	
+	// If the stack index is >=0, then an input component has "blocked" the input.
+	// Set any remaining component delegate values to zero. Otherwise, these delegate values
+	// have already been set and processed.
+	for ( ; StackIndex >= 0; --StackIndex)
+	{
+		if (UInputComponent* IC = InputComponentStack[StackIndex])
+		{
+			EvaluateBlockedInputComponent(IC);
+		}
+	}
+	
+	SortAndExecuteDelegates();
+}
 
-			for (const TSharedPtr<FInputActionBinding>& ActionBinding : PotentialActions)
-			{
-				GetChordsForAction(*ActionBinding.Get(), bGamePaused, FoundChords, KeysToConsume);
-			}
+bool UPlayerInput::EvaluateInputComponentDelegates(UInputComponent* const IC, const TArray<TPair<FKey, FKeyState*>>& KeysWithEvents, const float DeltaTime, const bool bGamePaused)
+{
+	if (!IsValid(IC))
+	{
+		return false;
+	}
 
-			PotentialActions.Reset();
+	using namespace UE::Input;
 
-			for (const FInputKeyBinding& KeyBinding : IC->KeyBindings)
-			{
-				GetChordForKey(KeyBinding, bGamePaused, FoundChords, KeysToConsume);
-			}
+	static TArray<FKey> KeysToConsume;
+	static TArray<FDelegateDispatchDetails> FoundChords;
+	static TArray<TSharedPtr<FInputActionBinding>> PotentialActions;
+	
+	check(!KeysToConsume.Num() && !FoundChords.Num() && !EventIndices.Num() && !PotentialActions.Num());
 
-			FoundChords.Sort(FDelegateDispatchDetailsSorter());
+	IC->ConditionalBuildKeyMap(this);
 
-			for (int32 ChordIndex=0; ChordIndex < FoundChords.Num(); ++ChordIndex)
-			{
-				const FDelegateDispatchDetails& FoundChord = FoundChords[ChordIndex];
-				bool bFireDelegate = true;
-				// If this is a paired action (implements both pressed and released) then we ensure that only one chord is
-				// handling the pairing
-				if (FoundChord.SourceAction && FoundChord.SourceAction->IsPaired())
-				{
-					FActionKeyDetails& KeyDetails = ActionKeyMap.FindChecked(FoundChord.SourceAction->GetActionName());
-					if (!KeyDetails.CapturingChord.Key.IsValid() || KeyDetails.CapturingChord == FoundChord.Chord || !IsPressed(KeyDetails.CapturingChord.Key))
-					{
-						if (FoundChord.SourceAction->KeyEvent == IE_Pressed)
-						{
-							KeyDetails.CapturingChord = FoundChord.Chord;
-						}
-						else
-						{
-							KeyDetails.CapturingChord.Key = EKeys::Invalid;
-						}
-					}
-					else
-					{
-						bFireDelegate = false;
-					}
-				}
+	for (const TPair<FKey,FKeyState*>& KeyWithEvent : KeysWithEvents)
+	{
+		if (!KeyWithEvent.Value->bConsumed)
+		{
+			FGetActionsBoundToKey::Get(IC, this, KeyWithEvent.Key, PotentialActions);
+		}
+	}
 
-				if (bFireDelegate && FoundChords[ChordIndex].ActionDelegate.IsBound())
-				{
-					FoundChords[ChordIndex].FoundIndex = NonAxisDelegates.Num();
-					NonAxisDelegates.Add(FoundChords[ChordIndex]);
-				}
-			}
+	for (const TSharedPtr<FInputActionBinding>& ActionBinding : PotentialActions)
+	{
+		GetChordsForAction(*ActionBinding.Get(), bGamePaused, FoundChords, KeysToConsume);
+	}
 
-			for (int32 TouchBindingIndex=0; TouchBindingIndex<IC->TouchBindings.Num(); ++TouchBindingIndex)
-			{
-				FInputTouchBinding& TB = IC->TouchBindings[TouchBindingIndex];
+	PotentialActions.Reset();
 
-				for (int32 TouchIndex = 0; TouchIndex < EKeys::NUM_TOUCH_KEYS; TouchIndex++)
-				{
-					const FKey& TouchKey = EKeys::TouchKeys[TouchIndex];
-					FKeyState* KeyState = KeyStateMap.Find(TouchKey);
-					if (KeyEventOccurred(TouchKey, TB.KeyEvent, EventIndices, KeyState) && !IsKeyConsumed(TouchKey, KeyState))
-					{
-						if (TB.bExecuteWhenPaused || !bGamePaused)
-						{
-							check(EventIndices.Num() > 0);
-							FVector* TouchedLocation = TouchEventLocations.Find(EventIndices[0]);
-							FDelegateDispatchDetails TouchInfo(EventIndices[0], NonAxisDelegates.Num(), TB.TouchDelegate, TouchedLocation != nullptr ? *TouchedLocation : FVector(-1.0f, -1.0f, 0.0f), (ETouchIndex::Type)TouchIndex);
-							NonAxisDelegates.Add(TouchInfo);
-							for (int32 EventsIndex = 1; EventsIndex < EventIndices.Num(); ++EventsIndex)
-							{
-								TouchInfo.EventIndex = EventIndices[EventsIndex];
-								TouchedLocation = TouchEventLocations.Find(TouchInfo.EventIndex);
-								TouchInfo.TouchLocation = TouchedLocation != nullptr ? *TouchedLocation : FVector(-1.0f, -1.0f, 0.0f);
-								NonAxisDelegates.Add(TouchInfo);
-							}
-						}
-						if (TB.bConsumeInput)
-						{
-							KeysToConsume.AddUnique(TouchKey);
-						}
-					}
-				}
+	for (const FInputKeyBinding& KeyBinding : IC->KeyBindings)
+	{
+		GetChordForKey(KeyBinding, bGamePaused, FoundChords, KeysToConsume);
+	}
 
-				EventIndices.Reset();
-			}
+	FoundChords.Sort(FDelegateDispatchDetailsSorter());
 
-			// look for any gestures that happened
-			for (FInputGestureBinding& GB : IC->GestureBindings)
+	for (int32 ChordIndex=0; ChordIndex < FoundChords.Num(); ++ChordIndex)
+	{
+		const FDelegateDispatchDetails& FoundChord = FoundChords[ChordIndex];
+		bool bFireDelegate = true;
+		// If this is a paired action (implements both pressed and released) then we ensure that only one chord is
+		// handling the pairing
+		if (FoundChord.SourceAction && FoundChord.SourceAction->IsPaired())
+		{
+			FActionKeyDetails& KeyDetails = ActionKeyMap.FindChecked(FoundChord.SourceAction->GetActionName());
+			if (!KeyDetails.CapturingChord.Key.IsValid() || KeyDetails.CapturingChord == FoundChord.Chord || !IsPressed(KeyDetails.CapturingChord.Key))
 			{
-				FKeyState* KeyState = KeyStateMap.Find(GB.GestureKey);
-				// treat gestures as fire-and-forget, so by convention we assume if they happen, it was a "pressed" event
-				if (KeyEventOccurred(GB.GestureKey, IE_Pressed, EventIndices, KeyState) && !IsKeyConsumed(GB.GestureKey, KeyState))
+				if (FoundChord.SourceAction->KeyEvent == IE_Pressed)
 				{
-					if (KeyState)
-					{
-						check(EventIndices.Num() > 0);
-						FDelegateDispatchDetails GestureInfo(EventIndices[0], NonAxisDelegates.Num(), GB.GestureDelegate, KeyState->Value.X);
-						NonAxisDelegates.Emplace(MoveTemp(GestureInfo));
-
-						if (GB.bConsumeInput)
-						{
-							KeysToConsume.AddUnique(GB.GestureKey);
-						}
-					}
+					KeyDetails.CapturingChord = FoundChord.Chord;
 				}
-
-				EventIndices.Reset();
-			}
-
-			// A flag that is set in DetermineAxisValue. If false, then all keys related to the axis binding have been consumed.
-			bool bHadAnyKeys = false;
-			
-			// Run though game axis bindings and accumulate axis values
-			for (FInputAxisBinding& AB : IC->AxisBindings)
-			{
-				AB.AxisValue = DetermineAxisValue(AB, bGamePaused, KeysToConsume, bHadAnyKeys); 
-				
-				if ((bHadAnyKeys || !UE::Input::AxisEventsCanBeConsumed) && AB.AxisDelegate.IsBound())
+				else
 				{
-					AxisDelegates.Emplace(FAxisDelegateDetails(AB.AxisDelegate, AB.AxisValue));
+					KeyDetails.CapturingChord.Key = EKeys::Invalid;
 				}
 			}
-			for (FInputAxisKeyBinding& AxisKeyBinding : IC->AxisKeyBindings)
+			else
 			{
-				if (!IsKeyConsumed(AxisKeyBinding.AxisKey))
-				{
-					if (!bGamePaused || AxisKeyBinding.bExecuteWhenPaused)
-					{
-						AxisKeyBinding.AxisValue = GetKeyValue(AxisKeyBinding.AxisKey);
-					}
-					else
-					{
-						AxisKeyBinding.AxisValue = 0.f;
-					}
+				bFireDelegate = false;
+			}
+		}
 
-					if (AxisKeyBinding.bConsumeInput)
-					{
-						KeysToConsume.AddUnique(AxisKeyBinding.AxisKey);
-					}
-				}
-				else if(UE::Input::bClearAxisValueIfConsumed)
-				{
-					AxisKeyBinding.AxisValue = 0.f;
-				}
+		if (bFireDelegate && FoundChords[ChordIndex].ActionDelegate.IsBound())
+		{
+			FoundChords[ChordIndex].FoundIndex = NonAxisDelegates.Num();
+			NonAxisDelegates.Add(FoundChords[ChordIndex]);
+		}
+	}
 
-				if (AxisKeyBinding.AxisDelegate.IsBound())
-				{
-					AxisDelegates.Emplace(FAxisDelegateDetails(AxisKeyBinding.AxisDelegate, AxisKeyBinding.AxisValue));
-				}
-			}
-			for (FInputVectorAxisBinding& VectorAxisBinding : IC->VectorAxisBindings)
+	for (int32 TouchBindingIndex=0; TouchBindingIndex<IC->TouchBindings.Num(); ++TouchBindingIndex)
+	{
+		FInputTouchBinding& TB = IC->TouchBindings[TouchBindingIndex];
+
+		for (int32 TouchIndex = 0; TouchIndex < EKeys::NUM_TOUCH_KEYS; TouchIndex++)
+		{
+			const FKey& TouchKey = EKeys::TouchKeys[TouchIndex];
+			FKeyState* KeyState = KeyStateMap.Find(TouchKey);
+			if (KeyEventOccurred(TouchKey, TB.KeyEvent, EventIndices, KeyState) && !IsKeyConsumed(TouchKey, KeyState))
 			{
-				if (!IsKeyConsumed(VectorAxisBinding.AxisKey))
+				if (TB.bExecuteWhenPaused || !bGamePaused)
 				{
-					if (!bGamePaused || VectorAxisBinding.bExecuteWhenPaused)
-					{
-						VectorAxisBinding.AxisValue = GetProcessedVectorKeyValue(VectorAxisBinding.AxisKey);
-					}
-					else
-					{
-						VectorAxisBinding.AxisValue = FVector::ZeroVector;
-					}
-
-					if (VectorAxisBinding.bConsumeInput)
+					check(EventIndices.Num() > 0);
+					FVector* TouchedLocation = TouchEventLocations.Find(EventIndices[0]);
+					FDelegateDispatchDetails TouchInfo(EventIndices[0], NonAxisDelegates.Num(), TB.TouchDelegate, TouchedLocation != nullptr ? *TouchedLocation : FVector(-1.0f, -1.0f, 0.0f), (ETouchIndex::Type)TouchIndex);
+					NonAxisDelegates.Add(TouchInfo);
+					for (int32 EventsIndex = 1; EventsIndex < EventIndices.Num(); ++EventsIndex)
 					{
-						KeysToConsume.AddUnique(VectorAxisBinding.AxisKey);
+						TouchInfo.EventIndex = EventIndices[EventsIndex];
+						TouchedLocation = TouchEventLocations.Find(TouchInfo.EventIndex);
+						TouchInfo.TouchLocation = TouchedLocation != nullptr ? *TouchedLocation : FVector(-1.0f, -1.0f, 0.0f);
+						NonAxisDelegates.Add(TouchInfo);
 					}
 				}
-
-				if (VectorAxisBinding.AxisDelegate.IsBound())
+				if (TB.bConsumeInput)
 				{
-					VectorAxisDelegates.Emplace(FVectorAxisDelegateDetails(VectorAxisBinding.AxisDelegate, VectorAxisBinding.AxisValue));
+					KeysToConsume.AddUnique(TouchKey);
 				}
 			}
+		}
 
-			if (IC->bBlockInput)
-			{
-				// stop traversing the stack, all input has been consumed by this InputComponent
-				--StackIndex;
-				KeysToConsume.Reset();
-				FoundChords.Reset();
-				break;
-			}
+		EventIndices.Reset();
+	}
 
-			// we do this after finishing the whole component, so we don't consume a key while there might be more bindings to it
-			for (int32 KeyIndex=0; KeyIndex<KeysToConsume.Num(); ++KeyIndex)
+	// look for any gestures that happened
+	for (FInputGestureBinding& GB : IC->GestureBindings)
+	{
+		FKeyState* KeyState = KeyStateMap.Find(GB.GestureKey);
+		// treat gestures as fire-and-forget, so by convention we assume if they happen, it was a "pressed" event
+		if (KeyEventOccurred(GB.GestureKey, IE_Pressed, EventIndices, KeyState) && !IsKeyConsumed(GB.GestureKey, KeyState))
+		{
+			if (KeyState)
 			{
-				if (FKeyState* KeyState = KeyStateMap.Find(KeysToConsume[KeyIndex]))
+				check(EventIndices.Num() > 0);
+				FDelegateDispatchDetails GestureInfo(EventIndices[0], NonAxisDelegates.Num(), GB.GestureDelegate, KeyState->Value.X);
+				NonAxisDelegates.Emplace(MoveTemp(GestureInfo));
+
+				if (GB.bConsumeInput)
 				{
-					KeyState->bConsumed = true;	
+					KeysToConsume.AddUnique(GB.GestureKey);
 				}
 			}
-			KeysToConsume.Reset();
-			FoundChords.Reset();
 		}
+
+		EventIndices.Reset();
 	}
 
-	for ( ; StackIndex >= 0; --StackIndex)
+	// A flag that is set in DetermineAxisValue. If false, then all keys related to the axis binding have been consumed.
+	bool bHadAnyKeys = false;
+	
+	// Run though game axis bindings and accumulate axis values
+	for (FInputAxisBinding& AB : IC->AxisBindings)
+	{
+		AB.AxisValue = DetermineAxisValue(AB, bGamePaused, KeysToConsume, bHadAnyKeys); 
+		
+		if ((bHadAnyKeys || !UE::Input::AxisEventsCanBeConsumed) && AB.AxisDelegate.IsBound())
+		{
+			AxisDelegates.Emplace(FAxisDelegateDetails(AB.AxisDelegate, AB.AxisValue));
+		}
+	}
+	for (FInputAxisKeyBinding& AxisKeyBinding : IC->AxisKeyBindings)
 	{
-		UInputComponent* IC = InputComponentStack[StackIndex];
-		if (IC)
+		if (!IsKeyConsumed(AxisKeyBinding.AxisKey))
 		{
-			for (FInputAxisBinding& AxisBinding : IC->AxisBindings)
+			if (!bGamePaused || AxisKeyBinding.bExecuteWhenPaused)
 			{
-				AxisBinding.AxisValue = 0.f;
+				AxisKeyBinding.AxisValue = GetKeyValue(AxisKeyBinding.AxisKey);
 			}
-			for (FInputAxisKeyBinding& AxisKeyBinding : IC->AxisKeyBindings)
+			else
 			{
 				AxisKeyBinding.AxisValue = 0.f;
 			}
-			for (FInputVectorAxisBinding& VectorAxisBinding : IC->VectorAxisBindings)
+
+			if (AxisKeyBinding.bConsumeInput)
 			{
-				VectorAxisBinding.AxisValue = FVector::ZeroVector;
+				KeysToConsume.AddUnique(AxisKeyBinding.AxisKey);
 			}
 		}
-	}
+		else if(UE::Input::bClearAxisValueIfConsumed)
+		{
+			AxisKeyBinding.AxisValue = 0.f;
+		}
 
-	// Dispatch the delegates in the order they occurred
-	NonAxisDelegates.Sort(FDelegateDispatchDetailsSorter());
-	for (const FDelegateDispatchDetails& Details : NonAxisDelegates)
-	{
-		if (Details.ActionDelegate.IsBound())
+		if (AxisKeyBinding.AxisDelegate.IsBound())
 		{
-			Details.ActionDelegate.Execute(Details.Chord.Key);
+			AxisDelegates.Emplace(FAxisDelegateDetails(AxisKeyBinding.AxisDelegate, AxisKeyBinding.AxisValue));
 		}
-		else if (Details.TouchDelegate.IsBound())
+	}
+	for (FInputVectorAxisBinding& VectorAxisBinding : IC->VectorAxisBindings)
+	{
+		if (!IsKeyConsumed(VectorAxisBinding.AxisKey))
 		{
-			Details.TouchDelegate.Execute(Details.FingerIndex, Details.TouchLocation);
+			if (!bGamePaused || VectorAxisBinding.bExecuteWhenPaused)
+			{
+				VectorAxisBinding.AxisValue = GetProcessedVectorKeyValue(VectorAxisBinding.AxisKey);
+			}
+			else
+			{
+				VectorAxisBinding.AxisValue = FVector::ZeroVector;
+			}
+
+			if (VectorAxisBinding.bConsumeInput)
+			{
+				KeysToConsume.AddUnique(VectorAxisBinding.AxisKey);
+			}
 		}
-		else if (Details.GestureDelegate.IsBound())
+
+		if (VectorAxisBinding.AxisDelegate.IsBound())
 		{
-			Details.GestureDelegate.Execute(Details.GestureValue);
+			VectorAxisDelegates.Emplace(FVectorAxisDelegateDetails(VectorAxisBinding.AxisDelegate, VectorAxisBinding.AxisValue));
 		}
 	}
-	// Now dispatch delegates for summed axes
-	for (const FAxisDelegateDetails& Details : AxisDelegates)
+
+	const bool bComponentBlocksInput = IC->bBlockInput; 
+
+	if (bComponentBlocksInput)
 	{
-		if (Details.Delegate.IsBound())
-		{
-			Details.Delegate.Execute(Details.Value);
-		}
+		// stop traversing the stack, all input has been consumed by this InputComponent
+		KeysToConsume.Reset();
+		FoundChords.Reset();
 	}
-	for (const FVectorAxisDelegateDetails& Details : VectorAxisDelegates)
+
+	// we do this after finishing the whole component, so we don't consume a key while there might be more bindings to it
+	for (int32 KeyIndex=0; KeyIndex<KeysToConsume.Num(); ++KeyIndex)
 	{
-		if (Details.Delegate.IsBound())
+		if (FKeyState* KeyState = KeyStateMap.Find(KeysToConsume[KeyIndex]))
 		{
-			Details.Delegate.Execute(Details.Value);
+			KeyState->bConsumed = true;	
 		}
 	}
+	KeysToConsume.Reset();
+	FoundChords.Reset();
+
+	// This will make the input delegates execute after each component is processed, instead of waiting until all the components are done
+	// This will make it so that the order of execution on the delegates is actually set by the UInputComponent priority level. 
+	if (!UE::Input::bMaintainLegacyDelegateExecutionOrder)
+	{
+		// Execute the delegates for this component 
+    	SortAndExecuteDelegates();
+	}	
+
+	return bComponentBlocksInput;	
+}
 
-	AxisDelegates.Reset();
-	VectorAxisDelegates.Reset();
-	NonAxisDelegates.Reset();
+void UPlayerInput::EvaluateBlockedInputComponent(UInputComponent* InputComponent)
+{
+	if (!InputComponent)
+	{
+		return;
+	}
+	
+	for (FInputAxisBinding& AxisBinding : InputComponent->AxisBindings)
+	{
+		AxisBinding.AxisValue = 0.f;
+	}
+	for (FInputAxisKeyBinding& AxisKeyBinding : InputComponent->AxisKeyBindings)
+	{
+		AxisKeyBinding.AxisValue = 0.f;
+	}
+	for (FInputVectorAxisBinding& VectorAxisBinding : InputComponent->VectorAxisBindings)
+	{
+		VectorAxisBinding.AxisValue = FVector::ZeroVector;
+	}
 }
 
 void UPlayerInput::DiscardPlayerInput()
