[Input] Fix to ensure that Input Components block input value is actually read. Also, fix a mistake where the bConsumed flag on keys would be set even if the UInputComponent was blocking, which was not what happened before. These issues were introduced with 46335000

#rnx

diff --git a/Engine/Source/Runtime/Engine/Private/UserInterface/PlayerInput.cpp b/Engine/Source/Runtime/Engine/Private/UserInterface/PlayerInput.cpp
--- a/Engine/Source/Runtime/Engine/Private/UserInterface/PlayerInput.cpp
+++ b/Engine/Source/Runtime/Engine/Private/UserInterface/PlayerInput.cpp
@@ -1459,7 +1459,13 @@ void UPlayerInput::EvaluateInputDelegates(const TArray<UInputComponent*>& InputC
 		UInputComponent* const IC = InputComponentStack[StackIndex];
 		if (IsValid(IC))
 		{
-			EvaluateInputComponentDelegates(IC, KeysWithEvents, DeltaTime, bGamePaused);
+			const bool bDoesComponentBlockFurtherInput = EvaluateInputComponentDelegates(IC, KeysWithEvents, DeltaTime, bGamePaused);
+			if (bDoesComponentBlockFurtherInput)
+			{
+				// stop traversing the stack, all input has been consumed by this InputComponent
+                --StackIndex;
+				break;
+			}
 		}
 	}
 	
@@ -1674,24 +1680,21 @@ bool UPlayerInput::EvaluateInputComponentDelegates(UInputComponent* const IC, co
 
 	const bool bComponentBlocksInput = IC->bBlockInput; 
 
-	if (bComponentBlocksInput)
-	{
-		// stop traversing the stack, all input has been consumed by this InputComponent
-		KeysToConsume.Reset();
-		FoundChords.Reset();
-	}
-
-	// we do this after finishing the whole component, so we don't consume a key while there might be more bindings to it
-	for (int32 KeyIndex=0; KeyIndex<KeysToConsume.Num(); ++KeyIndex)
+	if (!bComponentBlocksInput)
 	{
-		if (FKeyState* KeyState = KeyStateMap.Find(KeysToConsume[KeyIndex]))
+		// we do this after finishing the whole component, so we don't consume a key while there might be more bindings to it
+		for (int32 KeyIndex = 0; KeyIndex < KeysToConsume.Num(); ++KeyIndex)
 		{
-			KeyState->bConsumed = true;	
+			if (FKeyState* KeyState = KeyStateMap.Find(KeysToConsume[KeyIndex]))
+			{
+				KeyState->bConsumed = true;	
+			}
 		}
 	}
+
 	KeysToConsume.Reset();
 	FoundChords.Reset();
-
+	
 	// This will make the input delegates execute after each component is processed, instead of waiting until all the components are done
 	// This will make it so that the order of execution on the delegates is actually set by the UInputComponent priority level. 
 	if (!UE::Input::bMaintainLegacyDelegateExecutionOrder)
