[CustomHLOD] Fix crash when validating actors from another world

diff --git a/Engine/Source/Runtime/Engine/Private/WorldPartition/WorldPartitionStreamingGeneration.cpp b/Engine/Source/Runtime/Engine/Private/WorldPartition/WorldPartitionStreamingGeneration.cpp
--- a/Engine/Source/Runtime/Engine/Private/WorldPartition/WorldPartitionStreamingGeneration.cpp
+++ b/Engine/Source/Runtime/Engine/Private/WorldPartition/WorldPartitionStreamingGeneration.cpp
@@ -913,47 +913,50 @@ class FWorldPartitionStreamingGenerator
 
 			if (Iterator->GetActorNativeClass()->IsChildOf<AWorldPartitionCustomHLOD>())
 			{
-				// Custom HLOD actors, unlike regular HLOD actors don't have their SourceCell set during creation, so we need to compute and store that information during streaming generation
-				// We're doing this by adding a placeholder actor descriptor for each custom HLOD actor. Placeholder actor has the same transform, bounds, etc. as the Custom HLOD actor.
-				// Later, when placing actors in the appropriate cells we'll save that information and remove the placeholder actors.
-				// The easiest way to add actor desc view is to spawn an actual actor and use the same logic that is used to handle unsaved dirty actors. 
-				// Right after we have actor desc view, we can destroy actor to avoid undesired side effects.
-				// To keep streaming generation deterministic, we use a deterministic actor GUID and name.
-				FActorSpawnParameters SpawnParameters;
-				SpawnParameters.ObjectFlags = RF_Transient;
-				SpawnParameters.bCreateActorPackage = false;
-				SpawnParameters.OverrideActorGuid = FGuid::NewDeterministicGuid(TEXT("CustomHLODPlaceholder") + Iterator->GetActorSoftPath().ToString());
-				SpawnParameters.Name = FName(SpawnParameters.OverrideActorGuid.ToString());
-				SpawnParameters.NameMode = FActorSpawnParameters::ESpawnActorNameMode::Requested;
-
-				// Spawn placeholder actor and initialize properties based on Custom HLOD actor
-				AWorldPartitionCustomHLODPlaceholder* CustomHLODPlaceholderActor = WorldPartitionContext->GetWorld()->SpawnActor<AWorldPartitionCustomHLODPlaceholder>(SpawnParameters);
-				check(CustomHLODPlaceholderActor);
-				CustomHLODPlaceholderActor->InitFrom(*Iterator);
-
-				if (UHLODLayer* HLODLayer = Cast<UHLODLayer>(Iterator->GetHLODLayer().TryLoad()))
+				if (WorldPartitionContext)
 				{
-					if (HLODLayer->GetLayerType() == EHLODLayerType::CustomHLODActor)
+					// Custom HLOD actors, unlike regular HLOD actors don't have their SourceCell set during creation, so we need to compute and store that information during streaming generation
+					// We're doing this by adding a placeholder actor descriptor for each custom HLOD actor. Placeholder actor has the same transform, bounds, etc. as the Custom HLOD actor.
+					// Later, when placing actors in the appropriate cells we'll save that information and remove the placeholder actors.
+					// The easiest way to add actor desc view is to spawn an actual actor and use the same logic that is used to handle unsaved dirty actors. 
+					// Right after we have actor desc view, we can destroy actor to avoid undesired side effects.
+					// To keep streaming generation deterministic, we use a deterministic actor GUID and name.
+					FActorSpawnParameters SpawnParameters;
+					SpawnParameters.ObjectFlags = RF_Transient;
+					SpawnParameters.bCreateActorPackage = false;
+					SpawnParameters.OverrideActorGuid = FGuid::NewDeterministicGuid(TEXT("CustomHLODPlaceholder") + Iterator->GetActorSoftPath().ToString());
+					SpawnParameters.Name = FName(SpawnParameters.OverrideActorGuid.ToString());
+					SpawnParameters.NameMode = FActorSpawnParameters::ESpawnActorNameMode::Requested;
+
+					// Spawn placeholder actor and initialize properties based on Custom HLOD actor
+					AWorldPartitionCustomHLODPlaceholder* CustomHLODPlaceholderActor = WorldPartitionContext->GetWorld()->SpawnActor<AWorldPartitionCustomHLODPlaceholder>(SpawnParameters);
+					check(CustomHLODPlaceholderActor);
+					CustomHLODPlaceholderActor->InitFrom(*Iterator);
+
+					if (UHLODLayer* HLODLayer = Cast<UHLODLayer>(Iterator->GetHLODLayer().TryLoad()))
 					{
-						if (UHLODLayer* LinkedHLODLayer = HLODLayer->GetLinkedLayer())
+						if (HLODLayer->GetLayerType() == EHLODLayerType::CustomHLODActor)
 						{
-							FName ResolvedRuntimeGrid = GetResolvedRuntimeGridForHLODLayer(CustomHLODPlaceholderActor->GetRuntimeGrid(), LinkedHLODLayer);
-							if (!ResolvedRuntimeGrid.IsNone())
+							if (UHLODLayer* LinkedHLODLayer = HLODLayer->GetLinkedLayer())
 							{
-								CustomHLODPlaceholderActor->SetRuntimeGrid(ResolvedRuntimeGrid);
+								FName ResolvedRuntimeGrid = GetResolvedRuntimeGridForHLODLayer(CustomHLODPlaceholderActor->GetRuntimeGrid(), LinkedHLODLayer);
+								if (!ResolvedRuntimeGrid.IsNone())
+								{
+									CustomHLODPlaceholderActor->SetRuntimeGrid(ResolvedRuntimeGrid);
+								}
 							}
 						}
 					}
-				}
 				
-				// Register placeholder actor desc view
-				TUniquePtr<FStreamingGenerationTempActorDescInstance>& CustomHLODPlaceholderRef = OutCustomHLODActorPlaceholderInstances.Add_GetRef(MakeUnique<FStreamingGenerationTempActorDescInstance>(Iterator->GetContainerInstance(), CustomHLODPlaceholderActor));
-				RegisterActorDescView(FStreamingGenerationActorDescView(OutActorDescViewMap, CustomHLODPlaceholderRef.Get(), false));
-
-				// Destroy placeholder actor since we only need its actor descriptor.
-				WorldPartitionContext->GetWorld()->DestroyActor(CustomHLODPlaceholderActor, false, false);
-				FName UniqueTrashName = MakeUniqueObjectName(CustomHLODPlaceholderActor->GetOuter(), CustomHLODPlaceholderActor->GetClass(), NAME_TrashedPackage);
-				CustomHLODPlaceholderActor->Rename(*UniqueTrashName.ToString(), nullptr, REN_DontCreateRedirectors | REN_NonTransactional | REN_DoNotDirty);
+					// Register placeholder actor desc view
+					TUniquePtr<FStreamingGenerationTempActorDescInstance>& CustomHLODPlaceholderRef = OutCustomHLODActorPlaceholderInstances.Add_GetRef(MakeUnique<FStreamingGenerationTempActorDescInstance>(Iterator->GetContainerInstance(), CustomHLODPlaceholderActor));
+					RegisterActorDescView(FStreamingGenerationActorDescView(OutActorDescViewMap, CustomHLODPlaceholderRef.Get(), false));
+
+					// Destroy placeholder actor since we only need its actor descriptor.
+					WorldPartitionContext->GetWorld()->DestroyActor(CustomHLODPlaceholderActor, false, false);
+					FName UniqueTrashName = MakeUniqueObjectName(CustomHLODPlaceholderActor->GetOuter(), CustomHLODPlaceholderActor->GetClass(), NAME_TrashedPackage);
+					CustomHLODPlaceholderActor->Rename(*UniqueTrashName.ToString(), nullptr, REN_DontCreateRedirectors | REN_NonTransactional | REN_DoNotDirty);
+				}
 
 				// If Custom HLOD actor is used with some HLOD layer different than CustomHLODActor layer don't register its actor descriptor
 				if (bIncludeCustomHLODActorsWithOtherLayerType == false)
