[NetPhysics] Fix for inputs getting lost and minor fixes
- Implements a growing ordered array for marshaling inputs and states from GT to PT
- Don't produce AsyncOutput during resim in the NetworkPhysicsComponent
- Fix potential wrong modulo operation due to signage
- Guard against potential divide by 0 when calculating interpolation lerp
- Check correct TOptional property before dereferencing it

diff --git a/Engine/Source/Runtime/Engine/Private/PhysicsEngine/NetworkPhysicsComponent.cpp b/Engine/Source/Runtime/Engine/Private/PhysicsEngine/NetworkPhysicsComponent.cpp
--- a/Engine/Source/Runtime/Engine/Private/PhysicsEngine/NetworkPhysicsComponent.cpp
+++ b/Engine/Source/Runtime/Engine/Private/PhysicsEngine/NetworkPhysicsComponent.cpp
@@ -1143,9 +1143,7 @@ void UNetworkPhysicsComponent::OnRep_SetReplicatedStates()
 			AsyncInput->StateData = StateHelper->CreateUniqueRewindHistory(ReplicatedStates.History->GetHistorySize());
 		}
 
-		AsyncInput->StateData->ResetFast();
-		AsyncInput->StateData->ResizeDataHistory(ReplicatedStates.History->GetHistorySize(), EAllowShrinking::No);
-		ReplicatedStates.History->CopyAllData(*AsyncInput->StateData.Get(), /*bIncludeUnimportant*/ true, /*bIncludeImportant*/ true);
+		ReplicatedStates.History->CopyAllDataGrowingOrdered(*AsyncInput->StateData.Get());
 	}
 }
 
@@ -1172,9 +1170,7 @@ void UNetworkPhysicsComponent::OnRep_SetReplicatedInputs()
 			AsyncInput->InputData = InputHelper->CreateUniqueRewindHistory(ReplicatedInputs.History->GetHistorySize());
 		}
 
-		AsyncInput->InputData->ResetFast();
-		AsyncInput->InputData->ResizeDataHistory(ReplicatedInputs.History->GetHistorySize(), EAllowShrinking::No);
-		ReplicatedInputs.History->CopyAllData(*AsyncInput->InputData.Get(), /*bIncludeUnimportant*/ true, /*bIncludeImportant*/ true);
+		ReplicatedInputs.History->CopyAllDataGrowingOrdered(*AsyncInput->InputData.Get());
 	}
 }
 
@@ -1201,9 +1197,7 @@ void UNetworkPhysicsComponent::OnRep_SetReplicatedRemoteInputs()
 			AsyncInput->InputData = InputHelper->CreateUniqueRewindHistory(ReplicatedRemoteInputs.History->GetHistorySize());
 		}
 
-		AsyncInput->InputData->ResetFast();
-		AsyncInput->InputData->ResizeDataHistory(ReplicatedRemoteInputs.History->GetHistorySize(), EAllowShrinking::No);
-		ReplicatedRemoteInputs.History->CopyAllData(*AsyncInput->InputData.Get(), /*bIncludeUnimportant*/ true, /*bIncludeImportant*/ true);
+		ReplicatedRemoteInputs.History->CopyAllDataGrowingOrdered(*AsyncInput->InputData.Get());
 	}
 }
 
@@ -1234,9 +1228,7 @@ void UNetworkPhysicsComponent::ServerReceiveInputData_Implementation(const FNetw
 			ClientInputs.History->ValidateDataInHistory(ActorComponent.Get());
 		}
 
-		AsyncInput->InputData->ResizeDataHistory(ClientInputs.History->GetHistorySize(), EAllowShrinking::No);
-		AsyncInput->InputData->SetRecordDataIncremental(true); // Only record data that is newer than already cached data
-		ClientInputs.History->CopyAllData(*AsyncInput->InputData.Get(), /*bIncludeUnimportant*/ true, /*bIncludeImportant*/ true);
+		ClientInputs.History->CopyAllDataGrowingOrdered(*AsyncInput->InputData.Get());
 
 		// Send received inputs to remote clients
 		ReplicatedRemoteInputs.History->SetRecordDataIncremental(true); // Only record data that is newer than already cached data
@@ -1359,10 +1351,7 @@ void UNetworkPhysicsComponent::ServerReceiveInputCollection_Implementation(const
 			AsyncInput->InputData = InputHelper->CreateUniqueRewindHistory(ClientInputCollection.DataArray.Num());
 		}
 
-		AsyncInput->InputData->ResizeDataHistory(ClientInputCollection.DataArray.Num(), EAllowShrinking::No);
-		AsyncInput->InputData->SetRecordDataIncremental(true); // Only record data that is newer than already cached data
-
-		InputHelper->CopyData(ClientInputCollection, AsyncInput->InputData.Get());
+		InputHelper->CopyDataGrowingOrdered(ClientInputCollection, AsyncInput->InputData.Get());
 
 		// Validate Inputs
 		if (ImplementationInterface_External)
@@ -1396,10 +1385,7 @@ void UNetworkPhysicsComponent::OnRep_SetReplicatedInputCollection()
 			AsyncInput->InputData = InputHelper->CreateUniqueRewindHistory(ReplicatedInputCollection.DataArray.Num());
 		}
 
-		AsyncInput->InputData->ResetFast();
-		AsyncInput->InputData->ResizeDataHistory(ReplicatedInputCollection.DataArray.Num(), EAllowShrinking::No);
-
-		InputHelper->CopyData(ReplicatedInputCollection, AsyncInput->InputData.Get());
+		InputHelper->CopyDataGrowingOrdered(ReplicatedInputCollection, AsyncInput->InputData.Get());
 	}
 }
 
@@ -1422,10 +1408,7 @@ void UNetworkPhysicsComponent::OnRep_SetReplicatedRemoteInputCollection()
 			AsyncInput->InputData = InputHelper->CreateUniqueRewindHistory(ReplicatedRemoteInputCollection.DataArray.Num());
 		}
 
-		AsyncInput->InputData->ResetFast();
-		AsyncInput->InputData->ResizeDataHistory(ReplicatedRemoteInputCollection.DataArray.Num(), EAllowShrinking::No);
-
-		InputHelper->CopyData(ReplicatedRemoteInputCollection, AsyncInput->InputData.Get());
+		InputHelper->CopyDataGrowingOrdered(ReplicatedRemoteInputCollection, AsyncInput->InputData.Get());
 	}
 }
 
@@ -1448,10 +1431,7 @@ void UNetworkPhysicsComponent::OnRep_SetReplicatedStateCollection()
 			AsyncInput->StateData = StateHelper->CreateUniqueRewindHistory(ReplicatedStateCollection.DataArray.Num());
 		}
 
-		AsyncInput->StateData->ResetFast();
-		AsyncInput->StateData->ResizeDataHistory(ReplicatedStateCollection.DataArray.Num(), EAllowShrinking::No);
-
-		StateHelper->CopyData(ReplicatedStateCollection, AsyncInput->StateData.Get());
+		StateHelper->CopyDataGrowingOrdered(ReplicatedStateCollection, AsyncInput->StateData.Get());
 	}
 }
 
@@ -1834,7 +1814,7 @@ void FAsyncNetworkPhysicsComponent::ConsumeAsyncInput(const int32 PhysicsStep)
 			{
 				bCompareStateToTriggerRewind = *AsyncInput->bCompareStateToTriggerRewind;
 			}
-			if (AsyncInput->bCompareStateToTriggerRewind.IsSet())
+			if (AsyncInput->bCompareStateToTriggerRewindIncludeSimProxies.IsSet())
 			{
 				bCompareStateToTriggerRewindIncludeSimProxies = *AsyncInput->bCompareStateToTriggerRewindIncludeSimProxies;
 			}
@@ -2306,11 +2286,14 @@ void FAsyncNetworkPhysicsComponent::OnPostProcessInputs_Internal(const int32 Phy
 #endif
 	}
 
-	// Marshal inputs and states from PT to GT for networking
-	FAsyncNetworkPhysicsComponentOutput& AsyncOutput = GetAsyncOutput_Internal();
-	SendInputData_Internal(AsyncOutput, PhysicsStep);
-	SendStateData_Internal(AsyncOutput, PhysicsStep);
-	FinalizeOutputData_Internal();
+	if (bIsSolverResim == false)
+	{
+		// Marshal inputs and states from PT to GT for networking
+		FAsyncNetworkPhysicsComponentOutput& AsyncOutput = GetAsyncOutput_Internal();
+		SendInputData_Internal(AsyncOutput, PhysicsStep);
+		SendStateData_Internal(AsyncOutput, PhysicsStep);
+		FinalizeOutputData_Internal();
+	}
 }
 
 void FAsyncNetworkPhysicsComponent::SendInputData_Internal(FAsyncNetworkPhysicsComponentOutput& AsyncOutput, const int32 PhysicsStep)
diff --git a/Engine/Source/Runtime/Engine/Public/Physics/NetworkPhysicsComponent.h b/Engine/Source/Runtime/Engine/Public/Physics/NetworkPhysicsComponent.h
--- a/Engine/Source/Runtime/Engine/Public/Physics/NetworkPhysicsComponent.h
+++ b/Engine/Source/Runtime/Engine/Public/Physics/NetworkPhysicsComponent.h
@@ -238,7 +238,7 @@ struct TNetRewindHistory : public Chaos::TDataRewindHistory<DataType>
 
 		int32 RewindFrame = INDEX_NONE;
 		DataType& ReceiveData = NetNewData.GetAndLoadEarliestData();
-		
+
 		for (int32 Itr = 0; Itr < NetNewData.GetHistorySize(); Itr++)
 		{
 			DataType* NextReceiveData = NetNewData.GetAndLoadNextIncrementalData();
@@ -284,7 +284,7 @@ struct TNetRewindHistory : public Chaos::TDataRewindHistory<DataType>
 				break;
 			}
 		}
-		
+
 		return RewindFrame;
 	}
 
@@ -319,7 +319,7 @@ struct TNetRewindHistory : public Chaos::TDataRewindHistory<DataType>
 				}
 
 				// Mark the already cached data as received so that we don't perform this merge again when receiving redundant inputs / states
-				Data.bReceivedData = true; 
+				Data.bReceivedData = true;
 				return false;
 			}
 			else
@@ -806,7 +806,12 @@ struct FNetworkPhysicsPayload
 	UE_DEPRECATED(5.7, "Temporary virtual function for backwards compatibility with FNetworkPhysicsData. Use the InterpolateData function that passes the LerpAlpha instead")
 	virtual void DoInterpolateData(const FNetworkPhysicsPayload& MinData, const FNetworkPhysicsPayload& MaxData)
 	{
-		const float LerpAlpha = static_cast<float>(LocalFrame - MinData.LocalFrame) / static_cast<float>(MaxData.LocalFrame - MinData.LocalFrame);
+		float LerpAlpha = 0.5f;
+		if (MaxData.LocalFrame != MinData.LocalFrame)
+		{
+			LerpAlpha = static_cast<float>(LocalFrame - MinData.LocalFrame) / static_cast<float>(MaxData.LocalFrame - MinData.LocalFrame);
+			LerpAlpha = FMath::Clamp(LerpAlpha, 0.0f, 1.0f);
+		}
 		InterpolateData(MinData, MaxData, LerpAlpha);
 	}
 };
@@ -1121,16 +1126,19 @@ struct FNetworkPhysicsDataHelper
 
 	/** Copy data from the networked data collection to the rewind history */
 	virtual void CopyData(const FNetworkPhysicsDataCollection& From, Chaos::FBaseRewindHistory* To) = 0;
-	
+
 	/** Copy data from the rewind history to the networked data collection */
 	virtual void CopyData(const Chaos::FBaseRewindHistory* From, FNetworkPhysicsDataCollection& To) = 0;
-	
+
 	/** Copy data from the rewind history to the networked data collection and only copy data that is newer than the already existing data in the data collection */
 	virtual void CopyIncrementalData(const Chaos::FBaseRewindHistory* From, FNetworkPhysicsDataCollection& To) = 0;
-	
+
 	/** Copy data from the rewind history to the networked data collection and only copy data that has been altered since it was created (i.e. the client created an input that the server later altered) */
 	virtual bool CopyAlteredData(const Chaos::FBaseRewindHistory* From, FNetworkPhysicsDataCollection& To) = 0;
-	
+
+	/** Copy data from the networked data collection to the rewind history and store the data ordered instead of circular and grow the history to fit all unique data */
+	virtual void CopyDataGrowingOrdered(const FNetworkPhysicsDataCollection& From, Chaos::FBaseRewindHistory* To) = 0;
+
 	/** Call the ValidateData callback on the external implementation interface for data in the provided rewind history */
 	virtual void ValidateData(Chaos::FBaseRewindHistory* History, INetworkPhysicsInputState_External& Interface) = 0;
 
@@ -1237,6 +1245,17 @@ struct TNetworkPhysicsDataHelper : FNetworkPhysicsDataHelper
 		return bHasCopiedData;
 	}
 
+	/** Copy data from the networked data collection to the rewind history and store the data ordered instead of circular and grow the history to fit all unique data */
+	void CopyDataGrowingOrdered(const FNetworkPhysicsDataCollection& From, Chaos::FBaseRewindHistory* To) override
+	{
+		Chaos::TDataRewindHistory<DataType>* NetHistory = static_cast<Chaos::TDataRewindHistory<DataType>*>(To);
+
+		for (const TInstancedStruct<FNetworkPhysicsPayload>& DataInstance : From.DataArray)
+		{
+			NetHistory->RecordDataGrowingOrdered(&DataInstance.Get<DataType>());
+		}
+	}
+
 	/** Call the ValidateData callback on the external implementation interface for data in the provided rewind history */
 	void ValidateData(Chaos::FBaseRewindHistory* History, INetworkPhysicsInputState_External& Interface) override
 	{
@@ -1818,10 +1837,12 @@ struct FAsyncNetworkPhysicsComponentInput : public Chaos::FSimCallbackInput
 
 		if (InputData)
 		{
+			InputData->ResizeDataHistory(0, EAllowShrinking::No);
 			InputData->ResetFast();
 		}
 		if (StateData)
 		{
+			StateData->ResizeDataHistory(0, EAllowShrinking::No);
 			StateData->ResetFast();
 		}
 
diff --git a/Engine/Source/Runtime/Experimental/Chaos/Public/RewindData.h b/Engine/Source/Runtime/Experimental/Chaos/Public/RewindData.h
--- a/Engine/Source/Runtime/Experimental/Chaos/Public/RewindData.h
+++ b/Engine/Source/Runtime/Experimental/Chaos/Public/RewindData.h
@@ -73,8 +73,11 @@ struct FBaseRewindHistory
 	/** Record data into the history buffer at a given time */
 	FORCEINLINE virtual bool RecordData(const int32 RecordFrame, const void* HistoryData) { return true; }
 
+	/** Record data in order and allow growing the history */
+	FORCEINLINE virtual bool RecordDataGrowingOrdered(const void* HistoryData) { return true; }
+
 	/** Set if this history should only allow overriding of data if it has a higher value than current recorded data */
-	FORCEINLINE virtual void SetRecordDataIncremental(const bool bInIncremental) { }
+	FORCEINLINE virtual void SetRecordDataIncremental(const bool bInIncremental) {}
 
 	/** Copy all data from local history into into @param OutHistory
 	* @param bIncludeUnimportant = If to copy unimportant data entries
@@ -93,6 +96,9 @@ struct FBaseRewindHistory
 	* @param bIncludeImportant = If to copy important data entries */
 	virtual bool CopyData(Chaos::FBaseRewindHistory& OutHistory, const uint32 StartFrame, const uint32 EndFrame, bool bIncludeUnimportant = true, bool bIncludeImportant = false) { return false; }
 
+	/** Copy all data and record it ordered and growing */
+	virtual bool CopyAllDataGrowingOrdered(Chaos::FBaseRewindHistory& OutHistory) { return false; }
+
 	/** Create a polymorphic copy of only a range of frames, applying the frame offset to the copies
 	* @param StartFrame = Included
 	* @param EndFrame = Excluded */
@@ -338,11 +344,58 @@ public :
 		return true;
 	}
 
+	FORCEINLINE virtual bool RecordDataGrowingOrdered(const void* HistoryData) override
+	{
+		const DataType* HistoryDataCast = static_cast<const DataType*>(HistoryData);
+
+		int32 InsertAtIdx = 0;
+		int32 StoredFrame = 0;
+
+		if (NumFrames > 0)
+		{
+			for (int32 Idx = CurrentIndex; Idx >= 0; Idx--)
+			{
+				StoredFrame = DataHistory[Idx].LocalFrame;
+				if (StoredFrame == HistoryDataCast->LocalFrame)
+				{
+					// Data for frame already stored
+					return false;
+				}
+				else if (StoredFrame < HistoryDataCast->LocalFrame)
+				{
+					// Found index where earlier data is stored, insert newer data on the next index
+					InsertAtIdx = Idx + 1;
+					break;
+				}
+			}
+		}
+
+		DataHistory.Insert(*HistoryDataCast, InsertAtIdx);
+
+		NumFrames++;
+		CurrentIndex = NumFrames - 1;
+		LatestFrame = FMath::Max(LatestFrame, HistoryDataCast->LocalFrame);
+
+		return true;
+	}
+
 	FORCEINLINE virtual void SetRecordDataIncremental(const bool bInIncremental) override
 	{
 		bIncremental = bInIncremental;
 	}
 
+	virtual bool CopyAllDataGrowingOrdered(Chaos::FBaseRewindHistory& OutHistory) override
+	{
+		TDataRewindHistory& OutDataHistory = static_cast<TDataRewindHistory&>(OutHistory);
+		bool bHasCopiedData = false;
+
+		for (int32 CopyIndex = 0; CopyIndex < NumFrames; ++CopyIndex)
+		{
+			bHasCopiedData |= OutDataHistory.RecordDataGrowingOrdered(&DataHistory[CopyIndex]);
+		}
+		return bHasCopiedData;
+	}
+
 	/** Current data that is being loaded/recorded*/
 	DataType& GetCurrentData() { return DataHistory[CurrentIndex]; }
 
@@ -369,13 +422,16 @@ public :
 	/** Returns the next incremental data entry after the CurrentIndex (and loads it as the current data of the history), if there isn't a data entry with a higher LocalFrame value, nullptr is returned */
 	DataType* GetAndLoadNextIncrementalData()
 	{
-		const int32 NextIndex = (CurrentIndex + 1) % NumFrames;
-		if (DataHistory[NextIndex].LocalFrame > CurrentFrame)
+		if (NumFrames > 0)
 		{
-			CurrentIndex = NextIndex;
-			CurrentFrame = DataHistory[NextIndex].LocalFrame;
+			const int32 NextIndex = GetFrameIndex(CurrentIndex + 1);
+			if (DataHistory[NextIndex].LocalFrame > CurrentFrame)
+			{
+				CurrentIndex = NextIndex;
+				CurrentFrame = DataHistory[NextIndex].LocalFrame;
 
-			return &DataHistory[NextIndex];
+				return &DataHistory[NextIndex];
+			}
 		}
 
 		return nullptr;
@@ -435,7 +491,7 @@ public :
 	/** Resize the history */
 	FORCEINLINE void ResizeDataHistory(const int32 FrameCount, const EAllowShrinking AllowShrinking = EAllowShrinking::Default) override
 	{
-		if (FrameCount > 0 && NumFrames != FrameCount)
+		if (NumFrames != FrameCount)
 		{
 			NumFrames = FrameCount;
 			DataHistory.SetNum(NumFrames, AllowShrinking);
@@ -443,9 +499,21 @@ public :
 		}
 	}
 
+	/** Return circular index for given frame value */
 	FORCEINLINE const uint32 GetFrameIndex(const int32 Frame) const
 	{
-		return NumFrames == 0 ? 0 : FMath::Abs(Frame % NumFrames);
+		if (NumFrames <= 0)
+		{
+			return 0u;
+		}
+
+		int32 Idx = Frame % NumFrames;
+		if (Idx < 0)
+		{
+			Idx = Idx + NumFrames;
+		}
+
+		return static_cast<uint32>(Idx);
 	}
 
 	FORCEINLINE virtual void ResetFast()
