Fix up logic for resetting to Additive ref-pose, !bUseRawDataForPoseExtraction would short-circuit the conditional if Skeleton == nullptr (bug-hawk)

diff --git a/Engine/Source/Runtime/Engine/Private/Animation/AnimSequence.cpp b/Engine/Source/Runtime/Engine/Private/Animation/AnimSequence.cpp
--- a/Engine/Source/Runtime/Engine/Private/Animation/AnimSequence.cpp
+++ b/Engine/Source/Runtime/Engine/Private/Animation/AnimSequence.cpp
@@ -1595,6 +1595,11 @@ FTransform UAnimSequence::ExtractRootMotionFromRange(double StartTime, double En
 			RootTransformRefPose = RefSkeleton.GetRefBonePose()[0];
 		}
 	}
+	else
+	{
+		UE_LOG(LogAnimation, Error, TEXT("ExtractRootMotionFromRange for %s failed, Skeleton == nullptr."), *GetName());
+		return FTransform::Identity;
+	}
 
 	UE_COMPRESSED_DATA_READ_SCOPE(this);
 	FAnimExtractContext Context = ExtractionContext;
@@ -1738,7 +1743,7 @@ void UAnimSequence::GetBonePose(FAnimationPoseData& OutAnimationPoseData, const
 	const bool bInvalidSkeleton = MySkeleton == nullptr;
 	const bool bUseRawDataForPoseExtraction = bInvalidSkeleton || (CanEvaluateRawAnimationData() && bForceUseRawData) || ShouldUseRawDataForPoseExtraction_Lockless(RequiredBones, ExtractionContext);
 
-	const bool bIsBakedAdditive = !bUseRawDataForPoseExtraction && IsValidAdditive();
+	const bool bIsBakedAdditive = (!bUseRawDataForPoseExtraction || bInvalidSkeleton) && IsValidAdditive();
 	if (bInvalidSkeleton)
 	{
 		if (bIsBakedAdditive)
