World Partition - HLOD: Improved handling of materials in the HLOD build report
* Generate stable keys for MaterialInstanceDynamic (they are usually transient and their names can change)
* Register parent materials as references
* Use different prefixes for MaterialInstanceDynamic, MaterialInstanceConstant and MaterialInstance (MID, MIC & MI)
* These changes do NOT impact the resulting HLOD hash.

#FailedValidations Preflight

diff --git a/Engine/Source/Runtime/Engine/Private/WorldPartition/HLOD/HLODHashBuilder.cpp b/Engine/Source/Runtime/Engine/Private/WorldPartition/HLOD/HLODHashBuilder.cpp
--- a/Engine/Source/Runtime/Engine/Private/WorldPartition/HLOD/HLODHashBuilder.cpp
+++ b/Engine/Source/Runtime/Engine/Private/WorldPartition/HLOD/HLODHashBuilder.cpp
@@ -10,13 +10,49 @@
 #include "Engine/StaticMesh.h"
 #include "Engine/Texture.h"
 #include "Materials/MaterialInterface.h"
+#include "Materials/MaterialInstance.h"
+#include "Materials/MaterialInstanceDynamic.h"
 #include "WorldPartition/HLOD/HLODBuilder.h"
 
+// For most object, use the path name as the key
+// For MaterialInstanceDynamic, which are created dynamically, the name will not be stable from one execution to the next.
+// In that case, use the parent material name + a hash representing the dynamic variation.
+// This is purely for report stability, this doesn't affect the resulting HLOD hash
+static FString GetReportObjectKey(const UObject* InObject)
+{
+	FString ObjectKey;
+
+	if (const UMaterialInstanceDynamic* MaterialInstanceDynamic = Cast<const UMaterialInstanceDynamic>(InObject))
+	{
+		uint32 HashValue = UHLODProxy::GetCRC(const_cast<UMaterialInstanceDynamic*>(MaterialInstanceDynamic));
+
+		TArray<UTexture*> Textures;
+		TArray<uint32> TexturesHashes;
+		MaterialInstanceDynamic->GetUsedTextures(Textures);
+		for (UTexture* Texture : Textures)
+		{
+			TexturesHashes.Add(UHLODProxy::GetCRC(Texture));
+		}
+		TexturesHashes.Sort();
+		for (uint32 TextureHash : TexturesHashes)
+		{
+			HashValue = HashCombine(HashValue, TextureHash);
+		}
+
+		ObjectKey = FString::Printf(TEXT("%s (MID Key=%08X)"), *MaterialInstanceDynamic->Parent->GetPathName(), HashValue);
+	}
+	else
+	{
+		UObject* StopOuter = InObject->IsA<UActorComponent>() ? InObject->GetTypedOuter<ULevel>() : nullptr;
+		ObjectKey = InObject->GetPathName(StopOuter);
+	}
+
+	return ObjectKey;
+}
+
 void FHLODHashBuilder::PushObjectContext(const UObject* InObjectContext)
 {
-	UObject* StopOuter = InObjectContext->IsA<UActorComponent>() ? InObjectContext->GetTypedOuter<ULevel>() : nullptr;
-	FString ObjectPath = InObjectContext->GetFullName(StopOuter);
-	ObjectContextStack.Push(*ObjectPath);
+	ObjectContextStack.Push(InObjectContext->GetClass()->GetName() + TEXT(" ") + GetReportObjectKey(InObjectContext));
 }
 
 void FHLODHashBuilder::PopObjectContext()
@@ -36,7 +72,7 @@ uint32 FHLODHashBuilder::AddAssetReference(UObject* Asset, TFunctionRef<uint32()
 {
 	check(Asset != nullptr);
 
-	FName AssetName(Asset->GetPathName());
+	FName AssetName(GetReportObjectKey(Asset));
 
 	FAssetHash& AssetHash = AssetsHashes.FindOrAdd(AssetName);
 	if (AssetHash.Hash == 0)
@@ -57,17 +93,26 @@ FArchive& FHLODHashBuilder::operator<<(UObject*& InObject)
 {
 	if (UMaterialInterface* MaterialInterface = Cast<UMaterialInterface>(InObject))
 	{
-		*this << AddAssetReference(MaterialInterface, [&]() { return UHLODProxy::GetCRC(MaterialInterface); });
+		*this << AddAssetReference(MaterialInterface, [MaterialInterface]() { return UHLODProxy::GetCRC(MaterialInterface); });
 
 		FHLODHashScope HashScopeMaterial(*this, MaterialInterface);
 
+		// Add the whole material chain as references
+		UMaterialInterface* Parent = Cast<UMaterialInstance>(MaterialInterface) ? Cast<UMaterialInstance>(MaterialInterface)->Parent : nullptr;
+		while (Parent)
+		{
+			AddAssetReference(Parent, [Parent]() { return UHLODProxy::GetCRC(Parent); });
+			Parent = Cast<UMaterialInstance>(Parent) ? Cast<UMaterialInstance>(Parent)->Parent : nullptr;
+		}
+
+		// Textures
 		TArray<UTexture*> Textures;
 		MaterialInterface->GetUsedTextures(Textures);
 		*this << Textures;
 	
 		if (UMaterialInterface* NaniteOverride = MaterialInterface->GetNaniteOverride())
 		{ 
-			*this << AddAssetReference(NaniteOverride, [&]() { return UHLODProxy::GetCRC(NaniteOverride); });
+			*this << AddAssetReference(NaniteOverride, [NaniteOverride]() { return UHLODProxy::GetCRC(NaniteOverride); });
 
 			FHLODHashScope HashScopeNaniteMaterial(*this, NaniteOverride);
 
@@ -78,11 +123,11 @@ FArchive& FHLODHashBuilder::operator<<(UObject*& InObject)
 	}
 	else if (UTexture* Texture = Cast<UTexture>(InObject))
 	{
-		*this << AddAssetReference(Texture, [&]() { return UHLODProxy::GetCRC(Texture); });
+		*this << AddAssetReference(Texture, [Texture]() { return UHLODProxy::GetCRC(Texture); });
 	}
 	else if (UStaticMesh* StaticMesh = Cast<UStaticMesh>(InObject))
 	{
-		*this << AddAssetReference(StaticMesh, [&]() { return UHLODProxy::GetCRC(StaticMesh); });
+		*this << AddAssetReference(StaticMesh, [StaticMesh]() { return UHLODProxy::GetCRC(StaticMesh); });
 	}
 	else
 	{
@@ -119,6 +164,9 @@ FArchive& FHLODHashBuilder::operator<<(USkinnedAsset* InSkinnedAsset)
 static FString GetTypePrefixFromClassName(const FName& ClassFName)
 {
 	const FString C = ClassFName.ToString();
+	if (C.Contains(TEXT("MaterialInstanceDynamic")))                     return TEXT("MID");
+	if (C.Contains(TEXT("MaterialInstanceConstant")))                    return TEXT("MIC");
+	if (C.Contains(TEXT("MaterialInstance")))                    		 return TEXT("MI");
 	if (C.Contains(TEXT("Material")))                             		 return TEXT("MAT");
 	if (C.Contains(TEXT("Texture")))                              		 return TEXT("TEX");
 	if (C.Contains(TEXT("StaticMesh")))                           		 return TEXT("SM");
