Make the extra memory detection message optional due to Debug HUD perf cost on lower end devices

#rnx

diff --git a/Engine/Source/Runtime/Engine/Private/UnrealEngine.cpp b/Engine/Source/Runtime/Engine/Private/UnrealEngine.cpp
--- a/Engine/Source/Runtime/Engine/Private/UnrealEngine.cpp
+++ b/Engine/Source/Runtime/Engine/Private/UnrealEngine.cpp
@@ -547,6 +547,12 @@ TAutoConsoleVariable<bool> CVarCsvAlwaysShowFrameCount(
 	false,
 	TEXT("If enabled, we show the frame count in non-shipping builds, even if screen messages are disabled")
 );
+
+TAutoConsoleVariable<bool> CVarAllowExtraDevMemoryMessageDisplayed(
+	TEXT("Debug.AllowExtraDevMemoryMessageDisplayed"),
+	true,
+	TEXT("If true, we allow the extra memory detection message to be shown in non-shipping builds")
+);
 #endif
 
 
@@ -13352,7 +13358,7 @@ void DrawStatsHUD( UWorld* World, FViewport* Viewport, FCanvas* Canvas, UCanvas*
 		}
 #endif // STATS
 
-		if (FPlatformMemory::IsExtraDevelopmentMemoryAvailable())
+		if (CVarAllowExtraDevMemoryMessageDisplayed.GetValueOnGameThread() && FPlatformMemory::IsExtraDevelopmentMemoryAvailable())
 		{
 			SmallTextItem.SetColor(FLinearColor(1.0f, 1.0f, 1.0f, 1.0f));
 			SmallTextItem.Text = LOCTEXT("LLMWARNING", "WARNING: Running with Debug Memory Enabled!");
