Added a check at skeletal mesh actor to avoid calling a null skeletal mesh component

diff --git a/Engine/Source/Runtime/Engine/Private/Animation/SkeletalMeshActor.cpp b/Engine/Source/Runtime/Engine/Private/Animation/SkeletalMeshActor.cpp
--- a/Engine/Source/Runtime/Engine/Private/Animation/SkeletalMeshActor.cpp
+++ b/Engine/Source/Runtime/Engine/Private/Animation/SkeletalMeshActor.cpp
@@ -156,15 +156,18 @@ bool ASkeletalMeshActor::GetReferencedContentObjects(TArray<UObject*>& Objects)
 {
 	Super::GetReferencedContentObjects(Objects);
 
-	if (SkeletalMeshComponent->GetSkeletalMeshAsset())
+	if (SkeletalMeshComponent)
 	{
-		Objects.Add(SkeletalMeshComponent->GetSkeletalMeshAsset());
-	}
+		if (SkeletalMeshComponent->GetSkeletalMeshAsset())
+		{
+			Objects.Add(SkeletalMeshComponent->GetSkeletalMeshAsset());
+		}
 
-	if (SkeletalMeshComponent->GetAnimationMode() == EAnimationMode::Type::AnimationSingleNode &&
-		SkeletalMeshComponent->AnimationData.AnimToPlay)
-	{
-		Objects.Add(SkeletalMeshComponent->AnimationData.AnimToPlay);
+		if (SkeletalMeshComponent->GetAnimationMode() == EAnimationMode::Type::AnimationSingleNode &&
+			SkeletalMeshComponent->AnimationData.AnimToPlay)
+		{
+			Objects.Add(SkeletalMeshComponent->AnimationData.AnimToPlay);
+		}
 	}
 
 	return true;
