Arrays preallocation Streaming optimizations

diff --git a/Engine/Source/Runtime/Engine/Private/Streaming/AsyncTextureStreaming.cpp b/Engine/Source/Runtime/Engine/Private/Streaming/AsyncTextureStreaming.cpp
--- a/Engine/Source/Runtime/Engine/Private/Streaming/AsyncTextureStreaming.cpp
+++ b/Engine/Source/Runtime/Engine/Private/Streaming/AsyncTextureStreaming.cpp
@@ -910,14 +910,16 @@ void FRenderAssetStreamingMipCalcTask::UpdateLoadAndCancelationRequests_Async()
 	TArray<FStreamingRenderAsset>& StreamingRenderAssets = StreamingManager.AsyncUnsafeStreamingRenderAssets;
 	const FRenderAssetStreamingSettings& Settings = StreamingManager.Settings;
 
-	LoadRequests.Empty();
 	CancelationRequests.Empty();
 
 	int64 StreamOutMemoryBudget = TempMemoryBudget;
 	int64 StreamInMemoryBudget = TempMemoryBudget;
 
 	TArray<int32> PrioritizedRenderAssets;
-	PrioritizedRenderAssets.Empty(StreamingRenderAssets.Num());
+	if (!IsAborted())
+	{
+		PrioritizedRenderAssets.Empty(StreamingRenderAssets.Num());
+	}
 	for (int32 AssetIndex = 0; AssetIndex < StreamingRenderAssets.Num() && !IsAborted(); ++AssetIndex)
 	{
 		FStreamingRenderAsset& StreamingRenderAsset = StreamingRenderAssets[AssetIndex];
@@ -963,8 +965,17 @@ void FRenderAssetStreamingMipCalcTask::UpdateLoadAndCancelationRequests_Async()
 
 	// If possible, free as much memory with stream out operations, as will be required with new stream in requests.
 	// This prevents becoming overbudget momentarily, when we are already at the budget limit.
-	TArray<int32> StreamOutRequests;
-	TArray<int32> StreamInRequests;
+	bool bHasStreamOutRequest = false;
+	bool bHasStreamInRequest = false;
+
+	if (PrioritizedRenderAssets.Num() && !IsAborted())
+	{
+		LoadRequests.Reset(PrioritizedRenderAssets.Num());
+	}
+	else
+	{
+		LoadRequests.Empty();
+	}
 
 	// Now fill in the LoadRequest and CancelationRequests
 	for (int32 PriorityIndex = 0; PriorityIndex < PrioritizedRenderAssets.Num() && !IsAborted(); ++PriorityIndex)
@@ -977,21 +988,19 @@ void FRenderAssetStreamingMipCalcTask::UpdateLoadAndCancelationRequests_Async()
 		const int64 TempMemoryRequired = StreamingRenderAsset.GetSize(StreamingRenderAsset.WantedMips);
 
 		// Check whether the budget allows the update, with the exception of always allowing a single update of any size (otherwise completion might never happen).
-		if (StreamingRenderAsset.WantedMips < StreamingRenderAsset.ResidentMips && (TempMemoryRequired <= StreamOutMemoryBudget || !StreamOutRequests.Num()))
+		if (StreamingRenderAsset.WantedMips < StreamingRenderAsset.ResidentMips && (TempMemoryRequired <= StreamOutMemoryBudget || !bHasStreamOutRequest))
 		{
-			StreamOutRequests.Add(AssetIndex);
+			bHasStreamOutRequest = true;
+			LoadRequests.Add(AssetIndex);
 			StreamOutMemoryBudget -= TempMemoryRequired;
 		}
-		else if (StreamingRenderAsset.WantedMips > StreamingRenderAsset.ResidentMips && (TempMemoryRequired <= StreamInMemoryBudget || !StreamInRequests.Num()))
+		else if (StreamingRenderAsset.WantedMips > StreamingRenderAsset.ResidentMips && (TempMemoryRequired <= StreamInMemoryBudget || !bHasStreamInRequest))
 		{
-			StreamInRequests.Add(AssetIndex);
+			bHasStreamInRequest = true;
+			LoadRequests.Add(AssetIndex);
 			StreamInMemoryBudget -= TempMemoryRequired;
 		}
 	}
-
-	// Process stream out requests first since they execute faster, freeing the memory for the stream in requests.
-	LoadRequests.Append(StreamOutRequests);
-	LoadRequests.Append(StreamInRequests);
 }
 
 void FRenderAssetStreamingMipCalcTask::UpdatePendingStreamingStatus_Async()
@@ -1047,8 +1056,8 @@ void FRenderAssetStreamingMipCalcTask::DoWork()
 			StreamingRenderAsset.DynamicBoostFactor = 1.f; // Reset after every computation.
 		}
 
-		uint64 LenghtCycles64 = (FPlatformTime::Cycles64() - StartTime);
-		double ElapsedMSTime = FPlatformTime::ToMilliseconds64(LenghtCycles64);
+		uint64 LengthCycles64 = (FPlatformTime::Cycles64() - StartTime);
+		double ElapsedMSTime = FPlatformTime::ToMilliseconds64(LengthCycles64);
 
 		// use of a custom stat to have the elapsed time as a global stat and not a stat split accross multiple threads
 		CSV_CUSTOM_STAT(TextureStreaming, RenderAssetStreamingUpdate, float(ElapsedMSTime), ECsvCustomStatOp::Set);
diff --git a/Engine/Source/Runtime/Engine/Private/Streaming/TextureInstanceState.cpp b/Engine/Source/Runtime/Engine/Private/Streaming/TextureInstanceState.cpp
--- a/Engine/Source/Runtime/Engine/Private/Streaming/TextureInstanceState.cpp
+++ b/Engine/Source/Runtime/Engine/Private/Streaming/TextureInstanceState.cpp
@@ -331,9 +331,16 @@ EAddComponentResult FRenderAssetInstanceState::AddComponents(const TArray<FRende
 	TRACE_CPUPROFILER_EVENT_SCOPE(FRenderAssetInstanceState::AddComponents);
 	checkf(!bIsDynamicInstanceState, TEXT("Error: trying to add component to dynamic instance manager as static"));
 	const int32 PayloadCount = Payloads.Num();
-	for (int i=0; i<PayloadCount; ++i)
+	int32 TotalBoundsCount = 0;
+	for (int i = 0; i < PayloadCount; ++i)
 	{
-		EAddComponentResult Result = AddComponent(Payloads[i]);
+		TotalBoundsCount += Payloads[i].Bounds.Num();
+	}
+
+	Bounds4.Reserve(Bounds4.Num() + TotalBoundsCount);
+	for (int i = 0; i < PayloadCount; ++i)
+	{
+		EAddComponentResult Result = AddComponentInternal(Payloads[i]);
 		if (Result != EAddComponentResult::Success)
 		{
 			OnAddFailedFunc(Payloads[i], Result);
@@ -343,6 +350,15 @@ EAddComponentResult FRenderAssetInstanceState::AddComponents(const TArray<FRende
 }
 
 EAddComponentResult FRenderAssetInstanceState::AddComponent(const FPreAddComponentPayload& Payload)
+{
+	if (Payload.Result == EAddComponentResult::Success)
+	{
+		Bounds4.Reserve(Bounds4.Num() + Payload.Bounds.Num());
+	}
+	return AddComponentInternal(Payload);
+}
+
+EAddComponentResult FRenderAssetInstanceState::AddComponentInternal(const FPreAddComponentPayload& Payload)
 {
 	checkf(!bIsDynamicInstanceState, TEXT("Error: trying to add component to dynamic instance manager as static"));
 	if (Payload.Result == EAddComponentResult::Success)
@@ -445,6 +461,11 @@ EAddComponentResult FRenderAssetInstanceState::AddComponentInternal(const UPrimi
 
 		int32* ComponentLink = State ? State->ComponentMap.Find(Component) : nullptr;
 
+		if (State)
+		{
+			State->Bounds4.Reserve(State->Bounds4.Num() + RenderAssetInstanceInfos.Num());
+		}
+
 		// Loop for each bound.
 		for (int32 InfoIndex = 0; InfoIndex < RenderAssetInstanceInfos.Num();)
 		{
@@ -493,6 +514,11 @@ EAddComponentResult FRenderAssetInstanceState::AddComponentInternal(const UPrimi
 		int32* ComponentLink = State ? State->ComponentMap.Find(Component) : nullptr;
 		float MinDistanceSq = 0, MinRangeSq = 0, MaxRangeSq = FLT_MAX;
 
+		if (State)
+		{
+			State->Bounds4.Reserve(State->Bounds4.Num() + RenderAssetInstanceInfos.Num());
+		}
+
 		// Loop for each bound.
 		for (int32 InfoIndex = 0; InfoIndex < RenderAssetInstanceInfos.Num();)
 		{
diff --git a/Engine/Source/Runtime/Engine/Private/Streaming/TextureInstanceState.h b/Engine/Source/Runtime/Engine/Private/Streaming/TextureInstanceState.h
--- a/Engine/Source/Runtime/Engine/Private/Streaming/TextureInstanceState.h
+++ b/Engine/Source/Runtime/Engine/Private/Streaming/TextureInstanceState.h
@@ -135,6 +135,7 @@ private:
 	FORCEINLINE int32 AddBounds(const UPrimitiveComponent* Component);
 	void RemoveBounds(int32 Index);
 
+	EAddComponentResult AddComponentInternal(const FPreAddComponentPayload& Payload);
 	static EAddComponentResult AddComponentInternal(const UPrimitiveComponent* Component, FStreamingTextureLevelContext& LevelContext, float MaxAllowedUIDensity, FRenderAssetInstanceState* State, FPreAddComponentPayload* OutPreAddPayload);
 	static EAddComponentResult AddComponentIgnoreBoundsInternal(const UPrimitiveComponent* Component, FStreamingTextureLevelContext& LevelContext, FRenderAssetInstanceState* State, FPreAddComponentIgnoreBoundsPayload* OutPreAddPayload);
 
