Streamable Manager - add CVar to control async load error verbosity

diff --git a/Engine/Source/Runtime/Engine/Private/StreamableManager.cpp b/Engine/Source/Runtime/Engine/Private/StreamableManager.cpp
--- a/Engine/Source/Runtime/Engine/Private/StreamableManager.cpp
+++ b/Engine/Source/Runtime/Engine/Private/StreamableManager.cpp
@@ -59,6 +59,14 @@ static FAutoConsoleVariableRef CVarStreamableStripDebugNameInShipping(
 	ECVF_Default
 );
 
+static bool GStreamableLogAsyncLoadErrorsAsErrors = false;
+static FAutoConsoleVariableRef CVarLogAsyncLoadErrorsAsErrors(
+	TEXT("s.LogAsyncLoadErrorsAsErrors"),
+	GStreamableLogAsyncLoadErrorsAsErrors,
+	TEXT("Wether async loading errors should be logged as errors. If disabled, async load errors will be logged as verbose."),
+	ECVF_Default
+);
+
 #if !UE_BUILD_SHIPPING
 static FString GStreamableErrorOnRequest = "";
 static FAutoConsoleVariableRef CVarStreamableErrorOnRequest(
@@ -2158,12 +2166,28 @@ void FStreamableManager::AsyncLoadCallback(FSoftObjectPath TargetName, UPackage*
 		{
 			// Async load failed to find the object
 			Existing->bLoadFailed = true;
-			UE_LOG(LogStreamableManager, Verbose, TEXT("    Failed async load for %s"), *TargetName.ToString());
+			constexpr static const TCHAR LogFmt[] = TEXT("    Failed async load for %s");
+			if (GStreamableLogAsyncLoadErrorsAsErrors)
+			{
+				UE_LOG(LogStreamableManager, Error, LogFmt, *TargetName.ToString());
+			}
+			else
+			{
+				UE_LOG(LogStreamableManager, Verbose, LogFmt, *TargetName.ToString());
+			}
 		}
 
 		if (AsyncLoadError)
 		{
-			UE_LOG(LogStreamableManager, Verbose, TEXT("    Encountered error while async loading %s: %s"), *TargetName.ToString(), *AsyncLoadError->GetErrorMessage(true).ToString());
+			constexpr static const TCHAR LogFmt[] = TEXT("    Encountered error while async loading %s: %s");
+			if (GStreamableLogAsyncLoadErrorsAsErrors)
+			{
+				UE_LOG(LogStreamableManager, Error, LogFmt, *TargetName.ToString(), *AsyncLoadError->GetErrorMessage(true).ToString());
+			}
+			else
+			{
+				UE_LOG(LogStreamableManager, Verbose, LogFmt, *TargetName.ToString(), *AsyncLoadError->GetErrorMessage(true).ToString());
+			}
 
 			for (const TSharedRef<FStreamableHandle>& Handle : Existing->LoadingHandles)
 			{
