[Substrate] Remove r.SubstrateBackCompatibility as we now rely on the Substrate Rough Diffuse option to drive that instead.

#[review] [REDACTED_HANDLE]

diff --git a/Engine/Source/Runtime/Engine/Private/ShaderCompiler/ShaderCompiler.cpp b/Engine/Source/Runtime/Engine/Private/ShaderCompiler/ShaderCompiler.cpp
--- a/Engine/Source/Runtime/Engine/Private/ShaderCompiler/ShaderCompiler.cpp
+++ b/Engine/Source/Runtime/Engine/Private/ShaderCompiler/ShaderCompiler.cpp
@@ -3992,9 +3992,6 @@ void GlobalBeginCompileShader(
 			SET_SHADER_DEFINE(Input.Environment, SUBSTRATE_MATERIAL_CLOSURE_COUNT, 0);
 		}
 
-		const bool bSubstrateBackCompatibility = bSubstrate && Substrate::IsBackCompatibilityEnabled();
-		SET_SHADER_DEFINE(Input.Environment, PROJECT_SUBSTRATE_BACKCOMPATIBILITY, bSubstrateBackCompatibility ? 1 : 0);
-
 		const bool bSubstrateOpaqueRoughRefrac = bSubstrate && Substrate::IsOpaqueRoughRefractionEnabled(Target.GetPlatform());
 		SET_SHADER_DEFINE(Input.Environment, SUBSTRATE_OPAQUE_ROUGH_REFRACTION_ENABLED, bSubstrateOpaqueRoughRefrac ? 1 : 0);
 
@@ -4008,7 +4005,7 @@ void GlobalBeginCompileShader(
 	{
 		static IConsoleVariable* CVar = IConsoleManager::Get().FindConsoleVariable(TEXT("r.Material.RoughDiffuse"));
 		const bool bMaterialRoughDiffuse = CVar && CVar->GetInt() != 0;
-		const bool bSubstrateRoughDiffuse = Substrate::IsRoughDiffuseEnabled(Target.GetPlatform()) && !Substrate::IsBackCompatibilityEnabled();
+		const bool bSubstrateRoughDiffuse = Substrate::IsRoughDiffuseEnabled(Target.GetPlatform());
 		SET_SHADER_DEFINE(Input.Environment, MATERIAL_ROUGHDIFFUSE, (bSubstrate ? bSubstrateRoughDiffuse : bMaterialRoughDiffuse) ? 1 : 0);
 	}
 
diff --git a/Engine/Source/Runtime/RenderCore/Private/RenderUtils.cpp b/Engine/Source/Runtime/RenderCore/Private/RenderUtils.cpp
--- a/Engine/Source/Runtime/RenderCore/Private/RenderUtils.cpp
+++ b/Engine/Source/Runtime/RenderCore/Private/RenderUtils.cpp
@@ -1978,12 +1978,6 @@ static TAutoConsoleVariable<int32> CVarSubstrateProjectClosuresPerPixelOverride(
 	TEXT("Select how the maximum number of closures evaluated per pixel is defined for each platforms.\n - Use Platform Closure Count: Use the mininum between the platform's closure count and the project's closure count.\n - Force Closure Count : Force the project's closure count on platforms supporting Adaptive GBuffer."),
 	ECVF_ReadOnly | ECVF_RenderThreadSafe | ECVF_Preview);
 
-static TAutoConsoleVariable<int32> CVarSubstrateBackCompatibility(
-	TEXT("r.SubstrateBackCompatibility"),
-	0,
-	TEXT("Disables Substrate multiple scattering and replaces Chan diffuse by Lambert."),
-	ECVF_ReadOnly | ECVF_RenderThreadSafe);
-
 static TAutoConsoleVariable<int32> CVarSubstrateOpaqueMaterialRoughRefraction(
 	TEXT("r.Substrate.OpaqueMaterialRoughRefraction"),
 	0, // Substrate 60fps by default
@@ -2214,11 +2208,6 @@ namespace Substrate
 
 		return PayloadSizeBytes;
 	}
-
-	bool IsBackCompatibilityEnabled()
-	{
-		return CVarSubstrateBackCompatibility.GetValueOnAnyThread() > 0 ? 1 : 0;
-	}
 	
 	bool IsRoughDiffuseEnabled()
 	{
diff --git a/Engine/Source/Runtime/RenderCore/Private/Shader.cpp b/Engine/Source/Runtime/RenderCore/Private/Shader.cpp
--- a/Engine/Source/Runtime/RenderCore/Private/Shader.cpp
+++ b/Engine/Source/Runtime/RenderCore/Private/Shader.cpp
@@ -2478,12 +2478,6 @@ void ShaderMapAppendKey(EShaderPlatform Platform, FShaderKeyGenerator& KeyGen)
 			KeyGen.Append(TEXT("HIDDENCONV"));
 		}
 
-		if (Substrate::IsBackCompatibilityEnabled())
-		{
-			KeyGen.AppendSeparator();
-			KeyGen.Append(TEXT("BACKCOMPAT"));
-		}
-
 		if (Substrate::IsOpaqueRoughRefractionEnabled(Platform))
 		{
 			KeyGen.AppendSeparator();
diff --git a/Engine/Source/Runtime/RenderCore/Public/RenderUtils.h b/Engine/Source/Runtime/RenderCore/Public/RenderUtils.h
--- a/Engine/Source/Runtime/RenderCore/Public/RenderUtils.h
+++ b/Engine/Source/Runtime/RenderCore/Public/RenderUtils.h
@@ -634,7 +634,6 @@ namespace Substrate
 	RENDERCORE_API float GlintLevelMin();
 	RENDERCORE_API bool IsSpecularProfileEnabled();
 	RENDERCORE_API bool IsSpecularProfileEnabled(EShaderPlatform InPlatform);
-	RENDERCORE_API bool IsBackCompatibilityEnabled();
 	RENDERCORE_API bool IsDBufferPassEnabled(EShaderPlatform InPlatform);
 	RENDERCORE_API bool IsOpaqueRoughRefractionEnabled(EShaderPlatform InPlatform);
 	UE_DEPRECATED(5.7, "Plase use IsAdvancedVisualizationEnabled(EShaderPlatform InPlatform)")
