Fixed correct ordering of LoadRequests in FRenderAssetStreamingMipCalcTask::UpdateLoadAndCancelationRequests_Async()

diff --git a/Engine/Source/Runtime/Engine/Private/Streaming/AsyncTextureStreaming.cpp b/Engine/Source/Runtime/Engine/Private/Streaming/AsyncTextureStreaming.cpp
--- a/Engine/Source/Runtime/Engine/Private/Streaming/AsyncTextureStreaming.cpp
+++ b/Engine/Source/Runtime/Engine/Private/Streaming/AsyncTextureStreaming.cpp
@@ -965,16 +965,13 @@ void FRenderAssetStreamingMipCalcTask::UpdateLoadAndCancelationRequests_Async()
 
 	// If possible, free as much memory with stream out operations, as will be required with new stream in requests.
 	// This prevents becoming overbudget momentarily, when we are already at the budget limit.
-	bool bHasStreamOutRequest = false;
-	bool bHasStreamInRequest = false;
+	TArray<int32> StreamOutRequests;
+	TArray<int32> StreamInRequests;
 
-	if (PrioritizedRenderAssets.Num() && !IsAborted())
-	{
-		LoadRequests.Reset(PrioritizedRenderAssets.Num());
-	}
-	else
+	if (!IsAborted())
 	{
-		LoadRequests.Empty();
+		StreamOutRequests.Reserve(PrioritizedRenderAssets.Num());
+		StreamInRequests.Reserve(PrioritizedRenderAssets.Num());
 	}
 
 	// Now fill in the LoadRequest and CancelationRequests
@@ -988,19 +985,22 @@ void FRenderAssetStreamingMipCalcTask::UpdateLoadAndCancelationRequests_Async()
 		const int64 TempMemoryRequired = StreamingRenderAsset.GetSize(StreamingRenderAsset.WantedMips);
 
 		// Check whether the budget allows the update, with the exception of always allowing a single update of any size (otherwise completion might never happen).
-		if (StreamingRenderAsset.WantedMips < StreamingRenderAsset.ResidentMips && (TempMemoryRequired <= StreamOutMemoryBudget || !bHasStreamOutRequest))
+		if (StreamingRenderAsset.WantedMips < StreamingRenderAsset.ResidentMips && (TempMemoryRequired <= StreamOutMemoryBudget || !StreamOutRequests.Num()))
 		{
-			bHasStreamOutRequest = true;
-			LoadRequests.Add(AssetIndex);
+			StreamOutRequests.Add(AssetIndex);
 			StreamOutMemoryBudget -= TempMemoryRequired;
 		}
-		else if (StreamingRenderAsset.WantedMips > StreamingRenderAsset.ResidentMips && (TempMemoryRequired <= StreamInMemoryBudget || !bHasStreamInRequest))
+		else if (StreamingRenderAsset.WantedMips > StreamingRenderAsset.ResidentMips && (TempMemoryRequired <= StreamInMemoryBudget || !StreamInRequests.Num()))
 		{
-			bHasStreamInRequest = true;
-			LoadRequests.Add(AssetIndex);
+			StreamInRequests.Add(AssetIndex);
 			StreamInMemoryBudget -= TempMemoryRequired;
 		}
 	}
+
+	LoadRequests.Reset(StreamOutRequests.Num() + StreamInRequests.Num());
+	// Process stream out requests first since they execute faster, freeing the memory for the stream in requests.
+	LoadRequests.Append(StreamOutRequests);
+	LoadRequests.Append(StreamInRequests);
 }
 
 void FRenderAssetStreamingMipCalcTask::UpdatePendingStreamingStatus_Async()
