TextureDefines.h : clarify texture compression display names and tooltips
biggest change is "UserInterface2d" changing to "Uncompressed" to emphasize that

diff --git a/Engine/Source/Runtime/Engine/Classes/Engine/TextureDefines.h b/Engine/Source/Runtime/Engine/Classes/Engine/TextureDefines.h
--- a/Engine/Source/Runtime/Engine/Classes/Engine/TextureDefines.h
+++ b/Engine/Source/Runtime/Engine/Classes/Engine/TextureDefines.h
@@ -392,16 +392,16 @@ enum TextureCompressionSettings : int
 {
 	TC_Default					UMETA(DisplayName = "Default (BC1 or BC3 with A)"),
 	TC_Normalmap				UMETA(DisplayName = "Normalmap (BC5)"),
-	TC_Masks					UMETA(DisplayName = "Masks (no sRGB)"),
-	TC_Grayscale				UMETA(DisplayName = "Grayscale (G8/16, RGB8 sRGB) from source R"),
+	TC_Masks					UMETA(DisplayName = "Masks (no sRGB)", ToolTip = "Same compression as Default"),
+	TC_Grayscale				UMETA(DisplayName = "Grayscale (G8/16) from source R", ToolTip = "If SRGB is set, becomes uncompressed RGBA8; prefer BC4"),
 	TC_Displacementmap			UMETA(DisplayName = "Displacementmap (G8/16) from source A", ToolTip = "For Nanite displacement use TC_Alpha"),
 	TC_VectorDisplacementmap	UMETA(DisplayName = "VectorDisplacementmap (RGBA8)"),
-	TC_HDR						UMETA(DisplayName = "HDR (RGBA16F, no sRGB)"),
-	TC_EditorIcon				UMETA(DisplayName = "UserInterface2D (RGBA8)"),
+	TC_HDR						UMETA(DisplayName = "HDR (RGBA16F, no sRGB)", ToolTip = "Prefer HDR Compressed"),
+	TC_EditorIcon				UMETA(DisplayName = "Uncompressed (RGBA8)", ToolTip = "Very high memory use; prefer BC7 or Default instead"),
 	TC_Alpha					UMETA(DisplayName = "Alpha (no sRGB, BC4) from source R", ToolTip = "Single channel compressed; can also use for displacement and gray"),
 	TC_DistanceFieldFont		UMETA(DisplayName = "DistanceFieldFont (G8) from source A"),
-	TC_HDR_Compressed			UMETA(DisplayName = "HDR Compressed (RGB, BC6H)"),
-	TC_BC7						UMETA(DisplayName = "BC7 (RGBA)"),
+	TC_HDR_Compressed			UMETA(DisplayName = "HDR Compressed (BC6H)"),
+	TC_BC7						UMETA(DisplayName = "BC7 Compressed (BC7)", ToolTip = "High quality compressed texture, same size as BC3"),
 	TC_HalfFloat				UMETA(DisplayName = "Half Float (R16F)"),
 	TC_LQ				        UMETA(Hidden, DisplayName = "Low Quality (BGR565/BGR555A1)", ToolTip = "BGR565/BGR555A1, fallback to DXT1/DXT5 on Mac platform"),
 	TC_EncodedReflectionCapture	UMETA(Hidden), 
