Adding structured log context, skip trying to fetch distance field from base render data LOD 0 if that doesn't exist for some reason.

diff --git a/Engine/Source/Runtime/Engine/Private/StaticMesh.cpp b/Engine/Source/Runtime/Engine/Private/StaticMesh.cpp
--- a/Engine/Source/Runtime/Engine/Private/StaticMesh.cpp
+++ b/Engine/Source/Runtime/Engine/Private/StaticMesh.cpp
@@ -86,6 +86,7 @@
 #include "ProfilingDebugging/CookStats.h"
 #include "Streaming/UVChannelDensity.h"
 #include "Logging/MessageLog.h"
+#include "Logging/StructuredLog.h"
 #include "Misc/UObjectToken.h"
 #include "Streaming/StaticMeshUpdate.h"
 #include "Streaming/TextureStreamingHelpers.h"
@@ -318,6 +319,7 @@ namespace UE::StaticMesh::Private::Constants
 void FStaticMeshAsyncBuildWorker::DoWork()
 {
 	FStaticMeshAsyncBuildScope AsyncBuildScope(StaticMesh);
+	UE_LOG_CONTEXT("StaticMesh", FSoftObjectPath(StaticMesh));
 
 	if (PostLoadContext.IsValid())
 	{
@@ -4139,6 +4141,7 @@ void FStaticMeshRenderData::Cache(const ITargetPlatform* TargetPlatform, UStatic
 		UE_TRACE_LOG_SCOPED_T(Cpu, StaticMeshCache, CpuChannel)
 			<< StaticMeshCache.KeySuffix(*KeySuffix);
 #endif
+		UE_LOG_CONTEXT("StaticMeshKeySuffix", KeySuffix);
 		FString NewDerivedDataKey = BuildStaticMeshDerivedDataKey(KeySuffix);
 		if (NumCacheAttempts == 0 || !IsRunningCookCommandlet())
 		{
@@ -4477,10 +4480,18 @@ void FStaticMeshRenderData::Cache(const ITargetPlatform* TargetPlatform, UStatic
 			{
 				// If base data is already built (implying it was loaded from DDC), copy that to platform data, otherwise
 				// it will be copied later when the async queue finishes generating it.
-				FDistanceFieldVolumeData* BaseDistanceFieldData = Owner->GetRenderData()->LODResources[0].DistanceFieldData;
-				if (BaseDistanceFieldData && BaseDistanceFieldData->bAsyncBuilding == false)
+				if (Owner->GetRenderData()->LODResources.IsValidIndex(0))
 				{
-					*LODResources[0].DistanceFieldData = *BaseDistanceFieldData;
+					FDistanceFieldVolumeData* BaseDistanceFieldData = Owner->GetRenderData()->LODResources[0].DistanceFieldData;
+					if (BaseDistanceFieldData && BaseDistanceFieldData->bAsyncBuilding == false)
+					{
+						*LODResources[0].DistanceFieldData = *BaseDistanceFieldData;
+					}
+				}
+				else
+				{
+					UE_LOGFMT(LogStaticMesh, Error, "{StaticMesh}: Failed to fetch distance field data from LOD 0 on base render data for platform {PlatformName}", 
+						FSoftObjectPath(Owner), TargetPlatform->PlatformName());
 				}
 
 				// For platform data, this either copies the base render data, or allocates a placeholder structure for data to be copied into
