[NetPhysics] Only keep the replicated properties that will be used active.
- The new Iris compliant flow will never overlap with the Legacy RepGraph flow.

diff --git a/Engine/Source/Runtime/Engine/Private/PhysicsEngine/NetworkPhysicsComponent.cpp b/Engine/Source/Runtime/Engine/Private/PhysicsEngine/NetworkPhysicsComponent.cpp
--- a/Engine/Source/Runtime/Engine/Private/PhysicsEngine/NetworkPhysicsComponent.cpp
+++ b/Engine/Source/Runtime/Engine/Private/PhysicsEngine/NetworkPhysicsComponent.cpp
@@ -500,6 +500,8 @@ void UNetworkPhysicsComponent::InitializeComponent()
 
 	if (AActor* Owner = GetOwner())
 	{
+		Owner->SetCallPreReplication(true);
+
 		// Get settings from NetworkPhysicsSettingsComponent, if there is one
 		UNetworkPhysicsSettingsComponent* SettingsComponent = Owner->FindComponentByClass<UNetworkPhysicsSettingsComponent>();
 		if (SettingsComponent)
@@ -631,7 +633,7 @@ void UNetworkPhysicsComponent::GetLifetimeReplicatedProps(TArray< FLifetimePrope
 
 	DOREPLIFETIME_CONDITION(UNetworkPhysicsComponent, InputsToNetwork_Owner, COND_OwnerOnly);
 
-	// Legacy
+	// RepGraph / Legacy
 	DOREPLIFETIME_WITH_PARAMS_FAST(UNetworkPhysicsComponent, ReplicatedDeltaSourceInput, ReplicatedParamsAll);
 	DOREPLIFETIME_WITH_PARAMS_FAST(UNetworkPhysicsComponent, ReplicatedDeltaSourceState, ReplicatedParamsAll);
 	DOREPLIFETIME_WITH_PARAMS_FAST(UNetworkPhysicsComponent, ReplicatedInputs, ReplicatedParamsOwner);
@@ -644,6 +646,23 @@ void UNetworkPhysicsComponent::GetLifetimeReplicatedProps(TArray< FLifetimePrope
 	DOREPLIFETIME_WITH_PARAMS_FAST(UNetworkPhysicsComponent, ReplicatedStateCollection, ReplicatedParamsAll);
 }
 
+void UNetworkPhysicsComponent::PreReplication(IRepChangedPropertyTracker& ChangedPropertyTracker)
+{
+	Super::PreReplication(ChangedPropertyTracker);
+
+	// RepGraph / Legacy
+	DOREPLIFETIME_ACTIVE_OVERRIDE_FAST(UNetworkPhysicsComponent, ReplicatedDeltaSourceInput, bIsUsingLegacyData);
+	DOREPLIFETIME_ACTIVE_OVERRIDE_FAST(UNetworkPhysicsComponent, ReplicatedDeltaSourceState, bIsUsingLegacyData);
+	DOREPLIFETIME_ACTIVE_OVERRIDE_FAST(UNetworkPhysicsComponent, ReplicatedInputs, bIsUsingLegacyData);
+	DOREPLIFETIME_ACTIVE_OVERRIDE_FAST(UNetworkPhysicsComponent, ReplicatedRemoteInputs, bIsUsingLegacyData);
+	DOREPLIFETIME_ACTIVE_OVERRIDE_FAST(UNetworkPhysicsComponent, ReplicatedStates, bIsUsingLegacyData);
+
+	// Iris
+	DOREPLIFETIME_ACTIVE_OVERRIDE_FAST(UNetworkPhysicsComponent, ReplicatedInputCollection, !bIsUsingLegacyData);
+	DOREPLIFETIME_ACTIVE_OVERRIDE_FAST(UNetworkPhysicsComponent, ReplicatedRemoteInputCollection, !bIsUsingLegacyData);
+	DOREPLIFETIME_ACTIVE_OVERRIDE_FAST(UNetworkPhysicsComponent, ReplicatedStateCollection, !bIsUsingLegacyData);
+}
+
 // Called every Game Thread frame
 void UNetworkPhysicsComponent::TickComponent(float DeltaTime, ELevelTick TickType, FActorComponentTickFunction* ThisTickFunction)
 {
diff --git a/Engine/Source/Runtime/Engine/Public/Physics/NetworkPhysicsComponent.h b/Engine/Source/Runtime/Engine/Public/Physics/NetworkPhysicsComponent.h
--- a/Engine/Source/Runtime/Engine/Public/Physics/NetworkPhysicsComponent.h
+++ b/Engine/Source/Runtime/Engine/Public/Physics/NetworkPhysicsComponent.h
@@ -1339,6 +1339,9 @@ public:
 	// Function to init the replicated properties
 	ENGINE_API virtual void GetLifetimeReplicatedProps(TArray<FLifetimeProperty>& OutLifeTimeProps) const override;
 
+	// Function called just before replication, alter the active state of replicated properties
+	ENGINE_API virtual void PreReplication(IRepChangedPropertyTracker& ChangedPropertyTracker) override;
+
 	// Used to create any physics engine information for this component 
 	ENGINE_API virtual void BeginPlay() override;
 
