Substrate - Fix shading model override not visible on material instance when it should if hidden conversion is enabled.

#tests PC editor dx12
#rnx

#DisabledValidations CheckCopyright, CheckBalancedMacros, CheckUnacceptableWords

diff --git a/Engine/Source/Runtime/Engine/Private/Materials/Material.cpp b/Engine/Source/Runtime/Engine/Private/Materials/Material.cpp
--- a/Engine/Source/Runtime/Engine/Private/Materials/Material.cpp
+++ b/Engine/Source/Runtime/Engine/Private/Materials/Material.cpp
@@ -6576,9 +6576,10 @@ void UMaterial::GetAllExpressionsForCustomInterpolators(TArray<class UMaterialEx
 
 bool UMaterial::SupportsShadingModelOverride() const
 {
-	// If the material contains a Substrate's SubstrateShadingModels node, then we can support shading model override
+	// If the material contains a Substrate's SubstrateShadingModels node, then we can support shading model override.
+	// That is when SubstrateShadingModels nodes are used. And only if hidden conversion is not enabled, otherwise we always display it.
 	bool Out = true;
-	if (Substrate::IsSubstrateEnabled())
+	if (Substrate::IsSubstrateEnabled() && !Substrate::IsHiddenMaterialAssetConversionEnabled())
 	{
 		for (UMaterialExpression* Expression : GetExpressions())
 		{		
