Fix potential copy-paste error when implementing RequiresLODScreenSizeConversion accessors.

diff --git a/Engine/Source/Runtime/Engine/Classes/Engine/StaticMesh.h b/Engine/Source/Runtime/Engine/Classes/Engine/StaticMesh.h
--- a/Engine/Source/Runtime/Engine/Classes/Engine/StaticMesh.h
+++ b/Engine/Source/Runtime/Engine/Classes/Engine/StaticMesh.h
@@ -825,12 +825,12 @@ public:
 
 	bool GetRequiresLODScreenSizeConversion() const
 	{
-		return bRequiresLODDistanceConversion;
+		return bRequiresLODScreenSizeConversion;
 	}
 
 	void SetRequiresLODScreenSizeConversion(bool bRequiresConversion)
 	{
-		bRequiresLODDistanceConversion = bRequiresConversion;
+		bRequiresLODScreenSizeConversion = bRequiresConversion;
 	}
 
 	const FMeshNaniteSettings& GetNaniteSettings() const
