Add support for enabling the AssetCompilingManager to loan to/borrow from the SCW's memory allowance when running in memory constrained environments:
- Borrow from SCW: helps assets that require more memory than has been allocated to the AssetCompilingManager to pass without OOM
- Loan to SCW: when asset compilation memory usage is "low" - loan memory to SCW to increase throughput

#rnx

diff --git a/Engine/Source/Runtime/Core/Public/Misc/QueuedThreadPoolWrapper.h b/Engine/Source/Runtime/Core/Public/Misc/QueuedThreadPoolWrapper.h
--- a/Engine/Source/Runtime/Core/Public/Misc/QueuedThreadPoolWrapper.h
+++ b/Engine/Source/Runtime/Core/Public/Misc/QueuedThreadPoolWrapper.h
@@ -97,7 +97,21 @@ protected:
 				return TEXT("FScheduledWork");
 			}
 		}
+	protected:
+		virtual void DoThreadedWork() override
+		{
+			{
+				// Add this object as an execution context that can be retrieved via
+				// FExecutionResourceContext::Get() if a task needs to hold on the
+				// resources acquired (i.e. Concurrency Limit, Memory Pressure, etc...)
+				// longer than for the DoThreadedWork() scope.
+				FExecutionResourceContextScope ExecutionContextScope(this);
 
+				Work->DoThreadedWork();
+			}
+
+			Release();
+		}
 	private:
 		FReturnedRefCountValue AddRef() const override
 		{
@@ -132,20 +146,7 @@ protected:
 			AddRef();
 		}
 
-		void DoThreadedWork() override
-		{
-			{
-				// Add this object as an execution context that can be retrieved via
-				// FExecutionResourceContext::Get() if a task needs to hold on the
-				// resources acquired (i.e. Concurrency Limit, Memory Pressure, etc...)
-				// longer than for the DoThreadedWork() scope.
-				FExecutionResourceContextScope ExecutionContextScope(this);
-
-				Work->DoThreadedWork();
-			}
-
-			Release();
-		}
+		
 
 		void Abandon() override
 		{
@@ -182,7 +183,9 @@ protected:
 		mutable FThreadSafeCounter NumRefs;
 		friend class FQueuedThreadPoolWrapper;
 		FQueuedThreadPoolWrapper* ParentPool;
+	protected:
 		IQueuedWork* Work;
+	private:
 		EQueuedWorkPriority Priority;
 	};
 
diff --git a/Engine/Source/Runtime/Engine/Private/AssetCompilingManager.cpp b/Engine/Source/Runtime/Engine/Private/AssetCompilingManager.cpp
--- a/Engine/Source/Runtime/Engine/Private/AssetCompilingManager.cpp
+++ b/Engine/Source/Runtime/Engine/Private/AssetCompilingManager.cpp
@@ -24,6 +24,8 @@ LLM_DEFINE_TAG(AssetCompilation, NAME_None, NAME_None, GET_STATFNAME(STAT_AssetC
 #include "ProfilingDebugging/CountersTrace.h"
 #include "Animation/AnimationSequenceCompiler.h"
 #include "Animation/AnimBankCompiler.h"
+#include "ShaderCompiler.h"
+#include "Misc/ConfigCacheIni.h"
 
 #if WITH_EDITOR
 #include "DerivedDataBuildSchedulerQueue.h"
@@ -73,6 +75,12 @@ TRACE_DECLARE_INT_COUNTER(AsyncCompilationConcurrency, TEXT("AsyncCompilation/Co
 TRACE_DECLARE_MEMORY_COUNTER(AsyncCompilationAvailableMemory, TEXT("AsyncCompilation/AvailableMemory"));
 TRACE_DECLARE_MEMORY_COUNTER(AsyncCompilationTotalEstimatedMemory, TEXT("AsyncCompilation/TotalEstimatedMemory"));
 
+TRACE_DECLARE_MEMORY_COUNTER(AsyncCompilationSCWCurrentLimit, TEXT("AsyncCompilation/SCWCurrentLimit"));
+TRACE_DECLARE_MEMORY_COUNTER(AsyncCompilationSCWBaseLimit, TEXT("AsyncCompilation/SCWBaseLimit"));
+
+TRACE_DECLARE_INT_COUNTER(AsyncDelayedTask, TEXT("AsyncCompilation/DelayedTask"));
+TRACE_DECLARE_INT_COUNTER(AsyncDelayedTaskActive, TEXT("AsyncCompilation/DelayedTaskActive"));
+
 #endif //#if WITH_EDITOR
 
 namespace AssetCompilingManagerImpl
@@ -122,15 +130,20 @@ public:
 	 */
 	FMemoryBoundQueuedThreadPoolWrapper(FQueuedThreadPool* InWrappedQueuedThreadPool, TFunction<int32 ()> InMaxForegroundConcurrency, TFunction<int32()> InMaxBackgroundConcurrency, TFunction<EQueuedWorkPriority(EQueuedWorkPriority)> InPriorityMapper = [](EQueuedWorkPriority InPriority) { return InPriority; })
 		: FQueuedThreadPoolWrapper(InWrappedQueuedThreadPool, -1, InPriorityMapper)
+		, MemoryInfo([this]() { return GetMemoryAvailableForAssetCompilation(); })
 		, MaxForegroundConcurrency(InMaxForegroundConcurrency)
 		, MaxBackgroundConcurrency(InMaxBackgroundConcurrency)
 	{
 	}
 
-	int64 GetHardMemoryLimit() const
+	static int64 GetHardMemoryLimit()
 	{
 		return (int64)CVarAsyncAssetCompilationMaxMemoryUsage.GetValueOnAnyThread() * 1024 * 1024 * 1024;
 	}
+	static bool IsMemoryConstrained()
+	{
+		return GetHardMemoryLimit() != 0;
+	}
 
 	int64 GetDefaultMemoryPerAsset() const
 	{
@@ -202,16 +215,159 @@ public:
 		return RequiredMemory;
 	}
 
+	class MemorySharingInfo {
+	private:
+		const TFunction<int64 ()> GetAvailableMemory;
+		mutable FCriticalSection UpdateLock;
+
+		FShaderCompilerMemoryLimits SCWMemoryLimits;
+		FShaderCompileMemoryUsage SCWMemoryUsage;
+
+		int64 AvailableAssetCompilationMemory;
+	public:
+		int64 MemoryToLoan;
+		int64 MemoryLoanThreshold;
+		int64 MemoryToBorrowFromSCW;
+
+		bool IsMemorySharingEnabled;
+
+		MemorySharingInfo(TFunction<int64()> InGetAvailableMemory)
+			: GetAvailableMemory(InGetAvailableMemory)
+		{
+			MemoryToLoan = GConfig->GetIntOrDefault(TEXT("AssetCompilation"), TEXT("MemoryToLoanToSCW"), 2000, GEngineIni);	
+			MemoryLoanThreshold = GConfig->GetIntOrDefault(TEXT("AssetCompilation"), TEXT("MemoryLoanToSCWThreshold"), 3000, GEngineIni);	
+			MemoryToBorrowFromSCW = GConfig->GetIntOrDefault(TEXT("AssetCompilation"), TEXT("MemoryToBorrowFromSCW"), 2000, GEngineIni);	
+
+			IsMemorySharingEnabled = GConfig->GetBoolOrDefault(TEXT("AssetCompilation"), TEXT("SCWMemorySharingEnabled"), false, GEngineIni);
+		}	
+
+		int32 GetHighSCWMemory() const
+		{
+			FScopeLock MemLock(&UpdateLock);
+			return SCWMemoryLimits.BaseLimit + MemoryToLoan;
+		}
+
+		int32 GetNormalSCWMemory() const
+		{
+			FScopeLock MemLock(&UpdateLock);
+			return SCWMemoryLimits.BaseLimit;
+		}
+
+		bool IsSCWMemoryWithinLimit() const 
+		{ 
+			FScopeLock MemLock(&UpdateLock);
+			return SCWMemoryUsage.VirtualMemory <= (SCWMemoryLimits.CurrentLimit * 1024 * 1024);
+		}
+		bool IsMemoryAvailableToLoan() const
+		{
+			FScopeLock MemLock(&UpdateLock);
+			return AvailableAssetCompilationMemory >= MemoryLoanThreshold;
+		}
+		bool ShouldRetrieveFromSCW() const 
+		{
+			FScopeLock MemLock(&UpdateLock);
+			return IsMemorySharingEnabled && IsMemoryAvailableToLoan() == false && SCWMemoryLimits.CurrentLimit == GetHighSCWMemory();
+		}
+		bool CanLoanToSCW() const 
+		{
+			FScopeLock MemLock(&UpdateLock);
+			return IsMemorySharingEnabled && IsMemoryAvailableToLoan() && SCWMemoryLimits.CurrentLimit == GetNormalSCWMemory();
+		}
+
+		void Update()
+		{	
+			FScopeLock MemLock(&UpdateLock);			
+
+			GetShaderCompilerCurrentMemoryLimits(SCWMemoryLimits);								
+			SCWMemoryUsage = GShaderCompilingManager->GetExternalMemoryUsage();
+			AvailableAssetCompilationMemory = GetAvailableMemory() / (1024 * 1024);				
+		}
+	};
+
 	class FMemoryBoundScheduledWork : public FQueuedThreadPoolWrapper::FScheduledWork
 	{
 	public:
+		FMemoryBoundScheduledWork(MemorySharingInfo& InMemoryInfo) 
+		: MemoryInfo(InMemoryInfo)
+		{
+			FQueuedThreadPoolWrapper::FScheduledWork();
+		}
 		// Used to make sure this value remains constant between OnScheduled and OnUnscheduled
 		int64 RequiredMemory = 0;
+		MemorySharingInfo& MemoryInfo;
+	protected:
+		void DoThreadedWork() override
+		{			
+			MemoryInfo.Update();					
+
+			if (MemoryInfo.IsMemorySharingEnabled && FMemoryBoundQueuedThreadPoolWrapper::IsMemoryConstrained())
+			{
+				auto WaitForSCWMemory = [this]() 
+				{
+					TRACE_CPUPROFILER_EVENT_SCOPE(FMemoryBoundScheduledWork::WaitForSCWMemory);
+
+					const float MaxWaitTimeSeconds = 2.0f;
+					const float PerIterationWaitTimeSeconds  = 0.2f;
+					const int WaitIterations = static_cast<int>(MaxWaitTimeSeconds / PerIterationWaitTimeSeconds );
+
+					TRACE_COUNTER_INCREMENT(AsyncDelayedTask);
+					TRACE_COUNTER_INCREMENT(AsyncDelayedTaskActive);			
+
+					int Counter = 0;
+					for (Counter = 0; Counter < WaitIterations; Counter++) 
+					{
+						MemoryInfo.Update();
+
+						if (MemoryInfo.IsSCWMemoryWithinLimit())
+							break;
+
+						FPlatformProcess::Sleep(PerIterationWaitTimeSeconds);
+					}
+
+					if (Counter > 0)
+					{
+						UE_LOG(LogAsyncCompilation, Display, TEXT("Waited for SCW memory: %d iterations for a total of %f seconds"),
+							   Counter,
+							   Counter * PerIterationWaitTimeSeconds);
+
+
+						if (Counter == WaitIterations)
+						{
+							UE_LOG(LogAsyncCompilation, Display, TEXT("Waited for the maximum allowed time for SCW to reduce memory usage, proceeding with work anyway"));
+						}
+					}
+				
+					TRACE_COUNTER_DECREMENT(AsyncDelayedTaskActive);
+				};
+
+				if (RequiredMemory > FMemoryBoundQueuedThreadPoolWrapper::GetHardMemoryLimit())
+				{
+					SetShaderCompilerTargetMemoryLimitInMiB(MemoryInfo.GetNormalSCWMemory() - MemoryInfo.MemoryToBorrowFromSCW);
+					WaitForSCWMemory();
+					FQueuedThreadPoolWrapper::FScheduledWork::DoThreadedWork();
+					SetShaderCompilerTargetMemoryLimitInMiB(MemoryInfo.GetNormalSCWMemory());
+				}
+				else if (MemoryInfo.ShouldRetrieveFromSCW())
+				{
+					SetShaderCompilerTargetMemoryLimitInMiB(MemoryInfo.GetNormalSCWMemory());
+					WaitForSCWMemory();
+					FQueuedThreadPoolWrapper::FScheduledWork::DoThreadedWork();
+				}
+				else
+				{
+					FQueuedThreadPoolWrapper::FScheduledWork::DoThreadedWork();
+				}
+			} 
+			else 
+			{
+				FQueuedThreadPoolWrapper::FScheduledWork::DoThreadedWork();
+			}
+		}
 	};
 
 	FScheduledWork* AllocateScheduledWork() override
 	{
-		return new FMemoryBoundScheduledWork();
+		return new FMemoryBoundScheduledWork(MemoryInfo);
 	}
 
 	void UpdateCounters()
@@ -219,6 +375,11 @@ public:
 		TRACE_COUNTER_SET(AsyncCompilationAvailableMemory, GetMemoryAvailableForAssetCompilation());
 		TRACE_COUNTER_SET(AsyncCompilationTotalEstimatedMemory, TotalEstimatedMemory);
 		TRACE_COUNTER_SET(AsyncCompilationConcurrency, GetCurrentConcurrency());
+
+		FShaderCompilerMemoryLimits SCWMemoryLimits; 
+		GetShaderCompilerCurrentMemoryLimits(SCWMemoryLimits);
+		TRACE_COUNTER_SET(AsyncCompilationSCWCurrentLimit, SCWMemoryLimits.CurrentLimit * 1024 * 1024);
+		TRACE_COUNTER_SET(AsyncCompilationSCWBaseLimit, SCWMemoryLimits.BaseLimit * 1024 * 1024);
 	}
 
 	void OnScheduled(const IQueuedWork* InWork) override
@@ -253,6 +414,12 @@ public:
 				TotalEstimatedMemory/(1024*1024.f));
 
 		UpdateCounters();
+		
+		MemoryInfo.Update();
+		if (MemoryInfo.CanLoanToSCW() && IsMemoryConstrained())
+		{
+			SetShaderCompilerTargetMemoryLimitInMiB(MemoryInfo.GetHighSCWMemory());
+		} 		
 	}
 
 	int32 GetAdjustedMaxConcurrency(EQueuedWorkPriority Priority) const
@@ -391,6 +558,8 @@ public:
 	}
 
 private:
+	MemorySharingInfo MemoryInfo;
+
 	TFunction<int32()> MaxForegroundConcurrency;
 	TFunction<int32()> MaxBackgroundConcurrency;
 
diff --git a/Engine/Source/Runtime/Engine/Private/ShaderCompiler/ShaderCompilerThreadRunnable.cpp b/Engine/Source/Runtime/Engine/Private/ShaderCompiler/ShaderCompilerThreadRunnable.cpp
--- a/Engine/Source/Runtime/Engine/Private/ShaderCompiler/ShaderCompilerThreadRunnable.cpp
+++ b/Engine/Source/Runtime/Engine/Private/ShaderCompiler/ShaderCompilerThreadRunnable.cpp
@@ -29,10 +29,11 @@ static FAutoConsoleVariableRef CVarShaderCompilerTooLongIOThresholdSeconds(
 	TEXT("We don't default to parallel writes as it increases the CPU overhead from the shader compiler."),
 	ECVF_Default);
 
-int32 GShaderCompilerMemoryLimit = 0;
+int32 GShaderCompilerCurrentMemoryLimitInMiB = -1;
+int32 GShaderCompilerConfiguredMemoryLimitInMiB = 0;
 static FAutoConsoleVariableRef CVarShaderCompilerMemoryLimit(
 	TEXT("r.ShaderCompiler.MemoryLimit"),
-	GShaderCompilerMemoryLimit,
+	GShaderCompilerConfiguredMemoryLimitInMiB,
 	TEXT("Specifies a memory limit (in MiB) for all ShaderCompileWorker (SCW) processes.") \
 	TEXT("If the total memory consumption of all SCW processes exceeds this limit, the editor will start to suspend workers and reschedule compile jobs.") \
 	TEXT("By default 0, effectively disabling the limitation. If this is non-zero, it must be greater than or equal to 1024 since shader compilation must be granted at least 1024 MiB of memory in total."),
@@ -55,6 +56,21 @@ static constexpr int32 GSingleThreadedRunsMaxCount = (1 << 24);
 static const TCHAR* GWorkerInputFilename = TEXT("WorkerInputOnly.in");
 static const TCHAR* GWorkerOutputFilename = TEXT("WorkerOutputOnly.out");
 
+TRACE_DECLARE_MEMORY_COUNTER(AsyncCompilationSCWMemory, TEXT("AsyncCompilation/SCWMemory"));
+
+void GetShaderCompilerCurrentMemoryLimits(FShaderCompilerMemoryLimits& Limits)
+{
+	Limits = {0};
+
+	Limits.CurrentLimit = GShaderCompilerCurrentMemoryLimitInMiB;
+	Limits.BaseLimit = GShaderCompilerConfiguredMemoryLimitInMiB;
+}
+
+void SetShaderCompilerTargetMemoryLimitInMiB(int32 NewLimit)
+{
+	GShaderCompilerCurrentMemoryLimitInMiB = NewLimit;
+}
+
 /** Output structure for polling shader process memory violation status. See FShaderCompileThreadRunnable::QueryMemoryLimitViolationStatus(). */
 struct FShaderCompilerJobMemoryLimitInfo
 {
@@ -155,6 +171,8 @@ FShaderCompileThreadRunnableBase::FShaderCompileThreadRunnableBase(FShaderCompil
 	, MaxPriorityIndex(NumShaderCompileJobPriorities - 1)
 	, bForceFinish(false)
 {
+	if (GShaderCompilerCurrentMemoryLimitInMiB == -1)
+		GShaderCompilerCurrentMemoryLimitInMiB = GShaderCompilerConfiguredMemoryLimitInMiB;
 }
 
 bool FShaderCompileThreadRunnableBase::WorkerStateHeartbeat(uint64 InWorkerStateHash)
@@ -336,7 +354,7 @@ int32 FShaderCompileThreadRunnable::PullTasksFromQueue()
 		};
 
 	// Check if memory limitations has been violated and suspend workers as needed
-	if (GShaderCompilerMemoryLimit > 0)
+	if (GShaderCompilerCurrentMemoryLimitInMiB > 0)
 	{
 		CheckMemoryLimitViolation();
 	}
@@ -1084,6 +1102,7 @@ void FShaderCompileThreadRunnable::CheckMemoryLimitViolation()
 				MemoryMonitoringState.bHasFailedToSuspendWorkers = true; // Don't show this warning again unless we were able to suspend workers again
 			}
 		}
+		TRACE_COUNTER_SET(AsyncCompilationSCWMemory, LimitInfo.MemoryUsed);
 	}
 
 	// Check if we can resume previously suspended workers periodically
@@ -1128,26 +1147,22 @@ void FShaderCompileThreadRunnable::CheckMemoryLimitViolation()
 bool FShaderCompileThreadRunnable::QueryMemoryStatus(FShaderCompilerJobMemoryLimitInfo& OutInfo)
 {
 	const FShaderCompileMemoryUsage MemoryUsage = GetExternalWorkerMemoryUsage();
-	if (MemoryUsage.VirtualMemory > 0)
-	{
-		OutInfo.MemoryLimit = static_cast<int64>(GShaderCompilerMemoryLimit) * 1024 * 1024;
-		OutInfo.MemoryUsed = MemoryUsage.VirtualMemory;
-		return true;
-	}
-	return false;
+	
+	OutInfo.MemoryLimit = static_cast<int64>(GShaderCompilerCurrentMemoryLimitInMiB) * 1024 * 1024;
+	OutInfo.MemoryUsed = MemoryUsage.VirtualMemory;
+
+	return MemoryUsage.VirtualMemory > 0;
 }
 
 bool FShaderCompileThreadRunnable::QueryMemoryLimitViolationStatus(FShaderCompilerJobMemoryLimitInfo& OutInfo)
 {
 	const FShaderCompileMemoryUsage MemoryUsage = GetExternalWorkerMemoryUsage();
-	const int64 MemoryLimitInBytes = static_cast<int64>(GShaderCompilerMemoryLimit) * 1024 * 1024;
-	if (MemoryUsage.VirtualMemory >= static_cast<uint64>(MemoryLimitInBytes))
-	{
-		OutInfo.MemoryLimit = MemoryLimitInBytes;
-		OutInfo.MemoryUsed = MemoryUsage.VirtualMemory;
-		return true;
-	}
-	return false;
+	const int64 MemoryLimitInBytes = static_cast<int64>(GShaderCompilerCurrentMemoryLimitInMiB) * 1024 * 1024;
+
+	OutInfo.MemoryLimit = MemoryLimitInBytes;
+	OutInfo.MemoryUsed = MemoryUsage.VirtualMemory;
+
+	return MemoryUsage.VirtualMemory >= static_cast<uint64>(MemoryLimitInBytes);
 }
 
 bool FShaderCompileThreadRunnable::PrintWorkerMemoryUsage(bool bAllowToWaitForLock)
diff --git a/Engine/Source/Runtime/Engine/Public/ShaderCompiler.h b/Engine/Source/Runtime/Engine/Public/ShaderCompiler.h
--- a/Engine/Source/Runtime/Engine/Public/ShaderCompiler.h
+++ b/Engine/Source/Runtime/Engine/Public/ShaderCompiler.h
@@ -45,6 +45,7 @@ class FShaderPipelineCompileJob;
 class UMaterialInterface;
 class FJsonObject;
 struct FAnalyticsEventAttribute;
+struct FShaderCompilerMemoryLimits;
 
 DECLARE_LOG_CATEGORY_EXTERN(LogShaderCompilers, Log, All);
 
@@ -54,6 +55,9 @@ bool CreateShadersOnLoad();
 extern ENGINE_API bool IsShaderJobCacheDDCEnabled();
 extern ENGINE_API bool IsMaterialMapDDCEnabled();
 
+extern ENGINE_API void GetShaderCompilerCurrentMemoryLimits(FShaderCompilerMemoryLimits& limits);
+extern ENGINE_API void SetShaderCompilerTargetMemoryLimitInMiB(int32 NewLimit);
+
 /**
  * Returns true if we should compile shaders that are only compiled in on demand shader compilation
  * modes and not stored in cooked shader maps. This can be useful for debug shaders.
@@ -876,6 +880,11 @@ private:
 	bool bMultiProcessAggregated = false;
 };
 
+struct FShaderCompilerMemoryLimits
+{
+	uint64 CurrentLimit;
+	uint64 BaseLimit;
+};
 
 /**  
  * Manager of asynchronous and parallel shader compilation.
@@ -1074,6 +1083,9 @@ public:
 
 	ENGINE_API int32 GetNumPendingJobs() const;
 	ENGINE_API int32 GetNumOutstandingJobs() const;
+	
+	ENGINE_API FShaderCompilerMemoryLimits GetCurrentMemoryLimits() const;
+	ENGINE_API void SetTargetMemoryLimitInMiB(int32 NewLimit);
 
 	/** 
 	 * Returns whether to display a notification that shader compiling is happening in the background. 
