Fixed a couple of instance data manager issues to do with optimizing the instance data.
* ReorderArray did not handle strides other than 1
* NumCustomDataFloats was not copied from the Component data to the ChangeSet leading to mismatch (similarly also fixed AbsMaxDisplacement)
* The change set requested for generating the optimization data was requesting redundant data as it only needs positions and bounding boxes.

diff --git a/Engine/Source/Runtime/Engine/Private/Components/InstancedSkinnedMeshComponent.cpp b/Engine/Source/Runtime/Engine/Private/Components/InstancedSkinnedMeshComponent.cpp
--- a/Engine/Source/Runtime/Engine/Private/Components/InstancedSkinnedMeshComponent.cpp
+++ b/Engine/Source/Runtime/Engine/Private/Components/InstancedSkinnedMeshComponent.cpp
@@ -271,11 +271,10 @@ void UInstancedSkinnedMeshComponent::SetTransformProvider(UTransformProviderData
 template <typename ArrayType>
 inline void ReorderArray(ArrayType& InOutDataArray, const TArray<int32>& OldIndexArray, int32 ElementStride = 1)
 {
-	// TODO: this fails for FSkinnedMeshInstanceData, it is not really a problem here because the ctor just performs initializaiton and we'll overwrite it anyway 
-	// static_assert(std::is_trivially_constructible_v<decltype(InOutDataArray[0])>, "Must have trivially constructible type to safely call SetNumUninitialized below");
+	check(OldIndexArray.Num() * ElementStride == InOutDataArray.Num());
 	ArrayType TmpDataArray = MoveTemp(InOutDataArray);
 	InOutDataArray.Empty(TmpDataArray.Num());
-	for (int32 NewIndex = 0; NewIndex < TmpDataArray.Num(); ++NewIndex)
+	for (int32 NewIndex = 0; NewIndex < OldIndexArray.Num(); ++NewIndex)
 	{
 		int32 OldIndex = OldIndexArray[NewIndex];
 		for (int32 SubIndex = 0; SubIndex < ElementStride; ++SubIndex)
diff --git a/Engine/Source/Runtime/Engine/Private/InstanceData/InstanceDataManager.cpp b/Engine/Source/Runtime/Engine/Private/InstanceData/InstanceDataManager.cpp
--- a/Engine/Source/Runtime/Engine/Private/InstanceData/InstanceDataManager.cpp
+++ b/Engine/Source/Runtime/Engine/Private/InstanceData/InstanceDataManager.cpp
@@ -540,6 +540,12 @@ struct FChangeDesc2
 	}
 };
 
+static float ComputeAbsMaxDisplacement(const FInstanceDataManagerSourceDataDesc& ComponentData)
+{
+	return FMath::Max(-ComponentData.PrimitiveMaterialDesc.MinMaxMaterialDisplacement.X, ComponentData.PrimitiveMaterialDesc.MinMaxMaterialDisplacement.Y)
+		+ ComponentData.PrimitiveMaterialDesc.MaxWorldPositionOffsetDisplacement;
+}
+
 void FInstanceDataManager::GatherDefaultData(const FInstanceDataManagerSourceDataDesc& ComponentData, FInstanceUpdateChangeSet& ChangeSet) const
 {
 	// Collect the delta data to be able to update the index mapping.
@@ -549,10 +555,10 @@ void FInstanceDataManager::GatherDefaultData(const FInstanceDataManagerSourceDat
 	const FVector3f PrimitiveWorldSpacePositionHigh = FDFVector3{ ComponentData.PrimitiveLocalToWorld.GetOrigin() }.High;
 	ChangeSet.PrimitiveWorldSpaceOffset = FVector{ PrimitiveWorldSpacePositionHigh };
 	ChangeSet.PrimitiveToRelativeWorld = FDFMatrix::MakeToRelativeWorldMatrix(PrimitiveWorldSpacePositionHigh, ComponentData.PrimitiveLocalToWorld).M;
-	ChangeSet.AbsMaxDisplacement = AbsMaxDisplacement;
+	ChangeSet.AbsMaxDisplacement = ComputeAbsMaxDisplacement(ComponentData);
 	if (ChangeSet.Flags.bHasPerInstanceCustomData)
 	{
-		ChangeSet.NumCustomDataFloats = NumCustomDataFloats;
+		ChangeSet.NumCustomDataFloats = ComponentData.NumCustomDataFloats;
 	}
 
 	// Only gather data after the rest is initialized
@@ -572,9 +578,9 @@ void FInstanceDataManager::GatherDefaultData(const FInstanceDataManagerSourceDat
 
 void FInstanceDataManager::InitChangeSet(const FInstanceDataManagerSourceDataDesc& ComponentData, FInstanceUpdateChangeSet& ChangeSet) const
 {
-	// Force all data since we don't use the delta state
-	ChangeSet.ForceFullFlags = ComponentData.Flags; 
-	ChangeSet.Flags = ComponentData.Flags;
+	// This path is intended only to work with a full update.
+	check(ChangeSet.bNeedFullUpdate);
+	check(ChangeSet.Flags == ChangeSet.ForceFullFlags);
 
 	GatherDefaultData(ComponentData, ChangeSet);
 }
@@ -606,8 +612,7 @@ bool FInstanceDataManager::FlushChanges(FInstanceDataManagerSourceDataDesc&& Com
 		InstanceDataProxy = GetOrCreateProxyInternal();
 	}
 
-	float NewAbsMaxDisplacement = FMath::Max(-ComponentData.PrimitiveMaterialDesc.MinMaxMaterialDisplacement.X, ComponentData.PrimitiveMaterialDesc.MinMaxMaterialDisplacement.Y)
-		+ ComponentData.PrimitiveMaterialDesc.MaxWorldPositionOffsetDisplacement;
+	float NewAbsMaxDisplacement = ComputeAbsMaxDisplacement(ComponentData);
 
 	FChangeDesc2 ChangeDesc;
 	ChangeDesc.bUntrackedState = GetState() != ETrackingState::Tracked && GetState() != ETrackingState::Optimized;
@@ -969,7 +974,10 @@ void FInstanceDataManager::BeginCacheForCookedPlatformData(const ITargetPlatform
 
 FPrecomputedInstanceSpatialHashData FInstanceDataManager::PrecomputeOptimizationData(FInstanceDataManagerSourceDataDesc&& ComponentData)
 {
-	FInstanceUpdateChangeSet ChangeSet(ComponentData.NumInstances, ComponentData.Flags);
+	FInstanceDataFlags RequestedFlags;
+	// Only retain the per instance local bounds flag, if present in the component, since none of the other data is used in the optimization process.
+	RequestedFlags.bHasPerInstanceLocalBounds = ComponentData.Flags.bHasPerInstanceLocalBounds;
+	FInstanceUpdateChangeSet ChangeSet(ComponentData.NumInstances, RequestedFlags);
 	InitChangeSet(ComponentData, ChangeSet);
 
 	// Callback to the owner to fill in change-set data.
diff --git a/Engine/Source/Runtime/Engine/Public/InstanceData/InstanceUpdateChangeSet.h b/Engine/Source/Runtime/Engine/Public/InstanceData/InstanceUpdateChangeSet.h
--- a/Engine/Source/Runtime/Engine/Public/InstanceData/InstanceUpdateChangeSet.h
+++ b/Engine/Source/Runtime/Engine/Public/InstanceData/InstanceUpdateChangeSet.h
@@ -298,17 +298,18 @@ class FInstanceUpdateChangeSet
 {
 public:
 	/**
-	 * Cosntruct a full change set, with no delta (which will collect all enabled data).
+	 * Construct a full change set, with no delta (which will collect all enabled data).
 	 */
 	FInstanceUpdateChangeSet(int32 InNumSourceInstances, FInstanceDataFlags InFlags) 
 		: bNeedFullUpdate(true) 
 		, Flags(InFlags)
+		, ForceFullFlags(InFlags)
 		, NumSourceInstances(InNumSourceInstances)
 	{
 	}
 
 	/**
-	 * Cosntruct a delta change set, but which can be forced to full using bInNeedFullUpdate.
+	 * Construct a delta change set, but which can be forced to full using bInNeedFullUpdate.
 	 */
 	FInstanceUpdateChangeSet(bool bInNeedFullUpdate, FInstanceAttributeTracker&& InInstanceAttributeTracker, int32 InNumSourceInstances) 
 		: InstanceAttributeTracker(MoveTemp(InInstanceAttributeTracker))
