Fix incorrect RequestByte check during Substrate simplification causing material compilation to fail when using Slab node with features with Format=0.

This is because the computed "requested size" does not account for actual feature demoting. This is done directly in SubstrateExport().

#[fyi] sebastien.hillaire

diff --git a/Engine/Source/Runtime/Engine/Private/Materials/HLSLMaterialTranslator.cpp b/Engine/Source/Runtime/Engine/Private/Materials/HLSLMaterialTranslator.cpp
--- a/Engine/Source/Runtime/Engine/Private/Materials/HLSLMaterialTranslator.cpp
+++ b/Engine/Source/Runtime/Engine/Private/Materials/HLSLMaterialTranslator.cpp
@@ -14017,6 +14017,13 @@ bool FHLSLMaterialTranslator::FSubstrateCompilationContext::SubstrateGenerateDer
 						SubstrateMaterialRequestedSizeByte += UintByteSize;
 						SubstrateMaterialRequestedSizeByte += UintByteSize;
 
+						// When using blendable GBuffer, do not account for features's byte requests. as downcasting is done SubstrateExport(). 
+						// Otherwise this would cause material to not compile as we do not demote feature during simplification.
+						if (Substrate::IsSubstrateBlendableGBufferEnabled(ShaderPlatform))
+						{
+							break;
+						}
+
 						if (It.Has(ESubstrateBsdfFeature::EdgeColor | ESubstrateBsdfFeature::SecondRoughnessOrSimpleClearCoat))
 						{
 							SubstrateMaterialRequestedSizeByte += UintByteSize;
