[Input] Fix for a crash where UPlayerInput::SmoothMouse would assert before it would attempt to divide by zero.

This assert was originally added with 2876652 back in 2016.

Because the EKeys::MouseX and MouseY have their key details set to UpdateAxisWithoutSamples, it is possible (although very rare) to not have a mouse sample provided, thus triggering this assert.

To fix this, instead of throwing an assert, we can instead clamp the sample count to 1 before smoothing to avoid dividing by zero, and do an ensureMsgf instead, which is more appropriate to the severity of this issue.

[RN] Fix for a crash where UPlayerInput::SmoothMouse would assert before it would attempt to divide by zero.

diff --git a/Engine/Source/Runtime/Engine/Private/UserInterface/PlayerInput.cpp b/Engine/Source/Runtime/Engine/Private/UserInterface/PlayerInput.cpp
--- a/Engine/Source/Runtime/Engine/Private/UserInterface/PlayerInput.cpp
+++ b/Engine/Source/Runtime/Engine/Private/UserInterface/PlayerInput.cpp
@@ -1794,10 +1794,10 @@ float UPlayerInput::SmoothMouse(float aMouse, uint8& SampleCount, int32 Index)
 		check(MouseSamples > 0);
 
 		// this is seconds/sample
-		const float MouseSamplingTime = MouseSamplingTotal/MouseSamples;
+		const float MouseSamplingTime = MouseSamplingTotal / MouseSamples;
 		check(MouseSamplingTime > 0.0f);
 
-		if ( aMouse == 0 )
+		if (aMouse == 0)
 		{
 			// no mouse movement received
 			ZeroTime[Index] += DeltaTime;		// increment length of time we've been at zero
@@ -1814,15 +1814,21 @@ float UPlayerInput::SmoothMouse(float aMouse, uint8& SampleCount, int32 Index)
 		else
 		{
 			ZeroTime[Index] = 0;
-			if ( SmoothedMouse[Index] != 0 )
-			{
+			if (SmoothedMouse[Index] != 0)
+			{
+				// SampleCount may be zero because the Mouse keys are flagged as UpdateAxisWithoutSamples, which means
+				// we will process them even if they do not have any samples.
+				//
+				// For that reason, in order to guard against dive by zero errors, we will just use 1 sample count when
+				// calculating any smoothing if there aren't any this frame.
+				const uint8 NonZeroSampleCount = FMath::Max<uint8>(1, SampleCount);
+				ensureMsgf(SampleCount > 0, TEXT("The sample count was less than zero when smoothing the mouse. A value of 1 will be used in its place"));
+			
 				// this isn't the first tick with non-zero mouse movement
 				if ( DeltaTime < MouseSamplingTime * (SampleCount + 1) )
 				{
-					check(SampleCount > 0);
-
 					// smooth mouse movement so samples/tick is constant
-					aMouse = aMouse * DeltaTime/(MouseSamplingTime * SampleCount);
+					aMouse = aMouse * DeltaTime / (MouseSamplingTime * NonZeroSampleCount);
 				}
 				else
 				{
@@ -1831,15 +1837,14 @@ float UPlayerInput::SmoothMouse(float aMouse, uint8& SampleCount, int32 Index)
 					SampleCount = (uint8)(DeltaTime / MouseSamplingTime);
 				}
 			}
-
-			check(SampleCount >= 0);
+			
 			if (SampleCount == 0)
 			{
 				SmoothedMouse[Index] = 0;
 			}
 			else
 			{
-				SmoothedMouse[Index] = aMouse/SampleCount;
+				SmoothedMouse[Index] = aMouse / SampleCount;
 			}
 		}
 	}
