Fix up comment and ensure message to have correct wording.

#rnx

diff --git a/Engine/Source/Runtime/Engine/Private/UserInterface/PlayerInput.cpp b/Engine/Source/Runtime/Engine/Private/UserInterface/PlayerInput.cpp
--- a/Engine/Source/Runtime/Engine/Private/UserInterface/PlayerInput.cpp
+++ b/Engine/Source/Runtime/Engine/Private/UserInterface/PlayerInput.cpp
@@ -1819,10 +1819,10 @@ float UPlayerInput::SmoothMouse(float aMouse, uint8& SampleCount, int32 Index)
 				// SampleCount may be zero because the Mouse keys are flagged as UpdateAxisWithoutSamples, which means
 				// we will process them even if they do not have any samples.
 				//
-				// For that reason, in order to guard against dive by zero errors, we will just use 1 sample count when
+				// For that reason, in order to guard against divide by zero errors, we will just use 1 sample count when
 				// calculating any smoothing if there aren't any this frame.
 				const uint8 NonZeroSampleCount = FMath::Max<uint8>(1, SampleCount);
-				ensureMsgf(SampleCount > 0, TEXT("The sample count was less than zero when smoothing the mouse. A value of 1 will be used in its place"));
+				ensureMsgf(SampleCount > 0, TEXT("The sample count was zero when smoothing the mouse. A value of 1 will be used in its place"));
 			
 				// this isn't the first tick with non-zero mouse movement
 				if ( DeltaTime < MouseSamplingTime * (SampleCount + 1) )
