Fix lifetime issues with string table

#rnx

diff --git a/Engine/Plugins/Runtime/GameFeatures/Source/GameFeatures/Private/GameFeaturePluginStateMachine.cpp b/Engine/Plugins/Runtime/GameFeatures/Source/GameFeatures/Private/GameFeaturePluginStateMachine.cpp
--- a/Engine/Plugins/Runtime/GameFeatures/Source/GameFeatures/Private/GameFeaturePluginStateMachine.cpp
+++ b/Engine/Plugins/Runtime/GameFeatures/Source/GameFeatures/Private/GameFeaturePluginStateMachine.cpp
@@ -15,6 +15,7 @@
 #include "BundlePrereqCombinedStatusHelper.h"
 #include "Interfaces/IPluginManager.h"
 #include "Internationalization/PackageLocalizationManager.h"
+#include "Internationalization/StringTable.h"
 #include "IO/IoStoreOnDemand.h"
 #include "Logging/StructuredLog.h"
 #include "Materials/MaterialInterface.h"
@@ -3128,6 +3129,7 @@ struct FGameFeaturePluginState_Unregistering : public FGameFeaturePluginState
 
 		if (StateProperties.GameFeatureData)
 		{
+			UStringTable::OnPluginUnloaded(StateProperties.PluginName);
 			UGameFeaturesSubsystem::Get().OnGameFeatureUnregistering(StateProperties.GameFeatureData, StateProperties.PluginName, StateProperties.PluginIdentifier);
 
 			UGameFeaturesSubsystem::RemoveGameFeatureFromAssetManager(StateProperties.GameFeatureData, StateProperties.PluginName, StateProperties.AddedPrimaryAssetTypes);
diff --git a/Engine/Source/Runtime/Core/Public/Internationalization/StringTableCore.h b/Engine/Source/Runtime/Core/Public/Internationalization/StringTableCore.h
--- a/Engine/Source/Runtime/Core/Public/Internationalization/StringTableCore.h
+++ b/Engine/Source/Runtime/Core/Public/Internationalization/StringTableCore.h
@@ -26,7 +26,7 @@
 class FStringTable;
 class UStringTable;
 
-DECLARE_LOG_CATEGORY_EXTERN(LogStringTable, Log, All);
+CORE_API DECLARE_LOG_CATEGORY_EXTERN(LogStringTable, Log, All);
 
 /** Entry within a string table */
 class FStringTableEntry
diff --git a/Engine/Source/Runtime/Engine/Private/Internationalization/StringTable.cpp b/Engine/Source/Runtime/Engine/Private/Internationalization/StringTable.cpp
--- a/Engine/Source/Runtime/Engine/Private/Internationalization/StringTable.cpp
+++ b/Engine/Source/Runtime/Engine/Private/Internationalization/StringTable.cpp
@@ -112,6 +112,27 @@ public:
 		KeepAliveStringTables.RemoveSwap(InStringTableAsset);
 	}
 
+	void UnregisterForGCUnderPath(FStringView PluginPackageRoot)
+	{
+		KeepAliveStringTables.RemoveAllSwap(
+			[&PluginPackageRoot](const TObjectPtr<UStringTable>& InStringTableAsset)
+			{
+				if (!InStringTableAsset || InStringTableAsset->HasAnyFlags(RF_MirroredGarbage | RF_NewerVersionExists))
+				{
+					return true;
+				}
+
+				const FNameBuilder PackageName(InStringTableAsset->GetPackage()->GetFName());
+				if (PackageName.ToView().StartsWith(PluginPackageRoot))
+				{
+					UE_LOG(LogStringTable, Log, TEXT("Removing string table '%s' from keep-alive list"), *InStringTableAsset->GetName());
+					return true;
+				}
+
+				return false;
+			});
+	}
+
 private:
 	void HandleStringTableAssetAsyncLoadCompleted(const FName& InLoadedPackageName, UPackage* InLoadedPackage, EAsyncLoadingResult::Type InLoadingResult)
 	{
@@ -446,6 +467,12 @@ bool UStringTable::Rename(const TCHAR* NewName, UObject* NewOuter, ERenameFlags
 	return bRenamed;
 }
 
+void UStringTable::OnPluginUnloaded(FStringView PluginName)
+{
+	auto PackageNameBasis = WriteToString<256>(TEXT("/"), PluginName, TEXT("/"));
+	FStringTableEngineBridge::Get().UnregisterForGCUnderPath(PackageNameBasis);
+}
+
 FName UStringTable::GetStringTableId() const
 {
 	return StringTableId;
diff --git a/Engine/Source/Runtime/Engine/Public/Internationalization/StringTable.h b/Engine/Source/Runtime/Engine/Public/Internationalization/StringTable.h
--- a/Engine/Source/Runtime/Engine/Public/Internationalization/StringTable.h
+++ b/Engine/Source/Runtime/Engine/Public/Internationalization/StringTable.h
@@ -41,6 +41,9 @@ public:
 	ENGINE_API virtual void PostLoad() override;
 	ENGINE_API virtual bool Rename(const TCHAR* NewName = nullptr, UObject* NewOuter = nullptr, ERenameFlags Flags = REN_None) override;
 
+	/** Notifies that a plugin unmounted */
+	ENGINE_API static void OnPluginUnloaded(FStringView PluginName);
+
 	/** Get the string table ID that should be used by this asset */
 	ENGINE_API FName GetStringTableId() const;
 
