[Substrate] Fix for LOCTEXT key conflict

diff --git a/Engine/Source/Runtime/Engine/Private/RendererSettings.cpp b/Engine/Source/Runtime/Engine/Private/RendererSettings.cpp
--- a/Engine/Source/Runtime/Engine/Private/RendererSettings.cpp
+++ b/Engine/Source/Runtime/Engine/Private/RendererSettings.cpp
@@ -364,7 +364,7 @@ void URendererSettings::PostEditChangeProperty(FPropertyChangedEvent& PropertyCh
 		{
 			if (bEnableSubstrate)
 			{
-				const EAppReturnType::Type Choice = FMessageDialog::Open(EAppMsgType::OkCancel, LOCTEXT("Substrate", "Warning: Be aware that any materials saved when Substrate is enabled won't be rendered correctly if Substrate is disabled later on."));
+				const EAppReturnType::Type Choice = FMessageDialog::Open(EAppMsgType::OkCancel, LOCTEXT("EnableSubstrateWarning", "Warning: Be aware that any materials saved when Substrate is enabled won't be rendered correctly if Substrate is disabled later on."));
 				if (Choice == EAppReturnType::Cancel)
 				{
 					bEnableSubstrate = false;
@@ -372,7 +372,7 @@ void URendererSettings::PostEditChangeProperty(FPropertyChangedEvent& PropertyCh
 			}
 			else
 			{
-				const EAppReturnType::Type Choice = FMessageDialog::Open(EAppMsgType::OkCancel, LOCTEXT("Substrate", "Warning: Disabling Substrate will cause all materials authored with Substrate nodes to appear incorrect, as they are not backward compatible with the legacy root node and GBuffer format."));
+				const EAppReturnType::Type Choice = FMessageDialog::Open(EAppMsgType::OkCancel, LOCTEXT("DisableSubstrateWarning", "Warning: Disabling Substrate will cause all materials authored with Substrate nodes to appear incorrect, as they are not backward compatible with the legacy root node and GBuffer format."));
 				if (Choice == EAppReturnType::Cancel)
 				{
 					bEnableSubstrate = true;
