Make UMaterialInstanceConstant::ComputeAllStateCRC account for changes to base material's StateId.  Fixes missing updates to landscape physical material when using a material instance.

diff --git a/Engine/Source/Runtime/Engine/Private/Materials/Material.cpp b/Engine/Source/Runtime/Engine/Private/Materials/Material.cpp
--- a/Engine/Source/Runtime/Engine/Private/Materials/Material.cpp
+++ b/Engine/Source/Runtime/Engine/Private/Materials/Material.cpp
@@ -8077,6 +8077,7 @@ void UMaterial::GetLightingGuidChain(bool bIncludeTextures, TArray<FGuid>& OutGu
 uint32 UMaterial::ComputeAllStateCRC() const
 {
 	uint32 CRC = Super::ComputeAllStateCRC();
+	// TODO:  Find a way to avoid invalidating StateID on changes that can't affect computation, like repositioning nodes in the material graph.
 	CRC = FCrc::TypeCrc32(StateId, CRC);
 	return CRC;
 }
diff --git a/Engine/Source/Runtime/Engine/Private/Materials/MaterialInstanceConstant.cpp b/Engine/Source/Runtime/Engine/Private/Materials/MaterialInstanceConstant.cpp
--- a/Engine/Source/Runtime/Engine/Private/Materials/MaterialInstanceConstant.cpp
+++ b/Engine/Source/Runtime/Engine/Private/Materials/MaterialInstanceConstant.cpp
@@ -248,10 +248,34 @@ void UMaterialInstanceConstant::SetNaniteOverrideMaterial(bool bInEnableOverride
 	NaniteOverrideMaterial.OverrideMaterialEditor = InOverrideMaterial;
 }
 
+// Walk up the parent chain as long as they are all UMaterialInstanceConstant.  Gather the CRC32 of all parent ParameterStateId values.
+static uint32 GetMICParameterStateIdCRCRecursive(const UMaterialInstanceConstant* MIC)
+{
+	uint32 CRC = 0;
+	if (MIC && !MIC->GetReentrantFlag())
+	{
+		FMICReentranceGuard Guard(MIC);
+		UMaterialInstanceConstant* Next = Cast<UMaterialInstanceConstant>(MIC->Parent);
+		CRC = GetMICParameterStateIdCRCRecursive(Next);
+
+		CRC = FCrc::TypeCrc32(MIC->ParameterStateId, CRC);
+	}
+
+	return CRC;
+}
+
 uint32 UMaterialInstanceConstant::ComputeAllStateCRC() const
 {
 	uint32 CRC = Super::ComputeAllStateCRC();
-	CRC = FCrc::TypeCrc32(ParameterStateId, CRC);
+
+	const UMaterial* BaseMaterial = GetMaterial();
+	if (BaseMaterial)
+	{
+		CRC = FCrc::TypeCrc32(BaseMaterial->StateId, CRC);
+	}
+
+	CRC = FCrc::TypeCrc32(GetMICParameterStateIdCRCRecursive(this), CRC);
+
 	return CRC;
 }
 
