Fixing bug in PIE where the add SplitScreen option wasn't correctly adding a new player.

#rnx
#tests Tested that adding a splitscreen player in PIE adds a player correctly and adds another view

diff --git a/Engine/Source/Runtime/Engine/Classes/Engine/LocalPlayer.h b/Engine/Source/Runtime/Engine/Classes/Engine/LocalPlayer.h
--- a/Engine/Source/Runtime/Engine/Classes/Engine/LocalPlayer.h
+++ b/Engine/Source/Runtime/Engine/Classes/Engine/LocalPlayer.h
@@ -232,10 +232,10 @@ public:
 	 *	* (Server) Set the connection identifier on the player controller spawned on the game server.
 	 *	* (Client) Match the player controller connection identifier to the correct local user connection identifier.
 	 */
-	void SetConnectionIdentifier(const int32 Id);
+	ENGINE_API void SetConnectionIdentifier(const int32 Id);
 
 	/** Get the client connection identifier for this local user. */
-	int32 GetConnectionIdentifier() const;
+	ENGINE_API int32 GetConnectionIdentifier() const;
 
 private:
 	TArray<FSceneViewStateReference> ViewStates;
