Fixed SkinCache bugs.
 - Fixed race condition with syncing the GPU skin update task in GetCachedGeometry.
 - Removed 16 bit support for skin cache shader as it introduces precision artifacts.

diff --git a/Engine/Source/Runtime/Engine/Private/GPUSkinCache.cpp b/Engine/Source/Runtime/Engine/Private/GPUSkinCache.cpp
--- a/Engine/Source/Runtime/Engine/Private/GPUSkinCache.cpp
+++ b/Engine/Source/Runtime/Engine/Private/GPUSkinCache.cpp
@@ -946,11 +946,6 @@ public:
 		OutEnvironment.SetDefine(TEXT("GPUSKIN_RWBUFFER_OFFSET_TANGENT_Z"), FGPUSkinCache::RWTangentZOffsetInFloats);
 
 		DynamicMeshBoundsModifyCompilationEnvironment(Parameters, OutEnvironment);
-
-		if (FDataDrivenShaderPlatformInfo::GetSupportsRealTypes(Parameters.Platform) == ERHIFeatureSupport::RuntimeGuaranteed)
-		{
-			OutEnvironment.CompilerFlags.Add(CFLAG_AllowRealTypes);
-		}
 	}
 
 	BEGIN_SHADER_PARAMETER_STRUCT(FParameters, )
@@ -2466,7 +2461,6 @@ FRWBuffer* FGPUSkinCache::GetPositionBuffer(FRDGBuilder& GraphBuilder, FGPUSkinC
 {
 	if (Entry)
 	{
-		Entry->SkinCache->AddAsyncComputeWait(GraphBuilder);
 		FGPUSkinCacheEntry::FSectionDispatchData const& DispatchData = Entry->GetDispatchData()[SectionIndex];
 		FSkinCacheRWBuffer* SkinCacheRWBuffer = DispatchData.PositionBuffer;
 		return SkinCacheRWBuffer != nullptr ? &SkinCacheRWBuffer->Buffer : nullptr;
@@ -2478,7 +2472,6 @@ FRWBuffer* FGPUSkinCache::GetPreviousPositionBuffer(FRDGBuilder& GraphBuilder, F
 {
 	if (Entry)
 	{
-		Entry->SkinCache->AddAsyncComputeWait(GraphBuilder);
 		FGPUSkinCacheEntry::FSectionDispatchData const& DispatchData = Entry->GetDispatchData()[SectionIndex];
 		FSkinCacheRWBuffer* SkinCacheRWBuffer = DispatchData.PreviousPositionBuffer;
 		return SkinCacheRWBuffer != nullptr ? &SkinCacheRWBuffer->Buffer : nullptr;
@@ -2490,7 +2483,6 @@ FRWBuffer* FGPUSkinCache::GetTangentBuffer(FRDGBuilder& GraphBuilder, FGPUSkinCa
 {
 	if (Entry)
 	{
-		Entry->SkinCache->AddAsyncComputeWait(GraphBuilder);
 		FGPUSkinCacheEntry::FSectionDispatchData const& DispatchData = Entry->GetDispatchData()[SectionIndex];
 		FSkinCacheRWBuffer* SkinCacheRWBuffer = DispatchData.TangentBuffer;
 		return SkinCacheRWBuffer != nullptr ? &SkinCacheRWBuffer->Buffer : nullptr;
diff --git a/Engine/Source/Runtime/Engine/Private/SkeletalRenderGPUSkin.cpp b/Engine/Source/Runtime/Engine/Private/SkeletalRenderGPUSkin.cpp
--- a/Engine/Source/Runtime/Engine/Private/SkeletalRenderGPUSkin.cpp
+++ b/Engine/Source/Runtime/Engine/Private/SkeletalRenderGPUSkin.cpp
@@ -2579,6 +2579,8 @@ bool FSkeletalMeshObjectGPUSkin::GetCachedGeometry(FRDGBuilder& GraphBuilder, FC
 		return false;
 	}
 
+	FGPUSkinCache::AddAsyncComputeWait(GraphBuilder);
+
 	for (uint32 SectionIndex = 0; SectionIndex < SectionCount; ++SectionIndex)
 	{
 		FCachedGeometry::Section& CachedSection = OutCachedGeometry.Sections.AddDefaulted_GetRef();
diff --git a/Engine/Source/Runtime/Engine/Public/GPUSkinCache.h b/Engine/Source/Runtime/Engine/Public/GPUSkinCache.h
--- a/Engine/Source/Runtime/Engine/Public/GPUSkinCache.h
+++ b/Engine/Source/Runtime/Engine/Public/GPUSkinCache.h
@@ -196,8 +196,8 @@ public:
 
 	ENGINE_API UE::Tasks::FTask Dispatch(FRDGBuilder& GraphBuilder, const UE::Tasks::FTask& PrerequisitesTask, ERHIPipeline Pipeline);
 
-	ENGINE_API void AddAsyncComputeSignal(FRDGBuilder& GraphBuilder);
-	ENGINE_API void AddAsyncComputeWait(FRDGBuilder& GraphBuilder);
+	ENGINE_API static void AddAsyncComputeSignal(FRDGBuilder& GraphBuilder);
+	ENGINE_API static void AddAsyncComputeWait(FRDGBuilder& GraphBuilder);
 
 	inline ERHIFeatureLevel::Type GetFeatureLevel() const
 	{
