Remove old assert which could activate when deleting partially loaded blueprints, it crashes too late to be helpful for diagnosing any issue

diff --git a/Engine/Source/Runtime/Engine/Classes/Engine/BlueprintGeneratedClass.h b/Engine/Source/Runtime/Engine/Classes/Engine/BlueprintGeneratedClass.h
--- a/Engine/Source/Runtime/Engine/Classes/Engine/BlueprintGeneratedClass.h
+++ b/Engine/Source/Runtime/Engine/Classes/Engine/BlueprintGeneratedClass.h
@@ -104,7 +104,8 @@ public:
 
 	~FPointerToUberGraphFrame()
 	{
-		check(!RawPointer);
+		// check(!RawPointer);
+		// We don't want to check here because partially loaded blueprint classes may not be cleaned up properly. Re-enable if you are looking for specific ubergraph issues
 	}
 };
 
