World Partition - Destructible HLODs: Fix condition leading to a fatal check
* "VisibilityBuffer.Num() > DestStride" test was incorrect. With destructible instanced HLODs, the visibility buffer is now used to store the state of destructible actors for both the HLOD mesh AND instanced HLODs. Actors that are destructible in the mesh are at the beginning of the array.
* FMath::Min() was already mitigating the issue by only copying a valid range.

diff --git a/Engine/Source/Runtime/Engine/Private/WorldPartition/HLOD/DestructibleHLODComponent.cpp b/Engine/Source/Runtime/Engine/Private/WorldPartition/HLOD/DestructibleHLODComponent.cpp
--- a/Engine/Source/Runtime/Engine/Private/WorldPartition/HLOD/DestructibleHLODComponent.cpp
+++ b/Engine/Source/Runtime/Engine/Private/WorldPartition/HLOD/DestructibleHLODComponent.cpp
@@ -17,6 +17,7 @@ DEFINE_LOG_CATEGORY_STATIC(LogHLODDestruction, Log, All);
 
 UWorldPartitionDestructibleHLODComponent::UWorldPartitionDestructibleHLODComponent(const FObjectInitializer& ObjectInitializer)
 	: Super(ObjectInitializer)
+	, VisibilityTextureSize(0)
 {
 	SetIsReplicatedByDefault(true);
 }
@@ -237,6 +238,7 @@ void UWorldPartitionDestructibleHLODComponent::SetupVisibilityTexture()
 			if (DynamicInstanceTexture)
 			{
 				VisibilityTexture = DynamicInstanceTexture;
+				VisibilityTextureSize = TextureSize;
 			}
 
 			if (VisibilityTexture)
@@ -281,35 +283,23 @@ void UWorldPartitionDestructibleHLODComponent::UpdateVisibilityTexture()
 		if (FApp::CanEverRender() && ensure(TextureResource))
 		{
 			ENQUEUE_RENDER_COMMAND(FUpdateHLODVisibilityTexture)(
-				[TextureResource, VisibilityBuffer = VisibilityBuffer](FRHICommandListImmediate& RHICmdList)
-			{
-					UpdateVisibilityTexture_RenderThread(TextureResource, VisibilityBuffer);
-			});
+				[TextureResource, VisibilityTextureSize = VisibilityTextureSize, VisibilityBuffer = VisibilityBuffer](FRHICommandListImmediate& RHICmdList)
+				{
+					FRHITexture* TextureRHI = TextureResource->GetTexture2DRHI();
+					FRHILockTextureArgs LockArgs = FRHILockTextureArgs::Lock2D(TextureRHI, 0, RLM_WriteOnly, false, false);
+					FRHILockTextureResult LockResult = RHICmdList.LockTexture(LockArgs);
+					{
+						checkf(VisibilityTextureSize <= LockResult.Stride, TEXT("UpdateVisibilityTexture: copy dest (%d) is smaller than source (%d)"), LockResult.Stride, VisibilityTextureSize);
+						checkf(VisibilityTextureSize <= (uint32)VisibilityBuffer.Num(), TEXT("UpdateVisibilityTexture: Visibility texture size (%d) is greater than the visibility buffer size (%d)"), VisibilityTextureSize, VisibilityBuffer.Num());
+						uint32 CopyLen = FMath::Min(VisibilityTextureSize, FMath::Min(LockResult.Stride, (uint32)VisibilityBuffer.Num()));
+						FMemory::Memcpy((uint8*)LockResult.Data, VisibilityBuffer.GetData(), CopyLen);
+					}
+					RHICmdList.UnlockTexture(LockArgs);
+				});
 		}
 	}
 }
 
-void UWorldPartitionDestructibleHLODComponent::UpdateVisibilityTexture_RenderThread(FTexture2DDynamicResource* TextureResource, const TArray<uint8>& VisibilityBuffer)
-{
-	check(IsInRenderingThread());
-
-	FRHITexture* TextureRHI = TextureResource->GetTexture2DRHI();
-
-	uint32 DestStride = 0;
-	uint8* DestData = reinterpret_cast<uint8*>(RHILockTexture2D(TextureRHI, 0, RLM_WriteOnly, DestStride, false, false));
-
-#if !(UE_BUILD_SHIPPING)
-	if ((uint32)VisibilityBuffer.Num() > DestStride)
-	{
-		UE_LOG(LogHLODDestruction, Fatal, TEXT("UpdateVisibilityTexture_RenderThread: copy dest (%d) is smaller than source (%d). Memory will be stomped."), DestStride, VisibilityBuffer.Num());
-	}
-#endif
-
-	FMemory::Memcpy(DestData, VisibilityBuffer.GetData(), FMath::Min((uint32)VisibilityBuffer.Num(), DestStride));
-
-	RHIUnlockTexture2D(TextureRHI, 0, false, false);
-}
-
 #if WITH_EDITOR
 void UWorldPartitionDestructibleHLODComponent::SetDestructibleActors(const TArray<FName>& InDestructibleActors)
 {
diff --git a/Engine/Source/Runtime/Engine/Public/WorldPartition/HLOD/DestructibleHLODComponent.h b/Engine/Source/Runtime/Engine/Public/WorldPartition/HLOD/DestructibleHLODComponent.h
--- a/Engine/Source/Runtime/Engine/Public/WorldPartition/HLOD/DestructibleHLODComponent.h
+++ b/Engine/Source/Runtime/Engine/Public/WorldPartition/HLOD/DestructibleHLODComponent.h
@@ -298,7 +298,6 @@ private:
 
 	void SetupVisibilityTexture();
 	void UpdateVisibilityTexture();
-	static void UpdateVisibilityTexture_RenderThread(FTexture2DDynamicResource* TextureResource, const TArray<uint8>& VisibilityBuffer);
 
 private:
 	UPROPERTY()
@@ -316,6 +315,8 @@ private:
 	UPROPERTY(Transient)
 	TObjectPtr<UTexture2DDynamic> VisibilityTexture;
 
+	uint32 VisibilityTextureSize;
+
 	// Client only, visibility buffer that is meant to be sent to the GPU
 	TArray<uint8> VisibilityBuffer;
 
