Post Process Volume:  Improvement to determinism of post process volume sorting.  Post process volumes are now sorted by the volume of their bounds and an arbitrary guid, if Priority is otherwise equal.  Smaller volumes take precedence over larger volumes (sorted later in the order of application), so nested volumes behave intuitively.

diff --git a/Engine/Plugins/Experimental/Water/Source/Runtime/Private/WaterBodyComponent.cpp b/Engine/Plugins/Experimental/Water/Source/Runtime/Private/WaterBodyComponent.cpp
--- a/Engine/Plugins/Experimental/Water/Source/Runtime/Private/WaterBodyComponent.cpp
+++ b/Engine/Plugins/Experimental/Water/Source/Runtime/Private/WaterBodyComponent.cpp
@@ -505,6 +505,10 @@ FPostProcessVolumeProperties UWaterBodyComponent::GetPostProcessProperties() con
 	Ret.BlendWeight = UnderwaterPostProcessSettings.BlendWeight;
 	Ret.Priority = UnderwaterPostProcessSettings.Priority;
 	Ret.Settings = &CurrentPostProcessSettings;
+	Ret.Size = Bounds.GetBox().GetVolume();
+
+	// We can give this a singleton Guid, because only one water body is activated at a time, depending on player location.
+	Ret.VolumeGuid = FGuid(0xe0809956, 0xa2eb49f2, 0xb4f0aa6c, 0x143c89b1);
 	return Ret;
 }
 
diff --git a/Engine/Plugins/Experimental/Water/Source/Runtime/Public/WaterSubsystem.h b/Engine/Plugins/Experimental/Water/Source/Runtime/Public/WaterSubsystem.h
--- a/Engine/Plugins/Experimental/Water/Source/Runtime/Public/WaterSubsystem.h
+++ b/Engine/Plugins/Experimental/Water/Source/Runtime/Public/WaterSubsystem.h
@@ -39,7 +39,9 @@ struct FUnderwaterPostProcessVolume : public IInterface_PostProcessVolume
 {
 	FUnderwaterPostProcessVolume()
 		: PostProcessProperties()
-	{}
+	{
+		PostProcessProperties.VolumeGuid = FGuid(0x857a3d0e, 0x222e4c40, 0x833b5e81, 0x499b8d95);
+	}
 
 	virtual bool EncompassesPoint(FVector Point, float SphereRadius/*=0.f*/, float* OutDistanceToPoint) override
 	{
diff --git a/Engine/Source/Runtime/Engine/Classes/Components/PostProcessComponent.h b/Engine/Source/Runtime/Engine/Classes/Components/PostProcessComponent.h
--- a/Engine/Source/Runtime/Engine/Classes/Components/PostProcessComponent.h
+++ b/Engine/Source/Runtime/Engine/Classes/Components/PostProcessComponent.h
@@ -49,6 +49,10 @@ class UPostProcessComponent : public USceneComponent, public IInterface_PostProc
 	UPROPERTY(EditAnywhere, BlueprintReadWrite, Category = PostProcessVolume)
 	uint32 bUnbound: 1;
  
+	/** Automatically generated GUID for this post process volume, used for deterministic sorting of equal priority volumes. */
+	UPROPERTY(TextExportTransient, NonPIEDuplicateTransient, NonTransactional)
+	FGuid VolumeGuid;
+
 	//~ Begin IInterface_PostProcessVolume Interface
 	ENGINE_API virtual bool EncompassesPoint(FVector Point, float SphereRadius/*=0.f*/, float* OutDistanceToPoint) override;
 	virtual FPostProcessVolumeProperties GetProperties() const override
@@ -60,6 +64,8 @@ class UPostProcessComponent : public USceneComponent, public IInterface_PostProc
 		Ret.BlendWeight = BlendWeight;
 		Ret.Priority = Priority;
 		Ret.Settings = &Settings;
+		Ret.Size = Ret.bIsUnbound ? DBL_MAX : Cast<UShapeComponent>(GetAttachParent())->CalcLocalBounds().GetBox().GetVolume();
+		Ret.VolumeGuid = VolumeGuid;
 		return Ret;
 	}
 #if DEBUG_POST_PROCESS_VOLUME_ENABLE
@@ -81,6 +87,7 @@ protected:
 
 	//~ Begin UObject Interface
 	virtual void Serialize(FArchive& Ar) override;
+	virtual void PostInitProperties() override;
 	//~ End UObject Interface
 };
 
diff --git a/Engine/Source/Runtime/Engine/Classes/Engine/PostProcessVolume.h b/Engine/Source/Runtime/Engine/Classes/Engine/PostProcessVolume.h
--- a/Engine/Source/Runtime/Engine/Classes/Engine/PostProcessVolume.h
+++ b/Engine/Source/Runtime/Engine/Classes/Engine/PostProcessVolume.h
@@ -50,6 +50,10 @@ class APostProcessVolume : public AVolume, public IInterface_PostProcessVolume
 	UPROPERTY(EditAnywhere, BlueprintReadWrite, Category=PostProcessVolumeSettings, meta=(DisplayName = "Infinite Extent (Unbound)"))
 	uint32 bUnbound:1;
 
+	/** Automatically generated GUID for this post process volume, used for deterministic sorting of equal priority volumes. */
+	UPROPERTY(TextExportTransient, NonPIEDuplicateTransient, NonTransactional)
+	FGuid VolumeGuid;
+
 	//~ Begin IInterface_PostProcessVolume Interface
 	ENGINE_API virtual bool EncompassesPoint(FVector Point, float SphereRadius/*=0.f*/, float* OutDistanceToPoint) override;
 	virtual FPostProcessVolumeProperties GetProperties() const override
@@ -61,6 +65,8 @@ class APostProcessVolume : public AVolume, public IInterface_PostProcessVolume
 		Ret.BlendWeight = BlendWeight;
 		Ret.Priority = Priority;
 		Ret.Settings = &Settings;
+		Ret.Size = Ret.bIsUnbound ? DBL_MAX : GetBounds().GetBox().GetVolume();
+		Ret.VolumeGuid = VolumeGuid;
 		return Ret;
 	}
 #if DEBUG_POST_PROCESS_VOLUME_ENABLE
@@ -92,6 +98,7 @@ public:
 	ENGINE_API virtual bool CanEditChange(const FProperty* InProperty) const override;
 #endif // WITH_EDITOR
 	ENGINE_API virtual void Serialize(FArchive& Ar);
+	ENGINE_API virtual void PostInitProperties() override;
 	//~ End UObject Interface
 
 	/** Adds an Blendable (implements IBlendableInterface) to the array of Blendables (if it doesn't exist) and update the weight */
diff --git a/Engine/Source/Runtime/Engine/Classes/Engine/World.h b/Engine/Source/Runtime/Engine/Classes/Engine/World.h
--- a/Engine/Source/Runtime/Engine/Classes/Engine/World.h
+++ b/Engine/Source/Runtime/Engine/Classes/Engine/World.h
@@ -1893,8 +1893,8 @@ public:
 	/** Inserts a post process volume into the world in priority order */
 	UE_API void InsertPostProcessVolume(IInterface_PostProcessVolume* InVolume);
 
-	/** Removes a post process volume from the world */
-	UE_API void RemovePostProcessVolume(IInterface_PostProcessVolume* InVolume);
+	/** Removes a post process volume from the world.  Returns true if it was successfully removed. */
+	UE_API bool RemovePostProcessVolume(IInterface_PostProcessVolume* InVolume);
 
 	/** Called when a scene view for this world needs the worlds post process settings computed */
 	UE_API void AddPostProcessingSettings(FVector ViewLocation, FSceneView* SceneView);
@@ -1902,6 +1902,11 @@ public:
 	/** An array of post processing volumes, sorted in ascending order of priority.					*/
 	TArray< IInterface_PostProcessVolume * > PostProcessVolumes;
 
+private:
+	/** Caches the size of each volume at the time of insertion, for sort determinism in case it later changes.  One for one with PostProcessVolumes above. */
+	TArray<double> PostProcessVolumeCachedSizes;
+public:
+
 	/** Set of AudioVolumes sorted by  */
 	// TODO: Make this be property UPROPERTY(Transient)
 	TArray<class AAudioVolume*> AudioVolumes;
diff --git a/Engine/Source/Runtime/Engine/Classes/Interfaces/Interface_PostProcessVolume.h b/Engine/Source/Runtime/Engine/Classes/Interfaces/Interface_PostProcessVolume.h
--- a/Engine/Source/Runtime/Engine/Classes/Interfaces/Interface_PostProcessVolume.h
+++ b/Engine/Source/Runtime/Engine/Classes/Interfaces/Interface_PostProcessVolume.h
@@ -19,6 +19,13 @@ struct FPostProcessVolumeProperties
 	float BlendWeight;
 	bool bIsEnabled;
 	bool bIsUnbound;
+
+	// Size (volume of the PPV's bounding box), is used as an additional sort key if Priority is equal.  Smaller volumes take higher
+	// priority.  Unbounded volumes are treated as having max volume.  If volume is also equal, an arbitrary guid is used as a final
+	// sort key.  The goal is to improve determinism in the order equal Priority volumes are processed, which otherwise is affected
+	// by the order in which they were loaded, which can vary between editor and engine, platform, and due to streaming.
+	double Size = DBL_MAX;
+	FGuid VolumeGuid;
 };
 
 /** Interface for general PostProcessVolume access **/
diff --git a/Engine/Source/Runtime/Engine/Private/Components/PostProcessComponent.cpp b/Engine/Source/Runtime/Engine/Private/Components/PostProcessComponent.cpp
--- a/Engine/Source/Runtime/Engine/Private/Components/PostProcessComponent.cpp
+++ b/Engine/Source/Runtime/Engine/Private/Components/PostProcessComponent.cpp
@@ -40,6 +40,18 @@ void UPostProcessComponent::Serialize(FArchive& Ar)
 	}
 }
 
+void UPostProcessComponent::PostInitProperties()
+{
+	Super::PostInitProperties();
+
+#if WITH_EDITOR
+	if (!HasAnyFlags(EObjectFlags::RF_ClassDefaultObject | EObjectFlags::RF_DefaultSubObject) && !VolumeGuid.IsValid())
+	{
+		VolumeGuid = FGuid::NewGuid();
+	}
+#endif
+}
+
 bool UPostProcessComponent::EncompassesPoint(FVector Point, float SphereRadius/*=0.f*/, float* OutDistanceToPoint)
 {
 	UShapeComponent* ParentShape = Cast<UShapeComponent>(GetAttachParent());
diff --git a/Engine/Source/Runtime/Engine/Private/PostProcessVolume.cpp b/Engine/Source/Runtime/Engine/Private/PostProcessVolume.cpp
--- a/Engine/Source/Runtime/Engine/Private/PostProcessVolume.cpp
+++ b/Engine/Source/Runtime/Engine/Private/PostProcessVolume.cpp
@@ -3,6 +3,7 @@
 #include "Engine/PostProcessVolume.h"
 #include "Engine/BlendableInterface.h"
 #include "Engine/CollisionProfile.h"
+#include "Engine/Engine.h"
 #include "Components/BrushComponent.h"
 #include "EngineUtils.h"
 #include "SceneInterface.h"
@@ -22,15 +23,59 @@ APostProcessVolume::APostProcessVolume(const FObjectInitializer& ObjectInitializ
 	BlendWeight = 1.0f;
 }
 
+#if WITH_EDITOR
+
+// This delegate is needed to handle update of sort order when a post process volume's dimensions are changed in the editor.
+static FDelegateHandle GPostProcessVolumeActorMovedDelegateHandle;
+
+// Count of post process volume actors using the delegate above.  When this goes to zero, we can free the delegate.
+static int32 GPostProcessVolumeDelegateCount = 0;
+
+// Post process volumes are re-inserted to update their sort order following an editor move operation, as changes to the transform
+// can affect the bounds used as part of the sort key.  By design, outside the editor, sorting is NOT updated after insertion, to
+// avoid unexpected changes in order.  This would only matter if the post process volume is dynamically moved, which is likely
+// rare to begin with, but we'd rather such movement not later affect sorting.
+static void PostProcessVolumeOnEditorActorMoved(AActor* InActor)
+{
+	APostProcessVolume* VolumeActor = Cast<APostProcessVolume>(InActor);
+	if (VolumeActor)
+	{
+		UWorld* World = VolumeActor->GetWorld();
+		if (World && World->WorldType == EWorldType::Editor)
+		{
+			// Don't re-insert it if we didn't actually remove it
+			if (World->RemovePostProcessVolume(VolumeActor))
+			{
+				World->InsertPostProcessVolume(VolumeActor);
+			}
+		}
+	}
+}
+#endif  // WITH_EDITOR
+
 void APostProcessVolume::PostUnregisterAllComponents()
 {
+#if WITH_EDITOR
+	if (GEngine)
+	{
+		GPostProcessVolumeDelegateCount--;
+		check(GPostProcessVolumeDelegateCount >= 0);
+
+		if (GPostProcessVolumeDelegateCount == 0)
+		{
+			GEngine->OnActorMoved().Remove(GPostProcessVolumeActorMovedDelegateHandle);
+
+			GPostProcessVolumeActorMovedDelegateHandle.Reset();
+		}
+	}
+#endif
+
 	// Route clear to super first.
 	Super::PostUnregisterAllComponents();
 	// World will be NULL during exit purge.
 	if (GetWorld())
 	{
 		GetWorld()->RemovePostProcessVolume(this);
-		GetWorld()->PostProcessVolumes.RemoveSingle(this);
 	}
 }
 
@@ -39,6 +84,18 @@ void APostProcessVolume::PostRegisterAllComponents()
 	// Route update to super first.
 	Super::PostRegisterAllComponents();
 	GetWorld()->InsertPostProcessVolume(this);
+
+#if WITH_EDITOR
+	if (GEngine)
+	{
+		// Add a delegate so moved post process volumes can update their sort order, as the Volume size is affected by the transform.
+		if (GPostProcessVolumeDelegateCount == 0)
+		{
+			GPostProcessVolumeActorMovedDelegateHandle = GEngine->OnActorMoved().AddStatic(&PostProcessVolumeOnEditorActorMoved);
+		}
+		GPostProcessVolumeDelegateCount++;
+	}
+#endif
 }
 
 bool APostProcessVolume::IsPPVEnabled() const
@@ -72,6 +129,18 @@ void APostProcessVolume::Serialize(FArchive& Ar)
 #endif
 }
 
+void APostProcessVolume::PostInitProperties()
+{
+	Super::PostInitProperties();
+
+#if WITH_EDITOR
+	if (!HasAnyFlags(EObjectFlags::RF_ClassDefaultObject) && !VolumeGuid.IsValid())
+	{
+		VolumeGuid = FGuid::NewGuid();
+	}
+#endif
+}
+
 #if WITH_EDITOR
 void APostProcessVolume::PostEditChangeProperty(FPropertyChangedEvent& PropertyChangedEvent)
 {
diff --git a/Engine/Source/Runtime/Engine/Private/World.cpp b/Engine/Source/Runtime/Engine/Private/World.cpp
--- a/Engine/Source/Runtime/Engine/Private/World.cpp
+++ b/Engine/Source/Runtime/Engine/Private/World.cpp
@@ -10239,6 +10239,27 @@ FPrimaryAssetId UWorld::GetPrimaryAssetId() const
 	return FPrimaryAssetId();
 }
 
+static bool PostProcessVolumeLess(const FPostProcessVolumeProperties& PropertiesA, const FPostProcessVolumeProperties& PropertiesB, double SizeB)
+{
+	if (PropertiesA.Priority != PropertiesB.Priority)
+	{
+		return PropertiesA.Priority < PropertiesB.Priority;
+	}
+	if (PropertiesA.bIsUnbound != PropertiesB.bIsUnbound)
+	{
+		// Unbounded volumes sort first.
+		return PropertiesA.bIsUnbound;
+	}
+	if (PropertiesA.Size != SizeB)
+	{
+		// Larger sizes sort first.
+		return PropertiesA.Size > SizeB;
+	}
+
+	// Finally, sort by guid.
+	return PropertiesA.VolumeGuid < PropertiesB.VolumeGuid;
+}
+
 void UWorld::InsertPostProcessVolume(IInterface_PostProcessVolume* InVolume)
 {
 	const int32 NumVolumes = PostProcessVolumes.Num();
@@ -10248,31 +10269,43 @@ void UWorld::InsertPostProcessVolume(IInterface_PostProcessVolume* InVolume)
 	for (; InsertIndex < NumVolumes; InsertIndex++)
 	{
 		IInterface_PostProcessVolume* CurrentVolume = PostProcessVolumes[InsertIndex];
-		FPostProcessVolumeProperties CurrentProperties = CurrentVolume->GetProperties();
-
-		if (TargetProperties.Priority < CurrentProperties.Priority)
-		{
-			break;
-		}
-		else if (TargetProperties.Priority > CurrentProperties.Priority)
+		if (CurrentVolume == InVolume)
 		{
-			continue;
+			return;
 		}
-		if (TargetProperties.bIsUnbound && !CurrentProperties.bIsUnbound)
+
+		FPostProcessVolumeProperties CurrentProperties = CurrentVolume->GetProperties();
+
+		if (PostProcessVolumeLess(TargetProperties, CurrentProperties, PostProcessVolumeCachedSizes[InsertIndex]))
 		{
 			break;
 		}
-		if (CurrentVolume == InVolume)
+	}
+
+	// Check the remainder of the array to see if InVolume is already present.  The caller may have changed TargetProperties.Size
+	// from what was originally cached, making the "Size > CurrentCachedSize" comparison inconsistent, meaning the existing item will not
+	// have been reached in the loop.
+	for (int32 Index = InsertIndex + 1; Index < NumVolumes; Index++)
+	{
+		if (PostProcessVolumes[Index] == InVolume)
 		{
 			return;
 		}
 	}
 	PostProcessVolumes.Insert(InVolume, InsertIndex);
+	PostProcessVolumeCachedSizes.Insert(TargetProperties.Size, InsertIndex);
 }
 
-void UWorld::RemovePostProcessVolume(IInterface_PostProcessVolume* InVolume)
+bool UWorld::RemovePostProcessVolume(IInterface_PostProcessVolume* InVolume)
 {
-	PostProcessVolumes.RemoveSingle(InVolume);
+	int32 RemoveIndex;
+	if (PostProcessVolumes.Find(InVolume, RemoveIndex))
+	{
+		PostProcessVolumes.RemoveAt(RemoveIndex);
+		PostProcessVolumeCachedSizes.RemoveAt(RemoveIndex);
+		return true;
+	}
+	return false;
 }
 
 void UWorld::InitializeSubsystems()
