[UE] Always export references to non-default-constructed inner subobjects, rather than only doing so when they differ from their archetype on a deep comparison.

Notes:
- This brings T3D export to parity with the logic used since 4.26 for exporting subobject references as part of delta serialization on write during package save.
- The instancing graph is shared across the entire import of an exported outer root object. Historically this allows the import to resolve inner references to the root object on import.
- SourceToDestination mappings in the instancing graph can change during an outer object's import if subobjects are exported with the same archetype (e.g. CDO).
- Previously, subobject instancing on import would then incorrectly resolve a default-initialized reference to the previously-imported non-default-constructed instance, if it had already mapped the archetype.
- This bug was limited to the case where a deep comparison on export resulted in no difference between the instanced subobject and its archetype (i.e. no value overrides were also exported in the redefinition block).

diff --git a/Engine/Source/Runtime/Engine/Private/UnrealExporter.cpp b/Engine/Source/Runtime/Engine/Private/UnrealExporter.cpp
--- a/Engine/Source/Runtime/Engine/Private/UnrealExporter.cpp
+++ b/Engine/Source/Runtime/Engine/Private/UnrealExporter.cpp
@@ -205,6 +205,15 @@ bool UExporter::ExportToOutputDevice(const FExportObjectInnerContext* Context, U
 	if ( (PortFlags&PPF_SubobjectsOnly) == 0 )
 	{
 		PortFlags |= PPF_DeepComparison;
+
+		// always export references to instanced subobjects that are constructed through non-default means, regardless of equality
+		// this is needed to properly rebuild the outer object's instance graph during import after initializing from CDO/archetype
+		// notes:
+		// - a DSO reference can be uniquely resolved from its initialized value, since constructed DSOs are mapped by the initializer
+		// - a DSO with an outer that's initialized from a native CDO will not need to fix up any references to it post-initialization
+		// - a DSO overridden at edit time is considered unique when a deep comparison differs, and so any references must be exported
+		// - references to non-default subobjects must also be exported as they can't be resolved through the initializer's DSO mapping
+		PortFlags |= PPF_DeepCompareDSOsOnly;
 	}
 
 	if ( FCString::Stricmp(FileType, TEXT("COPY")) == 0 )
