Splines: Make spline input keys actually control the internal spline parameter space

Comes with a suite of fixes necessary for this hookup which couldn't necessarily be submitted independently due to the fact that they would apparently break things without the hookup:
1) Fixes auto tangent computation (previously ignored knot deltas)
2) Fixes tangent bezier evaluation (ignored knot deltas at evaluation time)
3) Fixes mutations at looping point
4) Fixes loop position override
5) Fixes several conversion between spline types
6) Fixes polyline conversion logic
7) Fixes reparam table generation and evaluation for non-uniform knot vectors
8) Fixes segment length computation for non-uniform knot vectors

Also synchronizes spline at load time to ensure backing data structures are in a consistent state.

#FailedValidations Preflight

diff --git a/Engine/Source/Editor/DetailCustomizations/Private/SplineComponentDetails.cpp b/Engine/Source/Editor/DetailCustomizations/Private/SplineComponentDetails.cpp
--- a/Engine/Source/Editor/DetailCustomizations/Private/SplineComponentDetails.cpp
+++ b/Engine/Source/Editor/DetailCustomizations/Private/SplineComponentDetails.cpp
@@ -1004,20 +1004,12 @@ bool FSplinePointDetails::CanSetInputKey() const
 		return false;
 	}
 
-PRAGMA_DISABLE_DEPRECATION_WARNINGS
-	bool bUsingSplineCurves = SplineComp->GetSplinePropertyName() == GET_MEMBER_NAME_CHECKED(USplineComponent, SplineCurves);
-PRAGMA_ENABLE_DEPRECATION_WARNINGS
-
-	return IsOnePointSelected() && bUsingSplineCurves;
+	return IsOnePointSelected();
 }
 
 void FSplinePointDetails::OnSetInputKey(float NewValue, ETextCommit::Type CommitInfo)
 {
 PRAGMA_DISABLE_DEPRECATION_WARNINGS
-	// Note: This function strongly assumes that SplineCurves is the authoritative data structure backing the spline component.
-	// I have made this assumption valid by introducing CanSetInputKey() which verifies this by checking with the selected component.
-	// This assumption is necessary because there is no interface on the component which allows us to modify input keys, we must directly write to SplineCurves.
-	
 	TryAcquireDetailsInterface();
 	
 	if ((CommitInfo != ETextCommit::OnEnter && CommitInfo != ETextCommit::OnUserMovedFocus) || !SplineComp || !SplineDetailsInterface)
@@ -1027,13 +1019,12 @@ PRAGMA_DISABLE_DEPRECATION_WARNINGS
 
 	check(SelectedKeys.Num() == 1);
 	const int32 Index = *SelectedKeys.CreateConstIterator();
-	TArray<FInterpCurvePoint<FVector>>& Positions = SplineComp->GetSplinePointsPosition().Points;
 
-	const int32 NumPoints = Positions.Num();
+	const int32 NumPoints = SplineComp->GetNumberOfSplinePoints();
 
 	bool bModifyOtherPoints = false;
-	if ((Index > 0 && NewValue <= Positions[Index - 1].InVal) ||
-		(Index < NumPoints - 1 && NewValue >= Positions[Index + 1].InVal))
+	if ((Index > 0 && NewValue <= SplineComp->GetInputKeyValueAtSplinePoint(Index - 1) ||
+		(Index < NumPoints - 1 && NewValue >= SplineComp->GetInputKeyValueAtSplinePoint(Index + 1))))
 	{
 		const FText Title(LOCTEXT("InputKeyTitle", "Input key out of range"));
 		const FText Message(LOCTEXT("InputKeyMessage", "Spline input keys must be numerically ascending. Would you like to modify other input keys in the spline in order to be able to set this value?"));
@@ -1052,37 +1043,28 @@ PRAGMA_DISABLE_DEPRECATION_WARNINGS
 		const FScopedTransaction Transaction(LOCTEXT("SetSplinePointInputKey", "Set spline point input key"));
 		SplineComp->Modify();
 
-		TArray<FInterpCurvePoint<FQuat>>& Rotations = SplineComp->GetSplinePointsRotation().Points;
-		TArray<FInterpCurvePoint<FVector>>& Scales = SplineComp->GetSplinePointsScale().Points;
-
 		if (bModifyOtherPoints)
 		{
 			// Shuffle the previous or next input keys down or up so the input value remains in sequence
-			if (Index > 0 && NewValue <= Positions[Index - 1].InVal)
+			if (Index > 0 && NewValue <= SplineComp->GetInputKeyValueAtSplinePoint(Index - 1))
 			{
-				float Delta = (NewValue - Positions[Index].InVal);
+				float Delta = (NewValue - SplineComp->GetInputKeyValueAtSplinePoint(Index));
 				for (int32 PrevIndex = 0; PrevIndex < Index; PrevIndex++)
 				{
-					Positions[PrevIndex].InVal += Delta;
-					Rotations[PrevIndex].InVal += Delta;
-					Scales[PrevIndex].InVal += Delta;
+					SplineComp->SetInputKeyValueAtSplinePoint(PrevIndex, SplineComp->GetInputKeyValueAtSplinePoint(PrevIndex) + Delta);
 				}
 			}
-			else if (Index < NumPoints - 1 && NewValue >= Positions[Index + 1].InVal)
+			else if (Index < NumPoints - 1 && NewValue >= SplineComp->GetInputKeyValueAtSplinePoint(Index + 1))
 			{
-				float Delta = (NewValue - Positions[Index].InVal);
+				float Delta = (NewValue - SplineComp->GetInputKeyValueAtSplinePoint(Index));
 				for (int32 NextIndex = Index + 1; NextIndex < NumPoints; NextIndex++)
 				{
-					Positions[NextIndex].InVal += Delta;
-					Rotations[NextIndex].InVal += Delta;
-					Scales[NextIndex].InVal += Delta;
+					SplineComp->SetInputKeyValueAtSplinePoint(NextIndex, SplineComp->GetInputKeyValueAtSplinePoint(NextIndex) + Delta);
 				}
 			}
 		}
 
-		Positions[Index].InVal = NewValue;
-		Rotations[Index].InVal = NewValue;
-		Scales[Index].InVal = NewValue;
+		SplineComp->SetInputKeyValueAtSplinePoint(Index, NewValue);
 	}
 
 	SplineComp->UpdateSpline();
diff --git a/Engine/Source/Runtime/Core/Public/UObject/UE5ReleaseStreamObjectVersions.inl b/Engine/Source/Runtime/Core/Public/UObject/UE5ReleaseStreamObjectVersions.inl
--- a/Engine/Source/Runtime/Core/Public/UObject/UE5ReleaseStreamObjectVersions.inl
+++ b/Engine/Source/Runtime/Core/Public/UObject/UE5ReleaseStreamObjectVersions.inl
@@ -186,5 +186,8 @@ UE5_RELEASE_VERSION(ExtraOutputAttributesOnDataDomainPCG, 58)
 // Media Profile: Changed storage of capture cameras list from lazy to soft pointers
 UE5_RELEASE_VERSION(MediaProfilePluginCaptureCameraSoftPtr, 59)
 
+// Reparameterize Spline in SplineComponent based on SplineCurves parameterization
+UE5_RELEASE_VERSION(SplineComponentReparameterizeOnLoad, 60)
+
 // -----<new versions can be added above this line>-------------------------------------------------
 #undef UE5_RELEASE_VERSION
\ No newline at end of file
diff --git a/Engine/Source/Runtime/Engine/Classes/Components/SplineComponent.h b/Engine/Source/Runtime/Engine/Classes/Components/SplineComponent.h
--- a/Engine/Source/Runtime/Engine/Classes/Components/SplineComponent.h
+++ b/Engine/Source/Runtime/Engine/Classes/Components/SplineComponent.h
@@ -405,7 +405,7 @@ public:
 	FSimpleMulticastDelegate& GetOnSplineDisplayChanged() { return OnSplineDisplayChanged; }
 	
 	/** Reset all spline data to match the current state of SplineCurves */
-	void SynchronizeSplines() { SetSpline(WarninglessSplineCurves()); }
+	ENGINE_API void SynchronizeSplines();
 
 	ENGINE_API int32 GetVersion() const;
 	ENGINE_API const FInterpCurveVector& GetSplinePointsPosition() const;
@@ -445,6 +445,9 @@ public:
 	UFUNCTION(BlueprintCallable, Category = Spline)
 	ENGINE_API virtual void SetOverrideConstructionScript(bool InOverride);
 
+	UFUNCTION(BlueprintCallable, Category = Spline)
+	ENGINE_API int32 GetPointIndexAtSplineInputKey(float InKey) const;
+	
 	/** Get location along spline at the provided input key value */
 	UFUNCTION(BlueprintCallable, Category = Spline)
 	ENGINE_API FVector GetLocationAtSplineInputKey(float InKey, ESplineCoordinateSpace::Type CoordinateSpace) const;
@@ -609,10 +612,14 @@ public:
 	UFUNCTION(BlueprintCallable, Category = Spline)
 	ENGINE_API int32 GetNumberOfSplineSegments() const;
 
-	/** Get the input key (e.g. the time) of the control point of the spline at the specified index. */
+	/** Get the input key of the control point of the spline at the specified index. */
 	UFUNCTION(BlueprintCallable, Category = Spline)
 	ENGINE_API float GetInputKeyValueAtSplinePoint(int32 PointIndex) const;
 
+	/** Set the input key of the control point of the spline at the specified index. */
+	UFUNCTION(BlueprintCallable, Category = Spline)
+	ENGINE_API void SetInputKeyValueAtSplinePoint(int32 PointIndex, float InputKey);
+
 	/** Gets the spline point of the spline at the specified index */
 	UFUNCTION(BlueprintCallable, Category = Spline)
 	ENGINE_API FSplinePoint GetSplinePointAt(int32 PointIndex, ESplineCoordinateSpace::Type CoordinateSpace) const;
diff --git a/Engine/Source/Runtime/Engine/Classes/Curves/Spline.h b/Engine/Source/Runtime/Engine/Classes/Curves/Spline.h
--- a/Engine/Source/Runtime/Engine/Classes/Curves/Spline.h
+++ b/Engine/Source/Runtime/Engine/Classes/Curves/Spline.h
@@ -56,9 +56,11 @@ struct FSpline
 	ENGINE_API void SetSplinePointType(int32 Index, EInterpCurveMode Type);
 	ENGINE_API EInterpCurveMode GetSplinePointType(int32 Index) const;
 
+	ENGINE_API void SetParameterAtIndex(int32 Index, float Parameter);
 	ENGINE_API float GetParameterAtIndex(int32 Index) const;
 	ENGINE_API float GetParameterAtDistance(float Distance) const;
 	ENGINE_API float GetDistanceAtParameter(float Parameter) const;
+	ENGINE_API int32 GetIndexAtParameter(float Parameter) const;
 
 	ENGINE_API FQuat GetOrientation(int32 Index) const;
 	ENGINE_API void SetOrientation(int32 Index, const FQuat& InOrientation);
diff --git a/Engine/Source/Runtime/Engine/Private/Components/SplineComponent.cpp b/Engine/Source/Runtime/Engine/Private/Components/SplineComponent.cpp
--- a/Engine/Source/Runtime/Engine/Private/Components/SplineComponent.cpp
+++ b/Engine/Source/Runtime/Engine/Private/Components/SplineComponent.cpp
@@ -8,6 +8,7 @@
 #include "Engine/Engine.h"
 #include "SceneView.h"
 #include "UObject/EditorObjectVersion.h"
+#include "UObject/UE5ReleaseStreamObjectVersion.h"
 #include "Math/RotationMatrix.h"
 #include "MeshElementCollector.h"
 #include "PrimitiveViewRelevance.h"
@@ -237,6 +238,7 @@ void USplineComponent::Serialize(FArchive& Ar)
 	Super::Serialize(Ar);
 
 	Ar.UsingCustomVersion(FEditorObjectVersion::GUID);
+	Ar.UsingCustomVersion(FUE5ReleaseStreamObjectVersion::GUID);
 
 #if WITH_EDITORONLY_DATA
 	if (Ar.IsLoading())
@@ -250,13 +252,22 @@ void USplineComponent::Serialize(FArchive& Ar)
 			WarninglessSplineCurves().ReparamTable = SplineReparamTable_DEPRECATED;
 		}
 
+		bool bHasSynchronized = false;
+		
 #if WITH_EDITOR
 		// WasEnabled() is editor only, that data gets stripped from packaged builds.
-		if (!Spline.WasEnabled())
+		if (!bHasSynchronized && !Spline.WasEnabled())
 		{
 			SynchronizeSplines();
+			bHasSynchronized = true;
 		}
 #endif
+
+		if (!bHasSynchronized && Ar.CustomVer(FUE5ReleaseStreamObjectVersion::GUID) < FUE5ReleaseStreamObjectVersion::SplineComponentReparameterizeOnLoad)
+		{
+			SynchronizeSplines();
+			bHasSynchronized = true;
+		}
 	}
 #endif
 
@@ -350,6 +361,7 @@ void FSplineCurves::UpdateSpline(bool bClosedLoop, bool bStationaryEndpoints, in
 	// Now initialize the spline reparam table
 	const int32 NumSegments = bClosedLoop ? NumPoints : FMath::Max(0, NumPoints - 1);
 
+	// todo: reparam table is wrong, stored input keys are computed as functions of point indices. we must refer to interp curve input keys directly.
 	// Start by clearing it
 	ReparamTable.Points.Reset(NumSegments * ReparamStepsPerSegment + 1);
 	float AccumulatedLength = 0.0f;
@@ -530,6 +542,13 @@ float USplineComponent::GetSegmentParamFromLength(const int32 Index, const float
 	return Param;
 }
 
+int32 USplineComponent::GetPointIndexAtSplineInputKey(float InKey) const
+{
+	return UE::SplineComponent::ShouldUseSplineCurves()
+		? WarninglessSplineCurves().Position.GetPointIndexForInputValue(InKey)
+		: Spline.GetIndexAtParameter(InKey);
+}
+
 FVector USplineComponent::GetLocationAtSplineInputKey(float InKey, ESplineCoordinateSpace::Type CoordinateSpace) const
 {
 	FVector Location = UE::SplineComponent::ShouldUseSplineCurves()
@@ -1126,6 +1145,16 @@ void USplineComponent::AddSplinePointAtIndex(const FVector& Position, int32 Inde
 		NewPoint.Rotation = FQuat::Identity.Rotator();
 		NewPoint.Scale = FVector(1.f);
 
+		// Adjust subsequent points' input keys to make room for the value we are about to add
+		int32 WarninglessSplinePointsNum = WarninglessSplineCurves().Position.Points.Num();
+		for (int I = WarninglessSplinePointsNum - 1; I >= Index; --I)
+		{
+			Spline.SetParameterAtIndex(I, Spline.GetParameterAtIndex(I) + 1.0f);
+			WarninglessSplineCurves().Position.Points[I].InVal += 1.0f;
+			WarninglessSplineCurves().Rotation.Points[I].InVal += 1.0f;
+			WarninglessSplineCurves().Scale.Points[I].InVal += 1.0f;
+		}
+		
 		Spline.InsertPoint(NewPoint, Index);
 
 		WarninglessSplineCurves().Position.Points.Insert(FInterpCurvePoint<FVector>(InKey, TransformedPosition, FVector::ZeroVector, FVector::ZeroVector, CIM_CurveAuto), Index);
@@ -1137,15 +1166,6 @@ void USplineComponent::AddSplinePointAtIndex(const FVector& Position, int32 Inde
 			Metadata->InsertPoint(Index, 0.5f, bClosedLoop);
 		}
 		
-		// Adjust subsequent points' input keys to make room for the value just added
-		int32 WarninglessSplinePointsNum = WarninglessSplineCurves().Position.Points.Num();
-		for (int I = Index + 1; I < WarninglessSplinePointsNum; ++I)
-		{
-			WarninglessSplineCurves().Position.Points[I].InVal += 1.0f;
-			WarninglessSplineCurves().Rotation.Points[I].InVal += 1.0f;
-			WarninglessSplineCurves().Scale.Points[I].InVal += 1.0f;
-		}
-
 		if (bLoopPositionOverride)
 		{
 			LoopPosition += 1.0f;
@@ -1166,11 +1186,13 @@ void USplineComponent::RemoveSplinePoint(int32 Index, bool bUpdateSpline)
 
 	if (Index >= 0 && Index < NumPoints)
 	{
-		Spline.RemovePoint(Index);
+		const float TargetPointParameter = GetInputKeyValueAtSplinePoint(Index);
+		const float PreviousPointParameter = Index > 0 ? GetInputKeyValueAtSplinePoint(Index - 1) : 0.f;
 
-		const float KnotDelta = Index > 0
-			? WarninglessSplineCurves().Position.Points[Index].InVal - WarninglessSplineCurves().Position.Points[Index - 1].InVal
-			: WarninglessSplineCurves().Position.Points[Index].InVal;
+		// The parameter range of the segment being removed.
+		const float KnotDelta = TargetPointParameter - PreviousPointParameter;
+		
+		Spline.RemovePoint(Index);
 		
 		WarninglessSplineCurves().Position.Points.RemoveAt(Index, EAllowShrinking::No);
 		WarninglessSplineCurves().Rotation.Points.RemoveAt(Index, EAllowShrinking::No);
@@ -1186,6 +1208,7 @@ void USplineComponent::RemoveSplinePoint(int32 Index, bool bUpdateSpline)
 		// Adjust all following spline point input keys to close the gap left by the removed point
 		while (Index < NumPoints)
 		{
+			Spline.SetParameterAtIndex(Index, Spline.GetParameterAtIndex(Index) - KnotDelta);
 			WarninglessSplineCurves().Position.Points[Index].InVal -= KnotDelta;
 			WarninglessSplineCurves().Rotation.Points[Index].InVal -= KnotDelta;
 			WarninglessSplineCurves().Scale.Points[Index].InVal -= KnotDelta;
@@ -1474,6 +1497,43 @@ float USplineComponent::GetInputKeyValueAtSplinePoint(int32 PointIndex) const
 	}
 }
 
+void USplineComponent::SetInputKeyValueAtSplinePoint(int32 PointIndex, float InputKey)
+{
+	const int32 NumPoints = GetNumberOfSplinePoints();
+	
+	if (NumPoints == 0)
+	{
+		return;
+	}
+
+	PointIndex = GetClampedIndex(PointIndex);
+
+	// todo: handle the following early-outs better (requires reordering the spline points and spline point metadata).
+	
+	// Prevents reordering the knot vector. If the knot vector is monotonic, we only need to check the immediate neighbors of the point.
+	if ((PointIndex + 1 < NumPoints && InputKey >= GetInputKeyValueAtSplinePoint(PointIndex + 1)) ||
+		(PointIndex - 1 >= 0        && InputKey <= GetInputKeyValueAtSplinePoint(PointIndex - 1)))
+	{
+		return;
+	}
+
+	// This is a bit annoying, we have to catch this case because FSpline implements closed splines with a true closing segment,
+	// so we would actually end up flipping this segment if not for this early out.
+	if (PointIndex == NumPoints - 1 && IsClosedLoop())
+	{
+		if (const float LoopInputKey = bLoopPositionOverride ? LoopPosition : GetInputKeyValueAtSplinePoint(PointIndex) + 1.f;
+			InputKey >= LoopInputKey)
+		{
+			return;
+		}
+	}
+
+	WarninglessSplineCurves().Position.Points[PointIndex].InVal = InputKey;
+	WarninglessSplineCurves().Rotation.Points[PointIndex].InVal = InputKey;
+	WarninglessSplineCurves().Scale.Points[PointIndex].InVal = InputKey;
+	Spline.SetParameterAtIndex(PointIndex, InputKey);
+}
+
 FSplinePoint USplineComponent::GetSplinePointAt(int32 PointIndex, ESplineCoordinateSpace::Type CoordinateSpace) const
 {
 	PointIndex = GetClampedIndex(PointIndex);
@@ -2095,7 +2155,7 @@ bool USplineComponent::DivideSplineIntoPolylineRecursiveWithDistancesHelper(floa
 		
 	}
 
-	check(OutPoints.Num() == OutDistancesAlongSpline.Num())
+	check(OutPoints.Num() == OutDistancesAlongSpline.Num());
 	return (OutPoints.Num() > 0);
 }
 
@@ -2279,21 +2339,21 @@ bool USplineComponent::ConvertSplineToPolyline_InDistanceRange(ESplineCoordinate
 	}
 	else
 	{
-		// with new spline, we are strongly assuming that point index is always same as param
-		// so, we can convert param directly to segment index.
-		// we do this without making the assumptions about the reparameterization table (unlike the other branch)
-
-		const float StartParam = Spline.GetParameterAtDistance(RangeStart);
-		const float EndParam = Spline.GetParameterAtDistance(RangeEnd);
-
-		SegmentStart = FMath::FloorToInt32(StartParam);
-		SegmentEnd = FMath::Min(NumSegments, FMath::CeilToInt32(EndParam));
+		const float RangeStartParam = GetInputKeyValueAtDistanceAlongSpline(RangeStart);
+		const float RangeEndParam = GetInputKeyValueAtDistanceAlongSpline(RangeEnd);
+		SegmentStart = GetPointIndexAtSplineInputKey(RangeStartParam);
+		SegmentEnd = FMath::Min(NumSegments, 1 + GetPointIndexAtSplineInputKey(RangeEndParam));
 	}
 
 	TArray<FVector> NewPoints;
 	TArray<double> NewDistances;
 	for (int32 SegmentIndex = SegmentStart; SegmentIndex < SegmentEnd; ++SegmentIndex)
 	{
+		if (SegmentIndex == INDEX_NONE)
+		{
+			continue;
+		}
+		
 		// Get the segment range as distances, clipped with the input range
 		double StartDist = FMath::Max(RangeStart, GetDistanceAlongSplineAtSplinePoint(SegmentIndex));
 		double StopDist = FMath::Min(RangeEnd, GetDistanceAlongSplineAtSplinePoint(SegmentIndex + 1));
@@ -3135,6 +3195,12 @@ void USplineComponent::SetSpline(const FSplineCurves& InSplineCurves)
 	OnSplineChanged.Broadcast();
 }
 
+void USplineComponent::SynchronizeSplines()
+{
+	SetSpline(WarninglessSplineCurves());
+	UpdateSpline();
+}
+
 int32 USplineComponent::GetVersion() const
 {
 	return UE::SplineComponent::ShouldUseSplineCurves()
@@ -3274,11 +3340,15 @@ void USplineComponent::PostEditChangeChainProperty(FPropertyChangedChainEvent& P
 		static const FName StationaryEndpointsName = GET_MEMBER_NAME_CHECKED(USplineComponent, bStationaryEndpoints);
 		static const FName DefaultUpVectorName = GET_MEMBER_NAME_CHECKED(USplineComponent, DefaultUpVector);
 		static const FName ClosedLoopName = GET_MEMBER_NAME_CHECKED(USplineComponent, bClosedLoop);
+		static const FName LoopPositionName = GET_MEMBER_NAME_CHECKED(USplineComponent, LoopPosition);
+		static const FName LoopPositionOverrideName = GET_MEMBER_NAME_CHECKED(USplineComponent, bLoopPositionOverride);
 
 		if (PropertyName == ReparamStepsPerSegmentName ||
 			PropertyName == StationaryEndpointsName ||
 			PropertyName == DefaultUpVectorName ||
-			PropertyName == ClosedLoopName)
+			PropertyName == ClosedLoopName ||
+			PropertyName == LoopPositionName ||
+			PropertyName == LoopPositionOverrideName)
 		{
 			UpdateSpline();
 
diff --git a/Engine/Source/Runtime/Engine/Private/Curves/Spline.cpp b/Engine/Source/Runtime/Engine/Private/Curves/Spline.cpp
--- a/Engine/Source/Runtime/Engine/Private/Curves/Spline.cpp
+++ b/Engine/Source/Runtime/Engine/Private/Curves/Spline.cpp
@@ -167,11 +167,6 @@ public:
         return Ar;
     }
 
-	/**
-	 * Returns the parameter space of FNewSpline. All public functions taking a parameter are only defined when the provided value is in this interval.
-	 * @return [0, NumSegments] when NumSegment != 0, otherwise undefined (empty interval).
-	 */
-	FInterval1f GetSegmentSpace() const;
 	float GetSegmentLength(const int32 Index, const float Param, const FVector& Scale3D = FVector(1.0f)) const;
 	float GetSplineLength();
 	
@@ -205,9 +200,11 @@ public:
 	void SetSplinePointType(int32 Index, EInterpCurveMode Type);
 	EInterpCurveMode GetSplinePointType(int32 Index) const;
 
+	void SetParameterAtIndex(int32 Index, float Parameter);
 	float GetParameterAtIndex(int32 Index) const;
 	float GetParameterAtDistance(float Distance);
 	float GetDistanceAtParameter(float Parameter);
+	int32 GetIndexAtParameter(float Parameter) const;
 
     virtual float FindNearest(const FVector& InLocation, float& OutSquaredDistance) const override;
     float FindNearestOnSegment(const FVector& InLocation, int32 SegmentIndex, float& OutSquaredDistance) const;
@@ -295,23 +292,17 @@ private:
 	
 	int32 ReparamStepsPerSegment = 10;
 
+	// valid when the spline has 1 point. this lets us remember the desired parameter value to integrate it into the actual knot vector which doesn't exist til we have a full segment
+	TOptional<float> FirstPointParameter;
+
+	// Valid when the spline is looped and the update parameters specify an override loop parameter (i.e. the parameter of the virtual closing point).
+	TOptional<float> LoopParameter;
+	
 private:
 
 	// In order for FSpline to access the parameter mapping functions below.
 	// The alternative is to redefine a bunch of MultiSpline functions to do the remapping internally, but this is a bit of a pain.
 	friend struct FSpline;
-	
-	/** Convert from [0, 1] to [0, NumSegments] */
-	float FromNormalizedSpace(float Parameter) const;
-	
-	/** Convert from [0, NumSegments] to [0, 1] */
-	float ToNormalizedSpace(float Parameter) const;
-
-	/** Convert from internal spline space to [0, NumSegments] */
-	float FromInternalSplineSpace(float Parameter) const;
-	
-	/** Convert from [0, NumSegments] to internal spline space */
-	float ToInternalSplineSpace(float Parameter) const;
 
 	void MarkReparamTableDirty();
 	
@@ -357,7 +348,7 @@ int32 FNewSpline::AddAttributeValue(FName Name, const AttrType& Value, float Par
 		ControlPointType ControlPoint = ControlPointType(Value);
 		int32 NumPoints = Channel->GetNumPoints();
 
-		float ParentSpaceParameter = ToInternalSplineSpace(Parameter);
+		float ParentSpaceParameter = Parameter;
 		FInterval1f MappedChildSpace = GetMappedChildSpace(Name);
 		
 		// Special case for empty spline
@@ -448,13 +439,13 @@ void FNewSpline::RemoveAttributeValue(FName Name, int32 Index)
 			if (Index == 0)
 			{
 				FInterval1f NewMappingRange = GetMappedChildSpace(Name);
-				NewMappingRange.Min = ToInternalSplineSpace(GetAttributeParameter<AttrType>(Name, Index + 1));
+				NewMappingRange.Min = GetAttributeParameter<AttrType>(Name, Index + 1);
 				SetAttributeChannelRange(Name, NewMappingRange, EMappingRangeSpace::Parent);
 			}
 			else if (Index == Channel->GetNumPoints() - 1)
 			{
 				FInterval1f NewMappingRange = GetMappedChildSpace(Name);
-				NewMappingRange.Max = ToInternalSplineSpace(GetAttributeParameter<AttrType>(Name, Index - 1));
+				NewMappingRange.Max = GetAttributeParameter<AttrType>(Name, Index - 1);
 				SetAttributeChannelRange(Name, NewMappingRange, EMappingRangeSpace::Parent);
 			}
 		}
@@ -496,7 +487,7 @@ int32 FNewSpline::SetAttributeParameter(const FName& Name, int32 Index, float Pa
 	
 	if (AttrSplineType* Channel = GetTypedAttributeChannel<AttrSplineType>(Name))
 	{
-		float ParentSpaceParameter = ToInternalSplineSpace(Parameter);
+		float ParentSpaceParameter = Parameter;
 
 		// Prevent collapse of the channel space and keep it starting at 0.
 		auto SanitizeChannelSpace = [Channel]()
@@ -517,7 +508,7 @@ int32 FNewSpline::SetAttributeParameter(const FName& Name, int32 Index, float Pa
 		// Uses current value of Index and ParentSpaceParameter at the time of calling to update the attribute parameter.
 		auto SetParentSpaceParameter = [this, &ParentSpaceParameter, Name, Channel, &Index]() -> int32
 		{
-			const float CurrentChildSpaceParameter = MapParameterToChildSpace(Name, ToInternalSplineSpace(GetAttributeParameter<AttrType>(Name, Index)));
+			const float CurrentChildSpaceParameter = MapParameterToChildSpace(Name, GetAttributeParameter<AttrType>(Name, Index));
 			const float DesiredChildSpaceParameter = MapParameterToChildSpace(Name, ParentSpaceParameter);
 
 			constexpr float MinStep = 2.f * UE_KINDA_SMALL_NUMBER;
@@ -571,7 +562,7 @@ int32 FNewSpline::SetAttributeParameter(const FName& Name, int32 Index, float Pa
 			}
 
 			float ParentSpaceUpperBound = GetMappedChildSpace(Name).Max;
-			const float NeighborParentSpaceParameter = ToInternalSplineSpace(GetAttributeParameter<AttrType>(Name, Index + 1));
+			const float NeighborParentSpaceParameter = GetAttributeParameter<AttrType>(Name, Index + 1);
 			const bool bAttributesWillReorder = ParentSpaceParameter > NeighborParentSpaceParameter;
 
 			// If the end point is passing its neighbor, we need to actually reshuffle the points.
@@ -627,7 +618,7 @@ int32 FNewSpline::SetAttributeParameter(const FName& Name, int32 Index, float Pa
 			}
 
 			float ParentSpaceLowerBound = GetMappedChildSpace(Name).Min;
-			const float NeighborParentSpaceParameter = ToInternalSplineSpace(GetAttributeParameter<AttrType>(Name, Index - 1));
+			const float NeighborParentSpaceParameter = GetAttributeParameter<AttrType>(Name, Index - 1);
 			const bool bAttributesWillReorder = ParentSpaceParameter < NeighborParentSpaceParameter;
 			
 			// If the end point is passing its neighbor, we need to actually reshuffle the points.
@@ -705,7 +696,7 @@ float FNewSpline::GetAttributeParameter(const FName& Name, int32 Index) const
 	
 	if (AttrSplineType* Channel = GetTypedAttributeChannel<AttrSplineType>(Name))
 	{
-		return FromInternalSplineSpace(MapParameterFromChildSpace(Name, Channel->GetParameter(Index)));
+		return MapParameterFromChildSpace(Name, Channel->GetParameter(Index));
 	}
 	return 0.0f;
 }
@@ -750,9 +741,11 @@ public:
 	void SetSplinePointType(int32 Index, EInterpCurveMode Type);
 	EInterpCurveMode GetSplinePointType(int32 Index) const;
 
+	void SetParameterAtIndex(int32 Index, float Parameter);
 	float GetParameterAtIndex(int32 Index) const;
 	float GetParameterAtDistance(float Distance) const;
 	float GetDistanceAtParameter(float Parameter) const;
+	int32 GetIndexAtParameter(float Parameter) const;
 	
 	/* Parameter Interface */
 
@@ -1055,6 +1048,16 @@ EInterpCurveMode FSpline::GetSplinePointType(int32 Index) const
 	}
 }
 
+void FSpline::SetParameterAtIndex(int32 Index, float Parameter)
+{
+	switch (CurrentImplementation)
+	{
+	default: break;
+	case 1: LegacyData->SetParameterAtIndex(Index, Parameter); break;
+	case 2: NewData->SetParameterAtIndex(Index, Parameter); break;
+	}
+}
+
 float FSpline::GetParameterAtIndex(int32 Index) const
 {
 	switch (CurrentImplementation)
@@ -1085,6 +1088,16 @@ float FSpline::GetDistanceAtParameter(float Parameter) const
 	}
 }
 
+int32 FSpline::GetIndexAtParameter(float Parameter) const
+{
+	switch (CurrentImplementation)
+	{
+	default: return INDEX_NONE;
+	case 1: return LegacyData->GetIndexAtParameter(Parameter);
+	case 2: return NewData->GetIndexAtParameter(Parameter);
+	}
+}
+
 FQuat FSpline::GetOrientation(int32 Index) const
 {
 	return GetOrientation(GetParameterAtIndex(Index));
@@ -1352,7 +1365,7 @@ AttrType FSpline::EvaluateAttribute(float Param, FName AttributeName) const
 {
 	if (NewData && SupportsAttributes())
 	{
-		return NewData->EvaluateAttribute<AttrType>(AttributeName, NewData->ToInternalSplineSpace(Param));
+		return NewData->EvaluateAttribute<AttrType>(AttributeName, Param);
 	}
 	return AttrType(0);
 }
@@ -1799,6 +1812,8 @@ FNewSpline::FNewSpline(const FLegacySpline& Other)
 	CreateAttributeChannel<UE::Geometry::Spline::TTangentBezierSpline<FVector>>(ScaleAttrName);
 	
 	TArray<FTangentBezierControlPoint> Points;
+	TArray<UE::Geometry::Spline::FKnot> KnotVector;
+	
     for (int32 Idx = 0; Idx < Other.GetNumControlPoints(); Idx++)
     {
 	    FTangentBezierControlPoint& Point = Points.AddDefaulted_GetRef();
@@ -1807,10 +1822,15 @@ FNewSpline::FNewSpline(const FLegacySpline& Other)
     	Point.TangentIn = Other.GetInTangent(Idx);
     	Point.TangentOut = Other.GetOutTangent(Idx);
     	Point.TangentMode = ConvertInterpCurveModeToTangentMode(Other.GetSplinePointType(Idx));
+    	
+    	const uint32 KnotMultiplicity = Idx == 0 || Idx == Other.GetNumControlPoints() - 1 ? 4 : 3;
+        KnotVector.Emplace(Other.GetParameterAtIndex(Idx), KnotMultiplicity);
     }
 
 	// initialize spline
 	GetSpline().SetControlPoints(Points);
+	GetSpline().SetKnotVector(KnotVector);
+	
 	ResetRotation();
 	ResetScale();
 	
@@ -1825,6 +1845,9 @@ FNewSpline::FNewSpline(const FLegacySpline& Other)
 
 	FSpline::FUpdateSplineParams UpdateParams;
 	UpdateParams.bClosedLoop = Other.GetSplinePointsPosition().bIsLooped;
+	UpdateParams.LoopPosition = Other.GetSplinePointsPosition().Points.Num() > 0
+		? Other.GetSplinePointsPosition().Points.Last().InVal + Other.GetSplinePointsPosition().LoopKeyOffset
+		: 0.f;
 	
     UpdateSpline(UpdateParams);
 }
@@ -1835,6 +1858,8 @@ FNewSpline::FNewSpline(const FSplineCurves& Other)
 	CreateAttributeChannel<UE::Geometry::Spline::TTangentBezierSpline<FVector>>(ScaleAttrName);
 	
 	TArray<FTangentBezierControlPoint> Points;
+	TArray<UE::Geometry::Spline::FKnot> KnotVector;
+	
     for (int32 Idx = 0; Idx < Other.Position.Points.Num(); Idx++)
     {
 	    const FInterpCurvePoint<FVector>& Position = Other.Position.Points[Idx];
@@ -1845,27 +1870,41 @@ FNewSpline::FNewSpline(const FSplineCurves& Other)
     	Point.TangentIn = Position.ArriveTangent;
     	Point.TangentOut = Position.LeaveTangent;
     	Point.TangentMode = ConvertInterpCurveModeToTangentMode(Position.InterpMode);
+
+    	const uint32 KnotMultiplicity = Idx == 0 || Idx == Other.Position.Points.Num() - 1 ? 4 : 3;
+    	KnotVector.Emplace(Position.InVal, KnotMultiplicity);
     }
 	
 	// initialize spline
 	GetSpline().SetControlPoints(Points);
+	GetSpline().SetKnotVector(KnotVector);
+	
 	ResetRotation();
 	ResetScale();
 	
     // setup rotation/scale attributes
     for (int32 Idx = 0; Idx < GetSpline().GetNumPoints(); Idx++)
     {
-		const FQuat& Rotation = Other.Rotation.Points[Idx].OutVal;
-		const FVector& Scale = Other.Scale.Points[Idx].OutVal;
-    	
-		SetAttributeValue<FQuat>(RotationAttrName, Rotation, Idx);
-		SetAttributeValue<FVector>(ScaleAttrName, Scale, Idx);
+		if (Other.Rotation.Points.IsValidIndex(Idx))
+		{
+			const FQuat& Rotation = Other.Rotation.Points[Idx].OutVal;
+			SetAttributeValue<FQuat>(RotationAttrName, Rotation, Idx);
+		}
+
+    	if (Other.Scale.Points.IsValidIndex(Idx))
+    	{
+    		const FVector& Scale = Other.Scale.Points[Idx].OutVal;
+    		SetAttributeValue<FVector>(ScaleAttrName, Scale, Idx);
+    	}
     }
 
 	ValidateRotScale();
 	
 	FSpline::FUpdateSplineParams UpdateParams;
 	UpdateParams.bClosedLoop = Other.Position.bIsLooped;
+	UpdateParams.LoopPosition = Other.Position.Points.Num() > 0
+		? Other.Position.Points.Last().InVal + Other.Position.LoopKeyOffset
+		: 0.f;
 	
     UpdateSpline(UpdateParams);
 }
@@ -1904,17 +1943,13 @@ bool FNewSpline::Serialize(FArchive& Ar)
     return true;
 }
 
-FInterval1f FNewSpline::GetSegmentSpace() const
-{
-	return GetNumberOfSegments() == 0 ? FInterval1f::Empty() : FInterval1f(0.f, static_cast<float>(GetNumberOfSegments()));
-}
-
 float FNewSpline::GetSegmentLength(const int32 Index, const float Param, const FVector& Scale3D) const
 {
 	const int32 NumPoints = GetSpline().GetNumPoints();
 	const int32 LastPoint = NumPoints - 1;
+	const int32 LastSegment = GetNumberOfSegments() - 1;
 
-	check(Index >= 0 && ((IsClosedLoop() && Index < NumPoints) || (!IsClosedLoop() && Index < LastPoint)));
+	check(Index >= 0 && Index <= LastSegment);
 	check(Param >= 0.0f && Param <= 1.0f);
 
 	// Evaluate the length of a Hermite spline segment.
@@ -1938,10 +1973,12 @@ float FNewSpline::GetSegmentLength(const int32 Index, const float Param, const F
 		{ 0.90617985f, 0.23692688f }
 	};
 
+	const float TangentScaleFactor = GetSegmentParameterRange(Index).Length();
+	
 	const auto& P0 = GetSpline().GetValue(Index);
-	const auto& T0 = GetSpline().GetTangentOut(Index);
+	const auto& T0 = GetSpline().GetTangentOut(Index) * TangentScaleFactor;
 	const auto& P1 = GetSpline().GetValue(Index == LastPoint ? 0 : Index + 1);
-	const auto& T1 = GetSpline().GetTangentOut(Index == LastPoint ? 0 : Index + 1);
+	const auto& T1 = GetSpline().GetTangentOut(Index == LastPoint ? 0 : Index + 1) * TangentScaleFactor;
 
 	// Special cases for linear
 	if (GetSpline().GetTangentModes()[Index] == UE::Geometry::Spline::ETangentMode::Linear)
@@ -1972,21 +2009,47 @@ float FNewSpline::GetSegmentLength(const int32 Index, const float Param, const F
 float FNewSpline::GetSplineLength()
 {	
 	// Evaluate the ParamToLength channel at the very end of the spline
-	const float MaxParameter = static_cast<float>(GetNumberOfSegments());
-    return GetDistanceAtParameter(MaxParameter);
+    return GetDistanceAtParameter(GetParameterSpace().Max);
 }
 
 void FNewSpline::SetClosedLoop(bool bInClosedLoop)
 {
-    GetSpline().SetClosedLoop(bInClosedLoop);
+	if (!IsClosedLoop() && GetSpline().GetNumPoints() <= 1)
+	{
+		// Can't close splines with 1 point
+		return;
+	}
+
+	bool bShouldActuallyClose = bInClosedLoop;
+
+	// To mimic FSplineCurves properly, we need to ignore closing requests if the override loop parameter is not greater than our final point's parameter.
+	if (bInClosedLoop && LoopParameter.IsSet())
+	{
+		if (LoopParameter.GetValue() <= GetSpline().GetParameter(GetSpline().GetNumPoints() - 1))
+		{
+			bShouldActuallyClose = false;
+		}
+	}
+	
+    GetSpline().SetClosedLoop(bShouldActuallyClose);
+
+	if (bShouldActuallyClose)
+	{
+		// Do not subtract 1, GetNumPoints() ignores the closing point, so we add 1 back.
+		const int32 ClosingPointIdx = GetSpline().GetNumPoints();
+		const float LoopKeyParameter = LoopParameter.IsSet() ? LoopParameter.GetValue() : GetSpline().GetParameter(ClosingPointIdx - 1) + 1.f;
+		GetSpline().SetParameter(ClosingPointIdx, LoopKeyParameter);
+	}
 	
 	UE::Geometry::Spline::TTangentBezierSpline<FVector>* ScaleChild = GetTypedAttributeChannel<UE::Geometry::Spline::TTangentBezierSpline<FVector>>(ScaleAttrName);
-	ScaleChild->SetClosedLoop(bInClosedLoop);
+	ScaleChild->SetClosedLoop(bShouldActuallyClose);
 	ScaleChild->SetKnotVector(GetSpline().GetKnotVector());
 		
 	UE::Geometry::Spline::TTangentBezierSpline<FQuat>* RotChild = GetTypedAttributeChannel<UE::Geometry::Spline::TTangentBezierSpline<FQuat>>(RotationAttrName);
-	RotChild->SetClosedLoop(bInClosedLoop);
+	RotChild->SetClosedLoop(bShouldActuallyClose);
 	RotChild->SetKnotVector(GetSpline().GetKnotVector());
+
+	MarkLegacyCurvesDirty();
 	
     UpdateSpline();
 }
@@ -2030,6 +2093,28 @@ void FNewSpline::ResetScale()
 
 void FNewSpline::AddPoint(const FSplinePoint& Point)
 {
+	auto UpperBound = [this](float Value)
+	{
+		int32 Count = GetSpline().GetNumPoints();
+		int32 First = 0;
+
+		while (Count > 0)
+		{
+			const int32 Middle = Count / 2;
+			if (Value >= GetSpline().GetParameter(First + Middle))
+			{
+				First += Middle + 1;
+				Count -= Middle + 1;
+			}
+			else
+			{
+				Count = Middle;
+			}
+		}
+
+		return First;
+	};
+	
 	ValidateRotScale();
 	ON_SCOPE_EXIT
 	{
@@ -2041,14 +2126,25 @@ void FNewSpline::AddPoint(const FSplinePoint& Point)
 	FTangentBezierControlPoint ControlPoint = ConvertToTangentBezierControlPoint(Point);
 	const int32 NumPoints = GetSpline().GetNumPoints();
     
-	// Convert InputKey to parameter value
-	int32 Index = FMath::FloorToInt(Point.InputKey);
-	float Fraction = Point.InputKey - Index;
+	// This index is the endpoint of the segment we are going to insert on, not the segment for insertion
+	const int32 Index = UpperBound(Point.InputKey);
 	
 	if (NumPoints == 0 || Index >= NumPoints)
 	{
 		GetSpline().AppendPoint(ControlPoint);
 
+		if (NumPoints == 0)
+		{
+			FirstPointParameter = Point.InputKey;
+		}
+		else if (NumPoints == 1)
+		{
+			GetSpline().SetParameter(0, FirstPointParameter.Get(0.f));
+			FirstPointParameter.Reset();
+		}
+
+		GetSpline().SetParameter(NumPoints, Point.InputKey);
+		
 		UE::Geometry::Spline::TTangentBezierSpline<FVector>* ScaleChild = GetTypedAttributeChannel<UE::Geometry::Spline::TTangentBezierSpline<FVector>>(ScaleAttrName);
 		ScaleChild->AppendPoint(UE::Geometry::Spline::TTangentBezierSpline<FVector>::FTangentBezierControlPoint(FVector::OneVector));
 		ScaleChild->SetKnotVector(GetSpline().GetKnotVector());
@@ -2063,7 +2159,15 @@ void FNewSpline::AddPoint(const FSplinePoint& Point)
 	else if (Index == 0)
     {
     	GetSpline().PrependPoint(ControlPoint);
-    	
+
+		if (NumPoints == 1)
+		{
+			GetSpline().SetParameter(1, FirstPointParameter.Get(0.f));
+			FirstPointParameter.Reset();
+		}
+		
+		GetSpline().SetParameter(0, Point.InputKey);
+		
     	UE::Geometry::Spline::TTangentBezierSpline<FVector>* ScaleChild = GetTypedAttributeChannel<UE::Geometry::Spline::TTangentBezierSpline<FVector>>(ScaleAttrName);
     	ScaleChild->PrependPoint(UE::Geometry::Spline::TTangentBezierSpline<FVector>::FTangentBezierControlPoint(FVector::OneVector));
     	ScaleChild->SetKnotVector(GetSpline().GetKnotVector());
@@ -2076,16 +2180,36 @@ void FNewSpline::AddPoint(const FSplinePoint& Point)
     	return;
     }
 	
-	// Insert at calculated parameter
-	int32 InsertedIndex = GetSpline().InsertPointAtSegmentParam(Index, Fraction, ControlPoint);
-	float NewPointInternalParam = GetSpline().GetParameter(InsertedIndex);
+	// Should never be set, if it is set then we should have hit one of the above branches
+	ensureAlways(!FirstPointParameter.IsSet());
 	
-	float ScaleParam = MapParameterToChildSpace(ScaleAttrName, NewPointInternalParam);
+	// here we can assume Index is a valid point and not the first point (i.e. in [1, NumPoints - 1]).
+	// UpperBound yields the endpoint of the segment we should insert on, so we decrement to get the index of the segment to insert on
+	const FVector SegmentStartOutTangent = GetOutTangent(Index - 1);
+	const FVector SegmentEndInTangent = GetInTangent(Index);
+	const float InsertOnSegmentIdx = Index - 1;
+	
+	// Insert at calculated segment and segment T, preserve user specified parameter.
+	const float SegmentT = GetSpline().GetSegmentParameterRange(InsertOnSegmentIdx).GetT(Point.InputKey);
+	const int32 InsertedIndex = GetSpline().InsertPointAtSegmentParam(InsertOnSegmentIdx, SegmentT, ControlPoint);
+	GetSpline().SetParameter(InsertedIndex, Point.InputKey);
+	
+	if (GetSplinePointType(InsertedIndex - 1) == CIM_CurveUser)
+	{
+		SetOutTangent(InsertedIndex - 1, SegmentStartOutTangent);
+	}
+	
+	if (GetSplinePointType(InsertedIndex + 1) == CIM_CurveUser)
+	{
+		SetInTangent(InsertedIndex + 1, SegmentEndInTangent);
+	}
+	
+	float ScaleParam = MapParameterToChildSpace(ScaleAttrName, Point.InputKey);
 	UE::Geometry::Spline::TTangentBezierSpline<FVector>* ScaleChild = GetTypedAttributeChannel<UE::Geometry::Spline::TTangentBezierSpline<FVector>>(ScaleAttrName);
 	ScaleChild->InsertPointAtGlobalParam(ScaleParam, UE::Geometry::Spline::TTangentBezierSpline<FVector>::FTangentBezierControlPoint(FVector::OneVector));
 	ScaleChild->SetKnotVector(GetSpline().GetKnotVector());
 	
-	float RotParam = MapParameterToChildSpace(RotationAttrName, NewPointInternalParam);
+	float RotParam = MapParameterToChildSpace(RotationAttrName, Point.InputKey);
 	UE::Geometry::Spline::TTangentBezierSpline<FQuat>* RotChild = GetTypedAttributeChannel<UE::Geometry::Spline::TTangentBezierSpline<FQuat>>(RotationAttrName);
 	RotChild->InsertPointAtGlobalParam(RotParam, UE::Geometry::Spline::TTangentBezierSpline<FQuat>::FTangentBezierControlPoint(FQuat::Identity));
 	RotChild->SetKnotVector(GetSpline().GetKnotVector());
@@ -2110,6 +2234,18 @@ void FNewSpline::InsertPoint(const FSplinePoint& Point, int32 Index)
 	if (NumPoints == 0 || Index >= NumPoints)
 	{
 		GetSpline().AppendPoint(ControlPoint);
+
+		if (NumPoints == 0)
+		{
+			FirstPointParameter = Point.InputKey;
+		}
+		else if (NumPoints == 1)
+		{
+			GetSpline().SetParameter(0, FirstPointParameter.Get(0.f));
+			FirstPointParameter.Reset();
+		}
+
+		GetSpline().SetParameter(NumPoints, Point.InputKey);
 		
 		UE::Geometry::Spline::TTangentBezierSpline<FVector>* ScaleChild = GetTypedAttributeChannel<UE::Geometry::Spline::TTangentBezierSpline<FVector>>(ScaleAttrName);
 		ScaleChild->AppendPoint(UE::Geometry::Spline::TTangentBezierSpline<FVector>::FTangentBezierControlPoint(FVector::OneVector));
@@ -2126,6 +2262,13 @@ void FNewSpline::InsertPoint(const FSplinePoint& Point, int32 Index)
 	{
 		GetSpline().PrependPoint(ControlPoint);
 		
+		if (NumPoints == 1)
+		{
+			GetSpline().SetParameter(1, FirstPointParameter.Get(0.f));
+			FirstPointParameter.Reset();
+		}
+		
+		GetSpline().SetParameter(0, Point.InputKey);
 		UE::Geometry::Spline::TTangentBezierSpline<FVector>* ScaleChild = GetTypedAttributeChannel<UE::Geometry::Spline::TTangentBezierSpline<FVector>>(ScaleAttrName);
 		ScaleChild->PrependPoint(UE::Geometry::Spline::TTangentBezierSpline<FVector>::FTangentBezierControlPoint(FVector::OneVector));
 		ScaleChild->SetKnotVector(GetSpline().GetKnotVector());
@@ -2138,17 +2281,31 @@ void FNewSpline::InsertPoint(const FSplinePoint& Point, int32 Index)
 		return;
 	}
 	
-	Index = FMath::Clamp(Index, 0, GetSpline().GetNumPoints());
+	// Clamp with a Min of 1 because we handle the 0 case above.
+	Index = FMath::Clamp(Index, 1, GetSpline().GetNumPoints());
+
+	const FVector SegmentStartOutTangent = GetOutTangent(Index - 1);
+	const FVector SegmentEndInTangent = GetInTangent(Index);
+	
+	const int32 InsertedIndex = GetSpline().InsertPointAtPosition(Index, ControlPoint);
+	GetSpline().SetParameter(InsertedIndex, Point.InputKey);
+	
+	if (GetSplinePointType(Index - 1) == CIM_CurveUser)
+	{
+		SetOutTangent(Index - 1, SegmentStartOutTangent);
+	}
 	
-	int32 InsertedIndex = GetSpline().InsertPointAtPosition(Index, ControlPoint);
-	float NewPointInternalParam = GetSpline().GetParameter(InsertedIndex);
+	if (GetSplinePointType(Index + 1) == CIM_CurveUser)
+	{
+		SetInTangent(Index + 1, SegmentEndInTangent);
+	}
 	
-	float ScaleParam = MapParameterToChildSpace(ScaleAttrName, NewPointInternalParam);
+	float ScaleParam = MapParameterToChildSpace(ScaleAttrName, Point.InputKey);
 	UE::Geometry::Spline::TTangentBezierSpline<FVector>* ScaleChild = GetTypedAttributeChannel<UE::Geometry::Spline::TTangentBezierSpline<FVector>>(ScaleAttrName);
 	ScaleChild->InsertPointAtGlobalParam(ScaleParam, UE::Geometry::Spline::TTangentBezierSpline<FVector>::FTangentBezierControlPoint(FVector::OneVector));
 	ScaleChild->SetKnotVector(GetSpline().GetKnotVector());
 	
-	float RotParam = MapParameterToChildSpace(RotationAttrName, NewPointInternalParam);
+	float RotParam = MapParameterToChildSpace(RotationAttrName, Point.InputKey);
 	UE::Geometry::Spline::TTangentBezierSpline<FQuat>* RotChild = GetTypedAttributeChannel<UE::Geometry::Spline::TTangentBezierSpline<FQuat>>(RotationAttrName);
 	RotChild->InsertPointAtGlobalParam(RotParam, UE::Geometry::Spline::TTangentBezierSpline<FQuat>::FTangentBezierControlPoint(FQuat::Identity));
 	RotChild->SetKnotVector(GetSpline().GetKnotVector());
@@ -2193,6 +2350,13 @@ void FNewSpline::RemovePoint(int32 Index)
 	{
 		return;
 	}
+
+	if (GetSpline().GetNumPoints() == 2)
+	{
+		FirstPointParameter = Index == 0
+			? GetSpline().GetParameter(1)
+			: GetSpline().GetParameter(0);
+	}
 	
 	GetSpline().RemovePoint(Index);
 
@@ -2349,14 +2513,48 @@ EInterpCurveMode FNewSpline::GetSplinePointType(int32 Index) const
     return ConvertTangentModeToInterpCurveMode(TangentMode);
 }
 
+void FNewSpline::SetParameterAtIndex(int32 Index, float Parameter)
+{
+	const int32 NumPoints = GetSpline().GetNumPoints();
+	
+	if (Index < 0 || Index >= NumPoints)
+	{
+		return;
+	}
+
+	UE::Geometry::Spline::TTangentBezierSpline<FQuat>* RotSpline = GetTypedAttributeChannel<UE::Geometry::Spline::TTangentBezierSpline<FQuat>>(RotationAttrName);
+	UE::Geometry::Spline::TTangentBezierSpline<FVector>* ScaleSpline = GetTypedAttributeChannel<UE::Geometry::Spline::TTangentBezierSpline<FVector>>(ScaleAttrName);
+	const float NewParameter = Parameter;
+
+	if (NumPoints == 1)
+	{
+		FirstPointParameter = NewParameter;
+	}
+	else
+	{
+		GetSpline().SetParameter(Index, NewParameter);
+		RotSpline->SetParameter(Index, NewParameter);
+		ScaleSpline->SetParameter(Index, NewParameter);
+	}
+}
+
 float FNewSpline::GetParameterAtIndex(int32 Index) const
 {
-    if (Index < 0 || (!IsClosedLoop() && Index >= GetSpline().GetNumPoints()) || (IsClosedLoop() && Index > GetSpline().GetNumPoints()))
+	const int32 NumPoints = GetSpline().GetNumPoints();
+	
+    if (Index < 0 || (!IsClosedLoop() && Index >= NumPoints) || (IsClosedLoop() && Index > NumPoints))
     {
     	return 0.f;
     }
 
-	return FromInternalSplineSpace(GetSpline().GetParameter(Index));
+	if (NumPoints == 1)
+	{
+		return FirstPointParameter.Get(0.f);
+	}
+	else
+	{
+		return GetSpline().GetParameter(Index);
+	}
 }
 
 float FNewSpline::GetParameterAtDistance(float Distance)
@@ -2370,7 +2568,7 @@ float FNewSpline::GetParameterAtDistance(float Distance)
 
 float FNewSpline::GetDistanceAtParameter(float Parameter)
 {
-	if (!GetSegmentSpace().Contains(Parameter))
+	if (!GetParameterSpace().Contains(Parameter))
 	{
 		return 0.f;
 	}
@@ -2379,14 +2577,70 @@ float FNewSpline::GetDistanceAtParameter(float Parameter)
 	
 	FReadScopeLock Lock(ReparamTableRWLock);
 
-	const float ParameterMax = static_cast<float>(GetNumberOfSegments());
-	const float Key = (Parameter / ParameterMax) * (ReparamTable.Points.Num() - 1);
-	const int32 LowerKey = FMath::FloorToInt32(Key);
-	ensureAlways(LowerKey >= 0 && LowerKey < ReparamTable.Points.Num());
-	const int32 UpperKey = FMath::CeilToInt32(Key);
-	ensureAlways(UpperKey >= 0 && UpperKey < ReparamTable.Points.Num());
-	const float Alpha = FMath::Frac(Key);
-	return FMath::Lerp(ReparamTable.Points[LowerKey].InVal, ReparamTable.Points[UpperKey].InVal, Alpha);
+	if (ReparamTable.Points.IsEmpty())
+	{
+		return 0.f;
+	}
+	
+	if (Parameter <= ReparamTable.Points[0].OutVal)
+	{
+		return ReparamTable.Points[0].InVal;
+	}
+	
+	if (Parameter >= ReparamTable.Points.Last().OutVal)
+	{
+		return ReparamTable.Points.Last().InVal;
+	}
+
+	const int32 NumReparamPoints = ReparamTable.Points.Num();
+	int32 Low = 0;
+	int32 High = NumReparamPoints - 1;
+
+	while (High - Low > 1)
+	{
+		const int32 Mid = (Low + High) / 2;
+		
+		if (ReparamTable.Points[Mid].OutVal < Parameter)
+		{
+			Low = Mid;
+		}
+		else
+		{
+			High = Mid;
+		}
+	}
+	
+	const float Alpha = FInterval1f(ReparamTable.Points[Low].OutVal, ReparamTable.Points[High].OutVal).GetT(Parameter);
+	return FMath::Lerp(ReparamTable.Points[Low].InVal, ReparamTable.Points[High].InVal, Alpha);
+}
+
+int32 FNewSpline::GetIndexAtParameter(float Parameter) const
+{
+	if (!GetParameterSpace().Contains(Parameter))
+	{
+		return INDEX_NONE;
+	}
+	
+	const int32 NumSegments = GetSpline().GetNumberOfSegments();
+
+	int32 MinIndex = 0;
+	int32 MaxIndex = NumSegments;
+
+	while (MaxIndex - MinIndex > 1)
+	{
+		int32 MidIndex = (MinIndex + MaxIndex) / 2;
+
+		if (GetSpline().GetParameter(MidIndex) <= Parameter)
+		{
+			MinIndex = MidIndex;
+		}
+		else
+		{
+			MaxIndex = MidIndex;
+		}
+	}
+
+	return MinIndex;
 }
 
 float FNewSpline::FindNearest(const FVector& InLocation, float& OutSquaredDistance) const
@@ -2397,7 +2651,7 @@ float FNewSpline::FindNearest(const FVector& InLocation, float& OutSquaredDistan
     }
     else
     {
-	    return FromInternalSplineSpace(GetSpline().FindNearest(InLocation, OutSquaredDistance));
+	    return GetSpline().FindNearest(InLocation, OutSquaredDistance);
     }
 }
 
@@ -2414,47 +2668,43 @@ float FNewSpline::FindNearestOnSegment(const FVector& InLocation, int32 SegmentI
     }
     else
     {
-	   return FromInternalSplineSpace(GetSpline().FindNearestOnSegment(InLocation, SegmentIndex, OutSquaredDistance));
+	   return GetSpline().FindNearestOnSegment(InLocation, SegmentIndex, OutSquaredDistance);
     }
 }
 
 FVector FNewSpline::EvaluatePosition(float Parameter) const
 {
-	return Evaluate(ToInternalSplineSpace(Parameter));
+	return Evaluate(Parameter);
 }
 
 FVector FNewSpline::EvaluateDerivative(float Parameter) const
 {
-	Parameter = ToInternalSplineSpace(Parameter);
     return GetSpline().GetTangent(Parameter);
 }
 
 FQuat FNewSpline::EvaluateRotation(float Parameter) const
 {
-    return EvaluateAttribute<FQuat>(RotationAttrName, ToInternalSplineSpace(Parameter));
+    return EvaluateAttribute<FQuat>(RotationAttrName, Parameter);
 }
 
 FVector FNewSpline::EvaluateScale(float Parameter) const
 {
-    return EvaluateAttribute<FVector>(ScaleAttrName, ToInternalSplineSpace(Parameter));
+    return EvaluateAttribute<FVector>(ScaleAttrName, Parameter);
 }
 
 void FNewSpline::UpdateSpline()
 {
     GetSpline().UpdateTangents();			// updates the tangents on our version of the points
 	
-	// We do this so that parameterization remains proportional to the square root of segment chord lengths.
-	// todo: reenable this code when underlying knot vector issues are resolved.
-	// GetSpline().Reparameterize(EParameterizationPolicy::Centripetal);
-	
     MarkReparamTableDirty();
     RebuildLegacyCurves();					// updates legacy curves based on points, never evaluates internal spline.
 }
 
 void FNewSpline::UpdateSpline(const FSpline::FUpdateSplineParams& Params)
 {
-	// todo: handled unhandled params
+	// todo: use Scale param
     
+	LoopParameter = Params.bLoopPositionOverride ? Params.LoopPosition : TOptional<float>();
     SetClosedLoop(Params.bClosedLoop);
 	GetSpline().SetStationaryEndpoints(Params.bStationaryEndpoints);
     ReparamStepsPerSegment = Params.ReparamStepsPerSegment;
@@ -2462,31 +2712,6 @@ void FNewSpline::UpdateSpline(const FSpline::FUpdateSplineParams& Params)
     UpdateSpline();
 }
 
-float FNewSpline::FromNormalizedSpace(float Parameter) const
-{
-	return GetNumberOfSegments() * Parameter;
-}
-
-float FNewSpline::ToNormalizedSpace(float Parameter) const
-{
-	return Parameter / GetNumberOfSegments();
-}
-
-float FNewSpline::FromInternalSplineSpace(float Parameter) const
-{
-	float OutLocalParam;
-	const int32 SegmentIndex = GetSpline().FindSegmentIndex(Parameter, OutLocalParam);
-	return static_cast<float>(SegmentIndex) + OutLocalParam;
-}
-float FNewSpline::ToInternalSplineSpace(float Parameter) const
-{
-	const int32 NumSegments = GetNumberOfSegments();
-	float ClampedParam = FMath::Clamp(Parameter, 0.0f, static_cast<float>(NumSegments));
-	const int32 SegmentIndex = FMath::Min(FMath::FloorToInt(ClampedParam), NumSegments - 1);
-	const float LocalT = ClampedParam - static_cast<float>(SegmentIndex); // Parameter [0,1] within segment
-	return GetSegmentParameterRange(SegmentIndex).Interpolate(LocalT);
-}
-
 void FNewSpline::MarkReparamTableDirty()
 {
 	ReparamTableNextVersion.Increment();
@@ -2520,16 +2745,18 @@ void FNewSpline::UpdateReparamTable()
 	float AccumulatedLength = 0.0f;
 	for (int32 SegmentIndex = 0; SegmentIndex < NumSegments; ++SegmentIndex)
 	{
-		ReparamTable.Points.Emplace(AccumulatedLength, SegmentIndex, 0.0f, 0.0f, CIM_Linear);
+		FInterval1f SegmentParameterRange = GetSegmentParameterRange(SegmentIndex);
+		
+		ReparamTable.Points.Emplace(AccumulatedLength, SegmentParameterRange.Min, 0.0f, 0.0f, CIM_Linear);
 		for (int32 Step = 1; Step < ReparamStepsPerSegment; ++Step)
 		{
-			const float Param = static_cast<float>(Step) / ReparamStepsPerSegment;
-			ReparamTable.Points.Emplace(GetSegmentLength(SegmentIndex, Param) + AccumulatedLength, SegmentIndex + Param, 0.0f, 0.0f, CIM_Linear);
+			const float SegmentT = static_cast<float>(Step) / ReparamStepsPerSegment;
+			ReparamTable.Points.Emplace(GetSegmentLength(SegmentIndex, SegmentT) + AccumulatedLength, SegmentParameterRange.Interpolate(SegmentT), 0.0f, 0.0f, CIM_Linear);
 		}
 
 		AccumulatedLength += GetSegmentLength(SegmentIndex, 1.0f);
 	}
-	ReparamTable.Points.Emplace(AccumulatedLength, static_cast<float>(NumSegments), 0.0f, 0.0f, CIM_Linear);
+	ReparamTable.Points.Emplace(AccumulatedLength, GetParameterSpace().Max, 0.0f, 0.0f, CIM_Linear);
 
 	ReparamTableVersion = CachedNextVersion;
 }
@@ -2561,25 +2788,32 @@ void FNewSpline::RebuildLegacyCurves()
 		return;
 	}
 
+	PositionCurve.ClearLoopKey();
+	RotationCurve.ClearLoopKey();
+	ScaleCurve.ClearLoopKey();
+	
 	PositionCurve.Points.Reset(GetSpline().GetNumPoints());
 	RotationCurve.Points.Reset(GetSpline().GetNumPoints());
 	ScaleCurve.Points.Reset(GetSpline().GetNumPoints());
-
-    PositionCurve.bIsLooped = IsClosedLoop();
-    PositionCurve.LoopKeyOffset = 1.f;
-    RotationCurve.bIsLooped = IsClosedLoop();
-    ScaleCurve.bIsLooped = IsClosedLoop();
 	
     for (int32 i = 0; i < GetSpline().GetNumPoints(); ++i)
     {
-    	const float AttributeParameter = static_cast<float>(i);
+		const float Parameter = GetSpline().GetParameter(i);
     	
     	// Additionally populate interp curves here so we don't double loop
-    	PositionCurve.Points.Emplace(AttributeParameter, GetSpline().GetValue(i), GetSpline().GetTangentIn(i), GetSpline().GetTangentOut(i), ConvertTangentModeToInterpCurveMode(GetSpline().GetTangentModes()[i]));
-    	RotationCurve.AddPoint(AttributeParameter, GetAttributeValue<FQuat>(RotationAttrName, i));
-    	ScaleCurve.AddPoint(AttributeParameter, GetAttributeValue<FVector>(ScaleAttrName, i));
+    	PositionCurve.Points.Emplace(Parameter, GetSpline().GetValue(i), GetSpline().GetTangentIn(i), GetSpline().GetTangentOut(i), ConvertTangentModeToInterpCurveMode(GetSpline().GetTangentModes()[i]));
+    	RotationCurve.AddPoint(Parameter, GetAttributeValue<FQuat>(RotationAttrName, i));
+    	ScaleCurve.AddPoint(Parameter, GetAttributeValue<FVector>(ScaleAttrName, i));
     }
-    
+	
+	if (IsClosedLoop())
+	{
+		const FInterval1f ClosingSegmentInterval = GetSpline().GetSegmentParameterRange(GetNumberOfSegments() - 1);
+		PositionCurve.SetLoopKey(ClosingSegmentInterval.Max);
+		RotationCurve.SetLoopKey(ClosingSegmentInterval.Max);
+		ScaleCurve.SetLoopKey(ClosingSegmentInterval.Max);
+	}
+	
     LegacyCurvesVersion = LegacyCurvesNextVersion.GetValue();
 }
 
@@ -2933,6 +3167,16 @@ EInterpCurveMode FLegacySpline::GetSplinePointType(const int32 Index) const
 	return PositionCurve.Points.IsValidIndex(Index) ? PositionCurve.Points[Index].InterpMode.GetValue() : CIM_Unknown;
 }
 
+void FLegacySpline::SetParameterAtIndex(int32 Index, float Parameter)
+{
+	if (PositionCurve.Points.IsValidIndex(Index))
+	{
+		PositionCurve.Points[Index].InVal = Parameter;
+		RotationCurve.Points[Index].InVal = Parameter;
+		ScaleCurve.Points[Index].InVal = Parameter;
+	}
+}
+
 float FLegacySpline::GetParameterAtIndex(int32 Index) const
 {
 	return PositionCurve.Points.IsValidIndex(Index) ? PositionCurve.Points[Index].InVal : 0.f;
@@ -2959,6 +3203,11 @@ float FLegacySpline::GetDistanceAtParameter(float Parameter) const
 	return Distance;
 }
 
+int32 FLegacySpline::GetIndexAtParameter(float Parameter) const
+{
+	return PositionCurve.GetPointIndexForInputValue(Parameter);
+}
+
 FVector FLegacySpline::Evaluate(float Param) const
 {
 	return PositionCurve.Eval(Param);
diff --git a/Engine/Source/Runtime/GeometryCore/Public/Splines/BSpline.h b/Engine/Source/Runtime/GeometryCore/Public/Splines/BSpline.h
--- a/Engine/Source/Runtime/GeometryCore/Public/Splines/BSpline.h
+++ b/Engine/Source/Runtime/GeometryCore/Public/Splines/BSpline.h
@@ -126,7 +126,8 @@ public:
 
 	static constexpr int32 WindowSize = Degree + 1;
 	using FWindow = TStaticArray<const ValueType*, WindowSize>;
-
+	using FWindowStore = TStaticArray<ValueType, WindowSize>;
+	
 	DECLARE_SPLINE_TYPE_ID(
 		BSplineNameSelector<DEGREE>::Name,
 		*TSplineValueTypeTraits<VALUETYPE>::Name
@@ -275,7 +276,9 @@ public:
 		if (NumValues == 0) return ValueType();
 		if (NumValues == 1) return GetValue(0);
 
-		FWindow Window = FindWindow(Parameter);
+		FWindow Window;
+		FWindowStore Scratch;
+		FindWindow(Parameter, Window, Scratch);
 
 		// If FindWindow fails, it will return an array of nullptr.
 		if (Window[0] == nullptr) return ValueType();
@@ -1188,12 +1191,12 @@ private:
 		return Degree;
 	}
 
-	virtual FWindow FindWindow(float Parameter) const
+	virtual void FindWindow(float Parameter, FWindow& OutWindow, FWindowStore& Scratch) const
 	{
-		FWindow Window = {};
+		OutWindow = {};
 		if (NumKeys() < Degree + 1)
 		{
-			return Window;
+			return;
 		}
 
 		int32 Span = FindKnotSpan(Parameter);
@@ -1208,7 +1211,7 @@ private:
 			{
 				int32 Index = (Span - Degree + i) % NumPoints;
 				if (Index < 0) Index += NumPoints;
-				Window[i] = &GetValue(Index);
+				OutWindow[i] = &GetValue(Index);
 			}
 		}
 		else
@@ -1219,12 +1222,10 @@ private:
 				const int32 Index = Span - Degree + i;
 				if (Index >= 0 && Index < NumKeys())
 				{
-					Window[i] = &GetValue(Index);
+					OutWindow[i] = &GetValue(Index);
 				}
 			}
 		}
-
-		return Window;
 	}
 
 	/**
diff --git a/Engine/Source/Runtime/GeometryCore/Public/Splines/PolyBezierSpline.h b/Engine/Source/Runtime/GeometryCore/Public/Splines/PolyBezierSpline.h
--- a/Engine/Source/Runtime/GeometryCore/Public/Splines/PolyBezierSpline.h
+++ b/Engine/Source/Runtime/GeometryCore/Public/Splines/PolyBezierSpline.h
@@ -23,6 +23,11 @@ public:
 	using Base = TBSpline<VALUETYPE, 3>;
 	typedef typename Base::ValueType ValueType;
 	using FWindow = typename Base::FWindow;
+	using FWindowStore = typename Base::FWindowStore;
+	using FFindWindowOverride = void(*)(FWindow& /*InOut*/,
+										int32 /*SegmentIndex*/,
+										const FInterval1f& /*SegmentRange*/,
+										FWindowStore& /*Scratch*/);
 
 	// Generate compile-time type ID for PolyBezier
 	DECLARE_SPLINE_TYPE_ID(
@@ -339,10 +344,16 @@ public:
     template<int32 Order>
     ValueType EvaluateDerivative(float Parameter) const
     {
-		FWindow Window = FindWindow(Parameter);
+		FWindow Window;
+		FWindowStore Scratch;
+		FindWindow(Parameter, Window, Scratch);
 
 		// If FindWindow fails, it will return an array of nullptr. InterpolateWindow assumes validity of elements.
-		if (Window[0] == nullptr) { return ValueType(); }
+		if (!Math::FSplineValidation::IsValidWindow<ValueType, Base::WindowSize>(Window) || // Cubic Bezier needs 4 points
+				!Math::FSplineValidation::IsValidParameter(Parameter))
+		{
+			return Math::FSplineValidation::GetDefaultValue<ValueType>();
+		}
 		
         // Order 0 is just regular evaluation
 	    if constexpr (Order == 0)
@@ -356,10 +367,10 @@ public:
     	MapGlobalParameterToLocalSegment(Parameter, SegmentIndex, LocalT);
 		
 	    // For PolyBezier, each segment is normalized to unit parameter space
-        constexpr float SegmentScale = 1.0f;
+		const float SegmentScale = GetSegmentParameterRange(SegmentIndex).Length();
 	    
 	    // Compute the derivative with proper scaling
-	    return Math::TBezierDerivativeCalculator<ValueType, Order>::Compute(
+	    return Math::TBezierDerivativeCalculator<ValueType, Order, Base::WindowSize>::Compute(
 	        Window, LocalT, SegmentScale);
     }
 
@@ -738,7 +749,10 @@ public:
 	    }
 	    
 	    SegmentIndex = FMath::Clamp(SegmentIndex, 0, NumSegments - 1);
-	    
+
+		// todo: reconsider this approach, a bit of a hack but it works for now.
+		LocalT = FMath::Clamp(LocalT, UE_KINDA_SMALL_NUMBER, 1.f - UE_KINDA_SMALL_NUMBER);
+		
 	    // Get the segment's parameter range
 	    FInterval1f SegmentRange = GetSegmentParameterRange(SegmentIndex);
 	    
@@ -951,6 +965,16 @@ public:
 			{
 				Base::PairKnots[0].Multiplicity = Degree + 1;
 			}
+
+			if (const int32 NewNumSegments = GetNumberOfSegments();
+				NewNumSegments > 0 && Base::IsClosedLoop())
+			{
+				const ValueType P1 = Base::GetValue(1);
+				const ValueType P0 = Base::GetValue(0);
+				const ValueType Delta = P0 - P1;
+				UpdateSegmentPoint(NewNumSegments - 1, 2, P0 + Delta);
+				UpdateSegmentPoint(NewNumSegments - 1, 3, P0);
+			}
 		}
 		// else if last point
 		else if (PointIndex == GetNumberOfSegments())
@@ -1506,7 +1530,8 @@ public:
 	{
 		Base::SetCustomKnots(NewKnots);
 	}
-	
+
+	void SetFindWindowOverride(FFindWindowOverride Fn) { FindWindowOverrideFn = Fn; }
 private:
 	
 	/**
@@ -1591,12 +1616,12 @@ private:
 		this->MarkFlatKnotsCacheDirty();
 	}
 
-    virtual FWindow FindWindow(float Parameter) const override
+	virtual void FindWindow(float Parameter, FWindow& OutWindow, FWindowStore& Scratch) const override
     {
-        FWindow Window = {};
+        OutWindow = {};
         if (Base::NumKeys() < 4)  // Need at least one complete Bezier segment
         {
-            return Window;
+            return;
         }
 		const TArray<ValueType>& Entries = this->Values;
         // Find which segment we're in
@@ -1627,7 +1652,7 @@ private:
 
         	for (int Idx = 0; Idx < Base::WindowSize; ++Idx)
         	{
-        		Window[Idx] = &Entries[WrapIndex(BaseIndex + Idx)];
+        		OutWindow[Idx] = &Entries[WrapIndex(BaseIndex + Idx)];
         	}
         }
         else
@@ -1635,21 +1660,25 @@ private:
         	// Return null window if BaseIndex is too large to fully populate the window.
             if (BaseIndex + 3 >= Entries.Num())
             {
-                return Window;
+                return;
             }
 
         	// Standard case - no wrapping
         	for (int Idx = 0; Idx < Base::WindowSize; ++Idx)
         	{
-        		Window[Idx] = &Entries[BaseIndex + Idx];
+        		OutWindow[Idx] = &Entries[BaseIndex + Idx];
         	}
         }
 
-        return Window;
+		if (FindWindowOverrideFn && Math::FSplineValidation::IsValidWindow<ValueType, Base::WindowSize>(OutWindow))
+		{
+			const FInterval1f SegmentRange = GetSegmentParameterRange(SegmentIndex);
+			FindWindowOverrideFn(OutWindow, SegmentIndex, SegmentRange, Scratch);
+		}
     }
-
+	
     virtual ValueType InterpolateWindow(TArrayView<const ValueType* const> Window, float Parameter) const override 
-    {		
+    {
         if (Window.Num() < 4)
         {
             return ValueType();
@@ -1937,9 +1966,13 @@ private:
 	}
 
 private:
-
+	
 	/** This is the last segment fetched by parameter. Used to accelerate predictable sampling patterns. */
 	mutable int32 LastEvaluatedSegment = INDEX_NONE;
+	
+	/** Optional override for FindWindow*/
+	FFindWindowOverride FindWindowOverrideFn = nullptr;
+
 };
 
 
diff --git a/Engine/Source/Runtime/GeometryCore/Public/Splines/SplineMath.h b/Engine/Source/Runtime/GeometryCore/Public/Splines/SplineMath.h
--- a/Engine/Source/Runtime/GeometryCore/Public/Splines/SplineMath.h
+++ b/Engine/Source/Runtime/GeometryCore/Public/Splines/SplineMath.h
@@ -269,23 +269,16 @@ namespace Math
 			}
 		}
 
-		template <typename T>
-		static bool IsValidWindow(TArrayView<const T* const> Window, int32 RequiredSize)
+		template <typename T, uint32 N>
+		static bool IsValidWindow(const TStaticArray<const T*, N>& Window)
 		{
-			if (Window.Num() < RequiredSize)
+			for (int32 i = 0; i < N; ++i)
 			{
-				return false;
-			}
-
-			// Check for null pointers
-			for (const T* Ptr : Window)
-			{
-				if (!Ptr)
+				if (Window[i] == nullptr)
 				{
 					return false;
 				}
 			}
-
 			return true;
 		}
 
@@ -457,10 +450,10 @@ namespace Math
 	template <typename T, int32 Order>
 	struct TGenericDerivativeHelper
 	{
-		static T Compute(TArrayView<const T* const> Window, float Parameter)
+		static T Compute(TStaticArray<const T* const, Order+1> Window, float Parameter)
 		{
 			// Basic validation
-			if (!FSplineValidation::IsValidWindow(Window, Order + 1))
+			if (!FSplineValidation::IsValidWindow<T, Order+1>(Window))
 			{
 				return FSplineValidation::GetDefaultValue<T>();
 			}
@@ -696,17 +689,17 @@ namespace Math
 
 
 	// Bezier curve specific calculations
-	template <typename T, int32 Order>
+	template <typename T, int32 Order, int32 WindowSize>
 	struct TBezierDerivativeCalculator
 	{
-		static T Compute(TArrayView<const T* const> Window, float Parameter, float SegmentScale)
+		static T Compute(TStaticArray<const T*, WindowSize> Window, float Parameter, float SegmentScale)
 		{
 			if constexpr (Order > 3)
 			{
 				return FSplineValidation::GetDefaultValue<T>();
 			}
 
-			if (!FSplineValidation::IsValidWindow(Window, 4) || // Cubic Bezier needs 4 points
+			if (!FSplineValidation::IsValidWindow<T, WindowSize>(Window) || // Cubic Bezier needs 4 points
 				!FSplineValidation::IsValidParameter(Parameter))
 			{
 				return FSplineValidation::GetDefaultValue<T>();
@@ -728,6 +721,12 @@ namespace Math
 			// derivatives scale inversely with parameter range
 			if constexpr (Order > 0)
 			{
+				// If L == 0, the segment is degenerate in parameter space
+				// return the zero derivative (well-defined limit)
+				if (SegmentScale == 0.0f)
+				{
+					return FSplineValidation::GetDefaultValue<T>();
+				}
 				// Calculate scaling factor: 1 / SegmentScale^Order
 				float ScalingFactor = 1.0f;
 				for (int32 i = 0; i < Order; ++i)
diff --git a/Engine/Source/Runtime/GeometryCore/Public/Splines/TangentBezierSpline.h b/Engine/Source/Runtime/GeometryCore/Public/Splines/TangentBezierSpline.h
--- a/Engine/Source/Runtime/GeometryCore/Public/Splines/TangentBezierSpline.h
+++ b/Engine/Source/Runtime/GeometryCore/Public/Splines/TangentBezierSpline.h
@@ -116,7 +116,6 @@ struct TTangentBezierControlPoint
 			&& TangentMode == Other.TangentMode;
 	}
 };
-	
 /**
  * A spline that provides tangent-based control over curve shape while using 
  * piecewise Bezier curves internally for evaluation. Supports both manual
@@ -129,8 +128,14 @@ class TTangentBezierSpline :
 {
 public:
 	typedef typename TSplineWrapper<TPolyBezierSpline<VALUETYPE>>::ValueType ValueType;
-	using TSplineWrapper<TPolyBezierSpline<VALUETYPE>>::InternalSpline;
+	
+	using TSplineWrapper<TPolyBezierSpline<ValueType>>::InternalSpline;
+	
 	using FTangentBezierControlPoint = TTangentBezierControlPoint<ValueType>;
+	using FPolyBezierSpline = TPolyBezierSpline<ValueType>;
+	
+	using FWindow      = typename FPolyBezierSpline::FWindow;
+	using FWindowStore = typename FPolyBezierSpline::FWindowStore;
 
 	// Generate compile-time type ID for TangentBezier
 	DECLARE_SPLINE_TYPE_ID(
@@ -138,15 +143,20 @@ public:
 		*TSplineValueTypeTraits<VALUETYPE>::Name
 	);
 	
-	TTangentBezierSpline() = default;
+	TTangentBezierSpline()
+	{
+		InstallFindWindowOverride();
+	}
+	
 	virtual ~TTangentBezierSpline() override = default;
 	
 	/** Default constructor with at least one segment */
 	TTangentBezierSpline(const ValueType& StartPoint, const ValueType& EndPoint)
-		: Tension(0.0f)
-		, TangentModes({ETangentMode::Auto, ETangentMode::Auto})
+		: TTangentBezierSpline()
 	{
-		InternalSpline = FPolyBezierSpline3d::CreateLine(StartPoint, EndPoint);
+		Tension = 0.f;
+		TangentModes = {ETangentMode::Auto, ETangentMode::Auto};
+		AssignInternalSpline(FPolyBezierSpline3d::CreateLine(StartPoint, EndPoint));
 	}
 
 	/**
@@ -158,22 +168,24 @@ public:
 		const ValueType& StartTangent,
 		const ValueType& EndTangent,
 		bool bAutoTangents = false)
-		: Tension(0.0f)
-		, TangentModes({bAutoTangents ? ETangentMode::Auto : ETangentMode::User,
-							bAutoTangents ? ETangentMode::Auto : ETangentMode::User})
+		: TTangentBezierSpline()
 	{
-		InternalSpline = FPolyBezierSpline3d(StartPoint,
+		Tension = 0.f;
+		TangentModes = {bAutoTangents ? ETangentMode::Auto : ETangentMode::User,
+						bAutoTangents ? ETangentMode::Auto : ETangentMode::User};
+		AssignInternalSpline( FPolyBezierSpline3d(StartPoint,
 											StartPoint + StartTangent / 3.0f,
 											EndPoint - (EndTangent / 3.0f),
-											EndPoint);
+											EndPoint));
 	}
    
     /**  Copy constructor */
     TTangentBezierSpline(const TTangentBezierSpline& Other)
-	    : Tension(Other.Tension)
-	    , TangentModes(Other.TangentModes)
+	    : TTangentBezierSpline()
     {
-		InternalSpline = Other.InternalSpline;
+		Tension = Other.Tension;
+		TangentModes = Other.TangentModes;
+		AssignInternalSpline(Other.InternalSpline);
     }
 
     /**  Copy assignment */
@@ -184,11 +196,34 @@ public:
             Tension = Other.Tension;
             TangentModes = Other.TangentModes;
 			bStationaryEndpoint = Other.bStationaryEndpoint;
-            InternalSpline = Other.InternalSpline;
+            AssignInternalSpline(Other.InternalSpline);
         }
         return *this;
     }
 
+	/** Move constructor */
+	TTangentBezierSpline(TTangentBezierSpline&& Other) noexcept
+		: TTangentBezierSpline() 
+	{
+		Tension = Other.Tension;
+		TangentModes = MoveTemp(Other.TangentModes);
+		bStationaryEndpoint = Other.bStationaryEndpoint;
+		AssignInternalSpline(MoveTemp(Other.InternalSpline));
+	}
+
+	/** Move assignment */
+	TTangentBezierSpline& operator=(TTangentBezierSpline&& Other) noexcept
+	{
+		if (this != &Other)
+		{
+			Tension = Other.Tension;
+			TangentModes = MoveTemp(Other.TangentModes);
+			bStationaryEndpoint = Other.bStationaryEndpoint;
+			AssignInternalSpline(MoveTemp(Other.InternalSpline));
+		}
+		return *this;
+	}
+
 	virtual bool IsEqual(const ISplineInterface* OtherSpline) const override
 	{
 		if (OtherSpline->GetTypeId() == GetTypeId())
@@ -216,6 +251,7 @@ public:
 		if (Ar.IsLoading())
 		{
 			TangentModes.SetNum(NumTangentModes);
+			InstallFindWindowOverride();
 		}
         
 		for (int32 i = 0; i < NumTangentModes; ++i)
@@ -229,7 +265,6 @@ public:
 			}
 		}
 		Ar << bStationaryEndpoint;
-		
         return true;
     }
 
@@ -263,7 +298,7 @@ public:
 		TTangentBezierSpline Result;
         
 		// Clear the default initialization
-		Result.InternalSpline = FPolyBezierSpline3d::CreateLine(StartPoint, EndPoint);
+		Result.AssignInternalSpline(FPolyBezierSpline3d::CreateLine(StartPoint, EndPoint));
 		
         Result.Reparameterize(EParameterizationPolicy::Uniform);
         return Result;
@@ -289,9 +324,9 @@ public:
         TTangentBezierSpline Result;
         
         // Clear the default initialization
-        Result.InternalSpline = FPolyBezierSpline3d::CreateCircleArc(
-            Center, Radius, StartAngle, EndAngle, NumSegments);
-        
+        Result.AssignInternalSpline(FPolyBezierSpline3d::CreateCircleArc(
+            Center, Radius, StartAngle, EndAngle, NumSegments));
+		
 		// Set up tangent status for all points
 		Result.TangentModes.Init(ETangentMode::Auto, Result.GetNumPoints());
         
@@ -316,9 +351,8 @@ public:
     {
         // Create the circle
         TTangentBezierSpline Result;
-        Result.InternalSpline = FPolyBezierSpline3d::CreateCircle(
-            ValueType(Center), Radius, NumSegments);
-        
+        Result.AssignInternalSpline(FPolyBezierSpline3d::CreateCircle(
+            ValueType(Center), Radius, NumSegments));
         // Set up tangent status for all points
 		Result.TangentModes.Init(ETangentMode::Auto, Result.GetNumPoints());
 		
@@ -344,9 +378,9 @@ public:
     {
         // Create the ellipse
         TTangentBezierSpline Result;
-        Result.InternalSpline = FPolyBezierSpline3d::CreateEllipse(
-            Center, RadiusX, RadiusY, NumSegments);
-        
+        Result.AssignInternalSpline(FPolyBezierSpline3d::CreateEllipse(
+            Center, RadiusX, RadiusY, NumSegments));
+		
 		// Set up tangent status for all points
 		Result.TangentModes.Init(ETangentMode::Auto, Result.GetNumPoints());
         
@@ -360,7 +394,7 @@ public:
 		TUniquePtr<TTangentBezierSpline<VALUETYPE>> Clone = MakeUnique<TTangentBezierSpline<VALUETYPE>>();
     
 		// Copy internal spline
-		Clone->InternalSpline = this->InternalSpline;
+		Clone->AssignInternalSpline(this->InternalSpline);
     
 		// Copy tangent settings
 		Clone->Tension = this->Tension;
@@ -1025,6 +1059,7 @@ public:
 		{
 			return;
 		}
+		
 		// special case for handling the last point and second last point
 		if (NumPoints == 1)
 		{
@@ -1044,15 +1079,25 @@ public:
 			TangentModes.RemoveAt(Index);
 			return;
 		}
-    
+
+		// Index 1 can be assumed valid after special casing 1 and 2 point splines.
+		ValueType StartTangentIn = GetTangentIn(1);
+		
 		// Remove the corresponding tangent mode
 		if (TangentModes.IsValidIndex(Index))
 		{
 			TangentModes.RemoveAt(Index);
 		}
-
+		
 		InternalSpline.RemovePoint(Index);
-    
+
+		// If Index is 0, we are deleting the first segment. This segment stores the in tangent of point 1.
+		// Because the segment is going away, we need to restore the user-specified tangent (if possible).
+		if (Index == 0 && TangentModes.IsValidIndex(0) && TangentModes[0] == ETangentMode::User)
+		{
+			SetTangentIn(0, StartTangentIn);
+		}
+		
 		// Ensure tangent modes array size matches number of points after removal
 		if (TangentModes.Num() > GetNumPoints())
 		{
@@ -1261,28 +1306,31 @@ public:
 		{
 			return;
 		}
-		
-        if (Index == 0)
-        {
-            // Special case for first point in closed loop
-            if (InternalSpline.IsClosedLoop() && GetNumberOfSegments() > 0)
-            {
-                const int32 LastSegment = GetNumberOfSegments() - 1;
-                ValueType P3 = InternalSpline.GetValue(LastSegment * 4 + 3);
-                ValueType P2 = P3 - (NewTangent / 3.0f);
-                UpdateBezierControlPoint(LastSegment * 4 + 2, P2);
-            }
-            return;
-        }
-        
-        // Normal case - update previous segment's P2 control point
-        const int32 BezierIndex = (Index - 1) * 4 + 2;
-        if (BezierIndex < InternalSpline.NumKeys())
-        {
-            ValueType P3 = InternalSpline.GetValue(BezierIndex + 1);
-            ValueType P2 = P3 - (NewTangent / 3.0f);
-            UpdateBezierControlPoint(BezierIndex, P2);
-        }
+
+		int32 SegmentIndex = INDEX_NONE;
+		if (Index == 0)
+		{
+			if (InternalSpline.IsClosedLoop())
+			{
+				SegmentIndex = GetNumberOfSegments() - 1;
+			}
+		}
+		else
+		{
+			SegmentIndex = Index - 1;
+		}
+
+		if (SegmentIndex != INDEX_NONE)
+		{
+			const int32 BezierIndex = SegmentIndex * 4 + 2;
+			
+			if (BezierIndex < InternalSpline.NumKeys())
+			{
+				ValueType P3 = InternalSpline.GetValue(BezierIndex + 1);
+				ValueType P2 = P3 - NewTangent / 3.0f;
+				UpdateBezierControlPoint(BezierIndex, P2);
+			}
+		}
     }
     
     /** Sets tangent out for the specified point */
@@ -1340,7 +1388,7 @@ public:
 	  */
 	void SetPointTangentMode(int32 Index, ETangentMode Mode)
 	{
-		if (Index < 0 || Index > GetNumPoints())
+		if (Index < 0 || Index >= GetNumPoints())
 		{
 			return;
 		}
@@ -1368,15 +1416,7 @@ public:
 		const int32 ControlPointIndex = (Index == 0) ? 0 : ((Index - 1) * 4 + 3);
 		const int32 NewIndex = (InternalSpline.SetParameter(ControlPointIndex, NewParameter) + 1) / 4;
 
-		// We likely need to do better in PolyBezier's SetParameter to account for user tangents, but this works for now.
-		if (NewIndex != Index)
-		{
-			UpdateTangents();
-
-			UE_LOG(LogSpline, Verbose, TEXT("\t"));
-			UE_LOG(LogSpline, Verbose, TEXT("After Tangent Recompute:"));
-			InternalSpline.Dump();
-		}
+		UpdateTangents();
 
 		return NewIndex;
     }
@@ -1421,38 +1461,35 @@ public:
 	void UpdatePointTangents(int32 Index, ETangentMode Mode)
 	{
 		const int32 NumPoints = GetNumPoints();
+		const int32 NumSegments = GetNumberOfSegments();
+		
 		if (Index < 0 || Index >= NumPoints || Mode == ETangentMode::User || Mode == ETangentMode::Broken)
 		{
 			// Invalid index or user-defined mode, no action needed
 			return;
 		}
-		
-
+ 
 		// Get the position of the current point
 		const ValueType& Current = GetValue(Index);
-
 		// Get previous and next points (handling closed loop case)
 		const ValueType& Prev = (Index > 0) ? GetValue(Index - 1) : (IsClosedLoop() ? GetValue(NumPoints - 1) : Current);
-
 		const ValueType& Next = (Index < NumPoints - 1) ? GetValue(Index + 1) : (IsClosedLoop() ? GetValue(0) : Current);
-
-		// Get actual parameter values for points
-		float PrevParam = (Index > 0)
-			                  ? static_cast<float>(Index - 1)
-			                  : (IsClosedLoop() ? static_cast<float>(NumPoints - 1) : static_cast<float>(Index));
-
-		float CurrentParam = static_cast<float>(Index);
-
-		float NextParam = (Index < NumPoints - 1) ? static_cast<float>(Index + 1) : (IsClosedLoop() ? 0.0f : static_cast<float>(Index));
-
-		// Adjust params for closed loop to ensure proper weighting
-		if (IsClosedLoop() && Index == 0)
-		{
-			PrevParam -= static_cast<float>(NumPoints); // Make previous parameter negative for proper delta
-		}
-		else if (IsClosedLoop() && Index == NumPoints - 1)
+		const float CurrentParam = GetParameter(Index);
+		float PrevParam = Index > 0 ? GetParameter(Index - 1) : CurrentParam;
+		float NextParam = Index < NumPoints - 1 ? GetParameter(Index + 1) : CurrentParam;
+ 
+		// If we are closed and updating the first or last point, we need to do some special work to get virtual PrevParam or NextParam
+		// The 'virtual' parameter is not the actual parameter of the previous or next point, but the current parameter +- the knot delta.
+		if (IsClosedLoop())
 		{
-			NextParam += static_cast<float>(NumPoints); // Make next parameter greater than NumPoints for proper delta
+			if (Index == 0)
+			{
+				PrevParam = CurrentParam - GetSegmentParameterRange(NumSegments - 1).Length();
+			}
+			else if (Index == NumPoints - 1)
+			{
+				NextParam = CurrentParam + GetSegmentParameterRange(NumSegments - 1).Length();
+			}
 		}
 
 		// Calculate tangents based on mode
@@ -1504,25 +1541,27 @@ public:
 
 		case ETangentMode::Linear:
 			{
-				// Previous segment tangent (incoming)
+				// Linear tangents should also be per-unit-parameter
 				if (Index > 0 || IsClosedLoop())
 				{
-					InTangent = Current - Prev;
+					float PrevToCurDiff = FMath::Max<float>(UE_KINDA_SMALL_NUMBER, CurrentParam - PrevParam);
+					InTangent = (Current - Prev) / PrevToCurDiff;
 				}
 				else
 				{
 					InTangent = ValueType();
 				}
 
-				// Next segment tangent (outgoing)
 				if (Index < NumPoints - 1 || IsClosedLoop())
 				{
-					OutTangent = Next - Current;
+					float CurToNextDiff = FMath::Max<float>(UE_KINDA_SMALL_NUMBER, NextParam - CurrentParam);
+					OutTangent = (Next - Current) / CurToNextDiff;
 				}
 				else
 				{
 					OutTangent = ValueType();
 				}
+
 			}
 			break;
 
@@ -1541,7 +1580,9 @@ public:
 		default:
 			// Default to Auto mode if unknown
 			{
-				ValueType Tangent = (Next - Prev) * (1.0f - Tension);
+				float PrevToNextParamDiff = FMath::Max<float>(UE_KINDA_SMALL_NUMBER, NextParam - PrevParam);
+				ValueType Tangent = ((Next - Prev) / PrevToNextParamDiff) * (1.0f - Tension);
+
 				InTangent = Tangent;
 				OutTangent = Tangent;
 			}
@@ -1640,6 +1681,39 @@ protected:
 	}
 private:
 
+	void InstallFindWindowOverride()
+	{
+		InternalSpline.SetFindWindowOverride(&RescaleNormalizedHandlesForSegment);
+	}
+	
+	// ONLY adjusts P1/P2 in caller-provided scratch; generic over WindowSize
+	static void RescaleNormalizedHandlesForSegment(FWindow& Window, int32 SegmentIndex,  const FInterval1f& SegRange, FWindowStore& Scratch)
+	{
+		// bail if window not complete
+		if (!Math::FSplineValidation::IsValidWindow(Window))
+		{
+			return;
+		}
+
+		const float L = SegRange.Length();
+
+		const ValueType& P0 = *Window[0];
+		const ValueType& P1 = *Window[1];
+		const ValueType& P2 = *Window[2];
+		const ValueType& P3 = *Window[3];
+
+		// Storage is *normalized*; convert to world-param Bezier CPs
+		const ValueType TOut = (P1 - P0) * 3.0f; // normalized tangent @ start
+		const ValueType TIn  = (P3 - P2) * 3.0f; // normalized tangent @ end
+
+		// Only write adjusted P1/P2 into scratch; keep P0,P3 pointers intact
+		Scratch[1] = P0 + (TOut / 3.0f) * L;
+		Scratch[2] = P3 - (TIn  / 3.0f) * L;
+
+		Window[1] = &Scratch[1];
+		Window[2] = &Scratch[2];
+	}
+	
 	static bool IsCurveKey(ETangentMode Mode)
 	{
 		return ((Mode == ETangentMode::Auto) || (Mode == ETangentMode::AutoClamped) || (Mode == ETangentMode::User) || (Mode == ETangentMode::Broken));
@@ -1704,7 +1778,16 @@ private:
     	}
     	return false;
     }
-
+	void AssignInternalSpline(const FPolyBezierSpline& NewSpline)
+	{
+		InternalSpline = NewSpline;
+		InstallFindWindowOverride();
+	}
+	void AssignInternalSpline(FPolyBezierSpline&& NewSpline)
+	{
+		InternalSpline = MoveTemp(NewSpline);
+		InstallFindWindowOverride();
+	}
 protected:
 	/* Tension parameter for auto-computed tangents [0,1] */
 	float Tension = 0.0f;
@@ -1713,10 +1796,11 @@ protected:
 	TArray<ETangentMode> TangentModes;
 
 	bool bStationaryEndpoint = false;
+
 };
 
 using FTangentBezierSpline3f = TTangentBezierSpline<FVector3f>; 
 using FTangentBezierSpline3d = TTangentBezierSpline<FVector3d>; 
 } // end namespace UE::Geometry::Spline
 } // end namespace UE::Geometry
-} // end namespace UE
\ No newline at end of file
+} // end namespace UE
