[HWRT] Fix issues due to mismatch between static mesh sections used for rasterization and ray tracing.
- cache ray tracing mesh commands based on materials returned by GetCachedRayTracingInstance(...) instead of using SceneInfo->StaticMeshes.

diff --git a/Engine/Source/Runtime/Engine/Private/SkeletalMeshSceneProxy.cpp b/Engine/Source/Runtime/Engine/Private/SkeletalMeshSceneProxy.cpp
--- a/Engine/Source/Runtime/Engine/Private/SkeletalMeshSceneProxy.cpp
+++ b/Engine/Source/Runtime/Engine/Private/SkeletalMeshSceneProxy.cpp
@@ -494,7 +494,6 @@ void FSkeletalMeshSceneProxy::DrawStaticElements(FStaticPrimitiveDrawInterface*
 	if (!HasViewDependentDPG())
 	{
 		ESceneDepthPriorityGroup PrimitiveDPG = GetStaticDepthPriorityGroup();
-		bool bUseSelectedMaterial = false;
 
 		for (int32 LODIndex = ClampedLODIndex; LODIndex < NumLODs; ++LODIndex)
 		{
@@ -529,53 +528,20 @@ void FSkeletalMeshSceneProxy::DrawStaticElements(FStaticPrimitiveDrawInterface*
 				#if WITH_EDITOR
 					if (GIsEditor)
 					{
-						bUseSelectedMaterial = (MeshObject->SelectedEditorSection == SectionIndex);
 						PDI->SetHitProxy(SectionElementInfo.HitProxy);
 					}
 				#endif // WITH_EDITOR
 
 					FMeshBatch MeshElement;
-					FMeshBatchElement& BatchElement = MeshElement.Elements[0];
-					MeshElement.DepthPriorityGroup = PrimitiveDPG;
-					MeshElement.VertexFactory = VertexFactory;
-					MeshElement.MaterialRenderProxy = SectionElementInfo.Material->GetRenderProxy();
-					MeshElement.ReverseCulling = IsLocalToWorldDeterminantNegative();
-					MeshElement.CastShadow = SectionElementInfo.bEnableShadowCasting;
-				#if RHI_RAYTRACING
-					MeshElement.CastRayTracedShadow = MeshElement.CastShadow && bCastDynamicShadow;
-				#endif
-					MeshElement.Type = PT_TriangleList;
-					MeshElement.LODIndex = LODIndex;
-					MeshElement.SegmentIndex = SectionIndex;
-					MeshElement.MeshIdInPrimitive = SectionIndex;
-
-					int32 DynamicBoundsStartOffset = MeshObject->GetDynamicBoundsStartOffset(LODIndex);
-					if (bAllowDynamicMeshBounds && DynamicBoundsStartOffset >= 0)
-					{
-						MeshElement.DynamicMeshBoundsIndex = DynamicBoundsStartOffset + SectionIndex;
-					}
-
-					BatchElement.PrimitiveUniformBuffer = GetUniformBuffer();
-					BatchElement.FirstIndex = Section.BaseIndex;
-					BatchElement.MinVertexIndex = Section.BaseVertexIndex;
-					BatchElement.MaxVertexIndex = LODData.GetNumVertices() - 1;
-					BatchElement.NumPrimitives = Section.NumTriangles;
-					BatchElement.IndexBuffer = LODData.MultiSizeIndexContainer.GetIndexBuffer();
-					BatchElement.bFetchInstanceCountFromScene = bInstancedSkinnedMesh;
-					BatchElement.NumInstances = GetInstanceDataHeader().NumInstances;
+					GetStaticMeshBatch(LODData, LODIndex, Section, SectionIndex, SectionElementInfo, VertexFactory, PrimitiveDPG, MeshElement);
 
 					PDI->DrawMesh(MeshElement, ScreenSize);
 
 					UMaterialInterface* SectionOverlayMaterial = (SectionElementInfo.PerSectionOverlayMaterial != nullptr) ? SectionElementInfo.PerSectionOverlayMaterial : OverlayMaterial;
 					if (SectionOverlayMaterial != nullptr)
 					{
-						FMeshBatch OverlayMeshBatch(MeshElement);
-						OverlayMeshBatch.bOverlayMaterial = true;
-						OverlayMeshBatch.CastShadow = false;
-						OverlayMeshBatch.bSelectable = false;
-						OverlayMeshBatch.MaterialRenderProxy = SectionOverlayMaterial->GetRenderProxy();
-						// make sure overlay is always rendered on top of base mesh
-						OverlayMeshBatch.MeshIdInPrimitive += LODData.RenderSections.Num();
+						FMeshBatch OverlayMeshBatch;
+						GetOverlayMeshBatch(LODData, SectionOverlayMaterial, MeshElement, OverlayMeshBatch);
 						// Reuse mesh ScreenSize as cull distance for an overlay. Overlay does not need to compute LOD so we can avoid adding new members into MeshBatch or MeshRelevance
 						float OverlayMeshScreenSize = OverlayMaterialMaxDrawDistance;
 						PDI->DrawMesh(OverlayMeshBatch, OverlayMeshScreenSize);
@@ -730,6 +696,57 @@ const FCardRepresentationData* FSkeletalMeshSceneProxy::GetMeshCardRepresentatio
 	return CardRepresentationData;
 }
 
+void FSkeletalMeshSceneProxy::GetStaticMeshBatch(
+	const FSkeletalMeshLODRenderData& LODData,
+	const int32 LODIndex,
+	const FSkelMeshRenderSection& Section,
+	const int32 SectionIndex,
+	const FSectionElementInfo& SectionElementInfo,
+	const FVertexFactory* VertexFactory,
+	ESceneDepthPriorityGroup PrimitiveDPG,
+	FMeshBatch& OutMeshBatch) const
+{
+	FMeshBatchElement& BatchElement = OutMeshBatch.Elements[0];
+	OutMeshBatch.DepthPriorityGroup = PrimitiveDPG;
+	OutMeshBatch.VertexFactory = VertexFactory;
+	OutMeshBatch.MaterialRenderProxy = SectionElementInfo.Material->GetRenderProxy();
+	OutMeshBatch.ReverseCulling = IsLocalToWorldDeterminantNegative();
+	OutMeshBatch.CastShadow = SectionElementInfo.bEnableShadowCasting;
+#if RHI_RAYTRACING
+	OutMeshBatch.CastRayTracedShadow = OutMeshBatch.CastShadow && bCastDynamicShadow;
+#endif
+	OutMeshBatch.Type = PT_TriangleList;
+	OutMeshBatch.LODIndex = LODIndex;
+	OutMeshBatch.SegmentIndex = SectionIndex;
+	OutMeshBatch.MeshIdInPrimitive = SectionIndex;
+
+	int32 DynamicBoundsStartOffset = MeshObject->GetDynamicBoundsStartOffset(LODIndex);
+	if (bAllowDynamicMeshBounds && DynamicBoundsStartOffset >= 0)
+	{
+		OutMeshBatch.DynamicMeshBoundsIndex = DynamicBoundsStartOffset + SectionIndex;
+	}
+
+	BatchElement.PrimitiveUniformBuffer = GetUniformBuffer();
+	BatchElement.FirstIndex = Section.BaseIndex;
+	BatchElement.MinVertexIndex = Section.BaseVertexIndex;
+	BatchElement.MaxVertexIndex = LODData.GetNumVertices() - 1;
+	BatchElement.NumPrimitives = Section.NumTriangles;
+	BatchElement.IndexBuffer = LODData.MultiSizeIndexContainer.GetIndexBuffer();
+	BatchElement.bFetchInstanceCountFromScene = bInstancedSkinnedMesh;
+	BatchElement.NumInstances = GetInstanceDataHeader().NumInstances;
+}
+
+void FSkeletalMeshSceneProxy::GetOverlayMeshBatch(const FSkeletalMeshLODRenderData& LODData, UMaterialInterface* SectionOverlayMaterial, const FMeshBatch& MeshBatch, FMeshBatch& OutOverlayMeshBatch) const
+{
+	OutOverlayMeshBatch = FMeshBatch(MeshBatch);
+	OutOverlayMeshBatch.bOverlayMaterial = true;
+	OutOverlayMeshBatch.CastShadow = false;
+	OutOverlayMeshBatch.bSelectable = false;
+	OutOverlayMeshBatch.MaterialRenderProxy = SectionOverlayMaterial->GetRenderProxy();
+	// make sure overlay is always rendered on top of base mesh
+	OutOverlayMeshBatch.MeshIdInPrimitive += LODData.RenderSections.Num();
+}
+
 void FSkeletalMeshSceneProxy::CreateBaseMeshBatch(const FSceneView* View, const FSkeletalMeshLODRenderData& LODData, const int32 LODIndex, const int32 SectionIndex, const FSectionElementInfo& SectionElementInfo, FMeshBatch& Mesh, ESkinVertexFactoryMode VFMode) const
 {
 	const FSkelMeshRenderSection& Section = LODData.RenderSections[SectionIndex];
@@ -975,6 +992,71 @@ TArray<FRayTracingGeometry*> FSkeletalMeshSceneProxy::GetStaticRayTracingGeometr
 	return {};
 }
 
+ERayTracingPrimitiveFlags FSkeletalMeshSceneProxy::GetCachedRayTracingInstance(FRayTracingInstance& OutRayTracingInstance)
+{
+	const ERayTracingPrimitiveFlags Flags = FPrimitiveSceneProxy::GetCachedRayTracingInstance(OutRayTracingInstance);
+
+	// the following flags cause ray tracing mesh command caching to be disabled
+	static const ERayTracingPrimitiveFlags DisableCacheMeshCommandsFlags = ERayTracingPrimitiveFlags::Dynamic
+		| ERayTracingPrimitiveFlags::Exclude
+		| ERayTracingPrimitiveFlags::Skip
+		| ERayTracingPrimitiveFlags::UnsupportedProxyType;
+
+	if (!EnumHasAnyFlags(Flags, DisableCacheMeshCommandsFlags))
+	{
+		const ESceneDepthPriorityGroup PrimitiveDPG = GetStaticDepthPriorityGroup();
+		const bool bIgnorePreviewFilter = true;
+
+		for (int32 LODIndex = ClampedLODIndex; LODIndex < NumLODs; ++LODIndex)
+		{
+			const FSkeletalMeshLODRenderData& LODData = RenderData->LODRenderData[LODIndex];
+
+			if (LODSections.Num() > 0 && LODData.GetNumVertices() > 0)
+			{
+				const FLODSectionElements& LODSection = LODSections[LODIndex];
+				check(LODSection.SectionElements.Num() == LODData.RenderSections.Num());
+
+				for (FSkeletalMeshSectionIter Iter(LODIndex, *MeshObject, LODData, LODSection, bIgnorePreviewFilter); Iter; ++Iter)
+				{
+					const FSkelMeshRenderSection& Section = Iter.GetSection();
+					const int32 SectionIndex = Iter.GetSectionElementIndex();
+					const FVertexFactory* VertexFactory = MeshObject->GetStaticSkinVertexFactory(LODIndex, SectionIndex, ESkinVertexFactoryMode::Default);
+
+					if (!VertexFactory || Section.NumTriangles == 0)
+					{
+						// hide this part
+						continue;
+					}
+
+					const FSectionElementInfo& SectionElementInfo = Iter.GetSectionElementInfo();
+
+					// If hidden skip the draw
+					if (MeshObject->IsMaterialHidden(LODIndex, SectionElementInfo.UseMaterialIndex) || Section.bDisabled)
+					{
+						continue;
+					}
+
+					FMeshBatch MeshElement;
+					GetStaticMeshBatch(LODData, LODIndex, Section, SectionIndex, SectionElementInfo, VertexFactory, PrimitiveDPG, MeshElement);
+
+					OutRayTracingInstance.Materials.Add(MeshElement);
+
+					UMaterialInterface* SectionOverlayMaterial = (SectionElementInfo.PerSectionOverlayMaterial != nullptr) ? SectionElementInfo.PerSectionOverlayMaterial : OverlayMaterial;
+					if (SectionOverlayMaterial != nullptr)
+					{
+						FMeshBatch OverlayMeshBatch;
+						GetOverlayMeshBatch(LODData, SectionOverlayMaterial, MeshElement, OverlayMeshBatch);
+
+						OutRayTracingInstance.Materials.Add(OverlayMeshBatch);
+					}
+				}
+			}
+		}
+	}
+
+	return Flags;
+}
+
 void FSkeletalMeshSceneProxy::GetDynamicRayTracingInstances(FRayTracingInstanceCollector& Collector)
 {
 	if (!CVarRayTracingSkeletalMeshes.GetValueOnRenderThread()
diff --git a/Engine/Source/Runtime/Engine/Private/StaticMeshSceneProxy.cpp b/Engine/Source/Runtime/Engine/Private/StaticMeshSceneProxy.cpp
--- a/Engine/Source/Runtime/Engine/Private/StaticMeshSceneProxy.cpp
+++ b/Engine/Source/Runtime/Engine/Private/StaticMeshSceneProxy.cpp
@@ -2031,6 +2031,63 @@ TArray<FRayTracingGeometry*> FStaticMeshSceneProxy::GetStaticRayTracingGeometrie
 	return {};
 }
 
+ERayTracingPrimitiveFlags FStaticMeshSceneProxy::GetCachedRayTracingInstance(FRayTracingInstance& OutRayTracingInstance)
+{
+	const ERayTracingPrimitiveFlags Flags = FPrimitiveSceneProxy::GetCachedRayTracingInstance(OutRayTracingInstance);
+
+	// the following flags cause ray tracing mesh command caching to be disabled
+	static const ERayTracingPrimitiveFlags DisableCacheMeshCommandsFlags = ERayTracingPrimitiveFlags::Dynamic
+		| ERayTracingPrimitiveFlags::Exclude
+		| ERayTracingPrimitiveFlags::Skip
+		| ERayTracingPrimitiveFlags::UnsupportedProxyType;
+
+	if (!EnumHasAnyFlags(Flags, DisableCacheMeshCommandsFlags))
+	{
+		const int32 NumBatches = GetNumMeshBatches();
+		const ESceneDepthPriorityGroup PrimitiveDPG = GetStaticDepthPriorityGroup();
+
+		FStaticMeshRayTracingProxyLODArray& RayTracingLODs = RenderData->RayTracingProxy->LODs;
+		const int32 NumLODs = RayTracingLODs.Num();
+
+		for (int32 LODIndex = 0; LODIndex < NumLODs; ++LODIndex)
+		{
+			const FStaticMeshRayTracingProxyLOD& RayTracingLOD = RayTracingLODs[LODIndex];
+
+			for (int32 BatchIndex = 0; BatchIndex < NumBatches; BatchIndex++)
+			{
+				for (int32 SectionIndex = 0; SectionIndex < RayTracingLOD.Sections->Num(); SectionIndex++)
+				{
+					FMeshBatch& MeshBatch = OutRayTracingInstance.Materials.AddDefaulted_GetRef();
+
+					bool bResult;
+					if (RenderData->RayTracingProxy->bUsingRenderingLODs)
+					{
+						// when using rendering LODs we can reuse the main GetMeshElement(...)
+						bResult = GetMeshElement(LODIndex, BatchIndex, SectionIndex, PrimitiveDPG, false, false, MeshBatch);
+					}
+					else
+					{
+						// otherwise initialize MeshBatch using ray tracing proxy VB/IB/Section data
+						bResult = GetRayTracingMeshElement(LODIndex, BatchIndex, SectionIndex, PrimitiveDPG, MeshBatch);
+					}
+
+					if (!bResult)
+					{
+						// Hidden material
+						MeshBatch.MaterialRenderProxy = UMaterial::GetDefaultMaterial(MD_Surface)->GetRenderProxy();
+						MeshBatch.VertexFactory = &(*RenderData->RayTracingProxy->LODVertexFactories)[LODIndex].VertexFactory;
+					}
+					MeshBatch.ReverseCulling = bReverseCulling; // overwrite what came from GetMeshElement as DXR only needs the user driven flag, not the flipping implied by the transform
+					MeshBatch.SegmentIndex = SectionIndex;
+					MeshBatch.MeshIdInPrimitive = SectionIndex;
+				}
+			}
+		}
+	}
+
+	return Flags;
+}
+
 RayTracing::FGeometryGroupHandle FStaticMeshSceneProxy::GetRayTracingGeometryGroupHandle() const
 {
 	check(IsInRenderingThread() || IsInParallelRenderingThread());
diff --git a/Engine/Source/Runtime/Engine/Public/SkeletalMeshSceneProxy.h b/Engine/Source/Runtime/Engine/Public/SkeletalMeshSceneProxy.h
--- a/Engine/Source/Runtime/Engine/Public/SkeletalMeshSceneProxy.h
+++ b/Engine/Source/Runtime/Engine/Public/SkeletalMeshSceneProxy.h
@@ -58,6 +58,7 @@ public:
 	ENGINE_API virtual RayTracing::FGeometryGroupHandle GetRayTracingGeometryGroupHandle() const override;
 
 	ENGINE_API virtual TArray<FRayTracingGeometry*> GetStaticRayTracingGeometries() const override;
+	ENGINE_API virtual ERayTracingPrimitiveFlags GetCachedRayTracingInstance(FRayTracingInstance& OutRayTracingInstance);
 
 	ENGINE_API virtual void GetDynamicRayTracingInstances(FRayTracingInstanceCollector& Collector) override;
 #endif // RHI_RAYTRACING
@@ -277,6 +278,17 @@ protected:
 
 private:
 	ENGINE_API void CreateBaseMeshBatch(const FSceneView* View, const FSkeletalMeshLODRenderData& LODData, const int32 LODIndex, const int32 SectionIndex, const FSectionElementInfo& SectionElementInfo, FMeshBatch& Mesh, ESkinVertexFactoryMode VFMode = ESkinVertexFactoryMode::Default) const;
+	
+	ENGINE_API void GetStaticMeshBatch(
+		const FSkeletalMeshLODRenderData& LODData,
+		const int32 LODIndex,
+		const FSkelMeshRenderSection& Section,
+		const int32 SectionIndex,
+		const FSectionElementInfo& SectionElementInfo,
+		const FVertexFactory* VertexFactory,
+		ESceneDepthPriorityGroup PrimitiveDPG,
+		FMeshBatch& OutMeshBatch) const;
+	ENGINE_API void GetOverlayMeshBatch(const FSkeletalMeshLODRenderData& LODData, UMaterialInterface* SectionOverlayMaterial, const FMeshBatch& MeshBatch, FMeshBatch& OutOverlayMeshBatch) const;
 
 public:
 #if WITH_EDITORONLY_DATA
diff --git a/Engine/Source/Runtime/Engine/Public/StaticMeshSceneProxy.h b/Engine/Source/Runtime/Engine/Public/StaticMeshSceneProxy.h
--- a/Engine/Source/Runtime/Engine/Public/StaticMeshSceneProxy.h
+++ b/Engine/Source/Runtime/Engine/Public/StaticMeshSceneProxy.h
@@ -170,7 +170,7 @@ public:
 
 #if RHI_RAYTRACING
 	ENGINE_API virtual TArray<FRayTracingGeometry*> GetStaticRayTracingGeometries() const override;
-
+	ENGINE_API virtual ERayTracingPrimitiveFlags GetCachedRayTracingInstance(FRayTracingInstance& OutRayTracingInstance);
 	ENGINE_API virtual void GetDynamicRayTracingInstances(FRayTracingInstanceCollector& Collector) override;
 	ENGINE_API virtual bool HasRayTracingRepresentation() const override;
 	virtual bool IsRayTracingRelevant() const override { return true; }
diff --git a/Engine/Source/Runtime/Renderer/Private/PrimitiveSceneInfo.cpp b/Engine/Source/Runtime/Renderer/Private/PrimitiveSceneInfo.cpp
--- a/Engine/Source/Runtime/Renderer/Private/PrimitiveSceneInfo.cpp
+++ b/Engine/Source/Runtime/Renderer/Private/PrimitiveSceneInfo.cpp
@@ -1191,18 +1191,24 @@ int32 CacheRayTracingMeshCommands(
 
 		if (RayTracingInstance.Materials.Num() > 0)
 		{
-			// TODO: LOD w/ screen size support. Probably needs another array parallel to RayTracingInstances
-			// We assume it is exactly 1 LOD now (true for Nanite proxies)
-			LODCount = 1;
+			for (const FMeshBatch& Mesh : RayTracingInstance.Materials)
+			{
+				LODCount = LODCount < (Mesh.LODIndex + 1) ? (Mesh.LODIndex + 1) : LODCount;
+			}
 		}
 		else
 		{
+			// Keeping old behavior for one UE release. 
+			ensureMsgf(false, TEXT("GetCachedRayTracingInstance(...) must return materials when using static instances code path."));
+
 			for (const FStaticMeshBatch& Mesh : SceneInfo->StaticMeshes)
 			{
 				LODCount = LODCount < (Mesh.LODIndex + 1) ? (Mesh.LODIndex + 1) : LODCount;
 			}
 		}
 
+		checkf(!EnumHasAnyFlags(Flags, ERayTracingPrimitiveFlags::CacheInstances) || LODCount == 1, TEXT("Cached ray tracing instances are expected to only have materials for one LOD."));
+
 		check(SceneInfo->GetRayTracingLODDataNum() == 0);
 				
 		TArray<FPrimitiveSceneInfo::FRayTracingLODData> RayTracingLODData;
@@ -1224,7 +1230,8 @@ int32 CacheRayTracingMeshCommands(
 		}
 		else
 		{
-			for (const FStaticMeshBatch& Mesh : SceneInfo->StaticMeshes) // TODO: Handle !RayTracingProxy->bUsingRenderingLODs
+			// Keeping old behavior for one UE release. 
+			for (const FStaticMeshBatch& Mesh : SceneInfo->StaticMeshes)
 			{
 				CacheRayTracingMeshBatch(Mesh, SceneInfo, RayTracingLODData, Commands, CommandContext, RayTracingMeshProcessor, DeferredData);
 			}
