Fix error compiling SkeletalMeshSceneProxy.h.

diff --git a/Engine/Source/Runtime/Engine/Public/SkeletalMeshSceneProxy.h b/Engine/Source/Runtime/Engine/Public/SkeletalMeshSceneProxy.h
--- a/Engine/Source/Runtime/Engine/Public/SkeletalMeshSceneProxy.h
+++ b/Engine/Source/Runtime/Engine/Public/SkeletalMeshSceneProxy.h
@@ -13,6 +13,7 @@ FSkeletalMeshSceneProxy
 class FSkeletalMeshObject;
 
 struct FSkinnedMeshSceneProxyDesc;
+struct FSkelMeshRenderSection;
 
 /**
  * A skeletal mesh component scene proxy.
