Improved streamable manager error.

diff --git a/Engine/Source/Runtime/Engine/Classes/Engine/Experimental/StreamableManagerError.h b/Engine/Source/Runtime/Engine/Classes/Engine/Experimental/StreamableManagerError.h
--- a/Engine/Source/Runtime/Engine/Classes/Engine/Experimental/StreamableManagerError.h
+++ b/Engine/Source/Runtime/Engine/Classes/Engine/Experimental/StreamableManagerError.h
@@ -3,6 +3,7 @@
 
 #include "Experimental/UnifiedError/UnifiedError.h"
 #include "UObject/UObjectGlobals.h"
+#include "UObject/SoftObjectPath.h"
 
 
 UE_DECLARE_ERROR_MODULE(ENGINE_API, StreamableManager);
@@ -29,10 +30,16 @@ namespace UE::UnifiedError::StreamableManager
 	class FStreamableManagerAdditionalContext
 	{
 	public:
-		FString SoftObjectPath;
+		FString RequestedPackageName;
+		FString MissingObject;
 	};
 
-
+	class FStreamableManagerRequestContext
+	{
+	public:
+		FString DebugName;
+		TArray<FSoftObjectPath> RequestedAssets;
+	};
 	
 }
 UE_DECLARE_ERRORSTRUCT_FEATURES(StreamableManager, FStreamableManagerAdditionalContext);
@@ -40,9 +47,25 @@ inline void SerializeForLog(FCbWriter& Writer, const UE::UnifiedError::Streamabl
 {
 	Writer.BeginObject();
 	Writer.AddString(ANSITEXTVIEW("$type"), TErrorStructFeatures<UE::UnifiedError::StreamableManager::FStreamableManagerAdditionalContext>::GetErrorContextTypeNameAsString());
-	Writer.AddString(ANSITEXTVIEW("$format"), TEXT("(SoftObjectPath: {SoftObjectPath})"));
-	Writer.AddString(ANSITEXTVIEW("SoftObjectPath"), Context.SoftObjectPath);
+	Writer.AddString(ANSITEXTVIEW("$format"), TEXT("(RequstedPackage: {RequestedPackageName}, MissingObject: {MissingObject})"));
+	Writer.AddString(ANSITEXTVIEW("RequestedPackageName"), Context.RequestedPackageName);
+	Writer.AddString(ANSITEXTVIEW("MissingObject"), Context.MissingObject);
 	Writer.EndObject();
 }
 
+UE_DECLARE_ERRORSTRUCT_FEATURES(StreamableManager, FStreamableManagerRequestContext);
+inline void SerializeForLog(FCbWriter& Writer, const UE::UnifiedError::StreamableManager::FStreamableManagerRequestContext& Context)
+{
+	Writer.BeginObject(); 
+	Writer.AddString(ANSITEXTVIEW("$type"), TErrorStructFeatures<UE::UnifiedError::StreamableManager::FStreamableManagerRequestContext>::GetErrorContextTypeNameAsString());
+	Writer.AddString(ANSITEXTVIEW("$format"), TEXT("(RequestDebugName: {DebugName}, RequestedAssets: {RequestedAssets})"));
+	Writer.AddString(ANSITEXTVIEW("DebugName"), Context.DebugName);
+	Writer.BeginArray(ANSITEXTVIEW("RequestedAssets"));
+	for (const FSoftObjectPath& RequestedAsset : Context.RequestedAssets )
+	{
+		SerializeForLog(Writer, RequestedAsset);
+	}
+	Writer.EndArray();
 
+	Writer.EndObject();
+}
\ No newline at end of file
diff --git a/Engine/Source/Runtime/Engine/Private/AssetManager.cpp b/Engine/Source/Runtime/Engine/Private/AssetManager.cpp
--- a/Engine/Source/Runtime/Engine/Private/AssetManager.cpp
+++ b/Engine/Source/Runtime/Engine/Private/AssetManager.cpp
@@ -2435,7 +2435,7 @@ void UAssetManager::OnAssetStateChangeCompleted(TSharedPtr<FStreamableHandle> Bo
 
 				if (AdditionalContext)
 				{
-					UE_LOG(LogAssetManager, Warning, TEXT("OnAssetStateChangeCompleted: Received error (%s) loading asset! Primary asset id:(%s) erroring asset:%s"), *BoundHandle->GetError()->GetErrorMessage(true).ToString(), *PrimaryAssetId.ToString(), *AdditionalContext->SoftObjectPath);
+					UE_LOG(LogAssetManager, Warning, TEXT("OnAssetStateChangeCompleted: Received error (%s) loading asset! Primary asset id:(%s) erroring asset:%s"), *BoundHandle->GetError()->GetErrorMessage(true).ToString(), *PrimaryAssetId.ToString(), *AdditionalContext->MissingObject);
 				}
 				else
 				{
diff --git a/Engine/Source/Runtime/Engine/Private/StreamableManager.cpp b/Engine/Source/Runtime/Engine/Private/StreamableManager.cpp
--- a/Engine/Source/Runtime/Engine/Private/StreamableManager.cpp
+++ b/Engine/Source/Runtime/Engine/Private/StreamableManager.cpp
@@ -1259,18 +1259,14 @@ void FStreamableHandle::AsyncLoadCallbackWrapper(const FName& PackageName, UPack
 			/*TRefCountPtr<UE::UnifiedError::IErrorDetails> ErrorDetails(new UE::UnifiedError::StreamableManager::FStreamableManagerErrorDetails(TargetName.ToString()));
 			MaybeError->PushErrorDetails(ErrorDetails);*/
 
-			UE::UnifiedError::StreamableManager::FStreamableManagerAdditionalContext MoreContext = {TargetName.ToString()};
-			MaybeError->PushErrorContext(MoveTemp(MoreContext));
+			UE::UnifiedError::StreamableManager::FStreamableManagerAdditionalContext MoreContext = {PackageName.ToString(), TargetName.ToString()};
+			MaybeError->PushErrorContext(MoveTemp(MoreContext), UE::UnifiedError::EDetailFilter::Default);
 
-			// todo: Daniel, Devin, this shouldn't be nessisary just use PushErrorContent(T, true) after we have support for $format
-			MaybeError->AppendFormatString(NSLOCTEXT("StreamableManager", "AdditionalContextFormatString", "Missing streamable asset: {StreamableManager_FStreamableManagerAdditionalContext/SoftObjectPath}"));
+			UE::UnifiedError::StreamableManager::FStreamableManagerRequestContext RequestContext = {DebugName, RequestedAssets};
+			MaybeError->PushErrorContext(MoveTemp(RequestContext), UE::UnifiedError::EDetailFilter::Default);
 
 			SetError(*MaybeError);
 		}
-		else
-		{
-			//if (TargetName.ToString())
-		}
 
 		OwningManager->AsyncLoadCallback(MoveTemp(TargetName), Package, MaybeError);
 
