[AutoRTFM] Use a transactionally-safe lock in FObjectReverseLookupCache.

This can be hit by the Content Worker when AutoRTFM is enabled.

diff --git a/Engine/Source/Runtime/Engine/Private/ObjectCacheContext.cpp b/Engine/Source/Runtime/Engine/Private/ObjectCacheContext.cpp
--- a/Engine/Source/Runtime/Engine/Private/ObjectCacheContext.cpp
+++ b/Engine/Source/Runtime/Engine/Private/ObjectCacheContext.cpp
@@ -15,6 +15,7 @@
 #include "Logging/LogMacros.h"
 #include "HAL/LowLevelMemStats.h"
 #include "Misc/ScopeRWLock.h"
+#include "Misc/TransactionallySafeRWLock.h"
 #include "Components/ComponentInterfaceIterator.h"
 #include "Components/InstancedSkinnedMeshComponent.h"
 
@@ -179,7 +180,7 @@ class FObjectReverseLookupCache
 public:
 	void RemoveFrom(FromType* InFrom)
 	{
-		FWriteScopeLock ScopeLock(Lock);
+		UE::TWriteScopeLock ScopeLock(Lock);
 
 		TRawSet<ToType*>* ToObjects = FromToMapping.Find(InFrom);
 		if (ToObjects)
@@ -201,7 +202,7 @@ public:
 	void Update(ToType* InTo, const TArray<FromType*>& NewFromObjects)
 	{
 		LLM_SCOPE_BYTAG(ObjectReverseLookupCache);
-		FWriteScopeLock ScopeLock(Lock);
+		UE::TWriteScopeLock ScopeLock(Lock);
 
 		// Remove the old reverse lookup mappings as they can't be read from the
 		// objects anymore.
@@ -242,7 +243,7 @@ public:
 	template <typename PredicateType = FDefaultObjectSearchPredicate>
 	TObjectCacheIterator<ToType> GetFrom(FromType* InFrom, const PredicateType& Predicate = PredicateType()) const
 	{
-		FReadScopeLock ScopeLock(Lock);
+		UE::TReadScopeLock ScopeLock(Lock);
 
 		// Until we can get some kind of thread-safe CopyOnWrite behavior
 		// we need to extract a copy to avoid race conditions.
@@ -265,9 +266,9 @@ public:
 	template <typename PredicateType = FDefaultObjectSearchPredicate>
 	TObjectCacheIterator<FromType> GetTo(ToType* InTo, const PredicateType& Predicate = PredicateType()) const
 	{
-		FReadScopeLock ScopeLock(Lock);
+		UE::TReadScopeLock ScopeLock(Lock);
 
-		// Until we can get some kind of thread-safe CopyOnWrite behavior, 
+		// Until we can get some kind of thread-safe CopyOnWrite behavior,
 		// we need to extract a copy to avoid race conditions.
 		TArray<FromType*> OutFrom;
 		if (const TRawSet<FromType*>* Set = ToFromMapping.Find(InTo))
@@ -282,12 +283,12 @@ public:
 			}
 		}
 
-		return TObjectCacheIterator<FromType>(MoveTemp(OutFrom));	
+		return TObjectCacheIterator<FromType>(MoveTemp(OutFrom));
 	}
 
-	template<class T> 
+	template<class T>
 	FString GetFullName(T* Obj)
-	{	
+	{
 		return Obj->GetFullName();
 	}
 
@@ -299,8 +300,8 @@ public:
 
 	int32 Compare(const FObjectReverseLookupCache<FromType, ToType>& Other)
 	{
-		FWriteScopeLock ScopeLockA(Other.Lock);
-		FWriteScopeLock ScopeLockB(Lock);
+		UE::TWriteScopeLock ScopeLockA(Other.Lock);
+		UE::TWriteScopeLock ScopeLockB(Lock);
 
 		int32 ErrorCount = 0;
 		for (const auto& Kvp : FromToMapping)
@@ -334,7 +335,7 @@ public:
 		return ErrorCount;
 	}
 private:
-	mutable FRWLock Lock;
+	mutable FTransactionallySafeRWLock Lock;
 
 	// We use raw pointer to compare UObject here both for performance reason
 	// and because we can't use TObjectKey. Some objects are reallocated in-place
@@ -357,8 +358,8 @@ private:
 		static FORCEINLINE KeyInitType GetSetKey(ElementInitType Element) {	return (void*)Element; }
 		static FORCEINLINE bool Matches(KeyInitType A, KeyInitType B) { return A == B; }
 		static FORCEINLINE uint32 GetKeyHash(KeyInitType Key) { return GetTypeHash(Key); }
-	}; 
-	
+	};
+
 	template <typename InElementType>
 	using TRawSet = TSet<InElementType, TRawObjectKeyFuncs<InElementType>>;
 
@@ -827,7 +828,7 @@ TObjectCacheIterator<USkinnedMeshComponent> FObjectCacheContext::GetSkinnedMeshC
 
 TObjectCacheIterator<USkinnedMeshComponent> FObjectCacheContext::GetSkinnedMeshComponents(USkinnedAsset* InSkinnedAsset)
 {
-	return SkinnedMeshCache.Get<TObjectIterator<USkinnedMeshComponent>>(InSkinnedAsset, 
+	return SkinnedMeshCache.Get<TObjectIterator<USkinnedMeshComponent>>(InSkinnedAsset,
 		[](USkinnedMeshComponent* Component){ return 1; },
 		[](USkinnedMeshComponent* Component, int32 ) { return Component->GetSkinnedAsset(); });
 }
@@ -908,7 +909,7 @@ TObjectCacheIterator<IStaticMeshComponent> FObjectCacheContext::GetStaticMeshCom
 	}
 #endif
 
-	const TSet<IStaticMeshComponent*>& ComputeResult = StaticMeshCache.GetSet<TComponentInterfaceIterator<IStaticMeshComponent>>(InStaticMesh, 
+	const TSet<IStaticMeshComponent*>& ComputeResult = StaticMeshCache.GetSet<TComponentInterfaceIterator<IStaticMeshComponent>>(InStaticMesh,
 		[](IStaticMeshComponent*) { return 1; },
 		[](IStaticMeshComponent* StaticMeshComponent, int32 Index) { return StaticMeshComponent->GetStaticMesh(); });
 
@@ -1094,7 +1095,7 @@ TObjectCacheIterator<UMaterialInterface> FObjectCacheContext::GetMaterialsAffect
 	{
 		TSet<UMaterialInterface*> LookupResults;
 		ObjectCacheContextImpl::GetMaterialsAffectedByMaterials_Lookup(InMaterials, LookupResults);
-		
+
 		for (UMaterialInterface* MaterialInterface : LookupResults)
 		{
 			checkf(IterationResults.Contains(MaterialInterface), TEXT("Permanent Object Cache has an additional element %s"), *MaterialInterface->GetFullName());
@@ -1187,10 +1188,10 @@ FObjectCacheContextScope::~FObjectCacheContextScope()
 	}
 }
 
-FObjectCacheContext& FObjectCacheContextScope::GetContext() 
-{ 
+FObjectCacheContext& FObjectCacheContextScope::GetContext()
+{
 	using namespace ObjectCacheContextScopeImpl;
-	return *Current; 
+	return *Current;
 }
 
 #if WITH_EDITOR
@@ -1346,7 +1347,7 @@ void FObjectCacheEventSink::NotifyRenderStateChanged_Concurrent(IPrimitiveCompon
 		else
 		{
 			FObjectCacheEventSinkPrivate::NotifyRenderStateChanged_Concurrent(PrimitiveComponent);
-		}		
+		}
 	}
 }
 
@@ -1513,7 +1514,7 @@ void FObjectCacheEventSinkPrivate::NotifyMaterialDestroyed_Concurrent(UMaterialI
 void FObjectCacheEventSinkPrivate::NotifyUsedMaterialsChanged_Concurrent(const IPrimitiveComponent* PrimitiveComponent, const TArray<UMaterialInterface*>& UsedMaterials)
 {
 	check(PrimitiveComponent);
-	using namespace ObjectCacheContextImpl;	
+	using namespace ObjectCacheContextImpl;
 
 	GMaterialToPrimitiveLookupCache.Update(const_cast<IPrimitiveComponent*>(PrimitiveComponent), UsedMaterials);
 }
@@ -1579,7 +1580,7 @@ void FObjectCacheEventSinkPrivate::NotifyStaticMeshChanged_Concurrent(IStaticMes
 	}
 }
 
-} // namespace ObjectCacheContextImpl 
+} // namespace ObjectCacheContextImpl
 
 #endif
 
