Misc BlendSpace changes to try obtain a callstack of the source of invalid blend sample data instead of crashing at a later stage

#rnx

diff --git a/Engine/Source/Runtime/Engine/Private/Animation/BlendSpace.cpp b/Engine/Source/Runtime/Engine/Private/Animation/BlendSpace.cpp
--- a/Engine/Source/Runtime/Engine/Private/Animation/BlendSpace.cpp
+++ b/Engine/Source/Runtime/Engine/Private/Animation/BlendSpace.cpp
@@ -907,10 +907,13 @@ void UBlendSpace::ReplaceReferredAnimations(const TMap<UAnimationAsset*, UAnimat
 		// replace the referenced animation sequence (if there was one)
 		if (Sample.Animation)
 		{
-			if (UAnimSequence* const* ReplacementAsset = (UAnimSequence* const*)ReplacementMap.Find(Sample.Animation))
+			if (UAnimationAsset* const* ReplacementAsset = ReplacementMap.Find(Sample.Animation))
 			{
-				Sample.Animation = *ReplacementAsset;
-				Sample.Animation->ReplaceReferredAnimations(ReplacementMap);
+				if (UAnimSequence* ReplacementAnimSequence = Cast<UAnimSequence>(*ReplacementAsset))
+				{
+					Sample.Animation = ReplacementAnimSequence;
+					Sample.Animation->ReplaceReferredAnimations(ReplacementMap);
+				}
 			}
 		}
 
@@ -920,10 +923,13 @@ void UBlendSpace::ReplaceReferredAnimations(const TMap<UAnimationAsset*, UAnimat
 	// replace preview base pose sequence asset (if there was one)
 	if (PreviewBasePose)
 	{
-		if (UAnimSequence* const* ReplacementAsset = (UAnimSequence* const*)ReplacementMap.Find(PreviewBasePose))
+		if (UAnimationAsset* const* ReplacementAsset = ReplacementMap.Find(PreviewBasePose))
 		{
-			PreviewBasePose = *ReplacementAsset;
-			PreviewBasePose->ReplaceReferredAnimations(ReplacementMap);
+			if (UAnimSequence* ReplacementAnimSequence = Cast<UAnimSequence>(*ReplacementAsset))
+			{
+				PreviewBasePose = ReplacementAnimSequence;
+				PreviewBasePose->ReplaceReferredAnimations(ReplacementMap);
+			}
 		}
 	}
 
