Added ExternalStreamingObjects to the list of packages returned in FStreamingDescriptor::GenerateWorldPluginPackagesList

diff --git a/Engine/Source/Runtime/Engine/Private/WorldPartition/WorldPartitionRuntimeHash.cpp b/Engine/Source/Runtime/Engine/Private/WorldPartition/WorldPartitionRuntimeHash.cpp
--- a/Engine/Source/Runtime/Engine/Private/WorldPartition/WorldPartitionRuntimeHash.cpp
+++ b/Engine/Source/Runtime/Engine/Private/WorldPartition/WorldPartitionRuntimeHash.cpp
@@ -12,6 +12,7 @@
 #include "WorldPartition/DataLayer/ExternalDataLayerHelper.h"
 #include "WorldPartition/DataLayer/DataLayerInstanceNames.h"
 #include "WorldPartition/DataLayer/DataLayerManager.h"
+#include "WorldPartition/DataLayer/ExternalDataLayerManager.h"
 #include "Misc/ArchiveMD5.h"
 #include "Cooker/CookEvents.h"
 #if WITH_EDITOR
@@ -149,6 +150,19 @@ FString URuntimeHashExternalStreamingObjectBase::GetPackageNameToCreate() const
 	return FString();
 }
 
+FString URuntimeHashExternalStreamingObjectBase::GetPackagePathToCreate() const
+{
+	// GetPackagePathToCreate should not be called for ExternalStreamingObjects
+	if (ensure(ExternalDataLayerAsset))
+	{
+		UWorld* World = OuterWorld.Get();
+		check(World); 
+		
+		return World->GetWorldPartition()->GetExternalDataLayerManager()->GetExternalStreamingObjectPackagePath(ExternalDataLayerAsset);
+	}
+	return FString();
+}
+
 bool URuntimeHashExternalStreamingObjectBase::PrepareForCook(const IWorldPartitionCookPackageContext& InCookContext)
 {
 	bool bResult = true;
diff --git a/Engine/Source/Runtime/Engine/Private/WorldPartition/WorldPartitionStreamingDescriptor.cpp b/Engine/Source/Runtime/Engine/Private/WorldPartition/WorldPartitionStreamingDescriptor.cpp
--- a/Engine/Source/Runtime/Engine/Private/WorldPartition/WorldPartitionStreamingDescriptor.cpp
+++ b/Engine/Source/Runtime/Engine/Private/WorldPartition/WorldPartitionStreamingDescriptor.cpp
@@ -14,6 +14,7 @@
 #include "WorldPartition/RuntimeHashSet/RuntimePartitionLHGrid.h"
 #include "WorldPartition/WorldPartitionRuntimeLevelStreamingCell.h"
 #include "WorldPartition/DataLayer/DataLayerManager.h"
+#include "WorldPartition/DataLayer/ExternalDataLayerManager.h"
 #include "Engine/World.h"
 #include "HAL/FileManager.h"
 #include "Editor.h"
@@ -138,6 +139,20 @@ bool FStreamingDescriptor::GenerateStreamingDescriptor(UWorld* InWorld, FStreami
 				StreamingGrid.LoadingRange = SpatialHashStreamingGrid.LoadingRange;
 			});
 
+			for (const TWeakObjectPtr<URuntimeHashExternalStreamingObjectBase>& InjectedExternalStreamingObject : RuntimeSpatialHash->InjectedExternalStreamingObjects)
+			{
+				if (InjectedExternalStreamingObject.IsValid())
+				{
+					URuntimeSpatialHashExternalStreamingObject* ExternalStreamingObject = CastChecked<URuntimeSpatialHashExternalStreamingObject>(InjectedExternalStreamingObject.Get());
+					
+					for (FSpatialHashStreamingGrid& SpatialHashGrid : ExternalStreamingObject->StreamingGrids)
+					{
+						FStreamingDescriptor::FStreamingGrid& StreamingGrid = StreamingGrids.FindOrAdd(SpatialHashGrid.GridName);
+						StreamingGrid.ExternalStreamingObjects.AddUnique(ExternalStreamingObject->GetPackagePathToCreate());
+					}					
+				}
+			}
+
 			WorldPartition->RuntimeHash->ForEachStreamingCells([&StreamingGrids, DataLayerManager](const UWorldPartitionRuntimeCell* Cell)
 			{
 				UWorldPartitionRuntimeCellDataSpatialHash* CellDataSpatialHash = CastChecked<UWorldPartitionRuntimeCellDataSpatialHash>(Cell->RuntimeCellData);
@@ -197,6 +212,9 @@ bool FStreamingDescriptor::GenerateStreamingDescriptor(UWorld* InWorld, FStreami
 
 					for (const FRuntimePartitionStreamingData& StreamingData : ExternalStreamingObject->RuntimeStreamingData)
 					{
+						FStreamingDescriptor::FStreamingGrid& StreamingGrid = StreamingGrids.FindOrAdd(StreamingData.Name);
+						StreamingGrid.ExternalStreamingObjects.AddUnique(ExternalStreamingObject->GetPackagePathToCreate());
+
 						StreamingDataList.Add(&StreamingData);
 					}
 				}
@@ -289,17 +307,61 @@ TMap<FName, TArray<FName>> FStreamingDescriptor::GenerateWorldPluginPackagesList
 	{
 		for (FStreamingDescriptor::FStreamingCell& Cell : Grid.StreamingCells)
 		{
-			FName PluginName = FName(FPathViews::GetMountPointNameFromPath(Cell.CellPackage.ToString()));
+			FName PluginName = FName(FPathViews::GetMountPointNameFromPath(Cell.CellPackage.ToString()));			
 			FName PackageName = Cell.CellPackage;
 
 			PluginPackages.FindOrAdd(PluginName).Add(PackageName);
 		}
+
+		for (FString& ExternalStreamingObject : Grid.ExternalStreamingObjects)
+		{
+			FName PluginName = FName(FPathViews::GetMountPointNameFromPath(ExternalStreamingObject));			
+			FName PackageName = FName(ExternalStreamingObject);
+
+			PluginPackages.FindOrAdd(PluginName).Add(PackageName);
+		}
 	}
 
 	return PluginPackages;
 }
 
-
 };
 
+
+static FAutoConsoleCommand GenerateWorldPluginPackagesList(
+	TEXT("wp.editor.GenerateWorldPluginPackagesList"),
+	TEXT("Generate the world plugin package list for the current world in the speficied file."),
+	FConsoleCommandWithArgsDelegate::CreateLambda([](const TArray<FString>& Args)
+	{
+		if (Args.Num() == 1)
+		{
+			if (UWorld* World = GEditor->GetEditorWorldContext().World())
+			{
+				if (!World->IsGameWorld())
+				{
+					if (FArchive* LogFile = IFileManager::Get().CreateFileWriter(*Args[0]))
+					{
+						TMap<FName, TArray<FName>> PluginsPackageList = UE::Private::WorldPartition::FStreamingDescriptor::GenerateWorldPluginPackagesList(World);
+						{
+							FHierarchicalLogArchive Ar(*LogFile);
+
+							for (TPair<FName, TArray<FName>> PluginPackageList : PluginsPackageList)
+							{
+								UE_SCOPED_INDENT_LOG_ARCHIVE(Ar.PrintfIndent(TEXT("Plugin: %s"), *PluginPackageList.Key.ToString()));
+
+								for (FName Package : PluginPackageList.Value)
+								{
+									Ar.Printf(TEXT("Package: %s"), *Package.ToString());
+								}
+							}
+						}
+						LogFile->Close();
+						delete LogFile;
+					}
+				}
+			}
+		}
+	})
+);
+
 #endif
\ No newline at end of file
diff --git a/Engine/Source/Runtime/Engine/Public/WorldPartition/DataLayer/ExternalDataLayerManager.h b/Engine/Source/Runtime/Engine/Public/WorldPartition/DataLayer/ExternalDataLayerManager.h
--- a/Engine/Source/Runtime/Engine/Public/WorldPartition/DataLayer/ExternalDataLayerManager.h
+++ b/Engine/Source/Runtime/Engine/Public/WorldPartition/DataLayer/ExternalDataLayerManager.h
@@ -132,6 +132,7 @@ private:
 	friend class UWorldPartitionRuntimeLevelStreamingCell;
 	friend class UGameFeatureActionConvertContentBundleWorldPartitionBuilder;
 	friend struct FWorldPartitionUtils;
+	friend class URuntimeHashExternalStreamingObjectBase;
 };
 
 #undef UE_API
diff --git a/Engine/Source/Runtime/Engine/Public/WorldPartition/WorldPartitionRuntimeHash.h b/Engine/Source/Runtime/Engine/Public/WorldPartition/WorldPartitionRuntimeHash.h
--- a/Engine/Source/Runtime/Engine/Public/WorldPartition/WorldPartitionRuntimeHash.h
+++ b/Engine/Source/Runtime/Engine/Public/WorldPartition/WorldPartitionRuntimeHash.h
@@ -101,6 +101,7 @@ public:
 	//~End IWorldPartitionCookPackageObject interface
 
 	const static TCHAR* GetCookedExternalStreamingObjectName() { return TEXT("RuntimeHashExternalStreamingObjectBase"); }
+	ENGINE_API FString GetPackagePathToCreate() const;
 
 protected:
 	virtual void DumpStateLog(FHierarchicalLogArchive& Ar);
diff --git a/Engine/Source/Runtime/Engine/Public/WorldPartition/WorldPartitionRuntimeSpatialHash.h b/Engine/Source/Runtime/Engine/Public/WorldPartition/WorldPartitionRuntimeSpatialHash.h
--- a/Engine/Source/Runtime/Engine/Public/WorldPartition/WorldPartitionRuntimeSpatialHash.h
+++ b/Engine/Source/Runtime/Engine/Public/WorldPartition/WorldPartitionRuntimeSpatialHash.h
@@ -19,6 +19,11 @@ typedef FInt64Vector3 FGridCellCoord;
 typedef FInt64Vector2 FGridCellCoord2;
 class FWorldPartitionDraw2DContext;
 
+namespace UE::Private::WorldPartition
+{
+	struct FStreamingDescriptor;
+};
+
 UENUM()
 enum class EWorldPartitionCVarProjectDefaultOverride : uint8
 {
@@ -457,4 +462,5 @@ private:
 	friend class UWorldPartitionSubsystem;
 	friend class UWorldPartitionRuntimeHashSet;
 	friend struct FSpatialHashSettings;
+	friend UE::Private::WorldPartition::FStreamingDescriptor;
 };
diff --git a/Engine/Source/Runtime/Engine/Public/WorldPartition/WorldPartitionStreamingDescriptor.h b/Engine/Source/Runtime/Engine/Public/WorldPartition/WorldPartitionStreamingDescriptor.h
--- a/Engine/Source/Runtime/Engine/Public/WorldPartition/WorldPartitionStreamingDescriptor.h
+++ b/Engine/Source/Runtime/Engine/Public/WorldPartition/WorldPartitionStreamingDescriptor.h
@@ -83,6 +83,9 @@ namespace UE::Private::WorldPartition
 			/** Streaming cells of this grid */
 			TArray<FStreamingCell> StreamingCells;
 
+			/** ExternalStreaming Objects for this grid */
+			TArray<FString> ExternalStreamingObjects;
+
 			void DumpStateLog(FHierarchicalLogArchive& Ar);
 		};
 
