Substrate - fix for wrong SSS material on some material instance. (non equivalence was leading to different uniform scalar numeric parameter name and thus the prfile index could not be set correctly into constant buffer).

Also added a clear of the SubsurfaceProfilesRT to avoid having too long list of parameter.

#tests PC editor dx12
#rnx

#DisabledValidations checkUnacceptableWords

diff --git a/Engine/Source/Runtime/Engine/Private/Materials/HLSLMaterialTranslator.cpp b/Engine/Source/Runtime/Engine/Private/Materials/HLSLMaterialTranslator.cpp
--- a/Engine/Source/Runtime/Engine/Private/Materials/HLSLMaterialTranslator.cpp
+++ b/Engine/Source/Runtime/Engine/Private/Materials/HLSLMaterialTranslator.cpp
@@ -1871,7 +1871,8 @@ void FHLSLMaterialTranslator::TranslateMaterial()
 
 				if (BaseMaterial && BaseMaterial->bUseMaterialAttributes)
 				{
-					USubsurfaceProfile* SSSProfile = MaterialInterface ? MaterialInterface->GetSubsurfaceProfile_Internal() : nullptr;
+					// Need to root subsurface profile for the uniform scalar numeric parameter name to be correct.
+					USubsurfaceProfile* SSSProfile = MaterialInterface ? MaterialInterface->GetSubsurfaceProfileRoot_Internal() : nullptr;
 					const bool bHasSSS = SSSProfile != nullptr;
 					
 					return UMaterialExpressionSubstrateConvertMaterialAttributes::CompileCommon(this, 0,
@@ -1883,7 +1884,7 @@ void FHLSLMaterialTranslator::TranslateMaterial()
 				}
 				else if ((Material->GetMaterialDomain() == MD_Surface || Material->GetMaterialDomain() == MD_RuntimeVirtualTexture || Material->GetMaterialDomain() == MD_DeferredDecal) && EditorOnlyData)
 				{
-					USubsurfaceProfile* SSSProfile = MaterialInterface ? MaterialInterface->GetSubsurfaceProfile_Internal() : nullptr;
+					USubsurfaceProfile* SSSProfile = MaterialInterface ? MaterialInterface->GetSubsurfaceProfileRoot_Internal() : nullptr;
 					const bool bHasSSS = SSSProfile != nullptr;
 					const bool bHasAnisotropy = EditorOnlyData->Anisotropy.IsConnected();
 
diff --git a/Engine/Source/Runtime/Engine/Private/Materials/Material.cpp b/Engine/Source/Runtime/Engine/Private/Materials/Material.cpp
--- a/Engine/Source/Runtime/Engine/Private/Materials/Material.cpp
+++ b/Engine/Source/Runtime/Engine/Private/Materials/Material.cpp
@@ -7481,6 +7481,12 @@ USubsurfaceProfile* UMaterial::GetSubsurfaceProfile_Internal() const
 	return SubsurfaceProfile; 
 }
 
+USubsurfaceProfile* UMaterial::GetSubsurfaceProfileRoot_Internal() const
+{
+	checkSlow(IsInGameThread());
+	return SubsurfaceProfile; 
+}
+
 uint32 UMaterial::NumSubsurfaceProfileRoot_Internal() const
 {
 	return SubsurfaceProfiles.Num();
diff --git a/Engine/Source/Runtime/Engine/Private/Materials/MaterialInstance.cpp b/Engine/Source/Runtime/Engine/Private/Materials/MaterialInstance.cpp
--- a/Engine/Source/Runtime/Engine/Private/Materials/MaterialInstance.cpp
+++ b/Engine/Source/Runtime/Engine/Private/Materials/MaterialInstance.cpp
@@ -5217,6 +5217,13 @@ USubsurfaceProfile* UMaterialInstance::GetSubsurfaceProfile_Internal() const
 	return Parent ? Parent->GetSubsurfaceProfile_Internal() : 0;
 }
 
+USubsurfaceProfile* UMaterialInstance::GetSubsurfaceProfileRoot_Internal() const
+{
+	// Return the subsurface profile form the root material to create named parameters.
+	checkSlow(IsInGameThread());
+	return Parent ? Parent->GetSubsurfaceProfileRoot_Internal() : 0;
+}
+
 uint32 UMaterialInstance::NumSubsurfaceProfileRoot_Internal() const
 {
 	// Return the subsurface profile count form the root material.
diff --git a/Engine/Source/Runtime/Engine/Private/Materials/MaterialInterface.cpp b/Engine/Source/Runtime/Engine/Private/Materials/MaterialInterface.cpp
--- a/Engine/Source/Runtime/Engine/Private/Materials/MaterialInterface.cpp
+++ b/Engine/Source/Runtime/Engine/Private/Materials/MaterialInterface.cpp
@@ -1612,6 +1612,11 @@ USubsurfaceProfile* UMaterialInterface::GetSubsurfaceProfile_Internal() const
 	return nullptr;
 }
 
+USubsurfaceProfile* UMaterialInterface::GetSubsurfaceProfileRoot_Internal() const
+{
+	return nullptr;
+}
+
 USubsurfaceProfile* UMaterialInterface::GetSubsurfaceProfileRoot_Internal(uint32 Index) const
 {
 	return nullptr;
@@ -1869,6 +1874,7 @@ void UMaterialInterface::UpdateMaterialRenderProxy(FMaterialRenderProxy& Proxy)
 			ENQUEUE_RENDER_COMMAND(UpdateMaterialRenderProxySubsurfaceProfiles)(
 			[InProxy, Entries, SubsurfaceProfileOverride, SubsurfaceProfileOverrideSettings](FRHICommandListImmediate& RHICmdList)
 			{
+				InProxy->ClearSubsurfaceProfileRT(); // Clear to not add duplicates when updating from the editor.
 				for (const FEntry& Entry : Entries)
 				{
 					if (Entry.Profile)
diff --git a/Engine/Source/Runtime/Engine/Public/Materials/Material.h b/Engine/Source/Runtime/Engine/Public/Materials/Material.h
--- a/Engine/Source/Runtime/Engine/Public/Materials/Material.h
+++ b/Engine/Source/Runtime/Engine/Public/Materials/Material.h
@@ -1278,6 +1278,7 @@ public:
 	ENGINE_API virtual bool IsPostProcessMaterial() const;
 	ENGINE_API bool IsPostProcessMaterialOutputingAlpha() const;
 	ENGINE_API virtual USubsurfaceProfile* GetSubsurfaceProfile_Internal() const override;
+	ENGINE_API virtual USubsurfaceProfile* GetSubsurfaceProfileRoot_Internal() const override;
 	ENGINE_API virtual uint32 NumSubsurfaceProfileRoot_Internal() const override;
 	ENGINE_API virtual USubsurfaceProfile* GetSubsurfaceProfileRoot_Internal(uint32 Index) const override;
 	ENGINE_API virtual USubsurfaceProfile* GetSubsurfaceProfileOverride_Internal() const override;
diff --git a/Engine/Source/Runtime/Engine/Public/Materials/MaterialInstance.h b/Engine/Source/Runtime/Engine/Public/Materials/MaterialInstance.h
--- a/Engine/Source/Runtime/Engine/Public/Materials/MaterialInstance.h
+++ b/Engine/Source/Runtime/Engine/Public/Materials/MaterialInstance.h
@@ -971,6 +971,7 @@ public:
 	ENGINE_API virtual bool HasCustomizedUVs() const override;
 	ENGINE_API virtual bool HasPixelAnimation() const override;
 	ENGINE_API virtual USubsurfaceProfile* GetSubsurfaceProfile_Internal() const override;
+	ENGINE_API virtual USubsurfaceProfile* GetSubsurfaceProfileRoot_Internal() const override;
 	ENGINE_API uint32 NumSubsurfaceProfileRoot_Internal() const override;
 	ENGINE_API USubsurfaceProfile* GetSubsurfaceProfileRoot_Internal(uint32 Index) const override;
 	ENGINE_API virtual USubsurfaceProfile* GetSubsurfaceProfileOverride_Internal() const override;
diff --git a/Engine/Source/Runtime/Engine/Public/Materials/MaterialInterface.h b/Engine/Source/Runtime/Engine/Public/Materials/MaterialInterface.h
--- a/Engine/Source/Runtime/Engine/Public/Materials/MaterialInterface.h
+++ b/Engine/Source/Runtime/Engine/Public/Materials/MaterialInterface.h
@@ -1033,6 +1033,7 @@ public:
 	ENGINE_API virtual bool HasCustomizedUVs() const;
 	ENGINE_API virtual bool HasPixelAnimation() const;
 	ENGINE_API virtual USubsurfaceProfile* GetSubsurfaceProfile_Internal() const;
+	ENGINE_API virtual USubsurfaceProfile* GetSubsurfaceProfileRoot_Internal() const;
 	ENGINE_API virtual uint32 NumSubsurfaceProfileRoot_Internal() const;
 	ENGINE_API virtual USubsurfaceProfile* GetSubsurfaceProfileRoot_Internal(uint32 Index) const;
 	ENGINE_API virtual USubsurfaceProfile* GetSubsurfaceProfileOverride_Internal() const;
diff --git a/Engine/Source/Runtime/Engine/Public/Materials/MaterialRenderProxy.h b/Engine/Source/Runtime/Engine/Public/Materials/MaterialRenderProxy.h
--- a/Engine/Source/Runtime/Engine/Public/Materials/MaterialRenderProxy.h
+++ b/Engine/Source/Runtime/Engine/Public/Materials/MaterialRenderProxy.h
@@ -225,6 +225,7 @@ public:
 	const USubsurfaceProfile* GetSubsurfaceProfileRT() const { return SubsurfaceProfileRT; }
 
 	// Subsurface profiles
+	void ClearSubsurfaceProfileRT() { SubsurfaceProfilesRT.Empty(); }
 	void AddSubsurfaceProfileRT(const USubsurfaceProfile* Ptr) { SubsurfaceProfilesRT.Add(Ptr); }
 	const USubsurfaceProfile* GetSubsurfaceProfileRT(uint32 Index) const { check(Index<uint32(SubsurfaceProfilesRT.Num())); return SubsurfaceProfilesRT[Index]; }
 	const uint32 NumSubsurfaceProfileRT() const { return SubsurfaceProfilesRT.Num(); }
