Fix for anim instance being spawned for an unregistered SMC

#rnx

diff --git a/Engine/Source/Runtime/Engine/Private/Components/SkeletalMeshComponent.cpp b/Engine/Source/Runtime/Engine/Private/Components/SkeletalMeshComponent.cpp
--- a/Engine/Source/Runtime/Engine/Private/Components/SkeletalMeshComponent.cpp
+++ b/Engine/Source/Runtime/Engine/Private/Components/SkeletalMeshComponent.cpp
@@ -866,6 +866,11 @@ void USkeletalMeshComponent::FinalizeAttributeEvaluationResults(const FBoneConta
 
 bool USkeletalMeshComponent::NeedToSpawnAnimScriptInstance() const
 {
+	if (!IsRegistered())
+	{
+		return false;
+	}
+
 	USkeletalMesh* MeshAsset = GetSkeletalMeshAsset();
 	IAnimClassInterface* AnimClassInterface = IAnimClassInterface::GetFromClass(AnimClass);
 	USkeleton* AnimSkeleton = (AnimClassInterface) ? AnimClassInterface->GetTargetSkeleton() : nullptr;
