Blend Space + Aim Space -- "Use Single Frame" broken
#fix reinstate Animation->GetPlayLength rather than sample playlength, when clamping evaluation time value

diff --git a/Engine/Source/Runtime/Engine/Private/Animation/BlendSpace.cpp b/Engine/Source/Runtime/Engine/Private/Animation/BlendSpace.cpp
--- a/Engine/Source/Runtime/Engine/Private/Animation/BlendSpace.cpp
+++ b/Engine/Source/Runtime/Engine/Private/Animation/BlendSpace.cpp
@@ -1220,7 +1220,7 @@ void UBlendSpace::GetAnimationPose_Internal(TArray<FBlendSampleData>& BlendSampl
 			{
 				if (Sample.Animation && Sample.Animation->GetSkeleton() != nullptr)
 				{
-					const float Time = FMath::Clamp<float>(BlendSampleDataCache[I].Time, 0.f, Sample.GetSamplePlayLength());
+					const float Time = FMath::Clamp<float>(BlendSampleDataCache[I].Time, 0.f, Sample.Animation->GetPlayLength());
 
 					FAnimationPoseData ChildAnimationPoseData = { Pose, ChildrenCurves[I], ChildrenAttributes[I] };
 					// first one always fills up the source one
