[HWRT] Fix crash trying to cache ray tracing mesh commands for LODs below ClampedMinLOD.

diff --git a/Engine/Source/Runtime/Engine/Private/StaticMeshSceneProxy.cpp b/Engine/Source/Runtime/Engine/Private/StaticMeshSceneProxy.cpp
--- a/Engine/Source/Runtime/Engine/Private/StaticMeshSceneProxy.cpp
+++ b/Engine/Source/Runtime/Engine/Private/StaticMeshSceneProxy.cpp
@@ -2046,10 +2046,11 @@ ERayTracingPrimitiveFlags FStaticMeshSceneProxy::GetCachedRayTracingInstance(FRa
 		const int32 NumBatches = GetNumMeshBatches();
 		const ESceneDepthPriorityGroup PrimitiveDPG = GetStaticDepthPriorityGroup();
 
-		FStaticMeshRayTracingProxyLODArray& RayTracingLODs = RenderData->RayTracingProxy->LODs;
+		const FStaticMeshRayTracingProxyLODArray& RayTracingLODs = RenderData->RayTracingProxy->LODs;
 		const int32 NumLODs = RayTracingLODs.Num();
+		const int32 RayTracingMinLOD = RenderData->RayTracingProxy->bUsingRenderingLODs ? ClampedMinLOD : 0;
 
-		for (int32 LODIndex = 0; LODIndex < NumLODs; ++LODIndex)
+		for (int32 LODIndex = RayTracingMinLOD; LODIndex < NumLODs; ++LODIndex)
 		{
 			const FStaticMeshRayTracingProxyLOD& RayTracingLOD = RayTracingLODs[LODIndex];
 
