Re-enable Quad overdraw mode for Metal SM6

diff --git a/Engine/Source/Runtime/Engine/Private/DebugViewModeHelpers.cpp b/Engine/Source/Runtime/Engine/Private/DebugViewModeHelpers.cpp
--- a/Engine/Source/Runtime/Engine/Private/DebugViewModeHelpers.cpp
+++ b/Engine/Source/Runtime/Engine/Private/DebugViewModeHelpers.cpp
@@ -86,10 +86,10 @@ bool SupportDebugViewShaderMode(EDebugViewShaderMode ShaderMode, EShaderPlatform
 		return false;
 	}
 
-	// Quad modes don't compile/work on Metal or Playstation. 
+	// Quad modes don't compile/work on Metal < SM6 or Playstation. 
 	const bool bSupportQuadComplexityModes =
 		!FDataDrivenShaderPlatformInfo::GetIsLanguageSony(Platform) && 
-		!FDataDrivenShaderPlatformInfo::GetIsLanguageMetal(Platform);
+		(!FDataDrivenShaderPlatformInfo::GetIsLanguageMetal(Platform) || IsMetalSM6Platform(Platform));
 	
 	switch (ShaderMode)
 	{
