[HWRT] Prevent crash due to RayTracingProxy being unexpectedly null.
- added ensure and workaround instead

diff --git a/Engine/Source/Runtime/Engine/Private/StaticMeshSceneProxy.cpp b/Engine/Source/Runtime/Engine/Private/StaticMeshSceneProxy.cpp
--- a/Engine/Source/Runtime/Engine/Private/StaticMeshSceneProxy.cpp
+++ b/Engine/Source/Runtime/Engine/Private/StaticMeshSceneProxy.cpp
@@ -377,6 +377,11 @@ FStaticMeshSceneProxy::FStaticMeshSceneProxy(const FStaticMeshSceneProxyDesc& In
 	bDynamicRayTracingGeometry = false;
 	bCreateDynamicRayTracingGeometries = true;
 
+	if (bSupportRayTracing && !ensureMsgf(RenderData->RayTracingProxy, TEXT("StaticMesh [%s] missing RayTracingProxy."), StaticMesh ? *StaticMesh->GetName() : TEXT("None")))
+	{
+		bSupportRayTracing = false;
+	}
+
 	if (bSupportRayTracing)
 	{
 		const bool bWantsRayTracingWPO = bEvaluateWorldPositionOffset && MaterialRelevance.bUsesWorldPositionOffset && InProxyDesc.bEvaluateWorldPositionOffsetInRayTracing;
