Added support to get ExternalStreamingObject generated package name for ContentBundles
Changed ExternalStreamingObject to store the name of it's package (set by CB or EDL on generation) to avoid duplicating the package name logic outside EDL and CB

diff --git a/Engine/Source/Runtime/Engine/Private/WorldPartition/ContentBundle/ContentBundleEditor.cpp b/Engine/Source/Runtime/Engine/Private/WorldPartition/ContentBundle/ContentBundleEditor.cpp
--- a/Engine/Source/Runtime/Engine/Private/WorldPartition/ContentBundle/ContentBundleEditor.cpp
+++ b/Engine/Source/Runtime/Engine/Private/WorldPartition/ContentBundle/ContentBundleEditor.cpp
@@ -363,6 +363,7 @@ void FContentBundleEditor::GenerateStreaming(TArray<FString>* OutPackageToGenera
 	WorldPartition->GenerateContainerStreaming(Params, Context);
 
 	ExternalStreamingObject = WorldPartition->FlushStreamingToExternalStreamingObject();
+	ExternalStreamingObject->SetPackagePathToCreate(GetExternalStreamingObjectPackagePath()); 
 
 	uint32 CellCount = 0;
 	ExternalStreamingObject->ForEachStreamingCells([&CellCount](const UWorldPartitionRuntimeCell& Cell)
diff --git a/Engine/Source/Runtime/Engine/Private/WorldPartition/DataLayer/ExternalDataLayerManager.cpp b/Engine/Source/Runtime/Engine/Private/WorldPartition/DataLayer/ExternalDataLayerManager.cpp
--- a/Engine/Source/Runtime/Engine/Private/WorldPartition/DataLayer/ExternalDataLayerManager.cpp
+++ b/Engine/Source/Runtime/Engine/Private/WorldPartition/DataLayer/ExternalDataLayerManager.cpp
@@ -914,6 +914,7 @@ URuntimeHashExternalStreamingObjectBase* UExternalDataLayerManager::CreateExtern
 
 	ExternalStreamingObjects.Emplace(InExternalDataLayerAsset, ExternalStreamingObject);
 	ExternalStreamingObject->ExternalDataLayerAsset = InExternalDataLayerAsset;
+	ExternalStreamingObject->SetPackagePathToCreate(GetExternalStreamingObjectPackagePath(InExternalDataLayerAsset));
 
 	// Make a copy of External Data Layer WorldDataLayers DataLayerInstances
 	AWorldDataLayers* EDLWorldDataLayers = ExternalDataLayerInstance->GetDirectOuterWorldDataLayers();
diff --git a/Engine/Source/Runtime/Engine/Private/WorldPartition/WorldPartitionRuntimeHash.cpp b/Engine/Source/Runtime/Engine/Private/WorldPartition/WorldPartitionRuntimeHash.cpp
--- a/Engine/Source/Runtime/Engine/Private/WorldPartition/WorldPartitionRuntimeHash.cpp
+++ b/Engine/Source/Runtime/Engine/Private/WorldPartition/WorldPartitionRuntimeHash.cpp
@@ -142,7 +142,7 @@ UWorldPartitionRuntimeCell* URuntimeHashExternalStreamingObjectBase::GetCellForC
 
 FString URuntimeHashExternalStreamingObjectBase::GetPackageNameToCreate() const
 {
-	// GetPackageNameToCreate should not be called for ExternalStreamingObjects
+	// GetPackageNameToCreate should not be called for ExternalStreamingObjects belonging to content bundles
 	if (ensure(ExternalDataLayerAsset))
 	{
 		return TEXT("/") + FExternalDataLayerHelper::GetExternalStreamingObjectPackageName(ExternalDataLayerAsset);
@@ -150,17 +150,15 @@ FString URuntimeHashExternalStreamingObjectBase::GetPackageNameToCreate() const
 	return FString();
 }
 
-FString URuntimeHashExternalStreamingObjectBase::GetPackagePathToCreate() const
+
+void URuntimeHashExternalStreamingObjectBase::SetPackagePathToCreate(const  FString& InPath)
 {
-	// GetPackagePathToCreate should not be called for ExternalStreamingObjects
-	if (ensure(ExternalDataLayerAsset))
-	{
-		UWorld* World = OuterWorld.Get();
-		check(World); 
-		
-		return World->GetWorldPartition()->GetExternalDataLayerManager()->GetExternalStreamingObjectPackagePath(ExternalDataLayerAsset);
-	}
-	return FString();
+	PackagePathToCreate = InPath;
+}
+
+const FString& URuntimeHashExternalStreamingObjectBase::GetPackagePathToCreate() const
+{
+	return PackagePathToCreate;
 }
 
 bool URuntimeHashExternalStreamingObjectBase::PrepareForCook(const IWorldPartitionCookPackageContext& InCookContext)
diff --git a/Engine/Source/Runtime/Engine/Public/WorldPartition/WorldPartitionRuntimeHash.h b/Engine/Source/Runtime/Engine/Public/WorldPartition/WorldPartitionRuntimeHash.h
--- a/Engine/Source/Runtime/Engine/Public/WorldPartition/WorldPartitionRuntimeHash.h
+++ b/Engine/Source/Runtime/Engine/Public/WorldPartition/WorldPartitionRuntimeHash.h
@@ -101,7 +101,8 @@ public:
 	//~End IWorldPartitionCookPackageObject interface
 
 	const static TCHAR* GetCookedExternalStreamingObjectName() { return TEXT("RuntimeHashExternalStreamingObjectBase"); }
-	ENGINE_API FString GetPackagePathToCreate() const;
+	ENGINE_API const FString& GetPackagePathToCreate() const;
+	ENGINE_API void SetPackagePathToCreate(const FString& InPath);
 
 protected:
 	virtual void DumpStateLog(FHierarchicalLogArchive& Ar);
@@ -143,6 +144,10 @@ protected:
 	TWeakObjectPtr<UWorldPartition> TargetInjectedWorldPartition;
 #endif
 
+#if WITH_EDITOR
+	FString PackagePathToCreate;
+#endif
+
 	friend class UWorldPartition;
 	friend class UExternalDataLayerManager;
 };
