[ContentStreaming] Implement abandon support for RenderAssetUpdates to decouple a UStreamableRenderAsset from its FRenderAssetUpdate so it can be garbage collected independently (and also so as to not stall GC). This is a follow up to CL 45515138 which was a temporary mitigation to work around an issue that stalled GCs indefenitely.

#tests Tested that the asset update was ticked and cleaned up eventually by inducing abandons with `gc.ContinuousIncrementalGC 1` and `r.Streaming.TextureStreamInAbandonCallbackIOLatency`. Also ran replay runs with the same on multiple platforms.
[RN] Introduced a mechansim for render asset updates to be abandoned by owning `UStreamableRenderAssets` so `UStreamableRenderAssets` can be garbage collected independently (See r.Streaming.EnableAssetUpdateAbandons). Abandoned render assets are ticked during GC or every streaming update if  r.Streaming.TickAbandonedRenderAssetUpdatesOnStreamingUpdate is true (defaults to true).

diff --git a/Engine/Source/Runtime/Engine/Private/StreamableRenderAsset.cpp b/Engine/Source/Runtime/Engine/Private/StreamableRenderAsset.cpp
--- a/Engine/Source/Runtime/Engine/Private/StreamableRenderAsset.cpp
+++ b/Engine/Source/Runtime/Engine/Private/StreamableRenderAsset.cpp
@@ -475,23 +475,18 @@ bool UStreamableRenderAsset::IsReadyForFinishDestroy()
 
 	if (PendingUpdate)
 	{
-		if (GEnableAssetUpdateFastCancels)
+		bool bCanBeReleased = false;
+		if (GEnableAssetUpdateAbandons && PendingUpdate->CanAbandon())
 		{
-			PendingUpdate->FastCancel();
+			IStreamingManager::Get().GetRenderAssetStreamingManager().OnRenderAssetUpdateAbandoned(this, PendingUpdate);
+			bCanBeReleased = true;
 		}
-		
-		// To avoid async tasks from timing-out the GC, we tick as Async to force completion if this is relevant.
-		// This could lead the asset from releasing the PendingUpdate, which will be deleted once the async task completes.
-		if (PendingUpdate->GetRelevantThread() == FRenderAssetUpdate::TT_Async)
+		else if(FRenderAssetUpdateHelper::TickRenderAssetUpdateForGarbageCollection(PendingUpdate))
 		{
-			PendingUpdate->Tick(FRenderAssetUpdate::TT_GameRunningAsync);
-		}
-		else
-		{
-			PendingUpdate->Tick(GIsThreadedRendering ? FRenderAssetUpdate::TT_None : FRenderAssetUpdate::TT_Render);
+			bCanBeReleased = true;
 		}
 
-		if (PendingUpdate->IsCompleted())
+		if (bCanBeReleased)
 		{
 			PendingUpdate.SafeRelease();
 		}
diff --git a/Engine/Source/Runtime/Engine/Private/Streaming/AbandonedRenderAssetUpdateManager.cpp b/Engine/Source/Runtime/Engine/Private/Streaming/AbandonedRenderAssetUpdateManager.cpp
new file mode 100644
--- /dev/null
+++ b/Engine/Source/Runtime/Engine/Private/Streaming/AbandonedRenderAssetUpdateManager.cpp
@@ -0,0 +1,70 @@
+// Copyright Epic Games, Inc. All Rights Reserved.
+
+#include "AbandonedRenderAssetUpdateManager.h"
+#include "RenderAssetUpdate.h"
+
+#include "Streaming/TextureStreamingHelpers.h"
+#include "Engine/StreamableRenderAsset.h"
+#include "ProfilingDebugging/CsvProfiler.h"
+#include "UObject/UObjectGlobals.h"
+
+CSV_DECLARE_CATEGORY_EXTERN(TextureStreaming);
+
+bool GTickAbandonedRenderAssetUpdatesOnStreamingUpdate = true;
+static FAutoConsoleVariableRef CVarTickAbandonedRenderAssetUpdatesOnStreamingUpdate(
+	TEXT("r.Streaming.TickAbandonedRenderAssetUpdatesOnStreamingUpdate"),
+	GTickAbandonedRenderAssetUpdatesOnStreamingUpdate,
+	TEXT("Tick abandoned render asset updates on every streaming update"),
+	ECVF_Default
+);
+
+FAbandonedRenderAssetUpdateManager::FAbandonedRenderAssetUpdateManager()
+{
+	FCoreUObjectDelegates::GetPostGarbageCollect().AddRaw(this, &FAbandonedRenderAssetUpdateManager::OnPostGarbageCollect);
+}
+
+FAbandonedRenderAssetUpdateManager::~FAbandonedRenderAssetUpdateManager()
+{
+	UE_CLOG(!AbandonedRenderAssetUpdates.IsEmpty(),
+		LogContentStreaming,
+		Warning,
+		TEXT("FAbandonedRenderAssetUpdateManager shutdown before processing %d render asset updates"), AbandonedRenderAssetUpdates.Num());
+	FCoreUObjectDelegates::GetPostGarbageCollect().RemoveAll(this);
+}
+
+void FAbandonedRenderAssetUpdateManager::OnAbandoned(UStreamableRenderAsset* OwningRenderAsset, TRefCountPtr<class FRenderAssetUpdate> RenderAssetUpdate)
+{
+	UE_CLOG(!IsInGameThread(), LogContentStreaming, Fatal, TEXT("FAbandonedRenderAssetUpdateManager::OnAbandoned is expected to only be called on the game thread"));
+	RenderAssetUpdate->OnAbandoned();
+	AbandonedRenderAssetUpdates.Add(RenderAssetUpdate);
+	UE_LOG(LogContentStreaming, Log, TEXT("FAbandonedRenderAssetUpdateManager Abandoned pending render asset update [Name:%s]"), *OwningRenderAsset->GetName());
+}
+
+void FAbandonedRenderAssetUpdateManager::OnPostGarbageCollect()
+{
+	TRACE_CPUPROFILER_EVENT_SCOPE(FAbandonedRenderAssetUpdateManager::OnPostGarbageCollect);
+	
+	if (GTickAbandonedRenderAssetUpdatesOnStreamingUpdate)
+	{
+		return;
+	}
+
+	TickAbandoned();	
+}
+
+void FAbandonedRenderAssetUpdateManager::TickAbandoned()
+{
+	TRACE_CPUPROFILER_EVENT_SCOPE(FAbandonedRenderAssetUpdateManager::TickAbandoned);
+	UE_CLOG(!IsInGameThread(), LogContentStreaming, Fatal, TEXT("FAbandonedRenderAssetUpdateManager::TickAbandoned is expected to only tick on the game thread"));
+	const int32 OutTotal = AbandonedRenderAssetUpdates.Num();
+	for (auto It = AbandonedRenderAssetUpdates.CreateIterator(); It; ++It)
+	{
+		if (FRenderAssetUpdateHelper::TickRenderAssetUpdateForGarbageCollection(*It))
+		{
+			It.RemoveCurrentSwap();
+		}
+	}
+	const int32 OutCompleted = OutTotal - AbandonedRenderAssetUpdates.Num();
+	CSV_CUSTOM_STAT(TextureStreaming, NumAbandonedRenderAssetUpdatesTotal, OutTotal, ECsvCustomStatOp::Set);
+	CSV_CUSTOM_STAT(TextureStreaming, NumAbandonedRenderAssetUpdatesCompleted, OutCompleted, ECsvCustomStatOp::Set);
+}
diff --git a/Engine/Source/Runtime/Engine/Private/Streaming/AbandonedRenderAssetUpdateManager.h b/Engine/Source/Runtime/Engine/Private/Streaming/AbandonedRenderAssetUpdateManager.h
new file mode 100644
--- /dev/null
+++ b/Engine/Source/Runtime/Engine/Private/Streaming/AbandonedRenderAssetUpdateManager.h
@@ -0,0 +1,35 @@
+// Copyright Epic Games, Inc. All Rights Reserved.
+
+/*=============================================================================
+AbandonedRenderAssetUpdateManager.h: Definitions of classes used for handling abandoned render asset updates
+=============================================================================*/
+
+#pragma once
+
+#include "CoreMinimal.h"
+
+class UStreamableRenderAsset;
+class FRenderAssetUpdate;
+
+extern bool GTickAbandonedRenderAssetUpdatesOnStreamingUpdate;
+
+// The streaming data of a level.
+class FAbandonedRenderAssetUpdateManager
+{
+public:
+	FAbandonedRenderAssetUpdateManager();
+	~FAbandonedRenderAssetUpdateManager();
+
+	/** Adds a render asset update to the abandoned list to be processed and deallocated independent of the streamable render asset */
+	void OnAbandoned(UStreamableRenderAsset* OwningRenderAsset, TRefCountPtr<FRenderAssetUpdate> RenderAssetUpdate);
+
+	/** Tick abandoned render asset updates, called post GC or on each streaming update (see GTickAbandonedRenderAssetUpdatesOnStreamingUpdate) */
+	void TickAbandoned();
+
+private:
+	/** Internal post garbage collection callbacks to trigger processing of abandoned render asset updates */
+	void OnPostGarbageCollect();
+
+private:
+	TArray < TRefCountPtr<FRenderAssetUpdate> > AbandonedRenderAssetUpdates;
+};
diff --git a/Engine/Source/Runtime/Engine/Private/Streaming/RenderAssetUpdate.cpp b/Engine/Source/Runtime/Engine/Private/Streaming/RenderAssetUpdate.cpp
--- a/Engine/Source/Runtime/Engine/Private/Streaming/RenderAssetUpdate.cpp
+++ b/Engine/Source/Runtime/Engine/Private/Streaming/RenderAssetUpdate.cpp
@@ -27,11 +27,11 @@ static TAutoConsoleVariable<int32> CVarStreamingStressTestExtraAsyncLatency(
 	TEXT("An extra latency in milliseconds for each async task when doing the stress test."),
 	ECVF_Cheat);
 
-ENGINE_API bool GEnableAssetUpdateFastCancels = false;
-static FAutoConsoleVariableRef CVarEnableAssetUpdateFastCancels(
-	TEXT("r.Streaming.EnableAssetUpdateFastCancels"),
-	GEnableAssetUpdateFastCancels,
-	TEXT("Enable asset update fast cancel support, enables calls to FRenderAssetUpdate::FastCancel which can then be used to optionally decouple Asset Update internals that stall/prevent objects from being GC'd"),
+ENGINE_API bool GEnableAssetUpdateAbandons = false;
+static FAutoConsoleVariableRef CVarEnableAssetUpdateAbandons(
+	TEXT("r.Streaming.EnableAssetUpdateAbandons"),
+	GEnableAssetUpdateAbandons,
+	TEXT("Enable asset update abandon support, enables calls to FRenderAssetUpdate::TryAbandon & FRenderAssetUpdate::OnAbandon which can then be used to optionally decouple Asset Update internals that stall/prevent objects from being GC'd"),
 	ECVF_Default
 );
 
@@ -134,6 +134,24 @@ void ResumeRenderAssetStreaming()
 	GRenderAssetStreamingSuspension.fetch_sub(1, std::memory_order_relaxed);
 }
 
+namespace FRenderAssetUpdateHelper
+{
+	bool TickRenderAssetUpdateForGarbageCollection(TRefCountPtr<FRenderAssetUpdate> RenderAssetUpdate)
+	{
+		// To avoid async tasks from timing-out the GC, we tick as Async to force completion if this is relevant.
+		// This could lead the asset from releasing the PendingUpdate, which will be deleted once the async task completes.
+		if (RenderAssetUpdate->GetRelevantThread() == FRenderAssetUpdate::TT_Async)
+		{
+			RenderAssetUpdate->Tick(FRenderAssetUpdate::TT_GameRunningAsync);
+		}
+		else
+		{
+			RenderAssetUpdate->Tick(GIsThreadedRendering ? FRenderAssetUpdate::TT_None : FRenderAssetUpdate::TT_Render);
+		}
+		return RenderAssetUpdate->IsCompleted();
+	}
+}
+
 FRenderAssetUpdate::FRenderAssetUpdate(const UStreamableRenderAsset* InAsset)
 	: ResourceState(InAsset->GetStreamableResourceState())
 	, CurrentFirstLODIdx(InAsset->GetStreamableResourceState().ResidentFirstLODIdx())
diff --git a/Engine/Source/Runtime/Engine/Private/Streaming/StreamingManagerTexture.cpp b/Engine/Source/Runtime/Engine/Private/Streaming/StreamingManagerTexture.cpp
--- a/Engine/Source/Runtime/Engine/Private/Streaming/StreamingManagerTexture.cpp
+++ b/Engine/Source/Runtime/Engine/Private/Streaming/StreamingManagerTexture.cpp
@@ -2066,6 +2066,11 @@ void FRenderAssetStreamingManager::UpdateResourceStreaming( float DeltaTime, boo
 
 	TickDeferredMipLevelChangeCallbacks();
 
+	if (GEnableAssetUpdateAbandons && GTickAbandonedRenderAssetUpdatesOnStreamingUpdate)
+	{
+		AbandonedRenderAssetUpdateManager.TickAbandoned();
+	}
+
 	if (bUseThreadingForPerf)
 	{
 		RenderAssetInstanceAsyncWork->StartBackgroundTask(GThreadPool);
@@ -2293,6 +2298,12 @@ void FRenderAssetStreamingManager::AddRenderedTextureStats(TMap<FString, FRender
 	}
 }
 
+void FRenderAssetStreamingManager::OnRenderAssetUpdateAbandoned(UStreamableRenderAsset* OwningRenderAsset, TRefCountPtr<class FRenderAssetUpdate>& RenderAssetUpdate)
+{
+	check(GEnableAssetUpdateAbandons);
+	AbandonedRenderAssetUpdateManager.OnAbandoned(OwningRenderAsset, RenderAssetUpdate);
+}
+
 #if !UE_BUILD_SHIPPING
 
 bool FRenderAssetStreamingManager::HandleDumpTextureStreamingStatsCommand( const TCHAR* Cmd, FOutputDevice& Ar )
diff --git a/Engine/Source/Runtime/Engine/Private/Streaming/StreamingManagerTexture.h b/Engine/Source/Runtime/Engine/Private/Streaming/StreamingManagerTexture.h
--- a/Engine/Source/Runtime/Engine/Private/Streaming/StreamingManagerTexture.h
+++ b/Engine/Source/Runtime/Engine/Private/Streaming/StreamingManagerTexture.h
@@ -8,6 +8,7 @@ TextureStreamingManager.h: Definitions of classes used for texture streaming.
 
 #include "CoreMinimal.h"
 #include "ContentStreaming.h"
+#include "AbandonedRenderAssetUpdateManager.h"
 #include "Streaming/TextureStreamingHelpers.h"
 #include "Streaming/StreamingTexture.h"
 #include "Streaming/LevelTextureManager.h"
@@ -184,6 +185,8 @@ struct FRenderAssetStreamingManager final : public IRenderAssetStreamingManager
 
 	void AddRenderedTextureStats(TMap<FString, FRenderedTextureStats>& InOutRenderedTextureStats) override;
 
+	void OnRenderAssetUpdateAbandoned(UStreamableRenderAsset* OwningRenderAsset, TRefCountPtr<class FRenderAssetUpdate>& RenderAssetUpdate) override;
+
 	/**
 	 * Mark the textures/meshes with a timestamp. They're about to lose their location-based heuristic and we don't want them to
 	 * start using LastRenderTime heuristic for a few seconds until they are garbage collected!
@@ -455,6 +458,8 @@ private:
 	FGraphEventRef StreamingRenderAssetsSyncEvent;
 	TArray<FStreamingRenderAsset>& GetStreamingRenderAssetsAsyncSafe();
 
+	FAbandonedRenderAssetUpdateManager AbandonedRenderAssetUpdateManager;
+
 	friend bool TrackRenderAssetEvent( FStreamingRenderAsset* StreamingRenderAsset, UStreamableRenderAsset* RenderAsset, bool bForceMipLevelsToBeResident, const FRenderAssetStreamingManager* Manager);
 };
 
diff --git a/Engine/Source/Runtime/Engine/Private/Streaming/Texture2DStreamIn_IO.cpp b/Engine/Source/Runtime/Engine/Private/Streaming/Texture2DStreamIn_IO.cpp
--- a/Engine/Source/Runtime/Engine/Private/Streaming/Texture2DStreamIn_IO.cpp
+++ b/Engine/Source/Runtime/Engine/Private/Streaming/Texture2DStreamIn_IO.cpp
@@ -23,11 +23,11 @@ CSV_DECLARE_CATEGORY_EXTERN(TextureStreaming);
 #endif
 
 #if !UE_BUILD_SHIPPING
-float GTextureStreamInFastCancelCallbackIOLatency = 0.0f;
+float GTextureStreamInAbandonCallbackIOLatency = 0.0f;
 static FAutoConsoleVariableRef CVarTextureStreamInFastCancelCallbackIOLatency(
-	TEXT("r.Streaming.TextureStreamInFastCancelCallbackIOLatency"),
-	GTextureStreamInFastCancelCallbackIOLatency,
-	TEXT("Texture stream in latency (in ms), only applied to the bulk request callkbacks when fast cancels are enabled, useful to induce a case where FastCancel is called while the IO request is in flight"),
+	TEXT("r.Streaming.TextureStreamInAbandonCallbackIOLatency"),
+	GTextureStreamInAbandonCallbackIOLatency,
+	TEXT("Texture stream in latency (in ms), only applied to the bulk request callkbacks when abandons are enabled, useful to induce a case where CanAbandon is called while the IO request is in flight"),
 	ECVF_Default);
 #endif // #if !UE_BUILD_SHIPPING
 
@@ -123,68 +123,42 @@ void FTexture2DStreamIn_IO::SetIORequests(const FContext& Context)
 
 	TaskSynchronization.Increment();
 
-	FBulkDataRequest::FCompletionCallback OnBulkDataRequestComplete = [this, InFlightBulkSize](FBulkDataBatchRequest::EStatus Status)
-	{
-		TaskSynchronization.Decrement();
-
-		if (Status != FBulkDataBatchRequest::EStatus::Ok && !bIsCancelled)
-		{
-			// If IO requests was cancelled but the streaming request wasn't, this is an IO error.
-			bFailedOnIOError = true;
-			MarkAsCancelled();
-		}
-#if !UE_BUILD_SHIPPING
-		// On some platforms the IO is too fast to test cancelation requests timing issues.
-		if (FRenderAssetStreamingSettings::ExtraIOLatency > 0 && TaskSynchronization.GetValue() == 0)
-		{
-			FPlatformProcess::Sleep(FRenderAssetStreamingSettings::ExtraIOLatency * .001f); // Slow down the streaming.
-		}
-#endif
-		if (bIsCancelled)
-		{
-			CSV_CUSTOM_STAT(TextureStreaming, CancelledMipsInMB, float(InFlightBulkSize) / (1024.0f * 1024.0f), ECsvCustomStatOp::Accumulate);
-		}
-		// The tick here is intended to schedule the success or cancel callback.
-		// Using TT_None ensure gets which could create a dead lock.
-		Tick(FTexture2DUpdate::TT_None);
-	};
-
-	if (!GEnableAssetUpdateFastCancels)
-	{
-		Batch.Issue(MoveTemp(OnBulkDataRequestComplete), BatchRequest);
-	}
-	else
-	{
-		// Explicitly add/release ref to keep this object alive as the owning object can release their references if FastCanceled.
-		AddRef();
-
-		Batch.Issue([this, OnBulkDataRequestComplete = MoveTemp(OnBulkDataRequestComplete)](FBulkDataBatchRequest::EStatus Status)
+	Batch.Issue([this, InFlightBulkSize](FBulkDataBatchRequest::EStatus Status)
 		{
 #if !UE_BUILD_SHIPPING
-			// This is separate from FRenderAssetStreamingSettings::ExtraIOLatency to make it easier to test fast cancels.
+			// This is separate from FRenderAssetStreamingSettings::ExtraIOLatency to make it easier to test abandons.
 			// FRenderAssetStreamingSettings::ExtraIOLatency is applied to many parts of the asset update that can in turn stall GC.
-			// Stalled GCs would not trigger the call to FastCancel, testing which is the purpose of this CVar.
-			if (GTextureStreamInFastCancelCallbackIOLatency > 0)
+			// Stalled GCs would not trigger the call to CanAbandon, testing which is the purpose of this CVar.
+			if (GTextureStreamInAbandonCallbackIOLatency > 0)
 			{
-				FPlatformProcess::Sleep(GTextureStreamInFastCancelCallbackIOLatency * .001f); // Slow down the streaming.
+				FPlatformProcess::Sleep(GTextureStreamInAbandonCallbackIOLatency * .001f); // Slow down the streaming.
 			}
 #endif // !UE_BUILD_SHIPPING
 
-			// Handle fast cancels with early out.
-			if (TaskSynchronization.GetValue() <= 0)
+			TaskSynchronization.Decrement();
+
+			if (Status != FBulkDataBatchRequest::EStatus::Ok && !bIsCancelled)
 			{
-				check(bIsCancelled);
-				UE_LOG(LogContentStreaming, Warning, TEXT("Texture stream-in request was FastCanceled"));
+				// If IO requests was cancelled but the streaming request wasn't, this is an IO error.
+				bFailedOnIOError = true;
+				MarkAsCancelled();
 			}
-			else
+#if !UE_BUILD_SHIPPING
+			// On some platforms the IO is too fast to test cancelation requests timing issues.
+			if (FRenderAssetStreamingSettings::ExtraIOLatency > 0 && TaskSynchronization.GetValue() == 0)
 			{
-				OnBulkDataRequestComplete(Status);
+				FPlatformProcess::Sleep(FRenderAssetStreamingSettings::ExtraIOLatency * .001f); // Slow down the streaming.
 			}
-
-			Release();
+#endif
+			if (bIsCancelled)
+			{
+				CSV_CUSTOM_STAT(TextureStreaming, CancelledMipsInMB, float(InFlightBulkSize) / (1024.0f * 1024.0f), ECsvCustomStatOp::Accumulate);
+			}
+			// The tick here is intended to schedule the success or cancel callback.
+			// Using TT_None ensure gets which could create a dead lock.
+			Tick(FTexture2DUpdate::TT_None);
 		},
 		BatchRequest);
-	}
 }
 
 void FTexture2DStreamIn_IO::CancelIORequests()
@@ -230,14 +204,12 @@ void FTexture2DStreamIn_IO::Abort()
 	}
 }
 
-void FTexture2DStreamIn_IO::FastCancel()
+bool FTexture2DStreamIn_IO::CanAbandon() const
 {
-	check(GEnableAssetUpdateFastCancels);
-	if (!IsCompleted() && IsCancelled() && BatchRequest.IsPending() && TaskSynchronization.GetValue() > 0)
-	{
-		TaskSynchronization.Set(0);
-		UE_LOG(LogContentStreaming, Warning, TEXT("Decoupled FTexture2DStreamIn_IO from pending IO request in order to unblock GC"));
-	}
+	check(GEnableAssetUpdateAbandons);
+	// Abandon if we have any in flight IO.
+	return (FTexture2DStreamIn::CanAbandon()) 
+			|| (!IsCompleted() && IsCancelled() && BatchRequest.IsPending() && TaskSynchronization.GetValue() > 0);
 }
 
 void FTexture2DStreamIn_IO::FCancelIORequestsTask::DoWork()
diff --git a/Engine/Source/Runtime/Engine/Private/Streaming/Texture2DStreamIn_IO.h b/Engine/Source/Runtime/Engine/Private/Streaming/Texture2DStreamIn_IO.h
--- a/Engine/Source/Runtime/Engine/Private/Streaming/Texture2DStreamIn_IO.h
+++ b/Engine/Source/Runtime/Engine/Private/Streaming/Texture2DStreamIn_IO.h
@@ -39,7 +39,7 @@ protected:
 	// Start an async task to cancel pending IO requests.
 	void Abort() override;
 
-	void FastCancel() override;
+	bool CanAbandon() const override;
 
 private:
 
diff --git a/Engine/Source/Runtime/Engine/Private/Streaming/Texture2DUpdate.cpp b/Engine/Source/Runtime/Engine/Private/Streaming/Texture2DUpdate.cpp
--- a/Engine/Source/Runtime/Engine/Private/Streaming/Texture2DUpdate.cpp
+++ b/Engine/Source/Runtime/Engine/Private/Streaming/Texture2DUpdate.cpp
@@ -21,23 +21,25 @@ FTexture2DUpdateContext::FTexture2DUpdateContext(const UTexture2D* InTexture, ET
 	: Texture(InTexture)
 	, CurrentThread(InCurrentThread)
 {
-	check(InTexture);
-	checkSlow(InCurrentThread != FTexture2DUpdate::TT_Render || IsInRenderingThread());
-	Resource = Texture && Texture->GetResource() ? const_cast<UTexture2D*>(Texture)->GetResource()->GetTexture2DResource() : nullptr;
-	if (Resource)
+	if (InTexture)
 	{
-		MipsView = Resource->GetPlatformMipsView();
+		checkSlow(InCurrentThread != FTexture2DUpdate::TT_Render || IsInRenderingThread());
+		Resource = Texture->GetResource() ? const_cast<UTexture2D*>(Texture)->GetResource()->GetTexture2DResource() : nullptr;
+		if (Resource)
+		{
+			MipsView = Resource->GetPlatformMipsView();
+		}
 	}
 }
 
 FTexture2DUpdateContext::FTexture2DUpdateContext(const UStreamableRenderAsset* InTexture, EThreadType InCurrentThread)
-	: FTexture2DUpdateContext(CastChecked<UTexture2D>(InTexture), InCurrentThread)
+	: FTexture2DUpdateContext(Cast<UTexture2D>(InTexture), InCurrentThread)
 {}
 
 FTexture2DUpdate::FTexture2DUpdate(UTexture2D* InTexture) 
 	: TRenderAssetUpdate<FTexture2DUpdateContext>(InTexture)
 {
-	if (!InTexture->GetResource())
+	if (!InTexture || !InTexture->GetResource())
 	{
 		bIsCancelled = true;
 	}
diff --git a/Engine/Source/Runtime/Engine/Public/ContentStreaming.h b/Engine/Source/Runtime/Engine/Public/ContentStreaming.h
--- a/Engine/Source/Runtime/Engine/Public/ContentStreaming.h
+++ b/Engine/Source/Runtime/Engine/Public/ContentStreaming.h
@@ -42,6 +42,7 @@ class UTexture2D;
 class UStaticMesh;
 class USkeletalMesh;
 class UStreamableRenderAsset;
+class FRenderAssetUpdate;
 class FSoundSource;
 class FAudioStreamCacheMemoryHandle;
 struct FWaveInstance;
@@ -496,6 +497,8 @@ struct IRenderAssetStreamingManager : public IStreamingManager
 	virtual void MarkMountedStateDirty(FIoFilenameHash FilenameHash) = 0;
 
 	virtual void AddRenderedTextureStats(TMap<FString, FRenderedTextureStats>& InOutRenderedTextureAssets) = 0;
+
+	virtual void OnRenderAssetUpdateAbandoned(UStreamableRenderAsset* OwningRenderAsset, TRefCountPtr<FRenderAssetUpdate>& RenderAssetUpdate) = 0;
 };
 
 enum class EAudioChunkLoadResult : uint8
diff --git a/Engine/Source/Runtime/Engine/Public/RenderAssetUpdate.h b/Engine/Source/Runtime/Engine/Public/RenderAssetUpdate.h
--- a/Engine/Source/Runtime/Engine/Public/RenderAssetUpdate.h
+++ b/Engine/Source/Runtime/Engine/Public/RenderAssetUpdate.h
@@ -21,7 +21,13 @@ class UStreamableRenderAsset;
 
 ENGINE_API bool IsAssetStreamingSuspended();
 
-extern ENGINE_API bool GEnableAssetUpdateFastCancels;
+extern ENGINE_API bool GEnableAssetUpdateAbandons;
+
+class FRenderAssetUpdate;
+namespace FRenderAssetUpdateHelper
+{
+	bool TickRenderAssetUpdateForGarbageCollection(TRefCountPtr<FRenderAssetUpdate> RenderAssetUpdate);
+}
 
 /**
 * This class provides a framework for loading and unloading the texture/mesh LODs.
@@ -78,11 +84,22 @@ public:
 	}
 
 	/**
-	* Allows the update to decouple any internal tasks it is awaiting on in preparation to be garbage collected.
+	* Allows the update to be decoupled from its UStreamableRenderAsset so it can be independently garbage collected.
 	* This is called when we query IsReadyForFinishDestroy on the UStreamableRenderAsset.
+	* Returns true if the asset update can handle being abandoned.
+	*/
+	virtual bool CanAbandon() const
+	{
+		return false;
+	}
+
+	/**
+	* Called when the asset update is being abandoned. Execute any task abandon specific cleanup here.
 	*/
-	virtual void FastCancel()
-	{	
+	virtual void OnAbandoned()
+	{
+		check(CanAbandon());
+		StreamableAsset = nullptr;
 	}
 
 	/** Returns whether the task was aborted through Abort() or cancelled.  */
