Improved streamable manager error reporting.
- Added additional context information
- Ensured that all cases will report warnings
- Removed duplicated context information in error logs.

diff --git a/Engine/Source/Runtime/Core/Private/Experimental/UnifiedError/UnifiedError.cpp b/Engine/Source/Runtime/Core/Private/Experimental/UnifiedError/UnifiedError.cpp
--- a/Engine/Source/Runtime/Core/Private/Experimental/UnifiedError/UnifiedError.cpp
+++ b/Engine/Source/Runtime/Core/Private/Experimental/UnifiedError/UnifiedError.cpp
@@ -96,26 +96,28 @@ FAnsiString FError::GetModuleIdString() const
 
 FText FError::GetErrorMessage(bool bIncludeContext) const
 {
-	FText FormatText = FText::Join(FText::FromString(": "), NSLOCTEXT("UnifiedError", "DefaultLogPrepend", "{Root/ModuleIdString}.{Root/ErrorCodeString}"), GetFormatErrorText());
 
+	EDetailFilter LogDetailFilter = EDetailFilter::IncludeInLogMessage;
 	if (bIncludeContext)
 	{
-		TArray<FText> AdditionalContext;
-		AdditionalContext.Add(FormatText);
-
-		const IErrorDetails* DetailsIt = ErrorDetails.GetReference();
-		while (DetailsIt != nullptr)
+		LogDetailFilter |= EDetailFilter::IncludeInContextLogMessage;
+	}
+	
+	TArray<FText> AdditionalContext;
+	AdditionalContext.Add(FText::Join(FText::FromString(": "), NSLOCTEXT("UnifiedError", "DefaultLogPrepend", "{Root/ModuleIdString}.{Root/ErrorCodeString}"), GetFormatErrorText()));
+	
+	const IErrorDetails* DetailsIt = ErrorDetails.GetReference();
+	while (DetailsIt != nullptr)
+	{
+		if (DetailsIt->ShouldInclude(EDetailFilter::IncludeInContextLogMessage))
 		{
-			if (DetailsIt->ShouldInclude(EDetailFilter::IncludeInContextLogMessage))
-			{
-				AdditionalContext.Add(FText::FromString(FUtf8String::Printf("(%s:{%s})", DetailsIt->GetErrorDetailsTypeNameForLog().GetData(), DetailsIt->GetErrorDetailsTypeNameForLog().GetData())));
-			}
-			const IErrorDetails* Next = DetailsIt->GetInnerErrorDetails().GetReference(); 
-			DetailsIt = Next;
+			AdditionalContext.Add(FText::FromString(FUtf8String::Printf("(%s:{%s})", DetailsIt->GetErrorDetailsTypeNameForLog().GetData(), DetailsIt->GetErrorDetailsTypeNameForLog().GetData())));
 		}
-		FormatText = FText::Join(FText::FromString(", "), AdditionalContext);
+		const IErrorDetails* Next = DetailsIt->GetInnerErrorDetails().GetReference(); 
+		DetailsIt = Next;
 	}
-
+	FText FormatText = FText::Join(FText::FromString(", "), AdditionalContext);
+	
 #if USE_LOCALIZED_STRUCTURED_LOG_FOR_FERRORMESSAGE
 	UE::FLogTemplateOptions TemplateOptions;
 	TemplateOptions.bAllowSubObjectReferences = true;
diff --git a/Engine/Source/Runtime/Core/Public/Experimental/UnifiedError/UnifiedError.h b/Engine/Source/Runtime/Core/Public/Experimental/UnifiedError/UnifiedError.h
--- a/Engine/Source/Runtime/Core/Public/Experimental/UnifiedError/UnifiedError.h
+++ b/Engine/Source/Runtime/Core/Public/Experimental/UnifiedError/UnifiedError.h
@@ -413,15 +413,6 @@ namespace UE::UnifiedError
 			SerializeForLog(Writer, ErrorDetail);
 		}
 
-		virtual const FText GetErrorFormatString(const FError& Error) const override
-		{
-			if (EnumHasAnyFlags(DetailFilterMask, EDetailFilter::IncludeInLogMessage))
-			{
-				return FText::Join(FText::FromString(" "), FDynamicErrorDetails::GetErrorFormatString(Error), FText::FromString(FUtf8String::Printf("{%s}", GetErrorDetailsTypeNameForLog().GetData())));
-			}
-			return FDynamicErrorDetails::GetErrorFormatString(Error);
-		}
-
 		const T& GetValue() const
 		{
 			return ErrorDetail;
diff --git a/Engine/Source/Runtime/Engine/Private/StreamableManager.cpp b/Engine/Source/Runtime/Engine/Private/StreamableManager.cpp
--- a/Engine/Source/Runtime/Engine/Private/StreamableManager.cpp
+++ b/Engine/Source/Runtime/Engine/Private/StreamableManager.cpp
@@ -67,6 +67,14 @@ static FAutoConsoleVariableRef CVarLogAsyncLoadErrorsAsErrors(
 	ECVF_Default
 );
 
+static bool GStreamableLogStreamableFailuresAsErrors = false;
+static FAutoConsoleVariableRef CVarLogStreamableFailuresAsErrors(
+		TEXT("s.LogStreamableFailuresAsErrors"),
+	GStreamableLogStreamableFailuresAsErrors,
+		TEXT("Log all streamable failures as errors."),
+		ECVF_Default
+	);
+
 #if !UE_BUILD_SHIPPING
 static FString GStreamableErrorOnRequest = "";
 static FAutoConsoleVariableRef CVarStreamableErrorOnRequest(
@@ -1082,11 +1090,18 @@ FStreamableHandle::~FStreamableHandle()
 
 void FStreamableHandle::CompleteLoad()
 {
-	// Only complete if it's still active
 	if (Error.IsSet())
 	{
-		UE_LOG(LogStreamableManager, Verbose, TEXT("FStreamableHandle::CompleteLoad failed load %s, with error %s!"), *DebugName, *Error->GetErrorMessage().ToString());
+		if (GStreamableLogStreamableFailuresAsErrors)
+		{
+			UE_LOG(LogStreamableManager, Error, TEXT("Encounted error: %s"), *Error->GetErrorMessage(true).ToString());
+		}
+		else
+		{
+			UE_LOG(LogStreamableManager, Verbose, TEXT("Encounted error: %s"), *Error->GetErrorMessage(true).ToString());
+		}
 	}
+	// Only complete if it's still active
 	if (IsActive())
 	{
 		bLoadCompleted = true;
@@ -1225,8 +1240,12 @@ void FStreamableHandle::SetError(const UE::UnifiedError::FError& InError)
 {
 	if (!Error)
 	{
-		SetErrorOnParents(InError);
 		Error = InError;
+		// add some additional context for this specific request 
+		UE::UnifiedError::StreamableManager::FStreamableManagerRequestContext RequestContext = { DebugName, RequestedAssets };
+		Error->PushErrorContext(MoveTemp(RequestContext), UE::UnifiedError::EDetailFilter::Default);
+
+		SetErrorOnParents(Error.GetValue());
 	}
 }
 
@@ -1262,9 +1281,6 @@ void FStreamableHandle::AsyncLoadCallbackWrapper(const FName& PackageName, UPack
 			UE::UnifiedError::StreamableManager::FStreamableManagerAdditionalContext MoreContext = {PackageName.ToString(), TargetName.ToString()};
 			MaybeError->PushErrorContext(MoveTemp(MoreContext), UE::UnifiedError::EDetailFilter::Default);
 
-			UE::UnifiedError::StreamableManager::FStreamableManagerRequestContext RequestContext = {DebugName, RequestedAssets};
-			MaybeError->PushErrorContext(MoveTemp(RequestContext), UE::UnifiedError::EDetailFilter::Default);
-
 			SetError(*MaybeError);
 		}
 
