Speculative fix for MutliProductBuild cooker assertion: FileNameMapAddr in FPackageDatas::FindOrAddPackageDatas, on a generated StreamingObject package.

Remove the addition of the optionally mounted external streaming objects to the predicted list of world packages since it affects package contents and causes the two cook sessions to disagree about the .umap or .uasset extension and triggers a cooker assertion in the second cook that finds the discrepency.

#rnx

diff --git a/Engine/Source/Runtime/Engine/Private/WorldPartition/WorldPartitionStreamingDescriptor.cpp b/Engine/Source/Runtime/Engine/Private/WorldPartition/WorldPartitionStreamingDescriptor.cpp
--- a/Engine/Source/Runtime/Engine/Private/WorldPartition/WorldPartitionStreamingDescriptor.cpp
+++ b/Engine/Source/Runtime/Engine/Private/WorldPartition/WorldPartitionStreamingDescriptor.cpp
@@ -318,7 +318,10 @@ TMap<FName, TArray<FName>> FStreamingDescriptor::GenerateWorldPluginPackagesList
 			FName PluginName = FName(FPathViews::GetMountPointNameFromPath(ExternalStreamingObject));			
 			FName PackageName = FName(ExternalStreamingObject);
 
-			PluginPackages.FindOrAdd(PluginName).Add(PackageName);
+			// Temporary Workaround to avoid cooker crash: Remove the addition of these packages since it affects
+			// package contents and causes the two cook sessions to disagree about the .umap or .uasset extension
+			// and triggers a cooker assertion in the second cook that finds the discrepency
+			// PluginPackages.FindOrAdd(PluginName).Add(PackageName);
 		}
 	}
 
