Fixed FScene leak by making UWorld::InitWorld call Scene->Release and clear the Scene pointer. The new behavior can be reverted using the cvar s.World.ShouldInitWorldReleaseScene.

diff --git a/Engine/Source/Runtime/Engine/Classes/Engine/World.h b/Engine/Source/Runtime/Engine/Classes/Engine/World.h
--- a/Engine/Source/Runtime/Engine/Classes/Engine/World.h
+++ b/Engine/Source/Runtime/Engine/Classes/Engine/World.h
@@ -3545,6 +3545,9 @@ public:
 	UE_API bool Exec( UWorld* InWorld, const TCHAR* Cmd, FOutputDevice& Ar=*GLog );
 
 private:
+	/** Internal helper to release the Scene and clear the pointer. */
+	void ReleaseScene();
+
 	/** Internal version of CleanupWorld. */
 	UE_API void CleanupWorldInternal(bool bSessionEnded, bool bCleanupResources, bool bWorldChanged);
 
diff --git a/Engine/Source/Runtime/Engine/Private/World.cpp b/Engine/Source/Runtime/Engine/Private/World.cpp
--- a/Engine/Source/Runtime/Engine/Private/World.cpp
+++ b/Engine/Source/Runtime/Engine/Private/World.cpp
@@ -259,6 +259,13 @@ FAutoConsoleVariableRef CVarCreateStaticLevelCollection(TEXT("s.World.CreateStat
 	TEXT("When set, create a separate level collection for static streaming levels that will not be duplicated by DuplicateRequestedLevels.\n")
 	TEXT("If this is 0, static streaming levels will be part of the main DynamicSourceLevels collection."));
 
+static bool GShouldInitWorldReleaseScene = true;
+FAutoConsoleVariableRef CVarShouldInitWorldReleaseScene(
+	TEXT("s.World.ShouldInitWorldReleaseScene"), 
+	GShouldInitWorldReleaseScene,
+	TEXT("Set to false to revert to earlier behavior (causes leaking Scene and associated resources in certain situations)."),
+	ECVF_Default);
+
 // This cvar activates detection of incorrectly unregistered components after a world is cleaned-up.
 // This is important since components left as registered/renderstatecreated after their world is cleaned up will 
 // crash in FGlobalComponentRecreateRenderStateContext unless a GC happens before the recreate.
@@ -2336,6 +2343,12 @@ void UWorld::InitWorld(const InitializationValues IVS)
 		return;
 	}
 
+	// This condition has to match the path that leads to GetRendererModule().AllocateScene below as that clobbers the Scene pointer
+	// However, we must call this before starting to create subsystem uobjects and such.
+	if (IVS.bInitializeScenes && GShouldInitWorldReleaseScene)
+	{
+		ReleaseScene();
+	}
 	// Reset flags in case of world reuse
 	bIsLevelStreamingFrozen = false;
 	bShouldForceUnloadStreamingLevels = false;
@@ -6154,6 +6167,17 @@ void ValidateComponentsRegistration(UWorld* World)
 #endif
 }
 
+void UWorld::ReleaseScene()
+{
+	if (Scene)
+	{
+		Scene->Release();
+		IRendererModule& RendererModule = GetRendererModule();
+		RendererModule.RemoveScene(Scene);
+		Scene = nullptr;
+	}
+}
+
 void UWorld::CleanupWorldInternal(bool bSessionEnded, bool bCleanupResources, bool bWorldChanged)
 {
 	TGuardValue<bool> IsBeingCleanedUp(bIsBeingCleanedUp, true);
@@ -9982,15 +10006,14 @@ void UWorld::RecreateScene(ERHIFeatureLevel::Type InFeatureLevel, bool bBroadcas
 			Scene->SetFXSystem(nullptr);
 		}
 
-		Scene->Release();
-		IRendererModule& RendererModule = GetRendererModule();
-		RendererModule.RemoveScene(Scene);
+		ReleaseScene();
 
 		if (bBroadcastChange)
 		{
 			FRenderResource::ChangeFeatureLevel(InFeatureLevel);
 		}
 
+		IRendererModule& RendererModule = GetRendererModule();
 		RendererModule.AllocateScene(this, bRequiresHitProxies, bCreateFXSystem, InFeatureLevel);
 
 		for (ULevel* Level : Levels)
