[HWRT] Fix FInstancedStaticMeshSceneProxy::GetDynamicRayTracingInstances(...) not using ray tracing proxy mesh batches.

#[robomerge] 5.7
#[robomerge][partner] megalights

diff --git a/Engine/Source/Runtime/Engine/Classes/Engine/InstancedStaticMesh.h b/Engine/Source/Runtime/Engine/Classes/Engine/InstancedStaticMesh.h
--- a/Engine/Source/Runtime/Engine/Classes/Engine/InstancedStaticMesh.h
+++ b/Engine/Source/Runtime/Engine/Classes/Engine/InstancedStaticMesh.h
@@ -418,6 +418,9 @@ public:
 	/** LOD render data from the static mesh. */
 	FStaticMeshLODResourcesArray& LODModels;
 
+	TIndirectArray<FInstancedStaticMeshVertexFactory> RayTracingProxyVertexFactories;
+	const FStaticMeshRayTracingProxyLODArray* RayTracingProxyLODs = nullptr;
+
 	/** Feature level used when creating instance data */
 	ERHIFeatureLevel::Type FeatureLevel;
 
@@ -499,6 +502,9 @@ public:
 	/** Sets up a wireframe FMeshBatch for a specific LOD. */
 	ENGINE_API virtual bool GetWireframeMeshElement(int32 LODIndex, int32 BatchIndex, const FMaterialRenderProxy* WireframeRenderProxy, uint8 InDepthPriorityGroup, bool bAllowPreCulledIndices, FMeshBatch& OutMeshBatch) const override;
 
+	/** Sets up a ray tracing FMeshBatch for a specific LOD. */
+	ENGINE_API virtual bool GetRayTracingMeshElement(int32 LODIndex, int32 BatchIndex, int32 SectionIndex, uint8 InDepthPriorityGroup, FMeshBatch& OutMeshBatch) const;
+
 	ENGINE_API virtual void GetDistanceFieldAtlasData(const FDistanceFieldVolumeData*& OutDistanceFieldData, float& SelfShadowBias) const override;
 
 	/**
@@ -565,6 +571,7 @@ protected:
 
 	/** Common path for the Get*MeshElement functions */
 	ENGINE_API void SetupInstancedMeshBatch(int32 LODIndex, int32 BatchIndex, FMeshBatch& OutMeshBatch) const;
+	ENGINE_API void SetupRayTracingInstancedMeshBatch(int32 LODIndex, int32 BatchIndex, FMeshBatch& OutMeshBatch) const;
 
 	/** Untransformed bounds of the static mesh */
 	FBoxSphereBounds StaticMeshBounds;
@@ -572,10 +579,14 @@ private:
 
 	void SetupProxy(const FInstancedStaticMeshSceneProxyDesc& InProxyDesc);
 	bool CanSetupInstancedMeshBatch(int32 LODIndex) const;
+	bool CanSetupRayTracingInstancedMeshBatch(int32 LODIndex) const;
 
 	/** Stores a loose uniform buffer per LOD, used for static view relevance. */
 	TMap<uint32, FInstancedStaticMeshVFLooseUniformShaderParametersRef> LODLooseUniformBuffers;
 
+	/** Stores a loose uniform buffer used for ray tracing. */
+	FInstancedStaticMeshVFLooseUniformShaderParametersRef RayTracingLooseUniformBuffer;
+
 	TSharedPtr<FInstanceDataSceneProxy, ESPMode::ThreadSafe> InstanceDataSceneProxy; 
 };
 
diff --git a/Engine/Source/Runtime/Engine/Private/InstancedStaticMesh.cpp b/Engine/Source/Runtime/Engine/Private/InstancedStaticMesh.cpp
--- a/Engine/Source/Runtime/Engine/Private/InstancedStaticMesh.cpp
+++ b/Engine/Source/Runtime/Engine/Private/InstancedStaticMesh.cpp
@@ -948,6 +948,11 @@ FInstancedStaticMeshRenderData::FInstancedStaticMeshRenderData(const FInstancedS
 	, LODModels(InDesc->GetStaticMesh()->GetRenderData()->LODResources)
 	, FeatureLevel(InFeatureLevel)
 {
+	if (IsRayTracingAllowed() && InDesc->GetStaticMesh()->bSupportRayTracing && InDesc->GetStaticMesh()->GetRenderData()->RayTracingProxy)
+	{
+		RayTracingProxyLODs = &InDesc->GetStaticMesh()->GetRenderData()->RayTracingProxy->LODs;
+	}
+
 	// Allocate the vertex factories for each LOD
 	InitVertexFactories();
 	RegisterSpeedTreeWind(InDesc);
@@ -968,6 +973,11 @@ void FInstancedStaticMeshRenderData::ReleaseResources(FSceneInterface* Scene, co
 	{
 		VertexFactories[LODIndex].ReleaseResource();
 	}
+
+	for (int32 LODIndex = 0; LODIndex < RayTracingProxyVertexFactories.Num(); LODIndex++)
+	{
+		RayTracingProxyVertexFactories[LODIndex].ReleaseResource();
+	}
 }
 
 void FInstancedStaticMeshVertexFactory::InitInstancedStaticMeshVertexFactoryComponents(
@@ -1033,6 +1043,19 @@ void FInstancedStaticMeshRenderData::BindBuffersToVertexFactories(FRHICommandLis
 		VertexFactory.SetData(RHICmdList, Data, InstanceBuffer ? &InstanceData : nullptr);
 		VertexFactory.InitResource(RHICmdList);
 	}
+
+	for (int32 LODIndex = 0; LODIndex < RayTracingProxyVertexFactories.Num(); LODIndex++)
+	{
+		const FStaticMeshRayTracingProxyLOD* RenderData = &(*RayTracingProxyLODs)[LODIndex];
+
+		FInstancedStaticMeshDataType InstanceData;
+		FInstancedStaticMeshVertexFactory::FDataType Data;
+		// Assign to the vertex factory for this LOD.
+		FInstancedStaticMeshVertexFactory& VertexFactory = RayTracingProxyVertexFactories[LODIndex];
+		FInstancedStaticMeshVertexFactory::InitInstancedStaticMeshVertexFactoryComponents(*RenderData->VertexBuffers, /*ColorVertexBuffer*/ nullptr, InstanceBuffer, &VertexFactory, LightMapCoordinateIndex, bRHISupportsManualVertexFetch, Data, InstanceData);
+		VertexFactory.SetData(RHICmdList, Data, InstanceBuffer ? &InstanceData : nullptr);
+		VertexFactory.InitResource(RHICmdList);
+	}
 }
 
 void FInstancedStaticMeshRenderData::InitVertexFactories()
@@ -1042,6 +1065,14 @@ void FInstancedStaticMeshRenderData::InitVertexFactories()
 	{
 		VertexFactories.Add(new FInstancedStaticMeshVertexFactory(FeatureLevel));
 	}
+
+	if (RayTracingProxyLODs)
+	{
+		for (int32 LODIndex = 0; LODIndex < RayTracingProxyLODs->Num(); LODIndex++)
+		{
+			RayTracingProxyVertexFactories.Add(new FInstancedStaticMeshVertexFactory(FeatureLevel));
+		}
+	}
 }
 
 void FInstancedStaticMeshRenderData::RegisterSpeedTreeWind(const FInstancedStaticMeshSceneProxyDesc* InProxyDesc)
@@ -1566,6 +1597,11 @@ bool FInstancedStaticMeshSceneProxy::CanSetupInstancedMeshBatch(int32 LODIndex)
 	return LODLooseUniformBuffers.Contains(LODIndex) && LODIndex < InstancedRenderData.VertexFactories.Num();
 }
 
+bool FInstancedStaticMeshSceneProxy::CanSetupRayTracingInstancedMeshBatch(int32 LODIndex) const
+{
+	return LODIndex < InstancedRenderData.RayTracingProxyVertexFactories.Num();
+}
+
 void FInstancedStaticMeshSceneProxy::CreateRenderThreadResources(FRHICommandListBase& RHICmdList)
 {
 	check(InstanceDataSceneProxy);
@@ -1578,6 +1614,10 @@ void FInstancedStaticMeshSceneProxy::CreateRenderThreadResources(FRHICommandList
 		FInstancedStaticMeshVFLooseUniformShaderParametersRef LooseUniformBuffer = CreateLooseUniformBuffer(nullptr, &UserData_AllInstances, 0, i, EUniformBufferUsage::UniformBuffer_MultiFrame);
 		LODLooseUniformBuffers.Add(i, LooseUniformBuffer);
 	}
+	
+	// currently we setup loose uniform buffer for ray tracing based on InstancedLODIndex = 0
+	// since currently non of the uniforms depend on the LOD index when InstancedLODRange = 0
+	RayTracingLooseUniformBuffer = CreateLooseUniformBuffer(nullptr, &UserData_AllInstances, /*InstancedLODRange*/ 0, /*InstancedLODIndex*/ 0, EUniformBufferUsage::UniformBuffer_MultiFrame);
 }
 
 void FInstancedStaticMeshSceneProxy::DestroyRenderThreadResources()
@@ -1630,6 +1670,36 @@ void FInstancedStaticMeshSceneProxy::SetupInstancedMeshBatch(int32 LODIndex, int
 	BatchElement0.NumInstances = GetInstanceDataHeader().NumInstances;
 }
 
+void FInstancedStaticMeshSceneProxy::SetupRayTracingInstancedMeshBatch(int32 LODIndex, int32 BatchIndex, FMeshBatch& OutMeshBatch) const
+{
+	OutMeshBatch.VertexFactory = &InstancedRenderData.RayTracingProxyVertexFactories[LODIndex];
+
+	FMeshBatchElement& BatchElement0 = OutMeshBatch.Elements[0];
+	BatchElement0.UserData = (void*)&UserData_AllInstances;
+	BatchElement0.bUserDataIsColorVertexBuffer = false;
+	BatchElement0.InstancedLODIndex = LODIndex;
+	BatchElement0.UserIndex = 0;
+	BatchElement0.PrimitiveUniformBuffer = GetUniformBuffer();
+	BatchElement0.LooseParametersUniformBuffer = RayTracingLooseUniformBuffer;
+	BatchElement0.bForceInstanceCulling = true; // force ISM through Generic path even for a single instance cases
+
+	if (OutMeshBatch.MaterialRenderProxy)
+	{
+		// If the material on the mesh batch is translucent, then preserve the instance draw order to prevent flickering
+		// TODO: This should use depth sorting of instances instead, once that is implemented for GPU Scene instances
+		const ERHIFeatureLevel::Type FeatureLevel = GetScene().GetFeatureLevel();
+
+		// NOTE: For now, this feature is not supported for mobile platforms
+		if (FeatureLevel > ERHIFeatureLevel::ES3_1)
+		{
+			const FMaterial& Material = OutMeshBatch.MaterialRenderProxy->GetIncompleteMaterialWithFallback(FeatureLevel);
+			BatchElement0.bPreserveInstanceOrder = IsTranslucentBlendMode(Material);
+		}
+	}
+	BatchElement0.bFetchInstanceCountFromScene = CVarISMFetchInstanceCountFromScene.GetValueOnRenderThread() != 0;
+	BatchElement0.NumInstances = GetInstanceDataHeader().NumInstances;
+}
+
 void FInstancedStaticMeshSceneProxy::GetLightRelevance(const FLightSceneProxy* LightSceneProxy, bool& bDynamic, bool& bRelevant, bool& bLightMapped, bool& bShadowMapped) const
 {
 	FStaticMeshSceneProxy::GetLightRelevance(LightSceneProxy, bDynamic, bRelevant, bLightMapped, bShadowMapped);
@@ -1672,6 +1742,16 @@ bool FInstancedStaticMeshSceneProxy::GetWireframeMeshElement(int32 LODIndex, int
 	return false;
 }
 
+bool FInstancedStaticMeshSceneProxy::GetRayTracingMeshElement(int32 LODIndex, int32 BatchIndex, int32 SectionIndex, uint8 InDepthPriorityGroup, FMeshBatch& OutMeshBatch) const
+{
+	if (CanSetupRayTracingInstancedMeshBatch(LODIndex) && FStaticMeshSceneProxy::GetRayTracingMeshElement(LODIndex, BatchIndex, SectionIndex, InDepthPriorityGroup, OutMeshBatch))
+	{
+		SetupRayTracingInstancedMeshBatch(LODIndex, BatchIndex, OutMeshBatch);
+		return true;
+	}
+	return false;
+}
+
 void FInstancedStaticMeshSceneProxy::GetDistanceFieldAtlasData(const FDistanceFieldVolumeData*& OutDistanceFieldData, float& SelfShadowBias) const
 {
 	FStaticMeshSceneProxy::GetDistanceFieldAtlasData(OutDistanceFieldData, SelfShadowBias);
@@ -1831,61 +1911,57 @@ void FInstancedStaticMeshSceneProxy::GetDynamicRayTracingInstances(FRayTracingIn
 		float DistanceToView = 0;
 	};
 
-	//setup a 'template' for the instance first, so we aren't duplicating work
-	//#dxr_todo: when multiple LODs are used, template needs to be an array of templates, probably best initialized on-demand via a lambda
-	FRayTracingInstance RayTracingInstanceTemplate;
-	RayTracingInstanceTemplate.Geometry = RayTracingGeometry;
-
-	FRayTracingInstance RayTracingWPOInstanceTemplate;  //template for evaluating the WPO instances into the world
-	FRayTracingInstance RayTracingWPODynamicTemplate;   //template for simulating the WPO instances
-
-	// Which index holds the reference to the particular simulated instance
-	TArray<uint32> ActiveInstances;
-
-	// Visible instances
-	TArray<FVisibleInstance> VisibleInstances;
-	TArray<FRayTracingInstance> RayTracingInstances;
-
-	const uint32 RequestedSimulatedInstances = CVarRayTracingSimulatedInstanceCount.GetValueOnRenderThread();
-	const uint32 SimulatedInstances = FMath::Min(RequestedSimulatedInstances == -1 ? InstanceCount : FMath::Clamp(RequestedSimulatedInstances, 1u, InstanceCount), MaxSimulatedInstances);
+	TArray<FMeshBatch> RayTracingMeshBatches;
 
-	if (bDynamicRayTracingGeometry)
 	{
-		int32 SectionCount = InstancedRenderData.LODModels[LODIndex].Sections.Num();
+		int32 SectionCount = (*InstancedRenderData.RayTracingProxyLODs)[LODIndex].Sections->Num();
 
 		for (int32 SectionIdx = 0; SectionIdx < SectionCount; ++SectionIdx)
 		{
-			//#dxr_todo: so far we use the parent static mesh path to get material data
-			{
-				FMeshBatch MeshBatch;
+			FMeshBatch MeshBatch;
 
-				if (!GetMeshElement(LODIndex, 0, SectionIdx, 0, false, false, MeshBatch))
-				{
-					MeshBatch.MaterialRenderProxy = UMaterial::GetDefaultMaterial(MD_Surface)->GetRenderProxy();
-					MeshBatch.VertexFactory = &RenderData->LODVertexFactories[LODIndex].VertexFactory;
-					MeshBatch.SegmentIndex = SectionIdx;
-					MeshBatch.MeshIdInPrimitive = SectionIdx;
-				};
-				MeshBatch.ReverseCulling = bReverseCulling;
-				RayTracingWPOInstanceTemplate.Materials.Add(MeshBatch);
+			bool bResult;
+			if (RenderData->RayTracingProxy->bUsingRenderingLODs)
+			{
+				// when using rendering LODs we can reuse the main GetMeshElement(...)
+				bResult = GetMeshElement(LODIndex, 0, SectionIdx, 0, false, false, MeshBatch);
 			}
-
+			else
 			{
-				FMeshBatch DynamicMeshBatch;
+				// otherwise initialize MeshBatch using ray tracing proxy VB/IB/Section data
+				bResult = GetRayTracingMeshElement(LODIndex, 0, SectionIdx, 0, MeshBatch);
+			}
 
-				if (!GetMeshElement(LODIndex, 0, SectionIdx, 0, false, false, DynamicMeshBatch)) // todo: RayTracingLOD vertex factory
+			if (bResult)
+			{
+				if (RenderData->RayTracingProxy->bUsingRenderingLODs)
+				{
+					check(MeshBatch.VertexFactory == &InstancedRenderData.VertexFactories[LODIndex]);
+				}
+				else
 				{
-					DynamicMeshBatch.MaterialRenderProxy = UMaterial::GetDefaultMaterial(MD_Surface)->GetRenderProxy();
-					DynamicMeshBatch.SegmentIndex = SectionIdx;
-					DynamicMeshBatch.MeshIdInPrimitive = SectionIdx;
+					check(MeshBatch.VertexFactory == &InstancedRenderData.RayTracingProxyVertexFactories[LODIndex]);
 				}
-				DynamicMeshBatch.ReverseCulling = bReverseCulling;
-				DynamicMeshBatch.VertexFactory = &InstancedRenderData.VertexFactories[LODIndex];
 
-				RayTracingWPODynamicTemplate.Materials.Add(DynamicMeshBatch);
+				MeshBatch.ReverseCulling = bReverseCulling;
+
+				RayTracingMeshBatches.Add(MeshBatch);
 			}
 		}
+	}
+
+	// Which index holds the reference to the particular simulated instance
+	TArray<uint32> ActiveInstances;
+
+	// Visible instances
+	TArray<FVisibleInstance> VisibleInstances;
+	TArray<FRayTracingInstance> RayTracingInstances;
 
+	const uint32 RequestedSimulatedInstances = CVarRayTracingSimulatedInstanceCount.GetValueOnRenderThread();
+	const uint32 SimulatedInstances = FMath::Min(RequestedSimulatedInstances == -1 ? InstanceCount : FMath::Clamp(RequestedSimulatedInstances, 1u, InstanceCount), MaxSimulatedInstances);
+
+	if (bDynamicRayTracingGeometry)
+	{
 		SetupRayTracingDynamicInstances(Collector.GetRHICommandList(), SimulatedInstances, LODIndex);
 
 		ActiveInstances.AddZeroed(SimulatedInstances);
@@ -2045,9 +2121,12 @@ void FInstancedStaticMeshSceneProxy::GetDynamicRayTracingInstances(FRayTracingIn
 		}
 	}
 
+	FRayTracingInstance StaticRayTracingInstance;
+	StaticRayTracingInstance.Geometry = RayTracingGeometry;
+	StaticRayTracingInstance.Materials = RayTracingMeshBatches;
 	//preallocate the worst-case to prevent an explosion of reallocs
 	//#dxr_todo: possibly track used instances and reserve based on previous behavior
-	RayTracingInstanceTemplate.PrimitiveInstanceIndices.Reserve(InstanceCount);
+	StaticRayTracingInstance.PrimitiveInstanceIndices.Reserve(InstanceCount);
 
 	// Add all visible instances
 	for (FVisibleInstance VisibleInstance : VisibleInstances)
@@ -2057,20 +2136,22 @@ void FInstancedStaticMeshSceneProxy::GetDynamicRayTracingInstances(FRayTracingIn
 		FMatrix InstanceToWorld = InstanceSceneDataBuffers->GetInstanceToWorld(InstanceIndex);
 		const uint32 DynamicInstanceIdx = InstanceIndex % SimulatedInstances;
 
-		if (bDynamicRayTracingGeometry && InstancedRenderData.VertexFactories[LODIndex].GetType()->SupportsRayTracingDynamicGeometry())
+		if (bDynamicRayTracingGeometry && InstancedRenderData.RayTracingProxyVertexFactories[LODIndex].GetType()->SupportsRayTracingDynamicGeometry())
 		{
 			FRayTracingInstance* DynamicInstance = nullptr;
 
 			if (ActiveInstances[DynamicInstanceIdx] == INDEX_NONE)
 			{
 				// first case of this dynamic instance, setup the material and add it
-				const FStaticMeshLODResources& LODModel = RenderData->LODResources[LODIndex];
+				const FStaticMeshRayTracingProxyLOD& RayTracingLOD = RenderData->RayTracingProxy->LODs[LODIndex];
+				const uint32 NumVertices = RayTracingLOD.VertexBuffers->PositionVertexBuffer.GetNumVertices();
 
 				FRayTracingDynamicInstanceData& DynamicData = RayTracingDynamicInstanceData[DynamicInstanceIdx];
 
 				ActiveInstances[DynamicInstanceIdx] = RayTracingInstances.Num();
-				FRayTracingInstance& RayTracingInstance = RayTracingInstances.Add_GetRef(RayTracingWPOInstanceTemplate);
+				FRayTracingInstance& RayTracingInstance = RayTracingInstances.AddDefaulted_GetRef();
 				RayTracingInstance.Geometry = &DynamicData.Geometry;
+				RayTracingInstance.Materials = RayTracingMeshBatches;
 				RayTracingInstance.PrimitiveInstanceIndices.Reserve(InstanceCount);
 
 				DynamicInstance = &RayTracingInstance;
@@ -2079,10 +2160,10 @@ void FInstancedStaticMeshSceneProxy::GetDynamicRayTracingInstances(FRayTracingIn
 					FirstActiveViewIndex,
 					FRayTracingDynamicGeometryUpdateParams
 					{
-						RayTracingWPODynamicTemplate.Materials,
+						RayTracingInstance.Materials,
 						false,
-						(uint32)LODModel.GetNumVertices(),
-						uint32((SIZE_T)LODModel.GetNumVertices() * sizeof(FVector3f)),
+						NumVertices,
+						uint32((SIZE_T)NumVertices * sizeof(FVector3f)),
 						DynamicData.Geometry.Initializer.TotalPrimitiveCount,
 						&DynamicData.Geometry,
 						nullptr,
@@ -2102,33 +2183,12 @@ void FInstancedStaticMeshSceneProxy::GetDynamicRayTracingInstances(FRayTracingIn
 		}
 		else
 		{
-			RayTracingInstanceTemplate.PrimitiveInstanceIndices.Add(InstanceIndex);
+			StaticRayTracingInstance.PrimitiveInstanceIndices.Add(InstanceIndex);
 		}
 	}
 
-	if (RayTracingInstanceTemplate.PrimitiveInstanceIndices.Num() > 0)
+	if (StaticRayTracingInstance.PrimitiveInstanceIndices.Num() > 0)
 	{
-		int32 SectionCount = InstancedRenderData.LODModels[LODIndex].Sections.Num();
-
-		for (int32 SectionIdx = 0; SectionIdx < SectionCount; ++SectionIdx)
-		{
-			//#dxr_todo: so far we use the parent static mesh path to get material data
-			FMeshBatch MeshBatch;
-
-			bool bResult = GetMeshElement(LODIndex, 0, SectionIdx, 0, false, false, MeshBatch);
-			if (!bResult)
-			{
-				// Hidden material
-				MeshBatch.MaterialRenderProxy = UMaterial::GetDefaultMaterial(MD_Surface)->GetRenderProxy();
-				MeshBatch.VertexFactory = &RenderData->LODVertexFactories[LODIndex].VertexFactory;
-				MeshBatch.SegmentIndex = SectionIdx;
-				MeshBatch.MeshIdInPrimitive = SectionIdx;
-			}
-			MeshBatch.ReverseCulling = bReverseCulling;
-
-			RayTracingInstanceTemplate.Materials.Add(MeshBatch);
-		}
-
 		for (int32 ViewIndex = 0; ViewIndex < Views.Num(); ViewIndex++)
 		{
 			if ((VisibilityMap & (1 << ViewIndex)) == 0)
@@ -2136,7 +2196,7 @@ void FInstancedStaticMeshSceneProxy::GetDynamicRayTracingInstances(FRayTracingIn
 				continue;
 			}
 
-			Collector.AddRayTracingInstance(ViewIndex, RayTracingInstanceTemplate);
+			Collector.AddRayTracingInstance(ViewIndex, StaticRayTracingInstance);
 		}
 	}
 
@@ -2250,6 +2310,10 @@ void FInstancedStaticMeshSceneProxy::SetInstanceCullDistance_RenderThread(float
 		FInstancedStaticMeshVFLooseUniformShaderParametersRef LooseUniformBuffer = CreateLooseUniformBuffer(nullptr, &UserData_AllInstances, 0, i, EUniformBufferUsage::UniformBuffer_MultiFrame);
 		LODLooseUniformBuffers.Add(i, LooseUniformBuffer);
 	}
+
+	// currently we setup loose uniform buffer for ray tracing based on InstancedLODIndex = 0
+	// since currently non of the uniforms depend on the LOD index when InstancedLODRange = 0
+	RayTracingLooseUniformBuffer = CreateLooseUniformBuffer(nullptr, &UserData_AllInstances, /*InstancedLODRange*/ 0, /*InstancedLODIndex*/ 0, EUniformBufferUsage::UniformBuffer_MultiFrame);
 }
 
 void FInstancedStaticMeshSceneProxy::GetPreSkinnedLocalBounds(FBoxSphereBounds& OutBounds) const
