Substrate - add default normal to default BSDF for material function input.
That is for normal to be able to be used in code in case parameter blendced operators are used (for the evaluation of the normal based parameter blending weight).

#tests PC editor dx12
#rnx

diff --git a/Engine/Source/Runtime/Engine/Private/Materials/MaterialExpressions.cpp b/Engine/Source/Runtime/Engine/Private/Materials/MaterialExpressions.cpp
--- a/Engine/Source/Runtime/Engine/Private/Materials/MaterialExpressions.cpp
+++ b/Engine/Source/Runtime/Engine/Private/Materials/MaterialExpressions.cpp
@@ -16615,11 +16615,16 @@ FSubstrateOperator* UMaterialExpressionFunctionInput::SubstrateGenerateMaterialT
 	}
 
 	// If we are parsing for a material function input we always needs to return a default valid BSDF operator at least 
-	// If the material requires an input thenthe UI will forcethe user to proviced one.
+	// If the material requires an input then the UI will force the user to proviced one.
 	// If this is not the case however, we need to compile a default material to compile the shader, or to preview the material function.
 	if(!SlabOperator)
 	{
+		int32 NormalCodeChunk = Compiler->VertexNormal();
+		const int32 NullTangentCodeChunk = INDEX_NONE;
+		const FSubstrateRegisteredSharedLocalBasis NewRegisteredSharedLocalBasis = SubstrateCompilationInfoCreateSharedLocalBasis(Compiler, NormalCodeChunk, NullTangentCodeChunk);
+
 		SlabOperator = &Compiler->SubstrateCompilationRegisterOperator(SUBSTRATE_OPERATOR_BSDF, Compiler->SubstrateTreeStackGetPathUniqueId(), this->MaterialExpressionGuid, Parent, Compiler->SubstrateTreeStackGetParentPathUniqueId());
+		SlabOperator->BSDFRegisteredSharedLocalBasis = NewRegisteredSharedLocalBasis;
 		SlabOperator->BSDFType = SUBSTRATE_BSDF_TYPE_SLAB;
 		SlabOperator->ThicknessIndex = Compiler->SubstrateThicknessStackGetThicknessIndex();
 	}
