Splines: Mark new spline types as experimental, mark SplinesFramework as experimental

diff --git a/Engine/Plugins/PCG/Source/PCG/Public/Data/PCGSplineStruct.h b/Engine/Plugins/PCG/Source/PCG/Public/Data/PCGSplineStruct.h
--- a/Engine/Plugins/PCG/Source/PCG/Public/Data/PCGSplineStruct.h
+++ b/Engine/Plugins/PCG/Source/PCG/Public/Data/PCGSplineStruct.h
@@ -170,8 +170,10 @@ protected:
 	TArray<int64> ControlPointsEntryKeys; // Needs to be int64 for UHT, but it is a PCGMetadataEntryKey
 
 private:
+	PRAGMA_DISABLE_EXPERIMENTAL_WARNINGS
 	UPROPERTY()
 	FSpline Spline;
+	PRAGMA_ENABLE_EXPERIMENTAL_WARNINGS
 
 	mutable UE::FSpinLock SpinLock;
 	mutable TArray<FName> CachedSortedAttributeChannelNames;
diff --git a/Engine/Source/Runtime/Engine/Classes/Components/SplineComponent.h b/Engine/Source/Runtime/Engine/Classes/Components/SplineComponent.h
--- a/Engine/Source/Runtime/Engine/Classes/Components/SplineComponent.h
+++ b/Engine/Source/Runtime/Engine/Classes/Components/SplineComponent.h
@@ -5,11 +5,17 @@
 #include "CoreMinimal.h"
 #include "UObject/ObjectMacros.h"
 #include "Components/PrimitiveComponent.h"
-#include "Curves/Spline.h"
 #include "Math/InterpCurve.h"
 #include "BoxTypes.h"
+
+PRAGMA_DISABLE_EXPERIMENTAL_WARNINGS
+#include "Curves/Spline.h"
+PRAGMA_ENABLE_EXPERIMENTAL_WARNINGS
+
 #include "SplineComponent.generated.h"
 
+PRAGMA_DISABLE_EXPERIMENTAL_WARNINGS
+
 class FPrimitiveSceneProxy;
 class FPrimitiveDrawInterface;
 class FSceneView;
@@ -211,7 +217,6 @@ class USplineComponent : public UPrimitiveComponent
 
 protected:
 
-	// todo: replicate, but I believe we need to make the struct able to replicate in the first place.
 	UPROPERTY(EditAnywhere, Category=Points)
 	FSpline Spline;
 	
@@ -1175,3 +1180,5 @@ inline FRotator USplineComponent::GetWorldRotationAtDistanceAlongSpline(float Di
 inline FVector USplineComponent::GetWorldLocationAtTime(float Time, bool bUseConstantVelocity) const { return GetLocationAtTime(Time, ESplineCoordinateSpace::World, bUseConstantVelocity); }
 inline FVector USplineComponent::GetWorldDirectionAtTime(float Time, bool bUseConstantVelocity) const { return GetDirectionAtTime(Time, ESplineCoordinateSpace::World, bUseConstantVelocity); }
 inline FRotator USplineComponent::GetWorldRotationAtTime(float Time, bool bUseConstantVelocity) const { return GetRotationAtTime(Time, ESplineCoordinateSpace::World, bUseConstantVelocity); }
+
+PRAGMA_ENABLE_EXPERIMENTAL_WARNINGS
\ No newline at end of file
diff --git a/Engine/Source/Runtime/Engine/Classes/Curves/Spline.h b/Engine/Source/Runtime/Engine/Classes/Curves/Spline.h
--- a/Engine/Source/Runtime/Engine/Classes/Curves/Spline.h
+++ b/Engine/Source/Runtime/Engine/Classes/Curves/Spline.h
@@ -13,6 +13,8 @@ struct FSplineCurves;
 class FLegacySpline;
 class FNewSpline;
 
+struct UE_EXPERIMENTAL(5.7, "New spline APIs are experimental.") FSpline;
+
 /**
  * A general purpose, reflected spline.
  * The implementation can be configured at runtime.
diff --git a/Engine/Source/Runtime/Engine/Private/Curves/Spline.cpp b/Engine/Source/Runtime/Engine/Private/Curves/Spline.cpp
--- a/Engine/Source/Runtime/Engine/Private/Curves/Spline.cpp
+++ b/Engine/Source/Runtime/Engine/Private/Curves/Spline.cpp
@@ -8,12 +8,17 @@
 #include "Misc/Base64.h"
 #include "Serialization/MemoryReader.h"
 #include "Serialization/MemoryWriter.h"
+#include "UObject/FortniteMainBranchObjectVersion.h"
+
+PRAGMA_DISABLE_EXPERIMENTAL_WARNINGS
 #include "Splines/TangentBezierSpline.h"
 #include "Splines/MultiSpline.h"
-#include "UObject/FortniteMainBranchObjectVersion.h"
+PRAGMA_ENABLE_EXPERIMENTAL_WARNINGS
 
 #include UE_INLINE_GENERATED_CPP_BY_NAME(Spline)
 
+PRAGMA_DISABLE_EXPERIMENTAL_WARNINGS
+
 using UE::Geometry::FInterval1f;
 
 namespace UE::Spline
@@ -3170,4 +3175,6 @@ void FLegacySpline::UpdateSpline(const FSpline::FUpdateSplineParams& InParams)
 		AccumulatedLength += GetSegmentLength(SegmentIndex, 1.0f, InParams.Scale3D);
 	}
 	ReparamTable.Points.Emplace(AccumulatedLength, static_cast<float>(NumSegments), 0.0f, 0.0f, CIM_Linear);
-}
\ No newline at end of file
+}
+
+PRAGMA_ENABLE_EXPERIMENTAL_WARNINGS
\ No newline at end of file
diff --git a/Engine/Source/Runtime/GeometryCore/Private/GeometryCoreModule.cpp b/Engine/Source/Runtime/GeometryCore/Private/GeometryCoreModule.cpp
--- a/Engine/Source/Runtime/GeometryCore/Private/GeometryCoreModule.cpp
+++ b/Engine/Source/Runtime/GeometryCore/Private/GeometryCoreModule.cpp
@@ -29,6 +29,7 @@ void FGeometryCoreModule::RegisterSplineTypes() const
 {
 	using namespace UE::Geometry::Spline;
 
+	PRAGMA_DISABLE_EXPERIMENTAL_WARNINGS
 	/**
 	 * TMultiSpline<FTangentBezierSpline3d> once had an incorrect type ID which failed to take into account the wrapped spline type (FTangentBezierSpline3d).
 	 * Now, the implementation name of a multi-spline is in the form of "MultiSpline.[WrappedSplineName]".
@@ -42,6 +43,7 @@ void FGeometryCoreModule::RegisterSplineTypes() const
 	{
 		return MakeUnique<TMultiSpline<FTangentBezierSpline3d>>();
 	});
+	PRAGMA_ENABLE_EXPERIMENTAL_WARNINGS
 }
 
 #undef LOCTEXT_NAMESPACE
diff --git a/Engine/Source/Runtime/GeometryCore/Public/Splines/BSpline.h b/Engine/Source/Runtime/GeometryCore/Public/Splines/BSpline.h
--- a/Engine/Source/Runtime/GeometryCore/Public/Splines/BSpline.h
+++ b/Engine/Source/Runtime/GeometryCore/Public/Splines/BSpline.h
@@ -14,7 +14,11 @@ namespace Geometry
 {
 namespace Spline
 {
-enum class EParameterizationPolicy
+
+struct UE_EXPERIMENTAL(5.7, "New spline APIs are experimental.") FKnot;
+template <typename VALUETYPE, int32 DEGREE> class UE_EXPERIMENTAL(5.7, "New spline APIs are experimental.") TBSpline;
+	
+enum class UE_EXPERIMENTAL(5.7, "New spline APIs are experimental.") EParameterizationPolicy
 {
 	/** Equal spacing between knots regardless of control point positions */
 	Uniform,
diff --git a/Engine/Source/Runtime/GeometryCore/Public/Splines/LinearSpline.h b/Engine/Source/Runtime/GeometryCore/Public/Splines/LinearSpline.h
--- a/Engine/Source/Runtime/GeometryCore/Public/Splines/LinearSpline.h
+++ b/Engine/Source/Runtime/GeometryCore/Public/Splines/LinearSpline.h
@@ -10,6 +10,8 @@ namespace Geometry
 namespace Spline
 {
 
+template<typename T> class UE_EXPERIMENTAL(5.7, "") TLinearSpline;
+	
 template<typename T>
 class TLinearSpline : public TBSpline<T, 1>
 {
diff --git a/Engine/Source/Runtime/GeometryCore/Public/Splines/MultiSpline.h b/Engine/Source/Runtime/GeometryCore/Public/Splines/MultiSpline.h
--- a/Engine/Source/Runtime/GeometryCore/Public/Splines/MultiSpline.h
+++ b/Engine/Source/Runtime/GeometryCore/Public/Splines/MultiSpline.h
@@ -15,6 +15,8 @@ namespace Geometry
 namespace Spline
 {
 
+template <typename SPLINETYPE> class UE_EXPERIMENTAL(5.7, "") TMultiSpline;
+	
 /**
  * Multi-spline implementation that can contain other splines
  * Provides attribute functionality while allowing direct access to implementation-specific methods
diff --git a/Engine/Source/Runtime/GeometryCore/Public/Splines/PolyBezierSpline.h b/Engine/Source/Runtime/GeometryCore/Public/Splines/PolyBezierSpline.h
--- a/Engine/Source/Runtime/GeometryCore/Public/Splines/PolyBezierSpline.h
+++ b/Engine/Source/Runtime/GeometryCore/Public/Splines/PolyBezierSpline.h
@@ -4,7 +4,6 @@
 
 #include "BSpline.h"
 #include "Misc/ScopeExit.h"
-#include "HAL/IConsoleManager.h"
 
 namespace UE
 {
@@ -13,6 +12,8 @@ namespace Geometry
 namespace Spline
 {
 
+template<typename VALUETYPE> class UE_EXPERIMENTAL(5.7, "New spline APIs are experimental.") TPolyBezierSpline;
+	
 /* Augments B-Spline interface */
 template<typename VALUETYPE>
 class TPolyBezierSpline :
diff --git a/Engine/Source/Runtime/GeometryCore/Public/Splines/Spline.h b/Engine/Source/Runtime/GeometryCore/Public/Splines/Spline.h
--- a/Engine/Source/Runtime/GeometryCore/Public/Splines/Spline.h
+++ b/Engine/Source/Runtime/GeometryCore/Public/Splines/Spline.h
@@ -13,6 +13,8 @@ namespace Geometry
 {
 namespace Spline
 {
+
+template <typename SPLINETYPE> class UE_EXPERIMENTAL(5.7, "New spline APIs are experimental.") TSplineWrapper;
 	
 /* 
  * A spline which wraps another fully implemented spline (SplineType must derive from TSplineInterface) 
diff --git a/Engine/Source/Runtime/GeometryCore/Public/Splines/SplineInterfaces.h b/Engine/Source/Runtime/GeometryCore/Public/Splines/SplineInterfaces.h
--- a/Engine/Source/Runtime/GeometryCore/Public/Splines/SplineInterfaces.h
+++ b/Engine/Source/Runtime/GeometryCore/Public/Splines/SplineInterfaces.h
@@ -14,13 +14,13 @@ namespace Spline
 {
 
 // Forward declarations
-class ISplineInterface;
-template<typename T> class TSplineInterface;
+class UE_EXPERIMENTAL(5.7, "New spline APIs are experimental.") ISplineInterface;
+template <typename VALUETYPE> class UE_EXPERIMENTAL(5.7, "New spline APIs are experimental.") TSplineInterface;
 
 /**
  * Enumeration for out-of-bounds handling modes
  */
-enum class EOutOfBoundsHandlingMode
+enum class UE_EXPERIMENTAL(5.7, "New spline APIs are experimental.") EOutOfBoundsHandlingMode
 {
     Zero,       // Null out the values outside of the bounds
     Constant,   // Choose the nearest of the first/last value
diff --git a/Engine/Source/Runtime/GeometryCore/Public/Splines/TangentBezierSpline.h b/Engine/Source/Runtime/GeometryCore/Public/Splines/TangentBezierSpline.h
--- a/Engine/Source/Runtime/GeometryCore/Public/Splines/TangentBezierSpline.h
+++ b/Engine/Source/Runtime/GeometryCore/Public/Splines/TangentBezierSpline.h
@@ -10,10 +10,13 @@ namespace Geometry
 namespace Spline
 {
 	
+template<typename VALUETYPE> struct UE_EXPERIMENTAL(5.7, "New spline APIs are experimental.") TTangentBezierControlPoint;
+template<typename VALUETYPE> class UE_EXPERIMENTAL(5.7, "New spline APIs are experimental.") TTangentBezierSpline;
+	
 /**
  * Enum defining how tangents are computed for a spline control point.
  */
-enum class ETangentMode : uint8
+enum class UE_EXPERIMENTAL(5.7, "New spline APIs are experimental.") ETangentMode : uint8
 {
 	/** Automatically compute unclamped tangents based on surrounding points */
 	Auto,
