TypedElement Actor + Component types store their respective objects as WeakObjectPtrs
* Fixes various use-after-free cases
* The lifetime of the TypedElement handles for these is mediated by the BeginDestroy functions on AActor and AActorComponent.  However, the handles are placed on a deferred queue for destruction.
* The handle data itself holds raw pointers to UObjects.  Problems can arise when flushing the deferred destruction of the handle and the UObject has had its destructor already called.  In some circumstances, code in this path atttempts to dereference the raw UObject, resulting in an access violation.
* The change here is to hold onto the UObjects by WeakObjectPtr so that the callsites can determine if the UObjects themselves have already been destroyed.
* One thing to note is that GetEvenIfUnreachable needed to be used as there are several places in the engine where the UObjects were being used even after they had been deemed "destroyed".  For example, component deselection occurs after it has been marked for deletion and deemed unreachable - deselection requires the use of TypedElementHandles though.  Use of .Get(bEvenIfPendingDestroy = true) adversely affects workflows at this time.  GetEvenIfUnreachable will retain the existing behaviour while preventing access to actually destroyed/freed UObjects.

* Explicit move + copy ctors and assignment operators as some compilers throw warnings from the implicitly defined versions of them.

diff --git a/Engine/Source/Runtime/Engine/Private/Elements/Actor/ActorElementData.cpp b/Engine/Source/Runtime/Engine/Private/Elements/Actor/ActorElementData.cpp
--- a/Engine/Source/Runtime/Engine/Private/Elements/Actor/ActorElementData.cpp
+++ b/Engine/Source/Runtime/Engine/Private/Elements/Actor/ActorElementData.cpp
@@ -2,7 +2,7 @@
 
 #include "Elements/Actor/ActorElementData.h"
 #include "Elements/Framework/TypedElementHandle.h"
-#include "UObject/Stack.h"
+#include "GameFramework/Actor.h"
 
 UE_DEFINE_TYPED_ELEMENT_DATA_RTTI(FActorElementData);
 
@@ -12,13 +12,70 @@ namespace ActorElementDataUtil
 AActor* GetActorFromHandle(const FTypedElementHandle& InHandle, const bool bSilent)
 {
 	const FActorElementData* ActorElement = InHandle.GetData<FActorElementData>(bSilent);
-	return ActorElement ? ActorElement->Actor : nullptr;
+	return ActorElement ? ActorElement->ActorWeak.GetEvenIfUnreachable() : nullptr;
 }
 
 AActor* GetActorFromHandleChecked(const FTypedElementHandle& InHandle)
 {
 	const FActorElementData& ActorElement = InHandle.GetDataChecked<FActorElementData>();
-	return ActorElement.Actor;
+	return ActorElement.ActorWeak.GetEvenIfUnreachable();
 }
 
 } // namespace ActorElementDataUtil
+
+FActorElementData::FActorElementData()
+{
+}
+
+FActorElementData::~FActorElementData()
+{
+}
+
+FActorElementData::FActorElementData(const FActorElementData& Other)
+{
+	PRAGMA_DISABLE_DEPRECATION_WARNINGS
+	this->Actor = Other.Actor;
+	PRAGMA_ENABLE_DEPRECATION_WARNINGS
+	this->ActorWeak = Other.ActorWeak;
+}
+
+FActorElementData::FActorElementData(FActorElementData&& Other)
+{
+	PRAGMA_DISABLE_DEPRECATION_WARNINGS
+	this->Actor = MoveTemp(Other.Actor);
+	PRAGMA_ENABLE_DEPRECATION_WARNINGS
+	this->ActorWeak = MoveTemp(Other.ActorWeak);
+}
+
+FActorElementData& FActorElementData::operator=(const FActorElementData& Other)
+{
+	if (this != &Other)
+	{
+		PRAGMA_DISABLE_DEPRECATION_WARNINGS
+		this->Actor = Other.Actor;
+		PRAGMA_ENABLE_DEPRECATION_WARNINGS
+		this->ActorWeak = Other.ActorWeak;
+	}
+	return *this;
+}
+
+FActorElementData& FActorElementData::operator=(FActorElementData&& Other)
+{
+	if (this != &Other)
+	{
+		PRAGMA_DISABLE_DEPRECATION_WARNINGS
+		this->Actor = Other.Actor;
+		PRAGMA_ENABLE_DEPRECATION_WARNINGS
+		this->ActorWeak = Other.ActorWeak;
+	}
+	return *this;
+}
+
+template <>
+FString GetTypedElementDebugId<FActorElementData>(const FActorElementData& InElementData)
+{
+	const AActor* Object = InElementData.ActorWeak.GetEvenIfUnreachable();
+	return IsValid(Object)
+		? Object->GetFullName()
+		: TEXT("null");
+}
diff --git a/Engine/Source/Runtime/Engine/Private/Elements/Component/ComponentElementData.cpp b/Engine/Source/Runtime/Engine/Private/Elements/Component/ComponentElementData.cpp
--- a/Engine/Source/Runtime/Engine/Private/Elements/Component/ComponentElementData.cpp
+++ b/Engine/Source/Runtime/Engine/Private/Elements/Component/ComponentElementData.cpp
@@ -1,6 +1,8 @@
 // Copyright Epic Games, Inc. All Rights Reserved.
 
 #include "Elements/Component/ComponentElementData.h"
+
+#include "Components/ActorComponent.h"
 #include "Elements/Framework/TypedElementHandle.h"
 #include "UObject/Stack.h"
 
@@ -12,13 +14,70 @@ namespace ComponentElementDataUtil
 UActorComponent* GetComponentFromHandle(const FTypedElementHandle& InHandle, const bool bSilent)
 {
 	const FComponentElementData* ComponentElement = InHandle.GetData<FComponentElementData>(bSilent);
-	return ComponentElement ? ComponentElement->Component : nullptr;
+	return ComponentElement ? ComponentElement->ComponentWeak.GetEvenIfUnreachable() : nullptr;
 }
 
 UActorComponent* GetComponentFromHandleChecked(const FTypedElementHandle& InHandle)
 {
 	const FComponentElementData& ComponentElement = InHandle.GetDataChecked<FComponentElementData>();
-	return ComponentElement.Component;
+	return ComponentElement.ComponentWeak.GetEvenIfUnreachable();
 }
 
 } // namespace ComponentElementDataUtil
+
+FComponentElementData::FComponentElementData()
+{
+}
+
+FComponentElementData::~FComponentElementData()
+{
+}
+
+FComponentElementData::FComponentElementData(const FComponentElementData& Other)
+{
+	PRAGMA_DISABLE_DEPRECATION_WARNINGS
+	this->Component = Other.Component;
+	PRAGMA_ENABLE_DEPRECATION_WARNINGS
+	this->ComponentWeak = Other.ComponentWeak;
+}
+
+FComponentElementData::FComponentElementData(FComponentElementData&& Other)
+{
+	PRAGMA_DISABLE_DEPRECATION_WARNINGS
+	this->Component = MoveTemp(Other.Component);
+	PRAGMA_ENABLE_DEPRECATION_WARNINGS
+	this->ComponentWeak = MoveTemp(Other.ComponentWeak);
+}
+
+FComponentElementData& FComponentElementData::operator=(const FComponentElementData& Other)
+{
+	if (this != &Other)
+	{
+		PRAGMA_DISABLE_DEPRECATION_WARNINGS
+		this->Component = Other.Component;
+		PRAGMA_ENABLE_DEPRECATION_WARNINGS
+		this->ComponentWeak = Other.ComponentWeak;
+	}
+	return *this;
+}
+
+FComponentElementData& FComponentElementData::operator=(FComponentElementData&& Other)
+{
+	if (this != &Other)
+	{
+		PRAGMA_DISABLE_DEPRECATION_WARNINGS
+		this->Component = Other.Component;
+		PRAGMA_ENABLE_DEPRECATION_WARNINGS
+		this->ComponentWeak = Other.ComponentWeak;
+	}
+	return *this;
+}
+
+template <>
+FString GetTypedElementDebugId<FComponentElementData>(const FComponentElementData& InElementData)
+{
+	UActorComponent* Object = InElementData.ComponentWeak.GetEvenIfUnreachable();
+	return Object
+		? Object->GetFullName()
+		: TEXT("null");
+}
diff --git a/Engine/Source/Runtime/Engine/Private/Elements/Framework/EngineElementsLibrary.cpp b/Engine/Source/Runtime/Engine/Private/Elements/Framework/EngineElementsLibrary.cpp
--- a/Engine/Source/Runtime/Engine/Private/Elements/Framework/EngineElementsLibrary.cpp
+++ b/Engine/Source/Runtime/Engine/Private/Elements/Framework/EngineElementsLibrary.cpp
@@ -246,7 +246,10 @@ void UEngineElementsLibrary::ReplaceEditorActorElementHandlesImpl(const TMap<con
 {
 	EngineElementsLibraryUtil::ReplaceEditorTypedElementHandles<FActorElementData, const AActor*>(OutUpdatedElements, ReplacementActors, GActorElementOwnerStore, [](const AActor* InActor, TTypedElementOwner<FActorElementData>& InOutActorElement)
 	{
+		PRAGMA_DISABLE_DEPRECATION_WARNINGS
 		InOutActorElement.GetDataChecked().Actor = const_cast<AActor*>(InActor);
+		PRAGMA_ENABLE_DEPRECATION_WARNINGS
+		InOutActorElement.GetDataChecked().ActorWeak = const_cast<AActor*>(InActor);
 	}, &UEngineElementsLibrary::DestroyActorElement);
 }
 
@@ -254,7 +257,10 @@ void UEngineElementsLibrary::ReplaceEditorComponentElementHandlesImpl(const TMap
 {
 	EngineElementsLibraryUtil::ReplaceEditorTypedElementHandles<FComponentElementData, const UActorComponent*>(OutUpdatedElements, ReplacementComponents, GComponentElementOwnerStore, [](const UActorComponent* InComponent, TTypedElementOwner<FComponentElementData>& InOutComponentElement)
 	{
+		PRAGMA_DISABLE_DEPRECATION_WARNINGS
 		InOutComponentElement.GetDataChecked().Component = const_cast<UActorComponent*>(InComponent);
+		PRAGMA_ENABLE_DEPRECATION_WARNINGS
+		InOutComponentElement.GetDataChecked().ComponentWeak = const_cast<UActorComponent*>(InComponent);
 	}, &UEngineElementsLibrary::DestroyComponentElement);
 }
 
@@ -359,7 +365,10 @@ TTypedElementOwner<FActorElementData> UEngineElementsLibrary::CreateActorElement
 	TTypedElementOwner<FActorElementData> ActorElement = EngineElementsLibraryUtil::CreateTypedElement<AActor, FActorElementData>(InActor, NAME_Actor);
 	if (ActorElement)
 	{
+		PRAGMA_DISABLE_DEPRECATION_WARNINGS
 		ActorElement.GetDataChecked().Actor = const_cast<AActor*>(InActor);
+		PRAGMA_ENABLE_DEPRECATION_WARNINGS
+		ActorElement.GetDataChecked().ActorWeak = const_cast<AActor*>(InActor);
 	}
 	
 	RegisterActorElement(InActor);
@@ -374,7 +383,7 @@ void UEngineElementsLibrary::DestroyActorElement(const AActor* InActor, TTypedEl
 
 	if (InOutActorElement)
 	{
-		checkf(InOutActorElement.GetDataChecked().Actor == InActor, TEXT("Actor element was not for this actor instance! %s"), *InActor->GetPathName());
+		checkf(InOutActorElement.GetDataChecked().ActorWeak.GetEvenIfUnreachable() == InActor, TEXT("Actor element was not for this actor instance! %s"), *InActor->GetPathName());
 		UTypedElementRegistry::GetInstance()->DestroyElement(InOutActorElement);
 	}
 }
@@ -447,10 +456,19 @@ void UEngineElementsLibrary::ReplaceEditorActorElementHandles(const TMap<const A
 
 TTypedElementOwner<FComponentElementData> UEngineElementsLibrary::CreateComponentElement(const UActorComponent* InComponent)
 {
+	TWeakObjectPtr<UActorComponent> ComponentWeak = const_cast<UActorComponent*>(InComponent);
+	if (!ComponentWeak.IsValid())
+	{
+		return TTypedElementOwner<FComponentElementData>();
+	}
 	TTypedElementOwner<FComponentElementData> ComponentElement = EngineElementsLibraryUtil::CreateTypedElement<UActorComponent, FComponentElementData>(InComponent, NAME_Components);
 	if (ComponentElement)
 	{
-		ComponentElement.GetDataChecked().Component = const_cast<UActorComponent*>(InComponent);
+		FComponentElementData& ElementData = ComponentElement.GetDataChecked();
+		PRAGMA_DISABLE_DEPRECATION_WARNINGS
+		ElementData.Component = const_cast<UActorComponent*>(InComponent);
+		PRAGMA_ENABLE_DEPRECATION_WARNINGS
+		ElementData.ComponentWeak = const_cast<UActorComponent*>(InComponent);
 	}
 	return ComponentElement;
 }
@@ -459,7 +477,7 @@ void UEngineElementsLibrary::DestroyComponentElement(const UActorComponent* InCo
 {
 	if (InOutComponentElement)
 	{
-		checkf(InOutComponentElement.GetDataChecked().Component == InComponent, TEXT("Component element was not for this component instance! %s"), *InComponent->GetPathName());
+		checkf(InOutComponentElement.GetDataChecked().ComponentWeak.GetEvenIfUnreachable() == InComponent, TEXT("Component element was not for this component instance! %s"), *InComponent->GetPathName());
 		UTypedElementRegistry::GetInstance()->DestroyElement(InOutComponentElement);
 	}
 }
diff --git a/Engine/Source/Runtime/Engine/Public/Elements/Actor/ActorElementData.h b/Engine/Source/Runtime/Engine/Public/Elements/Actor/ActorElementData.h
--- a/Engine/Source/Runtime/Engine/Public/Elements/Actor/ActorElementData.h
+++ b/Engine/Source/Runtime/Engine/Public/Elements/Actor/ActorElementData.h
@@ -4,6 +4,7 @@
 
 #include "Elements/Framework/TypedElementData.h"
 #include "UObject/Object.h"
+#include "UObject/WeakObjectPtr.h"
 
 class AActor;
 struct FTypedElementHandle;
@@ -15,17 +16,20 @@ struct FActorElementData
 {
 	UE_DECLARE_TYPED_ELEMENT_DATA_RTTI(FActorElementData);
 
+	ENGINE_API FActorElementData();
+	ENGINE_API ~FActorElementData();
+	ENGINE_API FActorElementData(const FActorElementData& Other);
+	ENGINE_API FActorElementData(FActorElementData&& Other);
+	ENGINE_API FActorElementData& operator=(const FActorElementData& Other);
+	ENGINE_API FActorElementData& operator=(FActorElementData&& Other);
+
+	UE_DEPRECATED(5.7, "Use ActorWeak")
 	AActor* Actor = nullptr;
+
+	TWeakObjectPtr<AActor> ActorWeak;
 };
 
-template <>
-inline FString GetTypedElementDebugId<FActorElementData>(const FActorElementData& InElementData)
-{
-	UObject* Object = (UObject*)InElementData.Actor;
-	return IsValid(Object)
-		? Object->GetFullName()
-		: TEXT("null");
-}
+template<> ENGINE_API FString GetTypedElementDebugId<FActorElementData>(const FActorElementData& InElementData);
 
 namespace ActorElementDataUtil
 {
diff --git a/Engine/Source/Runtime/Engine/Public/Elements/Component/ComponentElementData.h b/Engine/Source/Runtime/Engine/Public/Elements/Component/ComponentElementData.h
--- a/Engine/Source/Runtime/Engine/Public/Elements/Component/ComponentElementData.h
+++ b/Engine/Source/Runtime/Engine/Public/Elements/Component/ComponentElementData.h
@@ -4,6 +4,7 @@
 
 #include "Elements/Framework/TypedElementData.h"
 #include "UObject/Object.h"
+#include "UObject/WeakObjectPtr.h"
 
 class UActorComponent;
 struct FTypedElementHandle;
@@ -15,17 +16,24 @@ struct FComponentElementData
 {
 	UE_DECLARE_TYPED_ELEMENT_DATA_RTTI(FComponentElementData);
 
+	ENGINE_API FComponentElementData();
+	ENGINE_API ~FComponentElementData();
+	ENGINE_API FComponentElementData(const FComponentElementData& Other);
+	ENGINE_API FComponentElementData(FComponentElementData&& Other);
+	ENGINE_API FComponentElementData& operator=(const FComponentElementData& Other);
+	ENGINE_API FComponentElementData& operator=(FComponentElementData&& Other);
+
+	// TypedElementHandles are defer destroyed and can outlive their UObjects, leaving a stale
+	// handle to the Component.
+	// Use ComponentWeak instead to check for validity before use, especially in code that runs during
+	// deferred element handle destruction
+	UE_DEPRECATED(5.7, "Use ComponentWeak")
 	UActorComponent* Component = nullptr;
+	
+	TWeakObjectPtr<UActorComponent> ComponentWeak = nullptr;
 };
 
-template <>
-inline FString GetTypedElementDebugId<FComponentElementData>(const FComponentElementData& InElementData)
-{
-	UObject* Object = (UObject*)InElementData.Component;
-	return Object
-		? Object->GetFullName()
-		: TEXT("null");
-}
+template<> ENGINE_API FString GetTypedElementDebugId<FComponentElementData>(const FComponentElementData& InElementData);
 
 namespace ComponentElementDataUtil
 {
