Level editing crash fix - Delete actor inside a level instance then press save in the bottom of level viewport

diff --git a/Engine/Source/Runtime/Engine/Private/WorldPartition/ActorDescContainer.cpp b/Engine/Source/Runtime/Engine/Private/WorldPartition/ActorDescContainer.cpp
--- a/Engine/Source/Runtime/Engine/Private/WorldPartition/ActorDescContainer.cpp
+++ b/Engine/Source/Runtime/Engine/Private/WorldPartition/ActorDescContainer.cpp
@@ -554,12 +554,14 @@ void UActorDescContainer::OnPackageDeleted(UPackage* Package)
 		if (ShouldHandleDeletedObjectPlaceholderEvent(DeletedObjectPlaceholder))
 		{
 			// Here we validate that we didn't already processed the DeletedObjectPlaceholder in OnObjectPreSave
-			const AActor* OriginalActor = CastChecked<AActor>(DeletedObjectPlaceholder->GetOriginalObject());
-			if (FDeletedObjectPlaceholderAnnotation Annotation = DeletedObjectPlaceholdersAnnotation.GetAndRemoveAnnotation(OriginalActor); Annotation.IsValid())
+			if (const AActor* OriginalActor = Cast<AActor>(DeletedObjectPlaceholder->GetOriginalObject()))
 			{
-				check(Annotation.GetDeletedObjectPlaceholder() == DeletedObjectPlaceholder);
-				check(Annotation.GetActorDescContainer() == this);
-				verify(RemoveActor(OriginalActor->GetActorGuid()));
+				if (FDeletedObjectPlaceholderAnnotation Annotation = DeletedObjectPlaceholdersAnnotation.GetAndRemoveAnnotation(OriginalActor); Annotation.IsValid())
+				{
+					check(Annotation.GetDeletedObjectPlaceholder() == DeletedObjectPlaceholder);
+					check(Annotation.GetActorDescContainer() == this);
+					verify(RemoveActor(OriginalActor->GetActorGuid()));
+				}
 			}
 		}
 	}
