World Partition - HLOD: Ensure that always loaded HLOD actors that are using the "MainGrid:None" runtime grid naming are included correctly in the persistent level actors set

diff --git a/Engine/Source/Runtime/Engine/Private/WorldPartition/RuntimeHashSet/WorldPartitionRuntimeHashSet.cpp b/Engine/Source/Runtime/Engine/Private/WorldPartition/RuntimeHashSet/WorldPartitionRuntimeHashSet.cpp
--- a/Engine/Source/Runtime/Engine/Private/WorldPartition/RuntimeHashSet/WorldPartitionRuntimeHashSet.cpp
+++ b/Engine/Source/Runtime/Engine/Private/WorldPartition/RuntimeHashSet/WorldPartitionRuntimeHashSet.cpp
@@ -349,7 +349,7 @@ const URuntimePartition* UWorldPartitionRuntimeHashSet::ResolveRuntimePartition(
 		}
 	
 		// If an HLOD partition token was specified ("MainPartition:HLODPartition"), make sure it's valid too
-		if (!bMainPartitionLayer && HLODPartitionTokens.IsValidIndex(0))
+		if (!bMainPartitionLayer && HLODPartitionTokens.IsValidIndex(0) && !HLODPartitionTokens[0].IsNone())
 		{
 			const FRuntimePartitionHLODSetup* FoundRuntimePartitionHLODSetup = FoundRuntimePartitionDesc->HLODSetups.FindByPredicate([&HLODPartitionTokens](const FRuntimePartitionHLODSetup& RuntimePartitionHLODSetup)
 			{
