Fix for various errors and lost data related to the PinBase and Composite not transacting correctly by ensuring the Composite parent has modify called on it correctly, along with the expression collection.
Also fixed a nullptr crash with a simple check for when transactions overlap with a deletion of a material outside of the material editor meaning a Material wasn't assigned.

diff --git a/Engine/Source/Runtime/Engine/Private/Materials/MaterialExpressions.cpp b/Engine/Source/Runtime/Engine/Private/Materials/MaterialExpressions.cpp
--- a/Engine/Source/Runtime/Engine/Private/Materials/MaterialExpressions.cpp
+++ b/Engine/Source/Runtime/Engine/Private/Materials/MaterialExpressions.cpp
@@ -11444,14 +11444,19 @@ bool UMaterialExpressionComposite::Modify(bool bAlwaysMarkDirty)
 	// Modify pin bases so they can update the compilation graph
 	if (InputExpressions)
 	{
-		InputExpressions->Modify(bAlwaysMarkDirty);
+		InputExpressions->ModifyAsSubObject(bAlwaysMarkDirty);
 	}
 
 	if (OutputExpressions)
 	{
-		OutputExpressions->Modify(bAlwaysMarkDirty);
+		OutputExpressions->ModifyAsSubObject(bAlwaysMarkDirty);
 	}
 
+	return ModifyCompositeOnly(bAlwaysMarkDirty);
+}
+
+bool UMaterialExpressionComposite::ModifyCompositeOnly(bool bAlwaysMarkDirty)
+{
 	return Super::Modify(bAlwaysMarkDirty);
 }
 
@@ -11513,6 +11518,7 @@ void UMaterialExpressionPinBase::PostEditChangeProperty(FPropertyChangedEvent& P
 	{
 		Modify();
 		Material->Modify();
+		Material->GetEditorOnlyData()->Modify();
 
 		if (PropertyChangedEvent.ChangeType == EPropertyChangeType::ArrayAdd || PropertyChangedEvent.ChangeType == EPropertyChangeType::Duplicate)
 		{
@@ -11575,7 +11581,7 @@ void UMaterialExpressionPinBase::PostEditChangeProperty(FPropertyChangedEvent& P
 	PreEditRereouteExpresions.Empty();
 	Super::PostEditChangeProperty(PropertyChangedEvent);
 	
-	if (Material->GetExpressionCollection().Expressions.Remove(nullptr))
+	if (Material && Material->GetExpressionCollection().Expressions.Remove(nullptr))
 	{
 		UE_LOG(LogMaterial, Warning, TEXT("Null expressions detected in material expressions list after execution of UMaterialExpressionPinBase::PostEditChangeProperty."));
 	}
@@ -11679,7 +11685,7 @@ bool UMaterialExpressionPinBase::IsExpressionConnected(FExpressionInput* Input,
 
 void UMaterialExpressionPinBase::ConnectExpression(FExpressionInput* Input, int32 OutputIndex)
 {
-	if (PinDirection == EGPD_Output)
+	if (PinDirection == EGPD_Output && Input && OutputIndex >= 0 && OutputIndex < ReroutePins.Num())
 	{
 		ReroutePins[OutputIndex].Expression->ConnectExpression(Input, 0);
 	}
@@ -11691,6 +11697,15 @@ void UMaterialExpressionPinBase::GetCaption(TArray<FString>& OutCaptions) const
 }
 
 bool UMaterialExpressionPinBase::Modify(bool bAlwaysMarkDirty)
+{
+	if (UMaterialExpressionComposite* Composite = Cast<UMaterialExpressionComposite>(SubgraphExpression))
+	{
+		Composite->ModifyCompositeOnly(bAlwaysMarkDirty);
+	}
+	return ModifyAsSubObject(bAlwaysMarkDirty);
+}
+
+bool UMaterialExpressionPinBase::ModifyAsSubObject(bool bAlwaysMarkDirty)
 {
 	// Modify reroute pins so they can update the compilation graph
 	for (const FCompositeReroute& ReroutePin : ReroutePins)
@@ -11705,12 +11720,6 @@ bool UMaterialExpressionPinBase::Modify(bool bAlwaysMarkDirty)
 	return Super::Modify(bAlwaysMarkDirty);
 }
 
-bool UMaterialExpressionPinBase::CanUserDeleteExpression() const
-{
-	return false;
-}
-
-
 int32 UMaterialExpressionFresnel::Compile(class FMaterialCompiler* Compiler, int32 OutputIndex)
 {
 	// pow(1 - max(0,Normal dot Camera),Exponent) * (1 - BaseReflectFraction) + BaseReflectFraction
diff --git a/Engine/Source/Runtime/Engine/Public/Materials/MaterialExpressionComposite.h b/Engine/Source/Runtime/Engine/Public/Materials/MaterialExpressionComposite.h
--- a/Engine/Source/Runtime/Engine/Public/Materials/MaterialExpressionComposite.h
+++ b/Engine/Source/Runtime/Engine/Public/Materials/MaterialExpressionComposite.h
@@ -87,6 +87,8 @@ class UMaterialExpressionComposite : public UMaterialExpression
 	//~ Begin UObject Interface.
 	virtual bool Modify(bool bAlwaysMarkDirty = true) override;
 	//~ End UObject Interface.
+
+	bool ModifyCompositeOnly(bool bAlwaysMarkDirty = true);
 #endif // WITH_EDITOR
 };
 
diff --git a/Engine/Source/Runtime/Engine/Public/Materials/MaterialExpressionPinBase.h b/Engine/Source/Runtime/Engine/Public/Materials/MaterialExpressionPinBase.h
--- a/Engine/Source/Runtime/Engine/Public/Materials/MaterialExpressionPinBase.h
+++ b/Engine/Source/Runtime/Engine/Public/Materials/MaterialExpressionPinBase.h
@@ -96,8 +96,7 @@ class UMaterialExpressionPinBase : public UMaterialExpression
 	virtual bool Modify(bool bAlwaysMarkDirty = true) override;
 	//~ End UObject Interface.
 
-	virtual bool CanUserDeleteExpression() const override;
-
+	bool ModifyAsSubObject(bool bAlwaysMarkDirty = true);
 private:
 	/** Used to remove reroute expressions from the material after we remove them. */
 	TArray<UMaterialExpressionReroute*> PreEditRereouteExpresions;
