[HWRT] Fix a crash when required scratch buffer for BVH build is larger than 2GB

#rnx

diff --git a/Engine/Source/Runtime/Engine/Private/Rendering/RayTracingGeometryManager.cpp b/Engine/Source/Runtime/Engine/Private/Rendering/RayTracingGeometryManager.cpp
--- a/Engine/Source/Runtime/Engine/Private/Rendering/RayTracingGeometryManager.cpp
+++ b/Engine/Source/Runtime/Engine/Private/Rendering/RayTracingGeometryManager.cpp
@@ -1383,7 +1383,10 @@ void FRayTracingGeometryManager::ProcessBuildRequests(FRHIComputeCommandList& In
 	}
 
 	// kick actual build request to RHI command list
-	InCmdList.BuildAccelerationStructures(BuildParams);
+	if (BuildParams.Num())
+	{
+		InCmdList.BuildAccelerationStructures(BuildParams);
+	}
 
 	BuildParams.Reset();
 
@@ -1398,11 +1401,24 @@ void FRayTracingGeometryManager::ProcessBuildRequests(FRHIComputeCommandList& In
 
 void FRayTracingGeometryManager::SetupBuildParams(const FBuildRequest& InBuildRequest, TArray<FRayTracingGeometryBuildParams>& InBuildParams, TArray<FRayTracingGeometry*>& InReleaseBuffers, bool bRemoveFromRequestArray)
 {
+	FRHIRayTracingGeometry* RayTracingGeometryRHI = InBuildRequest.Owner->GetRHI();
+
 	check(InBuildRequest.RequestIndex != INDEX_NONE && InBuildRequest.RequestIndex == InBuildRequest.Owner->RayTracingBuildRequestIndex);
-	checkf(InBuildRequest.Owner->GetRHI() != nullptr, TEXT("Build request for FRayTracingGeometry without valid RHI. Was the FRayTracingGeometry evicted or released without calling RemoveBuildRequest()?"));
+	checkf(RayTracingGeometryRHI != nullptr, TEXT("Build request for FRayTracingGeometry without valid RHI. Was the FRayTracingGeometry evicted or released without calling RemoveBuildRequest()?"));
+	
+	const uint64 MaxScratchBufferSize = 2147483647u;
+	const uint64 RequiredScratchBufferSize = InBuildRequest.BuildMode == EAccelerationStructureBuildMode::Build
+		? RayTracingGeometryRHI->GetSizeInfo().BuildScratchSize
+		: RayTracingGeometryRHI->GetSizeInfo().UpdateScratchSize;
+
+	if (RequiredScratchBufferSize > MaxScratchBufferSize)
+	{
+		UE_LOG(LogRayTracingGeometryManager, Warning, TEXT("Ray Tracing Geometry (%s) with %d primitives requires too large scratch buffer (%llu) - skipping the build."), *InBuildRequest.Owner->Initializer.DebugName.ToString(), InBuildRequest.Owner->Initializer.TotalPrimitiveCount, RequiredScratchBufferSize);
+		return;
+	}
 
 	FRayTracingGeometryBuildParams BuildParam;
-	BuildParam.Geometry = InBuildRequest.Owner->GetRHI();
+	BuildParam.Geometry = RayTracingGeometryRHI;
 	BuildParam.BuildMode = InBuildRequest.BuildMode;
 	InBuildParams.Add(BuildParam);
 
