When /Script/Engine.RendererSettings/OverrideSVTTexture is enabled, the SVT type can be overridden during serialization (for example, the desktop type can be used for the mobile renderer). Setting is disabled by default.

#rnx

diff --git a/Engine/Source/Runtime/Engine/Private/VT/VirtualTextureBuilder.cpp b/Engine/Source/Runtime/Engine/Private/VT/VirtualTextureBuilder.cpp
--- a/Engine/Source/Runtime/Engine/Private/VT/VirtualTextureBuilder.cpp
+++ b/Engine/Source/Runtime/Engine/Private/VT/VirtualTextureBuilder.cpp
@@ -10,6 +10,7 @@
 
 #if WITH_EDITOR
 #include "Interfaces/ITargetPlatform.h"
+#include "Misc/ConfigCacheIni.h"
 #endif
 
 #include UE_INLINE_GENERATED_CPP_BY_NAME(VirtualTextureBuilder)
@@ -38,6 +39,20 @@ void UVirtualTextureBuilder::Serialize(FArchive& Ar)
 			TextureMobile = nullptr;
 		}
 
+		int32 OverrideSVTTexture = Ar.CookingTarget()->GetConfigSystem()->GetIntOrDefault(TEXT("/Script/Engine.RendererSettings"), TEXT("OverrideSVTTexture"), -1, GEngineIni);
+
+		if (OverrideSVTTexture >= 0)
+		{
+			if (OverrideSVTTexture == 0)
+			{
+				TextureMobile = Texture;
+			}
+			else
+			{
+				Texture = TextureMobile;
+			}
+		}
+
 		// Selectivly serialize VirtualTexture for platforms that support Deferred or/and Mobile rendering
 		if (!Ar.CookingTarget()->SupportsFeature(ETargetPlatformFeatures::DeferredRendering) && bSeparateTextureForMobile)
 		{
