Wrong tangent used in FNewSpline::GetSegmentLength

diff --git a/Engine/Source/Runtime/Engine/Private/Curves/Spline.cpp b/Engine/Source/Runtime/Engine/Private/Curves/Spline.cpp
--- a/Engine/Source/Runtime/Engine/Private/Curves/Spline.cpp
+++ b/Engine/Source/Runtime/Engine/Private/Curves/Spline.cpp
@@ -1946,7 +1946,7 @@ float FNewSpline::GetSegmentLength(const int32 Index, const float Param, const F
 	const auto& P0 = GetSpline().GetValue(Index);
 	const auto& T0 = GetSpline().GetTangentOut(Index);
 	const auto& P1 = GetSpline().GetValue(Index == LastPoint ? 0 : Index + 1);
-	const auto& T1 = GetSpline().GetTangentOut(Index == LastPoint ? 0 : Index + 1);
+	const auto& T1 = GetSpline().GetTangentIn(Index == LastPoint ? 0 : Index + 1);
 
 	// Special cases for linear
 	if (GetSpline().GetTangentModes()[Index] == UE::Geometry::Spline::ETangentMode::Linear)
