SplineComponent: Add Spline vs SplineCurves validation

This is by default non-functional, all new code is gated behind new CVars. This is a precautionary measure.

diff --git a/Engine/Source/Runtime/Engine/Classes/Components/SplineComponent.h b/Engine/Source/Runtime/Engine/Classes/Components/SplineComponent.h
--- a/Engine/Source/Runtime/Engine/Classes/Components/SplineComponent.h
+++ b/Engine/Source/Runtime/Engine/Classes/Components/SplineComponent.h
@@ -45,6 +45,14 @@ namespace ESplineCoordinateSpace
 	};
 }
 
+UENUM()
+enum class ELastAuthority : uint8
+{
+	Unset,
+	Spline,
+	SplineCurves
+};
+
 UCLASS(Abstract, MinimalAPI)
 class USplineMetadata : public UObject
 {
@@ -1064,6 +1072,12 @@ PRAGMA_DISABLE_DEPRECATION_WARNINGS
 PRAGMA_ENABLE_DEPRECATION_WARNINGS
 	}
 
+	UPROPERTY()
+	ELastAuthority LastAuthority = ELastAuthority::Unset;
+	
+	/** Checks for consistency between Spline and SplineCurves. */
+	bool Validate() const;
+	
 	/** Set the SplineCurves with the default shape i.e. a spline of two points (Used by default constructor) */
 	ENGINE_API void SetDefaultSpline();
 
diff --git a/Engine/Source/Runtime/Engine/Private/Components/SplineComponent.cpp b/Engine/Source/Runtime/Engine/Private/Components/SplineComponent.cpp
--- a/Engine/Source/Runtime/Engine/Private/Components/SplineComponent.cpp
+++ b/Engine/Source/Runtime/Engine/Private/Components/SplineComponent.cpp
@@ -36,6 +36,8 @@
 
 #define SPLINE_FAST_BOUNDS_CALCULATION 0
 
+DEFINE_LOG_CATEGORY_STATIC(LogSplineComponent, Log, All);
+
 const FInterpCurvePointVector USplineComponent::DummyPointPosition(0.0f, FVector::ZeroVector, FVector::ForwardVector, FVector::ForwardVector, CIM_Constant);
 const FInterpCurvePointQuat USplineComponent::DummyPointRotation(0.0f, FQuat::Identity);
 const FInterpCurvePointVector USplineComponent::DummyPointScale(0.0f, FVector::OneVector);
@@ -47,9 +49,76 @@ bool GUseSplineCurves = true;
 FAutoConsoleVariableRef CVarUseSplineCurves(
 	TEXT("SplineComponent.UseSplineCurves"),
 	GUseSplineCurves,
-	TEXT("When true, SplineCurves is the 'source of truth'.")
+	TEXT("When true, SplineCurves is the authoritative backing data.")
+);
+
+bool GValidateOnChange = false;
+FAutoConsoleVariableRef CVarValidateOnChange(
+	TEXT("SplineComponent.ValidateOnChange"),
+	GValidateOnChange,
+	TEXT("When true, the non-authoritative backing data is validated against the authoritative backing data when changing any data.")
+);
+
+bool GValidateOnLoad = false;
+FAutoConsoleVariableRef CVarValidateOnLoad(
+	TEXT("SplineComponent.ValidateOnLoad"),
+	GValidateOnLoad,
+	TEXT("When true, consistency between authoritative and non-authoritative backing data is validated at load time.")
+);
+
+bool GSynchronizeOnLoad = false;
+FAutoConsoleVariableRef CVarSynchronizeOnLoad(
+	TEXT("SplineComponent.SynchronizeOnLoad"),
+	GSynchronizeOnLoad,
+	TEXT("When true, consistency between authoritative and non-authoritative backing data is synchronized if inconsistent at load time.")
+);
+	
+bool GValidateOnSave = false;
+FAutoConsoleVariableRef CVarValidateOnSave(
+	TEXT("SplineComponent.ValidateOnSave"),
+	GValidateOnSave,
+	TEXT("When true, consistency between authoritative and non-authoritative backing data is validated at save time.")
+);
+
+bool GFailedValidationEnsures = false;
+FAutoConsoleVariableRef CVarFailedValidationEnsures(
+	TEXT("SplineComponent.FailedValidationEnsures"),
+	GFailedValidationEnsures,
+	TEXT("When true, all validation checks will ensure if not passed.")
 );
 
+/**
+ * This option exists because the tangent computation code for CurveClamped points is just not consistent.
+ * I suspect the FSplineCurves logic has a bug, but this has not yet been confirmed.
+ * This is not a major concern because auto-tangents can be considered 'transient' data because, while they are serialized,
+ * they are overwritten when USplineComponent::UpdateSpline is invoked.
+ */
+bool GIgnoreCurveClampedPointTangentsForValidation = true;
+	FAutoConsoleVariableRef CVarIgnoreCurveClampedPointTangentsForValidation(
+	TEXT("SplineComponent.IgnoreCurveClampedPointTangentsForValidation"),
+	GIgnoreCurveClampedPointTangentsForValidation,
+	TEXT("When true, validation checks will not consider inconsistent tangents to be invalid when the point type is CurveClamped.")
+);
+
+int GForceLastAuthority = 0;
+FAutoConsoleVariableRef CVarForceLastAuthority(
+	TEXT("SplineComponent.ForceLastAuthority"),
+	GForceLastAuthority,
+	TEXT("If 1, synchronization on load will treat FSplineCurves as the authority. If 2, synchronization on load will tread FSpline as the authority.")
+);
+
+static ELastAuthority GetEffectiveLastAuthority(const ELastAuthority InLastAuthority)
+{
+	ELastAuthority UseLastAuthority = InLastAuthority;
+	switch (GForceLastAuthority)
+	{
+	case 1: UseLastAuthority = ELastAuthority::SplineCurves; break;
+	case 2: UseLastAuthority = ELastAuthority::Spline; break;
+	default: break;
+	}
+	return UseLastAuthority;
+}
+
 static bool ShouldUseSplineCurves()
 {
 	return GUseSplineCurves || FSpline::IsEnabledStatic() == false;
@@ -62,6 +131,130 @@ USplineMetadata::USplineMetadata(const FObjectInitializer& ObjectInitializer)
 {
 }
 
+bool USplineComponent::Validate() const
+{
+	if (!Spline.IsEnabled())
+	{
+		// Can't be inconsistent if Spline is incapable of storing any data at all.
+		return true;
+	}
+
+	auto StringifySplinePointType = [](const ESplinePointType::Type InSplinePointType)
+	{
+		switch (InSplinePointType)
+		{
+		case ESplinePointType::Linear: return TEXT("Linear");
+		case ESplinePointType::Curve: return TEXT("Curve");
+		case ESplinePointType::Constant: return TEXT("Constant");
+		case ESplinePointType::CurveClamped: return TEXT("CurveClamped");
+		case ESplinePointType::CurveCustomTangent: return TEXT("CurveCustomTangent");
+		default: return TEXT("Unknown");
+		}
+	};
+
+	bool bValid = true;
+	
+	const int32 NumSplineCurvesPoints = WarninglessSplineCurves().Position.Points.Num();
+	const int32 NumSplinePoints = Spline.GetNumControlPoints();
+	const int32 NumPoints = FMath::Min(NumSplinePoints, NumSplineCurvesPoints);
+
+	UE_LOG(LogSplineComponent, Display, TEXT("Validating SplineComponent (%s) - Authority: %s"), *GetPathName(), UE::SplineComponent::ShouldUseSplineCurves() ? TEXT("FSplineCurves") : TEXT("FSpline"));
+	
+	if (NumSplineCurvesPoints != NumSplinePoints)
+	{
+		UE_LOG(LogSplineComponent, Display, TEXT("Internally inconsistent number of points (%d vs %d)!"), NumSplineCurvesPoints, NumSplinePoints);
+		bValid = false;
+	}
+
+	for (int32 Idx = 0; Idx < NumPoints; ++Idx)
+	{
+		const auto& SplineCurvesPositionPoint = WarninglessSplineCurves().Position.Points[Idx];
+		const auto& SplineCurvesRotationPoint = WarninglessSplineCurves().Rotation.Points[Idx];
+		const auto& SplineCurvesScalePoint = WarninglessSplineCurves().Scale.Points[Idx];
+
+		const ESplinePointType::Type SplineCurvesPointType = ConvertInterpCurveModeToSplinePointType(SplineCurvesPositionPoint.InterpMode);
+		const ESplinePointType::Type SplinePointType = ConvertInterpCurveModeToSplinePointType(Spline.GetSplinePointType(Idx));
+		if (SplineCurvesPointType != SplinePointType)
+		{
+			UE_LOG(LogSplineComponent, Display, TEXT("Internally inconsistent point type (%s vs %s) at index %d!"), StringifySplinePointType(SplineCurvesPointType), StringifySplinePointType(SplinePointType), Idx );
+			bValid = false;
+		}
+		
+		const FVector SplineCurvesPosition = SplineCurvesPositionPoint.OutVal;
+		const FVector SplinePosition = Spline.GetLocation(Idx);
+		if (!SplineCurvesPosition.Equals(SplinePosition))
+		{
+			UE_LOG(LogSplineComponent, Display, TEXT("Internally inconsistent position (%s vs %s) at index %d!"), *SplineCurvesPosition.ToCompactString(), *SplinePosition.ToCompactString(), Idx );
+			bValid = false;
+		}
+
+		const FVector SplineCurvesInTangent = SplineCurvesPositionPoint.ArriveTangent;
+		const FVector SplineInTangent = Spline.GetInTangent(Idx);
+		if (!SplineCurvesInTangent.Equals(SplineInTangent))
+		{
+			UE_LOG(LogSplineComponent, Display, TEXT("Internally inconsistent arrive tangent (%s vs %s) at index %d!"), *SplineCurvesInTangent.ToCompactString(), *SplineInTangent.ToCompactString(), Idx );
+			UE_LOG(LogSplineComponent, Display, TEXT("Point types: %s and %s"), StringifySplinePointType(SplineCurvesPointType), StringifySplinePointType(SplinePointType));
+			
+			if (SplineCurvesPointType != ESplinePointType::CurveClamped || !UE::SplineComponent::GIgnoreCurveClampedPointTangentsForValidation)
+			{
+				bValid = false;
+			}
+			else
+			{
+				UE_LOG(LogSplineComponent, Display, TEXT("Ignoring inconsistent arrive tangent for CurveClamped point."));
+			}
+		}
+
+		const FVector SplineCurvesOutTangent = SplineCurvesPositionPoint.LeaveTangent;
+		const FVector SplineOutTangent = Spline.GetOutTangent(Idx);
+		if (!SplineCurvesOutTangent.Equals(SplineOutTangent))
+		{
+			UE_LOG(LogSplineComponent, Display, TEXT("Internally inconsistent leave tangent (%s vs %s) at index %d!"), *SplineCurvesOutTangent.ToCompactString(), *SplineOutTangent.ToCompactString(), Idx );
+			UE_LOG(LogSplineComponent, Display, TEXT("Point types: %s and %s"), StringifySplinePointType(SplineCurvesPointType), StringifySplinePointType(SplinePointType));
+			
+			if (SplineCurvesPointType != ESplinePointType::CurveClamped || !UE::SplineComponent::GIgnoreCurveClampedPointTangentsForValidation)
+			{
+				bValid = false;
+			}
+			else
+			{
+				UE_LOG(LogSplineComponent, Display, TEXT("Ignoring inconsistent leave tangent for CurveClamped point."));
+			}
+		}
+
+		const float SplineCurvesParam = SplineCurvesPositionPoint.InVal;
+		const float SplineParam = Spline.GetParameterAtIndex(Idx);
+		if (!FMath::IsNearlyEqual(SplineCurvesParam, SplineParam))
+		{
+			UE_LOG(LogSplineComponent, Display, TEXT("Internally inconsistent parameter (%f vs %f) at index %d!"), SplineCurvesParam, SplineParam, Idx );
+			bValid = false;
+		}
+
+		const FQuat SplineCurvesRotation = SplineCurvesRotationPoint.OutVal;
+		const FQuat SplineRotation = Spline.GetRotation(Idx);
+		if (!SplineCurvesRotation.Equals(SplineRotation))
+		{
+			UE_LOG(LogSplineComponent, Display, TEXT("Internally inconsistent rotation (%s vs %s) at index %d!"), *SplineCurvesRotation.ToString(), *SplineRotation.ToString(), Idx );
+			bValid = false;
+		}
+
+		const FVector SplineCurvesScale = SplineCurvesScalePoint.OutVal;
+		const FVector SplineScale = Spline.GetScale(Idx);
+		if (!SplineCurvesScale.Equals(SplineScale))
+		{
+			UE_LOG(LogSplineComponent, Display, TEXT("Internally inconsistent scale (%s vs %s) at index %d!"), *SplineCurvesScale.ToCompactString(), *SplineScale.ToCompactString(), Idx );
+			bValid = false;
+		}
+	}
+
+	if (UE::SplineComponent::GFailedValidationEnsures)
+	{
+		ensureAlways(bValid);
+	}
+
+	return bValid;
+}
+
 void USplineComponent::PopulateFromLegacy()
 {
 	SynchronizeSplines();
@@ -158,6 +351,11 @@ void USplineComponent::ResetToDefault()
 	bShouldVisualizeScale = false;
 	ScaleVisualizationWidth = 30.0f;
 #endif
+
+	if (UE::SplineComponent::GValidateOnChange)
+	{
+		Validate();
+	}
 }
 
 bool USplineComponent::CanResetToDefault() const
@@ -193,6 +391,11 @@ void USplineComponent::SetDefaultSpline()
     WarninglessSplineCurves().Scale.Points.Emplace(EndParam, FVector(1.0f), FVector::ZeroVector, FVector::ZeroVector, CIM_CurveAuto);
 
 	OnSplineChanged.Broadcast();
+
+	if (UE::SplineComponent::GValidateOnChange)
+	{
+		Validate();
+	}
 }
 
 void USplineComponent::GetLifetimeReplicatedProps(TArray<FLifetimeProperty>& OutLifetimeProps) const
@@ -234,6 +437,11 @@ ESplinePointType::Type ConvertInterpCurveModeToSplinePointType(EInterpCurveMode
 
 void USplineComponent::Serialize(FArchive& Ar)
 {
+	if (Ar.IsSaving())
+	{
+		LastAuthority = UE::SplineComponent::ShouldUseSplineCurves() ? ELastAuthority::SplineCurves : ELastAuthority::Spline;
+	}
+	
 	Super::Serialize(Ar);
 
 	Ar.UsingCustomVersion(FEditorObjectVersion::GUID);
@@ -257,9 +465,61 @@ void USplineComponent::Serialize(FArchive& Ar)
 			SynchronizeSplines();
 		}
 #endif
+		
+		if (UE::SplineComponent::GSynchronizeOnLoad && Spline.IsEnabled() && !Validate())
+		{
+			UE_LOG(LogSplineComponent, Display, TEXT("Internally inconsistent SplineComponent (%s) at load time!"), *GetPathName())
+			
+			switch (UE::SplineComponent::GetEffectiveLastAuthority(LastAuthority))
+			{
+			case ELastAuthority::Unset:	// Intentional fallthrough - at the time this property was added, the authority had been Spline for some time.
+			case ELastAuthority::Spline:
+				UE_LOG(LogSplineComponent, Display, TEXT("Populating SplineCurves from Spline..."))
+				{
+					FSplineCurves TemporarySplineCurves = FSplineCurves::FromSplineInterface(Spline);
+					for (int32 Idx = 0; Idx < WarninglessSplineCurves().Position.Points.Num(); ++Idx)
+					{
+						TemporarySplineCurves.Position.Points[Idx].InVal = WarninglessSplineCurves().Position.Points[Idx].InVal;
+						TemporarySplineCurves.Rotation.Points[Idx].InVal = WarninglessSplineCurves().Rotation.Points[Idx].InVal;
+						TemporarySplineCurves.Scale.Points[Idx].InVal = WarninglessSplineCurves().Scale.Points[Idx].InVal;
+					}
+					WarninglessSplineCurves() = TemporarySplineCurves;
+				}
+				break;
+
+			case ELastAuthority::SplineCurves:
+				UE_LOG(LogSplineComponent, Display, TEXT("Populating Spline from SplineCurves..."))
+				SynchronizeSplines();
+				break;
+			}
+
+			UpdateSpline();
+			
+			// see if we got it right...
+			if (!Validate())
+			{
+				UE_LOG(LogSplineComponent, Warning, TEXT("Internally inconsistent SplineComponent (%s) after synchronization on load!"), *GetPathName())
+			}
+			else
+			{
+				UE_LOG(LogSplineComponent, Display, TEXT("Successfully synchronized SplineComponent (%s) on load."), *GetPathName())
+			}
+		}
+		else if (UE::SplineComponent::GValidateOnLoad && Spline.IsEnabled() && !Validate())
+		{
+			UE_LOG(LogSplineComponent, Display, TEXT("Internally inconsistent spline data at load time!"))
+		}
 	}
 #endif
 
+	if (Ar.IsSaving())
+	{
+		if (UE::SplineComponent::GValidateOnSave && !Validate())
+		{
+			UE_LOG(LogSplineComponent, Display, TEXT("Internally inconsistent spline data at save time!"))
+		}
+	}
+	
 	// Support old resources which don't have the rotation and scale splines present
 	const FPackageFileVersion ArchiveUEVersion = Ar.UEVer();
 	if (ArchiveUEVersion < VER_UE4_INTERPCURVE_SUPPORTS_LOOPING)
@@ -394,6 +654,11 @@ void USplineComponent::UpdateSpline()
 #endif
 
 	OnSplineUpdated.Broadcast();
+
+	if (UE::SplineComponent::GValidateOnChange)
+	{
+		Validate();
+	}
 }
 
 void USplineComponent::SetOverrideConstructionScript(bool bInOverride)
@@ -887,6 +1152,11 @@ void USplineComponent::SetClosedLoop(bool bInClosedLoop, bool bUpdateSpline)
 	}
 
 	OnSplineChanged.Broadcast();
+
+	if (UE::SplineComponent::GValidateOnChange)
+	{
+		Validate();
+	}
 }
 
 void USplineComponent::SetClosedLoopAtPosition(bool bInClosedLoop, float Key, bool bUpdateSpline)
@@ -901,6 +1171,11 @@ void USplineComponent::SetClosedLoopAtPosition(bool bInClosedLoop, float Key, bo
 	}
 
 	OnSplineChanged.Broadcast();
+
+	if (UE::SplineComponent::GValidateOnChange)
+	{
+		Validate();
+	}
 }
 
 bool USplineComponent::IsClosedLoop() const
@@ -953,6 +1228,11 @@ void USplineComponent::ClearSplinePoints(bool bUpdateSpline)
 	{
 		UpdateSpline();
 	}
+
+	if (UE::SplineComponent::GValidateOnChange)
+	{
+		Validate();
+	}
 }
 
 static int32 UpperBound(const TArray<FInterpCurvePoint<FVector>>& SplinePoints, float Value)
@@ -1038,6 +1318,11 @@ void USplineComponent::AddPoint(const FSplinePoint& InSplinePoint, bool bUpdateS
 		UpdateSpline();
 	}
 
+	if (UE::SplineComponent::GValidateOnChange)
+	{
+		Validate();
+	}
+
 	OnSplineChanged.Broadcast();
 }
 
@@ -1059,6 +1344,11 @@ void USplineComponent::AddPoints(const TArray<FSplinePoint>& InSplinePoints, boo
 		UpdateSpline();
 	}
 
+	if (UE::SplineComponent::GValidateOnChange)
+	{
+		Validate();
+	}
+	
 	OnSplineChanged.Broadcast();
 }
 
@@ -1104,6 +1394,11 @@ void USplineComponent::AddSplinePoint(const FVector& Position, ESplineCoordinate
 		UpdateSpline();
 	}
 
+	if (UE::SplineComponent::GValidateOnChange)
+	{
+		Validate();
+	}
+	
 	OnSplineChanged.Broadcast();
 }
 
@@ -1157,6 +1452,11 @@ void USplineComponent::AddSplinePointAtIndex(const FVector& Position, int32 Inde
 		UpdateSpline();
 	}
 
+	if (UE::SplineComponent::GValidateOnChange)
+	{
+		Validate();
+	}
+	
 	OnSplineChanged.Broadcast();
 }
 
@@ -1203,6 +1503,11 @@ void USplineComponent::RemoveSplinePoint(int32 Index, bool bUpdateSpline)
 		UpdateSpline();
 	}
 
+	if (UE::SplineComponent::GValidateOnChange)
+	{
+		Validate();
+	}
+	
 	OnSplineChanged.Broadcast();
 }
 
@@ -1257,6 +1562,11 @@ void USplineComponent::SetSplinePoints(const TArray<FVector>& Points, ESplineCoo
 		UpdateSpline();
 	}
 
+	if (UE::SplineComponent::GValidateOnChange)
+	{
+		Validate();
+	}
+
 	OnSplineChanged.Broadcast();
 }
 
@@ -1277,6 +1587,11 @@ void USplineComponent::SetLocationAtSplinePoint(int32 PointIndex, const FVector&
 			UpdateSpline();
 		}
 
+		if (UE::SplineComponent::GValidateOnChange)
+		{
+			Validate();
+		}
+		
 		OnSplineChanged.Broadcast();
 	}
 }
@@ -1309,6 +1624,11 @@ void USplineComponent::SetTangentsAtSplinePoint(int32 PointIndex, const FVector&
 			UpdateSpline();
 		}
 
+		if (UE::SplineComponent::GValidateOnChange)
+		{
+			Validate();
+		}
+
 		OnSplineChanged.Broadcast();
 	}
 }
@@ -1331,6 +1651,11 @@ void USplineComponent::SetUpVectorAtSplinePoint(int32 PointIndex, const FVector&
 			UpdateSpline();
 		}
 
+		if (UE::SplineComponent::GValidateOnChange)
+		{
+			Validate();
+		}
+		
 		OnSplineChanged.Broadcast();
 	}
 }
@@ -1374,6 +1699,8 @@ void USplineComponent::SetQuaternionAtSplinePoint(int32 PointIndex, const FQuat&
 			UpdateSpline();
 		}
 
+		// no need to validate, handled by SetTangentsAtSplinePoint
+		
 		OnSplineChanged.Broadcast();
 	}
 }
@@ -1392,6 +1719,11 @@ void USplineComponent::SetScaleAtSplinePoint(int32 PointIndex, const FVector& In
 			UpdateSpline();
 		}
 
+		if (UE::SplineComponent::GValidateOnChange)
+		{
+			Validate();
+		}
+
 		OnSplineChanged.Broadcast();
 	}
 }
@@ -1426,6 +1758,11 @@ void USplineComponent::SetSplinePointType(int32 PointIndex, ESplinePointType::Ty
 			UpdateSpline();
 		}
 
+		if (UE::SplineComponent::GValidateOnChange)
+		{
+			Validate();
+		}
+		
 		OnSplineChanged.Broadcast();
 	}
 }
@@ -3123,6 +3460,11 @@ void USplineComponent::SetSpline(const FSpline& InSpline)
 {
 	WarninglessSplineCurves() = FSplineCurves::FromSplineInterface(InSpline);
 	Spline = InSpline;
+
+	if (UE::SplineComponent::GValidateOnChange)
+	{
+		Validate();
+	}
 	
 	OnSplineChanged.Broadcast();
 }
@@ -3132,6 +3474,11 @@ void USplineComponent::SetSpline(const FSplineCurves& InSplineCurves)
 	WarninglessSplineCurves() = InSplineCurves;
 	Spline = InSplineCurves;
 
+	if (UE::SplineComponent::GValidateOnChange)
+	{
+		Validate();
+	}
+	
 	OnSplineChanged.Broadcast();
 }
 
@@ -3316,6 +3663,11 @@ void USplineComponent::PostTransacted(const FTransactionObjectEvent& Transaction
 	// Notify listeners that the spline data may now be invalid.
 	if (TransactionEvent.GetEventType() == ETransactionObjectEventType::UndoRedo)
 	{
+		if (UE::SplineComponent::GValidateOnChange)
+		{
+			Validate();
+		}
+		
 		OnSplineChanged.Broadcast();
 	}
 }
