Splines: Reparameterize FNewSpline in UpdateSpline

diff --git a/Engine/Source/Runtime/Engine/Private/Curves/Spline.cpp b/Engine/Source/Runtime/Engine/Private/Curves/Spline.cpp
--- a/Engine/Source/Runtime/Engine/Private/Curves/Spline.cpp
+++ b/Engine/Source/Runtime/Engine/Private/Curves/Spline.cpp
@@ -1903,6 +1903,7 @@ bool FNewSpline::Serialize(FArchive& Ar)
 	// Do not serialize legacy curves, only generate if loading.
 	if (Ar.IsLoading())
 	{
+		GetSpline().Reparameterize();
 		RebuildLegacyCurves();
 	}
 	
@@ -2449,8 +2450,7 @@ void FNewSpline::UpdateSpline()
     GetSpline().UpdateTangents();			// updates the tangents on our version of the points
 	
 	// We do this so that parameterization remains proportional to the square root of segment chord lengths.
-	// todo: reenable this code when underlying knot vector issues are resolved.
-	// GetSpline().Reparameterize(EParameterizationPolicy::Centripetal);
+	GetSpline().Reparameterize();
 	
     MarkReparamTableDirty();
     RebuildLegacyCurves();					// updates legacy curves based on points, never evaluates internal spline.
