Added missing member initialization from CL 46943010

diff --git a/Engine/Source/Runtime/Engine/Public/WorldPartition/WorldPartitionLevelHelper.h b/Engine/Source/Runtime/Engine/Public/WorldPartition/WorldPartitionLevelHelper.h
--- a/Engine/Source/Runtime/Engine/Public/WorldPartition/WorldPartitionLevelHelper.h
+++ b/Engine/Source/Runtime/Engine/Public/WorldPartition/WorldPartitionLevelHelper.h
@@ -52,6 +52,7 @@ public:
 			, DestLevel(nullptr)
 			, PackageReferencer(nullptr)
 			, bLoadAsync(false)
+			, bSilenceLoadFailures(false)
 			, AsyncRequestIDs(nullptr)
 		{}
 
