World Partition - HLOD: Fixed FWorldPartitionLevelHelper::LoadActorsInternal() to correctly remap soft object paths from the same container if no World Partition object is available

diff --git a/Engine/Source/Runtime/Engine/Private/WorldPartition/WorldPartitionLevelHelper.cpp b/Engine/Source/Runtime/Engine/Private/WorldPartition/WorldPartitionLevelHelper.cpp
--- a/Engine/Source/Runtime/Engine/Private/WorldPartition/WorldPartitionLevelHelper.cpp
+++ b/Engine/Source/Runtime/Engine/Private/WorldPartition/WorldPartitionLevelHelper.cpp
@@ -848,7 +848,15 @@ bool FWorldPartitionLevelHelper::LoadActorsInternal(FLoadActorsParams&& InParams
 					{
 						if (!Value.IsNull() && Value.GetAssetPathString().Equals(SourceWorldPath, ESearchCase::IgnoreCase))
 						{
-							OuterWorld->GetWorldPartition()->ConvertContainerPathToEditorPath(PackageObjectMapping->ContainerID, FSoftObjectPath(Value), Value);
+							if (OuterWorld->GetWorldPartition())
+							{
+								OuterWorld->GetWorldPartition()->ConvertContainerPathToEditorPath(PackageObjectMapping->ContainerID, FSoftObjectPath(Value), Value);
+							}
+							else
+							{
+								// Remap container path to source world path + container id
+								Value = RemapActorPath(PackageObjectMapping->ContainerID, SourceOuterWorldPath, Value);
+							}
 						}
 					});
 					FixupArchive.Fixup(Actor);
