Fixed issue with spline keys when deleting the first point in a spline.

#rnx

diff --git a/Engine/Source/Runtime/Engine/Private/Components/SplineComponent.cpp b/Engine/Source/Runtime/Engine/Private/Components/SplineComponent.cpp
--- a/Engine/Source/Runtime/Engine/Private/Components/SplineComponent.cpp
+++ b/Engine/Source/Runtime/Engine/Private/Components/SplineComponent.cpp
@@ -1468,9 +1468,15 @@ void USplineComponent::RemoveSplinePoint(int32 Index, bool bUpdateSpline)
 	{
 		Spline.RemovePoint(Index);
 
-		const float KnotDelta = Index > 0
-			? WarninglessSplineCurves().Position.Points[Index].InVal - WarninglessSplineCurves().Position.Points[Index - 1].InVal
-			: WarninglessSplineCurves().Position.Points[Index].InVal;
+		float KnotDelta = 0.0f;
+		if (Index > 0)
+		{
+			KnotDelta = WarninglessSplineCurves().Position.Points[Index].InVal - WarninglessSplineCurves().Position.Points[Index - 1].InVal;
+		}
+		else if (NumPoints >= 2) // e.g. Index == 0, Index < NumPoints - 1
+		{
+			KnotDelta = WarninglessSplineCurves().Position.Points[Index + 1].InVal - WarninglessSplineCurves().Position.Points[Index].InVal;
+		}
 		
 		WarninglessSplineCurves().Position.Points.RemoveAt(Index, EAllowShrinking::No);
 		WarninglessSplineCurves().Rotation.Points.RemoveAt(Index, EAllowShrinking::No);
