Splines: Add CVars which control the evaluation method used by FNewSpline

We have at least one unsolved problem related to evaluation of FQuat channels, however the fallback path works well. This change ensures we have a path forward with config changes if we fail to resolve the underlying issue for 38.00

preflight 68f126aada10cdc5c3412597

diff --git a/Engine/Source/Runtime/Engine/Private/Curves/Spline.cpp b/Engine/Source/Runtime/Engine/Private/Curves/Spline.cpp
--- a/Engine/Source/Runtime/Engine/Private/Curves/Spline.cpp
+++ b/Engine/Source/Runtime/Engine/Private/Curves/Spline.cpp
@@ -67,6 +67,27 @@ FAutoConsoleVariableRef CVarFallbackFindNearest(
 	TEXT("True if we should implement FindNearest and FindNearestOnSegment using an intermediate spline representation. Only applies if Spline.Implementation == 2.")
 );
 
+bool GUseLegacyPositionEvaluation = false;
+FAutoConsoleVariableRef CVarUseLegacyPositionEvaluation(
+	TEXT("Spline.UseLegacyPositionEvaluation"),
+	GUseLegacyPositionEvaluation,
+	TEXT("If true, evaluating the position channel always routes through an interp curve.")
+);
+
+bool GUseLegacyRotationEvaluation = true;
+FAutoConsoleVariableRef CVarUseLegacyRotationEvaluation(
+	TEXT("Spline.UseLegacyRotationEvaluation"),
+	GUseLegacyRotationEvaluation,
+	TEXT("If true, evaluating the rotation channel always routes through an interp curve.")
+);
+
+bool GUseLegacyScaleEvaluation = false;
+FAutoConsoleVariableRef CVarUseLegacyScaleEvaluation(
+	TEXT("Spline.UseLegacyScaleEvaluation"),
+	GUseLegacyScaleEvaluation,
+	TEXT("If true, evaluating the scale channel always routes through an interp curve.")
+);
+
 bool GImmediatelyUpdateLegacyCurves = false;
 FAutoConsoleVariableRef CVarImmediatelyUpdateLegacyCurves(
 	TEXT("Spline.ImmediatelyUpdateLegacyCurves"),
@@ -292,11 +313,11 @@ private:
 	static inline const FName ScaleAttrName = FName("Scale");
 
 	/** Legacy Curves: */
-    FInterpCurveVector PositionCurve;
-	FInterpCurveQuat RotationCurve;
-	FInterpCurveVector ScaleCurve;
+    mutable FInterpCurveVector PositionCurve;
+	mutable FInterpCurveQuat RotationCurve;
+	mutable FInterpCurveVector ScaleCurve;
 	FChangeStamp LegacyCurvesNextVersion;
-	uint32 LegacyCurvesVersion = 0;
+	mutable uint32 LegacyCurvesVersion = 0;
 	
 	int32 ReparamStepsPerSegment = 10;
 
@@ -326,7 +347,7 @@ private:
 	void MarkLegacyCurvesDirty();
 
 	/** Clears and repopulates PositionCurve, RotationCurve, and ScaleCurve */
-	void RebuildLegacyCurves();
+	void RebuildLegacyCurves() const;
 
 	/** Converts FSplinePoint to FTangentBezierControlPoint */
     FTangentBezierControlPoint ConvertToTangentBezierControlPoint(const FSplinePoint& Point) const;
@@ -2426,7 +2447,15 @@ float FNewSpline::FindNearestOnSegment(const FVector& InLocation, int32 SegmentI
 
 FVector FNewSpline::EvaluatePosition(float Parameter) const
 {
-	return Evaluate(ToInternalSplineSpace(Parameter));
+	if (UE::Spline::GUseLegacyPositionEvaluation)
+	{
+		RebuildLegacyCurves();
+		return PositionCurve.Eval(Parameter);
+	}
+	else
+	{
+		return Evaluate(ToInternalSplineSpace(Parameter));
+	}
 }
 
 FVector FNewSpline::EvaluateDerivative(float Parameter) const
@@ -2437,12 +2466,28 @@ FVector FNewSpline::EvaluateDerivative(float Parameter) const
 
 FQuat FNewSpline::EvaluateRotation(float Parameter) const
 {
-    return EvaluateAttribute<FQuat>(RotationAttrName, ToInternalSplineSpace(Parameter));
+	if (UE::Spline::GUseLegacyRotationEvaluation)
+	{
+		RebuildLegacyCurves();
+		return RotationCurve.Eval(Parameter);
+	}
+	else
+	{
+		return EvaluateAttribute<FQuat>(RotationAttrName, ToInternalSplineSpace(Parameter));
+	}
 }
 
 FVector FNewSpline::EvaluateScale(float Parameter) const
 {
-    return EvaluateAttribute<FVector>(ScaleAttrName, ToInternalSplineSpace(Parameter));
+	if (UE::Spline::GUseLegacyScaleEvaluation)
+	{
+		RebuildLegacyCurves();
+		return ScaleCurve.Eval(Parameter);
+	}
+	else
+	{
+		return EvaluateAttribute<FVector>(ScaleAttrName, ToInternalSplineSpace(Parameter));
+	}
 }
 
 void FNewSpline::UpdateSpline()
@@ -2551,7 +2596,7 @@ void FNewSpline::MarkLegacyCurvesDirty()
 	}
 }
 
-void FNewSpline::RebuildLegacyCurves()
+void FNewSpline::RebuildLegacyCurves() const
 {
 	// Curves are dirty if:
 	// 1) Current version is not the same as next version (if next version == version, it has not been dirtied)
