Add rendering support for FX system creation via anim notifies, etc
NOTE:  Nigara library and UFXSystemComponent changes, so confining change to current season and will need to redo this work in target branches if desired for future seasons
NOTE:  Cascade not yet supported, but could be as follow on work
NOTE:  Tested with full list of FX requested from Tech Art here - https://docs.google.com/document/d/1XzI4CKR4iepp_OGVL7XmCOlLcyHx-s0njsc7UqY-RXg/edit?tab=t.0
#rnx

diff --git a/Engine/Plugins/FX/Niagara/Source/Niagara/Private/NiagaraFunctionLibrary.cpp b/Engine/Plugins/FX/Niagara/Source/Niagara/Private/NiagaraFunctionLibrary.cpp
--- a/Engine/Plugins/FX/Niagara/Source/Niagara/Private/NiagaraFunctionLibrary.cpp
+++ b/Engine/Plugins/FX/Niagara/Source/Niagara/Private/NiagaraFunctionLibrary.cpp
@@ -78,6 +78,8 @@ UNiagaraComponent* CreateNiagaraSystem(UNiagaraSystem* SystemTemplate, UWorld* W
 			NiagaraComponent->SetAutoDestroy(bAutoDestroy);
 			NiagaraComponent->bAllowAnyoneToDestroyMe = true;
 			NiagaraComponent->SetVisibleInRayTracing(false);
+
+			UFXSystemComponent::OnNewFXSystemComponentCreated.Broadcast(NiagaraComponent);
 		}
 		else if (FNiagaraWorldManager* WorldManager = FNiagaraWorldManager::Get(World))
 		{
diff --git a/Engine/Source/Runtime/Engine/Classes/Particles/ParticleSystemComponent.h b/Engine/Source/Runtime/Engine/Classes/Particles/ParticleSystemComponent.h
--- a/Engine/Source/Runtime/Engine/Classes/Particles/ParticleSystemComponent.h
+++ b/Engine/Source/Runtime/Engine/Classes/Particles/ParticleSystemComponent.h
@@ -375,6 +375,9 @@ struct FFXSystemSpawnParameters
 	bool bIsPlayerEffect = false;
 };
 
+// Callback for when a new FX System Component object is created via anim notifies, etc
+DECLARE_MULTICAST_DELEGATE_OneParam(FNewFXSystemComponentCreatedEvent, UFXSystemComponent* /*InNewFXSystemComponent*/);
+
 UCLASS(Abstract, MinimalAPI)
 class UFXSystemComponent : public UPrimitiveComponent
 {
@@ -471,6 +474,9 @@ public:
 	/** Returns true if we have crossed LWC tiles to the point that we may introduce artifacts. */
 	static ENGINE_API bool RequiresLWCTileRecache(const FVector3f CurrentTile, const FVector CurrentLocation);
 
+	// Callback for when a new FX System Component object is created via anim notify, etc
+	static ENGINE_API FNewFXSystemComponentCreatedEvent OnNewFXSystemComponentCreated;
+
 #if WITH_PER_COMPONENT_PARTICLE_PERF_STATS
 	mutable FParticlePerfStats* ParticlePerfStats = nullptr;
 #endif
diff --git a/Engine/Source/Runtime/Engine/Private/Particles/ParticleSystemComponent.cpp b/Engine/Source/Runtime/Engine/Private/Particles/ParticleSystemComponent.cpp
--- a/Engine/Source/Runtime/Engine/Private/Particles/ParticleSystemComponent.cpp
+++ b/Engine/Source/Runtime/Engine/Private/Particles/ParticleSystemComponent.cpp
@@ -79,6 +79,9 @@ DEFINE_STAT(STAT_ParticlesOverview_RT_CNC);
 
 DEFINE_LOG_CATEGORY(LogParticles);
 
+// Callback for when a new FX System Component object is created via anim notify, etc
+FNewFXSystemComponentCreatedEvent UFXSystemComponent::OnNewFXSystemComponentCreated;
+
 const FGuid FParticleSystemCustomVersion::GUID(0x4A56EB40, 0x10F511DC, 0x92D3347E, 0xB2C96AE7);
 // Register the custom version with core
 FCustomVersionRegistration GRegisterParticleSystemCustomVersion(FParticleSystemCustomVersion::GUID, FParticleSystemCustomVersion::LatestVersion, TEXT("ParticleSystemVer"));
