SDCE, fixed invalid stripping of FNanitePostDeformVertex
- Added override CVar (from Jamie)

diff --git a/Engine/Source/Runtime/Core/Private/UObject/FortniteMainBranchObjectVersion.cpp b/Engine/Source/Runtime/Core/Private/UObject/FortniteMainBranchObjectVersion.cpp
--- a/Engine/Source/Runtime/Core/Private/UObject/FortniteMainBranchObjectVersion.cpp
+++ b/Engine/Source/Runtime/Core/Private/UObject/FortniteMainBranchObjectVersion.cpp
@@ -13,7 +13,7 @@ TMap<FGuid, FGuid> FFortniteMainBranchObjectVersion::GetSystemGuids()
 	SystemGuids.Add(DevGuids.MATERIALSHADERMAP_DERIVEDDATA_VER, FGuid("4DCE94C9E9994B3DA256F6A015B87CA9"));
 	SystemGuids.Add(DevGuids.NIAGARASHADERMAP_DERIVEDDATA_VER, FGuid("6360A977062842A29ED30E3A7ACB0E64"));
 	SystemGuids.Add(DevGuids.Niagara_LatestScriptCompileVersion, FGuid("783FF3700BB641C799B2D43552FBF91E"));
-	SystemGuids.Add(DevGuids.MaterialTranslationDDCVersion, FGuid("8F124D15BF2B4522933D5C6C52D6F3E0"));
+	SystemGuids.Add(DevGuids.MaterialTranslationDDCVersion, FGuid("0AE469C404D94DB5A5BDB77BB6CDE3F3"));
 
 	// These GUIDs are in AppendRestrictedFortniteMainStreamObjectVersionGuids as they should rarely be changed.
 	//SystemGuids.Add(DevGuids.NANITE_DERIVEDDATA_VER, FGuid("ACF98184978742A592537DA223A1E6BE"));
diff --git a/Engine/Source/Runtime/Engine/Private/Materials/HLSLMaterialTranslator.cpp b/Engine/Source/Runtime/Engine/Private/Materials/HLSLMaterialTranslator.cpp
--- a/Engine/Source/Runtime/Engine/Private/Materials/HLSLMaterialTranslator.cpp
+++ b/Engine/Source/Runtime/Engine/Private/Materials/HLSLMaterialTranslator.cpp
@@ -87,6 +87,11 @@ static FAutoConsoleVariableRef CVarEnableMaterialTranslationLogFile(
 	TEXT("Enables material translation log file generation for tracking materials translation time (written to  \"ShaderDebugInfo/MaterialTranslationLog-X.csv\")."),
 	ECVF_ReadOnly);
 
+TAutoConsoleVariable<bool> CVarShadersSDCEEnabled(
+	TEXT("r.Shaders.SDCE"),
+	true,
+	TEXT("Strip dead code using symbolic dead code elimination on structured member stores"));
+	
 /**
  * Utility used to create a MaterialTranslationLog.txt file during cooks that contains the list of all translated materials
  * and other info such as how long the translation took.
@@ -3093,7 +3098,7 @@ void FHLSLMaterialTranslator::GetMaterialEnvironment(EShaderPlatform InPlatform,
 		OutEnvironment.SetDefine(TEXT("SUBSTRATE_CLAMPED_CLOSURE_COUNT"), EnvironmentDefines->SubstrateClampedClosureCount);
 		OutEnvironment.SetDefine(TEXT("SUBSTRATE_COMPLEXSPECIALPATH"), EnvironmentDefines->bSubstrateComplexSpecialPath);
 
-		if (!Substrate::IsSubstrateBlendableGBufferEnabled(InPlatform))
+		if (CVarShadersSDCEEnabled.GetValueOnAnyThread() && !Substrate::IsSubstrateBlendableGBufferEnabled(InPlatform))
 		{
 			static const TCHAR* SDCESymbols[] = {
 				TEXT("FMaterialPixelParameters"),
@@ -3103,7 +3108,6 @@ void FHLSLMaterialTranslator::GetMaterialEnvironment(EShaderPlatform InPlatform,
 				TEXT("FCluster"),
 				TEXT("FNaniteView"),
 				TEXT("FInstanceDynamicData"),
-				TEXT("FNanitePostDeformVertex"),
 				TEXT("FMaterialVertexParameters"),
 				TEXT("FNanitePixelAttributes"),
 				TEXT("ViewState"),
