Fix for race condition when setting skin weight buffer.

diff --git a/Engine/Source/Runtime/Engine/Private/SkeletalRenderGPUSkin.cpp b/Engine/Source/Runtime/Engine/Private/SkeletalRenderGPUSkin.cpp
--- a/Engine/Source/Runtime/Engine/Private/SkeletalRenderGPUSkin.cpp
+++ b/Engine/Source/Runtime/Engine/Private/SkeletalRenderGPUSkin.cpp
@@ -1743,6 +1743,7 @@ void FSkeletalMeshObjectGPUSkin::FSkeletalMeshObjectLOD::UpdateSkinWeights(const
 			ENQUEUE_RENDER_COMMAND(UpdateSkinWeightsGPUSkin)(UE::RenderCommandPipe::SkeletalMesh,
 				[NewMeshObjectWeightBuffer, VertexBuffers, Self](FRHICommandList& RHICmdList)
 			{
+				Self->MeshObjectWeightBuffer_RenderThread = NewMeshObjectWeightBuffer;
 				Self->GPUSkinVertexFactories.UpdateVertexFactoryData(VertexBuffers);
 			});
 		}
@@ -1947,7 +1948,7 @@ FGPUBaseSkinVertexFactory const* FSkeletalMeshObjectGPUSkin::GetBaseSkinVertexFa
 const FSkinWeightVertexBuffer* FSkeletalMeshObjectGPUSkin::GetSkinWeightVertexBuffer(int32 LODIndex) const
 {
 	checkSlow(LODs.IsValidIndex(LODIndex));
-	return LODs[LODIndex].MeshObjectWeightBuffer;
+	return LODs[LODIndex].MeshObjectWeightBuffer_RenderThread;
 }
 
 FMatrix FSkeletalMeshObjectGPUSkin::GetTransform() const
@@ -2476,6 +2477,7 @@ void FSkeletalMeshObjectGPUSkin::FSkeletalMeshObjectLOD::InitResources(
 	ENQUEUE_RENDER_COMMAND(FSkeletalMeshObjectLOD_InitResources)(UE::RenderCommandPipe::SkeletalMesh,
 		[this, InGPUSkinCache, VertexBuffers, &LODData, InFeatureLevel, VertexAttributeMask](FRHICommandList& RHICmdList)
 	{
+		MeshObjectWeightBuffer_RenderThread = VertexBuffers.SkinWeightVertexBuffer;
 		GPUSkinVertexFactories.InitVertexFactories(RHICmdList, VertexBuffers, LODData.RenderSections, InFeatureLevel, VertexAttributeMask, GPUSkinTechnique);
 
 		if (DynamicBoundsStartOffset == INDEX_NONE && GPUSkinTechnique == ESkeletalMeshGPUSkinTechnique::GPUSkinCache)
diff --git a/Engine/Source/Runtime/Engine/Private/SkeletalRenderGPUSkin.h b/Engine/Source/Runtime/Engine/Private/SkeletalRenderGPUSkin.h
--- a/Engine/Source/Runtime/Engine/Private/SkeletalRenderGPUSkin.h
+++ b/Engine/Source/Runtime/Engine/Private/SkeletalRenderGPUSkin.h
@@ -579,6 +579,7 @@ private:
 			, LODIndex(InLOD)
 			, FeatureLevel(InFeatureLevel)
 			, MeshObjectWeightBuffer(nullptr)
+			, MeshObjectWeightBuffer_RenderThread(nullptr)
 			, MeshObjectColorBuffer(nullptr)
 			, GPUSkinTechnique(InSkinTechnique)
 		{
@@ -639,6 +640,7 @@ private:
 
 		/** Skin weight buffer to use, could be from asset or component override */
 		const FSkinWeightVertexBuffer* MeshObjectWeightBuffer;
+		const FSkinWeightVertexBuffer* MeshObjectWeightBuffer_RenderThread;
 
 		/** Color buffer to user, could be from asset or component override */
 		FColorVertexBuffer* MeshObjectColorBuffer;
