World Partition - HLOD: Clear legacy LODActors and Hierarchical LOD setup found in World Partition levels.
* Cleanup of this data was already performed in the conversion commandlet... However, sub levels to convert that had level scripts would retain their legacy setup.
* On WorldSettings PostLoad(), clear any left over Hierarchical LOD Setup
* Modified UHLODEngineSubsystem::CleanupHLOD() to also clear LODActors if some are still found in World Partition levels

diff --git a/Engine/Source/Editor/UnrealEd/Classes/Commandlets/WorldPartitionConvertCommandlet.h b/Engine/Source/Editor/UnrealEd/Classes/Commandlets/WorldPartitionConvertCommandlet.h
--- a/Engine/Source/Editor/UnrealEd/Classes/Commandlets/WorldPartitionConvertCommandlet.h
+++ b/Engine/Source/Editor/UnrealEd/Classes/Commandlets/WorldPartitionConvertCommandlet.h
@@ -65,6 +65,7 @@ protected:
 	UNREALED_API void FixupSoftObjectPaths(UPackage* OuterPackage);
 
 	UNREALED_API bool DetachDependantLevelPackages(ULevel* Level);
+	UNREALED_API void ClearHierarchicalLODSetup(ULevel* Level);
 	UNREALED_API bool RenameWorldPackageWithSuffix(UWorld* World);
 
 	UNREALED_API void SetupHLOD();
diff --git a/Engine/Source/Editor/UnrealEd/Private/Commandlets/WorldPartitionConvertCommandlet.cpp b/Engine/Source/Editor/UnrealEd/Private/Commandlets/WorldPartitionConvertCommandlet.cpp
--- a/Engine/Source/Editor/UnrealEd/Private/Commandlets/WorldPartitionConvertCommandlet.cpp
+++ b/Engine/Source/Editor/UnrealEd/Private/Commandlets/WorldPartitionConvertCommandlet.cpp
@@ -532,13 +532,18 @@ bool UWorldPartitionConvertCommandlet::DetachDependantLevelPackages(ULevel* Leve
 		Level->MapBuildData = nullptr;
 	}
 
+	return true;
+}
+
+void UWorldPartitionConvertCommandlet::ClearHierarchicalLODSetup(ULevel* Level)
+{
 	// Try to delete matching HLOD package
 	FHierarchicalLODUtilitiesModule& Module = FModuleManager::LoadModuleChecked<FHierarchicalLODUtilitiesModule>("HierarchicalLODUtilities");
 	IHierarchicalLODUtilities* Utilities = Module.GetUtilities();
 
 	const int32 NumHLODLevels = Level->GetWorldSettings()->GetNumHierarchicalLODLevels();
 
-	for (int32 HLODIndex=0; HLODIndex<NumHLODLevels; HLODIndex++)
+	for (int32 HLODIndex = 0; HLODIndex < NumHLODLevels; HLODIndex++)
 	{
 		if (UPackage* HLODPackage = Utilities->RetrieveLevelHLODPackage(Level, HLODIndex))
 		{
@@ -546,17 +551,7 @@ bool UWorldPartitionConvertCommandlet::DetachDependantLevelPackages(ULevel* Leve
 		}
 	}
 
-	for (AActor* Actor: Level->Actors)
-	{
-		if (IsValid(Actor) && Actor->IsA<ALODActor>())
-		{
-			Level->GetWorld()->DestroyActor(Actor);
-		}
-	}
-
 	Level->GetWorldSettings()->ResetHierarchicalLODSetup();
-
-	return true;
 }
 
 bool UWorldPartitionConvertCommandlet::RenameWorldPackageWithSuffix(UWorld* World)
@@ -1056,6 +1051,8 @@ int32 UWorldPartitionConvertCommandlet::Main(const FString& Params)
 			}
 		}
 
+		ClearHierarchicalLODSetup(Level);
+
 		if (bMainLevel)
 		{
 			if (LevelHasLevelScriptBlueprint(Level))
diff --git a/Engine/Source/Runtime/Engine/Private/HLOD/HLODEngineSubsystem.cpp b/Engine/Source/Runtime/Engine/Private/HLOD/HLODEngineSubsystem.cpp
--- a/Engine/Source/Runtime/Engine/Private/HLOD/HLODEngineSubsystem.cpp
+++ b/Engine/Source/Runtime/Engine/Private/HLOD/HLODEngineSubsystem.cpp
@@ -59,6 +59,7 @@ void UHLODEngineSubsystem::UnregisterRecreateLODActorsDelegates()
 {
 	FWorldDelegates::OnPostWorldInitialization.Remove(OnPostWorldInitializationDelegateHandle);
 	FWorldDelegates::LevelAddedToWorld.Remove(OnLevelAddedToWorldDelegateHandle);
+	FEditorDelegates::PreSaveWorldWithContext.Remove(OnPreSaveWorlDelegateHandle);
 }
 
 void UHLODEngineSubsystem::RegisterRecreateLODActorsDelegates()
@@ -97,6 +98,12 @@ void UHLODEngineSubsystem::RecreateLODActorsForLevel(ULevel* InLevel, UWorld* In
 		CleanupHLODs(InLevel);
 	}
 
+	// Legacy LODActors are not necessary anymore when using World Partition
+	if (InWorld->IsPartitionedWorld())
+	{
+		return;
+	}
+
 	// Look for HLODProxy packages associated with this level
 	int32 NumLODLevels = InLevel->GetWorldSettings()->GetHierarchicalLODSetup().Num();
 	for (int32 LODIndex = 0; LODIndex < NumLODLevels; ++LODIndex)
@@ -142,6 +149,11 @@ bool UHLODEngineSubsystem::CleanupHLOD(ALODActor* InLODActor)
 {
 	bool bShouldDestroyActor = false;
 
+	if (InLODActor->GetLevel()->GetWorldSettings()->IsPartitionedWorld())
+	{
+		UE_LOG(LogEngine, Warning, TEXT("Deleting legacy LODActor %s found in World Partition map. Resave %s to silence warning."), *InLODActor->GetName(), *InLODActor->GetOutermost()->GetPathName());
+		bShouldDestroyActor = true;
+	}
 	if (InLODActor->GetLevel()->GetWorldSettings()->GetHierarchicalLODSetup().Num() == 0)
 	{
 		UE_LOG(LogEngine, Warning, TEXT("Deleting LODActor %s found in map with no HLOD setup or disabled HLOD system. Resave %s to silence warning."), *InLODActor->GetName(), *InLODActor->GetOutermost()->GetPathName());
diff --git a/Engine/Source/Runtime/Engine/Private/WorldSettings.cpp b/Engine/Source/Runtime/Engine/Private/WorldSettings.cpp
--- a/Engine/Source/Runtime/Engine/Private/WorldSettings.cpp
+++ b/Engine/Source/Runtime/Engine/Private/WorldSettings.cpp
@@ -651,6 +651,9 @@ void AWorldSettings::PostLoad()
 	{
 		// Force to re-apply WorldPartition restrictions on WorldSettings (in case they changed)	
 		ApplyWorldPartitionForcedSettings();
+
+		// Partitionned levels have no need for this legacy data
+		ResetHierarchicalLODSetup();
 	}
 #endif// WITH_EDITOR
 
