UE-350439 Fix GPU skin cache crash on updating dynamic mesh bounds when disabling SM5 preview platforms

diff --git a/Engine/Source/Runtime/Engine/Private/GPUSkinCache.cpp b/Engine/Source/Runtime/Engine/Private/GPUSkinCache.cpp
--- a/Engine/Source/Runtime/Engine/Private/GPUSkinCache.cpp
+++ b/Engine/Source/Runtime/Engine/Private/GPUSkinCache.cpp
@@ -1645,7 +1645,7 @@ void FGPUSkinCache::DispatchPassSetup(FTaskData& TaskData)
 
 	TBitArray<> SlotsToUpdateMask;
 
-	if (TaskData.DynamicMeshBoundsBlackboardStruct)
+	if (TaskData.DynamicMeshBoundsBlackboardStruct && DynamicMeshBoundsBuffer)
 	{
 		SlotsToUpdateMask.Init(false, DynamicMeshBoundsBuffer->GetNumSlotsAllocated());
 	}
@@ -1792,7 +1792,7 @@ void FGPUSkinCache::DispatchPassSetup(FTaskData& TaskData)
 	PendingProcessRTGeometryEntries.Reset();
 #endif
 
-	if (TaskData.DynamicMeshBoundsBlackboardStruct)
+	if (TaskData.DynamicMeshBoundsBlackboardStruct && DynamicMeshBoundsBuffer)
 	{
 		TaskData.DynamicMeshBoundsBlackboardStruct->Parameters = DynamicMeshBoundsBuffer->Update(*TaskData.RHICmdList, SlotsToUpdateMask);
 		TaskData.DynamicMeshBoundsBlackboardStruct->bInitialized = true;
