Platforms with different preferred shader code format for preview platforms (like D3D12) must recompile their shader maps when switching preview platforms.
This is achived by adding PreferredShaderCodeFormat name to material and global shader map Id.
Also bump GUID for `MATERIALSHADERMAP_DERIVEDDATA_VER` because `FMaterialShaderMapId::Serialize()` includes this new entry.
To reduce code redundancy, add `GetPreferredShaderCodeFormat()` to FDataDrivenShaderPlatformInfo, which is used to set up a new compile job and to build the shader map Ids.
Platforms that don't have a different preferred shader code format for preview platforms (like most RHIs), can continue to share as many cached shaders between preview and non-preview as before.

#test Tested in Lyra and ContentExamples (NiagaraAdvancedParticles) by switching several times between multiple preview platforms, including PCD3D_ES3_1, Vulkan SM5, Android, Generic iOS.
#rnx

diff --git a/Engine/Source/Runtime/Core/Private/UObject/UE5ReleaseStreamObjectVersion.cpp b/Engine/Source/Runtime/Core/Private/UObject/UE5ReleaseStreamObjectVersion.cpp
--- a/Engine/Source/Runtime/Core/Private/UObject/UE5ReleaseStreamObjectVersion.cpp
+++ b/Engine/Source/Runtime/Core/Private/UObject/UE5ReleaseStreamObjectVersion.cpp
@@ -9,7 +9,7 @@ TMap<FGuid, FGuid> FUE5ReleaseStreamObjectVersion::GetSystemGuids()
 	const FDevSystemGuids& DevGuids = FDevSystemGuids::Get();
 
 	SystemGuids.Add(DevGuids.GLOBALSHADERMAP_DERIVEDDATA_VER, FGuid("D0BF3452816D46908073DFDD4B855AE5"));
-	SystemGuids.Add(DevGuids.MATERIALSHADERMAP_DERIVEDDATA_VER, FGuid("79090A66F9D94B5BB285D49E5D39468E"));
+	SystemGuids.Add(DevGuids.MATERIALSHADERMAP_DERIVEDDATA_VER, FGuid("4DD44E7A9F54481DB11D06CDD52FF390"));
 	SystemGuids.Add(DevGuids.NANITE_DERIVEDDATA_VER, FGuid("AC88CFBDDA614A9C8488F64D2DAF699D"));
 	SystemGuids.Add(DevGuids.NIAGARASHADERMAP_DERIVEDDATA_VER, FGuid("7BCE96ABAA6E46D689306FE673DED4DE"));
 	SystemGuids.Add(DevGuids.Niagara_LatestScriptCompileVersion, FGuid("14BEE7BD97194A7BBB8374D34DC49915"));
diff --git a/Engine/Source/Runtime/Engine/Private/Materials/MaterialShader.cpp b/Engine/Source/Runtime/Engine/Private/Materials/MaterialShader.cpp
--- a/Engine/Source/Runtime/Engine/Private/Materials/MaterialShader.cpp
+++ b/Engine/Source/Runtime/Engine/Private/Materials/MaterialShader.cpp
@@ -1285,11 +1285,13 @@ void FMaterialShaderMapId::Serialize(FArchive& Ar, bool bLoadingCooked)
 		if (!bIsLegacyPackage)
 		{
 			Ar << bUsingNewHLSLGenerator;
+			Ar << PreferredShaderCodeFormat;
 			Ar << ExternalCodeReferencesHash;
 		}
 		else
 		{
 			bUsingNewHLSLGenerator = false;
+			PreferredShaderCodeFormat = NAME_None;
 			ExternalCodeReferencesHash = FSHAHash();
 		}
 
@@ -1394,6 +1396,12 @@ void FMaterialShaderMapId::GetMaterialHash(FSHAHash& OutHash, bool bWithStaticPa
 
 	HashState.Update((const uint8*)&bUsingNewHLSLGenerator, sizeof(bUsingNewHLSLGenerator));
 
+	if (PreferredShaderCodeFormat != NAME_None)
+	{
+		const FString PreferredShaderCodeFormatString = PreferredShaderCodeFormat.ToString();
+		HashState.Update((const uint8*)PreferredShaderCodeFormatString.GetCharArray().GetData(), PreferredShaderCodeFormatString.GetCharArray().NumBytes());
+	}
+
 	HashState.Update((const uint8*)&ExternalCodeReferencesHash, sizeof(ExternalCodeReferencesHash));
 
 	SubstrateCompilationConfig.UpdateHash(HashState);
@@ -1425,6 +1433,11 @@ bool FMaterialShaderMapId::Equals(const FMaterialShaderMapId& ReferenceSet, bool
 		return false;
 	}
 
+	if (PreferredShaderCodeFormat != ReferenceSet.PreferredShaderCodeFormat)
+	{
+		return false;
+	}
+
 	if (SubstrateCompilationConfig != ReferenceSet.SubstrateCompilationConfig)
 	{
 		return false;
@@ -1772,6 +1785,11 @@ void FMaterialShaderMapId::RecordAndEmit(FMaterialKeyGeneratorContext& Context)
 
 	Context.Record("bUsingNewHLSLGenerator", bUsingNewHLSLGenerator);
 
+	if (PreferredShaderCodeFormat != NAME_None)
+	{
+		Context.Record("PreferredShaderCodeFormat", PreferredShaderCodeFormat);
+	}
+
 	// Recording of some fields necessary for IncludeUObjectData is not yet implemented; we only support emitting them
 	check(!Context.HasAllFlags(EMaterialKeyInclude::UObjectData) || !Context.IsRecording());
 	if (Context.HasAllFlags(EMaterialKeyInclude::UObjectData) && Context.IsEmitting())
diff --git a/Engine/Source/Runtime/Engine/Private/Materials/MaterialShared.cpp b/Engine/Source/Runtime/Engine/Private/Materials/MaterialShared.cpp
--- a/Engine/Source/Runtime/Engine/Private/Materials/MaterialShared.cpp
+++ b/Engine/Source/Runtime/Engine/Private/Materials/MaterialShared.cpp
@@ -774,6 +774,7 @@ void FMaterial::BuildShaderMapIdOverride(const FBuildShaderMapIdArgs& Args) cons
 
 		OutId.Usage = GetShaderMapUsage();
 		OutId.bUsingNewHLSLGenerator = IsUsingNewHLSLGenerator();
+		OutId.PreferredShaderCodeFormat = FDataDrivenShaderPlatformInfo::GetPreferredShaderCodeFormat(Args.Platform);
 		OutId.BaseMaterialId = GetMaterialId();
 		OutId.QualityLevel = GetQualityLevel();
 		OutId.FeatureLevel = GetFeatureLevel();
diff --git a/Engine/Source/Runtime/Engine/Private/ShaderCompiler/ShaderCompiler.cpp b/Engine/Source/Runtime/Engine/Private/ShaderCompiler/ShaderCompiler.cpp
--- a/Engine/Source/Runtime/Engine/Private/ShaderCompiler/ShaderCompiler.cpp
+++ b/Engine/Source/Runtime/Engine/Private/ShaderCompiler/ShaderCompiler.cpp
@@ -3179,7 +3179,7 @@ void GlobalBeginCompileShader(
 	Input.Target = Target;
 	Input.ShaderPlatformName = FDataDrivenShaderPlatformInfo::GetName(ShaderPlatform);
 	Input.ShaderFormat = ShaderFormatName;
-	Input.PreferredShaderCodeFormat = FDataDrivenShaderPlatformInfo::GetIsPreviewPlatform(ShaderPlatform) ? GRHIGlobals.PreferredPreviewShaderCodeFormat : NAME_None;
+	Input.PreferredShaderCodeFormat = FDataDrivenShaderPlatformInfo::GetPreferredShaderCodeFormat(ShaderPlatform);
 	Input.SupportedHardwareMask = TargetPlatform ? TargetPlatform->GetSupportedHardwareMask() : 0;
 	Input.CompressionFormat = GetShaderCompressionFormat();
 	GetShaderCompressionOodleSettings(Input.OodleCompressor, Input.OodleLevel);
diff --git a/Engine/Source/Runtime/Engine/Public/MaterialShared.h b/Engine/Source/Runtime/Engine/Public/MaterialShared.h
--- a/Engine/Source/Runtime/Engine/Public/MaterialShared.h
+++ b/Engine/Source/Runtime/Engine/Public/MaterialShared.h
@@ -1270,6 +1270,12 @@ public:
 	/** Is the material using the new HLSL generator? */
 	bool bUsingNewHLSLGenerator;
 
+	/**
+	 * Preferred format for the shader byte code the shader map is going to be compiled for.
+	 * This is used for preview platforms and affects all shaders, so make it part of the shader map Id.
+	 */
+	FName PreferredShaderCodeFormat;
+
 	/** The Substrate configuration used when compiling this material, can be tweaked in the material editor for live visualization of simplification. */
 	FSubstrateCompilationConfig SubstrateCompilationConfig;
 
diff --git a/Engine/Source/Runtime/RHI/Private/DataDrivenShaderPlatformInfo.cpp b/Engine/Source/Runtime/RHI/Private/DataDrivenShaderPlatformInfo.cpp
--- a/Engine/Source/Runtime/RHI/Private/DataDrivenShaderPlatformInfo.cpp
+++ b/Engine/Source/Runtime/RHI/Private/DataDrivenShaderPlatformInfo.cpp
@@ -635,6 +635,11 @@ const bool FGenericDataDrivenShaderPlatformInfo::CanUseForMaterialValidation(con
 	return DataDrivenShaderPlatformInfoEditorOnlyInfos[Platform].bCanUsePreviewPlatformForMaterialValidation || !Infos[Platform].bIsPreviewPlatform;
 }
 
+FName FGenericDataDrivenShaderPlatformInfo::GetPreferredShaderCodeFormat(const FStaticShaderPlatform Platform)
+{
+	return FDataDrivenShaderPlatformInfo::GetIsPreviewPlatform(Platform) ? GRHIGlobals.PreferredPreviewShaderCodeFormat : NAME_None;
+}
+
 #endif  // WITH_EDITOR
 
 const EShaderPlatform FGenericDataDrivenShaderPlatformInfo::GetShaderPlatformFromName(const FName ShaderPlatformName)
diff --git a/Engine/Source/Runtime/RHI/Public/DataDrivenShaderPlatformInfo.h b/Engine/Source/Runtime/RHI/Public/DataDrivenShaderPlatformInfo.h
--- a/Engine/Source/Runtime/RHI/Public/DataDrivenShaderPlatformInfo.h
+++ b/Engine/Source/Runtime/RHI/Public/DataDrivenShaderPlatformInfo.h
@@ -886,6 +886,7 @@ PRAGMA_ENABLE_DEPRECATION_WARNINGS
 	RHI_API static FText GetFriendlyName(const FStaticShaderPlatform Platform, FName DeviceProfileName = NAME_None);
 	RHI_API static const EShaderPlatform GetPreviewShaderPlatformParent(const FStaticShaderPlatform Platform);
 	RHI_API static const bool CanUseForMaterialValidation(const FStaticShaderPlatform Platform);
+	RHI_API static FName GetPreferredShaderCodeFormat(const FStaticShaderPlatform Platform);
 
 	RHI_API static TMap<FString, TFunction<bool(const FStaticShaderPlatform Platform)>> PropertyToShaderPlatformFunctionMap;
 #endif
diff --git a/Engine/Source/Runtime/RenderCore/Private/GlobalShader.cpp b/Engine/Source/Runtime/RenderCore/Private/GlobalShader.cpp
--- a/Engine/Source/Runtime/RenderCore/Private/GlobalShader.cpp
+++ b/Engine/Source/Runtime/RenderCore/Private/GlobalShader.cpp
@@ -480,6 +480,13 @@ void FGlobalShaderMapId::AppendKeyString(FString& KeyString, const TArray<FShade
 		KeyString.AppendChar('_');
 	}
 
+	const FName PreferredShaderCodeFormat = FDataDrivenShaderPlatformInfo::GetPreferredShaderCodeFormat(ShaderPlatform);
+	if (PreferredShaderCodeFormat != NAME_None)
+	{
+		KeyString.Append(PreferredShaderCodeFormat.ToString());
+		KeyString.AppendChar('_');
+	}
+
 	AppendKeyStringShaderDependencies(Dependencies, ShaderPipelineTypeDependencies, TConstArrayView<FVertexFactoryTypeDependency>(), LayoutParams, KeyString);
 
 	// Extra logic, just for Global shaders
