Reduced UE_LOG calls that can cause a lot of log entries, disproportionate to the relatively minor problem it is surfacing.

#rnx

diff --git a/Engine/Source/Runtime/Engine/Private/InstanceData/InstanceDataManager.cpp b/Engine/Source/Runtime/Engine/Private/InstanceData/InstanceDataManager.cpp
--- a/Engine/Source/Runtime/Engine/Private/InstanceData/InstanceDataManager.cpp
+++ b/Engine/Source/Runtime/Engine/Private/InstanceData/InstanceDataManager.cpp
@@ -651,7 +651,14 @@ bool FInstanceDataManager::FlushChanges(FInstanceDataManagerSourceDataDesc&& Com
 	// This should not happen in a cooked client, ideally.
 	if (GetState() == ETrackingState::Optimized)
 	{
-		UE_LOG(LogTemp, Log, TEXT("Discarded PrecomputedOptimizationData"));
+#if !UE_BUILD_SHIPPING
+		static bool bLoggedOnce = false;
+		if (!bLoggedOnce)
+		{
+			UE_LOG(LogTemp, Log, TEXT("Discarded PrecomputedOptimizationData"));
+			bLoggedOnce = true;
+		}
+#endif
 		PrecomputedOptimizationData.Reset();
 	}
 
diff --git a/Engine/Source/Runtime/Engine/Private/InstancedStaticMesh/ISMInstanceDataManager.cpp b/Engine/Source/Runtime/Engine/Private/InstancedStaticMesh/ISMInstanceDataManager.cpp
--- a/Engine/Source/Runtime/Engine/Private/InstancedStaticMesh/ISMInstanceDataManager.cpp
+++ b/Engine/Source/Runtime/Engine/Private/InstancedStaticMesh/ISMInstanceDataManager.cpp
@@ -597,7 +597,14 @@ bool FPrimitiveInstanceDataManager::FlushChanges(FInstanceUpdateComponentDesc &&
 	// If we got here & the state is "optimized" then we know the precomputed data is now invalid and we ditch it.
 	if (GetState() == ETrackingState::Optimized)
 	{
-		UE_LOG(LogInstanceProxy, Log, TEXT("Discarded PrecomputedOptimizationData"));
+#if !UE_BUILD_SHIPPING
+		static bool bLoggedOnce = false;
+		if (!bLoggedOnce)
+		{
+			UE_LOG(LogInstanceProxy, Log, TEXT("Discarded PrecomputedOptimizationData"));
+			bLoggedOnce = true;
+		}
+#endif
 		PrecomputedOptimizationData.Reset();
 	}
 
