[Backout] - CL44258326
#rnx
Original CL Desc
-----------------------------------------------------------------
PR #13545: Fixing UE-209483 a Stream Cache Handle Crash

diff --git a/Engine/Source/Runtime/Engine/Private/ContentStreaming.cpp b/Engine/Source/Runtime/Engine/Private/ContentStreaming.cpp
--- a/Engine/Source/Runtime/Engine/Private/ContentStreaming.cpp
+++ b/Engine/Source/Runtime/Engine/Private/ContentStreaming.cpp
@@ -1448,11 +1448,6 @@ FAudioChunkHandle::FAudioChunkHandle(const uint8* InData, uint32 NumBytes, const
 			CorrespondingWaveGuid = SoundWaveData->GetGUID();
 		}
 	}
-
-	if (IsValid())
-	{
-		IStreamingManager::Get().GetAudioStreamingManager().AddReferenceToChunk(*this);
-	}
 }
 
 FAudioChunkHandle::FAudioChunkHandle(const FAudioChunkHandle& Other)
