Fix compilation failure with Substrate.

diff --git a/Engine/Source/Runtime/Engine/Private/Materials/MaterialShared.cpp b/Engine/Source/Runtime/Engine/Private/Materials/MaterialShared.cpp
--- a/Engine/Source/Runtime/Engine/Private/Materials/MaterialShared.cpp
+++ b/Engine/Source/Runtime/Engine/Private/Materials/MaterialShared.cpp
@@ -5784,7 +5784,7 @@ FMaterialShaderParametersBase::FMaterialShaderParametersBase(const FMaterial* In
 		bHasMetallicConnected = InMaterial->HasMaterialPropertyConnected(MP_Metallic);
 		bHasEmissiveColorConnected = InMaterial->HasMaterialPropertyConnected(MP_EmissiveColor);
 		bHasAmbientOcclusionConnected = InMaterial->HasMaterialPropertyConnected(MP_AmbientOcclusion);
-		bHasAnisotropyConnected = InMaterial->HasMaterialPropertyConnected(MP_Anisotropy) || EnumHasAnyFlags(InMaterial->MaterialGetSubstrateMaterialBsdfFeatures_GameThread(), ESubstrateBsdfFeature::Anisotropy);
+		bHasAnisotropyConnected = InMaterial->HasMaterialPropertyConnected(MP_Anisotropy);
 	}
 	else
 	{
