PR #13887: Free pending-delete Virtual Texture immediately in Loadmap to avoid the fatal assert in FVirtualTextureAllocator::AcquireBlock
#rnx

diff --git a/Engine/Source/Runtime/Engine/Private/UnrealEngine.cpp b/Engine/Source/Runtime/Engine/Private/UnrealEngine.cpp
--- a/Engine/Source/Runtime/Engine/Private/UnrealEngine.cpp
+++ b/Engine/Source/Runtime/Engine/Private/UnrealEngine.cpp
@@ -16355,6 +16355,9 @@ void UEngine::TrimMemory()
 		[](FRHICommandList& RHICmdList)
 		{
 			GRHICommandList.GetImmediateCommandList().ImmediateFlush(EImmediateFlushType::FlushRHIThread);
+#if !WITH_EDITOR
+			GetRendererModule().ReleaseVirtualTexturePendingResources();
+#endif			
 			RHIFlushResources();
 			GRHICommandList.GetImmediateCommandList().ImmediateFlush(EImmediateFlushType::FlushRHIThreadFlushResources);
 		});
