Remove mutation of bCachingUniformExpressions off the game thread and only use it on the render thread + to check against CancelCacheUniformExpressions, i.e. it no longer feeds back state to the Instance. bResourceCreated is now the exclusive marker for IsReadyForFinishDestroy (mirroring ReleasedByRT on UMaterial).

diff --git a/Engine/Source/Runtime/Engine/Private/Materials/MaterialInstance.cpp b/Engine/Source/Runtime/Engine/Private/Materials/MaterialInstance.cpp
--- a/Engine/Source/Runtime/Engine/Private/Materials/MaterialInstance.cpp
+++ b/Engine/Source/Runtime/Engine/Private/Materials/MaterialInstance.cpp
@@ -118,7 +118,6 @@ void CacheMaterialInstanceUniformExpressions(const UMaterialInstance* MaterialIn
 {
 	if (MaterialInstance->Resource)
 	{
-		MaterialInstance->StartCacheUniformExpressions();
 		MaterialInstance->Resource->CacheUniformExpressions_GameThread(bRecreateUniformBuffer);
 	}
 }
@@ -592,8 +591,6 @@ void GameThread_UpdateMIParameter(const UMaterialInstance* Instance, const Param
 {
 	if (FApp::CanEverRender())
 	{
-		Instance->StartCacheUniformExpressions();
-
 		const UMaterial* Material = Instance->GetMaterial_Concurrent();
 		if (Material != nullptr)
 		{
@@ -3815,15 +3812,15 @@ void UMaterialInstance::BeginDestroy()
 			// Clear all references before assigning the atomic state below.
 			ResourcesToDestroy.Empty();
 
-			// Clear flag set when Resource was created
-			bResourceCreated = false;
-
 			// And remove from deferred uniform expression cache queue if it's in that
 			if (bCachingUniformExpressions)
 			{
 				Resource->CancelCacheUniformExpressions();
 				bCachingUniformExpressions = false;
 			}
+
+			// Clear flag set when Resource was created
+			bResourceCreated = false;
 		});
 	}
 }
@@ -3832,7 +3829,7 @@ bool UMaterialInstance::IsReadyForFinishDestroy()
 {
 	bool bIsReady = Super::IsReadyForFinishDestroy();
 
-	return bIsReady && !bResourceCreated && !bCachingUniformExpressions;
+	return bIsReady && !bResourceCreated;
 }
 
 void UMaterialInstance::FinishDestroy()
