World Partition - HLOD - Fix issue where instanced HLODs components always had their own HLOD batching policy set to "Instancing"
* The result was that if they were used to generated another level of HLOD (HLOD1 for example), they would be forced to be instanced again

#Synced-CL 48236868

diff --git a/Engine/Source/Runtime/Engine/Private/WorldPartition/HLOD/HLODInstancedStaticMeshComponent.cpp b/Engine/Source/Runtime/Engine/Private/WorldPartition/HLOD/HLODInstancedStaticMeshComponent.cpp
--- a/Engine/Source/Runtime/Engine/Private/WorldPartition/HLOD/HLODInstancedStaticMeshComponent.cpp
+++ b/Engine/Source/Runtime/Engine/Private/WorldPartition/HLOD/HLODInstancedStaticMeshComponent.cpp
@@ -42,7 +42,6 @@ void FHLODISMComponentDescriptor::InitFrom(const UStaticMeshComponent* Component
 	// Improve instance batching by ignoring some of the fields that aren't relevant for HLOD
 	Mobility = EComponentMobility::Static;
 	bHasCustomNavigableGeometry = EHasCustomNavigableGeometry::No;
-	HLODBatchingPolicy = EHLODBatchingPolicy::Instancing;
 	bSelectable = true;
 	bHasPerInstanceHitProxies = false;
 	bConsiderForActorPlacementWhenHidden = false;
@@ -51,7 +50,10 @@ void FHLODISMComponentDescriptor::InitFrom(const UStaticMeshComponent* Component
 	bOverrideNavigationExport = false;
 	bForceNavigationObstacle = false;
 	bCanEverAffectNavigation = false;
-	bFillCollisionUnderneathForNavmesh = false;	
+	bFillCollisionUnderneathForNavmesh = false;
+
+	// HLOD batching policy should be ignored for HLOD actors
+	HLODBatchingPolicy = EHLODBatchingPolicy::None;
 
 	// Force the correct ISM component class as it can be changed by the parent InitFrom() implementation
 	ComponentClass = UHLODBuilder::GetInstancedStaticMeshComponentClass();
