[NaniteAssemblies] Fix rare asserts with static mesh and skeletal mesh Nanite Assembly async compilation caused by trying to FinishCompilation on assets that still had dependencies that needed to be finished first.

#[fyi] Zulqarnain.Abbasi
#Synced-CL 48677745

diff --git a/Engine/Source/Runtime/Engine/Private/SkinnedAssetCompiler.cpp b/Engine/Source/Runtime/Engine/Private/SkinnedAssetCompiler.cpp
--- a/Engine/Source/Runtime/Engine/Private/SkinnedAssetCompiler.cpp
+++ b/Engine/Source/Runtime/Engine/Private/SkinnedAssetCompiler.cpp
@@ -416,18 +416,10 @@ void FSkinnedAssetCompilingManager::FinishCompilation(TArrayView<USkinnedAsset*
 			{
 				PendingSkinnedAssets.AddUnique(SkinnedAsset);
 			}
-
-			if (SkinnedAssetsWithPendingDependencies.Contains(SkinnedAsset))
+			else if (SkinnedAssetsWithPendingDependencies.Contains(SkinnedAsset))
 			{
-				// Add it to the next wave of meshes to finish, add its dependencies to the pending list
+				// Add it to the next wave of meshes to finish
 				NextPendingSkinnedAssets.AddUnique(SkinnedAsset);
-				for (USkinnedAsset* Dependency : SkinnedAsset->GetSkinnedAssetDependencies())
-				{
-					if (RegisteredSkinnedAsset.Contains(Dependency))
-					{
-						PendingSkinnedAssets.AddUnique(Dependency);
-					}
-				}
 			}
 
 			if (Options.bIncludeDependentAssets)
@@ -438,15 +430,32 @@ void FSkinnedAssetCompilingManager::FinishCompilation(TArrayView<USkinnedAsset*
 				{
 					for (TWeakObjectPtr<USkinnedAsset>& ReverseDependency : *ReverseLookup)
 					{
-						if (ReverseDependency.IsValid())
+						if (USkinnedAsset* Dependent = ReverseDependency.Get())
 						{
-							NextPendingSkinnedAssets.AddUnique(ReverseDependency.Get());
+							NextPendingSkinnedAssets.AddUnique(Dependent);
 						}
 					}
 				}
 			}
 		}
 
+		// Recursively filter in all dependencies
+		for (int32 NextPendingIndex = 0; NextPendingIndex < NextPendingSkinnedAssets.Num(); ++NextPendingIndex)
+		{
+			USkinnedAsset* SkinnedAsset = NextPendingSkinnedAssets[NextPendingIndex];
+			for (USkinnedAsset* Dependency : SkinnedAsset->GetSkinnedAssetDependencies())
+			{
+				if (RegisteredSkinnedAsset.Contains(Dependency))
+				{
+					PendingSkinnedAssets.AddUnique(Dependency);
+				}
+				else if (SkinnedAssetsWithPendingDependencies.Contains(Dependency))
+				{
+					NextPendingSkinnedAssets.AddUnique(Dependency);
+				}
+			}
+		}
+
 		if (PendingSkinnedAssets.Num() == 0 && NextPendingSkinnedAssets.Num() > 0)
 		{
 			PendingSkinnedAssets = MoveTemp(NextPendingSkinnedAssets);
diff --git a/Engine/Source/Runtime/Engine/Private/StaticMeshCompiler.cpp b/Engine/Source/Runtime/Engine/Private/StaticMeshCompiler.cpp
--- a/Engine/Source/Runtime/Engine/Private/StaticMeshCompiler.cpp
+++ b/Engine/Source/Runtime/Engine/Private/StaticMeshCompiler.cpp
@@ -571,18 +571,10 @@ void FStaticMeshCompilingManager::FinishCompilation(TArrayView<UStaticMesh* cons
 			{
 				PendingStaticMeshes.AddUnique(StaticMesh);
 			}
-
-			if (StaticMeshesWithPendingDependencies.Contains(StaticMesh))
+			else if (StaticMeshesWithPendingDependencies.Contains(StaticMesh))
 			{
-				// Add it to the next wave of meshes to finish, add its dependencies to the pending list
+				// Add it to the next wave of meshes to finish
 				NextPendingStaticMeshes.AddUnique(StaticMesh);
-				for (UStaticMesh* DependentMesh : StaticMesh->GetCachedNaniteAssemblyReferences())
-				{
-					if (DependentMesh && RegisteredStaticMesh.Contains(DependentMesh))
-					{
-						PendingStaticMeshes.AddUnique(DependentMesh);
-					}
-				}
 			}
 
 			if (Options.bIncludeDependentMeshes)
@@ -600,12 +592,29 @@ void FStaticMeshCompilingManager::FinishCompilation(TArrayView<UStaticMesh* cons
 					}
 				}
 			}
+		}
 
-			if (PendingStaticMeshes.Num() == 0 && NextPendingStaticMeshes.Num() > 0)
+		// Recursively filter in all dependencies
+		for (int32 NextPendingIndex = 0; NextPendingIndex < NextPendingStaticMeshes.Num(); ++NextPendingIndex)
+		{
+			UStaticMesh* StaticMesh = NextPendingStaticMeshes[NextPendingIndex];
+			for (UStaticMesh* Dependency : StaticMesh->GetCachedNaniteAssemblyReferences())
 			{
-				PendingStaticMeshes = MoveTemp(NextPendingStaticMeshes);
+				if (RegisteredStaticMesh.Contains(Dependency))
+				{
+					PendingStaticMeshes.AddUnique(Dependency);
+				}
+				else if (StaticMeshesWithPendingDependencies.Contains(Dependency))
+				{
+					NextPendingStaticMeshes.AddUnique(Dependency);
+				}
 			}
 		}
+
+		if (PendingStaticMeshes.Num() == 0 && NextPendingStaticMeshes.Num() > 0)
+		{
+			PendingStaticMeshes = MoveTemp(NextPendingStaticMeshes);
+		}
 	};
 
 	FilterIntoPendingLists(InStaticMeshes);
