Skeletal Mesh: Correctly update morphs during CPU skinning

#rnx

diff --git a/Engine/Source/Runtime/Engine/Private/SkeletalRenderCPUSkin.cpp b/Engine/Source/Runtime/Engine/Private/SkeletalRenderCPUSkin.cpp
--- a/Engine/Source/Runtime/Engine/Private/SkeletalRenderCPUSkin.cpp
+++ b/Engine/Source/Runtime/Engine/Private/SkeletalRenderCPUSkin.cpp
@@ -657,7 +657,7 @@ FORCEINLINE void UpdateMorphedVertex( VertexType& MorphedVertex, const VertexTyp
 		
 		if (!Info.VertexDeltaIterator.AtEnd() && Info.VertexDeltaIterator->SourceIdx == CurBaseVertIdx)
 		{
-			ApplyMorphBlend(MorphedVertex, *Info.VertexDeltaIterator, MorphWeights[MorphIdx]);
+			ApplyMorphBlend(MorphedVertex, *Info.VertexDeltaIterator, MorphWeights[Info.WeightIndex]);
 			++Info.VertexDeltaIterator;
 		}
 	}
