Fix for invalid cached nav bounds infecting other cached bounds.

diff --git a/Engine/Source/Runtime/Engine/Private/InstancedStaticMesh.cpp b/Engine/Source/Runtime/Engine/Private/InstancedStaticMesh.cpp
--- a/Engine/Source/Runtime/Engine/Private/InstancedStaticMesh.cpp
+++ b/Engine/Source/Runtime/Engine/Private/InstancedStaticMesh.cpp
@@ -3083,6 +3083,7 @@ FBoxSphereBounds UInstancedStaticMeshComponent::CalcBounds(const FTransform& InB
 	return NewBounds;
 }
 
+// This function CalcAndCacheNavigationBounds is set to be deprecated in v5.8, but because its original disuse is part of a hotfix, the header file isn't available to be modified yet.
 void UInstancedStaticMeshComponent::CalcAndCacheNavigationBounds()
 {
 	GetOrCalcBounds(GetComponentTransform(), EBoundsCacheType::Navigation);
@@ -3193,7 +3194,7 @@ void UInstancedStaticMeshComponent::UpdateBounds()
 
 	if (bCanEverAffectNavigation && (IsRegistered() || IsPreRegistering()))
 	{
-		CalcAndCacheNavigationBounds();
+		GetOrCalcBounds(GetComponentTransform(), EBoundsCacheType::Navigation);
 	}
 }
 
@@ -3646,6 +3647,8 @@ void UInstancedStaticMeshComponent::PreAllocateInstancesMemory(int32 AddedInstan
 
 int32 UInstancedStaticMeshComponent::AddInstanceInternal(int32 InstanceIndex, FInstancedStaticMeshInstanceData* InNewInstanceData, const FTransform& InstanceTransform, bool bWorldSpace)
 {
+	InvalidateCachedBounds();
+
 	FInstancedStaticMeshInstanceData* NewInstanceData = InNewInstanceData;
 
 	PrimitiveInstanceDataManager.Add(InstanceIndex, NewInstanceData != nullptr);
@@ -3677,9 +3680,6 @@ int32 UInstancedStaticMeshComponent::AddInstanceInternal(int32 InstanceIndex, FI
 	// Update navigation relevancy
 	bNavigationRelevant = IsNavigationRelevant();
 
-	// Indicates if the cached bounds were updated, it's expected to occur in PartialNavigationUpdate().
-	bNavigationCachedBoundsUpdated = false;
-	
 	// Perform navigation update if the component is relevant to navigation and registered.
 	// Otherwise this will be handled in OnRegister (e.g. Editor manipulations).
 	if (bNavigationRelevant && IsRegistered())
@@ -3706,15 +3706,6 @@ int32 UInstancedStaticMeshComponent::AddInstanceInternal(int32 InstanceIndex, FI
 		FInstancedStaticMeshDelegates::FInstanceIndexUpdateData IndexUpdate{ FInstancedStaticMeshDelegates::EInstanceIndexUpdateType::Added, InstanceIndex };
 		FInstancedStaticMeshDelegates::OnInstanceIndexUpdated.Broadcast(this, MakeArrayView(&IndexUpdate, 1));
 	}
-
-	const FBoundsCacheElement NavigationCacheElement = CachedBounds[(int32)EBoundsCacheType::Navigation];
-	InvalidateCachedBounds();
-
-	// Restore Navigation CachedBounds (the invalidation is not needed if it was updated).
-	if (bNavigationCachedBoundsUpdated)
-	{
-		CachedBounds[(int32)EBoundsCacheType::Navigation] = NavigationCacheElement;
-	}
 	
 	return InstanceIndex;
 }
@@ -3729,6 +3720,8 @@ TArray<int32> UInstancedStaticMeshComponent::AddInstancesInternal(TConstArrayVie
 	TRACE_CPUPROFILER_EVENT_SCOPE(UInstancedStaticMeshComponent::AddInstancesInternal);
 	CSV_SCOPED_TIMING_STAT_EXCLUSIVE(UInstancedStaticMeshComponent_AddInstancesInternal);
 
+	InvalidateCachedBounds();
+
 	const int32 Count = InstanceTransforms.Num();
 
 	TArray<int32> NewInstanceIndices;
@@ -3784,7 +3777,6 @@ TArray<int32> UInstancedStaticMeshComponent::AddInstancesInternal(TConstArrayVie
 			// since it was skipped because IsNavigationRelevant() requires at least one instance.
 			if (GetInstanceCount() == 1)
 			{
-				CalcAndCacheNavigationBounds();
 				FNavigationSystem::RegisterComponent(*this);
 			}
 
@@ -3815,7 +3807,6 @@ TArray<int32> UInstancedStaticMeshComponent::AddInstancesInternal(TConstArrayVie
 		}
 	}
 
-	InvalidateCachedBounds();
 	return NewInstanceIndices;
 }
 
@@ -4017,6 +4008,8 @@ bool UInstancedStaticMeshComponent::SupportsRemoveSwap() const
 
 bool UInstancedStaticMeshComponent::RemoveInstanceInternal(int32 InstanceIndex, bool InstanceAlreadyRemoved, bool bForceRemoveAtSwap, const bool bUpdateNavigation)
 {
+	InvalidateCachedBounds();
+
 #if WITH_EDITOR
 	DeletionState = InstanceAlreadyRemoved ? EInstanceDeletionReason::EntryAlreadyRemoved : EInstanceDeletionReason::EntryRemoval;
 #endif
@@ -4175,7 +4168,6 @@ bool UInstancedStaticMeshComponent::RemoveInstanceInternal(int32 InstanceIndex,
 	DeletionState = EInstanceDeletionReason::NotDeleting;
 #endif
 
-	InvalidateCachedBounds();
 	return true;
 }
 
@@ -5478,10 +5470,6 @@ void UInstancedStaticMeshComponent::PartialNavigationUpdates(const TConstArrayVi
 			CachedBounds[static_cast<int32>(EBoundsCacheType::Navigation)].Invalidate();
 		}
 
-		// Update cached navigation bounds
-		CalcAndCacheNavigationBounds();
-		bNavigationCachedBoundsUpdated = true;
-
 		// Dirty areas around the modified instances
 		TArray<FBox> DirtyAreas;
 		DirtyAreas.Reserve(InstanceTransforms.Num());
@@ -5647,9 +5635,7 @@ void UInstancedStaticMeshComponent::GetNavigationData(FNavigationRelevantData& D
 
 FBox UInstancedStaticMeshComponent::GetNavigationBounds() const
 {
-	// Navigation bounds must be calculated before being requested
-	constexpr bool bCheckValid = true;
-	return CachedBounds[(int)EBoundsCacheType::Navigation].GetValue(bCheckValid).GetBox();
+	return GetOrCalcBounds(GetComponentTransform(), EBoundsCacheType::Navigation).GetBox();
 }
 
 bool UInstancedStaticMeshComponent::IsNavigationRelevant() const
