World Partition - HLOD: Preserve original batching policy when creating HLOD ISM components
* Regression introduced with 45076387
* Revert back to the behavior from 5.6 to avoid breaking existing HLOD setup

diff --git a/Engine/Source/Runtime/Engine/Private/WorldPartition/HLOD/HLODInstancedStaticMeshComponent.cpp b/Engine/Source/Runtime/Engine/Private/WorldPartition/HLOD/HLODInstancedStaticMeshComponent.cpp
--- a/Engine/Source/Runtime/Engine/Private/WorldPartition/HLOD/HLODInstancedStaticMeshComponent.cpp
+++ b/Engine/Source/Runtime/Engine/Private/WorldPartition/HLOD/HLODInstancedStaticMeshComponent.cpp
@@ -52,9 +52,6 @@ void FHLODISMComponentDescriptor::InitFrom(const UStaticMeshComponent* Component
 	bCanEverAffectNavigation = false;
 	bFillCollisionUnderneathForNavmesh = false;
 
-	// HLOD batching policy should be ignored for HLOD actors
-	HLODBatchingPolicy = EHLODBatchingPolicy::None;
-
 	// Force the correct ISM component class as it can be changed by the parent InitFrom() implementation
 	ComponentClass = UHLODBuilder::GetInstancedStaticMeshComponentClass();
 }
