Animation: Revert mem stack allocator usage to avoid memory corruption issues

Change to using a thread_local scratch area.
BlendPosesTogetherPerBone can reallocate curves & attributes depending on the incoming poses, so we cannot use the mem stack allocator for the intermediate PerBoneIndices at this scope due to the need to return the reallocated curves/attributes to the calling scope

diff --git a/Engine/Source/Runtime/Engine/Private/Animation/AnimationRuntime.cpp b/Engine/Source/Runtime/Engine/Private/Animation/AnimationRuntime.cpp
--- a/Engine/Source/Runtime/Engine/Private/Animation/AnimationRuntime.cpp
+++ b/Engine/Source/Runtime/Engine/Private/Animation/AnimationRuntime.cpp
@@ -645,6 +645,11 @@ void FAnimationRuntime::BlendPosesTogetherPerBone(
 	BlendPosesTogetherPerBone(SourcePoses, SourceCurves, {}, InterpolationIndexProvider, BlendSampleDataCache, AnimationPoseData);	
 }
 
+namespace PerBoneScratch
+{
+	static thread_local TArray<int32> PerBoneIndices;
+	static thread_local TArray<float> SourceWeights;
+}
 
 void FAnimationRuntime::BlendPosesTogetherPerBone(TArrayView<const FCompactPose> SourcePoses, TArrayView<const FBlendedCurve> SourceCurves, TArrayView<const UE::Anim::FStackAttributeContainer> SourceAttributes, const IInterpolationIndexProvider* InterpolationIndexProvider, TArrayView<const FBlendSampleData> BlendSampleDataCache, /*out*/ FAnimationPoseData& OutAnimationPoseData)
 {
@@ -654,9 +659,8 @@ void FAnimationRuntime::BlendPosesTogetherPerBone(TArrayView<const FCompactPose>
 	FBlendedCurve& OutCurve = OutAnimationPoseData.GetCurve();
 	UE::Anim::FStackAttributeContainer& OutAttributes = OutAnimationPoseData.GetAttributes();
 
-	FMemMark Mark(FMemStack::Get());
-	TArray<int32, TMemStackAllocator<>> PerBoneIndices;
-	PerBoneIndices.SetNumUninitialized(OutPose.GetNumBones());
+	TArray<int32>& PerBoneIndices = PerBoneScratch::PerBoneIndices;
+	PerBoneIndices.SetNumUninitialized(OutPose.GetNumBones(), EAllowShrinking::No);
 
 	const FBoneContainer& RequiredBones = OutPose.GetBoneContainer();
 	TSharedPtr<IInterpolationIndexProvider::FPerBoneInterpolationData> Data = InterpolationIndexProvider->GetPerBoneInterpolationData(RequiredBones.GetSkeletonAsset());
@@ -676,20 +680,20 @@ void FAnimationRuntime::BlendPosesTogetherPerBone(TArrayView<const FCompactPose>
 
 	if (SourceCurves.Num() > 0)
 	{
-		TArray<float, TMemStackAllocator<>> SourceWeights;
-		SourceWeights.SetNumUninitialized(BlendSampleDataCache.Num());
+		TArray<float>& SourceWeights = PerBoneScratch::SourceWeights;
+		SourceWeights.SetNumUninitialized(BlendSampleDataCache.Num(), EAllowShrinking::No);
 		for (int32 CacheIndex = 0; CacheIndex < BlendSampleDataCache.Num(); ++CacheIndex)
 		{
 			SourceWeights[CacheIndex] = BlendSampleDataCache[CacheIndex].TotalWeight;
 		}
 
-		BlendCurves(SourceCurves, SourceWeights, OutCurve);		
+		BlendCurves(SourceCurves, SourceWeights, OutCurve);
 	}
-
+	
 	if (SourceAttributes.Num() > 0)
 	{
 		UE::Anim::Attributes::BlendAttributesPerBone(SourceAttributes, PerBoneIndices, BlendSampleDataCache, OutAttributes);
-	}	
+	}
 }
 
 void FAnimationRuntime::BlendPosesTogetherPerBone(
@@ -715,9 +719,8 @@ void FAnimationRuntime::BlendPosesTogetherPerBone(TArrayView<const FCompactPose>
 	FBlendedCurve& OutCurve = OutAnimationPoseData.GetCurve();
 	UE::Anim::FStackAttributeContainer& OutAttributes = OutAnimationPoseData.GetAttributes();
 
-	FMemMark Mark(FMemStack::Get());
-	TArray<int32, TMemStackAllocator<>> PerBoneIndices;
-	PerBoneIndices.SetNumUninitialized(OutPose.GetNumBones());
+	TArray<int32>& PerBoneIndices = PerBoneScratch::PerBoneIndices;
+	PerBoneIndices.SetNumUninitialized(OutPose.GetNumBones(), EAllowShrinking::No);
 
 	const FBoneContainer& RequiredBones = OutPose.GetBoneContainer();
 
@@ -739,8 +742,8 @@ void FAnimationRuntime::BlendPosesTogetherPerBone(TArrayView<const FCompactPose>
 
 	if (SourceCurves.Num() > 0)
 	{
-		TArray<float, TInlineAllocator<16>> SourceWeights;
-		SourceWeights.AddUninitialized(BlendSampleDataCacheIndices.Num());
+		TArray<float>& SourceWeights = PerBoneScratch::SourceWeights;
+		SourceWeights.SetNumUninitialized(BlendSampleDataCacheIndices.Num(), EAllowShrinking::No);
 		for (int32 CacheIndex = 0; CacheIndex < BlendSampleDataCacheIndices.Num(); ++CacheIndex)
 		{
 			SourceWeights[CacheIndex] = BlendSampleDataCache[BlendSampleDataCacheIndices[CacheIndex]].TotalWeight;
@@ -765,8 +768,8 @@ void FAnimationRuntime::BlendPosesTogetherPerBoneRemapped(TArrayView<const FComp
 	UE::Anim::FStackAttributeContainer& OutAttributes = OutAnimationPoseData.GetAttributes();
 
 	const FBoneContainer& RequiredBones = OutPose.GetBoneContainer();
-	TArray<int32> PerBoneIndices;
-	PerBoneIndices.AddUninitialized(OutPose.GetNumBones());
+	TArray<int32>& PerBoneIndices = PerBoneScratch::PerBoneIndices;
+	PerBoneIndices.SetNumUninitialized(OutPose.GetNumBones(), EAllowShrinking::No);
 	TSharedPtr<IInterpolationIndexProvider::FPerBoneInterpolationData> Data = InterpolationIndexProvider->GetPerBoneInterpolationData(OutPose.GetBoneContainer().GetSkeletonAsset());
 	for (FCompactPoseBoneIndex BoneIndex : OutPose.ForEachBoneIndex())
 	{
@@ -787,8 +790,8 @@ void FAnimationRuntime::BlendPosesTogetherPerBoneRemapped(TArrayView<const FComp
 
 	if (SourceCurves.Num() > 0)
 	{
-		TArray<float, TInlineAllocator<16>> SourceWeights;
-		SourceWeights.AddUninitialized(BlendSampleDataCacheIndices.Num());
+		TArray<float>& SourceWeights = PerBoneScratch::SourceWeights;
+		SourceWeights.SetNumUninitialized(BlendSampleDataCacheIndices.Num(), EAllowShrinking::No);
 		for (int32 CacheIndex = 0; CacheIndex < BlendSampleDataCacheIndices.Num(); ++CacheIndex)
 		{
 			SourceWeights[CacheIndex] = BlendSampleDataCache[BlendSampleDataCacheIndices[CacheIndex]].TotalWeight;
