Keep OwnerLastRenderTimePtr alive on SceneData after releasing proxy. The pointer can be referenced by the render thread. It references an Actor member which should be fenced already, so it is safe to keep assigned.

diff --git a/Engine/Source/Runtime/Engine/Private/Components/PrimitiveComponent.cpp b/Engine/Source/Runtime/Engine/Private/Components/PrimitiveComponent.cpp
--- a/Engine/Source/Runtime/Engine/Private/Components/PrimitiveComponent.cpp
+++ b/Engine/Source/Runtime/Engine/Private/Components/PrimitiveComponent.cpp
@@ -5380,7 +5380,7 @@ void UPrimitiveComponent::ReleaseSceneProxy()
 		const uint32 NumRefs = SceneData.OwnerLastRenderTimePtr->NumAlwaysVisibleComponents.fetch_sub(1, std::memory_order_relaxed);
 		check(NumRefs > 0);
 	}
-	SceneData.OwnerLastRenderTimePtr = nullptr;
+	// Leave SceneData.OwnerLastRenderTimePtr valid. It might still be referenced on the rendering thread.
 	SceneProxy = nullptr;
 	SceneData.SceneProxy = nullptr;
 	SceneData.bAlwaysVisible = false;
