Chaos Cloth - Fixed capsule orientation when colliding with the environment.

[RN] Chaos Cloth - Fixed capsule orientation when colliding with the environment.
#Synced-CL 49590414

diff --git a/Engine/Source/Runtime/Engine/Private/PhysicsEngine/EnvironmentalCollisions.cpp b/Engine/Source/Runtime/Engine/Private/PhysicsEngine/EnvironmentalCollisions.cpp
--- a/Engine/Source/Runtime/Engine/Private/PhysicsEngine/EnvironmentalCollisions.cpp
+++ b/Engine/Source/Runtime/Engine/Private/PhysicsEngine/EnvironmentalCollisions.cpp
@@ -105,7 +105,8 @@ void FEnvironmentalCollisions::AppendCollisionDataFromEnvironment(const USceneCo
 
 								const FReal Radius = ImplicitCapsule.GetRadiusf() * Scale.X;  // Assumes uniform scale
 								const FReal HalfHeight = ImplicitCapsule.GetHeightf() * (FReal)0.5;
-								const FVector HalfSegment = ComponentToClothTransform.GetUnitAxis(EAxis::X) * HalfHeight * Scale.X;  // Assumes uniform scale
+								const FVector TransformedAxis = ComponentToClothTransform.TransformVectorNoScale(FVec3(ImplicitCapsule.GetAxisf()));
+								const FVector HalfSegment = TransformedAxis * HalfHeight * Scale.X;  // Assumes uniform scale
 								const FVector TransformedCenter = ComponentToClothTransform.TransformPosition(FVec3(ImplicitCapsule.GetCenterf()));
 
 								CollisionData.Spheres.Emplace(Radius, TransformedCenter + HalfSegment);
