Cloth - Fixed a crash in FClothingSimulationInstance::GetBounds() when a SkeletalMeshComponent with simulating cloth in editor is transacted.

[RN] Cloth - Fixed a crash in FClothingSimulationInstance::GetBounds() when a SkeletalMeshComponent with simulating cloth in editor is transacted.
#tests Created a reliable reproducible and tested with the fix while making sure the simulation in editor functionality remained.

diff --git a/Engine/Source/Runtime/Engine/Classes/Components/SkeletalMeshComponent.h b/Engine/Source/Runtime/Engine/Classes/Components/SkeletalMeshComponent.h
--- a/Engine/Source/Runtime/Engine/Classes/Components/SkeletalMeshComponent.h
+++ b/Engine/Source/Runtime/Engine/Classes/Components/SkeletalMeshComponent.h
@@ -1612,7 +1612,7 @@ private:
 	FVector3f PrevRootRigidBodyScale;
 
 	/** Clothing simulation instance. */
-	UPROPERTY(Transient)
+	UPROPERTY(Transient, NonTransactional)
 	TArray<FClothingSimulationInstance> ClothingSimulationInstances;
 
 	/** Array of sources for cloth collision */
