[HWRT] Fix skeletal meshes with disabled sections flickering in ray tracing effects.
- Ray tracing geometry was getting unnecessarily recreated repeatedly.

#Synced-CL 49716836

diff --git a/Engine/Source/Runtime/Engine/Private/SkeletalRenderGPUSkin.cpp b/Engine/Source/Runtime/Engine/Private/SkeletalRenderGPUSkin.cpp
--- a/Engine/Source/Runtime/Engine/Private/SkeletalRenderGPUSkin.cpp
+++ b/Engine/Source/Runtime/Engine/Private/SkeletalRenderGPUSkin.cpp
@@ -1196,9 +1196,23 @@ void FSkeletalMeshObjectGPUSkin::UpdateRayTracingGeometry_Internal(
 		}
 		MeshObject->bHiddenMaterialVisibilityDirtyForRayTracing = false;
 
+		auto PatchSegmentVertexBuffer = [&LODModel](FRayTracingGeometrySegment& Segment)
+			{
+				if (Segment.bEnabled)
+				{
+					checkf(Segment.VertexBuffer != nullptr, TEXT("Must provide valid vertex buffer when segment is enabled."));
+				}
+				else if (Segment.VertexBuffer == nullptr) // TODO: Consider always doing this when !Segment.bEnabled
+				{
+					// TODO: RHIs should accept NULL VertexBuffer when Segment is disabled.
+					// For now use the static PositionVertexBuffer as placeholder.
+					// TODO: Consider also setting Segment.NumPrimitives = 0; in this case
+					Segment.VertexBuffer = LODModel.StaticVertexBuffers.PositionVertexBuffer.GetRHI();
+				}
+			};
+
 		if (bRequireRecreatingRayTracingGeometry)
 		{
-
 			FBufferRHIRef IndexBufferRHI = LODModel.MultiSizeIndexContainer.GetIndexBuffer()->IndexBufferRHI;
 			uint32 VertexBufferNumVertices = LODModel.StaticVertexBuffers.PositionVertexBuffer.GetNumVertices();
 			uint32 VertexBufferStride = LODModel.StaticVertexBuffers.PositionVertexBuffer.GetStride();
@@ -1242,17 +1256,7 @@ void FSkeletalMeshObjectGPUSkin::UpdateRayTracingGeometry_Internal(
 				// TODO: If we are at a dropped LOD, route material index through the LODMaterialMap in the LODInfo struct.
 				Segment.bEnabled = !MeshObject->IsMaterialHidden(LODIndex, Section.MaterialIndex) && Section.IsValid() && Section.bVisibleInRayTracing;
 
-				if (Segment.bEnabled)
-				{
-					checkf(Segment.VertexBuffer != nullptr, TEXT("Must provide valid vertex buffer when segment is enabled."));
-				}
-				else if (Segment.VertexBuffer == nullptr) // TODO: Consider always doing this when !Segment.bEnabled
-				{
-					// TODO: RHIs should accept NULL VertexBuffer when Segment is disabled.
-					// For now use the static PositionVertexBuffer as placeholder.
-					// TODO: Consider also setting Segment.NumPrimitives = 0; in this case
-					Segment.VertexBuffer = LODModel.StaticVertexBuffers.PositionVertexBuffer.GetRHI();
-				}
+				PatchSegmentVertexBuffer(Segment);
 
 				Initializer.Segments.Add(Segment);
 
@@ -1286,6 +1290,8 @@ void FSkeletalMeshObjectGPUSkin::UpdateRayTracingGeometry_Internal(
 				FRayTracingGeometrySegment& Segment = RayTracingGeometry.Initializer.Segments[SectionIndex];
 				Segment.VertexBuffer = VertexBuffers[SectionIndex];
 				Segment.VertexBufferOffset = 0;
+
+				PatchSegmentVertexBuffer(Segment);
 			}
 		}
 
