[GPU Skin Cache] Fixed morph not resettting when its weight is set to 0 directly, while skin cache is on

#Synced-CL 48737829

diff --git a/Engine/Source/Runtime/Engine/Private/GPUSkinCache.cpp b/Engine/Source/Runtime/Engine/Private/GPUSkinCache.cpp
--- a/Engine/Source/Runtime/Engine/Private/GPUSkinCache.cpp
+++ b/Engine/Source/Runtime/Engine/Private/GPUSkinCache.cpp
@@ -2207,7 +2207,9 @@ void FGPUSkinCache::ProcessEntry(FRHICommandList& RHICmdList, const FProcessEntr
 		}
 
 		FGPUSkinCacheEntry::FSectionDispatchData& SectionDispatchData = InOutEntry->DispatchData[SectionIndex];
-
+		
+		SectionDispatchData.DeformationType = EGPUSkinDeformationType::Default;
+		
 		if (Inputs.MorphVertexBuffer && Inputs.MorphVertexBuffer->SectionIds.Contains(SectionIndex))
 		{
 			InOutEntry->MorphBuffer = Inputs.MorphVertexBuffer->GetSRV();
