[HWRT] Prevent RayTracingGeometryGroup leak when FStaticMeshRenderData is initialized twice.
- FStaticMeshRenderData::InitResources(...) can end up being called twice in some cases (eg: during Undo-Redo flows). FRenderResource::InitResource(...) guards against double initialization, FRayTracingGeometry::InitResource(...) also has logic to prevent double registration, so we were missing similar guard for the RayTracingGeometryGroup.

diff --git a/Engine/Source/Runtime/Engine/Private/StaticMesh.cpp b/Engine/Source/Runtime/Engine/Private/StaticMesh.cpp
--- a/Engine/Source/Runtime/Engine/Private/StaticMesh.cpp
+++ b/Engine/Source/Runtime/Engine/Private/StaticMesh.cpp
@@ -2755,6 +2755,10 @@ void FStaticMeshRenderData::Serialize(FArchive& Ar, UStaticMesh* Owner, bool bCo
 
 void FStaticMeshRenderData::InitResources(ERHIFeatureLevel::Type InFeatureLevel, UStaticMesh* Owner)
 {
+	// TODO: This check is disabled because in some cases we end up calling InitResources(...) on already initialized RenderData
+	// due to an issue with PostEditUndo + async builds (UE-365996)
+	//check(!bIsInitialized);
+
 #if WITH_EDITOR
 	// Init the section info only for uncooked editor.
 	// Cooked packages don't need this and don't want any LOD screen size changes that it applies.
@@ -2802,7 +2806,10 @@ void FStaticMeshRenderData::InitResources(ERHIFeatureLevel::Type InFeatureLevel,
 
 				check(!RayTracingProxy->bUsingRenderingLODs || RayTracingLODs.Num() == LODResources.Num());
 
-				RayTracingGeometryGroupHandle = GRayTracingGeometryManager->RegisterRayTracingGeometryGroup(RayTracingLODs.Num(), RayTracingProxy->bUsingRenderingLODs ? CurrentFirstLODIdx : 0);
+				if (RayTracingGeometryGroupHandle == INDEX_NONE) // workaround for FStaticMeshRenderData::InitResources(...) being called on already initialized RenderData due to bug with PostEditUndo + async build (UE-365996)
+				{
+					RayTracingGeometryGroupHandle = GRayTracingGeometryManager->RegisterRayTracingGeometryGroup(RayTracingLODs.Num(), RayTracingProxy->bUsingRenderingLODs ? CurrentFirstLODIdx : 0);
+				}
 
 				for (int32 LODIndex = 0; LODIndex < RayTracingLODs.Num(); ++LODIndex)
 				{
