Cloth - Fixed a crash in USkeletalMeshComponent::AddClothingBounds when a skeletal mesh component containing in editor clothing  simulation is duplicated as template when spawned by the Sequencer.

[RN] Sequencer/Cloth - Fixed a crash in USkeletalMeshComponent::AddClothingBounds when a skeletal mesh component containing in editor clothing simulation is duplicated as template when spawned by the Sequencer.

diff --git a/Engine/Source/Runtime/Engine/Classes/Components/SkeletalMeshComponent.h b/Engine/Source/Runtime/Engine/Classes/Components/SkeletalMeshComponent.h
--- a/Engine/Source/Runtime/Engine/Classes/Components/SkeletalMeshComponent.h
+++ b/Engine/Source/Runtime/Engine/Classes/Components/SkeletalMeshComponent.h
@@ -1612,7 +1612,7 @@ private:
 	FVector3f PrevRootRigidBodyScale;
 
 	/** Clothing simulation instance. */
-	UPROPERTY(Transient, NonTransactional)
+	UPROPERTY(Transient, DuplicateTransient, TextExportTransient, SkipSerialization, NonTransactional)
 	TArray<FClothingSimulationInstance> ClothingSimulationInstances;
 
 	/** Array of sources for cloth collision */
