[AsyncLoadingTests]
- Fix potential deadlock when running tests with -dpcvars=s.AllowMultithreadedLoading=1

diff --git a/Engine/Source/Runtime/Engine/Private/Tests/Loading/AsyncLoadingTests_RecursiveLoads.cpp b/Engine/Source/Runtime/Engine/Private/Tests/Loading/AsyncLoadingTests_RecursiveLoads.cpp
--- a/Engine/Source/Runtime/Engine/Private/Tests/Loading/AsyncLoadingTests_RecursiveLoads.cpp
+++ b/Engine/Source/Runtime/Engine/Private/Tests/Loading/AsyncLoadingTests_RecursiveLoads.cpp
@@ -51,6 +51,13 @@ bool FLoadingTests_RecursiveLoads_FromSerialize::RunTest(const FString& Paramete
 	UAsyncLoadingTests_Shared::OnSerialize.BindLambda(
 		[this, &bHasBeenPostloaded, &PostLoadEvent, &bIsMultithreadedLoading](FArchive& Ar, UAsyncLoadingTests_Shared* Object)
 		{
+			// Loading from serialize is not supported in the multithreaded workers as it can lead to deadlocks.
+			// Classes that end up doing it should be added to the .ini (search ParallelPreloadExclusions for examples) to prevent them from running wide.
+			if (IsInParallelLoadingThread())
+			{
+				return;
+			}
+
 			if (Ar.IsLoading())
 			{
 				bIsMultithreadedLoading = !IsInGameThread();
