Static Spatial Index: properly initialize the elements box when sorting is needed.
FBox constructor properly initialize the IsValid member, while FBox2D don't.

diff --git a/Engine/Source/Runtime/Engine/Public/WorldPartition/RuntimeHashSet/StaticSpatialIndex.h b/Engine/Source/Runtime/Engine/Public/WorldPartition/RuntimeHashSet/StaticSpatialIndex.h
--- a/Engine/Source/Runtime/Engine/Public/WorldPartition/RuntimeHashSet/StaticSpatialIndex.h
+++ b/Engine/Source/Runtime/Engine/Public/WorldPartition/RuntimeHashSet/StaticSpatialIndex.h
@@ -274,7 +274,7 @@ private:
 		// Sort elements to maximize cache coherency during queries
 		if (ElementsSorter::NeedSort)
 		{
-			FBox ElementsBox;
+			FBox ElementsBox(ForceInit);
 			Algo::ForEach(Elements, [&ElementsBox](const TPair<FBox, ValueType>& Element) { ElementsBox += Element.Key; });
 
 			ElementsSorter Sorter;
