Adding delegates to WorldPartition to be able to inject transient streaming objects during PIE.

diff --git a/Engine/Source/Runtime/Engine/Private/WorldPartition/WorldPartition.cpp b/Engine/Source/Runtime/Engine/Private/WorldPartition/WorldPartition.cpp
--- a/Engine/Source/Runtime/Engine/Private/WorldPartition/WorldPartition.cpp
+++ b/Engine/Source/Runtime/Engine/Private/WorldPartition/WorldPartition.cpp
@@ -331,6 +331,9 @@ void UWorldPartition::FWorldPartitionExternalDirtyActorsTracker::Tick(float InDe
 }
 
 UWorldPartition::FWorldPartitionChangedEvent UWorldPartition::WorldPartitionChangedEvent;
+UWorldPartition::FWorldPartitionEditorGameWorldEvent UWorldPartition::OnPrepareEditorGameWorldForPIE;
+UWorldPartition::FWorldPartitionEditorGameWorldEvent UWorldPartition::OnShutdownEditorGameWorldForPIE;
+
 #endif
 
 #if !NO_LOGGING
@@ -451,6 +454,11 @@ void UWorldPartition::PrepareEditorGameWorld()
 	check(!bIsPIE);
 	bIsPIE = GIsEditor && !IsRunningGame();
 
+	if (bIsPIE)
+	{
+		OnPrepareEditorGameWorldForPIE.Broadcast(this);
+	}
+
 	// In PIE, we always want to populate the map check dialog
 	FStreamingGenerationMapCheckErrorHandler MapCheckErrorHandler;
 	FStreamingGenerationLogErrorHandler LogErrorHandler;
@@ -487,6 +495,7 @@ void UWorldPartition::ShutdownEditorGameWorld()
 		RuntimeHash->ShutdownEditorGameWorld();
 		ExternalDataLayerManager->ShutdownEditorGameWorld();
 		RuntimeHash->ShutdownEditorGameWorld();
+		OnShutdownEditorGameWorldForPIE.Broadcast(this);
 		bIsPIE = false;
 	}
 }
diff --git a/Engine/Source/Runtime/Engine/Public/WorldPartition/WorldPartition.h b/Engine/Source/Runtime/Engine/Public/WorldPartition/WorldPartition.h
--- a/Engine/Source/Runtime/Engine/Public/WorldPartition/WorldPartition.h
+++ b/Engine/Source/Runtime/Engine/Public/WorldPartition/WorldPartition.h
@@ -172,6 +172,11 @@ public:
 	ENGINE_API void OnEnableLoadingInEditorChanged();
 
 	ENGINE_API bool IsStreamingEnabledInEditor() const;
+
+	DECLARE_MULTICAST_DELEGATE_OneParam(FWorldPartitionEditorGameWorldEvent, UWorldPartition*);
+	static ENGINE_API FWorldPartitionEditorGameWorldEvent OnPrepareEditorGameWorldForPIE;
+	static ENGINE_API FWorldPartitionEditorGameWorldEvent OnShutdownEditorGameWorldForPIE;
+
 private:
 	ENGINE_API void SavePerUserSettings();
 		
