ObjectTrace: Enable return value for TRACE_OBJECT to allow conditional code paths based on the first encounter

#rnx

diff --git a/Engine/Source/Runtime/Engine/Private/ObjectTrace.cpp b/Engine/Source/Runtime/Engine/Private/ObjectTrace.cpp
--- a/Engine/Source/Runtime/Engine/Private/ObjectTrace.cpp
+++ b/Engine/Source/Runtime/Engine/Private/ObjectTrace.cpp
@@ -477,18 +477,18 @@ void FObjectTrace::OutputInstance(const UObject* InOuterObject, uint64 InInstanc
 }
 
 
-void FObjectTrace::OutputObject(const UObject* InObject)
+bool FObjectTrace::OutputObject(const UObject* InObject)
 {
 	if (UE::ObjectTrace::IsTracingDisabledForObject(InObject))
  	{
- 		return;
+ 		return false;
  	}
 
 	uint64 InstanceId = GetObjectId(InObject);
 	
 	if (!AddObjectToTracedSet(UE::ObjectTrace::FObjectId{InstanceId}))
 	{
-		return;
+		return false;
 	}
 	
 	uint8 Flags = 0;
@@ -499,6 +499,7 @@ void FObjectTrace::OutputObject(const UObject* InObject)
  	
 	OutputObject(InObject->GetOuter());
 	OutputInstanceInternal(InObject, InstanceId, GetObjectId(InObject->GetOuter()), InObject->GetClass(), InObject->GetFName().ToString(), InObject->GetPathName(), Flags);
+	return true;
 }
 
 void FObjectTrace::OutputObjectTransform(const UObject* InObject, const FTransform& InTransform, ETeleportType InTeleportType)
diff --git a/Engine/Source/Runtime/Engine/Public/ObjectTrace.h b/Engine/Source/Runtime/Engine/Public/ObjectTrace.h
--- a/Engine/Source/Runtime/Engine/Public/ObjectTrace.h
+++ b/Engine/Source/Runtime/Engine/Public/ObjectTrace.h
@@ -67,7 +67,7 @@ struct FObjectTrace
 	};
 
 	/** Helper function to output an object */
-	ENGINE_API static void OutputObject(const UObject* InObject);
+	ENGINE_API static bool OutputObject(const UObject* InObject);
 	
 	/** Helper function to output an "object instance" which is not an actual UObject */
 	ENGINE_API static void OutputInstance(const UObject* InOuterObject, uint64 InInstanceId, uint64 InOuterId, UStruct* InType, const FString& InName, const FString& InPathName = FString());
