Deprecate CachedSoftObjectPtr in Skeleton because it is no longer used

diff --git a/Engine/Source/Runtime/Engine/Classes/Animation/Skeleton.h b/Engine/Source/Runtime/Engine/Classes/Animation/Skeleton.h
--- a/Engine/Source/Runtime/Engine/Classes/Animation/Skeleton.h
+++ b/Engine/Source/Runtime/Engine/Classes/Animation/Skeleton.h
@@ -643,6 +643,7 @@ private:
 
 public:
 	/** A cached soft object pointer of this skeleton. This is done for performance reasons when searching for compatible skeletons when using IsCompatible(Skeleton). */
+	UE_DEPRECATED(5.7, "CachedSoftObjectPtr is no longer used and has been deprecated.")
 	TSoftObjectPtr<USkeleton> CachedSoftObjectPtr;
 
 	/** IInterface_PreviewMeshProvider interface */
diff --git a/Engine/Source/Runtime/Engine/Private/Animation/Skeleton.cpp b/Engine/Source/Runtime/Engine/Private/Animation/Skeleton.cpp
--- a/Engine/Source/Runtime/Engine/Private/Animation/Skeleton.cpp
+++ b/Engine/Source/Runtime/Engine/Private/Animation/Skeleton.cpp
@@ -296,7 +296,9 @@ void USkeleton::RemoveCompatibleSkeleton(const TSoftObjectPtr<USkeleton>& Source
 USkeleton::USkeleton(const FObjectInitializer& ObjectInitializer)
 	: Super(ObjectInitializer)
 {
+	PRAGMA_DISABLE_DEPRECATION_WARNINGS
 	CachedSoftObjectPtr = TSoftObjectPtr<USkeleton>(this);
+	PRAGMA_ENABLE_DEPRECATION_WARNINGS
 
 #if WITH_EDITORONLY_DATA
 	PreviewForwardAxis = EAxis::Y;
