[ChildActorComponent] Actors that are spawned by a ChildActorComponent should return their parent actor GetLevelInstanceType()/GetLevelInstanceFlags() as those flags are not updated on the child actors and need to match.

diff --git a/Engine/Source/Runtime/Engine/Classes/GameFramework/Actor.h b/Engine/Source/Runtime/Engine/Classes/GameFramework/Actor.h
--- a/Engine/Source/Runtime/Engine/Classes/GameFramework/Actor.h
+++ b/Engine/Source/Runtime/Engine/Classes/GameFramework/Actor.h
@@ -398,38 +398,44 @@ public:
 	/** If true, the actor belongs to a level instance which is currently being edited */
 	bool IsInEditLevelInstance() const
 	{
-		return LevelInstanceType == ELevelInstanceType::LevelInstanceEdit;
+		return GetLevelInstanceType() == ELevelInstanceType::LevelInstanceEdit;
 	}
 
 	/** If true, the actor belongs to a level instance which is currently in property override edit */
 	bool IsInPropertyOverrideLevelInstance() const
 	{
-		return LevelInstanceType == ELevelInstanceType::LevelInstancePropertyOverride;
+		return GetLevelInstanceType() == ELevelInstanceType::LevelInstancePropertyOverride;
 	}
 
 	/** If true, the actor belongs to a level instance which is currently being edited or a parent level instance being edited. */
 	bool IsInEditLevelInstanceHierarchy() const
 	{
-		return EnumHasAnyFlags(LevelInstanceFlags, ELevelInstanceFlags::IsInEditHierarchy);
+		return EnumHasAnyFlags(GetLevelInstanceFlags(), ELevelInstanceFlags::IsInEditHierarchy);
 	}
 
 	/** If true, this actor or one of its components has property overrides applied. */
 	bool HasLevelInstancePropertyOverrides() const
 	{
-		return EnumHasAnyFlags(LevelInstanceFlags, ELevelInstanceFlags::HasPropertyOverrides);
+		return EnumHasAnyFlags(GetLevelInstanceFlags(), ELevelInstanceFlags::HasPropertyOverrides);
 	}
 
 	/** If true, this actor or one of its components has property overrides applied from an editable level instance (level instance is in a non readonly level) */
 	bool HasEditableLevelLevelInstancePropertyOverrides() const
 	{
-		return EnumHasAnyFlags(LevelInstanceFlags, ELevelInstanceFlags::HasEditablePropertyOverrides);
+		return EnumHasAnyFlags(GetLevelInstanceFlags(), ELevelInstanceFlags::HasEditablePropertyOverrides);
 	}
 
 	/** If true, the actor belongs to a level instance. */
 	bool IsInLevelInstance() const
 	{
-		return LevelInstanceType != ELevelInstanceType::None;
+		return GetLevelInstanceType() != ELevelInstanceType::None;
 	}
+
+	/** Whether this actor belongs to a level instance or not and what type of level instance. */
+	ENGINE_API ELevelInstanceType GetLevelInstanceType() const;
+
+	/** Flags related to level instances for this actor. */
+	ENGINE_API ELevelInstanceFlags GetLevelInstanceFlags() const;
 #endif
 
 	/** Networking - Server - TearOff this actor to stop replication to clients. Will set bTearOff to true. */
diff --git a/Engine/Source/Runtime/Engine/Private/ActorEditor.cpp b/Engine/Source/Runtime/Engine/Private/ActorEditor.cpp
--- a/Engine/Source/Runtime/Engine/Private/ActorEditor.cpp
+++ b/Engine/Source/Runtime/Engine/Private/ActorEditor.cpp
@@ -1740,6 +1740,28 @@ EDataValidationResult AActor::IsDataValid(FDataValidationContext& Context) const
 	return Result;
 }
 
+ELevelInstanceType AActor::GetLevelInstanceType() const
+{
+	// ChildActorComponent Actors forward the call their ParentActor 
+	if (AActor* ParentActor = GetParentActor())
+	{
+		return ParentActor->GetLevelInstanceType();
+	}
+
+	return LevelInstanceType;
+}
+
+ELevelInstanceFlags AActor::GetLevelInstanceFlags() const
+{
+	// ChildActorComponent Actors forward the call their ParentActor 
+	if (AActor* ParentActor = GetParentActor())
+	{
+		return ParentActor->GetLevelInstanceFlags();
+	}
+
+	return LevelInstanceFlags;
+}
+
 //---------------------------------------------------------------------------
 // DataLayers (begin)
 
