Move interface virtual function implementations into a .cpp file, for easier iteration.

#rnx

diff --git a/Engine/Plugins/Runtime/GeometryCache/Source/GeometryCache/Private/GeometryCache.cpp b/Engine/Plugins/Runtime/GeometryCache/Source/GeometryCache/Private/GeometryCache.cpp
--- a/Engine/Plugins/Runtime/GeometryCache/Source/GeometryCache/Private/GeometryCache.cpp
+++ b/Engine/Plugins/Runtime/GeometryCache/Source/GeometryCache/Private/GeometryCache.cpp
@@ -12,6 +12,7 @@
 #include "UObject/RenderingObjectVersion.h"
 #include "Widgets/Notifications/SNotificationList.h"
 #include "Framework/Notifications/NotificationManager.h"
+#include "Templates/SubclassOf.h"
 
 #include UE_INLINE_GENERATED_CPP_BY_NAME(GeometryCache)
 
diff --git a/Engine/Source/Runtime/Engine/Classes/Interfaces/Interface_AssetUserData.h b/Engine/Source/Runtime/Engine/Classes/Interfaces/Interface_AssetUserData.h
--- a/Engine/Source/Runtime/Engine/Classes/Interfaces/Interface_AssetUserData.h
+++ b/Engine/Source/Runtime/Engine/Classes/Interfaces/Interface_AssetUserData.h
@@ -3,12 +3,15 @@
 #pragma once
 
 #include "UObject/ObjectMacros.h"
-#include "Templates/SubclassOf.h"
 #include "Templates/Casts.h"
 #include "UObject/Interface.h"
 #include "Engine/AssetUserData.h"
 #include "Interface_AssetUserData.generated.h"
 
+template <typename T> class TSubclassOf;
+
+class UAssetUserData;
+
 /** Interface for assets/objects that can own UserData **/
 UINTERFACE(MinimalApi, meta=(CannotImplementInterfaceInBlueprint))
 class UInterface_AssetUserData : public UInterface
@@ -20,7 +23,7 @@ class IInterface_AssetUserData
 {
 	GENERATED_IINTERFACE_BODY()
 
-	virtual void AddAssetUserData(UAssetUserData* InUserData) {}
+	ENGINE_API virtual void AddAssetUserData(UAssetUserData* InUserData);
 	
 	/**
 	* Returns an instance of the provided AssetUserData class if it's contained in the target asset.
@@ -30,10 +33,7 @@ class IInterface_AssetUserData
 	* @return	The instance of the UAssetUserData class contained, or null if it doesn't exist
 	*/
 	UFUNCTION(BlueprintCallable, Category=AssetUserData, meta=(DeterminesOutputType="InUserDataClass"))
-	virtual UAssetUserData* GetAssetUserDataOfClass(TSubclassOf<UAssetUserData> InUserDataClass)
-	{
-		return nullptr;
-	}
+	ENGINE_API virtual UAssetUserData* GetAssetUserDataOfClass(TSubclassOf<UAssetUserData> InUserDataClass);
 
 	/**
 	* Checks whether or not an instance of the provided AssetUserData class is contained.
@@ -43,12 +43,9 @@ class IInterface_AssetUserData
 	* @return	Whether or not an instance of InUserDataClass was found
 	*/
 	UFUNCTION(BlueprintCallable, Category=AssetUserData)
-	virtual bool HasAssetUserDataOfClass(TSubclassOf<UAssetUserData> InUserDataClass)
-	{		
-		return GetAssetUserDataOfClass(InUserDataClass) != nullptr;
-	}
-	
-/**
+	ENGINE_API virtual bool HasAssetUserDataOfClass(TSubclassOf<UAssetUserData> InUserDataClass);
+
+	/**
 	* Creates and adds an instance of the provided AssetUserData class to the target asset.
 	*
 	* @param	InUserDataClass		UAssetUserData sub class to create
@@ -56,17 +53,9 @@ class IInterface_AssetUserData
 	* @return	Whether or not an instance of InUserDataClass was succesfully added
 	*/
 	UFUNCTION(BlueprintCallable, Category=AssetUserData)
-	virtual bool AddAssetUserDataOfClass(TSubclassOf<UAssetUserData> InUserDataClass)
-	{
-		check(InUserDataClass.Get() != nullptr);
-		AddAssetUserData(NewObject<UAssetUserData>(_getUObject(), InUserDataClass.Get()));
-		return GetAssetUserDataOfClass(InUserDataClass.Get()) != nullptr;
-	}
+	ENGINE_API virtual bool AddAssetUserDataOfClass(TSubclassOf<UAssetUserData> InUserDataClass);
 
-	virtual const TArray<UAssetUserData*>* GetAssetUserDataArray() const
-	{
-		return nullptr;
-	}
+	ENGINE_API virtual const TArray<UAssetUserData*>* GetAssetUserDataArray() const;
 
 	template<typename T>
 	T* GetAssetUserData()
@@ -80,6 +69,10 @@ class IInterface_AssetUserData
 		return CastChecked<T>(GetAssetUserDataOfClass(T::StaticClass()));
 	}
 
-	virtual void RemoveUserDataOfClass(TSubclassOf<UAssetUserData> InUserDataClass) {}
+	ENGINE_API virtual void RemoveUserDataOfClass(TSubclassOf<UAssetUserData> InUserDataClass);
 
 };
+
+#if UE_ENABLE_INCLUDE_ORDER_DEPRECATED_IN_5_7
+#include "Templates/SubclassOf.h"
+#endif
diff --git a/Engine/Source/Runtime/Engine/Private/Interfaces/Interface_AssetUserData.cpp b/Engine/Source/Runtime/Engine/Private/Interfaces/Interface_AssetUserData.cpp
new file mode 100644
--- /dev/null
+++ b/Engine/Source/Runtime/Engine/Private/Interfaces/Interface_AssetUserData.cpp
@@ -0,0 +1,36 @@
+// Copyright Epic Games, Inc. All Rights Reserved.
+
+#include "Interfaces/Interface_AssetUserData.h"
+
+#include "Templates/SubclassOf.h"
+#include "Engine/AssetUserData.h"
+
+void IInterface_AssetUserData::AddAssetUserData(UAssetUserData* InUserData)
+{
+}
+
+UAssetUserData* IInterface_AssetUserData::GetAssetUserDataOfClass(TSubclassOf<UAssetUserData> InUserDataClass)
+{
+	return nullptr;
+}
+
+bool IInterface_AssetUserData::HasAssetUserDataOfClass(TSubclassOf<UAssetUserData> InUserDataClass)
+{
+	return GetAssetUserDataOfClass(InUserDataClass) != nullptr;
+}
+
+bool IInterface_AssetUserData::AddAssetUserDataOfClass(TSubclassOf<UAssetUserData> InUserDataClass)
+{
+	check(InUserDataClass.Get() != nullptr);
+	AddAssetUserData(NewObject<UAssetUserData>(_getUObject(), InUserDataClass.Get()));
+	return GetAssetUserDataOfClass(InUserDataClass.Get()) != nullptr;
+}
+
+const TArray<UAssetUserData*>* IInterface_AssetUserData::GetAssetUserDataArray() const
+{
+	return nullptr;
+}
+
+void IInterface_AssetUserData::RemoveUserDataOfClass(TSubclassOf<UAssetUserData> InUserDataClass)
+{
+}
