Clean up additional deprecated code in the skeleton class

diff --git a/Engine/Source/Runtime/Engine/Classes/Animation/Skeleton.h b/Engine/Source/Runtime/Engine/Classes/Animation/Skeleton.h
--- a/Engine/Source/Runtime/Engine/Classes/Animation/Skeleton.h
+++ b/Engine/Source/Runtime/Engine/Classes/Animation/Skeleton.h
@@ -483,9 +483,6 @@ public:
 
 	const TArray<TSoftObjectPtr<USkeleton>>& GetCompatibleSkeletons() const { return CompatibleSkeletons; }
 
-	UE_DEPRECATED(5.2, "Please use UE::Anim::FSkeletonRemappingRegistry::GetRemapping.")
-	ENGINE_API const FSkeletonRemapping* GetSkeletonRemapping(const USkeleton* SourceSkeleton) const;
-
 	ENGINE_API bool GetUseRetargetModesFromCompatibleSkeleton() const;
 	ENGINE_API void SetUseRetargetModesFromCompatibleSkeleton(bool bUse);
 
@@ -584,8 +581,6 @@ private:
 
 	/** rig property will be saved separately */
 	ENGINE_API virtual void GetAssetRegistryTags(FAssetRegistryTagsContext Context) const override;
-	UE_DEPRECATED(5.4, "Implement the version that takes FAssetRegistryTagsContext instead.")
-	ENGINE_API virtual void GetAssetRegistryTags(TArray<FAssetRegistryTag>& OutTags) const override;
 
 public:
 
@@ -653,9 +648,6 @@ public:
 
 #if WITH_EDITORONLY_DATA
 
-	UE_DEPRECATED(5.4, "Please do not use this function - notifies are stored collectively in the asset registry now rather than centrally on the skeleton")
-	ENGINE_API void CollectAnimationNotifies();
-
 	/*
 	 * Collect animation notifies that are referenced in all animations that use this skeleton (uses the asset registry).
 	 * @param	OutNotifies		All the notifies that were found
@@ -848,17 +840,6 @@ public:
 	ENGINE_API void GetRetargetSources(TArray<FName>& OutRetargetSources) const;
 #endif
 
-	/** 
-	 * Get Track index of InAnimSeq for the BoneTreeIndex of BoneTree
-	 * this is slow, and it's not supposed to be used heavily
-	 * @param	InBoneTreeIdx	BoneTree Index
-	 * @param	InAnimSeq		Animation Sequence to get track index for 
-	 *
-	 * @return	Index of Track of Animation Sequence
-	 */
-	UE_DEPRECATED(5.2, "GetRawAnimationTrackIndex has been deprecated, use tracks are referenced by name instead")
-	ENGINE_API int32 GetRawAnimationTrackIndex(const int32 InSkeletonBoneIndex, const UAnimSequence* InAnimSeq);
-
 	/** 
 	 * Get Bone Tree Index from Reference Bone Index
 	 * @param	InSkinnedAsset	SkinnedAsset for the ref bone idx
diff --git a/Engine/Source/Runtime/Engine/Private/Animation/Skeleton.cpp b/Engine/Source/Runtime/Engine/Private/Animation/Skeleton.cpp
--- a/Engine/Source/Runtime/Engine/Private/Animation/Skeleton.cpp
+++ b/Engine/Source/Runtime/Engine/Private/Animation/Skeleton.cpp
@@ -552,10 +552,6 @@ void USkeleton::ConvertToFReferenceSkeleton()
 #endif
 }
 
-const FSkeletonRemapping* USkeleton::GetSkeletonRemapping(const USkeleton* SourceSkeleton) const
-{
-	return &UE::Anim::FSkeletonRemappingRegistry::Get().GetRemapping(SourceSkeleton, this);
-}
 
 #if WITH_EDITOR
 
@@ -1126,30 +1122,6 @@ void USkeleton::GetRetargetSources(TArray<FName>& OutRetargetSources) const
 
 #endif
 
-int32 USkeleton::GetRawAnimationTrackIndex(const int32 InSkeletonBoneIndex, const UAnimSequence* InAnimSeq)
-{
-	if (InSkeletonBoneIndex != INDEX_NONE)
-	{
-#if WITH_EDITOR
-		TArray<FName> BoneTrackNames;
-		InAnimSeq->GetDataModel()->GetBoneTrackNames(BoneTrackNames);
-		const FName BoneName = ReferenceSkeleton.GetBoneName(InSkeletonBoneIndex);
-		
-		return BoneTrackNames.IndexOfByKey(BoneName);
-#else
-		UAnimSequence::FScopedCompressedAnimSequence CompressedAnimSequence = InAnimSeq->GetCompressedData();
-		const TArray<FTrackToSkeletonMap>& BoneMappings = CompressedAnimSequence.Get().CompressedTrackToSkeletonMapTable;
-		return BoneMappings.IndexOfByPredicate([InSkeletonBoneIndex](const FTrackToSkeletonMap& Mapping)
-			{
-				return Mapping.BoneTreeIndex == InSkeletonBoneIndex;
-			});
-#endif
-	}
-
-	return INDEX_NONE;
-}
-
-
 int32 USkeleton::GetSkeletonBoneIndexFromMeshBoneIndex(const USkinnedAsset* InSkinnedAsset, const int32 MeshBoneIndex)
 {
 	check(MeshBoneIndex != INDEX_NONE);
@@ -1262,11 +1234,6 @@ int32 USkeleton::GetChildBones(int32 ParentBoneIndex, TArray<int32> & Children)
 	return ReferenceSkeleton.GetDirectChildBones(ParentBoneIndex, Children);
 }
 
-void USkeleton::CollectAnimationNotifies()
-{
-	CollectAnimationNotifies(AnimationNotifies);
-}
-
 void USkeleton::CollectAnimationNotifies(TArray<FName>& OutNotifies) const
 {
 	// first merge in AnimationNotifies
@@ -2066,14 +2033,6 @@ void USkeleton::HandleVirtualBoneChanges()
 }
 
 #if WITH_EDITOR
-
-void USkeleton::GetAssetRegistryTags(TArray<FAssetRegistryTag>& OutTags) const
-{
-	PRAGMA_DISABLE_DEPRECATION_WARNINGS;
-	Super::GetAssetRegistryTags(OutTags);
-	PRAGMA_ENABLE_DEPRECATION_WARNINGS;
-}
-
 void USkeleton::GetAssetRegistryTags(FAssetRegistryTagsContext Context) const
 {
 	Super::GetAssetRegistryTags(Context);
