Rename the NetDriver's transaction for ServerReplicateActors_ForConnection to be the function name instead of the connection's view target.

diff --git a/Engine/Source/Runtime/Engine/Private/NetDriver.cpp b/Engine/Source/Runtime/Engine/Private/NetDriver.cpp
--- a/Engine/Source/Runtime/Engine/Private/NetDriver.cpp
+++ b/Engine/Source/Runtime/Engine/Private/NetDriver.cpp
@@ -6259,7 +6259,7 @@ int32 UNetDriver::ServerReplicateActors(float DeltaSeconds)
 
 			if (bUseGranularTransactions)
 			{
-				const FName TransactionalWorkName = Connection->ViewTarget.GetFName();
+				static const FName TransactionalWorkName("ServerReplicateActors_ForConnection");
 				UE::RemoteExecutor::EnqueueWork(TransactionalWorkName, true, [this, ConstParams = Params]()
 					{
 						UE::Net::FServerReplicateActors_ForConnectionParams Params = ConstParams;
