Fix incorrect comment.
#rnx

diff --git a/Engine/Source/Runtime/Engine/Classes/Engine/Texture2D.h b/Engine/Source/Runtime/Engine/Classes/Engine/Texture2D.h
--- a/Engine/Source/Runtime/Engine/Classes/Engine/Texture2D.h
+++ b/Engine/Source/Runtime/Engine/Classes/Engine/Texture2D.h
@@ -389,7 +389,7 @@ public:
 	virtual bool IsVirtualTexturedWithSinglePhysicalSpace() const { return false;  }
 
 	/**
-	 * Returns true if this virtual texture requires that it's physical pool is never split by the r.VT.PageTableSupportType=2 setting.
+	 * Returns true if this virtual texture requires that its physical pool is never split by the r.VT.PageTableMode=2 setting.
 	 * This is necessary if the texture is part of a system that assumes that all textures in the system are in the same physical pool.
 	 * Note that textures with this setting will necessarily be put in different physical spaces to those without the setting.
 	 */
