Fixed BodySetup not creating its physics meshes when both p.Chaos.EnableAsyncInitBody and p.EnableDeferredPhysicsCreation are enabled.
This is a regression caused by CL 42884400.

diff --git a/Engine/Source/Runtime/Engine/Private/PhysicsEngine/Experimental/PhysScene_AsyncPhysicsStateJobQueue.cpp b/Engine/Source/Runtime/Engine/Private/PhysicsEngine/Experimental/PhysScene_AsyncPhysicsStateJobQueue.cpp
--- a/Engine/Source/Runtime/Engine/Private/PhysicsEngine/Experimental/PhysScene_AsyncPhysicsStateJobQueue.cpp
+++ b/Engine/Source/Runtime/Engine/Private/PhysicsEngine/Experimental/PhysScene_AsyncPhysicsStateJobQueue.cpp
@@ -4,6 +4,7 @@
 #include "Physics/Experimental/PhysScene_Chaos.h"
 #include "PhysicsEngine/BodySetup.h"
 #include "Engine/World.h"
+#include "Components/ActorComponent.h"
 
 namespace Chaos
 {
@@ -244,6 +245,16 @@ void FPhysScene_AsyncPhysicsStateJobQueue::FJob::OnPreExecute_GameThread() const
 	{
 		case EJobType::CreatePhysicsState:
 		{
+			if (GEnableDeferredPhysicsCreation)
+			{
+				TSet<UBodySetup*> UniqueBodySetups;
+				Processor->CollectBodySetupsWithPhysicsMeshesToCreate(UniqueBodySetups);
+				if (UniqueBodySetups.Num())
+				{
+					TArray<UBodySetup*> BodySetups = UniqueBodySetups.Array();
+					ParallelFor(TEXT("CreatePhysicsMeshes.PF"), BodySetups.Num(), 1, [this, &BodySetups](int32 Index) { BodySetups[Index]->CreatePhysicsMeshes(); });
+				}
+			}
 			Processor->OnAsyncCreatePhysicsStateBegin_GameThread();
 		}
 		break;
