fix crash when calling recompileshaders changed without -odschost

diff --git a/Engine/Source/Runtime/Engine/Private/ShaderCompiler/ShaderCompiler.cpp b/Engine/Source/Runtime/Engine/Private/ShaderCompiler/ShaderCompiler.cpp
--- a/Engine/Source/Runtime/Engine/Private/ShaderCompiler/ShaderCompiler.cpp
+++ b/Engine/Source/Runtime/Engine/Private/ShaderCompiler/ShaderCompiler.cpp
@@ -4531,7 +4531,14 @@ bool RecompileShaders(const TCHAR* Cmd, FOutputDevice& Ar)
 
 		ERHIFeatureLevel::Type TargetFeatureLevel = GetMaxSupportedFeatureLevel(GMaxRHIShaderPlatform);
 		const EMaterialQualityLevel::Type ActiveQualityLevel = GetCachedScalabilityCVars().MaterialQualityLevel;
-		GODSCManager->AddThreadedRequest(MaterialsToLoad, ShaderTypesToLoad, GMaxRHIShaderPlatform, TargetFeatureLevel, ActiveQualityLevel, CommandType, RequestedMaterialName, ExtraCompilerFlags);
+		if (FODSCManager::IsODSCActive())
+		{
+			GODSCManager->AddThreadedRequest(MaterialsToLoad, ShaderTypesToLoad, GMaxRHIShaderPlatform, TargetFeatureLevel, ActiveQualityLevel, CommandType, RequestedMaterialName, ExtraCompilerFlags);
+		}
+		else
+		{
+			UE_LOG(LogShaders, Error, TEXT("Tried to use recompileshaders but ODSC is not active"));
+		}
 #endif
 		return true;
 	}
