Add more logging to ODSC so that we know how many shaders are sent from the ODSC server

diff --git a/Engine/Source/Runtime/Engine/Private/Materials/Material.cpp b/Engine/Source/Runtime/Engine/Private/Materials/Material.cpp
--- a/Engine/Source/Runtime/Engine/Private/Materials/Material.cpp
+++ b/Engine/Source/Runtime/Engine/Private/Materials/Material.cpp
@@ -6266,7 +6266,7 @@ void UMaterial::CompileODSCMaterialsForRemoteRecompile(TArray<FODSCRequestPayloa
 		FString OriginalMaterialName;
 	};
 	TMap<UMaterialInterface*, FShadersToCompile> CoalescedShadersToCompile;
-
+	
 	TSet<FString> MaterialsFailingToLoad;
 	for (const FODSCRequestPayload& payload : ShadersToRecompile)
 	{
@@ -6359,6 +6359,9 @@ void UMaterial::CompileODSCMaterialsForRemoteRecompile(TArray<FODSCRequestPayloa
 	// This is code that should be run on the cooker so shouldn't be a big deal.
 	FlushRenderingCommands();
 
+	int32 NumShaders = 0;
+	int32 NumShaderPipelines = 0;
+	int32 NumShaderMaps = 0;
 	for (TMap<FString, TArray<FMaterialResource*> >::TIterator It(CompilingResources); It; ++It)
 	{
 		TArray<FMaterialResource*>& ResourceArray = It.Value();
@@ -6367,10 +6370,21 @@ void UMaterial::CompileODSCMaterialsForRemoteRecompile(TArray<FODSCRequestPayloa
 		for (int32 Index = 0; Index < ResourceArray.Num(); Index++)
 		{
 			FMaterialResource* CurrentResource = ResourceArray[Index];
-			OutShaderMapArray.Add(CurrentResource->GetGameThreadShaderMap());
+			FMaterialShaderMap* MaterialShaderMap = CurrentResource->GetGameThreadShaderMap();
+			OutShaderMapArray.Add(MaterialShaderMap);
+
+			int32 ThisMaterialNumShaders = 0;
+			int32 ThisMaterialNumShaderPipelines = 0;
+			MaterialShaderMap->CountNumShaders(ThisMaterialNumShaders, ThisMaterialNumShaderPipelines);
+
+			NumShaders += ThisMaterialNumShaders;
+			NumShaderPipelines += ThisMaterialNumShaderPipelines;
+			++NumShaderMaps;
 		}
 	}
 
+	UE_LOG(LogShaders, Display, TEXT("Sending %d shaders, %d pipelines for %d shader maps"), NumShaders, NumShaderPipelines, NumShaderMaps);
+
 	// Report errors last to ensure visibility: some messages may have been output during shader compilation
 	for (FString& MaterialFailingToLoad : MaterialsFailingToLoad)
 	{
diff --git a/Engine/Source/Runtime/Engine/Private/ShaderCompiler/ShaderCompilerEditor.cpp b/Engine/Source/Runtime/Engine/Private/ShaderCompiler/ShaderCompilerEditor.cpp
--- a/Engine/Source/Runtime/Engine/Private/ShaderCompiler/ShaderCompilerEditor.cpp
+++ b/Engine/Source/Runtime/Engine/Private/ShaderCompiler/ShaderCompilerEditor.cpp
@@ -363,6 +363,11 @@ void RecompileShadersForRemote(
 	{
 		TMap<FString, TArray<TRefCountPtr<FMaterialShaderMap>>> CompiledShaderMaps;
 		{
+			if (!GShaderCompilingManager->AllowForIncompleteShaderMaps())
+			{
+				UE_LOG(LogShaders, Warning, TEXT("AllowForIncompleteShaderMaps is not set to true, ODSC shaders will most likely not be recompiled"));
+			}
+
 			// UMaterial::CompileODSCMaterialsForRemoteRecompile will call LoadObjects on the material names but doesn't keep them around. Add a GC guard to ensure we can still get them
 			// before they get unloaded, so that the whole chain UMaterial->FMaterial->FMaterialShaderMap is kept intact, and we can merge the next batch of ODSC requests
 			FGCScopeGuard NoGCScopeGuard;
diff --git a/Engine/Source/Runtime/Engine/Public/ShaderCompiler.h b/Engine/Source/Runtime/Engine/Public/ShaderCompiler.h
--- a/Engine/Source/Runtime/Engine/Public/ShaderCompiler.h
+++ b/Engine/Source/Runtime/Engine/Public/ShaderCompiler.h
@@ -1214,6 +1214,11 @@ public:
 		bAllowForIncompleteShaderMaps = toggle;
 	}
 
+	bool AllowForIncompleteShaderMaps() const
+	{
+		return bAllowForIncompleteShaderMaps;
+	}
+
 	ENGINE_API bool IsCompilingShaderMap(uint32 Id);
 
 	/** Prepares a job of the given type for compilation.  If a job with the given Id/Key already exists, it will attempt to adjust to the higher priority if possible, and nullptr will be returned.
