Fix UCanvas::MeasureStringInternal to correctly handle measuring strings containing newlines. More closely matches the logic in FSlateFontMeasure::MeasureStringInternal.
Added some unit tests to validate the fix.

#tests Ran new unit tests under UE::Engine::Canvas - with the fix, 8/8 pass, before the fix, 2/8 passed.

diff --git a/Engine/Source/Runtime/Engine/Private/Tests/UserInterface/CanvasTests.cpp b/Engine/Source/Runtime/Engine/Private/Tests/UserInterface/CanvasTests.cpp
new file mode 100644
--- /dev/null
+++ b/Engine/Source/Runtime/Engine/Private/Tests/UserInterface/CanvasTests.cpp
@@ -0,0 +1,232 @@
+// Copyright Epic Games, Inc. All Rights Reserved.
+
+#if WITH_TESTS
+
+#include "Engine/Canvas.h"
+#include "Engine/Engine.h"
+#include "Engine/Font.h"
+#include "Tests/TestHarnessAdapter.h"
+
+namespace UE::Canvas::Private
+{
+
+static UFont* GetTestFont()
+{
+	if (!GEngine)
+	{
+		return nullptr;
+	}
+	if (UFont* Font = GEngine->GetTinyFont())
+	{
+		return Font;
+	}
+	if (UFont* Font = GEngine->GetSmallFont())
+	{
+		return Font;
+	}
+	if (UFont* Font = GEngine->GetMediumFont())
+	{
+		return Font;
+	}
+	return nullptr;
+}
+
+} // namespace UE::Canvas::Private
+
+namespace UE::Canvas
+{
+
+TEST_CASE_NAMED(CanvasBasicMultilineWidth, "UE::Engine::Canvas::BasicMultilineWidth", "[Engine][Canvas][Text][EngineFilter]")
+{
+	UFont* TestFont = Private::GetTestFont();
+	REQUIRE_MESSAGE(TEXT("Engine font must be available for testing"), TestFont != nullptr);
+
+	FTextSizingParameters Params(TestFont, /* ScaleX = */ 1.0f, /* ScaleY = */ 1.0f);
+
+	// Multi-line width should equal widest individual line, not sum
+	constexpr FStringView SingleLine = TEXTVIEW("Hello");
+	constexpr FStringView MultiLine = TEXTVIEW("Hello\nWorld\nTest");
+
+	UCanvas::CanvasStringSize(Params, SingleLine);
+	const float HelloWidth = Params.DrawXL;
+
+	UCanvas::CanvasStringSize(Params, TEXTVIEW("World"));
+	const float WorldWidth = Params.DrawXL;
+
+	UCanvas::CanvasStringSize(Params, TEXTVIEW("Test"));
+	const float TestWidth = Params.DrawXL;
+
+	UCanvas::CanvasStringSize(Params, MultiLine);
+	const float MultiLineWidth = Params.DrawXL;
+
+	const float ExpectedWidth = FMath::Max(HelloWidth, WorldWidth, TestWidth);
+
+	REQUIRE_MESSAGE(TEXT("Multi-line width should equal widest individual line"),
+		FMath::IsNearlyEqual(MultiLineWidth, ExpectedWidth, 0.1f));
+}
+
+TEST_CASE_NAMED(CanvasSingleVsMultilineHeight, "UE::Engine::Canvas::SingleVsMultilineHeight", "[Engine][Canvas][Text][EngineFilter]")
+{
+	UFont* TestFont = Private::GetTestFont();
+	REQUIRE_MESSAGE(TEXT("Engine font must be available for testing"), TestFont != nullptr);
+
+	FTextSizingParameters Params(TestFont, /* ScaleX = */ 1.0f, /* ScaleY = */ 1.0f);
+
+	// Single line height
+	UCanvas::CanvasStringSize(Params, TEXTVIEW("Hello"));
+	const float SingleLineHeight = Params.DrawYL;
+
+	// Two line height should be approximately single + max font height
+	UCanvas::CanvasStringSize(Params, TEXTVIEW("Hello\nWorld"));
+	const float TwoLineHeight = Params.DrawYL;
+
+	const float MaxCharHeight = TestFont->GetMaxCharHeight();
+	const float ExpectedTwoLineHeight = MaxCharHeight + SingleLineHeight;
+
+	REQUIRE_MESSAGE(TEXT("Two-line height should be single line height plus font max height"),
+		FMath::IsNearlyEqual(TwoLineHeight, ExpectedTwoLineHeight, MaxCharHeight * 0.1f));
+}
+
+TEST_CASE_NAMED(CanvasEmptyLines, "UE::Engine::Canvas::EmptyLines", "[Engine][Canvas][Text][EngineFilter]")
+{
+	UFont* TestFont = Private::GetTestFont();
+	REQUIRE_MESSAGE(TEXT("Engine font must be available for testing"), TestFont != nullptr);
+
+	FTextSizingParameters Params(TestFont, /* ScaleX = */ 1.0f, /* ScaleY = */ 1.0f);
+
+	// Empty lines should still contribute to height
+	UCanvas::CanvasStringSize(Params, TEXTVIEW("Hello\n\nWorld"));
+	const float ThreeLineHeight = Params.DrawYL;
+
+	UCanvas::CanvasStringSize(Params, TEXTVIEW("Hello\nWorld"));
+	const float TwoLineHeight = Params.DrawYL;
+
+	const float MaxCharHeight = TestFont->GetMaxCharHeight();
+
+	REQUIRE_MESSAGE(TEXT("Empty line should add one more line advance"),
+		FMath::IsNearlyEqual(ThreeLineHeight, TwoLineHeight + MaxCharHeight, MaxCharHeight * 0.1f));
+}
+
+TEST_CASE_NAMED(CanvasLineWidthVariation, "UE::Engine::Canvas::LineWidthVariation", "[Engine][Canvas][Text][EngineFilter]")
+{
+	UFont* TestFont = Private::GetTestFont();
+	REQUIRE_MESSAGE(TEXT("Engine font must be available for testing"), TestFont != nullptr);
+
+	FTextSizingParameters Params(TestFont, /* ScaleX = */ 1.0f, /* ScaleY = */ 1.0f);
+
+	// Lines with very different widths
+	constexpr FStringView VaryingWidths = TEXTVIEW("A\nThis is a much longer line\nB");
+
+	UCanvas::CanvasStringSize(Params, VaryingWidths);
+	const float MultiLineWidth = Params.DrawXL;
+
+	UCanvas::CanvasStringSize(Params, TEXTVIEW("This is a much longer line"));
+	const float LongestLineWidth = Params.DrawXL;
+
+	REQUIRE_MESSAGE(TEXT("Multi-line width should match the longest individual line"),
+		FMath::IsNearlyEqual(MultiLineWidth, LongestLineWidth, 0.1f));
+}
+
+TEST_CASE_NAMED(CanvasStoppingOffsetMultiline, "UE::Engine::Canvas::StoppingOffsetMultiline", "[Engine][Canvas][Text][EngineFilter]")
+{
+	UFont* TestFont = Private::GetTestFont();
+	REQUIRE_MESSAGE(TEXT("Engine font must be available for testing"), TestFont != nullptr);
+
+	FTextSizingParameters Params(TestFont, /* ScaleX = */ 1.0f, /* ScaleY = */ 1.0f);
+
+	// Stopping offset on second line
+	constexpr FStringView MultiLine = TEXTVIEW("A\nWorld");
+
+	// Get width of "Wo" to calculate offset into second line
+	UCanvas::CanvasStringSize(Params, TEXTVIEW("Wo"));
+	const float TwoCharsWidth = Params.DrawXL;
+
+	// Stop after "Wo" on second line
+	int32 OutIndex;
+	UCanvas::MeasureStringInternal(Params, MultiLine, static_cast<int32>(TwoCharsWidth), UCanvas::ELastCharacterIndexFormat::LastWholeCharacterBeforeOffset, OutIndex);
+
+	// Should stop at index 3 ("A\nWo" = 4 chars, but LastWholeBefore gives index 3)
+	int32 ExpectedIndex = 3; // Position of 'o' in "World"
+
+	REQUIRE_MESSAGE(TEXT("Multi-line stopping offset should work correctly across line boundaries"),
+		OutIndex == ExpectedIndex);
+}
+
+TEST_CASE_NAMED(CanvasScalingConsistency, "UE::Engine::Canvas::ScalingConsistency", "[Engine][Canvas][Text][EngineFilter]")
+{
+	UFont* TestFont = Private::GetTestFont();
+	REQUIRE_MESSAGE(TEXT("Engine font must be available for testing"), TestFont != nullptr);
+
+	FTextSizingParameters Params1(TestFont, /* ScaleX = */ 1.0f, /* ScaleY = */ 1.0f);
+	FTextSizingParameters Params2(TestFont, /* ScaleX = */ 2.0f, /* ScaleY = */ 2.0f);
+
+	constexpr FStringView TestText = TEXTVIEW("Hello\nWorld");
+
+	UCanvas::CanvasStringSize(Params1, TestText);
+	const float Width1 = Params1.DrawXL;
+	const float Height1 = Params1.DrawYL;
+
+	UCanvas::CanvasStringSize(Params2, TestText);
+	const float Width2 = Params2.DrawXL;
+	const float Height2 = Params2.DrawYL;
+
+	REQUIRE_MESSAGE(TEXT("Scaled width should be approximately 2x original"),
+		FMath::IsNearlyEqual(Width2, Width1 * 2.0f, Width1 * 0.1f));
+	REQUIRE_MESSAGE(TEXT("Scaled height should be approximately 2x original"),
+		FMath::IsNearlyEqual(Height2, Height1 * 2.0f, Height1 * 0.1f));
+}
+
+TEST_CASE_NAMED(CanvasKerningAcrossLines, "UE::Engine::Canvas::KerningAcrossLines", "[Engine][Canvas][Text][EngineFilter]")
+{
+	UFont* TestFont = Private::GetTestFont();
+	REQUIRE_MESSAGE(TEXT("Engine font must be available for testing"), TestFont != nullptr);
+
+	FTextSizingParameters Params(TestFont, /* ScaleX = */ 1.0f, /* ScaleY = */ 1.0f);
+
+	// Kerning should not apply across line boundaries
+	UCanvas::CanvasStringSize(Params, TEXTVIEW("A"));
+	const float AWidth = Params.DrawXL;
+
+	UCanvas::CanvasStringSize(Params, TEXTVIEW("V"));
+	const float VWidth = Params.DrawXL;
+
+	UCanvas::CanvasStringSize(Params, TEXTVIEW("A\nV"));
+	const float MultiLineWidth = Params.DrawXL;
+
+	// Multi-line version should have width = Max(AWidth, VWidth) without kerning
+	const float ExpectedWidth = FMath::Max(AWidth, VWidth);
+
+	REQUIRE_MESSAGE(TEXT("Kerning should not apply across line boundaries"),
+		FMath::IsNearlyEqual(MultiLineWidth, ExpectedWidth, ExpectedWidth * 0.1f));
+}
+
+TEST_CASE_NAMED(CanvasSpacingAdjustments, "UE::Engine::Canvas::SpacingAdjustments", "[Engine][Canvas][Text][EngineFilter]")
+{
+	UFont* TestFont = Private::GetTestFont();
+	REQUIRE_MESSAGE(TEXT("Engine font must be available for testing"), TestFont != nullptr);
+
+	FTextSizingParameters ParamsNoSpacing(TestFont, /* ScaleX = */ 1.0f, /* ScaleY = */ 1.0f);
+	ParamsNoSpacing.SpacingAdjust = FVector2D(0.0f, 0.0f);
+	
+	FTextSizingParameters ParamsWithSpacing(TestFont, /* ScaleX = */ 1.0f, /* ScaleY = */ 1.0f);
+	ParamsWithSpacing.SpacingAdjust = FVector2D(5.0f, 10.0f);
+
+	constexpr FStringView TestText = TEXTVIEW("Hello\nWorld");
+
+	UCanvas::CanvasStringSize(ParamsNoSpacing, TestText);
+	const float WidthNoSpacing = ParamsNoSpacing.DrawXL;
+	const float HeightNoSpacing = ParamsNoSpacing.DrawYL;
+
+	UCanvas::CanvasStringSize(ParamsWithSpacing, TestText);
+	const float WidthWithSpacing = ParamsWithSpacing.DrawXL;
+	const float HeightWithSpacing = ParamsWithSpacing.DrawYL;
+
+	REQUIRE_MESSAGE(TEXT("Horizontal spacing adjustment should increase width"),
+		WidthWithSpacing > WidthNoSpacing);
+	REQUIRE_MESSAGE(TEXT("Vertical spacing adjustment should increase height"),
+		HeightWithSpacing > HeightNoSpacing);
+}
+
+} // namespace UE::Canvas::Private
+
+#endif // WITH_TESTS
\ No newline at end of file
diff --git a/Engine/Source/Runtime/Engine/Private/UserInterface/Canvas.cpp b/Engine/Source/Runtime/Engine/Private/UserInterface/Canvas.cpp
--- a/Engine/Source/Runtime/Engine/Private/UserInterface/Canvas.cpp
+++ b/Engine/Source/Runtime/Engine/Private/UserInterface/Canvas.cpp
@@ -1187,25 +1187,6 @@ ENGINE_API void StringSize(const UFont* Font,int32& XL,int32& YL, FStringView Te
 	YL = FMath::TruncToInt(Parameters.DrawYL);
 }
 
-/**
- * Calculates the width and height of a typical character in the specified font.
- *
- * @param	DrawFont			the font to use for calculating the character size
- * @param	DefaultCharWidth	[out] will be set to the width of the typical character
- * @param	DefaultCharHeight	[out] will be set to the height of the typical character
- * @param	pDefaultChar		if specified, pointer to a single character to use for calculating the default character size
- */
-static void GetDefaultCharSize( const UFont* DrawFont, float& DefaultCharWidth, float& DefaultCharHeight, const TCHAR* pDefaultChar=NULL )
-{
-	TCHAR DefaultChar = pDefaultChar != NULL ? *pDefaultChar : TEXT('0');
-	DrawFont->GetCharSize(DefaultChar, DefaultCharWidth, DefaultCharHeight);
-	if ( DefaultCharWidth == 0 )
-	{
-		// this font doesn't contain '0', try 'A'
-		DrawFont->GetCharSize(TEXT('A'), DefaultCharWidth, DefaultCharHeight);
-	}
-}
-
 void UCanvas::MeasureStringInternal(FTextSizingParameters& Parameters, FStringView Text, const int32 StopAfterHorizontalOffset, const ELastCharacterIndexFormat CharIndexFormat, int32& OutLastCharacterIndex)
 {
 	// initialize output so it always makes some sense
@@ -1216,51 +1197,53 @@ void UCanvas::MeasureStringInternal(FTextSizingParameters& Parameters, FStringVi
 
 	if (Parameters.DrawFont)
 	{
-		// get a default character width and height to be used for non-renderable characters
-		float DefaultCharWidth, DefaultCharHeight;
-		GetDefaultCharSize( Parameters.DrawFont, DefaultCharWidth, DefaultCharHeight );
-
 		// we'll need to use scaling in multiple places, so create a variable to hold it so that if we modify
 		// how the scale is calculated we only have to update these two lines
 		const float ScaleX = Parameters.Scaling.X;
 		const float ScaleY = Parameters.Scaling.Y;
 
-		const float DefaultCharIncrement = Parameters.SpacingAdjust.X * ScaleX;
+		const float UserCharSpacingX = Parameters.SpacingAdjust.X * ScaleX;
+		const float UserCharSpacingY = Parameters.SpacingAdjust.Y * ScaleY;
+
+		const float MaxCharHeight = Parameters.DrawFont->GetMaxCharHeight() * ScaleY;
+		// Advance by max height for consistent line spacing (plus any additional user-specified spacing adjustment).
+		const float AdvanceHeight = MaxCharHeight + UserCharSpacingY;
+
+		FVector2f MaxSize(ForceInitToZero);
+		FVector2f CurrentSize(ForceInitToZero);
+
 		const TCHAR* const BeginPos = Text.GetData();
-		const TCHAR* const EndPos = Text.GetData() + Text.Len();
-		const TCHAR* CurrentPos;
+		const TCHAR* const EndPos = BeginPos + Text.Len();
+		const TCHAR* CurrentPos = BeginPos;
 		const TCHAR* PrevPos = nullptr;
-		for (CurrentPos = BeginPos; CurrentPos < EndPos && *CurrentPos; ++CurrentPos)
+		for (; CurrentPos < EndPos && *CurrentPos; ++CurrentPos)
 		{
-			float CharWidth, CharHeight;
-
 			const TCHAR Ch = *CurrentPos;
-			Parameters.DrawFont->GetCharSize(Ch, CharWidth, CharHeight);
-			if (CharHeight == 0 && Ch == TEXT('\n'))
-			{
-				CharHeight = DefaultCharHeight;
-			}
 
-			float CharSpacing = DefaultCharIncrement;
-			if (PrevPos)
+			if (Ch == FChar::LineFeed)
 			{
-				CharSpacing += Parameters.DrawFont->GetCharKerning(*PrevPos, Ch) * ScaleX;
+				CurrentSize.X = 0.f;
+				CurrentSize.Y += AdvanceHeight;
+				PrevPos = nullptr;
+				continue;
 			}
 
+			float CharWidth, CharHeight;
+			Parameters.DrawFont->GetCharSize(Ch, CharWidth, CharHeight);
 			CharWidth *= ScaleX;
 			CharHeight *= ScaleY;
 
-			// If we have another character that isn't whitespace, append the character spacing
-			const TCHAR* const NextPos = CurrentPos + 1;
-			if (NextPos < EndPos && *NextPos && !FChar::IsWhitespace(*NextPos))
+			float CharKerning = 0.f;
+			if (PrevPos)
 			{
-				CharWidth += CharSpacing;
+				CharKerning = Parameters.DrawFont->GetCharKerning(*PrevPos, Ch) * ScaleX;
 			}
 
-			const float ScaledVertSpacing = Parameters.SpacingAdjust.Y * ScaleY;
+			const float AdvanceWidth = CharWidth + CharKerning + UserCharSpacingX;
 
-			Parameters.DrawXL += CharWidth;
-			Parameters.DrawYL = FMath::Max<float>(Parameters.DrawYL, CharHeight + ScaledVertSpacing );
+			CurrentSize.X += AdvanceWidth;
+			MaxSize.X = FMath::Max(MaxSize.X, CurrentSize.X);
+			MaxSize.Y = FMath::Max(MaxSize.Y, CurrentSize.Y + CharHeight + UserCharSpacingY);
 
 			// Were we asked to stop measuring after the specified horizontal offset in pixels?
 			if (StopAfterHorizontalOffset != INDEX_NONE)
@@ -1268,8 +1251,9 @@ void UCanvas::MeasureStringInternal(FTextSizingParameters& Parameters, FStringVi
 				if (CharIndexFormat == ELastCharacterIndexFormat::CharacterAtOffset)
 				{
 					// Round our test toward the character's center position
-					const float TotalCharWidth = CharWidth + Parameters.DrawFont->GetCharHorizontalOffset(Ch);
-					if (StopAfterHorizontalOffset < Parameters.DrawXL - TotalCharWidth / 2)
+					const float CharHorizontalOffset = Parameters.DrawFont->GetCharHorizontalOffset(Ch) * ScaleX;
+					const float TotalCharWidth = AdvanceWidth + CharHorizontalOffset;
+					if (StopAfterHorizontalOffset < CurrentSize.X - TotalCharWidth / 2)
 					{
 						// We've reached the stopping point, so bail
 						break;
@@ -1277,7 +1261,7 @@ void UCanvas::MeasureStringInternal(FTextSizingParameters& Parameters, FStringVi
 				}
 				else if (CharIndexFormat == ELastCharacterIndexFormat::LastWholeCharacterBeforeOffset)
 				{
-					if (StopAfterHorizontalOffset < Parameters.DrawXL - CharWidth)
+					if (StopAfterHorizontalOffset < CurrentSize.X)
 					{
 						--CurrentPos;
 						// We've reached the stopping point, so bail
@@ -1289,7 +1273,23 @@ void UCanvas::MeasureStringInternal(FTextSizingParameters& Parameters, FStringVi
 			PrevPos = CurrentPos;
 		}
 
-		OutLastCharacterIndex = CurrentPos - Text.GetData();
+		Parameters.DrawXL = MaxSize.X;
+		Parameters.DrawYL = MaxSize.Y;
+
+		if (CharIndexFormat == ELastCharacterIndexFormat::LastWholeCharacterBeforeOffset)
+		{
+			// LastWholeCharacterBeforeOffset is an inclusive index
+			if (CurrentPos < BeginPos)
+			{
+				CurrentPos = BeginPos;
+			}
+			else if (CurrentPos >= EndPos)
+			{
+				CurrentPos = EndPos - 1;
+			}
+		}
+
+		OutLastCharacterIndex = CurrentPos - BeginPos;
 	}
 }
 
