UCanvas::MeasureStringInternal: Addressing minor bughawk feedback.

#rnx

diff --git a/Engine/Source/Runtime/Engine/Private/Tests/UserInterface/CanvasTests.cpp b/Engine/Source/Runtime/Engine/Private/Tests/UserInterface/CanvasTests.cpp
--- a/Engine/Source/Runtime/Engine/Private/Tests/UserInterface/CanvasTests.cpp
+++ b/Engine/Source/Runtime/Engine/Private/Tests/UserInterface/CanvasTests.cpp
@@ -227,6 +227,6 @@ TEST_CASE_NAMED(CanvasSpacingAdjustments, "UE::Engine::Canvas::SpacingAdjustment
 		HeightWithSpacing > HeightNoSpacing);
 }
 
-} // namespace UE::Canvas::Private
+} // namespace UE::Canvas
 
 #endif // WITH_TESTS
\ No newline at end of file
diff --git a/Engine/Source/Runtime/Engine/Private/UserInterface/Canvas.cpp b/Engine/Source/Runtime/Engine/Private/UserInterface/Canvas.cpp
--- a/Engine/Source/Runtime/Engine/Private/UserInterface/Canvas.cpp
+++ b/Engine/Source/Runtime/Engine/Private/UserInterface/Canvas.cpp
@@ -1220,7 +1220,7 @@ void UCanvas::MeasureStringInternal(FTextSizingParameters& Parameters, FStringVi
 		{
 			const TCHAR Ch = *CurrentPos;
 
-			if (Ch == FChar::LineFeed)
+			if (FChar::IsLinebreak(Ch))
 			{
 				CurrentSize.X = 0.f;
 				CurrentSize.Y += AdvanceHeight;
