[shaders] remove submitjobs time and mutex wait stall stats; these are no longer useful and are further not thread safe and so causing TSAN report noise

diff --git a/Engine/Source/Runtime/Engine/Private/ShaderCompiler/ShaderCompilerJobCache.cpp b/Engine/Source/Runtime/Engine/Private/ShaderCompiler/ShaderCompilerJobCache.cpp
--- a/Engine/Source/Runtime/Engine/Private/ShaderCompiler/ShaderCompilerJobCache.cpp
+++ b/Engine/Source/Runtime/Engine/Private/ShaderCompiler/ShaderCompilerJobCache.cpp
@@ -617,7 +617,7 @@ public:
 	void SubmitJobs(const TArray<FShaderCommonCompileJobPtr>& InJobs);
 
 	/** This is an entry point for all jobs that have finished the compilation (whether real or cached). Can be called from multiple threads. Returns mutex stall time. */
-	double ProcessFinishedJob(FShaderCommonCompileJob* FinishedJob, EShaderCompileJobStatus Status);
+	void ProcessFinishedJob(FShaderCommonCompileJob* FinishedJob, EShaderCompileJobStatus Status);
 
 	/** Adds the job to cache. */
 	void AddToCacheAndProcessPending(FShaderCommonCompileJob* FinishedJob);
@@ -1393,7 +1393,6 @@ void FShaderJobCache::SubmitJob(FShaderCommonCompileJob* Job)
 		// the opposite order.
 		double StallStart = FPlatformTime::Seconds();
 		JobLock.WriteLock();
-		Job->TimeTaskSubmitJobsStall += FPlatformTime::Seconds() - StallStart;
 
 		// Job was linked in PendingSubmitJobTaskJobs before calling SubmitJob -- if it's not linked now, it means it was cancelled via
 		// call to RemoveAllPendingJobsWithId, so we can ignore it and just return.
@@ -1424,7 +1423,7 @@ void FShaderJobCache::SubmitJob(FShaderCommonCompileJob* Job)
 			Job->SerializeOutput(LoadContext);
 
 			// finish the job instantly
-			Job->TimeTaskSubmitJobsStall += ProcessFinishedJob(Job, EShaderCompileJobStatus::CompleteFoundInCache);
+			ProcessFinishedJob(Job, EShaderCompileJobStatus::CompleteFoundInCache);
 		}
 		else
 		{
@@ -1557,17 +1556,13 @@ void FShaderJobCache::SubmitJobs(const TArray<FShaderCommonCompileJobPtr>& InJob
 						JobLock.WriteUnlock();
 						ProcessFinishedJob(Job, EShaderCompileJobStatus::Skipped);
 					}
-
-					Job->TimeTaskSubmitJobs = FPlatformTime::Seconds() - TimeStart;
-
 				}, Prio);
 		}
 	}
 }
 
-double FShaderJobCache::ProcessFinishedJob(FShaderCommonCompileJob* FinishedJob, EShaderCompileJobStatus Status)
+void FShaderJobCache::ProcessFinishedJob(FShaderCommonCompileJob* FinishedJob, EShaderCompileJobStatus Status)
 {
-	double StallTime;
 
 	FinishedJob->UpdateStatus(Status);
 
@@ -1583,10 +1578,7 @@ double FShaderJobCache::ProcessFinishedJob(FShaderCommonCompileJob* FinishedJob,
 
 	{
 		// Need to protect writes to FShaderMapCompileResults
-		double StallStart = FPlatformTime::Seconds();
 		FScopeLock Lock(&CompileQueueSection);
-		StallTime = FPlatformTime::Seconds() - StallStart;
-
 		FShaderMapCompileResults& ShaderMapResults = *(FinishedJob->PendingShaderMap);
 		ShaderMapResults.FinishedJobs.Add(FinishedJob);
 		ShaderMapResults.bAllJobsSucceeded = ShaderMapResults.bAllJobsSucceeded && FinishedJob->bSucceeded;
@@ -1596,8 +1588,6 @@ double FShaderJobCache::ProcessFinishedJob(FShaderCommonCompileJob* FinishedJob,
 	}
 
 	InternalSubtractNumOutstandingJobs(1);
-
-	return StallTime;
 }
 
 void FShaderJobCache::AddToCacheAndProcessPending(FShaderCommonCompileJob* FinishedJob)
diff --git a/Engine/Source/Runtime/Engine/Private/ShaderCompiler/ShaderCompilerStats.cpp b/Engine/Source/Runtime/Engine/Private/ShaderCompiler/ShaderCompilerStats.cpp
--- a/Engine/Source/Runtime/Engine/Private/ShaderCompiler/ShaderCompilerStats.cpp
+++ b/Engine/Source/Runtime/Engine/Private/ShaderCompiler/ShaderCompilerStats.cpp
@@ -340,11 +340,6 @@ void FShaderCompilerStats::WriteStatSummary()
 
 	UE_LOG(LogShaderCompilers, Display, TEXT("Time at least one job was in flight (either pending or executed): %.2f s"), TotalTimeAtLeastOneJobWasInFlight);
 
-	if (Counters.AccumulatedTaskSubmitJobs > 0.0)
-	{
-		UE_LOG(LogShaderCompilers, Display, TEXT("Mutex wait stall in FShaderJobCache::SubmitJobs:  %.2f%%"), 100.0 * Counters.AccumulatedTaskSubmitJobsStall / Counters.AccumulatedTaskSubmitJobs );
-	}
-
 	// print stats about the batches
 	if (Counters.LocalJobBatchesSeen > 0 && Counters.DistributedJobBatchesSeen > 0)
 	{
@@ -937,12 +932,6 @@ void FShaderCompilerStats::RegisterFinishedJob(FShaderCommonCompileJob& Job)
 		AddToInterval(JobLifeTimeIntervals, TInterval<double>(Job.TimeAddedToPendingQueue, Job.TimeExecutionCompleted));
 	}
 	
-	if (Job.TimeTaskSubmitJobs)
-	{
-		Counters.AccumulatedTaskSubmitJobs += Job.TimeTaskSubmitJobs;
-		Counters.AccumulatedTaskSubmitJobsStall += Job.TimeTaskSubmitJobsStall;
-	}
-
 	auto RegisterStatsFromSingleJob = [this, bCompilationSkipped](const FShaderCompileJob& SingleJob)
 	{
 		// Register min/max/average shader code sizes for single job output
diff --git a/Engine/Source/Runtime/Engine/Public/ShaderCompiler.h b/Engine/Source/Runtime/Engine/Public/ShaderCompiler.h
--- a/Engine/Source/Runtime/Engine/Public/ShaderCompiler.h
+++ b/Engine/Source/Runtime/Engine/Public/ShaderCompiler.h
@@ -460,12 +460,6 @@ struct FShaderCompilerCounters
 	UPROPERTY()
 	double MaxJobLifeTime = 0;
 
-	/** Time spent in tasks generated in FShaderJobCache::SubmitJobs, plus stall time on mutex locks in those tasks */
-	UPROPERTY()
-	double AccumulatedTaskSubmitJobs = 0.0;
-	UPROPERTY()
-	double AccumulatedTaskSubmitJobsStall = 0.0;
-
 	/** Number of local job batches seen. */
 	UPROPERTY()
 	int64 LocalJobBatchesSeen = 0;
@@ -564,8 +558,6 @@ struct FShaderCompilerCounters
 		MaxJobExecutionTime = FMath::Max(Other.MaxJobExecutionTime, MaxJobExecutionTime);
 		AccumulatedJobLifeTime += Other.AccumulatedJobLifeTime;
 		MaxJobLifeTime = FMath::Max(Other.MaxJobLifeTime, MaxJobLifeTime);
-		AccumulatedTaskSubmitJobs += Other.AccumulatedTaskSubmitJobs;
-		AccumulatedTaskSubmitJobsStall += Other.AccumulatedTaskSubmitJobsStall;
 		LocalJobBatchesSeen += Other.LocalJobBatchesSeen;
 		TotalJobsReportedInLocalJobBatches += Other.TotalJobsReportedInLocalJobBatches;
 		DistributedJobBatchesSeen += Other.DistributedJobBatchesSeen;
diff --git a/Engine/Source/Runtime/RenderCore/Public/ShaderCompilerJobTypes.h b/Engine/Source/Runtime/RenderCore/Public/ShaderCompilerJobTypes.h
--- a/Engine/Source/Runtime/RenderCore/Public/ShaderCompilerJobTypes.h
+++ b/Engine/Source/Runtime/RenderCore/Public/ShaderCompilerJobTypes.h
@@ -197,9 +197,6 @@ public:
 	double TimeAssignedToExecution = 0.0;
 	/** In-engine timestamp of job being completed. Encompasses the compile time. Not set for jobs that are satisfied from the jobs cache */
 	double TimeExecutionCompleted = 0.0;
-	/** Time spent in tasks generated in FShaderJobCache::SubmitJobs, plus stall time on mutex locks in those tasks */
-	double TimeTaskSubmitJobs = 0.0;
-	double TimeTaskSubmitJobsStall = 0.0;
 
 	FShaderJobCacheRef JobCacheRef;
 
