Collision Filter Refactor Part 3.1: Updating call sites that use old APIs or directly access words that are simple.
These locations are roughly:
- Checking the flags of the old query data (sim data is left for later due to duplicate word 3)
- Things that used to do redundant unpack and pack of the old filter data
- Setting the actor id and component id

diff --git a/Engine/Plugins/PCG/Source/PCG/Private/Data/PCGCollisionWrapperData.cpp b/Engine/Plugins/PCG/Source/PCG/Private/Data/PCGCollisionWrapperData.cpp
--- a/Engine/Plugins/PCG/Source/PCG/Private/Data/PCGCollisionWrapperData.cpp
+++ b/Engine/Plugins/PCG/Source/PCG/Private/Data/PCGCollisionWrapperData.cpp
@@ -161,9 +161,9 @@ bool FPCGCollisionWrapper::GetShapeArray(FBodyInstance* Body, EPCGCollisionQuery
 			continue;
 		}
 
-		FCollisionFilterData ShapeFilter = Shape.Shape->GetQueryData();
-		const bool bShapeIsComplex = (ShapeFilter.Word3 & EPDF_ComplexCollision) != 0;
-		const bool bShapeIsSimple = (ShapeFilter.Word3 & EPDF_SimpleCollision) != 0;
+		const Chaos::Filter::FShapeFilterData& ShapeFilter = Shape.Shape->GetShapeFilterData();
+		const bool bShapeIsComplex = ShapeFilter.HasFlag(Chaos::EFilterFlags::ComplexCollision);
+		const bool bShapeIsSimple = ShapeFilter.HasFlag(Chaos::EFilterFlags::SimpleCollision);
 
 		if ((bShapeIsSimple && (QueryFlag == EPCGCollisionQueryFlag::Simple || QueryFlag == EPCGCollisionQueryFlag::SimpleFirst)) ||
 			(bShapeIsComplex && (QueryFlag == EPCGCollisionQueryFlag::Complex || QueryFlag == EPCGCollisionQueryFlag::ComplexFirst)))
diff --git a/Engine/Source/Programs/HeadlessChaos/Private/GeometryCollection/GeometryCollectionTestFramework.cpp b/Engine/Source/Programs/HeadlessChaos/Private/GeometryCollection/GeometryCollectionTestFramework.cpp
--- a/Engine/Source/Programs/HeadlessChaos/Private/GeometryCollection/GeometryCollectionTestFramework.cpp
+++ b/Engine/Source/Programs/HeadlessChaos/Private/GeometryCollection/GeometryCollectionTestFramework.cpp
@@ -209,6 +209,7 @@ namespace GeometryCollectionTest
 		// Enable all collisions
 		SimFilterData.Word1 = 0xFFFF; // this body channel
 		SimFilterData.Word3 = 0xFFFF; // collision candidate channels (also enables EPDF_ContactNotify)
+		QueryFilterData.Word3 = 0xFFFF;
 
 		FGeometryCollectionPhysicsProxy* PhysObject =
 			new FGeometryCollectionPhysicsProxy(
diff --git a/Engine/Source/Runtime/Engine/Private/Collision/SceneQuery.cpp b/Engine/Source/Runtime/Engine/Private/Collision/SceneQuery.cpp
--- a/Engine/Source/Runtime/Engine/Private/Collision/SceneQuery.cpp
+++ b/Engine/Source/Runtime/Engine/Private/Collision/SceneQuery.cpp
@@ -705,7 +705,8 @@ bool TSceneCastCommonImpWithRetryRequest(const UWorld* World, typename Traits::O
 								// is the best hit! But there could be other things we could've hit instead if we ignored the fact that we hit a cluster union.
 								OutRetryParams = Params;
 								// Ignore the actor that was hit (Check the shape data to be sure)
-								uint32 ActorIDFromShape = Traits::GetHits(HitBufferSync)->Shape->GetQueryData().Word0;
+								const Chaos::Filter::FInstanceData& FilterInstanceData = Traits::GetHits(HitBufferSync)->Shape->GetFilterInstanceData();
+								uint32 ActorIDFromShape = FilterInstanceData.GetActorId();
 								OutRetryParams.AddIgnoredActor(ActorIDFromShape);
 								if (ActorIDFromShape != ClusterUnionActor->GetUniqueID())
 								{
diff --git a/Engine/Source/Runtime/Engine/Private/PhysicsEngine/Experimental/PhysInterface_Chaos.cpp b/Engine/Source/Runtime/Engine/Private/PhysicsEngine/Experimental/PhysInterface_Chaos.cpp
--- a/Engine/Source/Runtime/Engine/Private/PhysicsEngine/Experimental/PhysInterface_Chaos.cpp
+++ b/Engine/Source/Runtime/Engine/Private/PhysicsEngine/Experimental/PhysInterface_Chaos.cpp
@@ -951,9 +951,9 @@ bool FPhysInterface_Chaos::LineTrace_Geom(FHitResult& OutHit, const FBodyInstanc
 						}
 
 						// Filter so we trace against the right kind of collision
-						FCollisionFilterData ShapeFilter = Shape->GetQueryData();
-						const bool bShapeIsComplex = (ShapeFilter.Word3 & EPDF_ComplexCollision) != 0;
-						const bool bShapeIsSimple = (ShapeFilter.Word3 & EPDF_SimpleCollision) != 0;
+						const Chaos::Filter::FShapeFilterData& ShapeFilter = Shape->GetShapeFilterData();
+						const bool bShapeIsComplex = ShapeFilter.HasFlag(Chaos::EFilterFlags::ComplexCollision);
+						const bool bShapeIsSimple = ShapeFilter.HasFlag(Chaos::EFilterFlags::SimpleCollision);
 						if ((bTraceComplex && bShapeIsComplex) || (!bTraceComplex && bShapeIsSimple))
 						{
 
@@ -1056,9 +1056,9 @@ bool FPhysInterface_Chaos::Sweep_Geom(FHitResult& OutHit, const FBodyInstance* I
 					}
 
 					// Filter so we trace against the right kind of collision
-					FCollisionFilterData ShapeFilter = Shape->GetQueryData();
-					const bool bShapeIsComplex = (ShapeFilter.Word3 & EPDF_ComplexCollision) != 0;
-					const bool bShapeIsSimple = (ShapeFilter.Word3 & EPDF_SimpleCollision) != 0;
+					const Chaos::Filter::FShapeFilterData& ShapeFilter = Shape->GetShapeFilterData();
+					const bool bShapeIsComplex = ShapeFilter.HasFlag(Chaos::EFilterFlags::ComplexCollision);
+					const bool bShapeIsSimple = ShapeFilter.HasFlag(Chaos::EFilterFlags::SimpleCollision);
 					if ((bSweepComplex && bShapeIsComplex) || (!bSweepComplex && bShapeIsSimple))
 					{
 						//question: this is returning first result, is that valid? Keeping it the same as physx for now
@@ -1141,9 +1141,9 @@ bool Overlap_GeomInternal(const FBodyInstance* InInstance, const Chaos::FImplici
 		const Chaos::FPerShapeData* Shape = ShapeRef.Shape;
 		check(Shape);
 
-		FCollisionFilterData ShapeFilter = Shape->GetQueryData();
-		const bool bShapeIsComplex = (ShapeFilter.Word3 & EPDF_ComplexCollision) != 0;
-		const bool bShapeIsSimple = (ShapeFilter.Word3 & EPDF_SimpleCollision) != 0;
+		const Chaos::Filter::FShapeFilterData& ShapeFilter = Shape->GetShapeFilterData();
+		const bool bShapeIsComplex = ShapeFilter.HasFlag(Chaos::EFilterFlags::ComplexCollision);
+		const bool bShapeIsSimple = ShapeFilter.HasFlag(Chaos::EFilterFlags::SimpleCollision);
 		if ((bTraceComplex && bShapeIsComplex) || (!bTraceComplex && bShapeIsSimple))
 		{
 			if (TargetInstance->IsShapeBoundToBody(ShapeRef))
diff --git a/Engine/Source/Runtime/Experimental/Chaos/Private/Chaos/CollisionFilterData.cpp b/Engine/Source/Runtime/Experimental/Chaos/Private/Chaos/CollisionFilterData.cpp
--- a/Engine/Source/Runtime/Experimental/Chaos/Private/Chaos/CollisionFilterData.cpp
+++ b/Engine/Source/Runtime/Experimental/Chaos/Private/Chaos/CollisionFilterData.cpp
@@ -244,6 +244,12 @@ namespace Chaos::Filter
 		return ShapeFilterData.IsValid() || InstanceData.IsValid();
 	}
 
+	bool FCombinedShapeFilterData::IsSimValid() const
+	{
+		FCollisionFilterData SimFilter = FShapeFilterBuilder::GetLegacyShapeSimFilter(*this);
+		return SimFilter.Word0 != 0 || SimFilter.Word1 != 0 || SimFilter.Word2 != 0 || SimFilter.Word3 != 0;
+	}
+
 	FShapeFilterData FShapeFilterBuilder::BuildLegacyShapeFilter(const uint8 InChannelIndex, const uint32 InQueryOverlapChannels, const uint32 InQueryBlockChannels, const uint32 InSimBlockChannels, const uint8 InMaskFilter, const EFilterFlags InFilterFlags)
 	{
 		// (Mask Filter 6 bits) | (Channel Index 5 bits) | (Flags 21 bits)
diff --git a/Engine/Source/Runtime/Experimental/Chaos/Private/Chaos/ConvexOptimizer.cpp b/Engine/Source/Runtime/Experimental/Chaos/Private/Chaos/ConvexOptimizer.cpp
--- a/Engine/Source/Runtime/Experimental/Chaos/Private/Chaos/ConvexOptimizer.cpp
+++ b/Engine/Source/Runtime/Experimental/Chaos/Private/Chaos/ConvexOptimizer.cpp
@@ -141,8 +141,7 @@ void FConvexOptimizer::BuildConvexShapes(const FShapesArray& UnionShapes)
 	const TUniquePtr<FPerShapeData>& TemplateData = UnionShapes[0];
 	const TUniquePtr<FShapeInstance>& ShapeData = ShapesArray.Last();
 
-	ShapeData->SetQueryData(TemplateData->GetQueryData());
-	ShapeData->SetSimData(TemplateData->GetSimData());
+	ShapeData->SetCombinedShapeFilterData(TemplateData->GetCombinedShapeFilterData());
 	ShapeData->SetCollisionTraceType(TemplateData->GetCollisionTraceType());
 
 	// Only simulation enabled if all the underlying shapes could be used for physics
diff --git a/Engine/Source/Runtime/Experimental/Chaos/Private/Chaos/PBDRigidClusteringAlgo.cpp b/Engine/Source/Runtime/Experimental/Chaos/Private/Chaos/PBDRigidClusteringAlgo.cpp
--- a/Engine/Source/Runtime/Experimental/Chaos/Private/Chaos/PBDRigidClusteringAlgo.cpp
+++ b/Engine/Source/Runtime/Experimental/Chaos/Private/Chaos/PBDRigidClusteringAlgo.cpp
@@ -472,7 +472,7 @@ namespace Chaos
 	{
 		SCOPE_CYCLE_COUNTER(STAT_GCUpdateFilterData);
 
-		FCollisionFilterData SelectedSimFilter, SelectedQueryFilter;
+		Filter::FCombinedShapeFilterData SelectedShapeFilter;
 		bool bFilterValid = false;
 		for (FPBDRigidParticleHandle* ChildHandle : Children)
 		{
@@ -480,9 +480,8 @@ namespace Chaos
 			{
 				if (Shape)
 				{
-					SelectedSimFilter = Shape->GetSimData();
-					bFilterValid = SelectedSimFilter.Word0 != 0 || SelectedSimFilter.Word1 != 0 || SelectedSimFilter.Word2 != 0 || SelectedSimFilter.Word3 != 0;
-					SelectedQueryFilter = Shape->GetQueryData();
+					SelectedShapeFilter = Shape->GetCombinedShapeFilterData();
+					bFilterValid = SelectedShapeFilter.IsSimValid();
 
 					if (bFilterValid)
 					{
@@ -503,8 +502,7 @@ namespace Chaos
 			const FShapesArray& ShapesArray = ClusterParent->ShapesArray();
 			for (const TUniquePtr<FPerShapeData>& Shape : ShapesArray)
 			{
-				Shape->SetQueryData(SelectedQueryFilter);
-				Shape->SetSimData(SelectedSimFilter);
+				Shape->SetCombinedShapeFilterData(SelectedShapeFilter);
 			}
 		}
 	}
diff --git a/Engine/Source/Runtime/Experimental/Chaos/Private/Chaos/PhysicsObjectCollisionInterface.cpp b/Engine/Source/Runtime/Experimental/Chaos/Private/Chaos/PhysicsObjectCollisionInterface.cpp
--- a/Engine/Source/Runtime/Experimental/Chaos/Private/Chaos/PhysicsObjectCollisionInterface.cpp
+++ b/Engine/Source/Runtime/Experimental/Chaos/Private/Chaos/PhysicsObjectCollisionInterface.cpp
@@ -141,9 +141,9 @@ namespace Chaos
 					continue;
 				}
 
-				FCollisionFilterData ShapeFilter = A->GetQueryData();
-				const bool bShapeIsComplex = (ShapeFilter.Word3 & static_cast<uint8>(EFilterFlags::ComplexCollision)) != 0;
-				const bool bShapeIsSimple = (ShapeFilter.Word3 & static_cast<uint8>(EFilterFlags::SimpleCollision)) != 0;
+				const Filter::FShapeFilterData& FilterData = A->GetShapeFilterData();
+				const bool bShapeIsComplex = FilterData.HasFlag(Chaos::EFilterFlags::ComplexCollision);
+				const bool bShapeIsSimple = FilterData.HasFlag(Chaos::EFilterFlags::SimpleCollision);
 				const bool bShouldTrace = (bTraceComplex && bShapeIsComplex) || (!bTraceComplex && bShapeIsSimple);
 				if (!bShouldTrace)
 				{
diff --git a/Engine/Source/Runtime/Experimental/Chaos/Public/Chaos/ClusterUnionManager.h b/Engine/Source/Runtime/Experimental/Chaos/Public/Chaos/ClusterUnionManager.h
--- a/Engine/Source/Runtime/Experimental/Chaos/Public/Chaos/ClusterUnionManager.h
+++ b/Engine/Source/Runtime/Experimental/Chaos/Public/Chaos/ClusterUnionManager.h
@@ -308,16 +308,9 @@ namespace Chaos
 			}
 
 			{
-				FCollisionFilterData Data = TemplateShape->GetQueryData();
-				Data.Word0 = ActorId;
-				ShapeData->SetQueryData(Data);
-			}
-
-			{
-				FCollisionFilterData Data = TemplateShape->GetSimData();
-				Data.Word0 = 0;
-				Data.Word2 = ComponentId;
-				ShapeData->SetSimData(Data);
+				const Chaos::Filter::FInstanceData InstanceData(ActorId, ComponentId);
+				ShapeData->SetFilterInstanceData(InstanceData);
+				ShapeData->SetShapeFilterData(TemplateShape->GetShapeFilterData());
 			}
 
 			{
diff --git a/Engine/Source/Runtime/Experimental/Chaos/Public/Chaos/CollisionFilterData.h b/Engine/Source/Runtime/Experimental/Chaos/Public/Chaos/CollisionFilterData.h
--- a/Engine/Source/Runtime/Experimental/Chaos/Public/Chaos/CollisionFilterData.h
+++ b/Engine/Source/Runtime/Experimental/Chaos/Public/Chaos/CollisionFilterData.h
@@ -198,6 +198,7 @@ namespace Chaos::Filter
 		CHAOS_API void SetShapeFilterData(const FShapeFilterData& InData);
 
 		CHAOS_API bool IsValid() const;
+		UE_INTERNAL CHAOS_API bool IsSimValid() const;
 
 	private:
 		friend class FShapeFilterBuilder;
diff --git a/Engine/Source/Runtime/Experimental/Chaos/Public/Chaos/ParticleHandle.h b/Engine/Source/Runtime/Experimental/Chaos/Public/Chaos/ParticleHandle.h
--- a/Engine/Source/Runtime/Experimental/Chaos/Public/Chaos/ParticleHandle.h
+++ b/Engine/Source/Runtime/Experimental/Chaos/Public/Chaos/ParticleHandle.h
@@ -3262,16 +3262,9 @@ void UpdateParticleShapes(const TArray<ParticleType*>& ShapesParticles, const FI
 				}
     
 				{
-					FCollisionFilterData Data = TemplateShape->GetQueryData();
-					Data.Word0 = ActorId;
-					ShapeData->SetQueryData(Data);
-				}
-    
-				{
-					FCollisionFilterData Data = TemplateShape->GetSimData();
-					Data.Word0 = 0;
-					Data.Word2 = ComponentID;
-					ShapeData->SetSimData(Data);
+					const Chaos::Filter::FInstanceData InstanceData(ActorId, ComponentID);
+					ShapeData->SetFilterInstanceData(InstanceData);
+					ShapeData->SetShapeFilterData(TemplateShape->GetShapeFilterData());
 				}
     
 				ShapeData->SetSimEnabled(TemplateShape->GetSimEnabled());
diff --git a/Engine/Source/Runtime/Experimental/Chaos/Public/Chaos/PhysicsObjectCollisionInterface.h b/Engine/Source/Runtime/Experimental/Chaos/Public/Chaos/PhysicsObjectCollisionInterface.h
--- a/Engine/Source/Runtime/Experimental/Chaos/Public/Chaos/PhysicsObjectCollisionInterface.h
+++ b/Engine/Source/Runtime/Experimental/Chaos/Public/Chaos/PhysicsObjectCollisionInterface.h
@@ -87,9 +87,9 @@ namespace Chaos
 					{
 						check(Shape);
 
-						FCollisionFilterData ShapeFilter = Shape->GetQueryData();
-						const bool bShapeIsComplex = (ShapeFilter.Word3 & static_cast<uint8>(EFilterFlags::ComplexCollision)) != 0;
-						const bool bShapeIsSimple = (ShapeFilter.Word3 & static_cast<uint8>(EFilterFlags::SimpleCollision)) != 0;
+						const Filter::FShapeFilterData& FilterData = Shape->GetShapeFilterData();
+						const bool bShapeIsComplex = FilterData.HasFlag(Chaos::EFilterFlags::ComplexCollision);
+						const bool bShapeIsSimple = FilterData.HasFlag(Chaos::EFilterFlags::SimpleCollision);
 						if ((bTraceComplex && bShapeIsComplex) || (!bTraceComplex && bShapeIsSimple))
 						{
 							FReal Distance;
@@ -205,9 +205,9 @@ namespace Chaos
 					{
 						check(Shape);
 
-						FCollisionFilterData ShapeFilter = Shape->GetQueryData();
-						const bool bShapeIsComplex = (ShapeFilter.Word3 & static_cast<uint8>(EFilterFlags::ComplexCollision)) != 0;
-						const bool bShapeIsSimple = (ShapeFilter.Word3 & static_cast<uint8>(EFilterFlags::SimpleCollision)) != 0;
+						const Filter::FShapeFilterData& FilterData = Shape->GetShapeFilterData();
+						const bool bShapeIsComplex = FilterData.HasFlag(Chaos::EFilterFlags::ComplexCollision);
+						const bool bShapeIsSimple = FilterData.HasFlag(Chaos::EFilterFlags::SimpleCollision);
 						if ((Params.bSweepComplex && bShapeIsComplex) || (!Params.bSweepComplex && bShapeIsSimple))
 						{
 							FVec3 WorldPosition;
