fix up now-incorrect comment caught by bughawk

diff --git a/Engine/Source/Runtime/Engine/Private/ShaderCompiler/ShaderCompilerJobCache.cpp b/Engine/Source/Runtime/Engine/Private/ShaderCompiler/ShaderCompilerJobCache.cpp
--- a/Engine/Source/Runtime/Engine/Private/ShaderCompiler/ShaderCompilerJobCache.cpp
+++ b/Engine/Source/Runtime/Engine/Private/ShaderCompiler/ShaderCompilerJobCache.cpp
@@ -616,7 +616,7 @@ public:
 	void SubmitJob(FShaderCommonCompileJob* Job);
 	void SubmitJobs(const TArray<FShaderCommonCompileJobPtr>& InJobs);
 
-	/** This is an entry point for all jobs that have finished the compilation (whether real or cached). Can be called from multiple threads. Returns mutex stall time. */
+	/** This is an entry point for all jobs that have finished the compilation (whether real or cached). Can be called from multiple threads. */
 	void ProcessFinishedJob(FShaderCommonCompileJob* FinishedJob, EShaderCompileJobStatus Status);
 
 	/** Adds the job to cache. */
