Added 'OBJ SERVERIDLIST' console command that dumps all existing objects with the specified FRemoteServerId in remote object enabled builds. Can save the output to a predefined file inside of the project's Saved folder.

Parameters:
-ID=ServerId (A number of one of the predefined ERemoteServerIdConstants, defaults to Local)
-SAVETOFILE writes the list to file instead of the log

diff --git a/Engine/Source/Runtime/Engine/Private/UnrealEngine.cpp b/Engine/Source/Runtime/Engine/Private/UnrealEngine.cpp
--- a/Engine/Source/Runtime/Engine/Private/UnrealEngine.cpp
+++ b/Engine/Source/Runtime/Engine/Private/UnrealEngine.cpp
@@ -10168,6 +10168,56 @@ bool UEngine::HandleObjCommand( const TCHAR* Cmd, FOutputDevice& Ar )
 		}
 		return true;
 	}
+	else if (FParse::Command(&Cmd, TEXT("SERVERIDLIST")))
+	{
+		// Lists all existing objects with the specified server id (e.g. -ID=Asset). Can write the output to a predefined file inside of the project's Saved folder (-SAVETOFILE).
+		FString ServerIdText;
+		FRemoteServerId ServerId(ERemoteServerIdConstants::Local);
+		if (FParse::Value(Cmd, TEXT("ID="), ServerIdText))
+		{
+			ServerId = FRemoteServerId::FromString(ServerIdText);
+		}
+
+		TArray<FRemoteObjectId> ObjectsIdList;
+		if (ServerId.IsValid())
+		{			
+			for (TObjectIterator<UObject> It; It; ++It)
+			{
+				FRemoteObjectId ObjectId(*It);
+				if (ObjectId.GetServerId() == ServerId)
+				{
+					ObjectsIdList.Add(ObjectId);
+				}
+			}
+			ObjectsIdList.Sort();
+		}
+		else
+		{
+			Ar.Logf(TEXT("Unable to parse server id '%s'"), *ServerIdText);
+			return true;
+		}
+
+		if (FParse::Param(Cmd, TEXT("SAVETOFILE")))
+		{
+			uint32 LocalServerIdNumber = FRemoteServerId::GetLocalServerId().GetIdNumber();;
+			const FString Filename = FPaths::Combine(FPaths::ProjectSavedDir(), FString::Printf(TEXT("ServerIdList-%s_S%u.txt"), *ServerId.ToString(), LocalServerIdNumber));
+			TArray<FString> Lines;
+
+			for (FRemoteObjectId ObjectId : ObjectsIdList)
+			{
+				Lines.Add(ObjectId.ToString(ERemoteIdToStringVerbosity::FullName));
+			}
+			FFileHelper::SaveStringArrayToFile(Lines, *Filename);
+		}
+		else
+		{
+			for (FRemoteObjectId ObjectId : ObjectsIdList)
+			{
+				Ar.Log(*ObjectId.ToString(ERemoteIdToStringVerbosity::FullName));
+			}
+		}
+		return true;
+	}
 #endif // UE_WITH_REMOTE_OBJECT_HANDLE
 #endif // !UE_BUILD_SHIPPING
 	else
