Fixing issue where attaching an actor to the socket of another skinned me actor would not transact correctly in Multi User or as part of an undo operation.

diff --git a/Engine/Source/Runtime/Engine/Private/Components/SkinnedMeshComponent.cpp b/Engine/Source/Runtime/Engine/Private/Components/SkinnedMeshComponent.cpp
--- a/Engine/Source/Runtime/Engine/Private/Components/SkinnedMeshComponent.cpp
+++ b/Engine/Source/Runtime/Engine/Private/Components/SkinnedMeshComponent.cpp
@@ -1526,6 +1526,9 @@ void USkinnedMeshComponent::PostEditUndo()
 	NotifyIfSkinnedAssetChanged();
 
 	Super::PostEditUndo();
+
+	// Post undo, make sure that any socket attachments are up to date.
+	UpdateCachedSocketAttachmentFlag();
 }
 
 void USkinnedMeshComponent::PostDuplicate(bool bDuplicateForPIE)
