TOverrideVoidReturnInvoker: fix return value always of type bool instead of ReturnType.

diff --git a/Engine/Source/Runtime/Engine/Public/OverrideVoidReturnInvoker.h b/Engine/Source/Runtime/Engine/Public/OverrideVoidReturnInvoker.h
--- a/Engine/Source/Runtime/Engine/Public/OverrideVoidReturnInvoker.h
+++ b/Engine/Source/Runtime/Engine/Public/OverrideVoidReturnInvoker.h
@@ -20,7 +20,7 @@ struct TOverrideVoidReturnInvoker
 	{}
 
 	template <typename... ArgTypes>
-	bool operator()(ArgTypes&&... Args)
+	ReturnType operator()(ArgTypes&&... Args)
 	{
 		using FuncResultType = std::invoke_result_t<FuncType, ArgTypes...>;
 		static_assert(std::disjunction_v<std::is_same<FuncResultType, void>, std::is_same<FuncResultType, ReturnType>>, "Passed in function should either return void or ReturnType");
