Scene Graph: Fixed up a few mesh implementation details

Ensure that ShouldCreateRenderState is actually called for derived mesh implementations
Added PURE_VIRTUAL to a few methods that are required in child classes (including ShouldCreateRenderState)
Made sure we invalidate cached material info for skeletal meshes when we undo/redo

diff --git a/Engine/Source/Runtime/Engine/Private/SkeletalMesh.cpp b/Engine/Source/Runtime/Engine/Private/SkeletalMesh.cpp
--- a/Engine/Source/Runtime/Engine/Private/SkeletalMesh.cpp
+++ b/Engine/Source/Runtime/Engine/Private/SkeletalMesh.cpp
@@ -1638,7 +1638,9 @@ void USkeletalMesh::PostEditUndo()
 		// A morph target remove has been undone, reinitialise
 		InitMorphTargets();
 	}
-	
+
+	UpdateSkeletalMaterialsInfoCache(GetMaterials());
+
 	bTransacting = false;
 }
 
