SkinnedMesh: Add missing engine impl for ISkinnedMeshComponent and fixed up FObjectCacheEventSink notification handling

Prevents crash loading into CitySample using -dpcvars=Editor.ObjectReverseLookupMode=2

diff --git a/Engine/Source/Runtime/Engine/Private/Components/SkinnedMeshComponent.cpp b/Engine/Source/Runtime/Engine/Private/Components/SkinnedMeshComponent.cpp
--- a/Engine/Source/Runtime/Engine/Private/Components/SkinnedMeshComponent.cpp
+++ b/Engine/Source/Runtime/Engine/Private/Components/SkinnedMeshComponent.cpp
@@ -659,8 +659,6 @@ void USkinnedMeshComponent::Serialize(FArchive& Ar)
 
 void USkinnedMeshComponent::PostLoad()
 {
-	NotifyIfSkinnedAssetChanged();
-
 	Super::PostLoad();
 
 	PRAGMA_DISABLE_DEPRECATION_WARNINGS
@@ -670,6 +668,8 @@ void USkinnedMeshComponent::PostLoad()
 	}
 	PRAGMA_ENABLE_DEPRECATION_WARNINGS
 
+	NotifyIfSkinnedAssetChanged();
+
 	if (GetLinkerCustomVersion(FUE5MainStreamObjectVersion::GUID) < FUE5MainStreamObjectVersion::AddSetMeshDeformerFlag)
 	{
 		// Set default for new bSetMeshDeformer flag.
@@ -2619,14 +2619,15 @@ void USkinnedMeshComponent::SetSkinnedAsset(class USkinnedAsset* InSkinnedAsset)
 USkinnedAsset* USkinnedMeshComponent::GetSkinnedAsset() const
 {
 	PRAGMA_DISABLE_DEPRECATION_WARNINGS
+	USkinnedAsset* CurrentSkinnedAsset = !SkeletalMesh ? SkinnedAsset : SkeletalMesh;	// Return the SkeletalMesh pointer for backward compatibility until it is retired
 #if WITH_EDITOR
 	// This should never happen and is a last resort, we should have caught the property overwrite well before we reach this code
-	if (KnownSkinnedAsset != SkinnedAsset)
+	if (KnownSkinnedAsset != CurrentSkinnedAsset)
 	{
 		OutdatedKnownSkinnedAssetDetected();
 	}
 #endif
-	return !SkeletalMesh ? SkinnedAsset : SkeletalMesh;  // Return the SkeletalMesh pointer for backward compatibility until it is retired
+	return CurrentSkinnedAsset;
 	PRAGMA_ENABLE_DEPRECATION_WARNINGS
 }
 
@@ -5709,9 +5710,10 @@ void USkinnedMeshComponent::NotifyIfSkinnedAssetChanged()
 {
 #if WITH_EDITOR
 	PRAGMA_DISABLE_DEPRECATION_WARNINGS
-	if (KnownSkinnedAsset != SkinnedAsset)
+	USkinnedAsset* CurrentSkinnedAsset = !SkeletalMesh ? SkinnedAsset : SkeletalMesh;	// Return the SkeletalMesh pointer for backward compatibility until it is retired
+	if (KnownSkinnedAsset != CurrentSkinnedAsset)
 	{
-		KnownSkinnedAsset = SkinnedAsset;
+		KnownSkinnedAsset = CurrentSkinnedAsset;
 		FObjectCacheEventSink::NotifySkinnedMeshChanged_Concurrent(GetSkinnedMeshComponentInterface());
 	}
 	PRAGMA_ENABLE_DEPRECATION_WARNINGS
@@ -5722,13 +5724,14 @@ void USkinnedMeshComponent::NotifyIfSkinnedAssetChanged()
 void USkinnedMeshComponent::OutdatedKnownSkinnedAssetDetected() const
 {
 	PRAGMA_DISABLE_DEPRECATION_WARNINGS
+	USkinnedAsset* CurrentSkinnedAsset = !SkeletalMesh ? SkinnedAsset : SkeletalMesh;	// Return the SkeletalMesh pointer for backward compatibility until it is retired
 	ensureMsgf(
-		KnownSkinnedAsset == SkinnedAsset, 
+		KnownSkinnedAsset == CurrentSkinnedAsset,
 		TEXT("SkinnedAsset property overwritten for component %s without a call to NotifyIfSkinnedAssetChanged(). KnownSkinnedAsset (%p) != SkinnedAsset (%p - %s)"),
 		*GetFullName(),
 		KnownSkinnedAsset,
-		SkinnedAsset.Get(),
-		SkinnedAsset ? *SkinnedAsset->GetFullName() : TEXT("nullptr")
+		CurrentSkinnedAsset,
+		CurrentSkinnedAsset ? *CurrentSkinnedAsset->GetFullName() : TEXT("nullptr")
 		);
 	PRAGMA_ENABLE_DEPRECATION_WARNINGS
 
diff --git a/Engine/Source/Runtime/Engine/Private/UnrealEngine.cpp b/Engine/Source/Runtime/Engine/Private/UnrealEngine.cpp
--- a/Engine/Source/Runtime/Engine/Private/UnrealEngine.cpp
+++ b/Engine/Source/Runtime/Engine/Private/UnrealEngine.cpp
@@ -341,6 +341,7 @@ void FEngineModule::StartupModule()
 
 	IPrimitiveComponent::AddImplementer({UPrimitiveComponent::StaticClass(),  [](UObject* Obj){return Cast<UPrimitiveComponent>(Obj)->GetPrimitiveComponentInterface();}});
 	IStaticMeshComponent::AddImplementer({UStaticMeshComponent::StaticClass(), [](UObject* Obj){return Cast<UStaticMeshComponent>(Obj)->GetStaticMeshComponentInterface();}});
+	ISkinnedMeshComponent::AddImplementer({USkinnedMeshComponent::StaticClass(), [](UObject* Obj) {return Cast<USkinnedMeshComponent>(Obj)->GetSkinnedMeshComponentInterface();}});
 
 #if OBJECT_TRACE_ENABLED
 	FObjectTrace::Init();
