Downgraded AnimBank check to an ensure to avoid crash.

diff --git a/Engine/Source/Runtime/Engine/Classes/Animation/AnimBank.h b/Engine/Source/Runtime/Engine/Classes/Animation/AnimBank.h
--- a/Engine/Source/Runtime/Engine/Classes/Animation/AnimBank.h
+++ b/Engine/Source/Runtime/Engine/Classes/Animation/AnimBank.h
@@ -208,7 +208,7 @@ public:
 	inline const FAnimBankData& GetData() const
 	{
 #if WITH_EDITOR
-		check(!IsCompiling());
+		ensure(!IsCompiling());
 #endif
 		return Data;
 	}
