Substrate - simple renaming.

#rnx

diff --git a/Engine/Source/Runtime/Engine/Classes/Engine/RendererSettings.h b/Engine/Source/Runtime/Engine/Classes/Engine/RendererSettings.h
--- a/Engine/Source/Runtime/Engine/Classes/Engine/RendererSettings.h
+++ b/Engine/Source/Runtime/Engine/Classes/Engine/RendererSettings.h
@@ -1553,7 +1553,7 @@ private:
 #endif // WITH_EDITOR
 
 	void SanatizeReflectionCaptureResolution();
-	void SanatizeSubstrateSettings();
+	void SanitizeSubstrateSettings();
 	void UpdateWorkingColorSpaceAndChromaticities();
 };
 
diff --git a/Engine/Source/Runtime/Engine/Private/RendererSettings.cpp b/Engine/Source/Runtime/Engine/Private/RendererSettings.cpp
--- a/Engine/Source/Runtime/Engine/Private/RendererSettings.cpp
+++ b/Engine/Source/Runtime/Engine/Private/RendererSettings.cpp
@@ -121,7 +121,7 @@ void URendererSettings::PostInitProperties()
 
 	UpdateWorkingColorSpaceAndChromaticities();
 
-	SanatizeSubstrateSettings();
+	SanitizeSubstrateSettings();
 
 #if WITH_EDITOR
 	CheckForMissingShaderModels();
@@ -443,7 +443,7 @@ void URendererSettings::PostEditChangeProperty(FPropertyChangedEvent& PropertyCh
 		if (PropertyChangedEvent.Property->GetFName() == GET_MEMBER_NAME_CHECKED(URendererSettings, SubstrateGBufferFormat) ||
 			PropertyChangedEvent.Property->GetFName() == GET_MEMBER_NAME_CHECKED(URendererSettings, SubstrateProjectClosuresPerPixel))
 		{
-			SanatizeSubstrateSettings();
+			SanitizeSubstrateSettings();
 		}
 	}
 }
@@ -731,7 +731,7 @@ void URendererSettings::SanatizeReflectionCaptureResolution()
 #endif // WITH_EDITOR
 }
 
-void URendererSettings::SanatizeSubstrateSettings()
+void URendererSettings::SanitizeSubstrateSettings()
 {
 	// When blendable GBuffer is used, force project closure count to 1 for clarity
 	if (SubstrateGBufferFormat == ESubstrateStorageFormat::Type::BlendableGBuffer)
