World Partition Runtime Spatial Hash: handle the problematic case where the world bounds is valid but empty.

diff --git a/Engine/Source/Runtime/Engine/Private/WorldPartition/RuntimeSpatialHash/RuntimeSpatialHashGridHelper.cpp b/Engine/Source/Runtime/Engine/Private/WorldPartition/RuntimeSpatialHash/RuntimeSpatialHashGridHelper.cpp
--- a/Engine/Source/Runtime/Engine/Private/WorldPartition/RuntimeSpatialHash/RuntimeSpatialHashGridHelper.cpp
+++ b/Engine/Source/Runtime/Engine/Private/WorldPartition/RuntimeSpatialHash/RuntimeSpatialHashGridHelper.cpp
@@ -48,7 +48,7 @@ FSquare2DGridHelper::FSquare2DGridHelper(const FBox& InWorldBounds, const FVecto
 	int64 GridSize = 1;
 	int32 GridLevelCount = 1;
 
-	if (WorldBounds.IsValid)
+	if (WorldBounds.IsValid && CellSize)
 	{
 		const FVector2D DistMin = FVector2D(WorldBounds.Min - Origin).GetAbs();
 		const FVector2D DistMax = FVector2D(WorldBounds.Max - Origin).GetAbs();
@@ -66,8 +66,11 @@ FSquare2DGridHelper::FSquare2DGridHelper(const FBox& InWorldBounds, const FVecto
 		}
 	}
 
-	check(GridSize);
-	check(FMath::IsPowerOfTwo(GridSize));
+	if (!GridSize || !FMath::IsPowerOfTwo(GridSize))
+	{
+		GridSize = 1;
+		GridLevelCount = 1;
+	}
 
 	Levels.Reserve(GridLevelCount);
 	int64 CurrentCellSize = CellSize;
