Fixed crash due to Nanite ISKM erroneously using the BoneMapIndirect storage mode.

diff --git a/Engine/Source/Runtime/Engine/Private/InstancedSkeletalMeshObject.cpp b/Engine/Source/Runtime/Engine/Private/InstancedSkeletalMeshObject.cpp
--- a/Engine/Source/Runtime/Engine/Private/InstancedSkeletalMeshObject.cpp
+++ b/Engine/Source/Runtime/Engine/Private/InstancedSkeletalMeshObject.cpp
@@ -25,7 +25,7 @@ FInstancedSkeletalMeshObject::FInstancedSkeletalMeshObject(
 	}
 
 	BoneTransformAllocationMode = EBoneTransformAllocationMode::Scene;
-	BoneTransformStorageMode    = ::GetBoneTransformStorageMode();
+	BoneTransformStorageMode    = Type == EType::GPUSkin ? ::GetBoneTransformStorageMode() : EBoneTransformStorageMode::Direct;
 
 	InitResources(InMeshDesc);
 
