#robmerge -Main
HACKS. REMOVE POST MILESTONE

hacks for:
1. Input causing resim (currently every input, but should only trigger on diff)
2. ConstantLinear renderinterp
3. ConstantLinear input decay

diff --git a/Engine/Source/Runtime/Engine/Private/PhysicsEngine/NetworkPhysicsComponent.cpp b/Engine/Source/Runtime/Engine/Private/PhysicsEngine/NetworkPhysicsComponent.cpp
--- a/Engine/Source/Runtime/Engine/Private/PhysicsEngine/NetworkPhysicsComponent.cpp
+++ b/Engine/Source/Runtime/Engine/Private/PhysicsEngine/NetworkPhysicsComponent.cpp
@@ -42,6 +42,16 @@ namespace PhysicsReplicationCVars
 		float DynamicInputScalingDecreaseTimeInterval = 10.0f;
 		static FAutoConsoleVariableRef CVarDynamicInputScalingDecreaseTimeInterval(TEXT("np2.Resim.DynamicInputScaling.DecreaseTimeInterval"), DynamicInputScalingDecreaseTimeInterval, TEXT("Default 10.0 (value in seconds). How often dynamic scaling can decrease the number of inputs to send."));
 
+		bool bForceTriggerResimOnInputRecieve = false;
+		static FAutoConsoleVariableRef CVarbForceTriggerResimOnInputRecieve(TEXT("np2.Resim.bForceTriggerResimOnInputRecieve"), bForceTriggerResimOnInputRecieve, TEXT(""));
+		bool bForceEnableApplyInputAtRuntime = false;
+		static FAutoConsoleVariableRef CVarbForceEnableApplyInputAtRuntime(TEXT("np2.Resim.bForceEnableApplyInputAtRuntime"), bForceEnableApplyInputAtRuntime, TEXT(""));
+		bool bInputDiffFromTickOffset = false;
+		static FAutoConsoleVariableRef CVarbInputDiffFromTickOffset(TEXT("np2.Resim.bInputDiffFromTickOffset"), bInputDiffFromTickOffset, TEXT(""));
+		bool bUseLinearDecayMecanicForAppliedInputsAtRuntime = true;
+		static FAutoConsoleVariableRef CVarbUseLinearDecayMecanicForAppliedInputsAtRuntime(TEXT("np2.Resim.bUseLinearDecayMecanicForAppliedInputsAtRuntime"), bUseLinearDecayMecanicForAppliedInputsAtRuntime, TEXT(""));
+		
+		
 		bool bAllowRewindToClosestState = true;
 		static FAutoConsoleVariableRef CVarResimAllowRewindToClosestState(TEXT("np2.Resim.AllowRewindToClosestState"), bAllowRewindToClosestState, TEXT("When rewinding to a specific frame, if the client doens't have state data for that frame, use closest data available. Only affects the first rewind frame, when FPBDRigidsEvolution is set to Reset."));
 		bool bCompareStateToTriggerRewind = false;
@@ -1936,6 +1946,17 @@ void FAsyncNetworkPhysicsComponent::ConsumeAsyncInput(const int32 PhysicsStep)
 				}
 
 				ReceiveHelper(InputHistory.Get(), AsyncInput->InputData.Get(), /*bImportant*/false, bCompareInput);
+			
+				if (PhysicsReplicationCVars::ResimulationCVars::bForceTriggerResimOnInputRecieve && !IsLocallyControlled())
+				{
+					const int32 LatestFrame = AsyncInput->InputData.Get()->GetLatestFrame();
+					int32 TargetFrame = LatestFrame;
+					if (PhysicsReplicationCVars::ResimulationCVars::bInputDiffFromTickOffset)
+					{
+						TargetFrame -= NetworkPhysicsTickOffset;
+					}
+					TriggerResimulation( TargetFrame );
+				}
 			}
 
 			// Receive States
@@ -1987,6 +2008,7 @@ FAsyncNetworkPhysicsComponentOutput& FAsyncNetworkPhysicsComponent::GetAsyncOutp
 	return AsyncOutput;
 }
 
+
 void FAsyncNetworkPhysicsComponent::OnPreProcessInputs_Internal(const int32 PhysicsStep)
 {
 	ConsumeAsyncInput(PhysicsStep);
@@ -2050,7 +2072,7 @@ void FAsyncNetworkPhysicsComponent::OnPreProcessInputs_Internal(const int32 Phys
 		}
 
 		// Apply replicated inputs on server and simulated proxies if setting is enabled, and on local player if we are resimulating
-		const bool bApplySimProxyInput = ComponentSettings.GetApplySimProxyInputAtRuntime() && !bIsServer && !IsLocallyControlled();
+		const bool bApplySimProxyInput = PhysicsReplicationCVars::ResimulationCVars::bForceEnableApplyInputAtRuntime && !bIsServer && !IsLocallyControlled();
 		const bool bApplyServerInput = bIsServer && !IsLocallyControlled();
 		if ((bApplyServerInput || bApplySimProxyInput || bIsSolverResim) && InputHistory && InputData)
 		{
@@ -2506,6 +2528,11 @@ const float FAsyncNetworkPhysicsComponent::GetCurrentInputDecay(const FNetworkPh
 	const float NumPredictedInputs = RewindData->CurrentFrame() - PhysicsData->LocalFrame; // Number of frames we have used the same PhysicsData for during resim
 	float MaxPredictedInputs = 0;
 
+	if (PhysicsReplicationCVars::ResimulationCVars::bUseLinearDecayMecanicForAppliedInputsAtRuntime)
+	{
+		return 1.0 - FMath::Min(NumPredictedInputs / 12.0,1.0);
+	}
+
 	if (ComponentSettings.GetApplyInputDecayOverSetTime())
 	{
 		// If the decay curve should be applied over a set amount of time, calculate how many frames that time translates to
@@ -2516,7 +2543,7 @@ const float FAsyncNetworkPhysicsComponent::GetCurrentInputDecay(const FNetworkPh
 		// Number of frames the input will predict and decay over until resimulation has completed
 		MaxPredictedInputs = RewindData->GetLatestFrame() - 1 - PhysicsData->LocalFrame;
 	}
-	
+		
 	// Linear decay
 	const float PredictionAlpha = MaxPredictedInputs > 0 ? FMath::Clamp(NumPredictedInputs / MaxPredictedInputs, 0.0f, 1.0f) : 0.0f;
 
diff --git a/Engine/Source/Runtime/Engine/Private/PhysicsEngine/PhysicsReplication.cpp b/Engine/Source/Runtime/Engine/Private/PhysicsEngine/PhysicsReplication.cpp
--- a/Engine/Source/Runtime/Engine/Private/PhysicsEngine/PhysicsReplication.cpp
+++ b/Engine/Source/Runtime/Engine/Private/PhysicsEngine/PhysicsReplication.cpp
@@ -120,6 +120,15 @@ namespace PhysicsReplicationCVars
 	{
 		extern bool bApplyPredictiveInterpolationWhenBehindServer;
 
+		float ToleranceX = 0.f;
+		static FAutoConsoleVariableRef CVarResimToleranceX(TEXT("np2.Resim.Tolerance.X"), ToleranceX, TEXT(""));
+		float ToleranceR = 0.f;
+		static FAutoConsoleVariableRef CVarResimToleranceR(TEXT("np2.Resim.Tolerance.R"), ToleranceR, TEXT(""));
+		float ToleranceV = 0.f;
+		static FAutoConsoleVariableRef CVarResimToleranceV(TEXT("np2.Resim.Tolerance.V"), ToleranceR, TEXT(""));
+		float ToleranceW = 0.f;
+		static FAutoConsoleVariableRef CVarResimToleranceW(TEXT("np2.Resim.Tolerance.W"), ToleranceR, TEXT(""));
+		
 		bool bRuntimeCorrectionEnabled = false;
 		static FAutoConsoleVariableRef CVarResimRuntimeCorrectionEnabled(TEXT("np2.Resim.RuntimeCorrectionEnabled"), bRuntimeCorrectionEnabled, TEXT("Apply positional and rotational runtime corrections while within resim trigger distance."));
 
@@ -2219,10 +2228,18 @@ bool FPhysicsReplicationAsync::ResimulationReplication(Chaos::FPBDRigidParticleH
 	const Chaos::FGeometryParticleState PastState = RewindData->GetPastStateAtFrame(*Handle, LocalFrame, RewindPhase);
 
 	// Check which comparisons to perform to trigger resimulation from
-	const bool bCompareX = Chaos::FPhysicsSolverBase::GetResimulationErrorPositionThresholdEnabled() || SettingsCurrent.ResimulationSettings.bOverrideResimulationErrorPositionThreshold;
-	const bool bCompareR = Chaos::FPhysicsSolverBase::GetResimulationErrorRotationThresholdEnabled() || SettingsCurrent.ResimulationSettings.bOverrideResimulationErrorRotationThreshold;
-	const bool bCompareV = Chaos::FPhysicsSolverBase::GetResimulationErrorLinearVelocityThresholdEnabled() || SettingsCurrent.ResimulationSettings.bOverrideResimulationErrorLinearVelocityThreshold;
-	const bool bCompareW = Chaos::FPhysicsSolverBase::GetResimulationErrorAngularVelocityThresholdEnabled() || SettingsCurrent.ResimulationSettings.bOverrideResimulationErrorAngularVelocityThreshold;
+	const bool bCompareX = PhysicsReplicationCVars::ResimulationCVars::ToleranceX > 0.f
+							|| Chaos::FPhysicsSolverBase::GetResimulationErrorPositionThresholdEnabled()
+							|| SettingsCurrent.ResimulationSettings.bOverrideResimulationErrorPositionThreshold;
+	const bool bCompareR = PhysicsReplicationCVars::ResimulationCVars::ToleranceR > 0.f
+							|| Chaos::FPhysicsSolverBase::GetResimulationErrorRotationThresholdEnabled()
+							|| SettingsCurrent.ResimulationSettings.bOverrideResimulationErrorRotationThreshold;
+	const bool bCompareV = PhysicsReplicationCVars::ResimulationCVars::ToleranceV > 0.f
+							|| Chaos::FPhysicsSolverBase::GetResimulationErrorLinearVelocityThresholdEnabled()
+							|| SettingsCurrent.ResimulationSettings.bOverrideResimulationErrorLinearVelocityThreshold;
+	const bool bCompareW = PhysicsReplicationCVars::ResimulationCVars::ToleranceW > 0.f
+							|| Chaos::FPhysicsSolverBase::GetResimulationErrorAngularVelocityThresholdEnabled() \
+							|| SettingsCurrent.ResimulationSettings.bOverrideResimulationErrorAngularVelocityThreshold;
 	bool bShouldTriggerResim = false;
 
 #if !(UE_BUILD_SHIPPING || UE_BUILD_TEST)
@@ -2235,7 +2252,11 @@ bool FPhysicsReplicationAsync::ResimulationReplication(Chaos::FPBDRigidParticleH
 	// Check for positional discrepancy in Distance between client and server
 	if (bCompareX)
 	{
-		const float ResimPositionErrorThreshold = SettingsCurrent.ResimulationSettings.GetResimulationErrorPositionThreshold(Chaos::FPhysicsSolverBase::GetResimulationErrorPositionThreshold());
+		float ResimPositionErrorThreshold = PhysicsReplicationCVars::ResimulationCVars::ToleranceX;
+		if (FMath::IsNearlyZero(ResimPositionErrorThreshold))
+		{
+			ResimPositionErrorThreshold = SettingsCurrent.ResimulationSettings.GetResimulationErrorPositionThreshold(Chaos::FPhysicsSolverBase::GetResimulationErrorPositionThreshold());
+		}
 		bShouldTriggerResim = Chaos::FRewindData::CheckVectorThreshold(Target.TargetState.Position, PastState.GetX(), ResimPositionErrorThreshold);
 #if !(UE_BUILD_SHIPPING || UE_BUILD_TEST)
 		if (bShouldTriggerResim)
@@ -2248,7 +2269,12 @@ bool FPhysicsReplicationAsync::ResimulationReplication(Chaos::FPBDRigidParticleH
 	// Check for linear velocity discrepancy in Distance / s between client and server
 	if (!bShouldTriggerResim && bCompareV)
 	{
-		const float ResimLinVelocityErrorThreshold = SettingsCurrent.ResimulationSettings.GetResimulationErrorLinearVelocityThreshold(Chaos::FPhysicsSolverBase::GetResimulationErrorLinearVelocityThreshold());
+		float ResimLinVelocityErrorThreshold = PhysicsReplicationCVars::ResimulationCVars::ToleranceV;
+		if (FMath::IsNearlyZero(ResimLinVelocityErrorThreshold))
+		{
+			ResimLinVelocityErrorThreshold = SettingsCurrent.ResimulationSettings.GetResimulationErrorLinearVelocityThreshold(Chaos::FPhysicsSolverBase::GetResimulationErrorLinearVelocityThreshold());
+		}
+
 		bShouldTriggerResim = Chaos::FRewindData::CheckVectorThreshold(Target.TargetState.LinVel, PastState.GetV(), ResimLinVelocityErrorThreshold);
 #if !(UE_BUILD_SHIPPING || UE_BUILD_TEST)
 		if (bShouldTriggerResim)
@@ -2261,7 +2287,11 @@ bool FPhysicsReplicationAsync::ResimulationReplication(Chaos::FPBDRigidParticleH
 	// Check for angular velocity discrepancy in Degrees / s between client and server
 	if (!bShouldTriggerResim && bCompareW)
 	{
-		const float ResimAngVelocityErrorThreshold = SettingsCurrent.ResimulationSettings.GetResimulationErrorAngularVelocityThreshold(Chaos::FPhysicsSolverBase::GetResimulationErrorAngularVelocityThreshold());
+		float ResimAngVelocityErrorThreshold = PhysicsReplicationCVars::ResimulationCVars::ToleranceW;
+		if (FMath::IsNearlyZero(ResimAngVelocityErrorThreshold))
+		{
+			ResimAngVelocityErrorThreshold = SettingsCurrent.ResimulationSettings.GetResimulationErrorAngularVelocityThreshold(Chaos::FPhysicsSolverBase::GetResimulationErrorAngularVelocityThreshold());
+		}
 		bShouldTriggerResim = Chaos::FRewindData::CheckVectorThreshold(Target.TargetState.AngVel, FMath::RadiansToDegrees(PastState.GetW()), ResimAngVelocityErrorThreshold);
 #if !(UE_BUILD_SHIPPING || UE_BUILD_TEST)
 		 if (bShouldTriggerResim)
@@ -2274,7 +2304,11 @@ bool FPhysicsReplicationAsync::ResimulationReplication(Chaos::FPBDRigidParticleH
 	// Check for rotational discrepancy in Degrees between client and server
 	if (!bShouldTriggerResim && bCompareR)
 	{
-		const float ResimRotationErrorThreshold = SettingsCurrent.ResimulationSettings.GetResimulationErrorRotationThreshold(Chaos::FPhysicsSolverBase::GetResimulationErrorRotationThreshold());
+		float ResimRotationErrorThreshold = PhysicsReplicationCVars::ResimulationCVars::ToleranceR;
+		if (FMath::IsNearlyZero(ResimRotationErrorThreshold))
+		{
+			ResimRotationErrorThreshold = SettingsCurrent.ResimulationSettings.GetResimulationErrorRotationThreshold(Chaos::FPhysicsSolverBase::GetResimulationErrorRotationThreshold());
+		}
 		bShouldTriggerResim = Chaos::FRewindData::CheckQuaternionThreshold(Target.TargetState.Quaternion, PastState.GetR(), ResimRotationErrorThreshold);
 #if !(UE_BUILD_SHIPPING || UE_BUILD_TEST)
 		if (bShouldTriggerResim)
diff --git a/Engine/Source/Runtime/Experimental/Chaos/Private/Framework/PhysicsProxyBase.cpp b/Engine/Source/Runtime/Experimental/Chaos/Private/Framework/PhysicsProxyBase.cpp
--- a/Engine/Source/Runtime/Experimental/Chaos/Private/Framework/PhysicsProxyBase.cpp
+++ b/Engine/Source/Runtime/Experimental/Chaos/Private/Framework/PhysicsProxyBase.cpp
@@ -5,6 +5,14 @@
 
 namespace RenderInterpolationCVars
 {
+	bool bUseConstantDecayMechanism = false;
+	static FAutoConsoleVariableRef CVarRenderInterpbUseConstantDecayMechanism(TEXT("p.RenderInterp.UseConstantDecayMechanism"), bUseConstantDecayMechanism, TEXT(""));
+	float ConstantDecayLinearRate = 4.f;
+	static FAutoConsoleVariableRef CVarRenderInterpConstantDecayLinearRate(TEXT("p.RenderInterp.ConstantDecayLinearRate"), ConstantDecayLinearRate, TEXT(""));
+	float ConstantDecayAngularRate = 5.f;
+	static FAutoConsoleVariableRef CVarRenderInterpConstantDecayAngularRate(TEXT("p.RenderInterp.ConstantDecayAngularRate"), ConstantDecayAngularRate, TEXT(""));
+
+	
 	float RenderInterpErrorCorrectionDuration = 0.3f;
 	static FAutoConsoleVariableRef CVarRenderInterpErrorCorrectionDuration(TEXT("p.RenderInterp.ErrorCorrectionDuration"), RenderInterpErrorCorrectionDuration, TEXT("How long in seconds to apply error correction over."));
 
@@ -69,6 +77,12 @@ int32 IPhysicsProxyBase::GetSolverSyncTimestamp_External() const
 
 const bool FProxyInterpolationError::IsErrorSmoothing() const
 {
+	if (RenderInterpolationCVars::bUseConstantDecayMechanism)
+	{
+		return ErrorX.SizeSquared() > RenderInterpolationCVars::RenderInterpMinimumLinearThreshold
+			|| ErrorR.GetAngle() > RenderInterpolationCVars::RenderInterpMinimumAngularThreshold;
+	}
+	
 	const bool ApplyExponentialDecay = ErrorInterpolationSettings.Get() ? ErrorInterpolationSettings.Get()->bApplyExponentialDecay : RenderInterpolationCVars::bRenderInterpApplyExponentialDecay;
 	if (ApplyExponentialDecay)
 	{
@@ -89,7 +103,7 @@ void FProxyInterpolationError::AccumlateErrorXR(const Chaos::FVec3 X, const FQua
 	SimTicks = 0;
 
 	const bool ApplyExponentialDecay = ErrorInterpolationSettings.Get() ? ErrorInterpolationSettings.Get()->bApplyExponentialDecay : RenderInterpolationCVars::bRenderInterpApplyExponentialDecay;
-	if (ApplyExponentialDecay || IsErrorSmoothing())
+	if (RenderInterpolationCVars::bUseConstantDecayMechanism || ApplyExponentialDecay || IsErrorSmoothing())
 	{
 		ErrorX += X;
 		ErrorR = R * ErrorR;
@@ -105,9 +119,26 @@ const bool FProxyInterpolationError::UpdateError(const int32 CurrentSimTick, con
 	// Cache how many simulation ticks have passed since last call
 	SimTicks = CurrentSimTick - LastSimTick;
 	LastSimTick = CurrentSimTick;
-
+		
 	if (SimTicks > 0)
 	{
+		if (RenderInterpolationCVars::bUseConstantDecayMechanism)
+		{
+			if (!IsErrorSmoothing())
+			{
+				Reset();
+				return false;
+			}
+			ErrorXPrev = ErrorX;
+			ErrorRPrev = ErrorR;
+			const double LinearDecayPerTick = RenderInterpolationCVars::ConstantDecayLinearRate * AsyncFixedTimeStep;
+			const double LinearErrorPercentToRemove = FMath::Min( LinearDecayPerTick * SimTicks, 1.);
+			ErrorX = ErrorX * (1. - LinearErrorPercentToRemove);
+			const double AngularDecayPerTick = RenderInterpolationCVars::ConstantDecayAngularRate * AsyncFixedTimeStep;
+			const double AngularErrorPercentToRemove = FMath::Min(AngularDecayPerTick * SimTicks, 1.);
+			ErrorR = FMath::Lerp(FQuat::Identity, ErrorR, 1.0 - AngularErrorPercentToRemove);
+			return true;
+		}
 		return DecayError(AsyncFixedTimeStep);
 	}
 	return false;
diff --git a/Engine/Source/Runtime/Experimental/Chaos/Private/RewindData.cpp b/Engine/Source/Runtime/Experimental/Chaos/Private/RewindData.cpp
--- a/Engine/Source/Runtime/Experimental/Chaos/Private/RewindData.cpp
+++ b/Engine/Source/Runtime/Experimental/Chaos/Private/RewindData.cpp
@@ -1216,6 +1216,8 @@ void FRewindData::AccumulateErrorIfNecessary(FGeometryParticleHandle& Obj, const
 			{
 				if (FDirtyParticleErrorInfo* Found = DirtyParticleErrors.Find(&Obj))
 				{
+					UE_LOG(LogChaos, Error, TEXT("CLIENT | PT | AccumulateErrorIfNecessary | Found an Existing Dirty Particle for Resim."));
+
 					return *Found;
 				}
 
diff --git a/Engine/Source/Runtime/Experimental/Chaos/Public/Chaos/Framework/PhysicsProxyBase.h b/Engine/Source/Runtime/Experimental/Chaos/Public/Chaos/Framework/PhysicsProxyBase.h
--- a/Engine/Source/Runtime/Experimental/Chaos/Public/Chaos/Framework/PhysicsProxyBase.h
+++ b/Engine/Source/Runtime/Experimental/Chaos/Public/Chaos/Framework/PhysicsProxyBase.h
@@ -29,6 +29,12 @@ namespace Chaos
 
 namespace RenderInterpolationCVars
 {
+	extern bool bUseConstantDecayMechanism;
+	extern float ConstantDecayLinearRate;
+	extern float ConstantDecayAngularRate;
+	extern float MinimumLinearThreshold;
+	extern float MinimumAngularThreshold;
+	
 	extern float RenderInterpErrorCorrectionDuration;
 	extern float RenderInterpMaximumErrorCorrectionBeforeSnapping;
 	extern float RenderInterpMaximumErrorCorrectionDesyncTimeBeforeSnapping;
@@ -334,6 +340,7 @@ struct FProxyInterpolationError : FProxyInterpolationBase
 
 	CHAOS_API virtual const bool IsErrorSmoothing() const override;
 
+
 	virtual const Chaos::FVec3 GetErrorX(const Chaos::FRealSingle Alpha) const override { return FMath::Lerp(ErrorXPrev, ErrorX, Alpha); }
 
 	virtual const FQuat GetErrorR(const Chaos::FRealSingle Alpha) const override { return FQuat::Slerp(ErrorRPrev, ErrorR, Alpha); }
@@ -343,7 +350,6 @@ struct FProxyInterpolationError : FProxyInterpolationBase
 	CHAOS_API virtual const bool UpdateError(const int32 CurrentSimTick, const Chaos::FReal AsyncFixedTimeStep) override;
 
 	CHAOS_API virtual const bool DirectionalDecay(Chaos::FVec3 Direction, float ErrorDirectionalDecayMultiplier) override;
-	
 protected:
 	CHAOS_API virtual const bool DecayError(const Chaos::FReal AsyncFixedTimeStep);
 
