FWorldPartitionClassDescRegistry: fix crash when reloading a blueprint actor class.

diff --git a/Engine/Source/Runtime/Engine/Private/WorldPartition/WorldPartitionClassDescRegistry.cpp b/Engine/Source/Runtime/Engine/Private/WorldPartition/WorldPartitionClassDescRegistry.cpp
--- a/Engine/Source/Runtime/Engine/Private/WorldPartition/WorldPartitionClassDescRegistry.cpp
+++ b/Engine/Source/Runtime/Engine/Private/WorldPartition/WorldPartitionClassDescRegistry.cpp
@@ -688,38 +688,39 @@ void FWorldPartitionClassDescRegistry::UpdateClassDescriptor(UObject* InObject,
 	const FTopLevelAssetPath ClassPath(ActorCDO->GetClass()->GetPathName());
 	const FTopLevelAssetPath CurrentParentClassPath = FTopLevelAssetPath(ActorCDO->GetClass()->GetSuperClass()->GetPathName());
 
-	TUniquePtr<FWorldPartitionActorDesc>* ExistingClassDesc = ClassByPath.FindChecked(ClassPath);
-
-	if (!ExistingClassDesc)
+	if (TUniquePtr<FWorldPartitionActorDesc>** ExistingClassDesc = ClassByPath.Find(ClassPath))
 	{
-		if (!bOnlyIfExists)
+		if (!*ExistingClassDesc)
 		{
-			ValidateInternalState();
+			if (!bOnlyIfExists)
+			{
+				ValidateInternalState();
 
-			FWorldPartitionActorDesc* NewActorDesc = ActorCDO->CreateActorDesc().Release();
-			RegisterClassDescriptor(NewActorDesc);
+				FWorldPartitionActorDesc* NewActorDesc = ActorCDO->CreateActorDesc().Release();
+				RegisterClassDescriptor(NewActorDesc);
 
-			ParentClassMap.Add(ClassPath, CurrentParentClassPath);
+				ParentClassMap.Add(ClassPath, CurrentParentClassPath);
 
-			ValidateInternalState();
+				ValidateInternalState();
+			}
 		}
-	}
-	else
-	{
-		FWorldPartitionActorDescUtils::UpdateActorDescriptorFromActor(ActorCDO, *ExistingClassDesc);
-
-		const FTopLevelAssetPath PreviousParentClassPath = ParentClassMap.FindChecked(ClassPath);
-		if (PreviousParentClassPath != CurrentParentClassPath)
+		else
 		{
-			ValidateInternalState();
+			FWorldPartitionActorDescUtils::UpdateActorDescriptorFromActor(ActorCDO, **ExistingClassDesc);
 
-			// We are reparenting a blueprint, update our parent map
-			ParentClassMap.Add(ClassPath, CurrentParentClassPath);
+			const FTopLevelAssetPath PreviousParentClassPath = ParentClassMap.FindChecked(ClassPath);
+			if (PreviousParentClassPath != CurrentParentClassPath)
+			{
+				ValidateInternalState();
 
-			ValidateInternalState();
-		}
+				// We are reparenting a blueprint, update our parent map
+				ParentClassMap.Add(ClassPath, CurrentParentClassPath);
+
+				ValidateInternalState();
+			}
 
-		ClassDescriptorUpdatedEvent.Broadcast(ExistingClassDesc->Get());
+			ClassDescriptorUpdatedEvent.Broadcast((*ExistingClassDesc)->Get());
+		}
 	}
 }
 
