[FastGeo] Fixed crash in SetHitResultFromShapeAndFaceIndex for maps using FastGeo following CL 45802378.

diff --git a/Engine/Source/Runtime/Engine/Private/Collision/CollisionConversions.cpp b/Engine/Source/Runtime/Engine/Private/Collision/CollisionConversions.cpp
--- a/Engine/Source/Runtime/Engine/Private/Collision/CollisionConversions.cpp
+++ b/Engine/Source/Runtime/Engine/Private/Collision/CollisionConversions.cpp
@@ -212,9 +212,7 @@ static void SetHitResultFromShapeAndFaceIndex(const FPhysicsShape& Shape,  const
 	OutResult.PhysMaterial = nullptr;
 
 	// Grab actor/component
-	if (!OwningComponent.IsExplicitlyNull()
-		|| IPhysicsBodyInstanceOwner::GetPhysicsBodyInstandeOwnerFromHitResult(OutResult)
-		)
+	if (!OwningComponent.IsExplicitlyNull())
 	{
 		OutResult.Component = OwningComponent;
 
