Fix potential div by zero

diff --git a/Engine/Source/Runtime/Engine/Private/Distributions.cpp b/Engine/Source/Runtime/Engine/Private/Distributions.cpp
--- a/Engine/Source/Runtime/Engine/Private/Distributions.cpp
+++ b/Engine/Source/Runtime/Engine/Private/Distributions.cpp
@@ -562,7 +562,7 @@ static float ComputeLookupTableError( const FDistributionLookupTable& InTable1,
 	const FDistributionLookupTable* Table1 = (InTable2.EntryCount > InTable1.EntryCount) ? &InTable2 : &InTable1;
 	const FDistributionLookupTable* Table2 = (Table1 == &InTable1) ? &InTable2 : &InTable1;
 	const int32 ValuesPerEntry = Table1->GetValuesPerEntry();
-	const float TimeStep = (MaxIn - MinIn) / (SampleCount-1);
+	const float TimeStep = (SampleCount > 1) ? ((MaxIn - MinIn) / (float)(SampleCount-1)) : 0.0f;
 
 	float Values1[4] = {0};
 	float Values2[4] = {0};
