Tweakable to disable forcing GC on PreviewScene exit
* Thumbnail rendering can cause large stalls due to full GC collection. CVar it determine if this is the cause of the stalls impacting LDs.
* Initially added to prevent increased GPU memory usage.

diff --git a/Engine/Source/Runtime/Engine/Private/PreviewScene.cpp b/Engine/Source/Runtime/Engine/Private/PreviewScene.cpp
--- a/Engine/Source/Runtime/Engine/Private/PreviewScene.cpp
+++ b/Engine/Source/Runtime/Engine/Private/PreviewScene.cpp
@@ -18,6 +18,13 @@
 #include "Components/ReflectionCaptureComponent.h"
 #include "GameFramework/GameModeBase.h"
 
+namespace PreviewScene::Private
+{
+	bool bForceGCOnPreviewSceneExit = true;
+	FAutoConsoleVariableRef CVarForceGCOnPreviewSceneExit(TEXT("r.ForceGCOnPreviewSceneExit"), bForceGCOnPreviewSceneExit, TEXT(""), ECVF_Default);
+
+}
+
 FPreviewScene::FPreviewScene(FPreviewScene::ConstructionValues CVS)
 	: PreviewWorld(nullptr)
 	, bForceAllUsedMipsResident(CVS.bForceMipsResident)
@@ -155,11 +162,14 @@ void FPreviewScene::Uninitialize()
 		// Release PhysicsScene for fixing big fbx importing bug
 		LocalPreviewWorld->ReleasePhysicsScene();
 
-		// The preview world is a heavy-weight object and may hold a significant amount of resources,
-		// including various GPU render targets and buffers required for rendering the scene.
-		// Since UWorld is garbage-collected, this memory may not be cleaned for an indeterminate amount of time.
-		// By forcing garbage collection explicitly, we allow memory to be reused immediately.
-		GEngine->ForceGarbageCollection(true /*bFullPurge*/);
+		if (PreviewScene::Private::bForceGCOnPreviewSceneExit)
+		{
+			// The preview world is a heavy-weight object and may hold a significant amount of resources,
+			// including various GPU render targets and buffers required for rendering the scene.
+			// Since UWorld is garbage-collected, this memory may not be cleaned for an indeterminate amount of time.
+			// By forcing garbage collection explicitly, we allow memory to be reused immediately.
+			GEngine->ForceGarbageCollection(true /*bFullPurge*/);
+		}
 	}
 }
 
