Fixes for header units

diff --git a/Engine/Plugins/VirtualProduction/CameraCalibrationCore/Source/CameraCalibrationCore/Public/LensFile.h b/Engine/Plugins/VirtualProduction/CameraCalibrationCore/Source/CameraCalibrationCore/Public/LensFile.h
--- a/Engine/Plugins/VirtualProduction/CameraCalibrationCore/Source/CameraCalibrationCore/Public/LensFile.h
+++ b/Engine/Plugins/VirtualProduction/CameraCalibrationCore/Source/CameraCalibrationCore/Public/LensFile.h
@@ -31,6 +31,10 @@ class UTextureRenderTarget2D;
 struct FBaseLensTable;
 struct FDisplacementMapBlendingParams;
 
+#if WITH_EDITORONLY_DATA
+class UAssetImportData;
+#endif
+
 DECLARE_MULTICAST_DELEGATE_OneParam(FOnLensFileModelChanged, const TSubclassOf<ULensModel>&);
 
 
diff --git a/Engine/Source/Runtime/Engine/Public/SceneView.h b/Engine/Source/Runtime/Engine/Public/SceneView.h
--- a/Engine/Source/Runtime/Engine/Public/SceneView.h
+++ b/Engine/Source/Runtime/Engine/Public/SceneView.h
@@ -1332,7 +1332,7 @@ END_GLOBAL_SHADER_PARAMETER_STRUCT()
 
 namespace InstancedViewParametersUtils
 {
-	static void CopyIntoInstancedViewParameters(FInstancedViewUniformShaderParameters& InstancedViewParameters, const FViewUniformShaderParameters& ViewParameters, uint32 CopyViewId)
+	inline void CopyIntoInstancedViewParameters(FInstancedViewUniformShaderParameters& InstancedViewParameters, const FViewUniformShaderParameters& ViewParameters, uint32 CopyViewId)
 	{
 		VIEW_UNIFORM_BUFFER_MEMBER_TABLE
 	}
@@ -1373,7 +1373,7 @@ namespace InstancedViewParametersUtils
 namespace ViewParametersUtils
 {
 	// This only copy the base view, not any instanced view data.
-	static void CopyViewParametersOnly(FViewUniformShaderParameters& DstView, const FViewUniformShaderParameters& SrcView)
+	inline void CopyViewParametersOnly(FViewUniformShaderParameters& DstView, const FViewUniformShaderParameters& SrcView)
 	{
 		VIEW_UNIFORM_BUFFER_MEMBER_TABLE
 	}
