Allow fast path through material update context when no static permutations need recaching.
Useful for mesh painting where we override textures that may be used by many materials inside a context and know that we don't need to recache static permutations.
#rnx

diff --git a/Engine/Source/Runtime/Engine/Private/Materials/MaterialShared.cpp b/Engine/Source/Runtime/Engine/Private/Materials/MaterialShared.cpp
--- a/Engine/Source/Runtime/Engine/Private/Materials/MaterialShared.cpp
+++ b/Engine/Source/Runtime/Engine/Private/Materials/MaterialShared.cpp
@@ -4812,6 +4812,8 @@ FMaterialUpdateContext::FMaterialUpdateContext(uint32 Options, EShaderPlatform I
 	bool bRecreateRenderStates = ((Options & EOptions::RecreateRenderStates) != 0) && FApp::CanEverRender();
 
 	bSyncWithRenderingThread = (Options & EOptions::SyncWithRenderingThread) != 0;
+	bSkipRecacheStaticPermutations = (Options & EOptions::SkipRecacheStaticPermutations) != 0;
+
 	if (bReregisterComponents)
 	{
 		ComponentReregisterContext = MakeUnique<FGlobalComponentReregisterContext>();
@@ -4933,7 +4935,10 @@ FMaterialUpdateContext::~FMaterialUpdateContext()
 		MI->UpdateCachedData();
 #endif
 		MI->RecacheUniformExpressions(true);
-		MI->InitStaticPermutation(EMaterialShaderPrecompileMode::None);//bHasStaticPermutation can change.
+		if (!bSkipRecacheStaticPermutations)
+		{
+			MI->InitStaticPermutation(EMaterialShaderPrecompileMode::None);//bHasStaticPermutation can change.
+		}
 		if (MI->bHasStaticPermutationResource)
 		{
 			NumInstancesWithStaticPermutations++;
diff --git a/Engine/Source/Runtime/Engine/Public/MaterialShared.h b/Engine/Source/Runtime/Engine/Public/MaterialShared.h
--- a/Engine/Source/Runtime/Engine/Public/MaterialShared.h
+++ b/Engine/Source/Runtime/Engine/Public/MaterialShared.h
@@ -3330,6 +3330,8 @@ class FMaterialUpdateContext
 	EShaderPlatform ShaderPlatform;
 	/** True if the SyncWithRenderingThread option was specified. */
 	bool bSyncWithRenderingThread;
+	/** True if the SkipRecacheStaticPermutations option was specified. */
+	bool bSkipRecacheStaticPermutations;
 
 public:
 
@@ -3348,6 +3350,8 @@ public:
 			SyncWithRenderingThread = 0x2,
 			/* Recreates only the render state for all components (mutually exclusive with ReregisterComponents) */
 			RecreateRenderStates = 0x4,
+			/** Skip recache of material instance permutations. Only safe to use if you know that no shader code or static switches are changing. */
+			SkipRecacheStaticPermutations = 0x08,
 			/** Default options: Recreate render state, sync with rendering thread. */
 			Default = RecreateRenderStates | SyncWithRenderingThread,
 		};
