FPrimitiveInstanceDataManager: Cooking: Downgrade verbosity of logline "Build Instance Spatial Hashes (%.2fms), ..." to reduce size of cooker log in large projects with precomputation on.
#rnx

diff --git a/Engine/Source/Runtime/Engine/Private/InstancedStaticMesh/ISMInstanceDataManager.cpp b/Engine/Source/Runtime/Engine/Private/InstancedStaticMesh/ISMInstanceDataManager.cpp
--- a/Engine/Source/Runtime/Engine/Private/InstancedStaticMesh/ISMInstanceDataManager.cpp
+++ b/Engine/Source/Runtime/Engine/Private/InstancedStaticMesh/ISMInstanceDataManager.cpp
@@ -1023,7 +1023,7 @@ void FPrimitiveInstanceDataManager::BeginCacheForCookedPlatformData(const ITarge
 		PrecomputedOptimizationData = MakeShared<FISMPrecomputedSpatialHashData>(PrecomputeOptimizationData(MoveTemp(ComponentData), InstanceTransforms));
 		uint32 EndtTime = FPlatformTime::Cycles();
 
-		UE_LOG(LogInstanceProxy, Log, TEXT("Build Instance Spatial Hashes (%.2fms), Instances: %d, Hashes: %d, Remap Size: %d"), FPlatformTime::ToMilliseconds( EndtTime - StartTime ), InstanceTransforms.Num(), PrecomputedOptimizationData->Hashes.Num(), PrecomputedOptimizationData->ProxyIndexToComponentIndexRemap.Num());
+		UE_LOG(LogInstanceProxy, Verbose, TEXT("Build Instance Spatial Hashes (%.2fms), Instances: %d, Hashes: %d, Remap Size: %d"), FPlatformTime::ToMilliseconds( EndtTime - StartTime ), InstanceTransforms.Num(), PrecomputedOptimizationData->Hashes.Num(), PrecomputedOptimizationData->ProxyIndexToComponentIndexRemap.Num());
 	}	
 }
 
