New Material Translator:

- Added support for custom outputs through the new API FMaterialIREmitter::SetCustomOutputs().
- Refactored the translator top to bottom to improve how SetMaterialOutput instructions are emitted and sorted.
- Added support for FMaterialExpressionTangentOutput. Other expressions support will come in subsequent change.

#rnx

diff --git a/Engine/Shaders/Private/MaterialTemplate.ush b/Engine/Shaders/Private/MaterialTemplate.ush
--- a/Engine/Shaders/Private/MaterialTemplate.ush
+++ b/Engine/Shaders/Private/MaterialTemplate.ush
@@ -644,6 +644,8 @@ struct FMaterialPixelParameters
 	FMaterialAttributes MaterialAttributes;
 
 	FMaterialLWCData LWCData;
+
+%{material_pixel_parameters_extras}
 };
 
 %{material_cache_types}
diff --git a/Engine/Source/Runtime/Engine/Private/Materials/MaterialExpressionsIR.cpp b/Engine/Source/Runtime/Engine/Private/Materials/MaterialExpressionsIR.cpp
--- a/Engine/Source/Runtime/Engine/Private/Materials/MaterialExpressionsIR.cpp
+++ b/Engine/Source/Runtime/Engine/Private/Materials/MaterialExpressionsIR.cpp
@@ -4213,4 +4213,10 @@ void UMaterialExpressionSwitch::Build(MIR::FEmitter& Em)
 	Em.Output(0, PreviousBranch);
 }
 
+void UMaterialExpressionTangentOutput::Build(MIR::FEmitter& Em)
+{
+	FValueRef InputValue = Em.CastToFloat(Em.Input(&Input), 1);
+	Em.SetCustomOutputs(TEXTVIEW("TangentOutput"), { &InputValue, 1 });
+}
+
 #endif // WITH_EDITOR
diff --git a/Engine/Source/Runtime/Engine/Private/Materials/MaterialIR.cpp b/Engine/Source/Runtime/Engine/Private/Materials/MaterialIR.cpp
--- a/Engine/Source/Runtime/Engine/Private/Materials/MaterialIR.cpp
+++ b/Engine/Source/Runtime/Engine/Private/Materials/MaterialIR.cpp
@@ -935,11 +935,6 @@ bool FFunction::Equals(const FFunction* Other) const
 		&& FMemory::Memcmp(Parameters, Other->Parameters, NumParameters * sizeof(FFunctionParameter)) == 0;
 }
 
-bool FFunctionHLSL::Equals(const FFunctionHLSL* Other) const
-{
-	return FFunction::Equals(Other) && Code == Other->Code;
-}
-
 } // namespace MIR
 
 #endif // #if WITH_EDITOR
diff --git a/Engine/Source/Runtime/Engine/Private/Materials/MaterialIRDebug.cpp b/Engine/Source/Runtime/Engine/Private/Materials/MaterialIRDebug.cpp
--- a/Engine/Source/Runtime/Engine/Private/Materials/MaterialIRDebug.cpp
+++ b/Engine/Source/Runtime/Engine/Private/Materials/MaterialIRDebug.cpp
@@ -313,7 +313,7 @@ struct FDebugDumpIRState
 	{
 		if (const FSetMaterialOutput* SetMaterialOutput = Instr->As<FSetMaterialOutput>())
 		{
-			Out.Appendf(TEXT(" \"%s\""), *FMaterialAttributeDefinitionMap::GetAttributeName(SetMaterialOutput->Property));
+			Out.Appendf(TEXT(" \"%s\""), *SetMaterialOutput->Name);
 		}
 		else if (const FOperator* Operator = Instr->As<FOperator>())
 		{
@@ -505,7 +505,7 @@ struct FDebugDumpIRState
 
 		for (const MIR::FFunctionHLSL* Function : ReferencedFunctionHLSLs)
 		{
-			Out.Appendf(TEXT("FunctionHLSL Name=\"C%d_%s\" ReturnType=\"%s\"\n"), Function->UniqueId, Function->Name.GetData(), *Function->ReturnType.GetSpelling());
+			Out.Appendf(TEXT("FunctionHLSL Name=\"%s\" ReturnType=\"%s\"\n"), Function->Name.GetData(), *Function->ReturnType.GetSpelling());
 			for (int i = 0; i < Function->NumParameters; ++i)
 			{
 				const TCHAR* Keyword = (i < Function->NumInputOnlyParams) ? TEXT("In")
diff --git a/Engine/Source/Runtime/Engine/Private/Materials/MaterialIREmitter.cpp b/Engine/Source/Runtime/Engine/Private/Materials/MaterialIREmitter.cpp
--- a/Engine/Source/Runtime/Engine/Private/Materials/MaterialIREmitter.cpp
+++ b/Engine/Source/Runtime/Engine/Private/Materials/MaterialIREmitter.cpp
@@ -173,6 +173,25 @@ static FComposite* MakeCompositePrototype(FEmitter& Emitter, FType Type, int Num
 	return Value;
 }
 
+// Emits a new value into the module, performing no de-duplication with existing values in the module.
+// NOTE: This function should **only** be used when emitting a "terminal" instruction that is known
+// to not have duplicates, condition checked externally by the caller.
+template <typename T>
+static T* EmitNew(FEmitter& Emitter, FType Type)
+{
+	// Otherwise, create a new value allocating the necessary memory in the module's arena.
+	T* Value = (T*)FEmitter::FPrivate::Allocate(Emitter, sizeof(T), alignof(T));
+	FMemory::Memzero(Value, sizeof(T));
+
+	// Set up the basics
+	Value->Kind = T::TypeKind;
+	Value->Type = Type;
+
+	FEmitter::FPrivate::PushValueToModule(Emitter, Value);
+
+	return Value;
+}
+
 // Emits a prototype value into the module, deduplicating if an identical value was already emitted.
 static FValueRef EmitPrototype(FEmitter& Emitter, const FValue& Prototype)
 {
@@ -266,9 +285,9 @@ FValueRef FValueRef::ToPoison() const
 	return To(FPoison::Get());
 }
 
-/*--------------------------------- FFunctionHLSLDesc ---------------------------------*/
+/*----------------------------------- FFunctionDesc -----------------------------------*/
 
-bool FFunctionHLSLDesc::PushInputOnlyParameter(FName InName, FType InType)
+bool FFunctionDesc::PushInputOnlyParameter(FName InName, FType InType)
 {
 	// You must input-only parameters first, before the others.
 	check(!NumInputOutputParams && !NumOutputOnlyParams);
@@ -282,7 +301,7 @@ bool FFunctionHLSLDesc::PushInputOnlyParameter(FName InName, FType InType)
 	return true;
 }
 	
-bool FFunctionHLSLDesc::PushInputOutputParameter(FName InName, FType InType)
+bool FFunctionDesc::PushInputOutputParameter(FName InName, FType InType)
 {
 	checkf(!NumOutputOnlyParams, TEXT("You must input-output parameters before you push any output-only parameters."));
 
@@ -296,7 +315,7 @@ bool FFunctionHLSLDesc::PushInputOutputParameter(FName InName, FType InType)
 	return true;
 }
 	
-bool FFunctionHLSLDesc::PushOutputOnlyParameter(FName InName, FType InType)
+bool FFunctionDesc::PushOutputOnlyParameter(FName InName, FType InType)
 {
 	if (GetNumParameters() == MIR::MaxNumFunctionParameters)
 	{
@@ -1300,30 +1319,81 @@ FValueRef FEmitter::ScreenTexture(EScreenTexture TextureKind)
 
 /*------------------------------ Instruction emission ------------------------------*/
 
-FSetMaterialOutput* FEmitter::SetMaterialOutput(EMaterialProperty InProperty, FValue* Arg)
+static FSetMaterialOutput* CreateSetMaterialOutput(FEmitter& Em, FMaterialIRModule* Module, FStringView Name, EMaterialProperty Property, EFrequency Frequency, FValueRef Arg)
 {
-	FSetMaterialOutput Proto = MakePrototype<FSetMaterialOutput>(FType::MakeVoid());
-	Proto.Property           = InProperty;
-	Proto.Arg                = Arg;
-	
-	FSetMaterialOutput* Instr = static_cast<FSetMaterialOutput*>(EmitPrototype(*this, Proto).Value);
+	// Material attributes is not a real material attribute, the translator should never create a SetMaterialOutput on this property.
+	check(Property != MP_MaterialAttributes); 
+
+	// Create the set material output instruction.
+	FSetMaterialOutput* Instr = EmitNew<FSetMaterialOutput>(Em, FType::MakeVoid());
+	Instr->Property           = Property;
+	Instr->Name               = Module->InternString(Name);
+	Instr->Arg                = Arg;
+
+	// Set the instruction priority according to this output property.
+	Instr->Priority = Property == MP_Normal ? 0 		// Normal goes first
+                    : Property == MP_CustomOutput ? 2 	// Custom outputs go last
+                    : 1; 								// All other attributes (BaseColor, Emissive, etc)
 
 	// Add the instruction to list of outputs of the stages it is evaluated in.
-	if (InProperty == MP_WorldPositionOffset)
+	if (Frequency == Frequency_PerVertex)
 	{
-		Module->GetEntryPoint(MIR::Stage_Vertex).Outputs[0] = Instr;
+		Module->GetEntryPoint(MIR::Stage_Vertex).Outputs.Push(Instr);
 	}
 	else
 	{
-		int32 StageOutputIndex = UMaterialAggregate::MaterialPropertyToAttributeIndex(InProperty);
-		check(StageOutputIndex != -1);
-		Module->GetEntryPoint(MIR::Stage_Pixel).Outputs[StageOutputIndex] = Instr;
-		Module->GetEntryPoint(MIR::Stage_Compute).Outputs[StageOutputIndex] = Instr;
+		Module->GetEntryPoint(MIR::Stage_Pixel).Outputs.Push(Instr);
+		Module->GetEntryPoint(MIR::Stage_Compute).Outputs.Push(Instr);
 	}
 
 	return Instr;
 }
 
+void FEmitter::Private_SetMaterialOutput(EMaterialProperty Property, FValueRef Arg)
+{
+	checkf(Property != MP_CustomOutput, TEXT("This API does not support CustomOutputs, use SetCustomOutput() instead."));
+
+	FString Name         = UMaterialAggregate::GetMaterialAttribute(Property)->Name.ToString();
+	EFrequency Frequency = (Property == MP_WorldPositionOffset) ? MIR::Frequency_PerVertex : MIR::Frequency_PerPixel;
+
+	CreateSetMaterialOutput(*this, Module, Name, Property, Frequency, Arg);
+}
+
+void FEmitter::SetCustomOutputs(FStringView Name, TConstArrayView<FValueRef> Args)
+{
+	// If this is a custom output, check that its name is unique.
+	for (const FMaterialIRModule::FCustomOutputGroup& Group : Module->GetCustomOutputGroups())
+	{
+		if (Group.Name == Name)
+		{
+			Errorf(TEXT("A custom output with name '%s' already exists."), Name.GetData());
+			return;
+		}
+	}
+
+	// Declare a new custom output group to include all generated individual outputs
+	Module->AddCustomOutputGroup(Module->InternString(Name), Args.Num());
+
+	// Allocate a small buffer to make room for the custom output name plus index
+	MIR::TTemporaryArray<TCHAR> TempNameBuffer{ Name.Len() + 8 };
+
+	// Copy the custom output base name to the temp buffer. 
+	FCString::Snprintf(TempNameBuffer.GetData(), TempNameBuffer.Num(), TEXT("%.*s"), Name.Len(), Name.GetData());
+
+	// Iterate over all the arguments and create a custom output SetMaterialOutput for each
+	for (int32 i = 0; i < Args.Num(); ++i)
+	{
+		// Append the index to the base name to make this custom output name, e.g. "TangentOutput0"
+		int32 NameLen = Name.Len() + FCString::Snprintf(TempNameBuffer.GetData() + Name.Len(), TempNameBuffer.Num() - Name.Len(), TEXT("%d"), i);
+
+		// Create the SetMaterialOutput instruction
+		FSetMaterialOutput* Instr = CreateSetMaterialOutput(*this, Module, { TempNameBuffer.GetData(), NameLen }, MP_CustomOutput, MIR::Frequency_PerPixel, Args[i]);
+	
+		// Push the instruction to the custom outputs in the module
+		Module->AddCustomOutput(Instr);
+	}
+}
+
 FValueRef FEmitter::Vector2(FValueRef InX, FValueRef InY)
 {
 	if (IsAnyNotValid(InX, InY))
@@ -4026,36 +4096,32 @@ FValueRef FEmitter::PromoteSubstrateParameter()
 	return EmitPrototype(*this, Prototype);
 }
 
+static void SetFunctionFromDesc(FFunction& Function, FFunctionKind Kind, const FFunctionDesc& Desc)
+{
+	Function.Kind = Kind;
+	Function.Name = Desc.Name;
+	Function.NumInputOnlyParams = Desc.NumInputOnlyParams;
+	Function.NumInputAndOutputParams = Desc.NumInputOnlyParams + Desc.NumInputOutputParams;
+	Function.NumParameters = Function.NumInputAndOutputParams + Desc.NumOutputOnlyParams;
+	Function.ReturnType = Desc.ReturnType;
+
+	// Copy over the parameter declarations from the description to the function prototype
+	for (uint32 i = 0; i < Function.NumParameters; ++i)
+	{
+		Function.Parameters[i] = Desc.Parameters[i];
+	} 
+}
+
 const FFunction* FEmitter::FunctionHLSL(const FFunctionHLSLDesc& Desc)
 {
 	check(Desc.GetNumParameters() <= MIR::MaxNumFunctionParameters);
 
 	FFunctionHLSL Prototype = {};
-	Prototype.Name = Desc.Name;
-	Prototype.ReturnType = Desc.ReturnType;
+	SetFunctionFromDesc(Prototype, FFunctionKind::HLSL, Desc);
 	Prototype.Code = Desc.Code;
-	Prototype.NumInputOnlyParams = Desc.NumInputOnlyParams;
-	Prototype.NumInputAndOutputParams = Desc.NumInputOnlyParams + Desc.NumInputOutputParams;
-	Prototype.NumParameters = Prototype.NumInputAndOutputParams + Desc.NumOutputOnlyParams;
-	Prototype.UniqueId = Module->FunctionHLSLs.Num();
 	Prototype.Defines = Desc.Defines;
 	Prototype.Includes = Desc.Includes;
-
-	// Copy over the parameter declarations from the description to the function prototype
-	for (uint32 i = 0; i < Prototype.NumParameters; ++i)
-	{
-		Prototype.Parameters[i] = Desc.Parameters[i];
-	}
-
-	// TODO: Optimize the lookup to be constant
-	for (const FFunctionHLSL* Function : Module->FunctionHLSLs)
-	{
-		if (Function->Equals(&Prototype))
-		{
-			return Function;
-		}
-	}
-
+	
 	// Create the new MIR HLSL function instance and set it up
 	FFunctionHLSL* Function = new (FPrivate::Allocate(*this, sizeof(FFunctionHLSL), alignof(FFunctionHLSL))) FFunctionHLSL { Prototype };
 	Function->Name = Module->InternString(Function->Name);
diff --git a/Engine/Source/Runtime/Engine/Private/Materials/MaterialIRModule.cpp b/Engine/Source/Runtime/Engine/Private/Materials/MaterialIRModule.cpp
--- a/Engine/Source/Runtime/Engine/Private/Materials/MaterialIRModule.cpp
+++ b/Engine/Source/Runtime/Engine/Private/Materials/MaterialIRModule.cpp
@@ -30,10 +30,12 @@ void FMaterialIRModule::Empty()
 	ParameterIdToData.Empty();
 	ParameterInfoToId.Empty();
 	Statistics = {};
-	MIR::ZeroArray(TArrayView<MIR::FValue*>{ PropertyValues, MP_MAX });
+	MIR::ZeroArray(TArrayView<MIR::FValue*>{ PropertyValues, UE_ARRAY_COUNT(PropertyValues) });
 	EntryPoints.Empty();
 	Values.Empty();
 	CompilationOutput = {};
+	CustomOutputs.Empty();
+	CustomOutputGroups.Empty();
 	Allocator.Flush();
 
 	// Initialize statistics used external input usage mask.
@@ -43,7 +45,7 @@ void FMaterialIRModule::Empty()
 	}
 }
 
-int32 FMaterialIRModule::AddEntryPoint(FStringView Name, MIR::EStage Stage, int32 NumOutputs)
+int32 FMaterialIRModule::AddEntryPoint(FStringView Name, MIR::EStage Stage, int32 NumOutputsHint)
 {
 	uint32 Index = EntryPoints.Num();
 
@@ -51,14 +53,18 @@ int32 FMaterialIRModule::AddEntryPoint(FStringView Name, MIR::EStage Stage, int3
 		.Name = InternString(Name),
 		.Stage = Stage,
 		.RootBlock = {},
-		.Outputs = AllocateArray<MIR::FValue*>(NumOutputs),
 	});
 
-	MIR::ZeroArray(EntryPoints.Last().Outputs);
+	EntryPoints.Last().Outputs.Reserve(NumOutputsHint);
 
 	return Index;
 }
 
+void FMaterialIRModule::AddCustomOutputGroup(FStringView Name, int32 NumOutputs)
+{
+	CustomOutputGroups.Push(FCustomOutputGroup{ Name, NumOutputs, CustomOutputs.Num() });
+}
+
 FStringView FMaterialIRModule::InternString(FStringView InString)
 {
 	// Allocate the buffer to contain the string
diff --git a/Engine/Source/Runtime/Engine/Private/Materials/MaterialIRModuleBuilder.cpp b/Engine/Source/Runtime/Engine/Private/Materials/MaterialIRModuleBuilder.cpp
--- a/Engine/Source/Runtime/Engine/Private/Materials/MaterialIRModuleBuilder.cpp
+++ b/Engine/Source/Runtime/Engine/Private/Materials/MaterialIRModuleBuilder.cpp
@@ -145,6 +145,14 @@ struct FMaterialIRModuleBuilderImpl
 				PushDependency(Input.Input->Expression);
 			}
 		}
+
+		// Push the custom output expressions to the dependency queue too.
+		TArray<UMaterialExpressionCustomOutput*> CustomOutputs;
+		Builder->Material->GetAllCustomOutputExpressions(CustomOutputs);
+		for (UMaterialExpressionCustomOutput* Expression : CustomOutputs)
+		{
+			PushDependency(Expression);
+		}
 	}
 
 	bool Step_BuildMaterialExpressionsToIRGraph()
@@ -470,8 +478,8 @@ struct FMaterialIRModuleBuilderImpl
 			Module->SetPropertyValue(Property, PropertyValue);
 
 			// Emit the SetMaterialProperty instruction
-			Emitter.SetMaterialOutput(Property, PropertyValue);
-			
+			Emitter.Private_SetMaterialOutput(Property, PropertyValue);
+
 			// Finally, push this connection insight
 			PushConnectionInsight(Builder->Material, (int32)Property, InputDesc.Input->Expression, InputDesc.Input->OutputIndex, PropertyValue->Type);
 		}
@@ -522,6 +530,20 @@ struct FMaterialIRModuleBuilderImpl
 
 			FMaterialIRModule::FEntryPoint& EntryPoint = Module->GetEntryPoint(EntryPointIndex);
 
+			// Stable sort the outputs in this entry point by their indicated priority.
+			// Note: this is needed to support situations where a final entry-point instruction needs to execute
+			//    before another, for instance we want to set the Normal attribute before all others, because the
+			//    others may read back the data. This shouldn't really be done this way, but it is currently a
+			//    restriction of how the HLSL backend works.
+			Algo::StableSort(EntryPoint.Outputs, [](const MIR::FValue* A, const MIR::FValue* B)
+				{
+					const MIR::FSetMaterialOutput* SMOA = MIR::As<MIR::FSetMaterialOutput>(A);
+					const MIR::FSetMaterialOutput* SMOB = MIR::As<MIR::FSetMaterialOutput>(B);
+					int PriorityA = SMOA ? SMOA->Priority : 0;
+					int PriorityB = SMOB ? SMOB->Priority : 0;
+					return PriorityA < PriorityB;
+				});
+
 			// Push this SetOutput instruction the value stack for processing.
 			for (MIR::FValue* Output : EntryPoint.Outputs)
 			{
diff --git a/Engine/Source/Runtime/Engine/Private/Materials/MaterialIRToHLSLTranslator.cpp b/Engine/Source/Runtime/Engine/Private/Materials/MaterialIRToHLSLTranslator.cpp
--- a/Engine/Source/Runtime/Engine/Private/Materials/MaterialIRToHLSLTranslator.cpp
+++ b/Engine/Source/Runtime/Engine/Private/Materials/MaterialIRToHLSLTranslator.cpp
@@ -144,6 +144,12 @@ struct FHLSLPrinter
         return *this;
     }
 
+	FHLSLPrinter& operator<<(TCHAR Ch)
+    {
+        Buffer.AppendChar(Ch);
+        return *this;
+    }
+
 	void BeginList()
 	{
 		bFirstListItem = true;
@@ -298,8 +304,7 @@ struct FPrivate : FMaterialIRToHLSLTranslation
 	// Generates the full HLSL of the specified entry point.
 	const FMaterialIRModule::FEntryPoint& GenerateEntryPoint(int32 EntryPointIndex)
 	{
-		Printer = {};
-		Printer.Tabs = 1;
+		Printer = { .Tabs = 1 };
 		Printer << Indentation;
 
 		const FMaterialIRModule::FEntryPoint& EntryPoint = Module->GetEntryPoint(EntryPointIndex);
@@ -1291,14 +1296,8 @@ struct FPrivate : FMaterialIRToHLSLTranslation
 		}
 		else
 		{
-			if (const FMaterialAggregateAttribute* PropertyAggregate = UMaterialAggregate::GetMaterialAttribute(SetMaterialOutput->Property))
-			{
-				Printer << TEXTVIEW("PixelMaterialInputs.") << PropertyAggregate->Name.ToString() << TEXTVIEW(" = ");
-			}
-			else
-			{
-				UE_MIR_UNREACHABLE();
-			}			
+			Printer << (SetMaterialOutput->Property != MP_CustomOutput ? TEXTVIEW("PixelMaterialInputs.") : TEXTVIEW("Parameters."));
+			Printer << SetMaterialOutput->Name.ToView() << TEXTVIEW(" = ");
 		}
 			
 		LowerValue(SetMaterialOutput->Arg);
@@ -1360,7 +1359,7 @@ struct FPrivate : FMaterialIRToHLSLTranslation
 				default: UE_MIR_UNREACHABLE();
 			}
 
-			Printer << TEXTVIEW("(") << LowerValue(Operator->AArg) << TEXTVIEW(" ") << OpString << TEXTVIEW(" ") << LowerValue(Operator->BArg) << TEXTVIEW(")");
+			Printer << TEXTVIEW("(") << LowerValue(Operator->AArg) << L' ' << OpString << L' ' << LowerValue(Operator->BArg) << TEXTVIEW(")");
 		}
 		else
 		{
@@ -1742,7 +1741,7 @@ struct FPrivate : FMaterialIRToHLSLTranslation
 
 		Printer << ListSeparator << LowerValue(TextureRead->VTPageTable)
 			<< ListSeparator << VTPageTableResult->VTPageTableIndex
-			<< ListSeparator << TEXTVIEW("VTUniform_Unpack(") << TEXT("Material.VTPackedUniform[") << VirtualTextureIndex << TEXTVIEW("]") << TEXTVIEW(")")
+			<< ListSeparator << TEXTVIEW("VTUniform_Unpack(") << TEXTVIEW("Material.VTPackedUniform[") << VirtualTextureIndex << TEXTVIEW("]") << TEXTVIEW(")")
 			<< EndArgs;
 	}
 
@@ -1985,30 +1984,30 @@ struct FPrivate : FMaterialIRToHLSLTranslation
 		// Generate locals to store the output and input-output parameters
 		for (int32 i = Call->Function->NumInputOnlyParams; i < Call->Function->NumParameters; ++i)
 		{
-			Printer << LowerType(Call->Function->Parameters[i].Type) << TEXT(" _") << (NumLocals++);
+			Printer << LowerType(Call->Function->Parameters[i].Type) << TEXTVIEW(" _") << (NumLocals++);
 			
 			if (i < Call->NumArguments)
 			{
-				Printer << TEXT(" = ") << LowerValue(Call->Arguments[i]);
+				Printer << TEXTVIEW(" = ") << LowerValue(Call->Arguments[i]);
 			}
 
 			Printer << EndOfStatement;
 		}
 
 		// Print the local that will store the result and assign it to the call to the custom function.
-		Printer << LowerType(Call->Function->ReturnType) << TEXT(" _") << NumLocals << TEXT(" = ") << TEXT("C") << Call->Function->UniqueId << TEXT("_") << Call->Function->Name;
+		Printer << LowerType(Call->Function->ReturnType) << TEXTVIEW(" _") << NumLocals << TEXTVIEW(" = ") << L'C' << Call->Function->UniqueId << L'_' << Call->Function->Name;
 		
-		Printer << BeginArgs << TEXT("Parameters");
+		Printer << BeginArgs << TEXTVIEW("Parameters");
 
 		// Print function call arguments
 		for (int32 i = 0; i < Call->Function->NumParameters; ++i)
 		{
-			Printer << TEXT(", ");
+			Printer << TEXTVIEW(", ");
 
 			// Outputs and input-output parameters are stored in special locals. Refer to them
 			if (i >= Call->Function->NumInputOnlyParams)
 			{
-				Printer << TEXT("_") << ParamLocal++;
+				Printer << L'_' << ParamLocal++;
 			}
 			else
 			{
@@ -2036,7 +2035,7 @@ struct FPrivate : FMaterialIRToHLSLTranslation
 		// Compute the index of the local that stores the additional output parameter 
 		ParamIndex = ParamIndex - Call->Function->NumInputAndOutputParams + Call->Function->NumInputOnlyParams + CallOutput->Index;
 
-		Printer << TEXT("_") << ParamIndex;
+		Printer << L'_' << ParamIndex;
 	}
 
 	ENoOp LowerType(const MIR::FType& Type)
@@ -2080,22 +2079,17 @@ struct FPrivate : FMaterialIRToHLSLTranslation
 					Printer << PrimitiveType->NumRows << TEXTVIEW("x") << PrimitiveType->NumColumns;
 				}
 			}
+
+			return NoOp;
 		}
-		else if (Type.IsVoid())
-		{
-			Printer << TEXTVIEW("void");
-		}
-		else if (Type.IsSubstrateData())
-		{
-			Printer << TEXTVIEW("FSubstrateData");
-		}
-		else if (Type.IsVTPageTableResult())
-		{
-			Printer << TEXTVIEW("VTPageTableResult");
-		}
-		else
+
+		switch (Type.GetKind())
 		{
-			UE_MIR_UNREACHABLE();
+			case MIR::ETypeKind::Void: Printer << TEXTVIEW("void"); break;
+			case MIR::ETypeKind::ShadingModel: Printer << TEXTVIEW("uint"); break;
+			case MIR::ETypeKind::SubstrateData: Printer << TEXTVIEW("FSubstrateData"); break;
+			case MIR::ETypeKind::VTPageTableResult: Printer << TEXTVIEW("VTPageTableResult"); break;
+			default: UE_MIR_UNREACHABLE();
 		}
 
 		return NoOp;
@@ -2103,6 +2097,52 @@ struct FPrivate : FMaterialIRToHLSLTranslation
 	
 	/* Finalization */
 
+	void GenerateCustomOutputs(TMap<FString, FString>& Params)
+	{
+		// Generate the custom output fields of the struct FMaterialPixelParameters. These fields will be computed and 
+		// set by the function that calculates all material outputs (e.g. BaseColor).
+		Printer = { .Tabs = 1 };
+		for (const MIR::FSetMaterialOutput* CustomOutput : Module->GetCustomOutputs())
+		{
+			Printer << Indentation << LowerType(CustomOutput->Arg->Type) << L' ' << *CustomOutput->Name << EndOfStatement;
+		}
+		Params.Add(TEXT("material_pixel_parameters_extras"), MoveTemp(Printer.Buffer));
+
+		// The material shader expects a getter function for each custom output. Generate these getter functions
+		// for each FMaterialPixelParameters field, generated above, matching each custom output in each group.
+		Printer = {};
+		for (const FMaterialIRModule::FCustomOutputGroup& Group : Module->GetCustomOutputGroups())
+		{
+			Printer << TEXTVIEW("#define NUM_MATERIAL_OUTPUTS_GET") << *FString(Group.Name).ToUpper() << L' ' << Group.NumOutputs << NewLine;
+
+			for (int i = 0; i < Group.NumOutputs; ++i)
+			{
+				const MIR::FSetMaterialOutput* Instr = Module->GetCustomOutputs()[i + Group.CustomOutputInstructionsBegin];
+				Printer << TEXTVIEW("#define HAVE_Get") << *Instr->Name << TEXTVIEW(" 1") << NewLine << NewLine;
+
+				// Print the non-LWC get function.
+				Printer << LowerType(Instr->Arg->Type) << TEXTVIEW(" Get") << Instr->Name << TEXTVIEW("(inout FMaterialPixelParameters Parameters)")
+						<< NewLine << OpenBrace << TEXTVIEW("return ");
+
+				Instr->Arg->Type.IsDouble()
+					? Printer << TEXTVIEW("WSDemote(Parameters.") << Instr->Name << L')'
+					: Printer << TEXTVIEW("Parameters.") << Instr->Name;
+
+				Printer << EndOfStatement << CloseBrace << NewLine;
+				
+				// Print the LWC get function.
+				Printer << LowerType(Instr->Arg->Type.GetPrimitive().ToScalarKind(MIR::EScalarKind::Double)) << TEXTVIEW(" Get") << Instr->Name << TEXTVIEW("_LWC(inout FMaterialPixelParameters Parameters)")
+					<< NewLine << OpenBrace << TEXTVIEW("return ");
+				
+				Instr->Arg->Type.IsDouble()
+					? Printer << TEXTVIEW("Parameters.") << Instr->Name
+					: Printer << TEXTVIEW("WSPromote(Parameters.") << Instr->Name << L')';
+				
+				Printer << EndOfStatement << CloseBrace << NewLine;
+			}
+		}
+	}
+
 	void GenerateTemplateStringParameters(TMap<FString, FString>& Params)
 	{
 		const FMaterialIRModule::FStatistics ModuleStatistics = Module->GetStatistics();
@@ -2209,20 +2249,20 @@ struct FPrivate : FMaterialIRToHLSLTranslation
 			for (const MIR::FFunctionHLSLDefine& Define : Function->Defines)
 			{
 				Printer << TEXTVIEW("#ifndef ") << Define.Name << NewLine;
-				Printer << TEXTVIEW("\t#define ") << Define.Name << TEXTVIEW(" ") << Define.Value << NewLine;
+				Printer << TEXTVIEW("\t#define ") << Define.Name << L' ' << Define.Value << NewLine;
 				Printer << TEXTVIEW("#endif") << NewLine;
 			}
 
 			// Print the user-specified include directives
 			for (FStringView Include : Function->Includes)
 			{
-				Printer << TEXTVIEW("#include \"") << Include << TEXTVIEW("\"") << NewLine;
+				Printer << TEXTVIEW("#include \"") << Include << L'\"' << NewLine;
 			}
 
 			// Write the custom function signature, for example "C5_MyCustomNode".
 			// - 'C' is only a "namespace" for custom functions
 			// - '5' is a unique id used to disambiguate distinct custom functions with the same name.
-			Printer << LowerType(Function->ReturnType) << TEXTVIEW(" C") << Function->UniqueId << TEXT("_") << Function->Name << BeginArgs << TEXT("FMaterialPixelParameters Parameters");
+			Printer << LowerType(Function->ReturnType) << TEXTVIEW(" C") << Function->UniqueId << L'_' << Function->Name << BeginArgs << TEXT("FMaterialPixelParameters Parameters");
 			
 			// Write the parameter declarations
 			for (uint16 i = 0; i < Function->NumParameters; ++i)
@@ -2235,7 +2275,7 @@ struct FPrivate : FMaterialIRToHLSLTranslation
 						: TEXTVIEW("out "));
 
 				// Type and name
-				Printer << LowerType(Function->Parameters[i].Type) << TEXTVIEW(" ") << *Function->Parameters[i].Name.ToString();
+				Printer << LowerType(Function->Parameters[i].Type) << L' ' << *Function->Parameters[i].Name.ToString();
 			}
 
 			Printer << EndArgs << NewLine << OpenBrace;
@@ -2244,7 +2284,7 @@ struct FPrivate : FMaterialIRToHLSLTranslation
 			bool bContainsReturn = Function->Code.Contains(TEXT("return"));
 			if (!bContainsReturn)
 			{
-				Printer << TEXT("return") << NewLine;
+				Printer << TEXTVIEW("return") << NewLine;
 			}
 			
 			// Write the function code
@@ -2252,7 +2292,7 @@ struct FPrivate : FMaterialIRToHLSLTranslation
 
 			if (!bContainsReturn)
 			{
-				Printer << TEXT(";");
+				Printer << L';';
 			}
 
 			Printer << NewLine << CloseBrace << TEXTVIEW("\n\n");
@@ -2490,6 +2530,7 @@ void FMaterialIRToHLSLTranslation::Run(TMap<FString, FString>& OutParameters, FS
 	Private.GenerateVertexStageHLSL();
 	Private.GenerateNonVertexStageHLSL(MIR::Stage_Pixel);
 	Private.GenerateNonVertexStageHLSL(MIR::Stage_Compute);
+	Private.GenerateCustomOutputs(OutParameters);
 	Private.GenerateTemplateStringParameters(OutParameters);
 	Private.GenerateCustomFunctionsHLSL(OutParameters);
 	Private.GetShaderCompilerEnvironment(OutEnvironment);
diff --git a/Engine/Source/Runtime/Engine/Public/Materials/MaterialExpressionTangentOutput.h b/Engine/Source/Runtime/Engine/Public/Materials/MaterialExpressionTangentOutput.h
--- a/Engine/Source/Runtime/Engine/Public/Materials/MaterialExpressionTangentOutput.h
+++ b/Engine/Source/Runtime/Engine/Public/Materials/MaterialExpressionTangentOutput.h
@@ -18,6 +18,7 @@ class UMaterialExpressionTangentOutput : public UMaterialExpressionCustomOutput
 	FExpressionInput Input;
 
 #if WITH_EDITOR
+	virtual void Build(MIR::FEmitter& Em) override;
 	virtual int32 Compile(class FMaterialCompiler* Compiler, int32 OutputIndex) override;
 	virtual void GetCaption(TArray<FString>& OutCaptions) const override;
 	virtual EMaterialValueType GetInputValueType(int32 InputIndex) override { return MCT_Float; }
diff --git a/Engine/Source/Runtime/Engine/Public/Materials/MaterialIR.h b/Engine/Source/Runtime/Engine/Public/Materials/MaterialIR.h
--- a/Engine/Source/Runtime/Engine/Public/Materials/MaterialIR.h
+++ b/Engine/Source/Runtime/Engine/Public/Materials/MaterialIR.h
@@ -28,6 +28,7 @@ struct FSimpleStringView
 	FString ToString() const { return FString{ ToView() }; }
 	operator FStringView() const { return ToView(); }
 	FSimpleStringView& operator=(FStringView View) { Ptr = View.GetData(); Len = View.Len(); return *this; }
+	const TCHAR* operator*() const { return Ptr; }
 };
 
 // Identifies the block of execution in which an instruction runs.
@@ -694,11 +695,18 @@ struct TComposite : FComposite
 // Instruction that sets material attribute (i.e., "BaseColor") to its value.
 struct FSetMaterialOutput : TInstruction<VK_SetMaterialOutput, 1>
 {
+	// The material attribute to set.
+	EMaterialProperty Property;
+
+	// The name of this material output (e.g. "BaseColor")
+	FSimpleStringView Name;
+
 	// The value this material attribute should be set to.
 	FValue* Arg;
 
-	// The material attribute to set.
-	EMaterialProperty Property;
+	// Values in a module EntryPoint are sorted by this priority index (lowest to highest).
+	// This value is used to give an order of evaluation of SetMaterialOutput instructions in an entry point.
+	int Priority;
 };
 
 // Inline HLSL instruction.
@@ -1134,6 +1142,8 @@ struct FFunctionHLSLDefine
 	
 	// The define value.
 	FStringView Value;
+
+	bool operator==(const FFunctionHLSLDefine& Other) const = default;
 };
 
 // A user-function defined with user-provided HLSL code.
diff --git a/Engine/Source/Runtime/Engine/Public/Materials/MaterialIREmitter.h b/Engine/Source/Runtime/Engine/Public/Materials/MaterialIREmitter.h
--- a/Engine/Source/Runtime/Engine/Public/Materials/MaterialIREmitter.h
+++ b/Engine/Source/Runtime/Engine/Public/Materials/MaterialIREmitter.h
@@ -123,8 +123,8 @@ struct FAttributeAssignment
 	FValueRef Value;
 };
 
-// Describes a user-specified HLSL function.
-struct FFunctionHLSLDesc
+// Describes a user-specified function.
+struct FFunctionDesc
 {
 	// The function name, used for debugging purposes and referenced in the generated HLSL. 
 	FStringView Name{};
@@ -144,9 +144,6 @@ struct FFunctionHLSLDesc
 	// Array of descriptions of the function parameters.
 	FFunctionParameter Parameters[MIR::MaxNumFunctionParameters];
 
-	// The user-specified HLSL code inside the function body.
-	FStringView Code{};
-
 	// Array of user-specified defines to declare before this function.
 	TConstArrayView<FFunctionHLSLDefine> Defines;
 
@@ -166,6 +163,13 @@ struct FFunctionHLSLDesc
 	uint32 GetNumParameters() const { return NumInputOnlyParams + NumInputOutputParams + NumOutputOnlyParams; }
 };
 
+//
+struct FFunctionHLSLDesc : FFunctionDesc
+{
+	// The user-specified HLSL code inside the function body.
+	FStringView Code{};
+};
+
 // Helper structure to simplify passing attributes to texture sampling emitter functions that share the same attributes.
 struct FTextureSampleBaseAttributes
 {
@@ -183,6 +187,13 @@ struct FTextureSampleBaseAttributes
 	bool bIsAdaptive : 1 { false };
 };
 
+// Used to distinguish whether a value has a per-vertex or per-pixel frequency.
+enum EFrequency
+{
+	Frequency_PerVertex,
+	Frequency_PerPixel,
+};
+
 // The FEmitter is a helper object responsible for emitting MIR values and instructions.
 // It is primarily used within UMaterialExpression::Build() implementations to lower the
 // high-level semantics of a material expression into corresponding lower-level MIR nodes.
@@ -446,11 +457,19 @@ public:
 	FValueRef ScreenTexture(EScreenTexture TextureKind);
 
 	/*----------------------------- Instructions emission -----------------------------*/
-
+	
 	// Emits an instruction that sets the material attribute with given property to given argument.
-	// Note: this API is used by the builder and you should not use it manually.
-	FSetMaterialOutput* SetMaterialOutput(EMaterialProperty InProperty, FValue* Arg);
-
+	// Note: This API is used by the builder and you should not use it manually.
+	void Private_SetMaterialOutput(EMaterialProperty Property, FValueRef Arg);
+	
+	// Emits a group of custom outputs.
+	// This function emits a new custom output group made of a number of custom outputs. The function emits a number
+	// of SetMaterialOutput instruction for each argument in Args, giving each the specified group name appended by
+	// the output index (e.g. "TangentOutput0" for a group named "TangentOutput").
+	// The type of each custom output is inferred by the type of the argument value it is set to, therefore make sure
+	// you ::Cast() each output argument to the desired type before calling this function.
+	void SetCustomOutputs(FStringView Name, TConstArrayView<FValueRef> Args);
+	
 	// Emits a 2D column vector from X and Y values.
 	FValueRef Vector2(FValueRef InX, FValueRef InY);
 
@@ -648,7 +667,7 @@ public:
 
 	// Emits an instruction that retrieves the output parameter value from a function call.
 	FValueRef CallParameterOutput(FValueRef Call, uint32 ParameterIndex);
-
+	
 	/*--------------------------- Unary operators shortcuts ----------------------------*/
 
 	FValueRef BitwiseNot(FValueRef A) { return Operator(UO_BitwiseNot, A); }
diff --git a/Engine/Source/Runtime/Engine/Public/Materials/MaterialIRModule.h b/Engine/Source/Runtime/Engine/Public/Materials/MaterialIRModule.h
--- a/Engine/Source/Runtime/Engine/Public/Materials/MaterialIRModule.h
+++ b/Engine/Source/Runtime/Engine/Public/Materials/MaterialIRModule.h
@@ -62,7 +62,22 @@ public:
 		MIR::FBlock RootBlock;
 
 		// An array of terminal values this entry point evaluates.
-		TArrayView<MIR::FValue*> Outputs;
+		TArray<MIR::FValue*> Outputs;
+	};
+
+	// Models a group of custom outputs contained in the module.
+	// Custom outputs are user-created and supplement the base outputs in the material (Normal, BaseColor, etc)
+	struct FCustomOutputGroup
+	{
+		// Name of the custom output group, e.g. "TangentOutput".
+		FStringView Name;
+		
+		// The number of outputs in this group.
+		int32 NumOutputs;
+
+		// The index in FMaterialIRModule::GetCustomOutputs() of the first output instruction in this group.
+		// E.g. the index of "TangentOutput0"
+		int32 CustomOutputInstructionsBegin;
 	};
 
 public:
@@ -91,7 +106,7 @@ public:
 	const FMaterialCompilationOutput& GetCompilationOutput() const { return CompilationOutput; }
 
 	// Adds a new entry point with specified name, stage and expected number of evaluated output values.
-	int32 AddEntryPoint(FStringView Name, MIR::EStage Stage, int32 NumOutputs);
+	int32 AddEntryPoint(FStringView Name, MIR::EStage Stage, int32 NumOutputsHint = 0);
 
 	// Returns the number of procedures inthe module.
 	int32 GetNumEntryPoints() const { return EntryPoints.Num(); }
@@ -103,6 +118,18 @@ public:
 	// Sets the value of a material property for future reference.
 	void SetPropertyValue(EMaterialProperty Property, MIR::FValue* Value) { PropertyValues[Property] = Value; }
 	
+	// Records a custom output instruction.
+	void AddCustomOutput(const MIR::FSetMaterialOutput* Instr) { CustomOutputs.Push(Instr); }
+
+	// Returns the array custom output instructions.
+	TConstArrayView<const MIR::FSetMaterialOutput*> GetCustomOutputs() const { return CustomOutputs; }
+
+	//
+	void AddCustomOutputGroup(FStringView Name, int32 NumOutputs);
+
+	// Returns the array custom output groups.
+	TConstArrayView<FCustomOutputGroup> GetCustomOutputGroups() const { return CustomOutputGroups; }
+
 	// Retrieves the value a meterial property is assigned to.
 	const MIR::FValue* GetPropertyValue(EMaterialProperty Property) const { return PropertyValues[Property]; }
 
@@ -188,14 +215,20 @@ private:
 	// Compilation output data.
 	FMaterialCompilationOutput CompilationOutput;
 
-	// List of all the IR values contained in this module.
+	// Array of all the IR values contained in this module.
 	TArray<MIR::FValue*> Values;
 
 	// Array of entry points in this module.
 	TArray<FEntryPoint> EntryPoints;
 
-	// The final value assigned to each material property.
-	MIR::FValue* PropertyValues[MP_MAX];
+	// The final value assigned to each concrete material property.
+	MIR::FValue* PropertyValues[MP_MaterialAttributes];
+
+	// Array of user-generated custom outputs in this module. 
+	TArray<const MIR::FSetMaterialOutput*> CustomOutputs;
+
+	// Array of custom output groups in this module.
+	TArray<FCustomOutputGroup> CustomOutputGroups;
 
 	// Compilation statistics.
 	FStatistics Statistics;
@@ -206,19 +239,19 @@ private:
 	// Parameter metadata.
 	TArray<TPair<FMaterialParameterInfo, FMaterialParameterMetadata>> ParameterIdToData;
 
-	// Stores user-defined strings.	
+	// Stores user-defined strings in this module.
 	TArray<FString> UserStrings;
 
-	// Environment define names for shader compilation.
+	// Environment define names for shader compilation in this module.
 	TSet<FName> EnvironmentDefines;
 
-	// Parameter collections used by this material.
+	// Parameter collections referenced by this module.
 	TArray<UMaterialParameterCollection*> ParameterCollections;
 
 	// Array of HLSL-defined user-functions in the module.
 	TArray<const MIR::FFunctionHLSL*> FunctionHLSLs;
 
-	/** Shading models used by this material. */
+	// Shading models used by this material.
 	FMaterialShadingModelField ShadingModelsFromCompilation;
 
 	// List of compilation errors.
diff --git a/Engine/Source/Runtime/Engine/Public/Materials/MaterialIRTypes.h b/Engine/Source/Runtime/Engine/Public/Materials/MaterialIRTypes.h
--- a/Engine/Source/Runtime/Engine/Public/Materials/MaterialIRTypes.h
+++ b/Engine/Source/Runtime/Engine/Public/Materials/MaterialIRTypes.h
@@ -250,6 +250,9 @@ struct FType
 	// Returns true if the given type is a matrix primitive.
 	bool IsMatrix() const { return Is(ETypeKind::Primitive) && Primitive.IsMatrix(); }
 
+	// Returns whether this type is ShadingModel.
+	bool IsShadingModel() const { return Kind == ETypeKind::ShadingModel; }
+
 	// Returns whether this type is a Texture object.
 	bool IsTexture() const { return Kind == ETypeKind::Texture; }
 
