New Material Translator:
- Fixed a couple of bugs in MaterialExpressionsIR.cpp
- Applied a couple of suggestions from BugHawk (more consistent documentation comments).

#rnx

diff --git a/Engine/Source/Runtime/Engine/Private/Materials/MaterialExpressionsIR.cpp b/Engine/Source/Runtime/Engine/Private/Materials/MaterialExpressionsIR.cpp
--- a/Engine/Source/Runtime/Engine/Private/Materials/MaterialExpressionsIR.cpp
+++ b/Engine/Source/Runtime/Engine/Private/Materials/MaterialExpressionsIR.cpp
@@ -301,6 +301,7 @@ void UMaterialExpressionFunctionInput::Build(MIR::FEmitter& Em)
 		case FunctionInput_Bool:
 		case FunctionInput_StaticBool:
 			OutputValue = Em.ConstantBool(PreviewValue.X != 0.0f);
+			break;
 
 		case FunctionInput_Texture2D:
 		case FunctionInput_TextureCube:
@@ -419,12 +420,7 @@ void UMaterialExpressionChannelMaskParameter::Build(MIR::FEmitter& Em)
 	FValueRef DotResult = Em.Dot(
 		Em.CastToFloat(Em.Input(&Input), 4),
 		BuildMaterialExpressionParameter(Em, this));
-
 	Em.Output(0, DotResult);
-	Em.Output(1, Em.Subscript(DotResult, 1));
-	Em.Output(2, Em.Subscript(DotResult, 2));
-	Em.Output(3, Em.Subscript(DotResult, 3));
-	Em.Output(4, Em.Subscript(DotResult, 4));
 }
 
 void UMaterialExpressionStaticBoolParameter::Build(MIR::FEmitter& Em)
diff --git a/Engine/Source/Runtime/Engine/Private/Materials/MaterialIREmitter.cpp b/Engine/Source/Runtime/Engine/Private/Materials/MaterialIREmitter.cpp
--- a/Engine/Source/Runtime/Engine/Private/Materials/MaterialIREmitter.cpp
+++ b/Engine/Source/Runtime/Engine/Private/Materials/MaterialIREmitter.cpp
@@ -1366,7 +1366,7 @@ void FEmitter::SetCustomOutputs(FStringView Name, TConstArrayView<FValueRef> Arg
 	{
 		if (Group.Name == Name)
 		{
-			Errorf(TEXT("A custom output with name '%s' already exists."), Name.GetData());
+			Errorf(TEXT("A custom output with name '%.*s' already exists."), Name.Len(), Name.GetData());
 			return;
 		}
 	}
diff --git a/Engine/Source/Runtime/Engine/Public/Materials/MaterialIREmitter.h b/Engine/Source/Runtime/Engine/Public/Materials/MaterialIREmitter.h
--- a/Engine/Source/Runtime/Engine/Public/Materials/MaterialIREmitter.h
+++ b/Engine/Source/Runtime/Engine/Public/Materials/MaterialIREmitter.h
@@ -467,7 +467,7 @@ public:
 	// of SetMaterialOutput instruction for each argument in Args, giving each the specified group name appended by
 	// the output index (e.g. "TangentOutput0" for a group named "TangentOutput").
 	// The type of each custom output is inferred by the type of the argument value it is set to, therefore make sure
-	// you ::Cast() each output argument to the desired type before calling this function.
+	// you cast each output argument to the desired type before calling this function.
 	void SetCustomOutputs(FStringView Name, TConstArrayView<FValueRef> Args);
 	
 	// Emits a 2D column vector from X and Y values.
diff --git a/Engine/Source/Runtime/Engine/Public/Materials/MaterialIRModule.h b/Engine/Source/Runtime/Engine/Public/Materials/MaterialIRModule.h
--- a/Engine/Source/Runtime/Engine/Public/Materials/MaterialIRModule.h
+++ b/Engine/Source/Runtime/Engine/Public/Materials/MaterialIRModule.h
@@ -105,7 +105,7 @@ public:
 	// Retrieves the material compilation output.
 	const FMaterialCompilationOutput& GetCompilationOutput() const { return CompilationOutput; }
 
-	// Adds a new entry point with specified name, stage and expected number of evaluated output values.
+	// Adds a new entry point with specified name, stage and optionally the expected number of evaluated output values hint.
 	int32 AddEntryPoint(FStringView Name, MIR::EStage Stage, int32 NumOutputsHint = 0);
 
 	// Returns the number of procedures inthe module.
