Added IsLevelTransformable in WorldPartitionRuntimeCellTransformer which can be overriden to skip entire cells from being transformed.

#FailedValidations Preflight

diff --git a/Engine/Source/Runtime/Engine/Public/WorldPartition/WorldPartition.h b/Engine/Source/Runtime/Engine/Public/WorldPartition/WorldPartition.h
--- a/Engine/Source/Runtime/Engine/Public/WorldPartition/WorldPartition.h
+++ b/Engine/Source/Runtime/Engine/Public/WorldPartition/WorldPartition.h
@@ -123,9 +123,9 @@ struct FRuntimeCellTransformerInstance
 	GENERATED_USTRUCT_BODY()
 
 #if WITH_EDITORONLY_DATA
-	inline void PreTransform(ULevel* InLevel) const { if (Instance && Instance->IsEnabled()) { Instance->PreTransform(InLevel); } }
-	inline void Transform(ULevel* InLevel) const { if (Instance && Instance->IsEnabled()) { Instance->Transform(InLevel); } }
-	inline void PostTransform(ULevel* InLevel) const { if (Instance && Instance->IsEnabled()) { Instance->PostTransform(InLevel); } }
+	inline void PreTransform(ULevel* InLevel) const { if (CanTransform(InLevel)) { Instance->PreTransform(InLevel); } }
+	inline void Transform(ULevel* InLevel) const { if (CanTransform(InLevel)) { Instance->Transform(InLevel); } }
+	inline void PostTransform(ULevel* InLevel) const { if (CanTransform(InLevel)) { Instance->PostTransform(InLevel); } }
 
 	/** Runtime cell transformer class */
 	UPROPERTY(EditAnywhere, Category = WorldPartitionSetup, AdvancedDisplay)
@@ -134,6 +134,9 @@ struct FRuntimeCellTransformerInstance
 	/** Transformer object instance */
 	UPROPERTY(VisibleAnywhere, Category = WorldPartitionSetup, Instanced, Meta = (EditCondition = "Class != nullptr", HideEditConditionToggle, NoResetToDefault))
 	TObjectPtr<UWorldPartitionRuntimeCellTransformer> Instance;
+
+private:
+	inline bool CanTransform(ULevel* InLevel) const { return Instance && Instance->IsEnabled() && Instance->IsLevelTransformable(InLevel); }
 #endif
 };
 
diff --git a/Engine/Source/Runtime/Engine/Public/WorldPartition/WorldPartitionRuntimeCellTransformer.h b/Engine/Source/Runtime/Engine/Public/WorldPartition/WorldPartitionRuntimeCellTransformer.h
--- a/Engine/Source/Runtime/Engine/Public/WorldPartition/WorldPartitionRuntimeCellTransformer.h
+++ b/Engine/Source/Runtime/Engine/Public/WorldPartition/WorldPartitionRuntimeCellTransformer.h
@@ -27,6 +27,8 @@ class UWorldPartitionRuntimeCellTransformer : public UObject
 
 #if WITH_EDITOR
 public:
+	/** Whether the level can be processed by the transformer. */
+	virtual bool IsLevelTransformable(ULevel* InLevel) const { return true; }
 	virtual void PreTransform(ULevel* InLevel) {}
 	virtual void Transform(ULevel* InLevel) {}
 	virtual void PostTransform(ULevel* InLevel) {}
