Fixed crash caused by actor being destroyed by life span timer while components are still in PreUnregistering state. Fix is to disable incremental pre-unregister components when a valid lifespan is set. Also, don't start life span timer when actor's level is being removed (unregistering incrementally).

diff --git a/Engine/Source/Runtime/Engine/Classes/GameFramework/Actor.h b/Engine/Source/Runtime/Engine/Classes/GameFramework/Actor.h
--- a/Engine/Source/Runtime/Engine/Classes/GameFramework/Actor.h
+++ b/Engine/Source/Runtime/Engine/Classes/GameFramework/Actor.h
@@ -1539,11 +1539,8 @@ private:
 	}
 
 	/** Whether actor supports incremental pre-unregister of components. */
-	virtual bool SupportsIncrementalPreUnregisterComponents() const
-	{
-		return true;
-	}
-
+	ENGINE_API virtual bool SupportsIncrementalPreUnregisterComponents() const;
+	
 	/** Incrementally pre-registers components associated with this actor (used during level streaming). */
 	ENGINE_API void IncrementalPreRegisterComponents();
 
diff --git a/Engine/Source/Runtime/Engine/Private/Actor.cpp b/Engine/Source/Runtime/Engine/Private/Actor.cpp
--- a/Engine/Source/Runtime/Engine/Private/Actor.cpp
+++ b/Engine/Source/Runtime/Engine/Private/Actor.cpp
@@ -6003,6 +6003,14 @@ static USceneComponent* GetNonPreRegisteredParent(UActorComponent* Component)
 } // namespace Private
 } // namespace UE
 
+bool AActor::SupportsIncrementalPreUnregisterComponents() const
+{
+	// Do not support incremental pre-unregister of components when a lifespan is set, 
+	// as this could cause UActorComponent::DestroyPhysicsState to be called 
+	// while the component is still in the pre-unregistering state (UActorComponent::IsPreUnregistering == true).
+	return !TimerHandle_LifeSpanExpired.IsValid();
+}
+
 void AActor::IncrementalPreUnregisterComponents()
 {
 	if (!SupportsIncrementalPreUnregisterComponents())
@@ -6520,7 +6528,18 @@ void AActor::SetLifeSpan( float InLifespan )
 	{
 		if( InLifespan > 0.0f)
 		{
-			GetWorldTimerManager().SetTimer( TimerHandle_LifeSpanExpired, this, &AActor::LifeSpanExpired, InLifespan );
+			// Do not start the lifespan timer if the level is already being removed from the world, 
+			// as this could cause UActorComponent::DestroyPhysicsState to be called while the component 
+			// is still in the pre-unregistering state (UActorComponent::IsPreUnregistering == true).
+			// This should not cause any issues since the actor is about to be destroyed.
+			if (GetLevel() && GetLevel()->bIsBeingRemoved)
+			{
+				UE_LOG(LogActor, Log, TEXT("AActor::SetLifeSpan call on Actor %s was ignored (level is already being removed)."), *GetName());
+			}
+			else
+			{
+				GetWorldTimerManager().SetTimer(TimerHandle_LifeSpanExpired, this, &AActor::LifeSpanExpired, InLifespan);
+			}
 		}
 		else
 		{
