Fix UWorldPartition::ResolveSubobject unable to resolve references to actor components.

diff --git a/Engine/Source/Runtime/Engine/Private/WorldPartition/WorldPartition.cpp b/Engine/Source/Runtime/Engine/Private/WorldPartition/WorldPartition.cpp
--- a/Engine/Source/Runtime/Engine/Private/WorldPartition/WorldPartition.cpp
+++ b/Engine/Source/Runtime/Engine/Private/WorldPartition/WorldPartition.cpp
@@ -1738,16 +1738,16 @@ bool UWorldPartition::ResolveSubobject(const TCHAR* SubObjectPath, UObject*& Out
 			FString SubObjectContext;	
 			if (!FString(SubObjectPath).Split(TEXT("."), &SubObjectContext, &SubObjectName))
 			{
-				SubObjectName = SubObjectPath;
+				SubObjectContext = SubObjectPath;
 			}
-
-			if (const FWorldPartitionActorDescInstance* ActorDescInstance = GetActorDescInstanceByPath(SubObjectName))
+ 
+			if (const FWorldPartitionActorDescInstance* ActorDescInstance = GetActorDescInstanceByPath(SubObjectContext))
 			{
 				if (bLoadIfExists)
 				{
 					LoadedSubobjects.Emplace(this, ActorDescInstance->GetGuid());
 				}
-
+ 
 				OutObject = StaticFindObject(UObject::StaticClass(), GetWorld()->PersistentLevel, SubObjectPath);
 				return true;
 			}
