Improve Mask Nodes tooltip information to reflect which ones lead to new permutations.
The changes have been made to the ChannelMask Parameter node and the StaticComponentMask Parameter node.
The Component Mask node(non static) is not a parameter that can be exposed, thus it cannot create permutations.

[RN] minor | Rendering Architecture

diff --git a/Engine/Source/Runtime/Engine/Private/Materials/MaterialExpressions.cpp b/Engine/Source/Runtime/Engine/Private/Materials/MaterialExpressions.cpp
--- a/Engine/Source/Runtime/Engine/Private/Materials/MaterialExpressions.cpp
+++ b/Engine/Source/Runtime/Engine/Private/Materials/MaterialExpressions.cpp
@@ -5414,6 +5414,18 @@ void UMaterialExpressionCrossProduct::GetCaption(TArray<FString>& OutCaptions) c
 #endif // WITH_EDITOR
 
 #if WITH_EDITOR
+static FString BuildStaticComponentMaskOutputText(const bool R, const bool G, const bool B, const bool A)
+{
+	FString Str(TEXT("( "));
+	if (R) Str += TEXT("R ");
+	if (G) Str += TEXT("G ");
+	if (B) Str += TEXT("B ");
+	if (A) Str += TEXT("A ");
+	Str += TEXT(")");
+
+	return Str;
+}
+
 int32 UMaterialExpressionComponentMask::Compile(class FMaterialCompiler* Compiler, int32 OutputIndex)
 {
 	if(!Input.GetTracedInput().Expression)
@@ -5432,13 +5444,8 @@ int32 UMaterialExpressionComponentMask::Compile(class FMaterialCompiler* Compile
 
 void UMaterialExpressionComponentMask::GetCaption(TArray<FString>& OutCaptions) const
 {
-	FString Str(TEXT("Mask ("));
-	if ( R ) Str += TEXT(" R");
-	if ( G ) Str += TEXT(" G");
-	if ( B ) Str += TEXT(" B");
-	if ( A ) Str += TEXT(" A");
-	Str += TEXT(" )");
-	OutCaptions.Add(Str);
+	const FString Caption = FString(TEXT("Mask ")) + BuildStaticComponentMaskOutputText(R, G, B, A);
+	OutCaptions.Add(Caption);
 }
 #endif // WITH_EDITOR
 
@@ -5467,8 +5474,17 @@ int32 UMaterialExpressionStaticComponentMaskParameter::Compile(class FMaterialCo
 
 void UMaterialExpressionStaticComponentMaskParameter::GetCaption(TArray<FString>& OutCaptions) const
 {
-	OutCaptions.Add(TEXT("Mask Param")); 
-	OutCaptions.Add(FString::Printf(TEXT("'%s'"), *ParameterName.ToString()));
+	const FString Caption = FString(TEXT("Mask ")) + BuildStaticComponentMaskOutputText(DefaultR, DefaultG, DefaultB, DefaultA);
+	OutCaptions.Add(Caption);
+}
+
+void UMaterialExpressionStaticComponentMaskParameter::GetExpressionToolTip(TArray<FString>& OutToolTip)
+{
+	const FString MaskOutput = BuildStaticComponentMaskOutputText(DefaultR, DefaultG, DefaultB, DefaultA);
+
+	OutToolTip.Add(TEXT("Outputs the selected channels."));
+	OutToolTip.Add(TEXT("Note: This DOES create new permutations."));
+	OutToolTip.Add(FString::Printf(TEXT("Current output: %s "), *MaskOutput));
 }
 
 bool  UMaterialExpressionStaticComponentMaskParameter::SetParameterValue(FName InParameterName, bool InR, bool InG, bool InB, bool InA, FGuid InExpressionGuid, EMaterialExpressionSetParameterValueFlags Flags)
@@ -8332,8 +8348,32 @@ void UMaterialExpressionChannelMaskParameter::GetCaption(TArray<FString>& OutCap
 		OutCaptions.Add(TEXT("Alpha")); break;
 	}
 
-	OutCaptions.Add(FString::Printf(TEXT("'%s'"), *ParameterName.ToString())); 
+	OutCaptions.Add(FString::Printf(TEXT("'%s'"), *ParameterName.ToString()));
+}
+
+
+void UMaterialExpressionChannelMaskParameter::GetExpressionToolTip(TArray<FString>& OutToolTip)
+{
+	const TCHAR* Output = nullptr;
+	switch (MaskChannel)
+	{
+	case EChannelMaskParameterColor::Red:
+		Output = TEXT("Red"); break;
+	case EChannelMaskParameterColor::Green:
+		Output = TEXT("Green"); break;
+	case EChannelMaskParameterColor::Blue:
+		Output = TEXT("Blue"); break;
+	case EChannelMaskParameterColor::Alpha:
+		Output = TEXT("Alpha"); break;
+	default:
+		checkf(false, TEXT("Invalid MaskChannel %d"), MaskChannel); break;
+	}
+
+	OutToolTip.Add(TEXT("Outputs the selected channel."));
+	OutToolTip.Add(TEXT("Note: This does NOT create new permutations."));
+	OutToolTip.Add(FString::Printf(TEXT("Current output: %s "), Output));
 }
+
 #endif // WITH_EDITOR
 
 //
diff --git a/Engine/Source/Runtime/Engine/Public/Materials/MaterialExpressionChannelMaskParameter.h b/Engine/Source/Runtime/Engine/Public/Materials/MaterialExpressionChannelMaskParameter.h
--- a/Engine/Source/Runtime/Engine/Public/Materials/MaterialExpressionChannelMaskParameter.h
+++ b/Engine/Source/Runtime/Engine/Public/Materials/MaterialExpressionChannelMaskParameter.h
@@ -37,7 +37,8 @@ class UMaterialExpressionChannelMaskParameter : public UMaterialExpressionVector
 	virtual void Build(MIR::FEmitter& Emitter) override;
 	virtual int32 Compile(class FMaterialCompiler* Compiler, int32 OutputIndex) override;
 	virtual void GetCaption(TArray<FString>& OutCaptions) const override;
-	
+	virtual void GetExpressionToolTip(TArray<FString>& OutToolTip) override;
+
 	virtual bool IsInputConnectionRequired(int32 InputIndex) const override {return true;}
 	virtual EMaterialValueType GetInputValueType(int32 InputIndex) override {return MCT_Float4;}
 
diff --git a/Engine/Source/Runtime/Engine/Public/Materials/MaterialExpressionStaticComponentMaskParameter.h b/Engine/Source/Runtime/Engine/Public/Materials/MaterialExpressionStaticComponentMaskParameter.h
--- a/Engine/Source/Runtime/Engine/Public/Materials/MaterialExpressionStaticComponentMaskParameter.h
+++ b/Engine/Source/Runtime/Engine/Public/Materials/MaterialExpressionStaticComponentMaskParameter.h
@@ -39,6 +39,7 @@ public:
 	virtual void Build(MIR::FEmitter& Emitter) override;
 	virtual int32 Compile(class FMaterialCompiler* Compiler, int32 OutputIndex) override;
 	virtual void GetCaption(TArray<FString>& OutCaptions) const override;
+	virtual void GetExpressionToolTip(TArray<FString>& OutToolTip) override;
 	virtual bool IsStaticExpression() const override { return true; }
 	virtual bool GetParameterValue(FMaterialParameterMetadata& OutMeta) const override
 	{
