ComponentInterfaces cleanup
- COMPONENT_INTERFACE macros should conform to the coding standard.
- Moving UE_DECLARE_COMPONENT_ACTOR_INTERFACE uses to the top of the classes.

diff --git a/Engine/Source/Runtime/Engine/Classes/Components/ComponentInterfaces.h b/Engine/Source/Runtime/Engine/Classes/Components/ComponentInterfaces.h
--- a/Engine/Source/Runtime/Engine/Classes/Components/ComponentInterfaces.h
+++ b/Engine/Source/Runtime/Engine/Classes/Components/ComponentInterfaces.h
@@ -178,44 +178,54 @@ protected:
 // 
 // Example...
 // 
-//		class FActorSomeInteface : ISomeInterface
+//		class FActorFooComponentInterface : IFooComponent
 //		{
 //			// Nothing but the overrides
 //			virtual void OverrideSomething() override;
 //		};
 //		
-//		class UHostClass
+//		class UFooComponent
 //		{
-//			UE_ComponentInterfaceDeclaration(SomeInterface, FActor);
+//			UE_DECLARE_COMPONENT_INTERFACE(FooComponent, FActor);
 //		}
 // 
-//	    void FActorSomeInterface::OverrideSomething()
+//	    void FActorFooComponentInterface::OverrideSomething()
 //		{
-//			UHostClass::GetHostClass(this)->OverrideSomethingImplementer();
+//			UFooComponent::GetFooComponent(this)->OverrideSomethingImplementer();
 //		}
 //
 
-
-
-#define UE_DECLARE_COMPONENT_INTERFACE_INTERNAL(actorcomponenttype, componentinterfacetype, actorcomponentinterfacetype, actorcomponentinterfacemember, interfacename )\
+#define UE_DECLARE_COMPONENT_INTERFACE_INTERNAL(ThisType, ComponentInterfaceType, InterfaceType, InterfaceMember, InterfaceName)\
 	public:\
-		componentinterfacetype* Get##interfacename ##Interface() const { return (componentinterfacetype*)&actorcomponentinterfacemember; } \
-	protected: \
-		actorcomponentinterfacetype	actorcomponentinterfacemember;\
+		ComponentInterfaceType* Get##InterfaceName##Interface()\
+		{\
+			return &InterfaceMember;\
+		}\
+		const ComponentInterfaceType* Get##InterfaceName##Interface() const\
+		{\
+			return &InterfaceMember;\
+		}\
+	protected:\
+		InterfaceType InterfaceMember;\
 		\
-		static actorcomponenttype* Get##interfacename(actorcomponentinterfacetype* InImpl)\
+		static uintptr_t GetAddress(const InterfaceType* InImpl)\
 		{\
-			return (actorcomponenttype*)(((size_t)InImpl) - offsetof(actorcomponenttype, actorcomponentinterfacemember));\
+			return reinterpret_cast<uintptr_t>(InImpl) - offsetof(ThisType, InterfaceMember);\
 		}\
-		static const actorcomponenttype* Get##interfacename(const actorcomponentinterfacetype* InImpl)\
+		static ThisType* Get##InterfaceName(InterfaceType* InImpl)\
 		{\
-			return (const actorcomponenttype*)(((size_t)InImpl) - offsetof(actorcomponenttype, actorcomponentinterfacemember));\
+			return reinterpret_cast<ThisType*>(GetAddress(InImpl));\
 		}\
-		friend class actorcomponentinterfacetype;
-
+		static const ThisType* Get##InterfaceName(const InterfaceType* InImpl)\
+		{\
+			return reinterpret_cast<const ThisType*>(GetAddress(InImpl));\
+		}\
+		friend class InterfaceType;
 
-#define UE_DECLARE_COMPONENT_INTERFACE(name, baseprefix)  UE_DECLARE_COMPONENT_INTERFACE_INTERNAL(U##name, I##name , PREPROCESSOR_JOIN(PREPROCESSOR_JOIN(baseprefix, name), Interface), PREPROCESSOR_JOIN(name, Interface), name)
+#define UE_DECLARE_COMPONENT_INTERFACE(ComponentName, BasePrefix)\
+	UE_DECLARE_COMPONENT_INTERFACE_INTERNAL(U##ComponentName, I##ComponentName , BasePrefix##ComponentName##Interface, ComponentName##Interface, ComponentName)
 
-#define UE_DECLARE_COMPONENT_ACTOR_INTERFACE(name)  UE_DECLARE_COMPONENT_INTERFACE(name, FActor)
+#define UE_DECLARE_COMPONENT_ACTOR_INTERFACE(ComponentName)\
+	UE_DECLARE_COMPONENT_INTERFACE(ComponentName, FActor)
 
 #pragma endregion
\ No newline at end of file
diff --git a/Engine/Source/Runtime/Engine/Classes/Components/SkinnedMeshComponent.h b/Engine/Source/Runtime/Engine/Classes/Components/SkinnedMeshComponent.h
--- a/Engine/Source/Runtime/Engine/Classes/Components/SkinnedMeshComponent.h
+++ b/Engine/Source/Runtime/Engine/Classes/Components/SkinnedMeshComponent.h
@@ -258,7 +258,9 @@ UCLASS(hidecategories=Object, config=Engine, editinlinenew, abstract, MinimalAPI
 class USkinnedMeshComponent : public UMeshComponent, public ILODSyncInterface, public IClothSimulationDataProvider
 {
 	GENERATED_UCLASS_BODY()
+	UE_DECLARE_COMPONENT_ACTOR_INTERFACE(SkinnedMeshComponent)
 
+public:
 	/** Access granted to the render state recreator in order to trigger state rebuild */
 	friend class FSkinnedMeshComponentRecreateRenderStateContext;
 
@@ -1280,9 +1282,6 @@ protected:
 	ENGINE_API virtual UObject const* AdditionalStatObject() const override;
 	//~ End UActorComponent Interface
 
-public:
-	UE_DECLARE_COMPONENT_ACTOR_INTERFACE(SkinnedMeshComponent)
-
 public:
 	//~ Begin USceneComponent Interface
 #if WITH_EDITOR
diff --git a/Engine/Source/Runtime/Engine/Classes/Components/StaticMeshComponent.h b/Engine/Source/Runtime/Engine/Classes/Components/StaticMeshComponent.h
--- a/Engine/Source/Runtime/Engine/Classes/Components/StaticMeshComponent.h
+++ b/Engine/Source/Runtime/Engine/Classes/Components/StaticMeshComponent.h
@@ -102,7 +102,9 @@ UCLASS(Blueprintable, ClassGroup=(Rendering, Common), hidecategories=(Object,Act
 class UStaticMeshComponent : public UMeshComponent
 {
 	GENERATED_UCLASS_BODY()
+	UE_DECLARE_COMPONENT_ACTOR_INTERFACE(StaticMeshComponent)
 
+public:
 	/** If 0, auto-select LOD level. if >0, force to (ForcedLodModel-1). */
 	UPROPERTY(EditAnywhere, AdvancedDisplay, BlueprintReadOnly, Category=LOD)
 	int32 ForcedLodModel;
@@ -549,8 +551,6 @@ public:
 #endif
 	//~ End USceneComponent Interface
 
-	UE_DECLARE_COMPONENT_ACTOR_INTERFACE(StaticMeshComponent)	
-
 	//~ Begin UActorComponent Interface.
 protected: 
 	ENGINE_API virtual void OnRegister() override;
diff --git a/Engine/Source/Runtime/Engine/Public/PrimitiveSceneProxyDesc.h b/Engine/Source/Runtime/Engine/Public/PrimitiveSceneProxyDesc.h
--- a/Engine/Source/Runtime/Engine/Public/PrimitiveSceneProxyDesc.h
+++ b/Engine/Source/Runtime/Engine/Public/PrimitiveSceneProxyDesc.h
@@ -26,7 +26,7 @@ struct FPrimitiveSceneProxyDesc
 	mutable TArray<uint32> ActorOwners;
 	const FCustomPrimitiveData* CustomPrimitiveData = nullptr;
 	FSceneInterface* Scene = nullptr;
-	IPrimitiveComponent* PrimitiveComponentInterface = nullptr;
+	const IPrimitiveComponent* PrimitiveComponentInterface = nullptr;
 
 	const UObject* AdditionalStatObjectPtr = nullptr;
 
@@ -260,7 +260,7 @@ struct FPrimitiveSceneProxyDesc
 		return StatId;
 	}
 
-	IPrimitiveComponent* GetPrimitiveComponentInterface() const
+	const IPrimitiveComponent* GetPrimitiveComponentInterface() const
 	{
 		return PrimitiveComponentInterface;
 	}
