Moving FActorPrimitiveComponentInterface, FActorStaticMeshComponentInterface, and FActorSkinnedMeshComponentInterface definitions to their own files.

diff --git a/Engine/Source/Runtime/Engine/Private/Components/ActorPrimitiveComponentInterface.cpp b/Engine/Source/Runtime/Engine/Private/Components/ActorPrimitiveComponentInterface.cpp
new file mode 100644
--- /dev/null
+++ b/Engine/Source/Runtime/Engine/Private/Components/ActorPrimitiveComponentInterface.cpp
@@ -0,0 +1,151 @@
+// Copyright Epic Games, Inc. All Rights Reserved.
+
+#include "Components/ActorPrimitiveComponentInterface.h"
+#include "Components/PrimitiveComponent.h"
+#include "Materials/MaterialInterface.h"
+
+bool FActorPrimitiveComponentInterface::IsRenderStateCreated() const
+{
+	return UPrimitiveComponent::GetPrimitiveComponent(this)->IsRenderStateCreated();
+}
+
+bool FActorPrimitiveComponentInterface::IsRenderStateDirty() const
+{
+	return UPrimitiveComponent::GetPrimitiveComponent(this)->IsRenderStateDirty();
+}
+
+bool FActorPrimitiveComponentInterface::ShouldCreateRenderState() const
+{
+	return UPrimitiveComponent::GetPrimitiveComponent(this)->ShouldCreateRenderState();
+}
+
+bool FActorPrimitiveComponentInterface::IsRegistered() const
+{
+	return UPrimitiveComponent::GetPrimitiveComponent(this)->IsRegistered();
+}
+
+bool FActorPrimitiveComponentInterface::IsUnreachable() const
+{
+	return UPrimitiveComponent::GetPrimitiveComponent(this)->IsUnreachable();
+}
+
+UWorld* FActorPrimitiveComponentInterface::GetWorld() const
+{
+	return UPrimitiveComponent::GetPrimitiveComponent(this)->GetWorld();
+}
+
+FSceneInterface* FActorPrimitiveComponentInterface::GetScene() const
+{
+	return UPrimitiveComponent::GetPrimitiveComponent(this)->GetScene();
+}
+
+FPrimitiveSceneProxy* FActorPrimitiveComponentInterface::GetSceneProxy() const
+{
+	return UPrimitiveComponent::GetPrimitiveComponent(this)->GetSceneProxy();
+}
+
+void FActorPrimitiveComponentInterface::GetUsedMaterials(TArray<UMaterialInterface*>& OutMaterials, bool bGetDebugMaterials) const
+{
+	UPrimitiveComponent::GetPrimitiveComponent(this)->GetUsedMaterials(OutMaterials, bGetDebugMaterials);
+}
+
+void FActorPrimitiveComponentInterface::GetStreamableRenderAssetInfo(TArray<FStreamingRenderAssetPrimitiveInfo>& StreamableRenderAssets) const
+{
+	FStreamingTextureLevelContext LevelContext{ GetCurrentMaterialQualityLevelChecked(),UPrimitiveComponent::GetPrimitiveComponent(this) };
+	UPrimitiveComponent::GetPrimitiveComponent(this)->GetStreamingRenderAssetInfo(LevelContext, StreamableRenderAssets);
+}
+
+void FActorPrimitiveComponentInterface::MarkRenderStateDirty()
+{
+	UPrimitiveComponent::GetPrimitiveComponent(this)->MarkRenderStateDirty();
+}
+
+void FActorPrimitiveComponentInterface::DestroyRenderState()
+{
+	UPrimitiveComponent::GetPrimitiveComponent(this)->DestroyRenderState_Concurrent();
+}
+
+void FActorPrimitiveComponentInterface::CreateRenderState(FRegisterComponentContext* Context)
+{
+	UPrimitiveComponent::GetPrimitiveComponent(this)->CreateRenderState_Concurrent(Context);
+}
+
+FString FActorPrimitiveComponentInterface::GetName() const
+{
+	return UPrimitiveComponent::GetPrimitiveComponent(this)->GetName();
+}
+
+FString FActorPrimitiveComponentInterface::GetFullName() const
+{
+	return UPrimitiveComponent::GetPrimitiveComponent(this)->GetFullName();
+}
+
+FTransform FActorPrimitiveComponentInterface::GetTransform() const
+{
+	return UPrimitiveComponent::GetPrimitiveComponent(this)->GetComponentTransform();
+}
+
+FBoxSphereBounds FActorPrimitiveComponentInterface::GetBounds() const
+{
+	return UPrimitiveComponent::GetPrimitiveComponent(this)->Bounds;
+}
+
+float FActorPrimitiveComponentInterface::GetLastRenderTimeOnScreen() const
+{
+	return UPrimitiveComponent::GetPrimitiveComponent(this)->GetLastRenderTimeOnScreen();
+}
+
+void FActorPrimitiveComponentInterface::GetPrimitiveStats(FPrimitiveStats& PrimitiveStats) const
+{
+	return UPrimitiveComponent::GetPrimitiveComponent(this)->GetPrimitiveStats(PrimitiveStats);
+}
+
+UObject* FActorPrimitiveComponentInterface::GetUObject()
+{
+	return UPrimitiveComponent::GetPrimitiveComponent(this);
+}
+
+const UObject* FActorPrimitiveComponentInterface::GetUObject() const
+{
+	return UPrimitiveComponent::GetPrimitiveComponent(this);
+}
+
+UObject* FActorPrimitiveComponentInterface::GetOwner() const
+{
+	return UPrimitiveComponent::GetPrimitiveComponent(this)->GetOwner();
+}
+
+FString FActorPrimitiveComponentInterface::GetOwnerName() const
+{
+	const UPrimitiveComponent* Component = UPrimitiveComponent::GetPrimitiveComponent(this);
+
+#if ACTOR_HAS_LABELS
+	if (Component->GetOwner())
+	{
+		return Component->GetOwner()->GetActorNameOrLabel();
+	}
+#endif
+
+	return Component->GetName();
+}
+
+FPrimitiveSceneProxy* FActorPrimitiveComponentInterface::CreateSceneProxy()
+{
+	UPrimitiveComponent* Component = UPrimitiveComponent::GetPrimitiveComponent(this);
+	FPrimitiveSceneProxy* SceneProxy = Component->CreateSceneProxy();
+	Component->AssignSceneProxy(SceneProxy);
+	return SceneProxy;
+}
+
+void FActorPrimitiveComponentInterface::PrecachePSOs()
+{
+	UPrimitiveComponent::GetPrimitiveComponent(this)->PrecachePSOs();
+}
+
+#if WITH_EDITOR
+HHitProxy* FActorPrimitiveComponentInterface::CreateMeshHitProxy(int32 SectionIndex, int32 MaterialIndex)
+{
+	UPrimitiveComponent* Component = UPrimitiveComponent::GetPrimitiveComponent(this);
+	return Component->CreateMeshHitProxy(SectionIndex, MaterialIndex);
+}
+#endif
diff --git a/Engine/Source/Runtime/Engine/Private/Components/ActorSkinnedMeshComponentInterface.cpp b/Engine/Source/Runtime/Engine/Private/Components/ActorSkinnedMeshComponentInterface.cpp
new file mode 100644
--- /dev/null
+++ b/Engine/Source/Runtime/Engine/Private/Components/ActorSkinnedMeshComponentInterface.cpp
@@ -0,0 +1,27 @@
+// Copyright Epic Games, Inc. All Rights Reserved.
+
+#include "Components/ActorSkinnedMeshComponentInterface.h"
+#include "Components/SkinnedMeshComponent.h"
+#include "ComponentReregisterContext.h"
+
+#if WITH_EDITOR
+void FActorSkinnedMeshComponentInterface::PostAssetCompilation()
+{
+	USkinnedMeshComponent::GetSkinnedMeshComponent(this)->PostAssetCompilation();
+}
+
+void FActorSkinnedMeshComponentInterface::PostSkeletalHierarchyChange()
+{
+	FComponentReregisterContext Context(USkinnedMeshComponent::GetSkinnedMeshComponent(this));
+}
+#endif
+
+USkinnedAsset* FActorSkinnedMeshComponentInterface::GetSkinnedAsset() const
+{
+	return USkinnedMeshComponent::GetSkinnedMeshComponent(this)->GetSkinnedAsset();
+}
+
+IPrimitiveComponent* FActorSkinnedMeshComponentInterface::GetPrimitiveComponentInterface()
+{
+	return USkinnedMeshComponent::GetSkinnedMeshComponent(this)->GetPrimitiveComponentInterface();
+}
diff --git a/Engine/Source/Runtime/Engine/Private/Components/ActorStaticMeshComponentInterface.cpp b/Engine/Source/Runtime/Engine/Private/Components/ActorStaticMeshComponentInterface.cpp
new file mode 100644
--- /dev/null
+++ b/Engine/Source/Runtime/Engine/Private/Components/ActorStaticMeshComponentInterface.cpp
@@ -0,0 +1,31 @@
+// Copyright Epic Games, Inc. All Rights Reserved.
+
+#include "Components/ActorStaticMeshComponentInterface.h"
+#include "Components/StaticMeshComponent.h"
+
+#if WITH_EDITOR
+void FActorStaticMeshComponentInterface::OnMeshRebuild(bool bRenderDataChanged)
+{
+	UStaticMeshComponent::GetStaticMeshComponent(this)->OnMeshRebuild(bRenderDataChanged);
+}
+
+void FActorStaticMeshComponentInterface::PreStaticMeshCompilation()
+{
+	UStaticMeshComponent::GetStaticMeshComponent(this)->PreStaticMeshCompilation();
+}
+
+void FActorStaticMeshComponentInterface::PostStaticMeshCompilation()
+{
+	UStaticMeshComponent::GetStaticMeshComponent(this)->PostStaticMeshCompilation();
+}
+#endif
+
+UStaticMesh* FActorStaticMeshComponentInterface::GetStaticMesh() const
+{
+	return UStaticMeshComponent::GetStaticMeshComponent(this)->GetStaticMesh();
+}
+
+IPrimitiveComponent* FActorStaticMeshComponentInterface::GetPrimitiveComponentInterface()
+{
+	return UStaticMeshComponent::GetStaticMeshComponent(this)->GetPrimitiveComponentInterface();
+}
diff --git a/Engine/Source/Runtime/Engine/Private/Components/PrimitiveComponent.cpp b/Engine/Source/Runtime/Engine/Private/Components/PrimitiveComponent.cpp
--- a/Engine/Source/Runtime/Engine/Private/Components/PrimitiveComponent.cpp
+++ b/Engine/Source/Runtime/Engine/Private/Components/PrimitiveComponent.cpp
@@ -5386,149 +5386,4 @@ void UPrimitiveComponent::ReleaseSceneProxy()
 	SceneData.bAlwaysVisible = false;
 }
 
-bool FActorPrimitiveComponentInterface::IsRenderStateCreated() const 
-{
-	return UPrimitiveComponent::GetPrimitiveComponent(this)->IsRenderStateCreated();
-}
-
-bool FActorPrimitiveComponentInterface::IsRenderStateDirty() const 
-{
-	return UPrimitiveComponent::GetPrimitiveComponent(this)->IsRenderStateDirty();
-}
-
-bool FActorPrimitiveComponentInterface::ShouldCreateRenderState() const 
-{
-	return UPrimitiveComponent::GetPrimitiveComponent(this)->ShouldCreateRenderState();
-}
-
-bool FActorPrimitiveComponentInterface::IsRegistered() const 
-{
-	return UPrimitiveComponent::GetPrimitiveComponent(this)->IsRegistered();
-}
-
-bool FActorPrimitiveComponentInterface::IsUnreachable() const 
-{
-	return UPrimitiveComponent::GetPrimitiveComponent(this)->IsUnreachable();
-}
-
-UWorld* FActorPrimitiveComponentInterface::GetWorld() const 
-{
-	return UPrimitiveComponent::GetPrimitiveComponent(this)->GetWorld();
-}
-
-FSceneInterface* FActorPrimitiveComponentInterface::GetScene() const 
-{
-	return UPrimitiveComponent::GetPrimitiveComponent(this)->GetScene();
-}
-
-FPrimitiveSceneProxy* FActorPrimitiveComponentInterface::GetSceneProxy() const 
-{
-	return UPrimitiveComponent::GetPrimitiveComponent(this)->SceneProxy;
-}
-
-void FActorPrimitiveComponentInterface::GetUsedMaterials(TArray<UMaterialInterface*>& OutMaterials, bool bGetDebugMaterials) const
-{
-	UPrimitiveComponent::GetPrimitiveComponent(this)->GetUsedMaterials(OutMaterials, bGetDebugMaterials);
-}
-
-void FActorPrimitiveComponentInterface::GetStreamableRenderAssetInfo(TArray<FStreamingRenderAssetPrimitiveInfo>& StreamableRenderAssets) const
-{
-	FStreamingTextureLevelContext LevelContext{GetCurrentMaterialQualityLevelChecked(),UPrimitiveComponent::GetPrimitiveComponent(this)};
-	UPrimitiveComponent::GetPrimitiveComponent(this)->GetStreamingRenderAssetInfo(LevelContext, StreamableRenderAssets);
-}
-
-void FActorPrimitiveComponentInterface::MarkRenderStateDirty()
-{
-	UPrimitiveComponent::GetPrimitiveComponent(this)->MarkRenderStateDirty();
-}
-
-void FActorPrimitiveComponentInterface::DestroyRenderState() 
-{
-	UPrimitiveComponent::GetPrimitiveComponent(this)->DestroyRenderState_Concurrent();
-}
-
-void FActorPrimitiveComponentInterface::CreateRenderState(FRegisterComponentContext* Context) 
-{
-	UPrimitiveComponent::GetPrimitiveComponent(this)->CreateRenderState_Concurrent(Context);
-}
-
-FString FActorPrimitiveComponentInterface::GetName() const 
-{
-	return UPrimitiveComponent::GetPrimitiveComponent(this)->GetName();
-}
-
-FString FActorPrimitiveComponentInterface::GetFullName() const 
-{
-	return UPrimitiveComponent::GetPrimitiveComponent(this)->GetFullName();
-}
-
-FTransform FActorPrimitiveComponentInterface::GetTransform() const 
-{
-	return UPrimitiveComponent::GetPrimitiveComponent(this)->GetComponentTransform();
-}
-
-FBoxSphereBounds FActorPrimitiveComponentInterface::GetBounds() const 
-{
-	return UPrimitiveComponent::GetPrimitiveComponent(this)->Bounds;
-}
-
-float FActorPrimitiveComponentInterface::GetLastRenderTimeOnScreen() const 
-{
-	return UPrimitiveComponent::GetPrimitiveComponent(this)->GetLastRenderTimeOnScreen();
-}
-
-void FActorPrimitiveComponentInterface::GetPrimitiveStats(FPrimitiveStats& PrimitiveStats) const
-{
-return UPrimitiveComponent::GetPrimitiveComponent(this)->GetPrimitiveStats(PrimitiveStats);
-}
-
-
-UObject* FActorPrimitiveComponentInterface::GetUObject() 
-{
-	return UPrimitiveComponent::GetPrimitiveComponent(this);
-}
-
-const UObject* FActorPrimitiveComponentInterface::GetUObject() const 
-{
-	return UPrimitiveComponent::GetPrimitiveComponent(this);
-}
-
-UObject* FActorPrimitiveComponentInterface::GetOwner() const 
-{
-	return UPrimitiveComponent::GetPrimitiveComponent(this)->GetOwner();
-}
-
-FString FActorPrimitiveComponentInterface::GetOwnerName() const 
-{
-	const UPrimitiveComponent* Component = UPrimitiveComponent::GetPrimitiveComponent(this);
-
-#if ACTOR_HAS_LABELS
-	return Component->GetOwner() ? Component->GetOwner()->GetActorNameOrLabel() : Component->GetName();
-#else
-	return Component->GetName();
-#endif
-}
-
-FPrimitiveSceneProxy* FActorPrimitiveComponentInterface::CreateSceneProxy() 
-{
-	UPrimitiveComponent* Component = UPrimitiveComponent::GetPrimitiveComponent(this);
-	FPrimitiveSceneProxy* SceneProxy = Component->CreateSceneProxy();
-	Component->AssignSceneProxy(SceneProxy);
-
-	return SceneProxy;
-}
-
-void FActorPrimitiveComponentInterface::PrecachePSOs()
-{
-	UPrimitiveComponent::GetPrimitiveComponent(this)->PrecachePSOs();
-}
-
-#if WITH_EDITOR
-HHitProxy* FActorPrimitiveComponentInterface::CreateMeshHitProxy(int32 SectionIndex, int32 MaterialIndex) 
-{
-	UPrimitiveComponent* Component = UPrimitiveComponent::GetPrimitiveComponent(this);	
-	return Component->CreateMeshHitProxy(SectionIndex, MaterialIndex);	
-}
-#endif
-
 #undef LOCTEXT_NAMESPACE
diff --git a/Engine/Source/Runtime/Engine/Private/Components/SkinnedMeshComponent.cpp b/Engine/Source/Runtime/Engine/Private/Components/SkinnedMeshComponent.cpp
--- a/Engine/Source/Runtime/Engine/Private/Components/SkinnedMeshComponent.cpp
+++ b/Engine/Source/Runtime/Engine/Private/Components/SkinnedMeshComponent.cpp
@@ -5741,28 +5741,6 @@ void USkinnedMeshComponent::OutdatedKnownSkinnedAssetDetected() const
 }
 #endif
 
-#if WITH_EDITOR
-void FActorSkinnedMeshComponentInterface::PostAssetCompilation()
-{
-	USkinnedMeshComponent::GetSkinnedMeshComponent(this)->PostAssetCompilation();
-}
-
-void FActorSkinnedMeshComponentInterface::PostSkeletalHierarchyChange()
-{
-	FComponentReregisterContext Context(USkinnedMeshComponent::GetSkinnedMeshComponent(this));
-}
-#endif
-
-USkinnedAsset* FActorSkinnedMeshComponentInterface::GetSkinnedAsset() const
-{
-	return USkinnedMeshComponent::GetSkinnedMeshComponent(this)->GetSkinnedAsset();
-}
-
-IPrimitiveComponent* FActorSkinnedMeshComponentInterface::GetPrimitiveComponentInterface()
-{
-	return USkinnedMeshComponent::GetSkinnedMeshComponent(this)->GetPrimitiveComponentInterface();
-}
-
 void FAnimUpdateRateParameters::SetTrailMode(float DeltaTime, uint8 UpdateRateShift, int32 NewUpdateRate, int32 NewEvaluationRate, bool bNewInterpSkippedFrames)
 {
 	OptimizeMode = TrailMode;
diff --git a/Engine/Source/Runtime/Engine/Private/Components/StaticMeshComponent.cpp b/Engine/Source/Runtime/Engine/Private/Components/StaticMeshComponent.cpp
--- a/Engine/Source/Runtime/Engine/Private/Components/StaticMeshComponent.cpp
+++ b/Engine/Source/Runtime/Engine/Private/Components/StaticMeshComponent.cpp
@@ -3942,33 +3942,4 @@ void UStaticMeshComponent::GetPrimitiveStats(FPrimitiveStats& PrimitiveStats) co
 	}
 }
 
-#if WITH_EDITOR
-void FActorStaticMeshComponentInterface::OnMeshRebuild(bool bRenderDataChanged)
-{
-	UStaticMeshComponent::GetStaticMeshComponent(this)->OnMeshRebuild(bRenderDataChanged);
-}
-
-void FActorStaticMeshComponentInterface::PreStaticMeshCompilation()
-{
-	UStaticMeshComponent::GetStaticMeshComponent(this)->PreStaticMeshCompilation();
-}
-
-void FActorStaticMeshComponentInterface::PostStaticMeshCompilation()
-{
-	UStaticMeshComponent::GetStaticMeshComponent(this)->PostStaticMeshCompilation();
-}
-#endif
-
-UStaticMesh* FActorStaticMeshComponentInterface::GetStaticMesh() const
-{
-	return UStaticMeshComponent::GetStaticMeshComponent(this)->GetStaticMesh();
-}
-
-IPrimitiveComponent* FActorStaticMeshComponentInterface::GetPrimitiveComponentInterface() 
-{
-	return UStaticMeshComponent::GetStaticMeshComponent(this)->GetPrimitiveComponentInterface();
-}
-
-
 #undef LOCTEXT_NAMESPACE
-
