[ObjectCacheContext]
- Remove checks since this is now used in multiple threads

diff --git a/Engine/Source/Runtime/Engine/Private/ObjectCacheContext.cpp b/Engine/Source/Runtime/Engine/Private/ObjectCacheContext.cpp
--- a/Engine/Source/Runtime/Engine/Private/ObjectCacheContext.cpp
+++ b/Engine/Source/Runtime/Engine/Private/ObjectCacheContext.cpp
@@ -103,7 +103,6 @@ EObjectReverseLookupMode GetObjectReverseLookupMode()
 {
 	if (!GIsGetObjectReverseLookupModeInitialized)
 	{
-		check(IsInGameThread());
 		int32 Mode = 0;
 		if (FParse::Value(FCommandLine::Get(), TEXT("-ObjectReverseLookupMode="), Mode))
 		{
@@ -120,7 +119,6 @@ EObjectReverseLookupMask GetObjectReverseLookupMask()
 {
 	if (!GIsGetObjectReverseLookupMaskInitialized)
 	{
-		check(IsInGameThread());
 		int32 Mask = 0;
 		if (FParse::Value(FCommandLine::Get(), TEXT("-ObjectReverseLookupMask="), Mask))
 		{
