[FastGeo] Automated tests to check that transformers can be enabled/disabled solely through actor/component filtering. Uses separate settings assets.

#tests Locally ran the 2 tests
#rnx
#virtualized

diff --git a/Engine/Source/Runtime/Engine/Public/WorldPartition/WorldPartition.h b/Engine/Source/Runtime/Engine/Public/WorldPartition/WorldPartition.h
--- a/Engine/Source/Runtime/Engine/Public/WorldPartition/WorldPartition.h
+++ b/Engine/Source/Runtime/Engine/Public/WorldPartition/WorldPartition.h
@@ -348,6 +348,7 @@ public:
 	ENGINE_API FBox GetEditorWorldBounds() const;
 	ENGINE_API FBox GetRuntimeWorldBounds() const;
 
+	inline TConstArrayView<FRuntimeCellTransformerInstance> GetRuntimeCellsTransformerStack() const { return RuntimeCellsTransformerStack; }
 	ENGINE_API void ApplyRuntimeCellsTransformerStack(ULevel* InLevel);
 	ENGINE_API bool HasRuntimeCellTransformerOfType(const UClass* InCellTransformerClass);
 	ENGINE_API void AddRuntimeCellTransformerOfType(TSubclassOf<UWorldPartitionRuntimeCellTransformer> InCellTransformerClass, int32 PositionToInsert = -1);
