PR #13845: Fix for world partition unsaved duplicates not handling serialized guid fields in PIE

diff --git a/Engine/Source/Runtime/Engine/Private/WorldPartition/WorldPartitionLevelStreamingDynamic.cpp b/Engine/Source/Runtime/Engine/Private/WorldPartition/WorldPartitionLevelStreamingDynamic.cpp
--- a/Engine/Source/Runtime/Engine/Private/WorldPartition/WorldPartitionLevelStreamingDynamic.cpp
+++ b/Engine/Source/Runtime/Engine/Private/WorldPartition/WorldPartitionLevelStreamingDynamic.cpp
@@ -18,6 +18,7 @@
 #include "WorldPartition/DataLayer/ExternalDataLayerHelper.h"
 #include "WorldPartition/ContentBundle/ContentBundlePaths.h"
 #include "LevelUtils.h"
+#include "Misc/PlayInEditorLoadingScope.h"
 #endif
 
 #define LOCTEXT_NAMESPACE "World"
@@ -422,6 +423,8 @@ bool UWorldPartitionLevelStreamingDynamic::IssueLoadRequests()
 	// Duplicate unsaved actors
 	if (UnsavedActorsContainer)
 	{
+		UE::Core::Private::FPlayInEditorLoadingScope PlayInEditorIDScope(RuntimeLevel->GetPackage()->GetPIEInstanceID());
+
 		// Set RuntimeLevel on container temporarily (for duplication) so that external code can find the Level in PostLoad code 
 		TGuardValue<TObjectPtr<ULevel>> ScopeRuntimeLevel(UnsavedActorsContainer->RuntimeLevel, RuntimeLevel);
 		FObjectDuplicationParameters Parameters(UnsavedActorsContainer, RuntimeLevel->GetPackage());
