Make bGPUSkinPassThrough part of the FLocalVertexFactoryUniformShaderParameters instead of global uniform buffer bool

diff --git a/Engine/Shaders/Private/LocalVertexFactory.ush b/Engine/Shaders/Private/LocalVertexFactory.ush
--- a/Engine/Shaders/Private/LocalVertexFactory.ush
+++ b/Engine/Shaders/Private/LocalVertexFactory.ush
@@ -128,14 +128,10 @@ struct FVertexFactoryInput
 	uint VertexId : SV_VertexID;
 };
 
-#if SUPPORT_GPUSKIN_PASSTHROUGH
-uint bIsGPUSkinPassThrough;
-#endif
-
 bool IsGPUSkinPassThrough()
 {
 	#if SUPPORT_GPUSKIN_PASSTHROUGH
-		return bIsGPUSkinPassThrough;
+		return LocalVF.bIsGPUSkinPassThrough;
 	#else
 		return false;
 	#endif
diff --git a/Engine/Source/Runtime/Engine/Private/LocalVertexFactory.cpp b/Engine/Source/Runtime/Engine/Private/LocalVertexFactory.cpp
--- a/Engine/Source/Runtime/Engine/Private/LocalVertexFactory.cpp
+++ b/Engine/Source/Runtime/Engine/Private/LocalVertexFactory.cpp
@@ -119,6 +119,7 @@ void GetLocalVFUniformShaderParameters(
 
 	UniformParameters.VertexFetch_Parameters = {ColorIndexMask, NumTexCoords, LightMapCoordinateIndex, EffectiveBaseVertexIndex};
 	UniformParameters.PreSkinBaseVertexIndex = EffectivePreSkinBaseVertexIndex;
+	UniformParameters.bIsGPUSkinPassThrough = LocalVertexFactory->IsGPUSkinPassThrough() ? 1 : 0;
 }
 
 TUniformBufferRef<FLocalVertexFactoryUniformShaderParameters> CreateLocalVFUniformBuffer(
@@ -192,12 +193,6 @@ void FLocalVertexFactoryShaderParametersBase::GetElementShaderBindingsBase(
 	}
 }
 
-void FLocalVertexFactoryShaderParameters::Bind(const FShaderParameterMap& ParameterMap)
-{
-	FLocalVertexFactoryShaderParametersBase::Bind(ParameterMap);
-	IsGPUSkinPassThrough.Bind(ParameterMap, TEXT("bIsGPUSkinPassThrough"));
-}
-
 void FLocalVertexFactoryShaderParameters::GetElementShaderBindings(
 	const FSceneInterface* Scene,
 	const FSceneView* View,
@@ -211,7 +206,6 @@ void FLocalVertexFactoryShaderParameters::GetElementShaderBindings(
 ) const
 {
 	FLocalVertexFactory const* LocalVertexFactory = static_cast<FLocalVertexFactory const*>(VertexFactory);
-	ShaderBindings.Add(IsGPUSkinPassThrough, (uint32)(LocalVertexFactory->bGPUSkinPassThrough ? 1 : 0));
 	if (LocalVertexFactory->bGPUSkinPassThrough)
 	{
 		const FGPUSkinPassthroughVertexFactory* GPUSkinPassthroughVF = static_cast<const FGPUSkinPassthroughVertexFactory*>(LocalVertexFactory);
diff --git a/Engine/Source/Runtime/Engine/Public/LocalVertexFactory.h b/Engine/Source/Runtime/Engine/Public/LocalVertexFactory.h
--- a/Engine/Source/Runtime/Engine/Public/LocalVertexFactory.h
+++ b/Engine/Source/Runtime/Engine/Public/LocalVertexFactory.h
@@ -20,6 +20,7 @@ BEGIN_GLOBAL_SHADER_PARAMETER_STRUCT(FLocalVertexFactoryUniformShaderParameters,
 	SHADER_PARAMETER(FIntVector4,VertexFetch_Parameters)
 	SHADER_PARAMETER(int32, PreSkinBaseVertexIndex)
 	SHADER_PARAMETER(uint32,LODLightmapDataIndex)
+	SHADER_PARAMETER(uint32,bIsGPUSkinPassThrough)
 	SHADER_PARAMETER_SRV(Buffer<float2>, VertexFetch_TexCoordBuffer)
 	SHADER_PARAMETER_SRV(Buffer<float>, VertexFetch_PositionBuffer)
 	SHADER_PARAMETER_SRV(Buffer<float>, VertexFetch_PreSkinPositionBuffer)
@@ -164,6 +165,11 @@ public:
 		return UniformBuffer.GetReference();
 	}
 
+	inline const bool IsGPUSkinPassThrough() const
+	{
+		return bGPUSkinPassThrough;
+	}
+
 #if WITH_EDITORONLY_DATA
 	virtual bool IsCoarseProxyMesh() const override { return Data.bIsCoarseProxy; }
 
@@ -245,8 +251,6 @@ class FLocalVertexFactoryShaderParameters : public FLocalVertexFactoryShaderPara
 {
 	DECLARE_EXPORTED_TYPE_LAYOUT(FLocalVertexFactoryShaderParameters, ENGINE_API, NonVirtual);
 public:
-	void Bind(const FShaderParameterMap& ParameterMap);
-
 	void GetElementShaderBindings(
 		const FSceneInterface* Scene,
 		const FSceneView* View,
@@ -258,7 +262,4 @@ public:
 		FMeshDrawSingleShaderBindings& ShaderBindings,
 		FVertexInputStreamArray& VertexStreams
 	) const; 
-
-private:
-	LAYOUT_FIELD(FShaderParameter, IsGPUSkinPassThrough);
 };
