Removing unnecessary interface.

diff --git a/Engine/Source/Runtime/Engine/Private/Streaming/DynamicTextureInstanceManager.cpp b/Engine/Source/Runtime/Engine/Private/Streaming/DynamicTextureInstanceManager.cpp
--- a/Engine/Source/Runtime/Engine/Private/Streaming/DynamicTextureInstanceManager.cpp
+++ b/Engine/Source/Runtime/Engine/Private/Streaming/DynamicTextureInstanceManager.cpp
@@ -243,11 +243,6 @@ void FDynamicRenderAssetInstanceManager::OnPreGarbageCollect(FRemovedRenderAsset
 	}
 }
 
-
-/*-----------------------------------
------- IRenderAssetInstanceManager ------
------------------------------------*/
-
 bool FDynamicRenderAssetInstanceManager::CanManage(const UPrimitiveComponent* Component) const
 {
 	return IsValid(Component) && !Component->HasAnyFlags(RF_BeginDestroyed|RF_FinishDestroyed);
diff --git a/Engine/Source/Runtime/Engine/Private/Streaming/DynamicTextureInstanceManager.h b/Engine/Source/Runtime/Engine/Private/Streaming/DynamicTextureInstanceManager.h
--- a/Engine/Source/Runtime/Engine/Private/Streaming/DynamicTextureInstanceManager.h
+++ b/Engine/Source/Runtime/Engine/Private/Streaming/DynamicTextureInstanceManager.h
@@ -7,7 +7,6 @@ DynamicTextureInstanceManager.h: Definitions of classes used for texture streami
 #pragma once
 
 #include "CoreMinimal.h"
-#include "Streaming/TextureInstanceManager.h"
 #include "Streaming/TextureInstanceTask.h"
 #include "ContentStreaming.h"
 
@@ -15,7 +14,7 @@ DynamicTextureInstanceManager.h: Definitions of classes used for texture streami
  * A texture/mesh instance manager to manage dynamic components. 
  * The async view generated is duplicated so that the state can change freely.
  */
-class FDynamicRenderAssetInstanceManager : public IRenderAssetInstanceManager
+class FDynamicRenderAssetInstanceManager
 {
 public:
 
@@ -35,37 +34,33 @@ public:
 	/** Remove all pending components that are marked for delete. This prevents searching in the pending list for each entry. */
 	void OnPreGarbageCollect(FRemovedRenderAssetArray& RemovedRenderAssets);
 
-	/*-----------------------------------
-	------ IRenderAssetInstanceManager ------
-	-----------------------------------*/
-
 	/** Return whether this component can be managed by this manager. */
-	bool IsReferenced(const UPrimitiveComponent* Component) const final;
+	bool IsReferenced(const UPrimitiveComponent* Component) const;
 
 	/** Return whether this component is be managed by this manager. */
-	bool CanManage(const UPrimitiveComponent* Component) const final override;
+	bool CanManage(const UPrimitiveComponent* Component) const;
 
 	/** Add a component streaming data, the LevelContext gives support for precompiled data. */
-	EAddComponentResult Add(const UPrimitiveComponent* Component, FStreamingTextureLevelContext& LevelContext, float MaxAllowedUIDensity = 0) final override;
+	EAddComponentResult Add(const UPrimitiveComponent* Component, FStreamingTextureLevelContext& LevelContext, float MaxAllowedUIDensity = 0);
 
 	/** Remove a component, the RemovedRenderAssets is the list of textures or meshes not referred anymore. */
-	 void Remove(const UPrimitiveComponent* Component, FRemovedRenderAssetArray* RemovedRenderAssets) final override;
+	 void Remove(const UPrimitiveComponent* Component, FRemovedRenderAssetArray* RemovedRenderAssets);
 
 	/** Notify the manager that an async view will be requested on the next frame. */
-	void PrepareAsyncView() final override;
+	void PrepareAsyncView();
 
 	/** Return a view of the data that has to be 100% thread safe. The content is allowed to be updated, but not memory must be reallocated. */
-	const FRenderAssetInstanceView* GetAsyncView(bool bCreateIfNull) final override;
+	const FRenderAssetInstanceView* GetAsyncView(bool bCreateIfNull);
 
 	/** Return the size taken for sub-allocation. */
-	uint32 GetAllocatedSize() const final override;
+	uint32 GetAllocatedSize() const;
 
 	const FRenderAssetInstanceView* GetGameThreadView();
 
 protected:
 	
 	/** Refresh component data (bounds, last render time, min and max view distance) - see TextureInstanceView. */
-	void Refresh(float Percentage) final override;
+	void Refresh(float Percentage);
 
 	void OnCreateViewDone(FRenderAssetInstanceView* InView);
 	void OnRefreshVisibilityDone(int32 BeginIndex, int32 EndIndex, const TArray<int32>& SkippedIndices, int32 FirstFreeBound, int32 LastUsedBound);
diff --git a/Engine/Source/Runtime/Engine/Private/Streaming/StaticTextureInstanceManager.cpp b/Engine/Source/Runtime/Engine/Private/Streaming/StaticTextureInstanceManager.cpp
--- a/Engine/Source/Runtime/Engine/Private/Streaming/StaticTextureInstanceManager.cpp
+++ b/Engine/Source/Runtime/Engine/Private/Streaming/StaticTextureInstanceManager.cpp
@@ -51,10 +51,6 @@ void FStaticRenderAssetInstanceManager::OnRefreshVisibilityDone(int32 InBeginInd
 	}
 }
 
-/*-----------------------------------
------- IRenderAssetInstanceManager ------
------------------------------------*/
-
 bool FStaticRenderAssetInstanceManager::CanManage(const UPrimitiveComponent* Component) const
 {
 	// This manager only manages static components from static actors.
diff --git a/Engine/Source/Runtime/Engine/Private/Streaming/StaticTextureInstanceManager.h b/Engine/Source/Runtime/Engine/Private/Streaming/StaticTextureInstanceManager.h
--- a/Engine/Source/Runtime/Engine/Private/Streaming/StaticTextureInstanceManager.h
+++ b/Engine/Source/Runtime/Engine/Private/Streaming/StaticTextureInstanceManager.h
@@ -7,7 +7,6 @@ StaticTextureInstanceManager.h: Definitions of classes used for texture streamin
 #pragma once
 
 #include "CoreMinimal.h"
-#include "Streaming/TextureInstanceManager.h"
 #include "Streaming/TextureInstanceTask.h"
 
 /** 
@@ -15,7 +14,7 @@ StaticTextureInstanceManager.h: Definitions of classes used for texture streamin
  * Once an async view has been requested, nothing can be changed as the async view and the internal state point to the same data. 
  * This allows for quicker refreshes and also prevent state duplication.
  */
-class FStaticRenderAssetInstanceManager : public IRenderAssetInstanceManager
+class FStaticRenderAssetInstanceManager
 {
 public:
 
@@ -32,42 +31,35 @@ public:
 
 	FORCEINLINE bool HasRenderAssetReferences() const { return StateSync.GetState()->NumBounds() > 0; }
 
-	/*-----------------------------------
-	------ IRenderAssetInstanceManager ------
-	-----------------------------------*/
-
 	/** Return whether this component can be managed by this manager. */
-	FORCEINLINE bool IsReferenced(const UPrimitiveComponent* Component) const final override { return StateSync.GetState()->HasComponentReferences(Component); }
+	FORCEINLINE bool IsReferenced(const UPrimitiveComponent* Component) const { return StateSync.GetState()->HasComponentReferences(Component); }
 
 	/** Return whether this component can be managed by this manager. */
 	FORCEINLINE void GetReferencedComponents(TArray<const UPrimitiveComponent*>& Components) const { StateSync.GetState()->GetReferencedComponents(Components); }
 
 	/** Return whether this component is be managed by this manager. */
-	bool CanManage(const UPrimitiveComponent* Component) const final override;
+	bool CanManage(const UPrimitiveComponent* Component) const;
 
 	/** Refresh component data (bounds, last render time, min and max view distance) - see RenderAssetInstanceView. */
-	void Refresh(float Percentage) final override;
+	void Refresh(float Percentage);
 
 	/** Add a component streaming data, the LevelContext gives support for precompiled data. */
-	EAddComponentResult Add(const UPrimitiveComponent* Component, FStreamingTextureLevelContext& LevelContext, float MaxAllowedUIDensity) final override;
+	EAddComponentResult Add(const UPrimitiveComponent* Component, FStreamingTextureLevelContext& LevelContext, float MaxAllowedUIDensity);
 
 	/** Add a batch of components streaming data. */
 	EAddComponentResult Add(const TArray<FRenderAssetInstanceState::FPreAddComponentPayload>& Payloads, TFunctionRef<void(const FRenderAssetInstanceState::FPreAddComponentPayload&, EAddComponentResult)> OnAddFailedFunc);
 
 	/** Remove a component, the RemoveTextures is the list of textures not referred anymore. */
-	 void Remove(const UPrimitiveComponent* Component, FRemovedRenderAssetArray* RemovedRenderAssets) final override;
-
-	/** Notify the manager that an async view will be requested on the next frame. */
-	FORCEINLINE void PrepareAsyncView() final override {}
+	 void Remove(const UPrimitiveComponent* Component, FRemovedRenderAssetArray* RemovedRenderAssets);
 
 	/** Return a view of the data that has to be 100% thread safe. The content is allowed to be updated, but not memory must be reallocated. */
-	const FRenderAssetInstanceView* GetAsyncView(bool bCreateIfNull) final override;
+	const FRenderAssetInstanceView* GetAsyncView(bool bCreateIfNull);
 
 	/** Returns whether we can add components to the manager. */
 	bool CanAddComponent() const;
 
 	/** Return the size taken for sub-allocation. */
-	uint32 GetAllocatedSize() const final override;
+	uint32 GetAllocatedSize() const;
 
 	/** Apply specified offset to all cached primitives bounds */
 	void OffsetBounds(const FVector& Offset);
diff --git a/Engine/Source/Runtime/Engine/Private/Streaming/TextureInstanceManager.h b/Engine/Source/Runtime/Engine/Private/Streaming/TextureInstanceManager.h
deleted file mode 100644
