AttributesContainer: no FName defns in headers

diff --git a/Engine/Source/Runtime/Engine/Classes/Animation/AttributesContainer.h b/Engine/Source/Runtime/Engine/Classes/Animation/AttributesContainer.h
--- a/Engine/Source/Runtime/Engine/Classes/Animation/AttributesContainer.h
+++ b/Engine/Source/Runtime/Engine/Classes/Animation/AttributesContainer.h
@@ -21,7 +21,8 @@ namespace UE
 {
 	namespace Anim
 	{
-		static FName BoneAttributeNamespace("bone");
+		extern ENGINE_API const FName BoneAttributeNamespace;
+
 		/** Runtime identifier for an attribute */
 		struct FAttributeId
 		{
diff --git a/Engine/Source/Runtime/Engine/Private/Animation/AttributesContainer.cpp b/Engine/Source/Runtime/Engine/Private/Animation/AttributesContainer.cpp
new file mode 100644
--- /dev/null
+++ b/Engine/Source/Runtime/Engine/Private/Animation/AttributesContainer.cpp
@@ -0,0 +1,8 @@
+// Copyright Epic Games, Inc. All Rights Reserved.
+
+#include "Animation/AttributesContainer.h"
+
+namespace UE::Anim
+{
+	const FName BoneAttributeNamespace{"bone"};
+}
