Removing outdated if check.

diff --git a/Engine/Source/Runtime/Engine/Private/Slate/SceneViewport.cpp b/Engine/Source/Runtime/Engine/Private/Slate/SceneViewport.cpp
--- a/Engine/Source/Runtime/Engine/Private/Slate/SceneViewport.cpp
+++ b/Engine/Source/Runtime/Engine/Private/Slate/SceneViewport.cpp
@@ -2132,7 +2132,7 @@ void FSceneViewport::InitRHI(FRHICommandListBase& RHICmdList)
 				RTRHI = BufferedRTRHI[i];
 				SRVRHI = BufferedSRVRHI[i];
 			}
-			else if (StereoRenderTargetManager == nullptr)
+			else
 			{
 				const FRHITextureCreateDesc Desc =
 					FRHITextureCreateDesc::Create2D(TEXT("BufferedRT"))
