Details cleanup across USoundBase/USoundWave/USoundCue/UMetaSoundSource to ensure consistency

diff --git a/Engine/Plugins/Runtime/Metasound/Source/MetasoundEditor/Private/MetasoundDetailCustomization.cpp b/Engine/Plugins/Runtime/Metasound/Source/MetasoundEditor/Private/MetasoundDetailCustomization.cpp
--- a/Engine/Plugins/Runtime/Metasound/Source/MetasoundEditor/Private/MetasoundDetailCustomization.cpp
+++ b/Engine/Plugins/Runtime/Metasound/Source/MetasoundEditor/Private/MetasoundDetailCustomization.cpp
@@ -292,13 +292,13 @@ namespace Metasound::Editor
 				GeneralCategoryBuilder.AddProperty(CategoryHierarchyHandle);
 				GeneralCategoryBuilder.AddProperty(KeywordsHandle);
 
-				DetailLayout.HideCategory("Attenuation");
-				DetailLayout.HideCategory("Developer");
-				DetailLayout.HideCategory("Effects");
-				DetailLayout.HideCategory("Loading");
-				DetailLayout.HideCategory("Modulation");
 				DetailLayout.HideCategory("Sound");
 				DetailLayout.HideCategory("Voice Management");
+				DetailLayout.HideCategory("Routing");
+				DetailLayout.HideCategory("Modulation");
+				DetailLayout.HideCategory("Memory");
+				DetailLayout.HideCategory("Developer");
+				DetailLayout.HideCategory("Advanced");
 			}
 			break;
 
@@ -345,15 +345,10 @@ namespace Metasound::Editor
 
 		// Hack to hide parent structs for nested metadata properties
 		DetailLayout.HideCategory("CustomView");
-
-		DetailLayout.HideCategory("Analysis");
+		
 		DetailLayout.HideCategory("Curves");
 		DetailLayout.HideCategory("File Path");
-		DetailLayout.HideCategory("Format");
 		DetailLayout.HideCategory("Info");
-		DetailLayout.HideCategory("Loading");
-		DetailLayout.HideCategory("Playback");
-		DetailLayout.HideCategory("Subtitles");
 		DetailLayout.HideCategory("Waveform Processing");
 	}
 
diff --git a/Engine/Source/Editor/DetailCustomizations/Private/SoundBaseDetailsUtils.cpp b/Engine/Source/Editor/DetailCustomizations/Private/SoundBaseDetailsUtils.cpp
--- a/Engine/Source/Editor/DetailCustomizations/Private/SoundBaseDetailsUtils.cpp
+++ b/Engine/Source/Editor/DetailCustomizations/Private/SoundBaseDetailsUtils.cpp
@@ -8,23 +8,14 @@ namespace SoundBaseDetailsUtils
 
 	static const TArray<FName> DefaultCategoryOrder =
 	{
-		"SoundWave",
-		"Format",
 		"Sound",
-		"Analysis",
+		"Voice Management",
+		"Routing",
 		"Modulation",
-		"Subtitles",
-		"Loading",
+		"Memory",
 		"Info",
-		"File Path",
-		"Waveform Processing",
 		"Developer",
-		"Voice Management",
-		"Effects",
-		"Attenuation",
-		"Curves",
 		"Advanced",
-		"Memory",
 	};
 } // namespace SoundBaseDetailsUtils
 
diff --git a/Engine/Source/Editor/DetailCustomizations/Private/SoundCueDetails.cpp b/Engine/Source/Editor/DetailCustomizations/Private/SoundCueDetails.cpp
--- a/Engine/Source/Editor/DetailCustomizations/Private/SoundCueDetails.cpp
+++ b/Engine/Source/Editor/DetailCustomizations/Private/SoundCueDetails.cpp
@@ -13,13 +13,11 @@ namespace SoundCueDetailsUtils
 	static const TArray<FName> CategoryOrder =
 	{
 		"Sound",
-		"Attenuation",
-		"Effects",
 		"Voice Management",
-		"AudioProperties",
+		"Routing",
 		"Memory",
-		"Advanced",
 		"Developer",
+		"Advanced",
 	};
 
 	void SortCategories(const TMap<FName, IDetailCategoryBuilder*>& AllCategoryMap)
@@ -36,6 +34,7 @@ TSharedRef<IDetailCustomization> FSoundCueDetails::MakeInstance()
 void FSoundCueDetails::CustomizeDetails(IDetailLayoutBuilder& DetailBuilder)
 {
 	DetailBuilder.SortCategories(SoundCueDetailsUtils::SortCategories);
+	DetailBuilder.HideCategory("Info");
 }
 
 #undef LOCTEXT_NAMESPACE
diff --git a/Engine/Source/Editor/DetailCustomizations/Private/SoundSourceBusDetails.cpp b/Engine/Source/Editor/DetailCustomizations/Private/SoundSourceBusDetails.cpp
--- a/Engine/Source/Editor/DetailCustomizations/Private/SoundSourceBusDetails.cpp
+++ b/Engine/Source/Editor/DetailCustomizations/Private/SoundSourceBusDetails.cpp
@@ -1,6 +1,7 @@
 // Copyright Epic Games, Inc. All Rights Reserved.
 
 #include "SoundSourceBusDetails.h"
+#include "Sound/SoundWave.h"
 
 #include "DetailLayoutBuilder.h"
 
@@ -15,6 +16,7 @@ void FSoundSourceBusDetails::CustomizeDetails(IDetailLayoutBuilder& DetailBuilde
 {
 	// Hide details of sound base that aren't relevant for buses
 	DetailBuilder.HideCategory("Waveform Processing");
+	DetailBuilder.HideProperty("Pitch", USoundWave::StaticClass());
 }
 
 #undef LOCTEXT_NAMESPACE
diff --git a/Engine/Source/Editor/DetailCustomizations/Private/SoundWaveDetails.cpp b/Engine/Source/Editor/DetailCustomizations/Private/SoundWaveDetails.cpp
--- a/Engine/Source/Editor/DetailCustomizations/Private/SoundWaveDetails.cpp
+++ b/Engine/Source/Editor/DetailCustomizations/Private/SoundWaveDetails.cpp
@@ -43,24 +43,17 @@ namespace SoundWaveDetailsUtils
 {
 	static const TArray<FName> CategoryOrder =
 	{
-		"SoundWave",
-		"Info",
-		"Format",
-		"Loading",
-		"Subtitles",
-		"Waveform Processing",
 		"Sound",
-		"Attenuation",
-		"Effects",
-		"Modulation",
 		"Voice Management",
-		"Audio Properties",
-		"Curves",
-		"Analysis",
+		"Routing",
+		"Modulation",
 		"Memory",
-		"Advanced",
+		"Info",
 		"Developer",
+		"Curves",
 		"File Path",
+		"Advanced",
+		"Waveform Processing",
 	};
 
 	void SortCategories(const TMap<FName, IDetailCategoryBuilder*>& AllCategoryMap)
@@ -225,7 +218,7 @@ void FSoundWaveDetails::CustomizeCloudStreamingPlatformDetails(IDetailLayoutBuil
 	}
 	if (!bEnabled)
 	{
-		DetailBuilder.HideCategory("Platform specific");
+		DetailBuilder.HideProperty("PlatformSettings");
 
 		// If there are objects that have the cloud streaming flag set without there being a suitable feature plugin available now,
 		// we do not hide the flag so that it can be disabled.
diff --git a/Engine/Source/Runtime/Engine/Classes/Sound/SoundBase.h b/Engine/Source/Runtime/Engine/Classes/Sound/SoundBase.h
--- a/Engine/Source/Runtime/Engine/Classes/Sound/SoundBase.h
+++ b/Engine/Source/Runtime/Engine/Classes/Sound/SoundBase.h
@@ -96,7 +96,7 @@ struct FSoundBaseEditorData
 
 public:
 	/** Override the Editor Preference to use submixes on preview playback */
-	UPROPERTY(EditAnywhere, Transient, SkipSerialization, Category = Developer, meta = (DisplayName = "Use Submixes on preview"))
+	UPROPERTY(EditAnywhere, Transient, SkipSerialization, Category = "Developer", meta = (DisplayName = "Use Submixes on preview"))
 	EUseSubmixOnPreviewMode OverrideUseSubmixesOnPreview = EUseSubmixOnPreviewMode::UseEditorPreference;
 };
 
@@ -113,7 +113,7 @@ private:
 
 public:
 	/** Sound class this sound belongs to */
-	UPROPERTY(EditAnywhere, BlueprintReadOnly, Category = Sound, meta = (DisplayName = "Class"), AssetRegistrySearchable)
+	UPROPERTY(EditAnywhere, BlueprintReadOnly, Category = "Sound", meta = (DisplayName = "Sound Class"), AssetRegistrySearchable)
 	TObjectPtr<USoundClass> SoundClassObject;
 
 #if WITH_EDITORONLY_DATA
@@ -123,11 +123,11 @@ public:
 #endif //WITH_EDITORONLY_DATA
 
 	/** When "au.3dVisualize.Attenuation" has been specified, draw this sound's attenuation shape when the sound is audible. For debugging purposes only. */
-	UPROPERTY(EditAnywhere, Category = Developer, meta = (DisplayName = "Enable Attenuation Debug"))
+	UPROPERTY(EditAnywhere, Category = "Developer", meta = (DisplayName = "Enable Attenuation Debug"))
 	uint8 bDebug : 1;
 
 	/** Whether or not to override the sound concurrency object with local concurrency settings. */
-	UPROPERTY(EditAnywhere, BlueprintReadWrite, Category = "Voice Management|Concurrency")
+	UPROPERTY(EditAnywhere, BlueprintReadWrite, Category = "Voice Management", meta=(InlineEditConditionToggle))
 	uint8 bOverrideConcurrency : 1;
 
 #if WITH_EDITORONLY_DATA
@@ -137,15 +137,15 @@ public:
 #endif //WITH_EDITORONLY_DATA
 
 	/** Whether or not to enable sending this audio's output to buses.  */
-	UPROPERTY(EditAnywhere, BlueprintReadWrite, Category = "Effects|Source")
+	UPROPERTY(EditAnywhere, BlueprintReadWrite, Category = "Routing", meta = (InlineEditConditionToggle))
 	uint8 bEnableBusSends : 1;
 
 	/** If enabled, sound will route to the Master Submix by default or to the Base Submix if defined. If disabled, sound will route ONLY to the Submix Sends and/or Bus Sends */
-	UPROPERTY(EditAnywhere, Category = "Effects|Submix")
+	UPROPERTY(EditAnywhere, Category = "Routing", meta = (InlineEditConditionToggle))
 	uint8 bEnableBaseSubmix : 1;
 
 	/** Whether or not to enable Submix Sends other than the Base Submix. */
-	UPROPERTY(EditAnywhere, Category = "Effects|Submix", meta = (DisplayAfter = "SoundSubmixObject"))
+	UPROPERTY(EditAnywhere, Category = "Routing", meta = (InlineEditConditionToggle))
 	uint8 bEnableSubmixSends : 1;
 
 	/** Whether or not this sound has a delay node */
@@ -162,7 +162,7 @@ public:
 #endif // WITH_EDITORONLY_DATA
 
 	/** Bypass volume weighting priority upon evaluating whether sound should remain active when max channel count is met (See platform Audio Settings). */
-	UPROPERTY(EditAnywhere, BlueprintReadWrite, Category = "Voice Management|Priority")
+	UPROPERTY(EditAnywhere, BlueprintReadWrite, Category = "Voice Management")
 	uint8 bBypassVolumeScaleForPriority : 1;
 
 	/** Virtualization behavior, determining if a sound may revive and how it continues playing when culled or evicted (limited to looping sounds). */
@@ -184,11 +184,11 @@ public:
 #endif // WITH_EDITORONLY_DATA
 
 	/** Set of concurrency settings to observe (if override is set to false).  Sound must pass all concurrency settings to play. */
-	UPROPERTY(EditAnywhere, BlueprintReadWrite, Category = "Voice Management|Concurrency", meta = (EditCondition = "!bOverrideConcurrency"))
+	UPROPERTY(EditAnywhere, BlueprintReadWrite, Category = "Voice Management", meta = (EditCondition = "!bOverrideConcurrency"))
 	TSet<TObjectPtr<USoundConcurrency>> ConcurrencySet;
 
 	/** If Override Concurrency is true, concurrency settings to use. */
-	UPROPERTY(EditAnywhere, BlueprintReadWrite, Category = "Voice Management|Concurrency", meta = (EditCondition = "bOverrideConcurrency"))
+	UPROPERTY(EditAnywhere, BlueprintReadWrite, Category = "Voice Management", meta = (EditCondition = "bOverrideConcurrency"))
 	FSoundConcurrencySettings ConcurrencyOverrides;
 
 #if WITH_EDITORONLY_DATA
@@ -198,56 +198,56 @@ public:
 #endif
 
 	/** Duration of sound in seconds. */
-	UPROPERTY(Category = Developer, AssetRegistrySearchable, VisibleAnywhere, BlueprintReadOnly)
+	UPROPERTY(Category = "Info|Format", AssetRegistrySearchable, VisibleAnywhere, BlueprintReadOnly)
 	float Duration;
 
 	/** The MaxDistance property is calculated statically on load or at asset edit time, but is not reliable at runtime.
 	  * the GetMaxDistance function should be used to determine the applied max distance based on runtime behavior. */
-	UPROPERTY(Category = Developer, AssetRegistrySearchable, VisibleAnywhere, BlueprintReadOnly)
+	UPROPERTY(Category = "Developer", AssetRegistrySearchable, VisibleAnywhere, BlueprintReadOnly)
 	float MaxDistance;
 
 	/** Total number of samples (in the thousands). Useful as a metric to analyze the relative size of a given sound asset in content browser. */
-	UPROPERTY(Category = Developer, AssetRegistrySearchable, VisibleAnywhere, BlueprintReadOnly)
+	UPROPERTY(Category = "Info|Format", AssetRegistrySearchable, VisibleAnywhere, BlueprintReadOnly)
 	float TotalSamples;
 
 	/** Used to determine whether sound can play or remain active if channel limit is met, where higher value is higher priority
 	  * (see platform's Audio Settings 'Max Channels' property). Unless bypassed, value is weighted with the final volume of the
 	  * sound to produce final runtime priority value.
 	  */
-	UPROPERTY(EditAnywhere, BlueprintReadWrite, Category = "Voice Management|Priority", meta = (ClampMin = "0.0", UIMin = "0.0", ClampMax = "100.0", UIMax = "100.0"))
+	UPROPERTY(EditAnywhere, BlueprintReadWrite, Category = "Voice Management", meta = (ClampMin = "0.0", UIMin = "0.0", ClampMax = "100.0", UIMax = "100.0", DisplayPriority=-1))
 	float Priority;
 
 	/** Attenuation settings package for the sound */
-	UPROPERTY(EditAnywhere, Category = Attenuation, meta = (EditCondition = "IsAttenuationSettingsEditable"))
+	UPROPERTY(EditAnywhere, Category = "Sound", meta = (EditCondition = "IsAttenuationSettingsEditable"))
 	TObjectPtr<USoundAttenuation> AttenuationSettings;
 
 	/** Submix to route sound output to. If unset, falls back to referenced SoundClass submix.
 	  * If SoundClass submix is unset, sends to the 'Master Submix' as set in the 'Audio' category of Project Settings'. */
-	UPROPERTY(EditAnywhere, BlueprintReadWrite, Category = "Effects|Submix", meta = (DisplayName = "Base Submix", EditCondition = "bEnableBaseSubmix"))
+	UPROPERTY(EditAnywhere, BlueprintReadWrite, Category = "Routing", meta = (DisplayName = "Base Submix", EditCondition = "bEnableBaseSubmix"))
 	TObjectPtr<USoundSubmixBase> SoundSubmixObject;
 
 	/** Array of submix sends to which a prescribed amount (see 'Send Level') of this sound is sent. */
-	UPROPERTY(EditAnywhere, BlueprintReadWrite, Category = "Effects|Submix", meta = (DisplayName = "Submix Sends", EditCondition = "bEnableSubmixSends"))
+	UPROPERTY(EditAnywhere, BlueprintReadWrite, Category = "Routing", meta = (DisplayName = "Submix Sends", EditCondition = "bEnableSubmixSends"))
 	TArray<FSoundSubmixSendInfo> SoundSubmixSends;
 
 	/** The source effect chain to use for this sound. */
-	UPROPERTY(EditAnywhere, BlueprintReadWrite, Category = "Effects|Source")
+	UPROPERTY(EditAnywhere, BlueprintReadWrite, Category = "Sound")
 	TObjectPtr<USoundEffectSourcePresetChain> SourceEffectChain;
 
 	/** This sound will send its audio output to this list of buses if there are bus instances playing after source effects are processed. */
-	UPROPERTY(EditAnywhere, BlueprintReadWrite, Category = "Effects|Source", meta = (DisplayName = "Post-Effect Bus Sends", EditCondition = "bEnableBusSends"))
+	UPROPERTY(EditAnywhere, BlueprintReadWrite, Category = "Routing", meta = (DisplayName = "Post-Effect Bus Sends", EditCondition = "bEnableBusSends"))
 	TArray<FSoundSourceBusSendInfo> BusSends;
 
 	/** This sound will send its audio output to this list of buses if there are bus instances playing before source effects are processed. */
-	UPROPERTY(EditAnywhere, BlueprintReadWrite, Category = "Effects|Source", meta = (DisplayName = "Pre-Effect Bus Sends", EditCondition = "bEnableBusSends"))
+	UPROPERTY(EditAnywhere, BlueprintReadWrite, Category = "Routing", meta = (DisplayName = "Pre-Effect Bus Sends", EditCondition = "bEnableBusSends"))
 	TArray<FSoundSourceBusSendInfo> PreEffectBusSends;
 
 	/** Array of user data stored with the asset */
-	UPROPERTY(EditAnywhere, AdvancedDisplay, Instanced, Category = Advanced)
+	UPROPERTY(EditAnywhere, Instanced, Category = "Advanced")
 	TArray<TObjectPtr<UAssetUserData>> AssetUserData;
 
 #if WITH_EDITORONLY_DATA	
-	UPROPERTY(EditAnywhere, Category = AudioProperties)
+	UPROPERTY(EditAnywhere, Category = "Sound")
 	TObjectPtr<UAudioPropertiesSheetAssetBase> AudioPropertiesSheet;
 
 	// //~ Begin IAudioPropertiesSheetAssetUserInterface
diff --git a/Engine/Source/Runtime/Engine/Classes/Sound/SoundCue.h b/Engine/Source/Runtime/Engine/Classes/Sound/SoundCue.h
--- a/Engine/Source/Runtime/Engine/Classes/Sound/SoundCue.h
+++ b/Engine/Source/Runtime/Engine/Classes/Sound/SoundCue.h
@@ -103,7 +103,7 @@ class USoundCue : public USoundBase
 	float PitchMultiplier;
 
 	/* Attenuation settings to use if Override Attenuation is set to true */
-	UPROPERTY(EditAnywhere, Category = Attenuation)
+	UPROPERTY(EditAnywhere, Category = "Sound")
 	FSoundAttenuationSettings AttenuationOverrides;
 
 #if WITH_EDITORONLY_DATA
@@ -116,7 +116,7 @@ class USoundCue : public USoundBase
 
 protected:
 	// NOTE: Use GetSubtitlePriority() to fetch this value for external use.
-	UPROPERTY(EditAnywhere, Category = "Voice Management|Priority", Meta = (Tooltip = "The priority of the subtitle.  Defaults to 10000.  Higher values will play instead of lower values."))
+	UPROPERTY(EditAnywhere, Category = "Voice Management", Meta = (Tooltip = "The priority of the subtitle.  Defaults to 10000.  Higher values will play instead of lower values."))
 	float SubtitlePriority;
 
 private:
@@ -125,15 +125,15 @@ private:
 public:
 
 	/* Makes this sound cue automatically load any sound waves it can play into the cache when it is loaded. */
-	UPROPERTY(EditAnywhere, Category = Memory)
+	UPROPERTY(EditAnywhere, Category = "Memory")
 	uint8 bPrimeOnLoad : 1;
 
 	/* Indicates whether attenuation should use the Attenuation Overrides or the Attenuation Settings asset */
-	UPROPERTY(EditAnywhere, Category = Attenuation)
+	UPROPERTY(EditAnywhere, Category = "Sound", meta=(DisplayAfter="AttenuationSettings"))
 	uint8 bOverrideAttenuation : 1;
 
 	/* Ignore per-platform random node culling for memory purposes */
-	UPROPERTY(EditAnywhere, Category = Memory, Meta = (DisplayName = "Disable Random Branch Culling"))
+	UPROPERTY(EditAnywhere, Category = "Memory", Meta = (DisplayName = "Disable Random Branch Culling"))
 	uint8 bExcludeFromRandomNodeBranchCulling : 1;
 
 private:
diff --git a/Engine/Source/Runtime/Engine/Classes/Sound/SoundModulationDestination.h b/Engine/Source/Runtime/Engine/Classes/Sound/SoundModulationDestination.h
--- a/Engine/Source/Runtime/Engine/Classes/Sound/SoundModulationDestination.h
+++ b/Engine/Source/Runtime/Engine/Classes/Sound/SoundModulationDestination.h
@@ -137,19 +137,19 @@ struct FSoundModulationDefaultRoutingSettings : public FSoundModulationDefaultSe
 	ENGINE_API FSoundModulationDefaultRoutingSettings();
 
 	/** What volume modulation settings to use */
-	UPROPERTY(EditAnywhere, BlueprintReadWrite, Category = "Modulation")
+	UPROPERTY(EditAnywhere, BlueprintReadWrite, Category = "Modulation", meta=(DisplayPriority=100))
 	EModulationRouting VolumeRouting = EModulationRouting::Inherit;
 
 	/** What pitch modulation settings to use */
-	UPROPERTY(EditAnywhere, BlueprintReadWrite, Category = "Modulation")
+	UPROPERTY(EditAnywhere, BlueprintReadWrite, Category = "Modulation", meta=(DisplayPriority=100))
 	EModulationRouting PitchRouting = EModulationRouting::Inherit;
 
 	/** What high-pass modulation settings to use */
-	UPROPERTY(EditAnywhere, BlueprintReadWrite, Category = "Modulation")
+	UPROPERTY(EditAnywhere, BlueprintReadWrite, Category = "Modulation", meta=(DisplayPriority=100))
 	EModulationRouting HighpassRouting = EModulationRouting::Inherit;
 
 	/** What low-pass modulation settings to use */
-	UPROPERTY(EditAnywhere, BlueprintReadWrite, Category = "Modulation")
+	UPROPERTY(EditAnywhere, BlueprintReadWrite, Category = "Modulation", meta=(DisplayPriority=100))
 	EModulationRouting LowpassRouting = EModulationRouting::Inherit;
 	
 	EModulationRouting GetRoutingForDestination(EModulationDestination Destination) const;
diff --git a/Engine/Source/Runtime/Engine/Classes/Sound/SoundSourceBus.h b/Engine/Source/Runtime/Engine/Classes/Sound/SoundSourceBus.h
--- a/Engine/Source/Runtime/Engine/Classes/Sound/SoundSourceBus.h
+++ b/Engine/Source/Runtime/Engine/Classes/Sound/SoundSourceBus.h
@@ -30,15 +30,15 @@ class USoundSourceBus : public USoundWave
 public:
 
 	/** How many channels to use for the source bus if the audio bus is not specified, otherwise it will use the audio bus object's channel count. */
-	UPROPERTY(EditAnywhere, Category = BusProperties)
+	UPROPERTY(EditAnywhere, Category = "Sound|Source Bus")
 	ESourceBusChannels SourceBusChannels;
 
 	/** The duration (in seconds) to use for the source bus. A duration of 0.0 indicates to play the source bus indefinitely. */
-	UPROPERTY(EditAnywhere, Category = BusProperties, meta = (UIMin = 0.0, ClampMin = 0.0))
+	UPROPERTY(EditAnywhere, Category = "Sound|Source Bus", meta = (UIMin = 0.0, ClampMin = 0.0))
 	float SourceBusDuration;
 
 	/** Audio bus to use as audio for this source bus. This source bus will sonify the audio from the audio bus. */
-	UPROPERTY(EditAnywhere, Category = BusProperties)
+	UPROPERTY(EditAnywhere, Category = "Sound|Source Bus", meta=(DisplayName="Audio Bus Return"))
 	TObjectPtr<UAudioBus> AudioBus;
 
 	/** Stop the source bus when the volume goes to zero. */
diff --git a/Engine/Source/Runtime/Engine/Classes/Sound/SoundWave.h b/Engine/Source/Runtime/Engine/Classes/Sound/SoundWave.h
--- a/Engine/Source/Runtime/Engine/Classes/Sound/SoundWave.h
+++ b/Engine/Source/Runtime/Engine/Classes/Sound/SoundWave.h
@@ -420,7 +420,7 @@ class USoundWave : public USoundBase, public IAudioProxyDataFactory, public IInt
 private:
 
 	/** Platform agnostic compression quality. 1..100 with 1 being best compression and 100 being best quality. ADPCM and PCM sound asset compression types ignore this parameter. */
-	UPROPERTY(Config, EditAnywhere, Category = "Format|Quality", meta = (DisplayName = "Compression", ClampMin = "1", ClampMax = "100", EditCondition = "SoundAssetCompressionType != ESoundAssetCompressionType::PCM && SoundAssetCompressionType != ESoundAssetCompressionType::ADPCM"), AssetRegistrySearchable)
+	UPROPERTY(Config, EditAnywhere, Category = "Info|Quality", meta = (DisplayName = "Compression", ClampMin = "1", ClampMax = "100", EditCondition = "SoundAssetCompressionType != ESoundAssetCompressionType::PCM && SoundAssetCompressionType != ESoundAssetCompressionType::ADPCM"), AssetRegistrySearchable)
 	int32 CompressionQuality;
 
 public:
@@ -432,17 +432,17 @@ public:
 	*	For example, if the platform enables Resampling For Device and specifies 32000 for High, then setting High here will
 	*	force the sound wave to be _at most_ 32000. Does nothing if Resampling For Device is disabled.
 	*/
-	UPROPERTY(EditAnywhere, Category = "Format|Quality")
+	UPROPERTY(EditAnywhere, Category = "Info|Quality")
 	ESoundwaveSampleRateSettings SampleRateQuality;
 
 	/** Type of buffer this wave uses. Set once on load */
 	TEnumAsByte<EDecompressionType> DecompressionType;
 
-	UPROPERTY(EditAnywhere, Category = Sound, meta = (DisplayName = "Group"))
+	UPROPERTY(EditAnywhere, Category = "Info", meta = (DisplayName = "Group"))
 	TEnumAsByte<ESoundGroup> SoundGroup;
 
 	/** If set, when played directly (not through a sound cue) the wave will be played looping. */
-	UPROPERTY(EditAnywhere, Category = Sound, AssetRegistrySearchable)
+	UPROPERTY(EditAnywhere, Category = "Sound", AssetRegistrySearchable)
 	uint8 bLooping : 1;
 
 	/** Here for legacy code. */
@@ -459,7 +459,7 @@ private:
 	uint8 bUseBinkAudio : 1;
 
 	/** The compression type to use for the sound wave asset. */
-	UPROPERTY(Config, EditAnywhere, Category = "Format", AssetRegistrySearchable)
+	UPROPERTY(Config, EditAnywhere, Category = "Info|Format", AssetRegistrySearchable)
 	ESoundAssetCompressionType SoundAssetCompressionType = ESoundAssetCompressionType::PlatformSpecific;
 
 public:
@@ -525,50 +525,50 @@ public:
 
 #if WITH_EDITORONLY_DATA
 	/** Specify a sound to use for the baked analysis. Will default to this USoundWave if not set. */
-	UPROPERTY(EditAnywhere, Category = "Analysis")
+	UPROPERTY(EditAnywhere, Category = "Info|Analysis")
 	TObjectPtr<USoundWave> OverrideSoundToUseForAnalysis;
 
 	/**
 		Whether or not we should treat the sound wave used for analysis (this or the override) as looping while performing analysis.
 		A looping sound may include the end of the file for inclusion in analysis for envelope and FFT analysis.
 	*/
-	UPROPERTY(EditAnywhere, Category = "Analysis")
+	UPROPERTY(EditAnywhere, Category = "Info|Analysis")
 	uint8 TreatFileAsLoopingForAnalysis : 1;
 
 	/** Whether or not to enable cook-time baked FFT analysis. */
-	UPROPERTY(EditAnywhere, Category = "Analysis|FFT")
+	UPROPERTY(EditAnywhere, Category = "Info|Analysis|FFT")
 	uint8 bEnableBakedFFTAnalysis : 1;
 
 	/** Whether or not to enable cook-time amplitude envelope analysis. */
-	UPROPERTY(EditAnywhere, Category = "Analysis|Envelope")
+	UPROPERTY(EditAnywhere, Category = "Info|Analysis|Envelope")
 	uint8 bEnableAmplitudeEnvelopeAnalysis : 1;
 
 	/** The FFT window size to use for fft analysis. */
-	UPROPERTY(EditAnywhere, Category = "Analysis|FFT", meta = (EditCondition = "bEnableBakedFFTAnalysis"))
+	UPROPERTY(EditAnywhere, Category = "Info|Analysis|FFT", meta = (EditCondition = "bEnableBakedFFTAnalysis"))
 	ESoundWaveFFTSize FFTSize;
 
 	/** How many audio frames analyze at a time. */
-	UPROPERTY(EditAnywhere, Category = "Analysis|FFT", meta = (EditCondition = "bEnableBakedFFTAnalysis", ClampMin = "512", UIMin = "512"))
+	UPROPERTY(EditAnywhere, Category = "Info|Analysis|FFT", meta = (EditCondition = "bEnableBakedFFTAnalysis", ClampMin = "512", UIMin = "512"))
 	int32 FFTAnalysisFrameSize;
 
 	/** Attack time in milliseconds of the spectral envelope follower. */
-	UPROPERTY(EditAnywhere, Category = "Analysis|FFT", meta = (EditCondition = "bEnableBakedFFTAnalysis", ClampMin = "0", UIMin = "0"))
+	UPROPERTY(EditAnywhere, Category = "Info|Analysis|FFT", meta = (EditCondition = "bEnableBakedFFTAnalysis", ClampMin = "0", UIMin = "0"))
 	int32 FFTAnalysisAttackTime;
 
 	/** Release time in milliseconds of the spectral envelope follower. */
-	UPROPERTY(EditAnywhere, Category = "Analysis|FFT", meta = (EditCondition = "bEnableBakedFFTAnalysis", ClampMin = "0", UIMin = "0"))
+	UPROPERTY(EditAnywhere, Category = "Info|Analysis|FFT", meta = (EditCondition = "bEnableBakedFFTAnalysis", ClampMin = "0", UIMin = "0"))
 	int32 FFTAnalysisReleaseTime;
 
 	/** How many audio frames to average a new envelope value. Larger values use less memory for audio envelope data but will result in lower envelope accuracy. */
-	UPROPERTY(EditAnywhere, Category = "Analysis|Envelope", meta = (EditCondition = "bEnableAmplitudeEnvelopeAnalysis", ClampMin = "512", UIMin = "512"))
+	UPROPERTY(EditAnywhere, Category = "Info|Analysis|Envelope", meta = (EditCondition = "bEnableAmplitudeEnvelopeAnalysis", ClampMin = "512", UIMin = "512"))
 	int32 EnvelopeFollowerFrameSize;
 
 	/** The attack time in milliseconds. Describes how quickly the envelope analyzer responds to increasing amplitudes. */
-	UPROPERTY(EditAnywhere, Category = "Analysis|Envelope", meta = (EditCondition = "bEnableAmplitudeEnvelopeAnalysis", ClampMin = "0", UIMin = "0"))
+	UPROPERTY(EditAnywhere, Category = "Info|Analysis|Envelope", meta = (EditCondition = "bEnableAmplitudeEnvelopeAnalysis", ClampMin = "0", UIMin = "0"))
 	int32 EnvelopeFollowerAttackTime;
 
 	/** The release time in milliseconds. Describes how quickly the envelope analyzer responds to decreasing amplitudes. */
-	UPROPERTY(EditAnywhere, Category = "Analysis|Envelope", meta = (EditCondition = "bEnableAmplitudeEnvelopeAnalysis", ClampMin = "0", UIMin = "0"))
+	UPROPERTY(EditAnywhere, Category = "Info|Analysis|Envelope", meta = (EditCondition = "bEnableAmplitudeEnvelopeAnalysis", ClampMin = "0", UIMin = "0"))
 	int32 EnvelopeFollowerReleaseTime;
 #endif // WITH_EDITORONLY_DATA
 
@@ -577,7 +577,7 @@ public:
 	FSoundModulationDefaultRoutingSettings ModulationSettings;
 
 	/** The frequencies (in hz) to analyze when doing baked FFT analysis. */
-	UPROPERTY(EditAnywhere, Category = "Analysis|FFT", meta = (EditCondition = "bEnableBakedFFTAnalysis"))
+	UPROPERTY(EditAnywhere, Category = "Info|Analysis|FFT", meta = (EditCondition = "bEnableBakedFFTAnalysis"))
 	TArray<float> FrequenciesToAnalyze;
 
 	/** The cooked spectral time data. */
@@ -681,15 +681,15 @@ public:
 	uint8 bDynamicResource : 1;
 
 	/** If set to true if this sound is considered to contain mature/adult content. */
-	UPROPERTY(EditAnywhere, BlueprintReadOnly, Category = Subtitles, AssetRegistrySearchable)
+	UPROPERTY(EditAnywhere, BlueprintReadOnly, Category = "Info", AssetRegistrySearchable)
 	uint8 bMature : 1;
 
 	/** If set to true will disable automatic generation of line breaks - use if the subtitles have been split manually. */
-	UPROPERTY(EditAnywhere, BlueprintReadOnly, Category = Subtitles)
+	UPROPERTY(EditAnywhere, BlueprintReadOnly, Category = "Info|Subtitles")
 	uint8 bManualWordWrap : 1;
 
 	/** If set to true the subtitles display as a sequence of single lines as opposed to multiline. */
-	UPROPERTY(EditAnywhere, BlueprintReadOnly, Category = Subtitles)
+	UPROPERTY(EditAnywhere, BlueprintReadOnly, Category = "Info|Subtitles")
 	uint8 bSingleLine : 1;
 
 #if WITH_EDITORONLY_DATA
@@ -700,7 +700,7 @@ public:
 	/** Whether or not this source is ambisonics file format. If set, sound always uses the
 	  * 'Master Ambisonics Submix' as set in the 'Audio' category of Project Settings'
 	  * and ignores submix if provided locally or in the referenced SoundClass. */
-	UPROPERTY(EditAnywhere, Category = Format)
+	UPROPERTY(EditAnywhere, Category = "Info")
 	uint8 bIsAmbisonics : 1;
 
 	/** Whether this SoundWave was decompressed from OGG. */
@@ -739,7 +739,7 @@ private:
 public:
 	// Loading behavior members are lazily initialized in const getters
 	/** Specifies how and when compressed audio data is loaded for asset if stream caching is enabled. */
-	UPROPERTY(EditAnywhere, Category = "Loading", meta = (DisplayName = "Loading Behavior Override"))
+	UPROPERTY(EditAnywhere, Category = "Memory", meta = (DisplayName = "Loading Behavior Override"))
 	mutable ESoundWaveLoadingBehavior LoadingBehavior;
 	
 private:
@@ -749,7 +749,7 @@ private:
 public:
 #if WITH_EDITORONLY_DATA
    	/** How much audio to add to First Audio Chunk (in seconds) */
-	UPROPERTY(EditAnywhere, Category = Loading, meta = (UIMin = 0, UIMax = 10, EditCondition = "LoadingBehavior == ESoundWaveLoadingBehavior::RetainOnLoad || LoadingBehavior == ESoundWaveLoadingBehavior::PrimeOnLoad"), DisplayName="Size of First Audio Chunk (seconds)")
+	UPROPERTY(EditAnywhere, Category = "Memory", meta = (UIMin = 0, UIMax = 10, EditCondition = "LoadingBehavior == ESoundWaveLoadingBehavior::RetainOnLoad || LoadingBehavior == ESoundWaveLoadingBehavior::PrimeOnLoad", EditConditionHides), DisplayName="Size of First Audio Chunk (seconds)")
    	FPerPlatformFloat SizeOfFirstAudioChunkInSeconds = 0.0f;
 #endif //WITH_EDITOR_ONLY_DATA
 
@@ -758,19 +758,19 @@ public:
 	FString SpokenText_DEPRECATED;
 
 	/** The priority of the subtitle. */
-	UPROPERTY(EditAnywhere, BlueprintReadOnly, Category = Subtitles)
+	UPROPERTY(EditAnywhere, BlueprintReadOnly, Category = "Info|Subtitles")
 	float SubtitlePriority;
 
 	/** Linear volume multiplier for sound playback, default is 1.0 */
-	UPROPERTY(Category = Sound, meta = (ClampMin = "0.0"), EditAnywhere)
+	UPROPERTY(Category = "Sound", meta = (ClampMin = "0.0"), EditAnywhere)
 	float Volume;
 
 	/** Playback rate multiplier for sound, range is 0.125 - 4.0, default is 1.0  */
-	UPROPERTY(Category = Sound, meta = (ClampMin = "0.125", ClampMax = "4.0"), EditAnywhere)
+	UPROPERTY(Category = "Sound", meta = (ClampMin = "0.125", ClampMax = "4.0"), EditAnywhere)
 	float Pitch;
 
 	/** Number of channels of multichannel data; 1 or 2 for regular mono and stereo files */
-	UPROPERTY(Category = Info, AssetRegistrySearchable, VisibleAnywhere)
+	UPROPERTY(Category = "Info|Format", AssetRegistrySearchable, VisibleAnywhere)
 	int32 NumChannels;
 
 #if WITH_EDITORONLY_DATA
@@ -783,11 +783,11 @@ public:
 	TArray<int32> ChannelSizes;
 
 	/** Gated Loudness Units Full Scale of the SoundWave. Can be modified by transformations. */
-	UPROPERTY(Category = Info, AssetRegistrySearchable, VisibleAnywhere)
+	UPROPERTY(Category = "Info|Analysis", AssetRegistrySearchable, VisibleAnywhere)
 	float LUFS;
 
 	/** Highest Sample Peak in db of the SoundWave. Can be modified by transformations. */
-	UPROPERTY(Category = Info, AssetRegistrySearchable, VisibleAnywhere)
+	UPROPERTY(Category = "Info|Analysis", AssetRegistrySearchable, VisibleAnywhere)
 	float SamplePeakDB;
 
 #endif // WITH_EDITORONLY_DATA
@@ -795,22 +795,22 @@ public:
 protected:
 
 	/** Cooked sample rate of the asset. Can be modified by sample rate override. */
-	UPROPERTY(Category = Info, AssetRegistrySearchable, VisibleAnywhere)
+	UPROPERTY(Category = "Info|Format", AssetRegistrySearchable, VisibleAnywhere)
 	int32 SampleRate;
 
 #if WITH_EDITORONLY_DATA
 	/** Sample rate of the imported sound wave. */
-	UPROPERTY(Category = Info, AssetRegistrySearchable, VisibleAnywhere)
+	UPROPERTY(Category = "Info|Format", AssetRegistrySearchable, VisibleAnywhere)
 	int32 ImportedSampleRate;
 
 	/** Cue point data parsed from the .wav file. Contains "Loop Regions" as cue points as well! */
-	UPROPERTY(Category = Info, VisibleAnywhere, BlueprintGetter = GetCuePoints)
+	UPROPERTY(Category = "Info|Markers", VisibleAnywhere, BlueprintGetter = GetCuePoints)
 	TArray<FSoundWaveCuePoint> CuePoints;
 
 	/** Dictates whether to use the CuePoints and Loop Regions from the .wav file or from the waveform editor 
 	 *	Marker transformations during playback.
 	 */
-	UPROPERTY(Category = Info, EditAnywhere)
+	UPROPERTY(Category = "Info|Markers", EditAnywhere)
 	ESoundWaveCuePointOrigin CuePointOrigin;
 #endif
 
@@ -832,12 +832,12 @@ public:
 	/**
 	 * Subtitle cues. 
 	 */
-	UPROPERTY(EditAnywhere, BlueprintReadOnly, Category = Subtitles)
+	UPROPERTY(EditAnywhere, BlueprintReadOnly, Category = "Info|Subtitles")
 	TArray<struct FSubtitleCue> Subtitles;
 
 #if WITH_EDITORONLY_DATA
 	/** Provides contextual information for the sound to the translator. */
-	UPROPERTY(EditAnywhere, Category = Subtitles)
+	UPROPERTY(EditAnywhere, Category = "Info|Subtitles")
 	FString Comment;
 
 #endif // WITH_EDITORONLY_DATA
@@ -845,7 +845,7 @@ public:
 #if WITH_EDITORONLY_DATA
 	
 	/** Information about the time-code from import, if available.  */
-	UPROPERTY(VisibleAnywhere, Category = Info)
+	UPROPERTY(VisibleAnywhere, Category = "Info", Meta=(DisplayName="Timecode"))
 	FSoundWaveTimecodeInfo TimecodeInfo;
 
 #endif // WITH_EDITORONLY_DATA
@@ -865,7 +865,7 @@ public:
 protected:
 
 	/** Curves associated with this sound wave */
-	UPROPERTY(EditAnywhere, Category = SoundWave, AdvancedDisplay)
+	UPROPERTY(EditAnywhere, Category = "Info|Curves")
 	TObjectPtr<class UCurveTable> Curves;
 
 	/** Hold a reference to our internal curve so we can switch back to it if we want to */
@@ -875,10 +875,10 @@ protected:
 #if WITH_EDITORONLY_DATA
 protected:
 	/** If enabled, this wave may be streamed from the cloud using the Opus format. Loading behavior must NOT be `Force Inline`. Requires a suitable support plugin to be installed. */
-	UPROPERTY(EditAnywhere, Category = "Format", Meta=(DisplayName="Enable cloud streaming", DisplayAfter="SoundAssetCompressionType", EditCondition = "LoadingBehavior != ESoundWaveLoadingBehavior::ForceInline"), AssetRegistrySearchable)
+	UPROPERTY(EditAnywhere, Category = "Info|Format", Meta=(DisplayName="Enable Cloud Streaming", DisplayAfter="SoundAssetCompressionType", EditCondition = "LoadingBehavior != ESoundWaveLoadingBehavior::ForceInline"), AssetRegistrySearchable)
 	uint8 bEnableCloudStreaming : 1;
 	/** Platform specific. */
-	UPROPERTY(EditAnywhere, config, Category="Platform specific", Meta=(DisplayName="Platform specific settings", ToolTip="Optionally disables cloud streaming per platform"))
+	UPROPERTY(EditAnywhere, config, Category="Info|Format", Meta=(DisplayName="Platform Specific Settings", ToolTip="Optionally disables cloud streaming per platform"))
 	TMap<FGuid, FSoundWaveCloudStreamingPlatformSettings> PlatformSettings;
 public:
 	static FName GetCloudStreamingEnabledPropertyName() { return GET_MEMBER_NAME_CHECKED(USoundWave, bEnableCloudStreaming); }
