[Engine]
* Fixed header unit compile errors

diff --git a/Engine/Source/Runtime/Engine/Classes/Components/StaticMeshComponent.h b/Engine/Source/Runtime/Engine/Classes/Components/StaticMeshComponent.h
--- a/Engine/Source/Runtime/Engine/Classes/Components/StaticMeshComponent.h
+++ b/Engine/Source/Runtime/Engine/Classes/Components/StaticMeshComponent.h
@@ -29,7 +29,9 @@ class FStaticMeshStaticLightingMesh;
 class FVertexFactoryType;
 class ULightComponent;
 class UStaticMesh;
+class UStaticMeshSocket;
 class UStaticMeshComponent;
+class UMaterialCacheVirtualTexture;
 class UNavCollisionBase;
 class UWorld;
 struct FConvexVolume;
@@ -127,13 +129,13 @@ public:
 private:
 	/** The static mesh that this component uses to render */
 	UPROPERTY(EditAnywhere, BlueprintReadOnly, Category=StaticMesh, ReplicatedUsing=OnRep_StaticMesh, meta=(AllowPrivateAccess="true"))
-	TObjectPtr<class UStaticMesh> StaticMesh;
+	TObjectPtr<UStaticMesh> StaticMesh;
 
 	ENGINE_API void SetStaticMeshInternal(UStaticMesh* StaticMesh);
 
 #if WITH_EDITOR
 	/** Used to track down when StaticMesh has been modified for notification purpose */
-	class UStaticMesh* KnownStaticMesh = nullptr;
+	UStaticMesh* KnownStaticMesh = nullptr;
 #endif
 	ENGINE_API void NotifyIfStaticMeshChanged();
 
@@ -146,7 +148,7 @@ public:
 	static const FName GetMemberNameChecked_StaticMesh() { return GET_MEMBER_NAME_CHECKED(UStaticMeshComponent, StaticMesh); }
 
 	UFUNCTION()
-	ENGINE_API void OnRep_StaticMesh(class UStaticMesh *OldStaticMesh);
+	ENGINE_API void OnRep_StaticMesh(UStaticMesh *OldStaticMesh);
 
 	/**
 	* Distance at which to disable World Position Offset for an entire instance (0 = Never disable WPO).
@@ -374,7 +376,7 @@ public:
 
 	/** All allocated material cache textures */
 	UPROPERTY(Transient, Experimental)
-	TArray<TObjectPtr<class UMaterialCacheVirtualTexture>> MaterialCacheTextures;
+	TArray<TObjectPtr<UMaterialCacheVirtualTexture>> MaterialCacheTextures;
 
 	/** 
 	 * Controls how dark the dynamic indirect shadow can be.
@@ -430,7 +432,7 @@ public:
 
 	/** Change the StaticMesh used by this instance. */
 	UFUNCTION(BlueprintCallable, Category="Components|StaticMesh")
-	ENGINE_API virtual bool SetStaticMesh(class UStaticMesh* NewMesh);
+	ENGINE_API virtual bool SetStaticMesh(UStaticMesh* NewMesh);
 
 	/** Get the StaticMesh used by this instance. */
 	TObjectPtr<UStaticMesh> GetStaticMesh() const 
@@ -889,7 +891,7 @@ public:
 	 *
 	 * @return UStaticMeshSocket of named socket on the static mesh component. None if not found.
 	 */
-	ENGINE_API class UStaticMeshSocket const* GetSocketByName( FName InSocketName ) const;
+	ENGINE_API UStaticMeshSocket const* GetSocketByName( FName InSocketName ) const;
 
 	/** Returns the wireframe color to use for this component. */
 	ENGINE_API FColor GetWireframeColor() const;
@@ -976,7 +978,7 @@ public:
 
 	/** Mesh being used by component */
 	UPROPERTY()
-	TObjectPtr<class UStaticMesh> StaticMesh = nullptr;
+	TObjectPtr<UStaticMesh> StaticMesh = nullptr;
 
 	/** Array of cached vertex colors for each LOD */
 	UPROPERTY()
diff --git a/Engine/Source/Runtime/Engine/Public/AlphaBlend.h b/Engine/Source/Runtime/Engine/Public/AlphaBlend.h
--- a/Engine/Source/Runtime/Engine/Public/AlphaBlend.h
+++ b/Engine/Source/Runtime/Engine/Public/AlphaBlend.h
@@ -30,7 +30,7 @@ enum class EAlphaBlendOption : uint8
 
 namespace UE::Anim
 {
-static constexpr EAlphaBlendOption DefaultBlendOption = EAlphaBlendOption::HermiteCubic;
+	inline constexpr EAlphaBlendOption DefaultBlendOption = EAlphaBlendOption::HermiteCubic;
 }
 
 /**
