Motion Matching slows down editor when editing curves in an animation
#fix Add FMutex to split out locking for async compression tasks map from compressed data map, as required locking behaviour could otherwise deadlock
#FailedValidations Preflight

diff --git a/Engine/Source/Runtime/Engine/Classes/Animation/AnimSequence.h b/Engine/Source/Runtime/Engine/Classes/Animation/AnimSequence.h
--- a/Engine/Source/Runtime/Engine/Classes/Animation/AnimSequence.h
+++ b/Engine/Source/Runtime/Engine/Classes/Animation/AnimSequence.h
@@ -810,6 +810,7 @@ protected:
 	TMap<FIoHash, TUniquePtr<FCompressedAnimSequence>> DataByPlatformKeyHash;
 	TMap<FIoHash, TPimplPtr<UE::Anim::FAnimationSequenceAsyncCacheTask>> CacheTasksByKeyHash;
 	mutable UE::FSharedRecursiveMutex SharedCompressedDataMutex;
+	mutable UE::FMutex CacheTasksMutex;
 
 	// Cached FIoHash keys stored according to ITargetPlatform type-hash
 	mutable TMap<uint32, FIoHash> PlatformHashToKeyHash;
diff --git a/Engine/Source/Runtime/Engine/Private/Animation/AnimSequence.cpp b/Engine/Source/Runtime/Engine/Private/Animation/AnimSequence.cpp
--- a/Engine/Source/Runtime/Engine/Private/Animation/AnimSequence.cpp
+++ b/Engine/Source/Runtime/Engine/Private/Animation/AnimSequence.cpp
@@ -1377,6 +1377,7 @@ bool UAnimSequence::HasCompressedDataForHash(const FIoHash& InKeyHash) const
 	{
 		if ((*CompressedDataPtr)->IsValid(this))
 		{
+			UE::TScopeLock LockGuard(CacheTasksMutex);
 			if (CacheTasksByKeyHash.Contains(InKeyHash))
 			{
 				return false;
@@ -1406,6 +1407,7 @@ bool UAnimSequence::IsCachedCookedPlatformDataLoaded(const ITargetPlatform* Targ
 	bool bShouldRecompress = false;
 	{
 		UE_COMPRESSED_DATA_READ_SCOPE(this);
+		UE::TScopeLock LockGuard(CacheTasksMutex);
 		if(const TUniquePtr<FCompressedAnimSequence>* CompressedDataPtr = DataByPlatformKeyHash.Find(KeyHash))
 		{
 			if (!CacheTasksByKeyHash.Contains(KeyHash) && (*CompressedDataPtr)->IsValid(this))
@@ -5149,6 +5151,7 @@ FIoHash UAnimSequence::BeginCacheDerivedData(const ITargetPlatform* TargetPlatfo
 	const FIoHash KeyHash = CreateDerivedDataKeyHash(TargetPlatform);
 
 	UE_COMPRESSED_DATA_WRITE_SCOPE(this);
+	UE::TScopeLock LockGuard(CacheTasksMutex);
 	FCompressedAnimSequence* TargetData = nullptr;
 	{
 		checkf(!CacheTasksByKeyHash.Contains(KeyHash), TEXT("[%s] Hash %s still/already has task in-flight"), *GetPathName(), *LexToString(KeyHash));
@@ -5231,7 +5234,7 @@ bool UAnimSequence::PollCacheDerivedData(const FIoHash& KeyHash) const
 		return true;
 	}
 
-	UE_COMPRESSED_DATA_READ_SCOPE(this);
+	UE::TScopeLock LockGuard(CacheTasksMutex);
 	if (const TPimplPtr<UE::Anim::FAnimationSequenceAsyncCacheTask>* Task = CacheTasksByKeyHash.Find(KeyHash))
 	{
 		return (*Task)->Poll();
@@ -5289,7 +5292,7 @@ void UAnimSequence::CacheDerivedDataForCurrentPlatform()
 				
 bool UAnimSequence::IsCompiling() const
 {
-	UE_COMPRESSED_DATA_READ_SCOPE(this);	
+	UE::TScopeLock LockGuard(CacheTasksMutex);
 	return CacheTasksByKeyHash.Num() > 0;
 }
 
@@ -5298,12 +5301,14 @@ void UAnimSequence::ClearAllCompressionData()
 	{
 		// Delete any cache tasks first because the destructor will cancel the cache and build tasks,
 		// and drop their pointers to the data.
-
-		for (auto It = CacheTasksByKeyHash.CreateIterator(); It; ++It)
 		{
-			if (!RequiresResidency(It->Key))
+			UE::TScopeLock LockGuard(CacheTasksMutex);
+			for (auto It = CacheTasksByKeyHash.CreateIterator(); It; ++It)
 			{
-				It.RemoveCurrent();
+				if (!RequiresResidency(It->Key))
+				{
+					It.RemoveCurrent();
+				}
 			}
 		}
 
@@ -5421,7 +5426,7 @@ bool UAnimSequence::HasCompressedDataForPlatform(const ITargetPlatform* InPlatfo
 
 bool UAnimSequence::TryCancelAsyncTasks()
 {
-	UE_COMPRESSED_DATA_READ_SCOPE(this);
+	UE::TScopeLock LockGuard(CacheTasksMutex);
 	for (auto It = CacheTasksByKeyHash.CreateIterator(); It; ++It)
 	{
 		It->Value->Cancel();
@@ -5432,8 +5437,8 @@ bool UAnimSequence::TryCancelAsyncTasks()
 
 bool UAnimSequence::WaitForAsyncTasks(float TimeLimitSeconds)
 {
-	UE_COMPRESSED_DATA_READ_SCOPE(this);
 	double StartTimeSeconds = FPlatformTime::Seconds();
+	UE::TScopeLock LockGuard(CacheTasksMutex);
 	for (auto& Pair : CacheTasksByKeyHash)
 	{
 		// Clamp to 0 as it implies polling
@@ -5448,7 +5453,7 @@ bool UAnimSequence::WaitForAsyncTasks(float TimeLimitSeconds)
 
 bool UAnimSequence::IsAsyncTaskComplete() const
 {
-	UE_COMPRESSED_DATA_READ_SCOPE(this);
+	UE::TScopeLock LockGuard(CacheTasksMutex);
 	bool bAllFinished = true;
 	for (auto& Pair : CacheTasksByKeyHash)
 	{
@@ -5463,6 +5468,7 @@ void UAnimSequence::FinishAsyncTasks()
 	const ITargetPlatform* RunningPlatform = GetTargetPlatformManagerRef().GetRunningTargetPlatform();
 	const FIoHash CurrentPlatformHash = CreateDerivedDataKeyHash(RunningPlatform);
 
+	UE::TScopeLock LockGuard(CacheTasksMutex);
 	const bool bHasInflightTasks = CacheTasksByKeyHash.Num() > 0;
 	if (bHasInflightTasks)
 	{
@@ -5526,7 +5532,7 @@ void UAnimSequence::FinishAsyncTasks()
 
 void UAnimSequence::Reschedule(FQueuedThreadPool* InThreadPool, EQueuedWorkPriority InPriority)
 {
-	UE_COMPRESSED_DATA_READ_SCOPE(this);
+	UE::TScopeLock LockGuard(CacheTasksMutex);
 	for (auto& Pair : CacheTasksByKeyHash)
 	{
 		Pair.Value->Reschedule(InThreadPool, InPriority);
