Added FMath::SafeDivide that performs a safe division, which returns the largest possible value when divisor is zero.
Fixed several cases of potential division by zeros in world partition code where we compute the minimum resulting division value.
Fixed code typos (soure -> source).

diff --git a/Engine/Source/Runtime/Core/Public/Math/UnrealMathUtility.h b/Engine/Source/Runtime/Core/Public/Math/UnrealMathUtility.h
--- a/Engine/Source/Runtime/Core/Public/Math/UnrealMathUtility.h
+++ b/Engine/Source/Runtime/Core/Public/Math/UnrealMathUtility.h
@@ -712,6 +712,13 @@ public:
 			: (Dividend - Divisor / 2 + 1) / Divisor;
 	}
 
+	/** Performs a safe division, which returns the largest possible value when divisor is nearly zero */
+	template <class T, class U> 
+	[[nodiscard]] static constexpr auto SafeDivide(T Dividend, U Divisor) -> decltype(Dividend / Divisor)
+	{
+		return FMath::IsNearlyZero(Divisor) ? TNumericLimits<decltype(Dividend / Divisor)>::Max() : (Dividend / Divisor);
+	};
+
 	/**
 	 * Computes the base 2 logarithm of the specified value
 	 *
diff --git a/Engine/Source/Runtime/Engine/Private/WorldPartition/WorldPartitionRuntimeCellData.cpp b/Engine/Source/Runtime/Engine/Private/WorldPartition/WorldPartitionRuntimeCellData.cpp
--- a/Engine/Source/Runtime/Engine/Private/WorldPartition/WorldPartitionRuntimeCellData.cpp
+++ b/Engine/Source/Runtime/Engine/Private/WorldPartition/WorldPartitionRuntimeCellData.cpp
@@ -84,20 +84,20 @@ void UWorldPartitionRuntimeCellData::AppendStreamingSourceInfo(const FWorldParti
 
 		CachedMinSquareDistanceToBlockingSource = FMath::Min(CellToSourceSquareDistance, CachedMinSquareDistanceToBlockingSource);
 
-		const double BlockOnSlowStreamingRatio = FMath::Sqrt(CachedMinSquareDistanceToBlockingSource) / SourceShape.GetRadius();
+		const double BlockOnSlowStreamingRatio = FMath::SafeDivide(FMath::Sqrt(CachedMinSquareDistanceToBlockingSource), SourceShape.GetRadius());
 		CachedMinBlockOnSlowStreamingRatio = FMath::Min(CachedMinBlockOnSlowStreamingRatio, BlockOnSlowStreamingRatio);
 	}
 
 	CachedMinSquareDistanceToSource = FMath::Min(CellToSourceSquareDistance, CachedMinSquareDistanceToSource);
 
-	const double SlowStreamingRatio = FMath::Sqrt(CachedMinSquareDistanceToSource) / SourceShape.GetRadius();
+	const double SlowStreamingRatio = FMath::SafeDivide(FMath::Sqrt(CachedMinSquareDistanceToSource), SourceShape.GetRadius());
 	CachedMinSlowStreamingRatio = FMath::Min(CachedMinSlowStreamingRatio, SlowStreamingRatio);
 
 	// Compute square distance from cell to source ratio
-	const double SoureDistanceRatio = FMath::Clamp(ContentBounds.ComputeSquaredDistanceToPoint(SourceShape.GetCenter()) / FMath::Square(SourceShape.GetRadius()), 0.0f, 1.0f);
+	const double SourceDistanceRatio = FMath::Clamp(FMath::SafeDivide(ContentBounds.ComputeSquaredDistanceToPoint(SourceShape.GetCenter()), FMath::Square(SourceShape.GetRadius())), 0.0f, 1.0f);
 
 	// Compute final cell priority for this source
-	const double SortingPriority = SoureDistanceRatio * SourceCosAngleRatio;
+	const double SortingPriority = SourceDistanceRatio * SourceCosAngleRatio;
 
 	// Update if lower
 	CachedMinSpatialSortingPriority = FMath::Min(CachedMinSpatialSortingPriority, SortingPriority);
diff --git a/Engine/Source/Runtime/Engine/Private/WorldPartition/WorldPartitionRuntimeCellDataSpatialHash.cpp b/Engine/Source/Runtime/Engine/Private/WorldPartition/WorldPartitionRuntimeCellDataSpatialHash.cpp
--- a/Engine/Source/Runtime/Engine/Private/WorldPartition/WorldPartitionRuntimeCellDataSpatialHash.cpp
+++ b/Engine/Source/Runtime/Engine/Private/WorldPartition/WorldPartitionRuntimeCellDataSpatialHash.cpp
@@ -100,10 +100,10 @@ void UWorldPartitionRuntimeCellDataSpatialHash::AppendStreamingSourceInfo(const
 	if (bBlockOnSlowLoading && Source.bBlockOnSlowLoading)
 	{
 		CachedMinSquareDistanceToBlockingSource2D = FMath::Min(SquareDistance2D, CachedMinSquareDistanceToBlockingSource2D);
-		CachedMinBlockOnSlowStreamingRatio2D = FMath::Min(CachedMinBlockOnSlowStreamingRatio2D, FMath::Sqrt(CachedMinSquareDistanceToBlockingSource2D) / SourceShape.GetRadius());
+		CachedMinBlockOnSlowStreamingRatio2D = FMath::Min(CachedMinBlockOnSlowStreamingRatio2D, FMath::SafeDivide(FMath::Sqrt(CachedMinSquareDistanceToBlockingSource2D), SourceShape.GetRadius()));
 	}
 	CachedMinSquareDistanceToSource2D = FMath::Min(SquareDistance2D, CachedMinSquareDistanceToSource2D);
-	CachedMinSlowStreamingRatio2D = FMath::Min(CachedMinSlowStreamingRatio2D, FMath::Sqrt(CachedMinSquareDistanceToSource2D) / SourceShape.GetRadius());
+	CachedMinSlowStreamingRatio2D = FMath::Min(CachedMinSlowStreamingRatio2D, FMath::SafeDivide(FMath::Sqrt(CachedMinSquareDistanceToSource2D), SourceShape.GetRadius()));
 
 	CachedSourceSquaredDistances.Add(SquareDistance2D);
 	CachedSourcePriorityWeights.Add(1.0f - ((float)Source.Priority / (float)EStreamingSourcePriority::Lowest));
