Reduce temporary memory allocations when iterating components registered on the streaming managers

diff --git a/Engine/Source/Runtime/Engine/Private/Streaming/DynamicTextureInstanceManager.cpp b/Engine/Source/Runtime/Engine/Private/Streaming/DynamicTextureInstanceManager.cpp
--- a/Engine/Source/Runtime/Engine/Private/Streaming/DynamicTextureInstanceManager.cpp
+++ b/Engine/Source/Runtime/Engine/Private/Streaming/DynamicTextureInstanceManager.cpp
@@ -167,11 +167,6 @@ void FDynamicRenderAssetInstanceManager::IncrementalUpdate(FRemovedRenderAssetAr
 	Refresh(Percentage);
 }
 
-void FDynamicRenderAssetInstanceManager::GetReferencedComponents(TArray<const UPrimitiveComponent *>& Components)
-{
-	StateSync.SyncAndGetState()->GetReferencedComponents(Components);
-}
-
 void FDynamicRenderAssetInstanceManager::OnCreateViewDone(FRenderAssetInstanceView* InView)
 {
 	// Don't call get state here to prevent recursion as this is a task callback.
diff --git a/Engine/Source/Runtime/Engine/Private/Streaming/DynamicTextureInstanceManager.h b/Engine/Source/Runtime/Engine/Private/Streaming/DynamicTextureInstanceManager.h
--- a/Engine/Source/Runtime/Engine/Private/Streaming/DynamicTextureInstanceManager.h
+++ b/Engine/Source/Runtime/Engine/Private/Streaming/DynamicTextureInstanceManager.h
@@ -28,8 +28,12 @@ public:
 
 	void IncrementalUpdate(FRemovedRenderAssetArray& RemovedRenderAssets, float Percentage);
 
-	// Get all (non removed) components refered by the manager. Debug only.
-	void GetReferencedComponents(TArray<const UPrimitiveComponent*>& Components);
+	// Iterate all (non removed) components refered by the manager. Debug only.
+	template<typename CallableT>
+	void ForEachReferencedComponent(CallableT Callable)
+	{
+		StateSync.SyncAndGetState()->ForEachReferencedComponent(Callable);
+	}
 
 	/** Remove all pending components that are marked for delete. This prevents searching in the pending list for each entry. */
 	void OnPreGarbageCollect(FRemovedRenderAssetArray& RemovedRenderAssets);
diff --git a/Engine/Source/Runtime/Engine/Private/Streaming/LevelTextureManager.cpp b/Engine/Source/Runtime/Engine/Private/Streaming/LevelTextureManager.cpp
--- a/Engine/Source/Runtime/Engine/Private/Streaming/LevelTextureManager.cpp
+++ b/Engine/Source/Runtime/Engine/Private/Streaming/LevelTextureManager.cpp
@@ -14,6 +14,7 @@
 #include "Streaming/StreamingManagerTexture.h"
 #include "Async/ParallelFor.h"
 #include "Materials/MaterialInterface.h"
+#include "Algo/ForEach.h"
 
 FLevelRenderAssetManager::FLevelRenderAssetManager(ULevel* InLevel, RenderAssetInstanceTask::FDoWorkTask& AsyncTask)
 	: Level(InLevel)
@@ -31,22 +32,24 @@ FLevelRenderAssetManager::FLevelRenderAssetManager(ULevel* InLevel, RenderAssetI
 
 void FLevelRenderAssetManager::Remove(FRemovedRenderAssetArray* RemovedRenderAssets)
 { 
-	TArray<const UPrimitiveComponent*> ReferencedComponents;
-	StaticInstances.GetReferencedComponents(ReferencedComponents);
-	ReferencedComponents.Append(UnprocessedComponents);
-	ReferencedComponents.Append(PendingComponents);
-	for (const UPrimitiveComponent* Component : ReferencedComponents)
+	auto DetachComponent = [](const UPrimitiveComponent* PrimitiveComponent)
 	{
-		if (Component)
+		if (PrimitiveComponent)
 		{
-			check(Component->IsValidLowLevelFast()); // Check that this component was not already destroyed.
+			// Check that this component was not already destroyed.
+			check(PrimitiveComponent->IsValidLowLevelFast());
+
 			// Don't check as there can be duplicates in PendingComponents
-			// check(Component->bAttachedToStreamingManagerAsStatic);  
+			//check(PrimitiveComponent->bAttachedToStreamingManagerAsStatic);
 
 			// A component can only be referenced in one level, so if it was here, we can clear the flag
-			Component->bAttachedToStreamingManagerAsStatic = false;
+			PrimitiveComponent->bAttachedToStreamingManagerAsStatic = false;
 		}
-	}
+	};
+
+	StaticInstances.ForEachReferencedComponent(DetachComponent);
+	Algo::ForEach(UnprocessedComponents, DetachComponent);
+	Algo::ForEach(PendingComponents, DetachComponent);
 
 	// Mark all static textures/meshes for removal.
 	if (RemovedRenderAssets)
diff --git a/Engine/Source/Runtime/Engine/Private/Streaming/StaticTextureInstanceManager.h b/Engine/Source/Runtime/Engine/Private/Streaming/StaticTextureInstanceManager.h
--- a/Engine/Source/Runtime/Engine/Private/Streaming/StaticTextureInstanceManager.h
+++ b/Engine/Source/Runtime/Engine/Private/Streaming/StaticTextureInstanceManager.h
@@ -34,8 +34,12 @@ public:
 	/** Return whether this component can be managed by this manager. */
 	FORCEINLINE bool IsReferenced(const UPrimitiveComponent* Component) const { return StateSync.GetState()->HasComponentReferences(Component); }
 
-	/** Return whether this component can be managed by this manager. */
-	FORCEINLINE void GetReferencedComponents(TArray<const UPrimitiveComponent*>& Components) const { StateSync.GetState()->GetReferencedComponents(Components); }
+	/** Invoke a callable for each referenced component. */
+	template<typename CallableT>
+	void ForEachReferencedComponent(CallableT Callable)
+	{
+		StateSync.GetState()->ForEachReferencedComponent(Callable);
+	}
 
 	/** Return whether this component is be managed by this manager. */
 	bool CanManage(const UPrimitiveComponent* Component) const;
diff --git a/Engine/Source/Runtime/Engine/Private/Streaming/StreamingManagerTexture.cpp b/Engine/Source/Runtime/Engine/Private/Streaming/StreamingManagerTexture.cpp
--- a/Engine/Source/Runtime/Engine/Private/Streaming/StreamingManagerTexture.cpp
+++ b/Engine/Source/Runtime/Engine/Private/Streaming/StreamingManagerTexture.cpp
@@ -677,12 +677,11 @@ void FRenderAssetStreamingManager::ConditionalUpdateStaticData()
 			}
 
 			// Reinsert dynamic components
-			TArray<const UPrimitiveComponent*> DynamicComponents;
-			DynamicComponentManager.GetReferencedComponents(DynamicComponents);
-			for (const UPrimitiveComponent* Primitive : DynamicComponents)
-			{
-				NotifyPrimitiveUpdated_Concurrent(Primitive);
-			}
+			DynamicComponentManager.ForEachReferencedComponent(
+				[this](const UPrimitiveComponent* PrimitiveComponent)
+				{
+					NotifyPrimitiveUpdated_Concurrent(PrimitiveComponent);
+				});
 		}
 
 		// Update the cache variables.
diff --git a/Engine/Source/Runtime/Engine/Private/Streaming/TextureInstanceState.cpp b/Engine/Source/Runtime/Engine/Private/Streaming/TextureInstanceState.cpp
--- a/Engine/Source/Runtime/Engine/Private/Streaming/TextureInstanceState.cpp
+++ b/Engine/Source/Runtime/Engine/Private/Streaming/TextureInstanceState.cpp
@@ -670,14 +670,6 @@ void FRenderAssetInstanceState::FlushPendingRemoveComponents(FRemovedRenderAsset
 	PendingRemoveComponents.Reset();
 }
 
-void FRenderAssetInstanceState::GetReferencedComponents(TArray<const UPrimitiveComponent*>& Components) const
-{
-	for (TMap<const UPrimitiveComponent*, int32>::TConstIterator It(ComponentMap); It; ++It)
-	{
-		Components.Add(It.Key());
-	}
-}
-
 void FRenderAssetInstanceState::UpdateBounds(const UPrimitiveComponent* Component)
 {
 	checkf(bIsDynamicInstanceState, TEXT("Bounds shouldn't be updated after creation unless the instances are dynamic"));
diff --git a/Engine/Source/Runtime/Engine/Private/Streaming/TextureInstanceState.h b/Engine/Source/Runtime/Engine/Private/Streaming/TextureInstanceState.h
--- a/Engine/Source/Runtime/Engine/Private/Streaming/TextureInstanceState.h
+++ b/Engine/Source/Runtime/Engine/Private/Streaming/TextureInstanceState.h
@@ -101,7 +101,14 @@ public:
 	// Remove any record associated with pending remove components. Returns an array of render assets with all components removed.
 	void FlushPendingRemoveComponents(FRemovedRenderAssetArray& RemovedRenderAssets);
 
-	void GetReferencedComponents(TArray<const UPrimitiveComponent*>& Components) const;
+	template<typename CallableT>
+	void ForEachReferencedComponent(CallableT Callable) const
+	{
+		for (const TPair<const UPrimitiveComponent*, int32>& Component : ComponentMap)
+		{
+			Callable(Component.Key);
+		}
+	}
 
 	void UpdateBounds(const UPrimitiveComponent* Component);
 	bool UpdateBounds(int32 BoundIndex);
