Create a new FWorldSubsystemCollection that can properly initialize a subsystem that gets added after the world has already executed PostInitialize and OnWorldBegin.

diff --git a/Engine/Source/Runtime/Engine/Classes/Engine/World.h b/Engine/Source/Runtime/Engine/Classes/Engine/World.h
--- a/Engine/Source/Runtime/Engine/Classes/Engine/World.h
+++ b/Engine/Source/Runtime/Engine/Classes/Engine/World.h
@@ -901,6 +901,16 @@ struct FDeferredComponentMoveData
 	int32 FinalOverlapCandidatesIndex = INDEX_NONE;
 };
 
+class FWorldSubsystemCollection : public FObjectSubsystemCollection<UWorldSubsystem>
+{
+public:
+	FWorldSubsystemCollection();
+	virtual ~FWorldSubsystemCollection();
+private:
+	virtual USubsystem* AddAndInitializeSubsystem(UClass* SubsystemClass) override;
+
+};
+
 
 /** 
  * The World is the top level object representing a map or a sandbox in which Actors and Components will exist and be rendered.  
@@ -4348,7 +4358,7 @@ public:
 	FWorldPSCPool PSCPool;
 
 	//PSC Pooling END
-	FObjectSubsystemCollection<UWorldSubsystem> SubsystemCollection;
+	FWorldSubsystemCollection SubsystemCollection;
 };
 
 /** Global UWorld pointer. Use of this pointer should be avoided whenever possible. */
diff --git a/Engine/Source/Runtime/Engine/Private/Subsystems/WorldSubsystem.cpp b/Engine/Source/Runtime/Engine/Private/Subsystems/WorldSubsystem.cpp
--- a/Engine/Source/Runtime/Engine/Private/Subsystems/WorldSubsystem.cpp
+++ b/Engine/Source/Runtime/Engine/Private/Subsystems/WorldSubsystem.cpp
@@ -44,21 +44,6 @@ bool UWorldSubsystem::DoesSupportWorldType(const EWorldType::Type WorldType) con
 void UWorldSubsystem::Initialize(FSubsystemCollectionBase& Collection)
 {
 	Super::Initialize(Collection);
-	//run relevant functions we may have missed if we are initialized after the world has set up
-	if (UWorld* World = GetWorld())
-	{
-		if (!bHasCalledPostInitialize && World->bIsWorldInitialized)
-		{
-			PostInitialize();
-			EnsureHasCalledPostInitialize();
-		}
-
-		if (!bHasCalledBeginPlay && World->HasBegunPlay())
-		{
-			OnWorldBeginPlay(*World);
-			EnsureHasCalledBeginPlay();
-		}
-	}
 }
 
 void UWorldSubsystem::PostInitialize()
diff --git a/Engine/Source/Runtime/Engine/Private/World.cpp b/Engine/Source/Runtime/Engine/Private/World.cpp
--- a/Engine/Source/Runtime/Engine/Private/World.cpp
+++ b/Engine/Source/Runtime/Engine/Private/World.cpp
@@ -729,6 +729,35 @@ void FWorldPartitionEvents::BroadcastWorldPartitionUninitialized(UWorld* InWorld
 	InWorld->BroadcastWorldPartitionUninitialized(InWorldPartition);
 }
 
+FWorldSubsystemCollection::FWorldSubsystemCollection() = default;
+FWorldSubsystemCollection::~FWorldSubsystemCollection() = default;
+
+USubsystem* FWorldSubsystemCollection::AddAndInitializeSubsystem(UClass* SubsystemClass)
+{
+	USubsystem* Subsystem = FObjectSubsystemCollection<UWorldSubsystem>::AddAndInitializeSubsystem(SubsystemClass);
+	if (!Subsystem)
+	{
+		return nullptr;
+	}
+
+	UWorldSubsystem* WorldSubsystem = CastChecked<UWorldSubsystem>(Subsystem);
+	UWorld* World = CastChecked<UWorld>(GetOuter());
+
+	//run relevant functions we may have missed if we are initialized after the world has set up
+	if (!WorldSubsystem->HasCalledPostInitialize() && World->bIsWorldInitialized)
+	{
+		WorldSubsystem->PostInitialize();
+		WorldSubsystem->EnsureHasCalledPostInitialize();
+	}
+
+	if (!WorldSubsystem->HasCalledBeginPlay() && World->HasBegunPlay())
+	{
+		WorldSubsystem->OnWorldBeginPlay(*World);
+		WorldSubsystem->EnsureHasCalledBeginPlay();
+	}
+	return Subsystem;
+}
+
 
 FAudioDeviceWorldDelegates::FOnWorldRegisteredToAudioDevice FAudioDeviceWorldDelegates::OnWorldRegisteredToAudioDevice;
 
diff --git a/Engine/Source/Runtime/Engine/Public/Subsystems/SubsystemCollection.h b/Engine/Source/Runtime/Engine/Public/Subsystems/SubsystemCollection.h
--- a/Engine/Source/Runtime/Engine/Public/Subsystems/SubsystemCollection.h
+++ b/Engine/Source/Runtime/Engine/Public/Subsystems/SubsystemCollection.h
@@ -92,11 +92,16 @@ protected:
 	/** Remove all subsystems in this list of packages */
 	ENGINE_API void RemoveSubsystemsInPackages(TConstArrayView<UPackage*> Packages);
 
-private:
-	ENGINE_API USubsystem* AddAndInitializeSubsystem(UClass* SubsystemClass);
+	/** Get the outer object associated with this collection */
+	UObject* GetOuter() const { return Outer; }
+
+	/** Adds a subsystem to the collection and initializes it */
+	ENGINE_API virtual USubsystem* AddAndInitializeSubsystem(UClass* SubsystemClass);
 
+	/** Removes a subsystem to the collection and deinitializes it */
 	ENGINE_API void RemoveAndDeinitializeSubsystem(USubsystem* Subsystem);
 
+private:
 	TMap<TObjectPtr<UClass>, TObjectPtr<USubsystem>> SubsystemMap;
 
 	mutable TMap<UClass*, TUniquePtr<FSubsystemArray>> SubsystemArrayMap;
diff --git a/Engine/Source/Runtime/Engine/Public/Subsystems/WorldSubsystem.h b/Engine/Source/Runtime/Engine/Public/Subsystems/WorldSubsystem.h
--- a/Engine/Source/Runtime/Engine/Public/Subsystems/WorldSubsystem.h
+++ b/Engine/Source/Runtime/Engine/Public/Subsystems/WorldSubsystem.h
@@ -50,6 +50,12 @@ public:
 	/** Called after world components (e.g. line batcher and all level components) have been updated */
 	virtual void OnWorldComponentsUpdated(UWorld& World) {}
 
+	/** Returns true once the subsystem has executed PostInitialize **/
+	bool HasCalledPostInitialize() const { return bHasCalledPostInitialize; }
+
+	/** Returns true once the subsystem has executed OnWorldBeginPlay **/
+	bool HasCalledBeginPlay() const { return bHasCalledBeginPlay; }
+
 	ENGINE_API void EnsureHasCalledPostInitialize() const;
 
 	ENGINE_API void EnsureHasCalledBeginPlay() const;
