Allow setting the load filter from code on UDataLayerAsset

diff --git a/Engine/Source/Runtime/Engine/Private/WorldPartition/DataLayer/DataLayerAsset.cpp b/Engine/Source/Runtime/Engine/Private/WorldPartition/DataLayer/DataLayerAsset.cpp
--- a/Engine/Source/Runtime/Engine/Private/WorldPartition/DataLayer/DataLayerAsset.cpp
+++ b/Engine/Source/Runtime/Engine/Private/WorldPartition/DataLayer/DataLayerAsset.cpp
@@ -94,5 +94,13 @@ void UDataLayerAsset::SetType(EDataLayerType InDataLayerType)
 		DataLayerType = InDataLayerType;
 	}
 }
+void UDataLayerAsset::SetLoadFilter(EDataLayerLoadFilter InLoadFilter)
+{
+	check(DataLayerType == EDataLayerType::Runtime || !IsPrivate());
+	if (CanEditDataLayerType())
+	{
+		LoadFilter = InLoadFilter;
+	}
+}
 
 #endif
diff --git a/Engine/Source/Runtime/Engine/Public/WorldPartition/DataLayer/DataLayerAsset.h b/Engine/Source/Runtime/Engine/Public/WorldPartition/DataLayer/DataLayerAsset.h
--- a/Engine/Source/Runtime/Engine/Public/WorldPartition/DataLayer/DataLayerAsset.h
+++ b/Engine/Source/Runtime/Engine/Public/WorldPartition/DataLayer/DataLayerAsset.h
@@ -41,6 +41,7 @@ class UDataLayerAsset : public UDataAsset
 public:
 	virtual bool CanEditDataLayerType() const;
 	ENGINE_API void SetType(EDataLayerType Type);
+	ENGINE_API void SetLoadFilter(EDataLayerLoadFilter LoadFilter);
 	void SetDebugColor(FColor InDebugColor) { DebugColor = InDebugColor; }
 	ENGINE_API bool CanBeReferencedByActor(AActor* InActor) const;
 	static ENGINE_API bool CanBeReferencedByActor(const TSoftObjectPtr<UDataLayerAsset>& InDataLayerAsset, AActor* InActor);
