Fixed loading the same partitioned map from 2 PIE clients running in the same process.

diff --git a/Engine/Source/Runtime/Engine/Private/World.cpp b/Engine/Source/Runtime/Engine/Private/World.cpp
--- a/Engine/Source/Runtime/Engine/Private/World.cpp
+++ b/Engine/Source/Runtime/Engine/Private/World.cpp
@@ -4300,6 +4300,13 @@ void UWorld::RenameToPIEWorld(int32 PIEInstanceID)
 	WorldPackage->Rename(*PIEPackageName, nullptr);
 	FSoftObjectPath::AddPIEPackageName(FName(*PIEPackageName));
 
+	for (UPackage* ExternalPackage : WorldPackage->GetExternalPackages())
+	{
+		const FString ExternalPackageName = *UWorld::ConvertToPIEPackageName(ExternalPackage->GetName(), PIEInstanceID);
+		ExternalPackage->Rename(*ExternalPackageName, nullptr);
+		FSoftObjectPath::AddPIEPackageName(FName(*ExternalPackageName));
+	}
+
 	StreamingLevelsPrefix = UWorld::BuildPIEPackagePrefix(PIEInstanceID);
 	
 	if (WorldComposition)
