Fix several global vars in header files that got copied into every TU.

diff --git a/Engine/Source/Runtime/Engine/Classes/Curves/Spline.h b/Engine/Source/Runtime/Engine/Classes/Curves/Spline.h
--- a/Engine/Source/Runtime/Engine/Classes/Curves/Spline.h
+++ b/Engine/Source/Runtime/Engine/Classes/Curves/Spline.h
@@ -163,9 +163,9 @@ struct FSpline
 
 private:
 	
-	static inline const FInterpCurveVector PositionCurve;
-	static inline const FInterpCurveQuat RotationCurve;
-	static inline const FInterpCurveVector ScaleCurve;
+	static ENGINE_API const FInterpCurveVector PositionCurve;
+	static ENGINE_API const FInterpCurveQuat RotationCurve;
+	static ENGINE_API const FInterpCurveVector ScaleCurve;
 	
 	// Used for upgrade logic in spline component.
 	// Not ideal, but allows us to automatically populate the proxy
diff --git a/Engine/Source/Runtime/Engine/Engine.Build.cs b/Engine/Source/Runtime/Engine/Engine.Build.cs
--- a/Engine/Source/Runtime/Engine/Engine.Build.cs
+++ b/Engine/Source/Runtime/Engine/Engine.Build.cs
@@ -173,6 +173,7 @@ public class Engine : ModuleRules
 				"XmlParser",
 				"IoStoreOnDemand",
 				"HTTP",
+				"GeometryCore",
 			}
 		);
 
diff --git a/Engine/Source/Runtime/Engine/Private/Curves/Spline.cpp b/Engine/Source/Runtime/Engine/Private/Curves/Spline.cpp
--- a/Engine/Source/Runtime/Engine/Private/Curves/Spline.cpp
+++ b/Engine/Source/Runtime/Engine/Private/Curves/Spline.cpp
@@ -16,6 +16,10 @@
 
 using UE::Geometry::FInterval1f;
 
+const FInterpCurveVector FSpline::PositionCurve;
+const FInterpCurveQuat FSpline::RotationCurve;
+const FInterpCurveVector FSpline::ScaleCurve;
+
 namespace UE::Spline
 {
 
@@ -3170,4 +3174,4 @@ void FLegacySpline::UpdateSpline(const FSpline::FUpdateSplineParams& InParams)
 		AccumulatedLength += GetSegmentLength(SegmentIndex, 1.0f, InParams.Scale3D);
 	}
 	ReparamTable.Points.Emplace(AccumulatedLength, static_cast<float>(NumSegments), 0.0f, 0.0f, CIM_Linear);
-}
\ No newline at end of file
+}
diff --git a/Engine/Source/Runtime/GeometryCore/Private/Splines/SplineTypeId.cpp b/Engine/Source/Runtime/GeometryCore/Private/Splines/SplineTypeId.cpp
new file mode 100644
--- /dev/null
+++ b/Engine/Source/Runtime/GeometryCore/Private/Splines/SplineTypeId.cpp
@@ -0,0 +1,9 @@
+// Copyright Epic Games, Inc. All Rights Reserved.
+
+#include "Splines/SplineTypeId.h"
+
+namespace UE::Geometry::Spline
+{
+DEFINE_LOG_CATEGORY(LogSpline);
+}
+
diff --git a/Engine/Source/Runtime/GeometryCore/Public/Splines/SplineTypeId.h b/Engine/Source/Runtime/GeometryCore/Public/Splines/SplineTypeId.h
--- a/Engine/Source/Runtime/GeometryCore/Public/Splines/SplineTypeId.h
+++ b/Engine/Source/Runtime/GeometryCore/Public/Splines/SplineTypeId.h
@@ -11,7 +11,8 @@ namespace Geometry
 namespace Spline
 {
     
-DEFINE_LOG_CATEGORY_STATIC(LogSpline, Log, All);
+GEOMETRYCORE_API DECLARE_LOG_CATEGORY_EXTERN(LogSpline, Log, All);
+
 /** Helper for static_assert that depends on template parameter */
 template<typename T>
 struct TDependentFalse : std::false_type {};
@@ -138,4 +139,4 @@ public:
 
 } // namespace Spline
 } // namespace Geometry
-} // namespace UE
\ No newline at end of file
+} // namespace UE
