[FastGeo] Fixed deadlock in FFastGeoComponentCluster::ForceUpdateVisibility. This is an edge case where SendAllEndOfFrameUpdates processes marked dirty and triggers recreation of render proxies combined with the primitive referencing a material that is missing a usage flag. UMaterial::SetMaterialUsage relies on FApp::IsGame() to detect whether we are running the game (pie, -game). During SendAllEndOfFrameUpdates, this function returns false and triggers the Editor code path (which is not intended) and ends up forcing a FlushRenderingCommands which forces any pending GameThread task (ForceUpdateVisibility) to run. To avoid this, combine FApp::IsGame with GEditor->IsPlayingSessionInEditor() to detect if running game.

diff --git a/Engine/Source/Runtime/Engine/Private/Materials/Material.cpp b/Engine/Source/Runtime/Engine/Private/Materials/Material.cpp
--- a/Engine/Source/Runtime/Engine/Private/Materials/Material.cpp
+++ b/Engine/Source/Runtime/Engine/Private/Materials/Material.cpp
@@ -2044,7 +2044,14 @@ bool UMaterial::SetMaterialUsage(bool &bNeedsRecompile, EMaterialUsage Usage, UM
 	{
 		// For materials which do not have their bUsedWith____ correctly set the DefaultMaterial<type> should be used in game
 		// Leaving this GIsEditor ensures that in game on PC will not look different than on the Consoles as we will not be compiling shaders on the fly
-		if( GIsEditor && !FApp::IsGame() && bAutomaticallySetUsageInEditor )
+		// 
+		// Using FApp::IsGame is not enough, during UWorld::SendAllEndOfFrameUpdates() it returns false even it PIE is running.
+		// Ideally, FApp::IsGame should be revisited, but to avoid causing regressions, we also check IsPlayingSessionInEditor.
+		bool bIsGame = FApp::IsGame();
+#if WITH_EDITOR
+		bIsGame = bIsGame || (GEditor && GEditor->IsPlayingSessionInEditor());
+#endif
+		if( GIsEditor && !bIsGame && bAutomaticallySetUsageInEditor )
 		{
 			check(IsInGameThread());
 			//Do not warn the user during automation testing
