Fix for actors not applying the full Data Layer set when applying the current editor context. When moving an actor into a data layer which itself exists in an External data layer (in a game features plugin) from either a game, or different GFP, the actor must first be moved in the EDL, for other data layers to be subsequently valid, or they will be skipped. Also, avoid RemoveActorFromAllDataLayers() here as it skips EDLs and this is not clear.

#tests Moved actors between EDL / DL in editor, and dropped actors into active editor contexts with EDL / DL preselected.
#rnx

diff --git a/Engine/Source/Editor/DataLayerEditor/Private/DataLayer/DataLayerEditorSubsystem.cpp b/Engine/Source/Editor/DataLayerEditor/Private/DataLayer/DataLayerEditorSubsystem.cpp
--- a/Engine/Source/Editor/DataLayerEditor/Private/DataLayer/DataLayerEditorSubsystem.cpp
+++ b/Engine/Source/Editor/DataLayerEditor/Private/DataLayer/DataLayerEditorSubsystem.cpp
@@ -2151,8 +2151,19 @@ bool UDataLayerEditorSubsystem::MoveActorToDataLayers(AActor* InActor, const TAr
 			}
 		}
 
-		// Move to new Data Layers (except for the External Data Layer)
-		RemoveActorFromAllDataLayers(InActor);
+		// Move to new Data Layers:
+		// First, remove all non-EDL layers
+		const bool bIncludeExternalDataLayerAsset = false;
+		for (const UDataLayerAsset* DataLayerAsset : InActor->GetDataLayerAssets(bIncludeExternalDataLayerAsset))
+		{
+			FAssignActorDataLayer::RemoveDataLayerAsset(InActor, DataLayerAsset);
+		}
+		// Second, move to new External Data Layer if any or different. Subsequent valid data layers should be in that EDL, and the EDL must be added first for package paths to allow adding the DLs.
+		if (OldExternalDataLayerAsset != NewExternalDataLayerAsset)
+		{
+			verify(FExternalDataLayerHelper::MoveActorsToExternalDataLayer({ InActor }, Params));
+		}
+		// Finally, apply the rest of the layers
 		for (UDataLayerInstance* DataLayerInstance : InDataLayerInstances)
 		{
 			if (!DataLayerInstance->IsA<UExternalDataLayerInstance>())
@@ -2160,11 +2171,6 @@ bool UDataLayerEditorSubsystem::MoveActorToDataLayers(AActor* InActor, const TAr
 				AddActorToDataLayer(InActor, DataLayerInstance);
 			}
 		}
-		// Move to new External Data Layer
-		if (OldExternalDataLayerAsset != NewExternalDataLayerAsset)
-		{
-			verify(FExternalDataLayerHelper::MoveActorsToExternalDataLayer({ InActor }, Params));
-		}
 	}
 	return true;
 }
diff --git a/Engine/Source/Runtime/Engine/Classes/GameFramework/Actor.h b/Engine/Source/Runtime/Engine/Classes/GameFramework/Actor.h
--- a/Engine/Source/Runtime/Engine/Classes/GameFramework/Actor.h
+++ b/Engine/Source/Runtime/Engine/Classes/GameFramework/Actor.h
@@ -4824,6 +4824,7 @@ private:
 
 	friend class UEngine;
 	friend class FContentBundleEditor;
+	friend class UDataLayerEditorSubsystem;
 	friend class UDataLayerInstanceWithAsset;
 	friend class UDataLayerInstancePrivate;
 	friend class UExternalDataLayerInstance;
diff --git a/Engine/Source/Runtime/Engine/Private/ActorPartition/ActorPartitionSubsystem.cpp b/Engine/Source/Runtime/Engine/Private/ActorPartition/ActorPartitionSubsystem.cpp
--- a/Engine/Source/Runtime/Engine/Private/ActorPartition/ActorPartitionSubsystem.cpp
+++ b/Engine/Source/Runtime/Engine/Private/ActorPartition/ActorPartitionSubsystem.cpp
@@ -591,12 +591,19 @@ bool UActorPartitionSubsystem::MoveActorToDataLayers(APartitionActor* InActor, c
 	InActor->Rename(*NewActorName, nullptr, REN_NonTransactional | REN_DontCreateRedirectors | REN_DoNotDirty);
 	APartitionActor::SetLabelForActor(InActor, ActorPartitionId, InActor->GetGridSize(), CellCoord.X, CellCoord.Y, CellCoord.Z);
 
-	// Move to new Data Layers (except for the External Data Layer)
+	// Move to new Data Layers:
+	// First, remove all non-EDL layers
 	const bool bIncludeExternalDataLayerAsset = false;
 	for (const UDataLayerAsset* DataLayerAsset : InActor->GetDataLayerAssets(bIncludeExternalDataLayerAsset))
 	{
 		FAssignActorDataLayer::RemoveDataLayerAsset(InActor, DataLayerAsset);
 	}
+	// Second, move to new External Data Layer if any or different. Subsequent valid data layers should be in that EDL, and the EDL must be added first for package paths to allow adding the DLs.
+	if (OldExternalDataLayerAsset != NewExternalDataLayerAsset)
+	{
+		verify(FExternalDataLayerHelper::MoveActorsToExternalDataLayer({ InActor }, Params));
+	}
+	// Finally, apply the rest of the layers
 	for (const UDataLayerInstance* DataLayerInstance : InDataLayerInstances)
 	{
 		if (const UDataLayerAsset* DataLayerAsset = !DataLayerInstance->IsA<UExternalDataLayerInstance>() ? DataLayerInstance->GetAsset() : nullptr)
@@ -604,11 +611,7 @@ bool UActorPartitionSubsystem::MoveActorToDataLayers(APartitionActor* InActor, c
 			FAssignActorDataLayer::AddDataLayerAsset(InActor, DataLayerAsset);
 		}
 	}
-	// Move to new External Data Layer
-	if (OldExternalDataLayerAsset != NewExternalDataLayerAsset)
-	{
-		verify(FExternalDataLayerHelper::MoveActorsToExternalDataLayer({ InActor }, Params));
-	}
+
 	return true;
 }
 
