PR #13459: Fix: typo in comment: "no is class" ? "no class is".

#rnx

diff --git a/Engine/Source/Runtime/Engine/Private/GameplayStatics.cpp b/Engine/Source/Runtime/Engine/Private/GameplayStatics.cpp
--- a/Engine/Source/Runtime/Engine/Private/GameplayStatics.cpp
+++ b/Engine/Source/Runtime/Engine/Private/GameplayStatics.cpp
@@ -1078,7 +1078,7 @@ AActor* UGameplayStatics::GetActorOfClass(const UObject* WorldContextObject, TSu
 {
 	QUICK_SCOPE_CYCLE_COUNTER(UGameplayStatics_GetActorOfClass);
 
-	// We do nothing if no is class provided
+	// We do nothing if no class is provided
 	if (ActorClass)
 	{
 		if (UWorld* World = GEngine->GetWorldFromContextObject(WorldContextObject, EGetWorldErrorMode::LogAndReturnNull))
@@ -1099,7 +1099,7 @@ void UGameplayStatics::GetAllActorsOfClass(const UObject* WorldContextObject, TS
 	QUICK_SCOPE_CYCLE_COUNTER(UGameplayStatics_GetAllActorsOfClass);
 	OutActors.Reset();
 
-	// We do nothing if no is class provided, rather than giving ALL actors!
+	// We do nothing if no class is provided, rather than giving ALL actors!
 	if (ActorClass)
 	{
 		if (UWorld* World = GEngine->GetWorldFromContextObject(WorldContextObject, EGetWorldErrorMode::LogAndReturnNull))
@@ -1169,7 +1169,7 @@ void UGameplayStatics::GetAllActorsOfClassWithTag(const UObject* WorldContextObj
 
 	UWorld* World = GEngine->GetWorldFromContextObject(WorldContextObject, EGetWorldErrorMode::LogAndReturnNull);
 
-	// We do nothing if no is class provided, rather than giving ALL actors!
+	// We do nothing if no class is provided, rather than giving ALL actors!
 	if (ActorClass && World)
 	{
 		for (TActorIterator<AActor> It(World, ActorClass); It; ++It)
