Fixing linker issue hidden by unity files.
Changing ForEachReferencedComponent methods to take rvalues.

diff --git a/Engine/Source/Runtime/Engine/Private/Streaming/DynamicTextureInstanceManager.cpp b/Engine/Source/Runtime/Engine/Private/Streaming/DynamicTextureInstanceManager.cpp
--- a/Engine/Source/Runtime/Engine/Private/Streaming/DynamicTextureInstanceManager.cpp
+++ b/Engine/Source/Runtime/Engine/Private/Streaming/DynamicTextureInstanceManager.cpp
@@ -34,24 +34,6 @@ void FDynamicRenderAssetInstanceManager::FTasks::SyncRefreshFullTask()
 	RefreshFullTask->TrySync();
 }
 
-template <typename TTasks>
-void FRenderAssetDynamicInstanceStateTaskSync<TTasks>::Sync()
-{
-	check(IsInGameThread());
-
-	Super::Tasks.SyncResults();
-	FRemovedRenderAssetArray RemovedRenderAssets;
-	Super::State->FlushPendingRemoveComponents(RemovedRenderAssets);
-	OnSyncDoneDelegate.Execute(RemovedRenderAssets);
-}
-
-template <typename TTasks>
-FRenderAssetInstanceState* FRenderAssetDynamicInstanceStateTaskSync<TTasks>::SyncAndGetState()
-{
-	Sync();
-	return Super::State.GetReference();
-}
-
 FDynamicRenderAssetInstanceManager::FDynamicRenderAssetInstanceManager(FOnSyncDoneDelegate&& InOnSyncDoneDelegate)
 	: StateSync(FRenderAssetDynamicInstanceStateTaskSync<FTasks>::FOnSyncDone::CreateLambda(MoveTemp(InOnSyncDoneDelegate)))
 	, DirtyIndex(0)
diff --git a/Engine/Source/Runtime/Engine/Private/Streaming/DynamicTextureInstanceManager.h b/Engine/Source/Runtime/Engine/Private/Streaming/DynamicTextureInstanceManager.h
--- a/Engine/Source/Runtime/Engine/Private/Streaming/DynamicTextureInstanceManager.h
+++ b/Engine/Source/Runtime/Engine/Private/Streaming/DynamicTextureInstanceManager.h
@@ -30,9 +30,9 @@ public:
 
 	// Iterate all (non removed) components refered by the manager. Debug only.
 	template<typename CallableT>
-	void ForEachReferencedComponent(CallableT Callable)
+	void ForEachReferencedComponent(CallableT&& Callable)
 	{
-		StateSync.SyncAndGetState()->ForEachReferencedComponent(Callable);
+		StateSync.SyncAndGetState()->ForEachReferencedComponent(Forward<CallableT>(Callable));
 	}
 
 	/** Remove all pending components that are marked for delete. This prevents searching in the pending list for each entry. */
diff --git a/Engine/Source/Runtime/Engine/Private/Streaming/StaticTextureInstanceManager.h b/Engine/Source/Runtime/Engine/Private/Streaming/StaticTextureInstanceManager.h
--- a/Engine/Source/Runtime/Engine/Private/Streaming/StaticTextureInstanceManager.h
+++ b/Engine/Source/Runtime/Engine/Private/Streaming/StaticTextureInstanceManager.h
@@ -36,9 +36,9 @@ public:
 
 	/** Invoke a callable for each referenced component. */
 	template<typename CallableT>
-	void ForEachReferencedComponent(CallableT Callable)
+	void ForEachReferencedComponent(CallableT&& Callable)
 	{
-		StateSync.GetState()->ForEachReferencedComponent(Callable);
+		StateSync.GetState()->ForEachReferencedComponent(Forward<CallableT>(Callable));
 	}
 
 	/** Return whether this component is be managed by this manager. */
diff --git a/Engine/Source/Runtime/Engine/Private/Streaming/TextureInstanceState.h b/Engine/Source/Runtime/Engine/Private/Streaming/TextureInstanceState.h
--- a/Engine/Source/Runtime/Engine/Private/Streaming/TextureInstanceState.h
+++ b/Engine/Source/Runtime/Engine/Private/Streaming/TextureInstanceState.h
@@ -102,7 +102,7 @@ public:
 	void FlushPendingRemoveComponents(FRemovedRenderAssetArray& RemovedRenderAssets);
 
 	template<typename CallableT>
-	void ForEachReferencedComponent(CallableT Callable) const
+	void ForEachReferencedComponent(CallableT&& Callable) const
 	{
 		for (const TPair<const UPrimitiveComponent*, int32>& Component : ComponentMap)
 		{
@@ -234,9 +234,21 @@ public:
 		, OnSyncDoneDelegate(InOnSyncDoneDelegate)
 	{}
 
-	void Sync();
+	void Sync()
+	{
+		check(IsInGameThread());
+
+		Super::Tasks.SyncResults();
+		FRemovedRenderAssetArray RemovedRenderAssets;
+		Super::State->FlushPendingRemoveComponents(RemovedRenderAssets);
+		OnSyncDoneDelegate.Execute(RemovedRenderAssets);
+	}
 
-	FRenderAssetInstanceState* SyncAndGetState();
+	FRenderAssetInstanceState* SyncAndGetState()
+	{
+		Sync();
+		return Super::State.GetReference();
+	}
 
 private:
 	FOnSyncDone OnSyncDoneDelegate;
