Removing more implicit RHICommandList usages by passing command lists around more. In some cases, we need to get the immediate command list directly.

diff --git a/Engine/Plugins/Runtime/XR/XRScribe/Source/XRScribe/Private/XRScribeEmulationLayer.cpp b/Engine/Plugins/Runtime/XR/XRScribe/Source/XRScribe/Private/XRScribeEmulationLayer.cpp
--- a/Engine/Plugins/Runtime/XR/XRScribe/Source/XRScribe/Private/XRScribeEmulationLayer.cpp
+++ b/Engine/Plugins/Runtime/XR/XRScribe/Source/XRScribe/Private/XRScribeEmulationLayer.cpp
@@ -1449,6 +1449,8 @@ XrResult FOpenXREmulationLayer::XrLayerCreateSwapchain(XrSession session, const
 		return XR_ERROR_RUNTIME_FAILURE;
 	}
 
+	FRHICommandListImmediate& RHICmdList = FRHICommandListImmediate::Get();
+
 	TUniquePtr<FOpenXREmulatedSwapchain> NewSwapchain = MakeUnique<FOpenXREmulatedSwapchain>();
 
 	static uint32 SwapchainCounter = 0;
@@ -1468,7 +1470,7 @@ XrResult FOpenXREmulationLayer::XrLayerCreateSwapchain(XrSession session, const
 	for (uint32 SwapchainImageIndex = 0; SwapchainImageIndex < kEmulatedSwapchainLength; SwapchainImageIndex++)
 	{
 		// TODO: handle alloc failure?!
-		FTextureRHIRef NewImage = RHICreateTexture(SwapchainImageDesc);
+		FTextureRHIRef NewImage = RHICmdList.CreateTexture(SwapchainImageDesc);
 		NewSwapchain->TextureList[SwapchainImageIndex] = NewImage;
 	}
 
diff --git a/Engine/Plugins/Runtime/nDisplay/Source/DisplayCluster/Private/Render/GUILayer/DisplayClusterGuiLayerController.cpp b/Engine/Plugins/Runtime/nDisplay/Source/DisplayCluster/Private/Render/GUILayer/DisplayClusterGuiLayerController.cpp
--- a/Engine/Plugins/Runtime/nDisplay/Source/DisplayCluster/Private/Render/GUILayer/DisplayClusterGuiLayerController.cpp
+++ b/Engine/Plugins/Runtime/nDisplay/Source/DisplayCluster/Private/Render/GUILayer/DisplayClusterGuiLayerController.cpp
@@ -246,7 +246,7 @@ FTextureRHIRef FDisplayClusterGuiLayerController::CreateTextureFrom_RenderThread
 		.SetInitialState(ERHIAccess::SRVMask);
 
 	// Create texture
-	return RHICreateTexture(Desc);
+	return RHICmdList.CreateTexture(Desc);
 }
 
 FRDGTextureRef FDisplayClusterGuiLayerController::CreateTextureFrom_RenderThread(
diff --git a/Engine/Plugins/Runtime/nDisplay/Source/DisplayClusterMedia/Private/Synchronization/LatencyQueue/DisplayClusterFrameQueueItem.cpp b/Engine/Plugins/Runtime/nDisplay/Source/DisplayClusterMedia/Private/Synchronization/LatencyQueue/DisplayClusterFrameQueueItem.cpp
--- a/Engine/Plugins/Runtime/nDisplay/Source/DisplayClusterMedia/Private/Synchronization/LatencyQueue/DisplayClusterFrameQueueItem.cpp
+++ b/Engine/Plugins/Runtime/nDisplay/Source/DisplayClusterMedia/Private/Synchronization/LatencyQueue/DisplayClusterFrameQueueItem.cpp
@@ -22,8 +22,8 @@ FDisplayClusterFrameQueueItem::FDisplayClusterFrameQueueItem(const FDisplayClust
 		FDisplayClusterFrameQueueItemView& NewItem = this->Views.Emplace(View.Key, View.Value);
 		if (View.Value.Texture.IsValid())
 		{
-			NewItem.Texture = CreateTexture(View.Value.Texture);
 			FRHICommandListImmediate& RHICmdList = GetImmediateCommandList_ForRenderCommand();
+			NewItem.Texture = CreateTexture(RHICmdList, View.Value.Texture);
 			RHICmdList.CopyTexture(View.Value.Texture.GetReference(), NewItem.Texture, FRHICopyTextureInfo());
 		}
 	}
@@ -45,7 +45,7 @@ void FDisplayClusterFrameQueueItem::SaveView(FRHICommandListImmediate& RHICmdLis
 		Texture->GetFormat() != View->Texture.GetReference()->GetFormat() ||
 		Texture->GetDesc().Extent != View->Texture.GetReference()->GetDesc().Extent)
 	{
-		View->Texture = CreateTexture(Texture);
+		View->Texture = CreateTexture(RHICmdList, Texture);
 		check(View->Texture.IsValid());
 	}
 
@@ -125,7 +125,7 @@ void FDisplayClusterFrameQueueItem::LoadData(FRHICommandListImmediate& RHICmdLis
 	}
 }
 
-FTextureRHIRef FDisplayClusterFrameQueueItem::CreateTexture(FRHITexture* ReferenceTexture)
+FTextureRHIRef FDisplayClusterFrameQueueItem::CreateTexture(FRHICommandListBase& RHICmdList, FRHITexture* ReferenceTexture)
 {
 	if (ReferenceTexture)
 	{
@@ -149,7 +149,7 @@ FTextureRHIRef FDisplayClusterFrameQueueItem::CreateTexture(FRHITexture* Referen
 		Desc.SetFlags(Flags);
 
 		// Create texture
-		return RHICreateTexture(Desc);
+		return RHICmdList.CreateTexture(Desc);
 	}
 
 	return nullptr;
diff --git a/Engine/Plugins/Runtime/nDisplay/Source/DisplayClusterMedia/Private/Synchronization/LatencyQueue/DisplayClusterFrameQueueItem.h b/Engine/Plugins/Runtime/nDisplay/Source/DisplayClusterMedia/Private/Synchronization/LatencyQueue/DisplayClusterFrameQueueItem.h
--- a/Engine/Plugins/Runtime/nDisplay/Source/DisplayClusterMedia/Private/Synchronization/LatencyQueue/DisplayClusterFrameQueueItem.h
+++ b/Engine/Plugins/Runtime/nDisplay/Source/DisplayClusterMedia/Private/Synchronization/LatencyQueue/DisplayClusterFrameQueueItem.h
@@ -56,7 +56,7 @@ public:
 
 private:
 	// Auxiliary function that creates a new texture based on the reference (same format, size)
-	FTextureRHIRef CreateTexture(FRHITexture* ReferenceTexture);
+	FTextureRHIRef CreateTexture(FRHICommandListBase& RHICmdList, FRHITexture* ReferenceTexture);
 
 private:
 	// Available views (viewports)
diff --git a/Engine/Source/Runtime/Engine/Private/Streaming/Texture2DArrayStreaming.cpp b/Engine/Source/Runtime/Engine/Private/Streaming/Texture2DArrayStreaming.cpp
--- a/Engine/Source/Runtime/Engine/Private/Streaming/Texture2DArrayStreaming.cpp
+++ b/Engine/Source/Runtime/Engine/Private/Streaming/Texture2DArrayStreaming.cpp
@@ -59,6 +59,8 @@ bool FTexture2DArrayMipAllocator_Reallocate::AllocateMips(const FTextureUpdateCo
 
 bool FTexture2DArrayMipAllocator_Reallocate::UploadMips(const FTextureUpdateContext& Context, const FTextureUpdateSyncOptions& SyncOptions)
 {
+	FRHICommandListImmediate& RHICmdList = FRHICommandListImmediate::Get();
+
 	const uint32 NumSlices = Context.Resource->GetSizeZ();
 
 	// Create new Texture.
@@ -71,24 +73,15 @@ bool FTexture2DArrayMipAllocator_Reallocate::UploadMips(const FTextureUpdateCont
 			.SetFlags(Context.Resource->GetCreationFlags())
 			.SetExtData(Context.Resource->GetExtData());
 
-		IntermediateTextureRHI = RHICreateTexture(Desc);
+		IntermediateTextureRHI = RHICmdList.CreateTexture(Desc);
 	}
 
 	// Copy shared mips.
 	{
-		bool bCopySharedMipsDone = false;
-		ENQUEUE_RENDER_COMMAND(FCopySharedMipsForTexture2DArray)(
-			[&](FRHICommandListImmediate& RHICmdList)
-		{
-			FRHITexture* SrcTexture = Context.Resource->GetTexture2DArrayRHI();
-			FRHITexture* DstTexture = IntermediateTextureRHI.GetReference();
-
-			UE::RHI::CopySharedMips_AssumeSRVMaskState(RHICmdList, SrcTexture, DstTexture);
+		FRHITexture* SrcTexture = Context.Resource->GetTexture2DArrayRHI();
+		FRHITexture* DstTexture = IntermediateTextureRHI.GetReference();
 
-			bCopySharedMipsDone = true;
-		});
-		// Expected to execute immediately since ran on the renderthread.
-		check(bCopySharedMipsDone);
+		UE::RHI::CopySharedMips_AssumeSRVMaskState(RHICmdList, SrcTexture, DstTexture);
 	}
 
 	// Update the streamed in mips if they were not initialized from the bulk data.
@@ -103,9 +96,9 @@ bool FTexture2DArrayMipAllocator_Reallocate::UploadMips(const FTextureUpdateCont
 			for (uint32 SliceIdx = 0; SliceIdx < NumSlices; ++SliceIdx)
 			{
 				uint32 DestStride = 0;
-				void* DestData = RHILockTexture2DArray(IntermediateTextureRHI, SliceIdx, RHIMipIndex, RLM_WriteOnly, DestStride, false);
+				void* DestData = RHICmdList.LockTexture2DArray(IntermediateTextureRHI, SliceIdx, RHIMipIndex, RLM_WriteOnly, DestStride, false);
 				FMemory::Memcpy(DestData, MipData + SliceIdx * SlizeSize, SlizeSize);
-				RHIUnlockTexture2DArray(IntermediateTextureRHI, SliceIdx, RHIMipIndex, false);
+				RHICmdList.UnlockTexture2DArray(IntermediateTextureRHI, SliceIdx, RHIMipIndex, false);
 			}
 
 		}
diff --git a/Engine/Source/Runtime/Engine/Private/Streaming/VolumeTextureStreaming.cpp b/Engine/Source/Runtime/Engine/Private/Streaming/VolumeTextureStreaming.cpp
--- a/Engine/Source/Runtime/Engine/Private/Streaming/VolumeTextureStreaming.cpp
+++ b/Engine/Source/Runtime/Engine/Private/Streaming/VolumeTextureStreaming.cpp
@@ -74,6 +74,8 @@ bool FVolumeTextureMipAllocator_Reallocate::AllocateMips(const FTextureUpdateCon
 
 bool FVolumeTextureMipAllocator_Reallocate::UploadMips(const FTextureUpdateContext& Context, const FTextureUpdateSyncOptions& SyncOptions)
 {
+	FRHICommandListImmediate& RHICmdList = FRHICommandListImmediate::Get();
+
 	EPixelFormat PixelFormat = Context.Resource->GetPixelFormat();
 
 	// Create new Texture.
@@ -93,24 +95,15 @@ bool FVolumeTextureMipAllocator_Reallocate::UploadMips(const FTextureUpdateConte
 
 		Desc.DetermineInititialState();
 
-		IntermediateTextureRHI = RHICreateTexture(Desc);
+		IntermediateTextureRHI = RHICmdList.CreateTexture(Desc);
 	}
 
 	// Copy shared mips.
 	{
-		bool bCopySharedMipsDone = false;
-		ENQUEUE_RENDER_COMMAND(FCopySharedMipsForTexture3D)(
-			[&](FRHICommandListImmediate& RHICmdList)
-		{
-			FRHITexture* SrcTexture = Context.Resource->GetTexture3DRHI();
-			FRHITexture* DstTexture = IntermediateTextureRHI.GetReference();
-
-			UE::RHI::CopySharedMips_AssumeSRVMaskState(RHICmdList, SrcTexture, DstTexture);
+		FRHITexture* SrcTexture = Context.Resource->GetTexture3DRHI();
+		FRHITexture* DstTexture = IntermediateTextureRHI.GetReference();
 
-			bCopySharedMipsDone = true;
-		});
-		// Expected to execute immediately since ran on the renderthread.
-		check(bCopySharedMipsDone);
+		UE::RHI::CopySharedMips_AssumeSRVMaskState(RHICmdList, SrcTexture, DstTexture);
 	}
 
 	// Update the streamed in mips if they were not initialized from the bulk data.
@@ -130,7 +123,7 @@ bool FVolumeTextureMipAllocator_Reallocate::UploadMips(const FTextureUpdateConte
 				FUpdateTextureRegion3D UpdateRegion(0, 0, 0, 0, 0, 0, Mip.SizeX, Mip.SizeY, Mip.SizeZ);
 				const uint32 NumBlockX = (uint32)FMath::DivideAndRoundUp<int32>(Mip.SizeX, BlockSizeX);
 				const uint32 NumBlockY = (uint32)FMath::DivideAndRoundUp<int32>(Mip.SizeY, BlockSizeY);
-				RHIUpdateTexture3D(IntermediateTextureRHI, MipIdx - PendingFirstLODIdx, UpdateRegion, NumBlockX * BlockBytes, NumBlockX * NumBlockY * BlockBytes, MipData);
+				RHICmdList.UpdateTexture3D(IntermediateTextureRHI, MipIdx - PendingFirstLODIdx, UpdateRegion, NumBlockX * BlockBytes, NumBlockX * NumBlockY * BlockBytes, MipData);
 			}
 		}
 	}
diff --git a/Engine/Source/Runtime/GameplayMediaEncoder/Private/GameplayMediaEncoder.cpp b/Engine/Source/Runtime/GameplayMediaEncoder/Private/GameplayMediaEncoder.cpp
--- a/Engine/Source/Runtime/GameplayMediaEncoder/Private/GameplayMediaEncoder.cpp
+++ b/Engine/Source/Runtime/GameplayMediaEncoder/Private/GameplayMediaEncoder.cpp
@@ -677,7 +677,7 @@ PRAGMA_ENABLE_DEPRECATION_WARNINGS
 
 		if (RHIType == ERHIInterfaceType::D3D11)
 		{
-			FTextureRHIRef Texture = RHICreateTexture(Desc);
+			FTextureRHIRef Texture = FRHICommandListImmediate::Get().CreateTexture(Desc);
 
 			PRAGMA_DISABLE_DEPRECATION_WARNINGS
 			InputFrame->SetTexture((ID3D11Texture2D*)Texture->GetNativeResource(), [&, InputFrame](ID3D11Texture2D* NativeTexture) { BackBuffers.Remove(InputFrame); });
@@ -686,7 +686,7 @@ PRAGMA_ENABLE_DEPRECATION_WARNINGS
 		}
 		else if (RHIType == ERHIInterfaceType::D3D12)
 		{
-			FTextureRHIRef Texture = RHICreateTexture(Desc);
+			FTextureRHIRef Texture = FRHICommandListImmediate::Get().CreateTexture(Desc);
 
 			PRAGMA_DISABLE_DEPRECATION_WARNINGS
 			InputFrame->SetTexture((ID3D12Resource*)Texture->GetNativeResource(), [&, InputFrame](ID3D12Resource* NativeTexture) { BackBuffers.Remove(InputFrame); });
diff --git a/Engine/Source/Runtime/SlateRHIRenderer/Private/SlatePostProcessResource.cpp b/Engine/Source/Runtime/SlateRHIRenderer/Private/SlatePostProcessResource.cpp
--- a/Engine/Source/Runtime/SlateRHIRenderer/Private/SlatePostProcessResource.cpp
+++ b/Engine/Source/Runtime/SlateRHIRenderer/Private/SlatePostProcessResource.cpp
@@ -24,18 +24,20 @@ void FSlatePostProcessResource::Update(const FIntPoint& NewSize, EPixelFormat Re
 {
 	if(NewSize.X > RenderTargetSize.X || NewSize.Y > RenderTargetSize.Y || RenderTargetSize == FIntPoint::ZeroValue || RenderTargets.Num() == 0 || PixelFormat != RequestedPixelFormat)
 	{
+		FRHICommandListImmediate& RHICmdList = FRHICommandListImmediate::Get();
+
 		if(!IsInitialized())
 		{
-			InitResource(FRHICommandListImmediate::Get());
+			InitResource(RHICmdList);
 		}
 
 		FIntPoint NewMaxSize(FMath::Max(NewSize.X, RenderTargetSize.X), FMath::Max(NewSize.Y, RenderTargetSize.Y));
-		ResizeTargets(NewMaxSize, RequestedPixelFormat);
+		ResizeTargets(RHICmdList, NewMaxSize, RequestedPixelFormat);
 	}
 	FrameUsed = GFrameCounter;
 }
 
-void FSlatePostProcessResource::ResizeTargets(const FIntPoint& NewSize, EPixelFormat RequestedPixelFormat)
+void FSlatePostProcessResource::ResizeTargets(FRHICommandListBase& RHICmdList, const FIntPoint& NewSize, EPixelFormat RequestedPixelFormat)
 {
 	check(IsInRenderingThread());
 
@@ -54,7 +56,7 @@ void FSlatePostProcessResource::ResizeTargets(const FIntPoint& NewSize, EPixelFo
 				.SetFlags(ETextureCreateFlags::RenderTargetable | ETextureCreateFlags::ShaderResource)
 				.SetInitialState(ERHIAccess::SRVMask);
 
-			FTextureRHIRef RenderTargetTextureRHI = RHICreateTexture(Desc);
+			FTextureRHIRef RenderTargetTextureRHI = RHICmdList.CreateTexture(Desc);
 
 			RenderTargets.Add(RenderTargetTextureRHI);
 		}
diff --git a/Engine/Source/Runtime/SlateRHIRenderer/Private/SlatePostProcessResource.h b/Engine/Source/Runtime/SlateRHIRenderer/Private/SlatePostProcessResource.h
--- a/Engine/Source/Runtime/SlateRHIRenderer/Private/SlatePostProcessResource.h
+++ b/Engine/Source/Runtime/SlateRHIRenderer/Private/SlatePostProcessResource.h
@@ -40,7 +40,7 @@ public:
 
 private:
 	/** Resizes targets to the new size */
-	void ResizeTargets(const FIntPoint& NewSize, EPixelFormat RequestedPixelFormat);
+	void ResizeTargets(FRHICommandListBase& RHICmdList, const FIntPoint& NewSize, EPixelFormat RequestedPixelFormat);
 
 private:
 	TArray<FTextureRHIRef, TInlineAllocator<2>> RenderTargets;
