Fixed missing forward declaration in physics constraint component

diff --git a/Engine/Source/Runtime/Engine/Classes/PhysicsEngine/PhysicsConstraintComponent.h b/Engine/Source/Runtime/Engine/Classes/PhysicsEngine/PhysicsConstraintComponent.h
--- a/Engine/Source/Runtime/Engine/Classes/PhysicsEngine/PhysicsConstraintComponent.h
+++ b/Engine/Source/Runtime/Engine/Classes/PhysicsEngine/PhysicsConstraintComponent.h
@@ -8,10 +8,6 @@
 #include "PhysicsEngine/ConstraintInstance.h"
 #include "PhysicsConstraintComponent.generated.h"
 
-#if UE_WITH_REMOTE_OBJECT_HANDLE
-#include "UObject/UObjectMigrationContext.h"
-#endif
-
 class AActor;
 class UPrimitiveComponent;
 struct FBodyInstance;
@@ -441,7 +437,7 @@ private:
 
 	void MigratePhysicsState(FArchive& Ar);
 	
-	ENGINE_API bool CanSendMigratedPhysicsState(const FUObjectMigrationContext* MigrationContext, EObjectMigrationSendType ObjMigrationSendType) const;
+	ENGINE_API bool CanSendMigratedPhysicsState(const struct FUObjectMigrationContext* MigrationContext, EObjectMigrationSendType ObjMigrationSendType) const;
 
 	void CheckAndUnbindPhysicsStateDelegates();
 #endif
diff --git a/Engine/Source/Runtime/Engine/Private/PhysicsEngine/PhysicsConstraintComponent.cpp b/Engine/Source/Runtime/Engine/Private/PhysicsEngine/PhysicsConstraintComponent.cpp
--- a/Engine/Source/Runtime/Engine/Private/PhysicsEngine/PhysicsConstraintComponent.cpp
+++ b/Engine/Source/Runtime/Engine/Private/PhysicsEngine/PhysicsConstraintComponent.cpp
@@ -16,6 +16,7 @@
 #include "UObject/SoftObjectPath.h"
 
 #if UE_WITH_REMOTE_OBJECT_HANDLE
+#include "UObject/UObjectMigrationContext.h"
 #include "PhysicsInterfaceDeclaresCore.h"
 #include "Physics/Experimental/PhysScene_Chaos.h"
 #include "Chaos/Serialization/SerializedMultiPhysicsState.h"
