FTreeParser: added versions of parsing functions that can take a profile structure to minimize the code needed to parse the same tree type.

diff --git a/Engine/Source/Runtime/Engine/Public/TreeParser.h b/Engine/Source/Runtime/Engine/Public/TreeParser.h
--- a/Engine/Source/Runtime/Engine/Public/TreeParser.h
+++ b/Engine/Source/Runtime/Engine/Public/TreeParser.h
@@ -69,6 +69,23 @@ public:
 		return Parse<TNodeType>(InRootNode, GetNumChildren, GetChild, EnterFunc, ExitFunc);
 	}
 
+	template <typename TProfile, class TEnterFunc, class TExitFunc>
+	static EParseAction ParseRooted(TProfile::TNodeType InRootNode, TProfile& InProfile, TEnterFunc EnterFunc, TExitFunc ExitFunc)
+	{
+		return ParseRooted(
+			InRootNode, 
+			[&InProfile](TProfile::TNodeType InNodeIndex)
+			{
+				return InProfile.GetNumChildren(InNodeIndex);
+			},
+			[&InProfile](TProfile::TNodeType InNodeIndex, TProfile::TNodeType InChildNodeIndex)
+			{
+				return InProfile.GetChild(InNodeIndex, InChildNodeIndex);
+			},
+			EnterFunc,
+			ExitFunc);
+	}
+
 	/**
 	 * Parses a list of nodes and edges from a a Directed Acyclic Graph (DAG).
 	 * This function is designed for a more general graph structure where a node can have multiple parents, which is common 
@@ -84,15 +101,7 @@ public:
 	{
 		TSet<TNodeType> VisitedNodes;
 		return ParseRooted<TNodeType>(
-			InRootNode,
-			[&GetNumChildren](TNodeType Node)
-			{
-				return GetNumChildren(Node);
-			},
-			[&GetChild](TNodeType Node, int32 ChildNodeIndex)
-			{
-				return GetChild(Node, ChildNodeIndex);
-			},
+            InRootNode, GetNumChildren, GetChild,
 			[&EnterFunc, &VisitedNodes](TNodeType Node)
 			{
 				bool bWasAlreadyInSet;
@@ -105,11 +114,24 @@ public:
 				}
 
 				return EParseAction::Skip;				
+			}, ExitFunc);
+	}
+
+	template <typename TProfile, class TEnterFunc, class TExitFunc>
+	static EParseAction ParseDirectedAcyclic(TProfile::TNodeType InRootNode, TProfile& InProfile, TEnterFunc EnterFunc, TExitFunc ExitFunc)
+	{
+		return ParseDirectedAcyclic(
+			InRootNode, 
+			[&InProfile](TProfile::TNodeType InNodeIndex)
+			{
+				return InProfile.GetNumChildren(InNodeIndex);
 			},
-			[&ExitFunc](TNodeType Node)
+			[&InProfile](TProfile::TNodeType InNodeIndex, TProfile::TNodeType InChildNodeIndex)
 			{
-				return ExitFunc(Node);
-			});
+				return InProfile.GetChild(InNodeIndex, InChildNodeIndex);
+			},
+			EnterFunc,
+			ExitFunc);
 	}
 };
 
