UE-318818:

- clean up transactional console commands a bit... fixed the logic that called into ExecuteTransactional, and removed the 'always open' tag on UPlayer::ConsoleCommand

diff --git a/Engine/Source/Runtime/Engine/Private/Player.cpp b/Engine/Source/Runtime/Engine/Private/Player.cpp
--- a/Engine/Source/Runtime/Engine/Private/Player.cpp
+++ b/Engine/Source/Runtime/Engine/Private/Player.cpp
@@ -29,7 +29,6 @@ UPlayer::UPlayer(const FObjectInitializer& ObjectInitializer)
 {
 }
 
-UE_AUTORTFM_ALWAYS_OPEN
 FString UPlayer::ConsoleCommand(const FString& Cmd, bool bWriteToLog)
 {
 	UNetConnection* NetConn = Cast<UNetConnection>(this);
diff --git a/Engine/Source/Runtime/Engine/Private/PlayerController.cpp b/Engine/Source/Runtime/Engine/Private/PlayerController.cpp
--- a/Engine/Source/Runtime/Engine/Private/PlayerController.cpp
+++ b/Engine/Source/Runtime/Engine/Private/PlayerController.cpp
@@ -539,16 +539,13 @@ FString APlayerController::ConsoleCommand(const FString& Cmd, bool bWriteToLog)
 	{
 		FString Result;
 #if UE_WITH_REMOTE_OBJECT_HANDLE
-		if (!AutoRTFM::IsClosed())
+		if (!AutoRTFM::IsTransactional())
 		{
-			UE_AUTORTFM_OPEN
+			static FName TransactionalWorkName(TEXT("APlayerController::ConsoleComand"));
+			UE::RemoteExecutor::ExecuteTransactional(TransactionalWorkName, [this, &Result, Cmd, bWriteToLog]()
 			{
-				static FName TransactionalWorkName(TEXT("APlyerController::ConsoleComand"));
-				UE::RemoteExecutor::ExecuteTransactional(TransactionalWorkName, [this, &Result, Cmd, bWriteToLog]()
-				{
-					Result = Player->ConsoleCommand(Cmd, bWriteToLog); 
-				});
-			};
+				Result = Player->ConsoleCommand(Cmd, bWriteToLog); 
+			});
 		}
 		else
 #endif
