Quick code clean-up based on review suggestions

[RN] minor | Rendering Architecture

diff --git a/Engine/Source/Runtime/Engine/Private/Materials/MaterialExpressions.cpp b/Engine/Source/Runtime/Engine/Private/Materials/MaterialExpressions.cpp
--- a/Engine/Source/Runtime/Engine/Private/Materials/MaterialExpressions.cpp
+++ b/Engine/Source/Runtime/Engine/Private/Materials/MaterialExpressions.cpp
@@ -8334,28 +8334,10 @@ int32 UMaterialExpressionChannelMaskParameter::Compile(class FMaterialCompiler*
 	return Ret;
 }
 
-void UMaterialExpressionChannelMaskParameter::GetCaption(TArray<FString>& OutCaptions) const
-{
-	switch (MaskChannel)
-	{
-	case EChannelMaskParameterColor::Red:
-		OutCaptions.Add(TEXT("Red")); break;
-	case EChannelMaskParameterColor::Green:
-		OutCaptions.Add(TEXT("Green")); break;
-	case EChannelMaskParameterColor::Blue:
-		OutCaptions.Add(TEXT("Blue")); break;
-	default:
-		OutCaptions.Add(TEXT("Alpha")); break;
-	}
-
-	OutCaptions.Add(FString::Printf(TEXT("'%s'"), *ParameterName.ToString()));
-}
-
-
-void UMaterialExpressionChannelMaskParameter::GetExpressionToolTip(TArray<FString>& OutToolTip)
+const TCHAR* GetChannelMaskParameterName(TEnumAsByte<EChannelMaskParameterColor::Type> InMaskChannel)
 {
 	const TCHAR* Output = nullptr;
-	switch (MaskChannel)
+	switch (InMaskChannel)
 	{
 	case EChannelMaskParameterColor::Red:
 		Output = TEXT("Red"); break;
@@ -8366,12 +8348,25 @@ void UMaterialExpressionChannelMaskParameter::GetExpressionToolTip(TArray<FStrin
 	case EChannelMaskParameterColor::Alpha:
 		Output = TEXT("Alpha"); break;
 	default:
-		checkf(false, TEXT("Invalid MaskChannel %d"), MaskChannel); break;
+		checkf(false, TEXT("Invalid MaskChannel %d"), InMaskChannel); break;
 	}
 
+	return Output;
+}
+
+void UMaterialExpressionChannelMaskParameter::GetCaption(TArray<FString>& OutCaptions) const
+{
+	const TCHAR* MaskName = GetChannelMaskParameterName(MaskChannel);
+	OutCaptions.Add(MaskName);
+	OutCaptions.Add(FString::Printf(TEXT("'%s'"), *ParameterName.ToString()));
+}
+
+void UMaterialExpressionChannelMaskParameter::GetExpressionToolTip(TArray<FString>& OutToolTip)
+{
+	const TCHAR* MaskName = GetChannelMaskParameterName(MaskChannel);
 	OutToolTip.Add(TEXT("Outputs the selected channel."));
 	OutToolTip.Add(TEXT("Note: This does NOT create new permutations."));
-	OutToolTip.Add(FString::Printf(TEXT("Current output: %s "), Output));
+	OutToolTip.Add(FString::Printf(TEXT("Current output: %s "), MaskName));
 }
 
 #endif // WITH_EDITOR
