[VisualLogger] Added opt-in functionality to apply LogCategories verbosity filter when recording.
Can be controlled by:
- Console commands: vislog.ActivateVerbosityFilterWhenRecording and vislog.DeactivateVerbosityFilterWhenRecording
- Config Property `bUseVerbosityFilterWhenRecording` in VisualLoggerSettings

#FailedValidations Preflight

diff --git a/Engine/Source/Runtime/Engine/Private/VisualLogger/VisualLogger.cpp b/Engine/Source/Runtime/Engine/Private/VisualLogger/VisualLogger.cpp
--- a/Engine/Source/Runtime/Engine/Private/VisualLogger/VisualLogger.cpp
+++ b/Engine/Source/Runtime/Engine/Private/VisualLogger/VisualLogger.cpp
@@ -1,19 +1,20 @@
 // Copyright Epic Games, Inc. All Rights Reserved.
 
 #include "VisualLogger/VisualLogger.h"
-#include "EngineStats.h"
+#include "AI/NavDataGenerator.h" // IWYU pragma: keep
+#include "AI/NavigationSystemBase.h"
 #include "Containers/StringFwd.h"
+#include "EngineStats.h"
+#include "Framework/Docking/TabManager.h"
 #include "Misc/CommandLine.h"
 #include "Misc/ScopeRWLock.h"
 #include "Modules/ModuleManager.h"
-#include "AI/NavDataGenerator.h" // IWYU pragma: keep
-#include "AI/NavigationSystemBase.h"
-#include "Framework/Docking/TabManager.h"
+#include "UnrealEngine.h"
 #include "VisualLogger/VisualLoggerBinaryFileDevice.h"
 #include "VisualLogger/VisualLoggerCustomVersion.h"
-#include "VisualLogger/VisualLoggerTraceDevice.h"
 #include "VisualLogger/VisualLoggerDebugSnapshotInterface.h"
-#include "UnrealEngine.h"
+#include "VisualLogger/VisualLoggerSettings.h"
+#include "VisualLogger/VisualLoggerTraceDevice.h"
 
 #if WITH_EDITOR
 #include "Editor.h"
@@ -27,6 +28,39 @@
 DEFINE_LOG_CATEGORY(LogVisual);
 #if ENABLE_VISUAL_LOG
 
+bool FVisualLogger::bUseVerbosityFilterWhenRecording = false;
+
+namespace UE::VisualLogger
+{
+FAutoConsoleCommandWithOutputDevice CmdActivateVerbosityFilterWhenRecording(
+	TEXT("vislog.ActivateVerbosityFilterWhenRecording"),
+	TEXT("Active log category verbosity will be used to filter out entries to record.\n"),
+	FConsoleCommandWithOutputDeviceDelegate::CreateLambda([](FOutputDevice& OutputDevice)
+		{
+			FVisualLogger::bUseVerbosityFilterWhenRecording = true;
+
+			// Keep setting in sync
+			GetMutableDefault<UVisualLoggerSettings>()->bUseVerbosityFilterWhenRecording = FVisualLogger::bUseVerbosityFilterWhenRecording;
+
+			OutputDevice.Log(ELogVerbosity::Log, TEXT("LogCategory verbosity based filtering is now 'ACTIVE'"));
+		}),
+	ECVF_Default);
+
+FAutoConsoleCommandWithOutputDevice CmdDeactivateVerbosityFilterWhenRecording(
+	TEXT("vislog.DeactivateVerbosityFilterWhenRecording"),
+	TEXT("Log category verbosity is deactivated and everything is recorded.\n"),
+	FConsoleCommandWithOutputDeviceDelegate::CreateLambda([](FOutputDevice& OutputDevice)
+		{
+			FVisualLogger::bUseVerbosityFilterWhenRecording = false;
+
+			// Keep setting in sync
+			GetMutableDefault<UVisualLoggerSettings>()->bUseVerbosityFilterWhenRecording = FVisualLogger::bUseVerbosityFilterWhenRecording;
+
+			OutputDevice.Log(ELogVerbosity::Log, TEXT("LogCategory verbosity based filtering is now 'INACTIVE'"));
+		}),
+	ECVF_Default);
+}
+
 namespace UE::VisLog::Private
 {
 	/** Denotes the consecutive run number, giving us the ability to differentiate UObjectName_0 in Run 1 from UObjectName_0 in Run 2 */
@@ -79,14 +113,6 @@ void FVisualLogger::SegmentLogfImpl(const UObject* Object, const FName& Category
 	);
 }
 
-// Deprecated
-void FVisualLogger::ArrowLogfImpl(const UObject* Object, const FName& CategoryName, ELogVerbosity::Type Verbosity, const FVector& Start, const FVector& End, const FColor& Color, const TCHAR* Fmt, ...)
-{
-	COLLAPSED_LOGF(
-		CurrentEntry->AddArrow(Start, End, CategoryName, Verbosity, Color, Buffer, 8);
-	);
-}
-
 void FVisualLogger::ArrowLineLogfImpl(const UObject* Object, const FName& CategoryName, ELogVerbosity::Type Verbosity, const FVector& Start, const FVector& End, const FColor& Color, const uint16 Mag, const TCHAR* Fmt, ...)
 {
 	COLLAPSED_LOGF(
@@ -94,14 +120,6 @@ void FVisualLogger::ArrowLineLogfImpl(const UObject* Object, const FName& Catego
 	);
 }
 
-// Deprecated
-void FVisualLogger::CircleLogfImpl(const UObject* Object, const FName& CategoryName, ELogVerbosity::Type Verbosity, const FVector& Center, const FVector& UpAxis, const float Radius, const FColor& Color, const uint16 Thickness, const TCHAR* Fmt, ...)
-{
-	COLLAPSED_LOGF(
-		CurrentEntry->AddCircle(Center, UpAxis, Radius, CategoryName, Verbosity, Color, Buffer, Thickness);
-	);
-}
-
 void FVisualLogger::DiscLogfImpl(const UObject* Object, const FName& CategoryName, ELogVerbosity::Type Verbosity, const FVector& Center, const FVector& UpAxis, const float Radius, const FColor& Color, const uint16 Thickness, bool bWireframe, const TCHAR* Fmt, ...)
 {
 	COLLAPSED_LOGF(
diff --git a/Engine/Source/Runtime/Engine/Private/VisualLogger/VisualLoggerAutomationTests.cpp b/Engine/Source/Runtime/Engine/Private/VisualLogger/VisualLoggerAutomationTests.cpp
--- a/Engine/Source/Runtime/Engine/Private/VisualLogger/VisualLoggerAutomationTests.cpp
+++ b/Engine/Source/Runtime/Engine/Private/VisualLogger/VisualLoggerAutomationTests.cpp
@@ -1,11 +1,11 @@
 // Copyright Epic Games, Inc. All Rights Reserved.
 
 #include "VisualLogger/VisualLoggerAutomationTests.h"
-#include "Engine/World.h"
-#include "Misc/AutomationTest.h"
 #include "Engine/Engine.h"
 #include "Engine/Level.h"
-
+#include "Engine/World.h"
+#include "Logging/LogScopedVerbosityOverride.h"
+#include "Misc/AutomationTest.h"
 #include "VisualLogger/VisualLogger.h"
 
 #include UE_INLINE_GENERATED_CPP_BY_NAME(VisualLoggerAutomationTests)
@@ -122,6 +122,10 @@ bool FVisualLogTest::RunTest(const FString& Parameters)
 	FVisualLogger::Get().SetIsRecording(false);
 	CHECK_FAIL(FVisualLogger::Get().IsRecording());
 
+	// Make sure our log is recorded by setting the active verbosity.
+	// This is required if verbosity filtering when recording is enabled.
+	LOG_SCOPE_VERBOSITY_OVERRIDE(LogVisual, ELogVerbosity::Log);
+
 	UE_VLOG(World, LogVisual, Log, TEXT("Simple text line to test vlog"));
 	CHECK_SUCCESS(Context.Device.LastObject == nullptr);
 	CHECK_SUCCESS(Context.Device.LastEntry.TimeStamp == -1);
@@ -177,6 +181,10 @@ bool FVisualLogSegmentsTest::RunTest(const FString& Parameters)
 	FVisualLogger::Get().SetIsRecording(false);
 	CHECK_FAIL(FVisualLogger::Get().IsRecording());
 
+	// Make sure our log is recorded by setting the active verbosity.
+	// This is required if verbosity filtering when recording is enabled.
+	LOG_SCOPE_VERBOSITY_OVERRIDE(LogVisual, ELogVerbosity::Log);
+
 	UE_VLOG_SEGMENT(World, LogVisual, Log, FVector(0, 0, 0), FVector(1, 0, 0), FColor::Red, TEXT("Simple segment log to test vlog"));
 	CHECK_SUCCESS(Context.Device.LastObject == nullptr);
 	CHECK_SUCCESS(Context.Device.LastEntry.TimeStamp == -1);
@@ -482,8 +490,8 @@ bool FVisualLogRedirectionsCreationOrderTest::RunTest(const FString& Parameters)
 	return true;
 }
 
-IMPLEMENT_SIMPLE_AUTOMATION_TEST(FVisualLogRedirectionsWithinHierachyTest, "System.Engine.VisualLogger.Redirections.WithinHierachy", EAutomationTestFlags::ClientContext | EAutomationTestFlags::EditorContext | EAutomationTestFlags::EngineFilter)
-bool FVisualLogRedirectionsWithinHierachyTest::RunTest(const FString& Parameters)
+IMPLEMENT_SIMPLE_AUTOMATION_TEST(FVisualLogRedirectionsWithinHierarchyTest, "System.Engine.VisualLogger.Redirections.WithinHierarchy", EAutomationTestFlags::ClientContext | EAutomationTestFlags::EditorContext | EAutomationTestFlags::EngineFilter)
+bool FVisualLogRedirectionsWithinHierarchyTest::RunTest(const FString& Parameters)
 {
 	UWorld* World = GetSimpleEngineAutomationTestWorld(GetTestFlags());
 	CHECK_NOT_NULL(World);
diff --git a/Engine/Source/Runtime/Engine/Private/VisualLogger/VisualLoggerSettings.cpp b/Engine/Source/Runtime/Engine/Private/VisualLogger/VisualLoggerSettings.cpp
new file mode 100644
--- /dev/null
+++ b/Engine/Source/Runtime/Engine/Private/VisualLogger/VisualLoggerSettings.cpp
@@ -0,0 +1,30 @@
+// Copyright Epic Games, Inc. All Rights Reserved.
+
+#include "VisualLogger/VisualLoggerSettings.h"
+#include "EngineDefines.h"
+
+#if ENABLE_VISUAL_LOG
+#include "VisualLogger/VisualLogger.h"
+
+void UVisualLoggerSettings::PostInitProperties()
+{
+	Super::PostInitProperties();
+
+	FVisualLogger::bUseVerbosityFilterWhenRecording = bUseVerbosityFilterWhenRecording;
+}
+
+#if WITH_EDITOR
+void UVisualLoggerSettings::PostEditChangeProperty(FPropertyChangedEvent& PropertyChangedEvent)
+{
+	static const FName NAME_UseVerbosityFilterWhenRecording = GET_MEMBER_NAME_CHECKED(UVisualLoggerSettings, bUseVerbosityFilterWhenRecording);
+
+	Super::PostEditChangeProperty(PropertyChangedEvent);
+
+	if (PropertyChangedEvent.Property
+		&& PropertyChangedEvent.Property->GetFName() == NAME_UseVerbosityFilterWhenRecording)
+	{
+		FVisualLogger::bUseVerbosityFilterWhenRecording = bUseVerbosityFilterWhenRecording;
+	}
+}
+#endif // WITH_EDITOR
+#endif // ENABLE_VISUAL_LOG
diff --git a/Engine/Source/Runtime/Engine/Public/VisualLogger/VisualLogger.h b/Engine/Source/Runtime/Engine/Public/VisualLogger/VisualLogger.h
--- a/Engine/Source/Runtime/Engine/Public/VisualLogger/VisualLogger.h
+++ b/Engine/Source/Runtime/Engine/Public/VisualLogger/VisualLogger.h
@@ -32,9 +32,6 @@
 // Text, log with output to regular unreal logs too.
 #define UE_VLOG_UELOG(LogOwner, CategoryName, Verbosity, Format, ...) { if(FVisualLogger::IsRecording()) FVisualLogger::CategorizedLogf(LogOwner, CategoryName, ELogVerbosity::Verbosity, Format, ##__VA_ARGS__); UE_LOG(CategoryName, Verbosity, Format, ##__VA_ARGS__); }
 
-// Deprecated. Text, log with output to regular unreal logs too.  Regular log will always happen even if the Visual Logger is disabled by via compiler switch.  See also UE_CVLOG_ALWAYS_UELOG below
-#define UE_VLOG_ALWAYS_UELOG(LogOwner, CategoryName, Verbosity, Format, ...) UE_DEPRECATED_MACRO(5.6, "Use UE_VLOG_UELOG instead, it now behaves like the 'always' version.") { UE_VLOG_UELOG(LogOwner, CategoryName, Verbosity, Format, ##__VA_ARGS__); }
-
 // Segment shape
 #define UE_VLOG_SEGMENT(LogOwner, CategoryName, Verbosity, SegmentStart, SegmentEnd, Color, Format, ...) if(FVisualLogger::IsRecording()) FVisualLogger::SegmentLogf(LogOwner, CategoryName, ELogVerbosity::Verbosity, SegmentStart, SegmentEnd, Color, 0, Format, ##__VA_ARGS__)
 #define UE_CVLOG_SEGMENT(Condition, LogOwner, CategoryName, Verbosity, SegmentStart, SegmentEnd, Color, Format, ...) if(FVisualLogger::IsRecording() && Condition) {UE_VLOG_SEGMENT(LogOwner, CategoryName, Verbosity, SegmentStart, SegmentEnd, Color, Format, ##__VA_ARGS__);}
@@ -135,9 +132,6 @@
 // Make sure to UE_LOG when Visual Logging is disabled.
 #define UE_VLOG_UELOG(LogOwner, CategoryName, Verbosity, Format, ...) { UE_LOG(CategoryName, Verbosity, Format, ##__VA_ARGS__); }
 
-// Deprecated.
-#define UE_VLOG_ALWAYS_UELOG(LogOwner, CategoryName, Verbosity, Format, ...) UE_DEPRECATED_MACRO(5.6, "Use UE_VLOG_UELOG instead, it now behaves like the 'always' version.") { UE_LOG(CategoryName, Verbosity, Format, ##__VA_ARGS__); }
-
 #define UE_VLOG_SEGMENT(LogOwner, CategoryName, Verbosity, SegmentStart, SegmentEnd, Color, DescriptionFormat, ...)
 #define UE_CVLOG_SEGMENT(Condition, LogOwner, CategoryName, Verbosity, SegmentStart, SegmentEnd, Color, DescriptionFormat, ...)
 #define UE_VLOG_SEGMENT_THICK(LogOwner, CategoryName, Verbosity, SegmentStart, SegmentEnd, Color, Thickness, DescriptionFormat, ...)
@@ -205,9 +199,6 @@
 // even when the compiler switch (ENABLE_VISUAL_LOG) is off.
 #define UE_CVLOG_UELOG(Condition, LogOwner, CategoryName, Verbosity, Format, ...) if (Condition) { UE_VLOG_UELOG(LogOwner, CategoryName, Verbosity, Format, ##__VA_ARGS__); }
 
-// Deprecated version.
-#define UE_CVLOG_ALWAYS_UELOG(Condition, LogOwner, CategoryName, Verbosity, Format, ...) UE_DEPRECATED_MACRO(5.6, "Use UE_CVLOG_UELOG instead, it now behaves like the 'always' version.")  if (Condition) { UE_VLOG_UELOG(LogOwner, CategoryName, Verbosity, Format, ##__VA_ARGS__); }
-
 // helper macros
 #define TEXT_EMPTY TEXT("")
 #define TEXT_NULL TEXT("NULL")
@@ -215,8 +206,6 @@
 #define TEXT_FALSE TEXT("FALSE")
 #define TEXT_CONDITION(Condition) ((Condition) ? TEXT_TRUE : TEXT_FALSE)
 
-class FVisualLogger;
-
 ENGINE_API DECLARE_LOG_CATEGORY_EXTERN(LogVisual, Display, All);
 
 #if ENABLE_VISUAL_LOG
@@ -242,23 +231,31 @@ class FVisualLogger : public FOutputDevice
 	static ENGINE_API void PulledConvexLogfImpl(const UObject* LogOwner, const FName& CategoryName, ELogVerbosity::Type Verbosity, const TArray<FVector>& ConvexPoints, float MinZ, float MaxZ, const FColor& Color, const TCHAR* Fmt, ...);
 	static ENGINE_API void MeshLogfImpl(const UObject* LogOwner, const FName& CategoryName, ELogVerbosity::Type Verbosity, const TArray<FVector>& Vertices, const TArray<int32>& Indices, const FColor& Color, const TCHAR* Fmt, ...);
 	static ENGINE_API void ConvexLogfImpl(const UObject* LogOwner, const FName& CategoryName, ELogVerbosity::Type Verbosity, const TArray<FVector>& Points, const FColor& Color, const TCHAR* Fmt, ...);
-	static ENGINE_API void HistogramDataLogfImpl (const UObject* LogOwner, const FName& CategoryName, ELogVerbosity::Type Verbosity, FName GraphName, FName DataName, const FVector2D& Data, const FColor& Color, const TCHAR* Fmt, ...);
-	UE_DEPRECATED(5.6, "Use ArrowLineLogfImpl")
-	static ENGINE_API void ArrowLogfImpl(const UObject* LogOwner, const FName& CategoryName, ELogVerbosity::Type Verbosity, const FVector& Start, const FVector& End, const FColor& Color, const TCHAR* Fmt, ...);
+	static ENGINE_API void HistogramDataLogfImpl(const UObject* LogOwner, const FName& CategoryName, ELogVerbosity::Type Verbosity, FName GraphName, FName DataName, const FVector2D& Data, const FColor& Color, const TCHAR* Fmt, ...);
 	static ENGINE_API void ArrowLineLogfImpl(const UObject* LogOwner, const FName& CategoryName, ELogVerbosity::Type Verbosity, const FVector& Start, const FVector& End, const FColor& Color, const uint16 Mag, const TCHAR* Fmt, ...);
-	UE_DEPRECATED(5.6, "Use DiscLogfImpl")
-	static ENGINE_API void CircleLogfImpl(const UObject* LogOwner, const FName& CategoryName, ELogVerbosity::Type Verbosity, const FVector& Center, const FVector& UpAxis, const float Radius, const FColor& Color, const uint16 Thickness, const TCHAR* Fmt, ...);
 	static ENGINE_API void DiscLogfImpl(const UObject* LogOwner, const FName& CategoryName, ELogVerbosity::Type Verbosity, const FVector& Center, const FVector& UpAxis, const float Radius, const FColor& Color, const uint16 Thickness, bool bWireframe, const TCHAR* Fmt, ...);
 	static ENGINE_API void CoordinateSystemLogfImpl(const UObject* LogOwner, const FName& CategoryName, ELogVerbosity::Type Verbosity, const FVector& AxisLoc, const FRotator& AxisRot, const float Scale, const FColor& Color, const uint16 Thickness, const TCHAR* Fmt, ...);
 
+	static bool IsFilteredOut(const FLogCategoryBase& Category, ELogVerbosity::Type Verbosity)
+	{
+		return bUseVerbosityFilterWhenRecording
+			&& ((Verbosity & ::ELogVerbosity::VerbosityMask) > Category.GetCompileTimeVerbosity()
+				|| Category.IsSuppressed(Verbosity));
+	}
+
 public:
+	ENGINE_API static bool bUseVerbosityFilterWhenRecording;
+
 	// Regular text log
 	template <typename FmtType, typename... Types>
 	static void CategorizedLogf(const UObject* LogOwner, const FLogCategoryBase& Category, ELogVerbosity::Type Verbosity, const FmtType& Fmt, Types... Args)
 	{
 		static_assert(TIsArrayOrRefOfTypeByPredicate<FmtType, TIsCharEncodingCompatibleWithTCHAR>::Value, "Formatting string must be a TCHAR array.");
 		static_assert((TIsValidVariadicFunctionArg<Types>::Value && ...), "Invalid argument(s) passed to FVisualLogger::CategorizedLogf");
-
+		if (IsFilteredOut(Category, Verbosity))
+		{
+			return;
+		}
 		CategorizedLogfImpl(LogOwner, Category.GetCategoryName(), Verbosity, (const TCHAR*)Fmt, Args...);
 	}
 	template <typename FmtType, typename... Types>
@@ -271,31 +268,15 @@ public:
 	}
 
 	// Segment log
-	template <typename FmtType, typename... Types>
-	UE_DEPRECATED(5.4, "Use SegmentLogf")
-	static void GeometryShapeLogf(const UObject* LogOwner, const FLogCategoryBase& Category, ELogVerbosity::Type Verbosity, const FVector& Start, const FVector& End, const FColor& Color, const uint16 Thickness, const FmtType& Fmt, Types... Args)
-	{
-		static_assert(TIsArrayOrRefOfTypeByPredicate<FmtType, TIsCharEncodingCompatibleWithTCHAR>::Value, "Formatting string must be a TCHAR array.");
-		static_assert((TIsValidVariadicFunctionArg<Types>::Value && ...), "Invalid argument(s) passed to FVisualLogger::GeometryShapeLogf");
-
-		SegmentLogfImpl(LogOwner, Category.GetCategoryName(), Verbosity, Start, End, Color, Thickness, (const TCHAR*)Fmt, Args...);
-	}
-	template <typename FmtType, typename... Types>
-	UE_DEPRECATED(5.4, "Use SegmentLogf")
-	static void GeometryShapeLogf(const UObject* LogOwner, const FName& CategoryName, ELogVerbosity::Type Verbosity, const FVector& Start, const FVector& End, const FColor& Color, const uint16 Thickness, const FmtType& Fmt, Types... Args)
-	{
-		static_assert(TIsArrayOrRefOfTypeByPredicate<FmtType, TIsCharEncodingCompatibleWithTCHAR>::Value, "Formatting string must be a TCHAR array.");
-		static_assert((TIsValidVariadicFunctionArg<Types>::Value && ...), "Invalid argument(s) passed to FVisualLogger::GeometryShapeLogf");
-
-		SegmentLogfImpl(LogOwner, CategoryName, Verbosity, Start, End, Color, Thickness, (const TCHAR*)Fmt, Args...);
-	}
-
 	template <typename FmtType, typename... Types>
 	static void SegmentLogf(const UObject* LogOwner, const FLogCategoryBase& Category, ELogVerbosity::Type Verbosity, const FVector& Start, const FVector& End, const FColor& Color, const uint16 Thickness, const FmtType& Fmt, Types... Args)
 	{
 		static_assert(TIsArrayOrRefOfTypeByPredicate<FmtType, TIsCharEncodingCompatibleWithTCHAR>::Value, "Formatting string must be a TCHAR array.");
 		static_assert((TIsValidVariadicFunctionArg<Types>::Value && ...), "Invalid argument(s) passed to FVisualLogger::SegmentLogf");
-
+		if (IsFilteredOut(Category, Verbosity))
+		{
+			return;
+		}
 		SegmentLogfImpl(LogOwner, Category.GetCategoryName(), Verbosity, Start, End, Color, Thickness, (const TCHAR*)Fmt, Args...);
 	}
 	template <typename FmtType, typename... Types>
@@ -308,30 +289,15 @@ public:
 	}
 
 	// Arrow
-	template <typename FmtType, typename... Types>
-	UE_DEPRECATED(5.6, "Use ArrowLineLogf")
-	static void ArrowLogf(const UObject* LogOwner, const FLogCategoryBase& Category, ELogVerbosity::Type Verbosity, const FVector& Start, const FVector& End, const FColor& Color, const FmtType& Fmt, Types... Args)
-	{
-		static_assert(TIsArrayOrRefOfTypeByPredicate<FmtType, TIsCharEncodingCompatibleWithTCHAR>::Value, "Formatting string must be a TCHAR array.");
-		static_assert((TIsValidVariadicFunctionArg<Types>::Value && ...), "Invalid argument(s) passed to FVisualLogger::ArrowLogf");
-		ArrowLineLogfImpl(LogOwner, Category.GetCategoryName(), Verbosity, Start, End, Color, 8, (const TCHAR*)Fmt, Args...);
-	}
-	template <typename FmtType, typename... Types>
-	UE_DEPRECATED(5.6, "Use ArrowLineLogf")
-	static void ArrowLogf(const UObject* LogOwner, const FName& CategoryName, ELogVerbosity::Type Verbosity, const FVector& Start, const FVector& End, const FColor& Color, const FmtType& Fmt, Types... Args)
-	{
-		static_assert(TIsArrayOrRefOfTypeByPredicate<FmtType, TIsCharEncodingCompatibleWithTCHAR>::Value, "Formatting string must be a TCHAR array.");
-		static_assert((TIsValidVariadicFunctionArg<Types>::Value && ...), "Invalid argument(s) passed to FVisualLogger::ArrowLogf");
-
-		ArrowLineLogfImpl(LogOwner, CategoryName, Verbosity, Start, End, Color, 8, (const TCHAR*)Fmt, Args...);
-	}
-
 	template <typename FmtType, typename... Types>
 	static void ArrowLineLogf(const UObject* LogOwner, const FLogCategoryBase& Category, ELogVerbosity::Type Verbosity, const FVector& Start, const FVector& End, const FColor& Color, const uint16 Mag, const FmtType& Fmt, Types... Args)
 	{
 		static_assert(TIsArrayOrRefOfTypeByPredicate<FmtType, TIsCharEncodingCompatibleWithTCHAR>::Value, "Formatting string must be a TCHAR array.");
 		static_assert((TIsValidVariadicFunctionArg<Types>::Value && ...), "Invalid argument(s) passed to FVisualLogger::ArrowLineLogf");
-		
+		if (IsFilteredOut(Category, Verbosity))
+		{
+			return;
+		}
 		ArrowLineLogfImpl(LogOwner, Category.GetCategoryName(), Verbosity, Start, End, Color, Mag, (const TCHAR*)Fmt, Args...);
 	}
 	template <typename FmtType, typename... Types>
@@ -343,34 +309,16 @@ public:
 		ArrowLineLogfImpl(LogOwner, CategoryName, Verbosity, Start, End, Color, Mag, (const TCHAR*)Fmt, Args...);
 	}
 
-
-	// Circle log
-	template <typename FmtType, typename... Types>
-	UE_DEPRECATED(5.6, "Use DiscLogf")
-	static void CircleLogf(const UObject* LogOwner, const FLogCategoryBase& Category, ELogVerbosity::Type Verbosity, const FVector& Center, const FVector& UpAxis, const float Radius, const FColor& Color, const uint16 Thickness, const FmtType& Fmt, Types... Args)
-	{
-		static_assert(TIsArrayOrRefOfTypeByPredicate<FmtType, TIsCharEncodingCompatibleWithTCHAR>::Value, "Formatting string must be a TCHAR array.");
-		static_assert((TIsValidVariadicFunctionArg<Types>::Value && ...), "Invalid argument(s) passed to FVisualLogger::CircleLogf");
-
-		DiscLogfImpl(LogOwner, Category.GetCategoryName(), Verbosity, Center, UpAxis, Radius, Color, Thickness, /*bWireframe = */false, (const TCHAR*)Fmt, Args...);
-	}
-	template <typename FmtType, typename... Types>
-	UE_DEPRECATED(5.6, "Use DiscLogf")
-	static void CircleLogf(const UObject* LogOwner, const FName& CategoryName, ELogVerbosity::Type Verbosity, const FVector& Center, const FVector& UpAxis, const float Radius, const FColor& Color, const uint16 Thickness, const FmtType& Fmt, Types... Args)
-	{
-		static_assert(TIsArrayOrRefOfTypeByPredicate<FmtType, TIsCharEncodingCompatibleWithTCHAR>::Value, "Formatting string must be a TCHAR array.");
-		static_assert((TIsValidVariadicFunctionArg<Types>::Value && ...), "Invalid argument(s) passed to FVisualLogger::CircleLogf");
-
-		DiscLogfImpl(LogOwner, CategoryName, Verbosity, Center, UpAxis, Radius, Color, Thickness, /*bWireframe = */false, (const TCHAR*)Fmt, Args...);
-	}
-
 	// Disc/circle log
 	template <typename FmtType, typename... Types>
 	static void DiscLogf(const UObject* LogOwner, const FLogCategoryBase& Category, ELogVerbosity::Type Verbosity, const FVector& Center, const FVector& UpAxis, const float Radius, const FColor& Color, const uint16 Thickness, bool bWireframe, const FmtType& Fmt, Types... Args)
 	{
 		static_assert(TIsArrayOrRefOfTypeByPredicate<FmtType, TIsCharEncodingCompatibleWithTCHAR>::Value, "Formatting string must be a TCHAR array.");
 		static_assert((TIsValidVariadicFunctionArg<Types>::Value && ...), "Invalid argument(s) passed to FVisualLogger::DiscLogf");
-
+		if (IsFilteredOut(Category, Verbosity))
+		{
+			return;
+		}
 		DiscLogfImpl(LogOwner, Category.GetCategoryName(), Verbosity, Center, UpAxis, Radius, Color, Thickness, bWireframe, (const TCHAR*)Fmt, Args...);
 	}
 	template <typename FmtType, typename... Types>
@@ -379,35 +327,19 @@ public:
 		static_assert(TIsArrayOrRefOfTypeByPredicate<FmtType, TIsCharEncodingCompatibleWithTCHAR>::Value, "Formatting string must be a TCHAR array.");
 		static_assert((TIsValidVariadicFunctionArg<Types>::Value && ...), "Invalid argument(s) passed to FVisualLogger::DiscLogf");
 
-		DiscLogfImpl(LogOwner, CategoryName, Verbosity, Center, UpAxis, Radius, Color, Thickness, bWireframe,(const TCHAR*)Fmt, Args...);
+		DiscLogfImpl(LogOwner, CategoryName, Verbosity, Center, UpAxis, Radius, Color, Thickness, bWireframe, (const TCHAR*)Fmt, Args...);
 	}
 
 	// Location log
-	template <typename FmtType, typename... Types>
-	UE_DEPRECATED(5.4, "Use LocationLogf")
-	static void GeometryShapeLogf(const UObject* LogOwner, const FLogCategoryBase& Category, ELogVerbosity::Type Verbosity, const FVector& Location, float Radius, const FColor& Color, const FmtType& Fmt, Types... Args)
-	{
-		static_assert(TIsArrayOrRefOfTypeByPredicate<FmtType, TIsCharEncodingCompatibleWithTCHAR>::Value, "Formatting string must be a TCHAR array.");
-		static_assert((TIsValidVariadicFunctionArg<Types>::Value && ...), "Invalid argument(s) passed to FVisualLogger::GeometryShapeLogf");
-
-		LocationLogfImpl(LogOwner, Category.GetCategoryName(), Verbosity, Location, static_cast<uint16>(Radius), Color, (const TCHAR*)Fmt, Args...);
-	}
-	template <typename FmtType, typename... Types>
-	UE_DEPRECATED(5.4, "Use LocationLogf")
-	static void GeometryShapeLogf(const UObject* LogOwner, const FName& CategoryName, ELogVerbosity::Type Verbosity, const FVector& Location, float Radius, const FColor& Color, const FmtType& Fmt, Types... Args)
-	{
-		static_assert(TIsArrayOrRefOfTypeByPredicate<FmtType, TIsCharEncodingCompatibleWithTCHAR>::Value, "Formatting string must be a TCHAR array.");
-		static_assert((TIsValidVariadicFunctionArg<Types>::Value && ...), "Invalid argument(s) passed to FVisualLogger::GeometryShapeLogf");
-
-		LocationLogfImpl(LogOwner, CategoryName, Verbosity, Location, static_cast<uint16>(Radius), Color, (const TCHAR*)Fmt, Args...);
-	}
-
 	template <typename FmtType, typename... Types>
 	static void LocationLogf(const UObject* LogOwner, const FLogCategoryBase& Category, ELogVerbosity::Type Verbosity, const FVector& Location, const uint16 Thickness, const FColor& Color, const FmtType& Fmt, Types... Args)
 	{
 		static_assert(TIsArrayOrRefOfTypeByPredicate<FmtType, TIsCharEncodingCompatibleWithTCHAR>::Value, "Formatting string must be a TCHAR array.");
 		static_assert((TIsValidVariadicFunctionArg<Types>::Value && ...), "Invalid argument(s) passed to FVisualLogger::LocationLogf");
-
+		if (IsFilteredOut(Category, Verbosity))
+		{
+			return;
+		}
 		LocationLogfImpl(LogOwner, Category.GetCategoryName(), Verbosity, Location, Thickness, Color, (const TCHAR*)Fmt, Args...);
 	}
 	template <typename FmtType, typename... Types>
@@ -425,7 +357,10 @@ public:
 	{
 		static_assert(TIsArrayOrRefOfTypeByPredicate<FmtType, TIsCharEncodingCompatibleWithTCHAR>::Value, "Formatting string must be a TCHAR array.");
 		static_assert(TAnd<TIsValidVariadicFunctionArg<Types>...>::Value, "Invalid argument(s) passed to FVisualLogger::SphereLogf");
-
+		if (IsFilteredOut(Category, Verbosity))
+		{
+			return;
+		}
 		SphereLogfImpl(LogOwner, Category.GetCategoryName(), Verbosity, Location, Radius, Color, bWireframe, (const TCHAR*)Fmt, Args...);
 	}
 	template <typename FmtType, typename... Types>
@@ -439,29 +374,14 @@ public:
 
 	// Box log
 	template <typename FmtType, typename... Types>
-	UE_DEPRECATED(5.4, "Use BoxLogf")
-	static void GeometryBoxLogf(const UObject* LogOwner, const FLogCategoryBase& Category, ELogVerbosity::Type Verbosity, const FBox& Box, const FMatrix& Matrix, const FColor& Color, const FmtType& Fmt, Types... Args)
-	{
-		static_assert(TIsArrayOrRefOfTypeByPredicate<FmtType, TIsCharEncodingCompatibleWithTCHAR>::Value, "Formatting string must be a TCHAR array.");
-		static_assert((TIsValidVariadicFunctionArg<Types>::Value && ...), "Invalid argument(s) passed to FVisualLogger::GeometryBoxLogf");
-
-		BoxLogfImpl(LogOwner, Category.GetCategoryName(), Verbosity, Box, Matrix, Color, /*bWireframe = */false, (const TCHAR*)Fmt, Args...);
-	}
-	template <typename FmtType, typename... Types>
-	UE_DEPRECATED(5.4, "Use BoxLogf")
-	static void GeometryBoxLogf(const UObject* LogOwner, const FName& CategoryName, ELogVerbosity::Type Verbosity, const FBox& Box, const FMatrix& Matrix, const FColor& Color, const FmtType& Fmt, Types... Args)
-	{
-		static_assert(TIsArrayOrRefOfTypeByPredicate<FmtType, TIsCharEncodingCompatibleWithTCHAR>::Value, "Formatting string must be a TCHAR array.");
-		static_assert((TIsValidVariadicFunctionArg<Types>::Value && ...), "Invalid argument(s) passed to FVisualLogger::GeometryBoxLogf");
-
-		BoxLogfImpl(LogOwner, CategoryName, Verbosity, Box, Matrix, Color, /*bWireframe = */false, (const TCHAR*)Fmt, Args...);
-	}
-	template <typename FmtType, typename... Types>
 	static void BoxLogf(const UObject* LogOwner, const FLogCategoryBase& Category, ELogVerbosity::Type Verbosity, const FBox& Box, const FMatrix& Matrix, const FColor& Color, bool bWireframe, const FmtType& Fmt, Types... Args)
 	{
 		static_assert(TIsArrayOrRefOfTypeByPredicate<FmtType, TIsCharEncodingCompatibleWithTCHAR>::Value, "Formatting string must be a TCHAR array.");
 		static_assert((TIsValidVariadicFunctionArg<Types>::Value && ...), "Invalid argument(s) passed to FVisualLogger::BoxLogf");
-
+		if (IsFilteredOut(Category, Verbosity))
+		{
+			return;
+		}
 		BoxLogfImpl(LogOwner, Category.GetCategoryName(), Verbosity, Box, Matrix, Color, bWireframe, (const TCHAR*)Fmt, Args...);
 	}
 	template <typename FmtType, typename... Types>
@@ -475,29 +395,14 @@ public:
 
 	// Cone log
 	template <typename FmtType, typename... Types>
-	UE_DEPRECATED(5.4, "Use ConeLogf")
-	static void GeometryShapeLogf(const UObject* LogOwner, const FLogCategoryBase& Category, ELogVerbosity::Type Verbosity, const FVector& Origin, const FVector& Direction, const float Length, const float Angle, const FColor& Color, const FmtType& Fmt, Types... Args)
-	{
-		static_assert(TIsArrayOrRefOfTypeByPredicate<FmtType, TIsCharEncodingCompatibleWithTCHAR>::Value, "Formatting string must be a TCHAR array.");
-		static_assert((TIsValidVariadicFunctionArg<Types>::Value && ...), "Invalid argument(s) passed to FVisualLogger::GeometryShapeLogf");
-
-		ConeLogfImpl(LogOwner, Category.GetCategoryName(), Verbosity, Origin, Direction, Length, Angle, Color, /*bWireframe = */false, (const TCHAR*)Fmt, Args...);
-	}
-	template <typename FmtType, typename... Types>
-	UE_DEPRECATED(5.4, "Use ConeLogf")
-	static void GeometryShapeLogf(const UObject* LogOwner, const FName& CategoryName, ELogVerbosity::Type Verbosity, const FVector& Origin, const FVector& Direction, const float Length, const float Angle, const FColor& Color, const FmtType& Fmt, Types... Args)
-	{
-		static_assert(TIsArrayOrRefOfTypeByPredicate<FmtType, TIsCharEncodingCompatibleWithTCHAR>::Value, "Formatting string must be a TCHAR array.");
-		static_assert((TIsValidVariadicFunctionArg<Types>::Value && ...), "Invalid argument(s) passed to FVisualLogger::GeometryShapeLogf");
-
-		ConeLogfImpl(LogOwner, CategoryName, Verbosity, Origin, Direction, Length, Angle, Color, /*bWireframe = */false, (const TCHAR*)Fmt, Args...);
-	}
-	template <typename FmtType, typename... Types>
 	static void ConeLogf(const UObject* LogOwner, const FLogCategoryBase& Category, ELogVerbosity::Type Verbosity, const FVector& Origin, const FVector& Direction, const float Length, const float Angle, const FColor& Color, bool bWireframe, const FmtType& Fmt, Types... Args)
 	{
 		static_assert(TIsArrayOrRefOfTypeByPredicate<FmtType, TIsCharEncodingCompatibleWithTCHAR>::Value, "Formatting string must be a TCHAR array.");
 		static_assert((TIsValidVariadicFunctionArg<Types>::Value && ...), "Invalid argument(s) passed to FVisualLogger::ConeLogf");
-
+		if (IsFilteredOut(Category, Verbosity))
+		{
+			return;
+		}
 		ConeLogfImpl(LogOwner, Category.GetCategoryName(), Verbosity, Origin, Direction, Length, Angle, Color, bWireframe, (const TCHAR*)Fmt, Args...);
 	}
 	template <typename FmtType, typename... Types>
@@ -511,29 +416,14 @@ public:
 
 	// Cylinder log
 	template <typename FmtType, typename... Types>
-	UE_DEPRECATED(5.4, "Use CylinderLogf")
-	static void GeometryShapeLogf(const UObject* LogOwner, const FLogCategoryBase& Category, ELogVerbosity::Type Verbosity, const FVector& Start, const FVector& End, const float Radius, const FColor& Color, const FmtType& Fmt, Types... Args)
-	{
-		static_assert(TIsArrayOrRefOfTypeByPredicate<FmtType, TIsCharEncodingCompatibleWithTCHAR>::Value, "Formatting string must be a TCHAR array.");
-		static_assert((TIsValidVariadicFunctionArg<Types>::Value && ...), "Invalid argument(s) passed to FVisualLogger::GeometryShapeLogf");
-
-		CylinderLogfImpl(LogOwner, Category.GetCategoryName(), Verbosity, Start, End, Radius, Color, /*bWireframe = */false, (const TCHAR*)Fmt, Args...);
-	}
-	template <typename FmtType, typename... Types>
-	UE_DEPRECATED(5.4, "Use CylinderLogf")
-	static void GeometryShapeLogf(const UObject* LogOwner, const FName& CategoryName, ELogVerbosity::Type Verbosity, const FVector& Start, const FVector& End, const float Radius, const FColor& Color, const FmtType& Fmt, Types... Args)
-	{
-		static_assert(TIsArrayOrRefOfTypeByPredicate<FmtType, TIsCharEncodingCompatibleWithTCHAR>::Value, "Formatting string must be a TCHAR array.");
-		static_assert((TIsValidVariadicFunctionArg<Types>::Value && ...), "Invalid argument(s) passed to FVisualLogger::GeometryShapeLogf");
-
-		CylinderLogfImpl(LogOwner, CategoryName, Verbosity, Start, End, Radius, Color, /*bWireframe = */false, (const TCHAR*)Fmt, Args...);
-	}
-	template <typename FmtType, typename... Types>
 	static void CylinderLogf(const UObject* LogOwner, const FLogCategoryBase& Category, ELogVerbosity::Type Verbosity, const FVector& Start, const FVector& End, const float Radius, const FColor& Color, bool bWireframe, const FmtType& Fmt, Types... Args)
 	{
 		static_assert(TIsArrayOrRefOfTypeByPredicate<FmtType, TIsCharEncodingCompatibleWithTCHAR>::Value, "Formatting string must be a TCHAR array.");
 		static_assert((TIsValidVariadicFunctionArg<Types>::Value && ...), "Invalid argument(s) passed to FVisualLogger::CylinderLogf");
-
+		if (IsFilteredOut(Category, Verbosity))
+		{
+			return;
+		}
 		CylinderLogfImpl(LogOwner, Category.GetCategoryName(), Verbosity, Start, End, Radius, Color, bWireframe, (const TCHAR*)Fmt, Args...);
 	}
 	template <typename FmtType, typename... Types>
@@ -547,33 +437,18 @@ public:
 
 	// Capsule log
 	template <typename FmtType, typename... Types>
-	UE_DEPRECATED(5.4, "Use CapsuleLogf")
-	static void GeometryShapeLogf(const UObject* LogOwner, const FLogCategoryBase& Category, ELogVerbosity::Type Verbosity, const FVector& Base, float HalfHeight, float Radius, const FQuat & Rotation, const FColor& Color, const FmtType& Fmt, Types... Args)
-	{
-		static_assert(TIsArrayOrRefOfTypeByPredicate<FmtType, TIsCharEncodingCompatibleWithTCHAR>::Value, "Formatting string must be a TCHAR array.");
-		static_assert((TIsValidVariadicFunctionArg<Types>::Value && ...), "Invalid argument(s) passed to FVisualLogger::GeometryShapeLogf");
-
-		CapsuleLogfImpl(LogOwner, Category.GetCategoryName(), Verbosity, Base, HalfHeight, Radius, Rotation, Color, /*bWireframe = */false, (const TCHAR*)Fmt, Args...);
-	}
-	template <typename FmtType, typename... Types>
-	UE_DEPRECATED(5.4, "Use CapsuleLogf")
-	static void GeometryShapeLogf(const UObject* LogOwner, const FName& CategoryName, ELogVerbosity::Type Verbosity, const FVector& Base, float HalfHeight, float Radius, const FQuat & Rotation, const FColor& Color, const FmtType& Fmt, Types... Args)
-	{
-		static_assert(TIsArrayOrRefOfTypeByPredicate<FmtType, TIsCharEncodingCompatibleWithTCHAR>::Value, "Formatting string must be a TCHAR array.");
-		static_assert((TIsValidVariadicFunctionArg<Types>::Value && ...), "Invalid argument(s) passed to FVisualLogger::GeometryShapeLogf");
-
-		CapsuleLogfImpl(LogOwner, CategoryName, Verbosity, Base, HalfHeight, Radius, Rotation, Color, /*bWireframe = */false, (const TCHAR*)Fmt, Args...);
-	}
-	template <typename FmtType, typename... Types>
-	static void CapsuleLogf(const UObject* LogOwner, const FLogCategoryBase& Category, ELogVerbosity::Type Verbosity, const FVector& Base, float HalfHeight, float Radius, const FQuat & Rotation, const FColor& Color, bool bWireframe, const FmtType& Fmt, Types... Args)
+	static void CapsuleLogf(const UObject* LogOwner, const FLogCategoryBase& Category, ELogVerbosity::Type Verbosity, const FVector& Base, float HalfHeight, float Radius, const FQuat& Rotation, const FColor& Color, bool bWireframe, const FmtType& Fmt, Types... Args)
 	{
 		static_assert(TIsArrayOrRefOfTypeByPredicate<FmtType, TIsCharEncodingCompatibleWithTCHAR>::Value, "Formatting string must be a TCHAR array.");
 		static_assert((TIsValidVariadicFunctionArg<Types>::Value && ...), "Invalid argument(s) passed to FVisualLogger::CapsuleLogf");
-
+		if (IsFilteredOut(Category, Verbosity))
+		{
+			return;
+		}
 		CapsuleLogfImpl(LogOwner, Category.GetCategoryName(), Verbosity, Base, HalfHeight, Radius, Rotation, Color, bWireframe, (const TCHAR*)Fmt, Args...);
 	}
 	template <typename FmtType, typename... Types>
-	static void CapsuleLogf(const UObject* LogOwner, const FName& CategoryName, ELogVerbosity::Type Verbosity, const FVector& Base, float HalfHeight, float Radius, const FQuat & Rotation, const FColor& Color, bool bWireframe, const FmtType& Fmt, Types... Args)
+	static void CapsuleLogf(const UObject* LogOwner, const FName& CategoryName, ELogVerbosity::Type Verbosity, const FVector& Base, float HalfHeight, float Radius, const FQuat& Rotation, const FColor& Color, bool bWireframe, const FmtType& Fmt, Types... Args)
 	{
 		static_assert(TIsArrayOrRefOfTypeByPredicate<FmtType, TIsCharEncodingCompatibleWithTCHAR>::Value, "Formatting string must be a TCHAR array.");
 		static_assert((TIsValidVariadicFunctionArg<Types>::Value && ...), "Invalid argument(s) passed to FVisualLogger::CapsuleLogf");
@@ -583,29 +458,14 @@ public:
 
 	// NavArea/Extruded convex log
 	template <typename FmtType, typename... Types>
-	UE_DEPRECATED(5.4, "Use PulledConvexLogf")
-	static void NavAreaShapeLogf(const UObject* LogOwner, const FLogCategoryBase& Category, ELogVerbosity::Type Verbosity, const TArray<FVector>& ConvexPoints, float MinZ, float MaxZ, const FColor& Color, const FmtType& Fmt, Types... Args)
-	{
-		static_assert(TIsArrayOrRefOfTypeByPredicate<FmtType, TIsCharEncodingCompatibleWithTCHAR>::Value, "Formatting string must be a TCHAR array.");
-		static_assert((TIsValidVariadicFunctionArg<Types>::Value && ...), "Invalid argument(s) passed to FVisualLogger::NavAreaShapeLogf");
-
-		PulledConvexLogfImpl(LogOwner, Category.GetCategoryName(), Verbosity, ConvexPoints, MinZ, MaxZ, Color, (const TCHAR*)Fmt, Args...);
-	}
-	template <typename FmtType, typename... Types>
-	UE_DEPRECATED(5.4, "Use PulledConvexLogf")
-	static void NavAreaShapeLogf(const UObject* LogOwner, const FName& CategoryName, ELogVerbosity::Type Verbosity, const TArray<FVector>& ConvexPoints, float MinZ, float MaxZ, const FColor& Color, const FmtType& Fmt, Types... Args)
-	{
-		static_assert(TIsArrayOrRefOfTypeByPredicate<FmtType, TIsCharEncodingCompatibleWithTCHAR>::Value, "Formatting string must be a TCHAR array.");
-		static_assert((TIsValidVariadicFunctionArg<Types>::Value && ...), "Invalid argument(s) passed to FVisualLogger::NavAreaShapeLogf");
-
-		PulledConvexLogfImpl(LogOwner, CategoryName, Verbosity, ConvexPoints, MinZ, MaxZ, Color, (const TCHAR*)Fmt, Args...);
-	}
-	template <typename FmtType, typename... Types>
 	static void PulledConvexLogf(const UObject* LogOwner, const FLogCategoryBase& Category, ELogVerbosity::Type Verbosity, const TArray<FVector>& ConvexPoints, float MinZ, float MaxZ, const FColor& Color, const FmtType& Fmt, Types... Args)
 	{
 		static_assert(TIsArrayOrRefOfTypeByPredicate<FmtType, TIsCharEncodingCompatibleWithTCHAR>::Value, "Formatting string must be a TCHAR array.");
 		static_assert((TIsValidVariadicFunctionArg<Types>::Value && ...), "Invalid argument(s) passed to FVisualLogger::PulledConvexLogf");
-
+		if (IsFilteredOut(Category, Verbosity))
+		{
+			return;
+		}
 		PulledConvexLogfImpl(LogOwner, Category.GetCategoryName(), Verbosity, ConvexPoints, MinZ, MaxZ, Color, (const TCHAR*)Fmt, Args...);
 	}
 	template <typename FmtType, typename... Types>
@@ -619,29 +479,14 @@ public:
 
 	// 3d Mesh log
 	template <typename FmtType, typename... Types>
-	UE_DEPRECATED(5.4, "Use MeshLogf")
-	static void GeometryShapeLogf(const UObject* LogOwner, const FLogCategoryBase& Category, ELogVerbosity::Type Verbosity, const TArray<FVector>& Vertices, const TArray<int32>& Indices, const FColor& Color, const FmtType& Fmt, Types... Args)
-	{
-		static_assert(TIsArrayOrRefOfTypeByPredicate<FmtType, TIsCharEncodingCompatibleWithTCHAR>::Value, "Formatting string must be a TCHAR array.");
-		static_assert((TIsValidVariadicFunctionArg<Types>::Value && ...), "Invalid argument(s) passed to FVisualLogger::GeometryShapeLogf");
-
-		MeshLogfImpl(LogOwner, Category.GetCategoryName(), Verbosity, Vertices, Indices, Color, (const TCHAR*)Fmt, Args...);
-	}
-	template <typename FmtType, typename... Types>
-	UE_DEPRECATED(5.4, "Use MeshLogf")
-	static void GeometryShapeLogf(const UObject* LogOwner, const FName& CategoryName, ELogVerbosity::Type Verbosity, const TArray<FVector>& Vertices, const TArray<int32>& Indices, const FColor& Color, const FmtType& Fmt, Types... Args)
-	{
-		static_assert(TIsArrayOrRefOfTypeByPredicate<FmtType, TIsCharEncodingCompatibleWithTCHAR>::Value, "Formatting string must be a TCHAR array.");
-		static_assert((TIsValidVariadicFunctionArg<Types>::Value && ...), "Invalid argument(s) passed to FVisualLogger::GeometryShapeLogf");
-
-		MeshLogfImpl(LogOwner, CategoryName, Verbosity, Vertices, Indices, Color, (const TCHAR*)Fmt, Args...);
-	}
-	template <typename FmtType, typename... Types>
 	static void MeshLogf(const UObject* LogOwner, const FLogCategoryBase& Category, ELogVerbosity::Type Verbosity, const TArray<FVector>& Vertices, const TArray<int32>& Indices, const FColor& Color, const FmtType& Fmt, Types... Args)
 	{
 		static_assert(TIsArrayOrRefOfTypeByPredicate<FmtType, TIsCharEncodingCompatibleWithTCHAR>::Value, "Formatting string must be a TCHAR array.");
 		static_assert((TIsValidVariadicFunctionArg<Types>::Value && ...), "Invalid argument(s) passed to FVisualLogger::MeshLogf");
-
+		if (IsFilteredOut(Category, Verbosity))
+		{
+			return;
+		}
 		MeshLogfImpl(LogOwner, Category.GetCategoryName(), Verbosity, Vertices, Indices, Color, (const TCHAR*)Fmt, Args...);
 	}
 	template <typename FmtType, typename... Types>
@@ -655,29 +500,14 @@ public:
 
 	// 2d Convex shape
 	template <typename FmtType, typename... Types>
-	UE_DEPRECATED(5.4, "Use ConvexLogf")
-	static void GeometryConvexLogf(const UObject* LogOwner, const FLogCategoryBase& Category, ELogVerbosity::Type Verbosity, const TArray<FVector>& Points, const FColor& Color, const FmtType& Fmt, Types... Args)
-	{
-		static_assert(TIsArrayOrRefOfTypeByPredicate<FmtType, TIsCharEncodingCompatibleWithTCHAR>::Value, "Formatting string must be a TCHAR array.");
-		static_assert((TIsValidVariadicFunctionArg<Types>::Value && ...), "Invalid argument(s) passed to FVisualLogger::GeometryConvexLogf");
-
-		ConvexLogfImpl(LogOwner, Category.GetCategoryName(), Verbosity, Points, Color, (const TCHAR*)Fmt, Args...);
-	}
-	template <typename FmtType, typename... Types>
-	UE_DEPRECATED(5.4, "Use ConvexLogf")
-	static void GeometryConvexLogf(const UObject* LogOwner, const FName& CategoryName, ELogVerbosity::Type Verbosity, const TArray<FVector>& Points, const FColor& Color, const FmtType& Fmt, Types... Args)
-	{
-		static_assert(TIsArrayOrRefOfTypeByPredicate<FmtType, TIsCharEncodingCompatibleWithTCHAR>::Value, "Formatting string must be a TCHAR array.");
-		static_assert((TIsValidVariadicFunctionArg<Types>::Value && ...), "Invalid argument(s) passed to FVisualLogger::GeometryConvexLogf");
-
-		ConvexLogfImpl(LogOwner, CategoryName, Verbosity, Points, Color, (const TCHAR*)Fmt, Args...);
-	}
-	template <typename FmtType, typename... Types>
 	static void ConvexLogf(const UObject* LogOwner, const FLogCategoryBase& Category, ELogVerbosity::Type Verbosity, const TArray<FVector>& Points, const FColor& Color, const FmtType& Fmt, Types... Args)
 	{
 		static_assert(TIsArrayOrRefOfTypeByPredicate<FmtType, TIsCharEncodingCompatibleWithTCHAR>::Value, "Formatting string must be a TCHAR array.");
 		static_assert((TIsValidVariadicFunctionArg<Types>::Value && ...), "Invalid argument(s) passed to FVisualLogger::ConvexLogf");
-
+		if (IsFilteredOut(Category, Verbosity))
+		{
+			return;
+		}
 		ConvexLogfImpl(LogOwner, Category.GetCategoryName(), Verbosity, Points, Color, (const TCHAR*)Fmt, Args...);
 	}
 	template <typename FmtType, typename... Types>
@@ -695,7 +525,10 @@ public:
 	{
 		static_assert(TIsArrayOrRefOfTypeByPredicate<FmtType, TIsCharEncodingCompatibleWithTCHAR>::Value, "Formatting string must be a TCHAR array.");
 		static_assert((TIsValidVariadicFunctionArg<Types>::Value && ...), "Invalid argument(s) passed to FVisualLogger::HistogramDataLogf");
-
+		if (IsFilteredOut(Category, Verbosity))
+		{
+			return;
+		}
 		HistogramDataLogfImpl(LogOwner, Category.GetCategoryName(), Verbosity, GraphName, DataName, Data, Color, (const TCHAR*)Fmt, Args...);
 	}
 	template <typename FmtType, typename... Types>
@@ -713,7 +546,10 @@ public:
 	{
 		static_assert(TIsArrayOrRefOfTypeByPredicate<FmtType, TIsCharEncodingCompatibleWithTCHAR>::Value, "Formatting string must be a TCHAR array.");
 		static_assert((TIsValidVariadicFunctionArg<Types>::Value && ...), "Invalid argument(s) passed to FVisualLogger::CoordinateSystemDataLogf");
-
+		if (IsFilteredOut(Category, Verbosity))
+		{
+			return;
+		}
 		CoordinateSystemLogfImpl(LogOwner, Category.GetCategoryName(), Verbosity, AxisLoc, AxisRot, Scale, Color, Thickness, (const TCHAR*)Fmt, Args...);
 	}
 	template <typename FmtType, typename... Types>
@@ -739,7 +575,6 @@ public:
 	static ENGINE_API void EventLog(const UObject* LogOwner, const FName EventTag1, const FVisualLogEventBase& Event1, const FVisualLogEventBase& Event2, const FVisualLogEventBase& Event3, const FVisualLogEventBase& Event4);
 	static ENGINE_API void EventLog(const UObject* LogOwner, const FName EventTag1, const FVisualLogEventBase& Event1, const FVisualLogEventBase& Event2, const FVisualLogEventBase& Event3, const FVisualLogEventBase& Event4, const FVisualLogEventBase& Event5);
 	static ENGINE_API void EventLog(const UObject* LogOwner, const FName EventTag1, const FVisualLogEventBase& Event1, const FVisualLogEventBase& Event2, const FVisualLogEventBase& Event3, const FVisualLogEventBase& Event4, const FVisualLogEventBase& Event5, const FVisualLogEventBase& Event6);
-	
 	static ENGINE_API void EventLog(const UObject* LogOwner, const FVisualLogEventBase& Event1, const FName EventTag1 = NAME_None, const FName EventTag2 = NAME_None, const FName EventTag3 = NAME_None, const FName EventTag4 = NAME_None, const FName EventTag5 = NAME_None, const FName EventTag6 = NAME_None);
 
 	// static getter
@@ -747,9 +582,6 @@ public:
 
 	virtual ~FVisualLogger() {}
 
-	UE_DEPRECATED(5.4, "Following the base class convention and using the name TearDown. Since FVisualLogger::Get() is used internally everywhere, this class isn't designed to be inherited.")
-	virtual void Shutdown() {}
-
 	// called on engine shutdown to flush all, etc.
 	ENGINE_API virtual void TearDown() override;
 
@@ -826,22 +658,19 @@ public:
 	const TArray<FVisualLogDevice*>& GetDevices() const { return OutputDevices; }
 	/** Check if log category can be recorded, verify before using GetEntryToWrite! */
 	ENGINE_API bool IsCategoryLogged(const FLogCategoryBase& Category) const;
-	/** Returns  current entry for given TimeStamp or creates another one  but first it serialize previous 
+
+	/** Returns current entry for given TimeStamp or creates another one, but first it serialize previous 
 	 *	entry as completed to vislog devices. Use VisualLogger::DontCreate to get current entry without serialization
 	 *	@note this function can return null */
 	UE_DEPRECATED_FORGAME(5.4, "Use the static GetEntryToWrite instead because this TimeStamp is inconsistent across multiple instances (or threads in Editor).  This function will be made private/protected.")
 	ENGINE_API FVisualLogEntry* GetEntryToWrite(const UObject* Object, double TimeStamp, ECreateIfNeeded ShouldCreate = ECreateIfNeeded::Create);
+
 	/** Returns  the current (or new) entry for the given object; alternatively nullptr if we aren't allowed to vlog with the given parameters.
 	 * @param LogOwner - The UObject (typically an AActor) we are going to write log entries about.  This becomes the row in the Visual Logger timeline.
 	 * @param LogCategory - The LogCategory we are logging about.  This function will only return a valid log entry if the visual logging is enabled for the category.
 	 */
 	[[nodiscard]] static ENGINE_API FVisualLogEntry* GetEntryToWrite(const UObject* LogOwner, const FLogCategoryBase& LogCategory);
 
-	/** Retrieves last used entry for given UObject
-	 *	@note this function can return null */
-	UE_DEPRECATED(5.6, "This method will no longer be exposed, use ExecuteOnLastEntryForObject instead.")
-	ENGINE_API FVisualLogEntry* GetLastEntryForObject(const UObject* Object);
-
 	/**
 	 * Executes the provided function using the last used entry for given UObject if any.
 	 * Method provides validation for thread safe access when adding information to an existing LogEntry.
@@ -869,21 +698,21 @@ public:
 	const TMap<FName, FVisualLogExtensionInterface*>& GetAllExtensions() const { return AllExtensions; }
 
 	/** internal check for each usage of visual logger */
-	static ENGINE_API bool CheckVisualLogInputInternal(const UObject* Object, const FName& CategoryName, ELogVerbosity::Type Verbosity, UWorld **OutWorld, FVisualLogEntry **OutCurrentEntry);
-	
+	static ENGINE_API bool CheckVisualLogInputInternal(const UObject* Object, const FName& CategoryName, ELogVerbosity::Type Verbosity, UWorld** OutWorld, FVisualLogEntry** OutCurrentEntry);
+
 	/** Returns a current time stamp to associate with a recorded event that occurred on Object (used for ordering events on a timeline) */
 	ENGINE_API double GetTimeStampForObject(const UObject* Object) const;
 
 	/** Sets function to call to get a timestamp instead of the default implementation (e.g. using a network synchronized clock instead of the local world time) */
 	ENGINE_API void SetGetTimeStampFunc(TFunction<double(const UObject*)> Function);
 
-	typedef TMap<FObjectKey, TArray<TWeakObjectPtr<const UObject> > > FOwnerToChildrenRedirectionMap;
+	typedef TMap<FObjectKey, TArray<TWeakObjectPtr<const UObject>>> FOwnerToChildrenRedirectionMap;
 	static ENGINE_API FOwnerToChildrenRedirectionMap& GetRedirectionMap(const UObject* InObject);
 
-	typedef TMap<FObjectKey, TWeakObjectPtr<const UObject> > FChildToOwnerRedirectionMap;
+	typedef TMap<FObjectKey, TWeakObjectPtr<const UObject>> FChildToOwnerRedirectionMap;
 	FChildToOwnerRedirectionMap& GetChildToOwnerRedirectionMap() { return ChildToOwnerMap; }
 
-	typedef TMap<FObjectKey, TWeakObjectPtr<const UWorld> > FObjectToWorldMapType;
+	typedef TMap<FObjectKey, TWeakObjectPtr<const UWorld>> FObjectToWorldMapType;
 	FObjectToWorldMapType& GetObjectToWorldMap() { return ObjectToWorldMap; }
 
 	ENGINE_API void AddClassToAllowList(UClass& InClass);
@@ -905,6 +734,11 @@ private:
 
 	ENGINE_API FVisualLoggerObjectEntryMap& GetThreadCurrentEntryMap();
 
+	/** Retrieves last used entry for given UObject
+	 * @note this function can return null
+	 */
+	ENGINE_API FVisualLogEntry* GetLastEntryForObject(const UObject* Object);
+
 	/** Get global entry to write where all logs are combined, not thread safe */
 	ENGINE_API FVisualLogEntry* GetEntryToWriteInternal(const UObject* Object, double TimeStamp, ECreateIfNeeded ShouldCreate);
 
diff --git a/Engine/Source/Runtime/Engine/Public/VisualLogger/VisualLoggerSettings.h b/Engine/Source/Runtime/Engine/Public/VisualLogger/VisualLoggerSettings.h
new file mode 100644
--- /dev/null
+++ b/Engine/Source/Runtime/Engine/Public/VisualLogger/VisualLoggerSettings.h
@@ -0,0 +1,36 @@
+// Copyright Epic Games, Inc. All Rights Reserved.
+
+#pragma once
+
+#include "Engine/DeveloperSettings.h"
+#include "VisualLoggerSettings.generated.h"
+
+#define UE_API ENGINE_API
+
+UCLASS(MinimalAPI, config = VisualLogger, defaultconfig, DisplayName = "VisualLogger")
+class UVisualLoggerSettings : public UDeveloperSettings
+{
+	GENERATED_BODY()
+
+public:
+
+	UVisualLoggerSettings() = default;
+
+	/**
+	 * Whether logging functions based on LogCategories should record an entry based on the category active verbosity level.
+	 * This reduces the amount of recorded entries which directly improves performance when using VisualLogger in an active session (i.e., PIE)
+	 * This also implies that the verbosity of the desired categories needs to be set accordingly before, or during, the recording.
+	 * Default behavior (false) is to record all entries and the verbosity based filtering is only used inside the tool when analysing the recording.
+	 */
+	UPROPERTY(config, EditAnywhere, Category = "Recording")
+	bool bUseVerbosityFilterWhenRecording = false;
+
+protected:
+	UE_API virtual void PostInitProperties() override;
+
+#if WITH_EDITOR
+	UE_API virtual void PostEditChangeProperty(FPropertyChangedEvent& PropertyChangedEvent) override;
+#endif // WITH_EDITOR
+};
+
+#undef UE_API
diff --git a/Engine/Source/Runtime/Engine/Public/VisualLogger/VisualLoggerTypes.h b/Engine/Source/Runtime/Engine/Public/VisualLogger/VisualLoggerTypes.h
--- a/Engine/Source/Runtime/Engine/Public/VisualLogger/VisualLoggerTypes.h
+++ b/Engine/Source/Runtime/Engine/Public/VisualLogger/VisualLoggerTypes.h
@@ -302,16 +302,12 @@ class FVisualLogDevice
 public:
 	struct FVisualLogEntryItem
 	{
-		FVisualLogEntryItem() {}		
+		FVisualLogEntryItem() {}
 		FVisualLogEntryItem(const FName& InOwnerName, const FName& InOwnerDisplayName, const FName& InOwnerClassName, const FVisualLogEntry& InLogEntry) 
 			: OwnerName(InOwnerName)
 			, OwnerDisplayName(InOwnerDisplayName)
 			, OwnerClassName(InOwnerClassName)
-			, Entry(InLogEntry) 
-		{}
-		UE_DEPRECATED(5.6, "Use the other constructor")
-		FVisualLogEntryItem(FName InOwnerName, FName InOwnerClassName, const FVisualLogEntry& LogEntry) 
-			: FVisualLogEntryItem(InOwnerName, InOwnerName, InOwnerClassName, LogEntry) 
+			, Entry(InLogEntry)
 		{}
 
 		FName OwnerName;
@@ -321,15 +317,13 @@ public:
 	};
 
 	virtual ~FVisualLogDevice() { }
-	UE_DEPRECATED(5.6, "Serialize now takes a display name in parameter, please use/implement the new variant")
-	virtual void Serialize(const UObject* LogOwner, FName OwnerName, FName InOwnerClassName, const FVisualLogEntry& LogEntry) { Serialize(LogOwner, OwnerName, OwnerName, InOwnerClassName, LogEntry); }
 	virtual void Serialize(const UObject* LogOwner, const FName& InOwnerName, const FName& InOwnerDisplayName, const FName& InOwnerClassName, const FVisualLogEntry& InLogEntry) = 0;
 	virtual void Cleanup(bool bReleaseMemory = false) { /* Empty */ }
 	virtual void StartRecordingToFile(double TimeStamp) { /* Empty */ }
 	virtual void StopRecordingToFile(double TimeStamp) { /* Empty */ }
 	virtual void DiscardRecordingToFile() { /* Empty */ }
 	virtual void SetFileName(const FString& InFileName) { /* Empty */ }
-	virtual void GetRecordedLogs(TArray<FVisualLogDevice::FVisualLogEntryItem>& OutLogs)  const { /* Empty */ }
+	virtual void GetRecordedLogs(TArray<FVisualLogDevice::FVisualLogEntryItem>& OutLogs) const { /* Empty */ }
 	virtual bool HasFlags(int32 InFlags) const { return false; }
 	FGuid GetSessionGUID() const { return SessionGUID; }
 	uint32 GetShortSessionID() const { return SessionGUID[0]; }
