[VisuaLogger] CIS fix for issue 1073490 - mismatching #define

#FailedValidations Preflight

diff --git a/Engine/Source/Runtime/Engine/Private/VisualLogger/VisualLoggerSettings.cpp b/Engine/Source/Runtime/Engine/Private/VisualLogger/VisualLoggerSettings.cpp
--- a/Engine/Source/Runtime/Engine/Private/VisualLogger/VisualLoggerSettings.cpp
+++ b/Engine/Source/Runtime/Engine/Private/VisualLogger/VisualLoggerSettings.cpp
@@ -1,7 +1,6 @@
 // Copyright Epic Games, Inc. All Rights Reserved.
 
 #include "VisualLogger/VisualLoggerSettings.h"
-#include "EngineDefines.h"
 
 #if ENABLE_VISUAL_LOG
 #include "VisualLogger/VisualLogger.h"
diff --git a/Engine/Source/Runtime/Engine/Public/VisualLogger/VisualLoggerSettings.h b/Engine/Source/Runtime/Engine/Public/VisualLogger/VisualLoggerSettings.h
--- a/Engine/Source/Runtime/Engine/Public/VisualLogger/VisualLoggerSettings.h
+++ b/Engine/Source/Runtime/Engine/Public/VisualLogger/VisualLoggerSettings.h
@@ -3,6 +3,7 @@
 #pragma once
 
 #include "Engine/DeveloperSettings.h"
+#include "EngineDefines.h"
 #include "VisualLoggerSettings.generated.h"
 
 #define UE_API ENGINE_API
@@ -25,12 +26,14 @@ public:
 	UPROPERTY(config, EditAnywhere, Category = "Recording")
 	bool bUseVerbosityFilterWhenRecording = false;
 
+#if ENABLE_VISUAL_LOG
 protected:
 	UE_API virtual void PostInitProperties() override;
 
 #if WITH_EDITOR
 	UE_API virtual void PostEditChangeProperty(FPropertyChangedEvent& PropertyChangedEvent) override;
 #endif // WITH_EDITOR
+#endif // ENABLE_VISUAL_LOG
 };
 
 #undef UE_API
