Collision Filter Refactor Part 3.4: Updating the physics filter builder and added query builders.
Updated call sites appropriately to use the builder directly. This removes a lot of direct word usage of the old filters.

diff --git a/Engine/Plugins/Experimental/SpatialReadiness/Source/SpatialReadiness/Private/SpatialReadiness/SpatialReadinessSimCallback.cpp b/Engine/Plugins/Experimental/SpatialReadiness/Source/SpatialReadiness/Private/SpatialReadiness/SpatialReadinessSimCallback.cpp
--- a/Engine/Plugins/Experimental/SpatialReadiness/Source/SpatialReadiness/Private/SpatialReadiness/SpatialReadinessSimCallback.cpp
+++ b/Engine/Plugins/Experimental/SpatialReadiness/Source/SpatialReadiness/Private/SpatialReadiness/SpatialReadinessSimCallback.cpp
@@ -73,46 +73,26 @@ int32 FSpatialReadinessSimCallback::AddUnreadyVolume_GT(const FBox& Bounds, cons
 		return INDEX_NONE;
 	}
 	FRigidBodyHandle_External& ParticleHandle = ParticleProxy->GetGameThreadAPI();
-
-	// Create collision response container
-	FCollisionResponseContainer CollisionResponse;
-	// We want this to block everything that simulates, even if it isn't in the ECC_PhysicsBody channel.
-	// Since bodies can be in any channel they want, we have to block all channels.
-	CollisionResponse.SetAllChannels(ECollisionResponse::ECR_Block);
-
+	
 	// Create collision filter data for the particle
-	FCollisionFilterData QueryFilterData, SimFilterData;
-	CreateShapeFilterData(
-		/* MyChannel */            static_cast<uint8>(ECollisionChannel::ECC_WorldDynamic),
-		/* MaskFilter */           FMaskFilter(0),
-		/* SourceObjectID */       0,
-		/* ResponseToChannels */   CollisionResponse,
-		/* ComponentID */          0,
-		/* BodyIndex */            0,
-		/* OutQueryData */         QueryFilterData,
-		/* OutSimData */           SimFilterData,
-		/* bEnableCCD */           true,
-		/* bEnableContactNotify */ false,
-		/* bStaticShape */         true);
-
 	// Notes of query filter data:
-	// * We want to block all channels in sim but overlap all channels query.
-	// * By default CreateShapeFilterData uses CollisionResponse for both
-	// * Word1 is blocking data, and Word2 is overlap data.
-	// * So we OR query Word1 into Word2 so that we block overlaps and clear
-	//   Word1 so that we don't block any queries.
-	QueryFilterData.Word2 |= QueryFilterData.Word1;
-	QueryFilterData.Word1 = 0;
-	QueryFilterData.Word3 |= EPDF_SimpleCollision | EPDF_ComplexCollision;
-	SimFilterData.Word3 = QueryFilterData.Word3;
+	// * We want to block all channels in sim but overlap all channels in query.
+	Chaos::Filter::FShapeFilterBuilder FilterBuilder;
+	FilterBuilder.SetCollisionChannelIndex(ECollisionChannel::ECC_WorldDynamic);
+	FilterBuilder.SetOverlapChannelMask(TNumericLimits<uint64>::Max());
+	FilterBuilder.SetBlockChannelMasks(TNumericLimits<uint64>::Max(), 0);
+	FilterBuilder.SetFilterFlags(Chaos::EFilterFlags::SimpleCollision);
+	FilterBuilder.SetFilterFlags(Chaos::EFilterFlags::ComplexCollision);
+	FilterBuilder.SetFilterFlags(Chaos::EFilterFlags::CCD);
+	FilterBuilder.SetFilterFlags(Chaos::EFilterFlags::StaticShape);
+	const Chaos::Filter::FShapeFilterData ShapeFilterData = FilterBuilder.Build();
 
 	// Make the geometry
 	ParticleHandle.SetGeometry(BoxGeom);
 	ParticleHandle.SetShapeSimCollisionEnabled(0, true);
 	ParticleHandle.SetShapeQueryCollisionEnabled(0, true);
-	ParticleHandle.SetShapeSimData(0, SimFilterData);
+	ParticleHandle.SetShapeFilterData(0, ShapeFilterData);
 	ParticleHandle.ShapesArray()[0]->SetIsProbe(true);
-	ParticleHandle.ShapesArray()[0]->SetQueryData(QueryFilterData);
 #if CHAOS_DEBUG_NAME
 	ParticleHandle.SetDebugName(MakeShared<FString>(FString::Printf(TEXT("UnreadyVolume: %s"), *Description)));
 #endif
diff --git a/Engine/Plugins/Experimental/SpatialReadiness/Tests/Private/PhysicsSpatialReadinessTests.cpp b/Engine/Plugins/Experimental/SpatialReadiness/Tests/Private/PhysicsSpatialReadinessTests.cpp
--- a/Engine/Plugins/Experimental/SpatialReadiness/Tests/Private/PhysicsSpatialReadinessTests.cpp
+++ b/Engine/Plugins/Experimental/SpatialReadiness/Tests/Private/PhysicsSpatialReadinessTests.cpp
@@ -53,30 +53,19 @@ namespace
 		FRigidBodyHandle_External& ParticleHandle = ParticleProxy->GetGameThreadAPI();
 		REQUIRE(ParticleProxy);
 
-		// Create collision response container
-		FCollisionResponseContainer CollisionResponse;
-		CollisionResponse.SetAllChannels(ECollisionResponse::ECR_Block);
-
 		// Create collision filter data for the particle
-		FCollisionFilterData QueryFilterData, SimFilterData;
-		CreateShapeFilterData(
-			/* MyChannel */            static_cast<uint8>(ECollisionChannel::ECC_WorldDynamic),
-			/* MaskFilter */           FMaskFilter(0),
-			/* SourceObjectID */       0,
-			/* ResponseToChannels */   CollisionResponse,
-			/* ComponentID */          0,
-			/* BodyIndex */            0,
-			/* OutQueryData */         QueryFilterData,
-			/* OutSimData */           SimFilterData,
-			/* bEnableCCD */           true,
-			/* bEnableContactNotify */ false,
-			/* bStaticShape */         false);
+		Chaos::Filter::FShapeFilterBuilder Builder;
+		Builder.SetCollisionChannelIndex(ECollisionChannel::ECC_WorldDynamic);
+		Builder.SetBlockChannelMask(TNumericLimits<uint64>::Max());
+		Builder.SetFilterFlags(Chaos::EFilterFlags::CCD);
+		const Chaos::Filter::FShapeFilterData ShapeFilterData = Builder.Build();
+		
 
 		// Make the geometry
 		ParticleHandle.SetGeometry(BoxGeom);
 		ParticleHandle.SetShapeSimCollisionEnabled(0, true);
 		ParticleHandle.SetShapeQueryCollisionEnabled(0, false);
-		ParticleHandle.SetShapeSimData(0, SimFilterData);
+		ParticleHandle.SetShapeFilterData(0, ShapeFilterData);
 		ParticleHandle.SetGravityEnabled(bGravityEnabled);
 		ParticleHandle.SetObjectState(EObjectStateType::Dynamic, false, true);
 #if CHAOS_DEBUG_NAME
diff --git a/Engine/Source/Runtime/Engine/Private/Collision/PhysicsFiltering.cpp b/Engine/Source/Runtime/Engine/Private/Collision/PhysicsFiltering.cpp
--- a/Engine/Source/Runtime/Engine/Private/Collision/PhysicsFiltering.cpp
+++ b/Engine/Source/Runtime/Engine/Private/Collision/PhysicsFiltering.cpp
@@ -3,27 +3,172 @@
 #include "Physics/PhysicsFiltering.h"
 #include "CollisionQueryParams.h"
 
-//////////////////////////////////////////////////////////////////////////
-// FPhysicsFilterBuilder
-
-FPhysicsFilterBuilder::FPhysicsFilterBuilder(TEnumAsByte<enum ECollisionChannel> InObjectType, FMaskFilter MaskFilter, const struct FCollisionResponseContainer& ResponseToChannels)
-	: BlockingBits(0)
-	, TouchingBits(0)
-	, Word3(0)
+void ExtractResponses(const FCollisionResponseContainer& InResponses, uint64& BlockingBits, uint64& TouchingBits)
 {
-	for (int32 i = 0; i < UE_ARRAY_COUNT(ResponseToChannels.EnumArray); i++)
+	BlockingBits = 0;
+	TouchingBits = 0;
+	for (int32 i = 0; i < UE_ARRAY_COUNT(InResponses.EnumArray); i++)
 	{
-		if (ResponseToChannels.EnumArray[i] == ECR_Block)
+		if (InResponses.EnumArray[i] == ECR_Block)
 		{
-			const uint32 ChannelBit = CRC_TO_BITFIELD(i);
+			const uint32 ChannelBit = (1LLU << i);
 			BlockingBits |= ChannelBit;
 		}
-		else if (ResponseToChannels.EnumArray[i] == ECR_Overlap)
+		else if (InResponses.EnumArray[i] == ECR_Overlap)
 		{
-			const uint32 ChannelBit = CRC_TO_BITFIELD(i);
+			const uint32 ChannelBit = (1LLU << i);
 			TouchingBits |= ChannelBit;
 		}
 	}
+}
+
+//////////////////////////////////////////////////////////////////////////
+// FPhysicsFilterBuilder
+
+FPhysicsFilterBuilder::FPhysicsFilterBuilder(const FChaosScene* InScene)
+{
+}
+
+FPhysicsFilterBuilder::FPhysicsFilterBuilder(TEnumAsByte<enum ECollisionChannel> InObjectType, FMaskFilter MaskFilter, const struct FCollisionResponseContainer& ResponseToChannels)
+{
+	SetCollisionChannelIndex(InObjectType);
+	SetResponses(ResponseToChannels);
+	SetMaskFilter(MaskFilter);
+	ShapeFilterBuilder.SetFilterFlags(Chaos::EFilterFlags::None);
+}
+
+FPhysicsFilterBuilder& FPhysicsFilterBuilder::SetCollisionChannelIndex(const TEnumAsByte<enum ECollisionChannel> InObjectType)
+{
+	ShapeFilterBuilder.SetCollisionChannelIndex(InObjectType);
+	return *this;
+}
+
+FPhysicsFilterBuilder& FPhysicsFilterBuilder::SetResponses(const FCollisionResponseContainer& InResponses)
+{
+	uint64 BlockingBits = 0;
+	uint64 TouchingBits = 0;
+	ExtractResponses(InResponses, BlockingBits, TouchingBits);
+	ShapeFilterBuilder.SetBlockChannelMask(BlockingBits);
+	ShapeFilterBuilder.SetOverlapChannelMask(TouchingBits);
+	return *this;
+}
 
-	Word3 = CreateChannelAndFilter(InObjectType, MaskFilter);
+FPhysicsFilterBuilder& FPhysicsFilterBuilder::SetMaskFilter(const FMaskFilter InMaskFilter)
+{
+	ShapeFilterBuilder.SetMaskFilter(InMaskFilter);
+	return *this;
+}
+
+FPhysicsFilterBuilder& FPhysicsFilterBuilder::SetFlags(const Chaos::EFilterFlags InFlags, bool bEnabled)
+{
+	ShapeFilterBuilder.SetFilterFlags(InFlags, bEnabled);
+	return *this;
+}
+
+FPhysicsFilterBuilder& FPhysicsFilterBuilder::SetOwnerID(const uint32 InOwnerId)
+{
+	InstanceData.SetOwnerId(InOwnerId);
+	return *this;
+}
+
+FPhysicsFilterBuilder& FPhysicsFilterBuilder::SetComponentID(const uint32 InComponentID)
+{
+	InstanceData.SetComponentId(InComponentID);
+	return *this;
+}
+
+Chaos::Filter::FShapeFilterData FPhysicsFilterBuilder::BuildShapeFilterData() const
+{
+	return ShapeFilterBuilder.Build();
+}
+
+Chaos::Filter::FInstanceData FPhysicsFilterBuilder::BuildInstanceData() const
+{
+	return InstanceData;
+}
+
+//////////////////////////////////////////////////////////////////////////
+// FPhysicsObjectQueryFilterBuilder
+
+FPhysicsObjectQueryFilterBuilder::FPhysicsObjectQueryFilterBuilder(const FChaosScene* InScene)
+{
+}
+
+FPhysicsObjectQueryFilterBuilder& FPhysicsObjectQueryFilterBuilder::SetObjectTypes(const uint64 ObjectTypes)
+{
+	Builder.SetObjectTypes(ObjectTypes);
+	return *this;
+}
+
+FPhysicsObjectQueryFilterBuilder& FPhysicsObjectQueryFilterBuilder::SetMultiQuery(const bool bMultiQuery)
+{
+	Builder.SetMultiQuery(bMultiQuery);
+	return *this;
+}
+
+FPhysicsObjectQueryFilterBuilder& FPhysicsObjectQueryFilterBuilder::SetMaskFilter(const uint8 InMaskFilter)
+{
+	Builder.SetMaskFilter(InMaskFilter);
+	return *this;
+}
+
+FPhysicsObjectQueryFilterBuilder& FPhysicsObjectQueryFilterBuilder::SetFlags(const Chaos::EFilterFlags InFlags, bool bEnabled)
+{
+	Builder.SetFilterFlags(InFlags, bEnabled);
+	return *this;
+}
+
+Chaos::Filter::FQueryFilterData FPhysicsObjectQueryFilterBuilder::Build() const
+{
+	return Builder.Build();
+}
+
+//////////////////////////////////////////////////////////////////////////
+// FPhysicsTraceQueryFilterBuilder
+
+FPhysicsTraceQueryFilterBuilder::FPhysicsTraceQueryFilterBuilder(const FChaosScene* InScene)
+{
+}
+
+FPhysicsTraceQueryFilterBuilder& FPhysicsTraceQueryFilterBuilder::SetCollisionChannelIndex(const TEnumAsByte<enum ECollisionChannel> InObjectType)
+{
+	Builder.SetCollisionChannelIndex(InObjectType);
+	return *this;
+}
+
+FPhysicsTraceQueryFilterBuilder& FPhysicsTraceQueryFilterBuilder::SetResponses(const FCollisionResponseContainer& InResponses)
+{
+	uint64 BlockingMask = 0;
+	uint64 OverlapMask = 0;
+	ExtractResponses(InResponses, BlockingMask, OverlapMask);
+	Builder.SetBlockChannelMask(BlockingMask);
+	Builder.SetOverlapChannelMask(OverlapMask);
+	return *this;
+}
+
+FPhysicsTraceQueryFilterBuilder& FPhysicsTraceQueryFilterBuilder::SetMaskFilter(const FMaskFilter InMaskFilter)
+{
+	Builder.SetMaskFilter(InMaskFilter);
+	return *this;
+}
+
+FPhysicsTraceQueryFilterBuilder& FPhysicsTraceQueryFilterBuilder::SetFlags(const Chaos::EFilterFlags InFlags, bool bEnabled)
+{
+	Builder.SetFilterFlags(InFlags, bEnabled);
+	return *this;
+}
+
+Chaos::Filter::FQueryFilterData FPhysicsTraceQueryFilterBuilder::Build() const
+{
+	return Builder.Build();
+}
+
+ECollisionChannel GetCollisionChannel(const Chaos::Filter::FShapeFilterData& ShapeFilterData)
+{
+	return (ECollisionChannel)ShapeFilterData.GetCollisionChannelIndex();
+}
+
+ECollisionChannel GetCollisionChannel(const Chaos::Filter::FQueryFilterData& QueryFilterData)
+{
+	return (ECollisionChannel)QueryFilterData.GetCollisionChannelIndex();
 }
diff --git a/Engine/Source/Runtime/Engine/Private/PhysicsEngine/BodyInstance.cpp b/Engine/Source/Runtime/Engine/Private/PhysicsEngine/BodyInstance.cpp
--- a/Engine/Source/Runtime/Engine/Private/PhysicsEngine/BodyInstance.cpp
+++ b/Engine/Source/Runtime/Engine/Private/PhysicsEngine/BodyInstance.cpp
@@ -4747,33 +4747,30 @@ void FBodyInstance::BuildBodyFilterData(FBodyCollisionFilterData& OutFilterData,
 		// CCD is determined by root body in case of welding
 		bool bRootCCD = (WeldParent != nullptr) ? WeldParent->bUseCCD : bUseCCD;
 
-		FCollisionFilterData SimFilterData;
-		FCollisionFilterData SimpleQueryData;
-
 		uint32 SourceObjectID = SourceObject.IsValid() ? SourceObject->GetUniqueID() : ActorOwner ? ActorOwner->GetUniqueID() : 0;
 		uint32 CompID = (OwnerComponentInst != nullptr) ? OwnerComponentInst->GetUniqueID() : 0;
-		CreateShapeFilterData(UseChannel, MaskFilter, SourceObjectID, UseResponse, CompID, InstanceBodyIndex, SimpleQueryData, SimFilterData, bRootCCD && !bPhysicsStatic, bUseNotifyRBCollision, bPhysicsStatic, bUseContactModification);
-
-		FCollisionFilterData ComplexQueryData = SimpleQueryData;
-			
-			// Set output sim data
-		OutFilterData.SimFilter = SimFilterData;
-
-			// Build filterdata variations for complex and simple
-		SimpleQueryData.Word3 |= EPDF_SimpleCollision;
-			if(bUseSimpleAsComplex)
-			{
-			SimpleQueryData.Word3 |= EPDF_ComplexCollision;
-			}
 
-		ComplexQueryData.Word3 |= EPDF_ComplexCollision;
-			if(bUseComplexAsSimple)
-			{
-			ComplexQueryData.Word3 |= EPDF_SimpleCollision;
-			}
-			
-		OutFilterData.QuerySimpleFilter = SimpleQueryData;
-		OutFilterData.QueryComplexFilter = ComplexQueryData;
+		FPhysicsFilterBuilder Builder;
+		Builder.SetOwnerID(SourceObjectID);
+		Builder.SetComponentID(CompID);
+		Builder.SetCollisionChannelIndex(UseChannel);
+		Builder.SetResponses(UseResponse);
+		Builder.SetMaskFilter(MaskFilter);
+		Builder.SetFlags(Chaos::EFilterFlags::CCD, bRootCCD && !bPhysicsStatic);
+		Builder.SetFlags(Chaos::EFilterFlags::ContactNotify, bUseNotifyRBCollision);
+		Builder.SetFlags(Chaos::EFilterFlags::StaticShape, bPhysicsStatic);
+		Builder.SetFlags(Chaos::EFilterFlags::ModifyContacts, bUseContactModification);
+
+		// Build shape filter variations for complex and simple
+		Builder.SetFlags(Chaos::EFilterFlags::SimpleCollision, true);
+		Builder.SetFlags(Chaos::EFilterFlags::ComplexCollision, bUseSimpleAsComplex);
+		const Chaos::Filter::FCombinedShapeFilterData SimpleCombinedFilter(Builder.BuildShapeFilterData(), Builder.BuildInstanceData());
+		Chaos::Filter::FShapeFilterBuilder::GetLegacyShapeFilter(SimpleCombinedFilter, OutFilterData.QuerySimpleFilter, OutFilterData.SimFilter);
+
+		Builder.SetFlags(Chaos::EFilterFlags::ComplexCollision, true);
+		Builder.SetFlags(Chaos::EFilterFlags::SimpleCollision, bUseComplexAsSimple);
+		const Chaos::Filter::FCombinedShapeFilterData ComplexCombinedFilter(Builder.BuildShapeFilterData(), Builder.BuildInstanceData());
+		Chaos::Filter::FShapeFilterBuilder::GetLegacyShapeFilter(ComplexCombinedFilter, OutFilterData.QueryComplexFilter, OutFilterData.SimFilter);
 	}
 }
 
diff --git a/Engine/Source/Runtime/Engine/Private/PhysicsEngine/PhysicsInterfaceUtils.cpp b/Engine/Source/Runtime/Engine/Private/PhysicsEngine/PhysicsInterfaceUtils.cpp
--- a/Engine/Source/Runtime/Engine/Private/PhysicsEngine/PhysicsInterfaceUtils.cpp
+++ b/Engine/Source/Runtime/Engine/Private/PhysicsEngine/PhysicsInterfaceUtils.cpp
@@ -5,110 +5,31 @@
 #include "Physics/PhysicsFiltering.h"
 #include "PhysicsInterfaceTypesCore.h"
 
-FCollisionFilterData CreateObjectQueryFilterData(const bool bTraceComplex, const int32 MultiTrace/*=1 if multi. 0 otherwise*/, const struct FCollisionObjectQueryParams & ObjectParam)
+FCollisionFilterData CreateQueryFilterData(const uint8 MyChannel, const bool bTraceComplex, const FCollisionResponseContainer& InCollisionResponseContainer, const struct FCollisionQueryParams& QueryParam, const struct FCollisionObjectQueryParams& ObjectParam, const bool bMultitrace)
 {
-	/**
-	* Format for QueryData :
-	*		word0 (meta data - ECollisionQuery. Extendable)
-	*
-	*		For object queries
-	*
-	*		word1 (object type queries)
-	*		word2 (unused)
-	*		word3 (Multi (1) or single (0) (top 8) + Flags (lower 24))
-	*/
-
-	FCollisionFilterData NewData;
-
-	NewData.Word0 = (uint32)ECollisionQuery::ObjectQuery;
-
-	if (bTraceComplex)
-	{
-		NewData.Word3 |= EPDF_ComplexCollision;
-	}
-	else
-	{
-		NewData.Word3 |= EPDF_SimpleCollision;
-	}
-
-	// get object param bits
-	NewData.Word1 = ObjectParam.GetQueryBitfield();
-
-	// if 'nothing', then set no bits
-	NewData.Word3 |= CreateChannelAndFilter((ECollisionChannel)MultiTrace, ObjectParam.IgnoreMask);
-
-	return NewData;
+	Chaos::Filter::FQueryFilterData Result = CreateChaosQueryFilterData(MyChannel, bTraceComplex, InCollisionResponseContainer, QueryParam, ObjectParam, bMultitrace);
+	return Chaos::Filter::FQueryFilterBuilder::GetLegacyQueryFilter(Result);
 }
 
-FCollisionFilterData CreateTraceQueryFilterData(const uint8 MyChannel, const bool bTraceComplex, const FCollisionResponseContainer& InCollisionResponseContainer, const FCollisionQueryParams& Params)
-{
-	/**
-	* Format for QueryData :
-	*		word0 (meta data - ECollisionQuery. Extendable)
-	*
-	*		For trace queries
-	*
-	*		word1 (blocking channels)
-	*		word2 (touching channels)
-	*		word3 (MyChannel (top 8) as ECollisionChannel + Flags (lower 24))
-	*/
-
-	FCollisionFilterData NewData;
-
-	NewData.Word0 = (uint32)ECollisionQuery::TraceQuery;
-
-	if (bTraceComplex)
-	{
-		NewData.Word3 |= EPDF_ComplexCollision;
-	}
-	else
-	{
-		NewData.Word3 |= EPDF_SimpleCollision;
-	}
-
-	// word1 encodes 'what i block', word2 encodes 'what i touch'
-	for (int32 i = 0; i<UE_ARRAY_COUNT(InCollisionResponseContainer.EnumArray); i++)
-	{
-		if (InCollisionResponseContainer.EnumArray[i] == ECR_Block)
-		{
-			// if i block, set that in word1
-			NewData.Word1 |= CRC_TO_BITFIELD(i);
-		}
-		else if (InCollisionResponseContainer.EnumArray[i] == ECR_Overlap)
-		{
-			// if i touch, set that in word2
-			NewData.Word2 |= CRC_TO_BITFIELD(i);
-		}
-	}
-
-	// if 'nothing', then set no bits
-	NewData.Word3 |= CreateChannelAndFilter((ECollisionChannel)MyChannel, Params.IgnoreMask);
-
-	return NewData;
-}
-
-#define TRACE_MULTI		1
-#define TRACE_SINGLE	0
-
-/** Utility for creating a PhysX PxFilterData for performing a query (trace) against the scene */
-FCollisionFilterData CreateQueryFilterData(const uint8 MyChannel, const bool bTraceComplex, const FCollisionResponseContainer& InCollisionResponseContainer, const struct FCollisionQueryParams& QueryParam, const struct FCollisionObjectQueryParams & ObjectParam, const bool bMultitrace)
+Chaos::Filter::FQueryFilterData CreateChaosQueryFilterData(const uint8 MyChannel, const bool bTraceComplex, const FCollisionResponseContainer& InCollisionResponseContainer, const struct FCollisionQueryParams& QueryParam, const struct FCollisionObjectQueryParams& ObjectParam, const bool bMultitrace)
 {
-	//#TODO implement chaos
-	if(ObjectParam.IsValid())
+	Chaos::EFilterFlags FlagsToSet = bTraceComplex ? Chaos::EFilterFlags::ComplexCollision : Chaos::EFilterFlags::SimpleCollision;
+	if (ObjectParam.IsValid())
 	{
-		return CreateObjectQueryFilterData(bTraceComplex, (bMultitrace ? TRACE_MULTI : TRACE_SINGLE), ObjectParam);
+		FPhysicsObjectQueryFilterBuilder Builder;
+		Builder.SetObjectTypes(ObjectParam.GetQueryBitfield());
+		Builder.SetMultiQuery(bMultitrace);
+		Builder.SetMaskFilter(ObjectParam.IgnoreMask);
+		Builder.SetFlags(FlagsToSet, true);
+		return Builder.Build();
 	}
 	else
 	{
-		return CreateTraceQueryFilterData(MyChannel, bTraceComplex, InCollisionResponseContainer, QueryParam);
+		FPhysicsTraceQueryFilterBuilder Builder;
+		Builder.SetCollisionChannelIndex((ECollisionChannel)MyChannel);
+		Builder.SetResponses(InCollisionResponseContainer);
+		Builder.SetMaskFilter(QueryParam.IgnoreMask);
+		Builder.SetFlags(FlagsToSet, true);
+		return Builder.Build();
 	}
 }
-
-Chaos::Filter::FQueryFilterData CreateChaosQueryFilterData(const uint8 MyChannel, const bool bTraceComplex, const FCollisionResponseContainer& InCollisionResponseContainer, const struct FCollisionQueryParams& QueryParam, const struct FCollisionObjectQueryParams& ObjectParam, const bool bMultitrace)
-{
-	// TODO @ JoshD: Update
-	PRAGMA_DISABLE_DEPRECATION_WARNINGS
-	FCollisionFilterData FilterData = CreateQueryFilterData(MyChannel, bTraceComplex, InCollisionResponseContainer, QueryParam, ObjectParam, bMultitrace);
-	return Chaos::Filter::FQueryFilterBuilder::BuildFromLegacyQueryFilter(FilterData);
-	PRAGMA_ENABLE_DEPRECATION_WARNINGS
-}
diff --git a/Engine/Source/Runtime/Engine/Public/Physics/PhysicsFiltering.h b/Engine/Source/Runtime/Engine/Public/Physics/PhysicsFiltering.h
--- a/Engine/Source/Runtime/Engine/Public/Physics/PhysicsFiltering.h
+++ b/Engine/Source/Runtime/Engine/Public/Physics/PhysicsFiltering.h
@@ -37,16 +37,36 @@ enum { NumFilterDataFlagBits = 32 - NumExtraFilterBits - NumCollisionChannelBits
 
 struct FPhysicsFilterBuilder
 {
+	UE_INTERNAL FPhysicsFilterBuilder() = default;
+	FPhysicsFilterBuilder(const FPhysicsFilterBuilder&) = default;
+	FPhysicsFilterBuilder(FPhysicsFilterBuilder&&) = default;
+	ENGINE_API FPhysicsFilterBuilder(const FChaosScene* InScene);
+	UE_DEPRECATED(5.8, "Use The constructor that takes a scene instead")
 	ENGINE_API FPhysicsFilterBuilder(TEnumAsByte<enum ECollisionChannel> InObjectType, FMaskFilter MaskFilter, const struct FCollisionResponseContainer& ResponseToChannels);
 
+	FPhysicsFilterBuilder& operator=(const FPhysicsFilterBuilder&) = default;
+	FPhysicsFilterBuilder& operator=(FPhysicsFilterBuilder&&) = default;
+
+	ENGINE_API FPhysicsFilterBuilder& SetCollisionChannelIndex(const TEnumAsByte<enum ECollisionChannel> InObjectType);
+	ENGINE_API FPhysicsFilterBuilder& SetResponses(const FCollisionResponseContainer& InResponses);
+	ENGINE_API FPhysicsFilterBuilder& SetMaskFilter(const FMaskFilter InMaskFilter);
+	ENGINE_API FPhysicsFilterBuilder& SetFlags(const Chaos::EFilterFlags InFlags, bool bEnabled);
+	ENGINE_API FPhysicsFilterBuilder& SetOwnerID(const uint32 InOwnerID);
+	ENGINE_API FPhysicsFilterBuilder& SetComponentID(const uint32 InComponentID);
+
+	ENGINE_API Chaos::Filter::FShapeFilterData BuildShapeFilterData() const;
+	ENGINE_API Chaos::Filter::FInstanceData BuildInstanceData() const;
+
+	UE_DEPRECATED(5.8, "Use SetFlags")
 	inline void ConditionalSetFlags(EPhysXFilterDataFlags Flag, bool bEnabled)
 	{
 		if (bEnabled)
 		{
-			Word3 |= Flag;
+			SetFlags((Chaos::EFilterFlags)Flag, bEnabled);
 		}
 	}
 
+	UE_DEPRECATED(5.8, "Use GetCombinedShapeFilterData instead")
 	inline void GetQueryData(uint32 SourceObjectID, uint32& OutWord0, uint32& OutWord1, uint32& OutWord2, uint32& OutWord3) const
 	{
 		/**
@@ -56,12 +76,16 @@ struct FPhysicsFilterBuilder
 		 *		word2 (touching channels)
 		 *		word3 (ExtraFilter (top NumExtraFilterBits) + MyChannel (next NumCollisionChannelBits) as ECollisionChannel + Flags (remaining NumFilterDataFlagBits)
 		 */
+		PRAGMA_DISABLE_INTERNAL_WARNINGS
+		const Chaos::Filter::FShapeFilterData& FilterData = BuildShapeFilterData();
 		OutWord0 = SourceObjectID;
-		OutWord1 = BlockingBits;
-		OutWord2 = TouchingBits;
-		OutWord3 = Word3;
+		OutWord1 = (uint32)FilterData.GetQueryBlockChannels();
+		OutWord2 = (uint32)FilterData.GetQueryOverlapChannels();
+		OutWord3 = FilterData.GetLegacyWord3();
+		PRAGMA_ENABLE_INTERNAL_WARNINGS
 	}
 
+	UE_DEPRECATED(5.8, "Use GetCombinedShapeFilterData instead")
 	inline void GetSimData(uint32 BodyIndex, uint32 ComponentID, uint32& OutWord0, uint32& OutWord1, uint32& OutWord2, uint32& OutWord3) const
 	{
 		/**
@@ -71,26 +95,82 @@ struct FPhysicsFilterBuilder
 		 *		word2 (skeletal mesh component ID)
 		 *		word3 (ExtraFilter (top NumExtraFilterBits) + MyChannel (next NumCollisionChannelBits) as ECollisionChannel + Flags (remaining NumFilterDataFlagBits)
 		 */
+		PRAGMA_DISABLE_INTERNAL_WARNINGS
+		const Chaos::Filter::FShapeFilterData& FilterData = BuildShapeFilterData();
 		OutWord0 = BodyIndex;
-		OutWord1 = BlockingBits;
+		OutWord1 = (uint32)FilterData.GetSimBlockChannels();
 		OutWord2 = ComponentID;
-		OutWord3 = Word3;
+		OutWord3 = FilterData.GetLegacyWord3();
+		PRAGMA_ENABLE_INTERNAL_WARNINGS
 	}
 
+	UE_DEPRECATED(5.8, "Use GetCombinedShapeFilterData instead")
 	inline void GetCombinedData(uint32& OutBlockingBits, uint32& OutTouchingBits, uint32& OutObjectTypeAndFlags) const
 	{
-		OutBlockingBits = BlockingBits;
-		OutTouchingBits = TouchingBits;
-		OutObjectTypeAndFlags = Word3;
+		PRAGMA_DISABLE_INTERNAL_WARNINGS
+		const Chaos::Filter::FShapeFilterData& FilterData = BuildShapeFilterData();
+		OutBlockingBits = (uint32)FilterData.GetQueryBlockChannels();
+		OutTouchingBits = (uint32)FilterData.GetQueryOverlapChannels();
+		OutObjectTypeAndFlags = FilterData.GetLegacyWord3();
+		PRAGMA_ENABLE_INTERNAL_WARNINGS
 	}
 
 private:
-	uint32 BlockingBits;
-	uint32 TouchingBits;
-	uint32 Word3;
+	Chaos::Filter::FInstanceData InstanceData;
+	PRAGMA_DISABLE_INTERNAL_WARNINGS
+	Chaos::Filter::FShapeFilterBuilder ShapeFilterBuilder;
+	PRAGMA_ENABLE_INTERNAL_WARNINGS
+};
+
+struct FPhysicsObjectQueryFilterBuilder
+{
+	UE_INTERNAL FPhysicsObjectQueryFilterBuilder() = default;
+	FPhysicsObjectQueryFilterBuilder(const FPhysicsObjectQueryFilterBuilder&) = default;
+	FPhysicsObjectQueryFilterBuilder(FPhysicsObjectQueryFilterBuilder&&) = default;
+	ENGINE_API FPhysicsObjectQueryFilterBuilder(const FChaosScene* InScene);
+
+	FPhysicsObjectQueryFilterBuilder& operator=(const FPhysicsObjectQueryFilterBuilder&) = default;
+	FPhysicsObjectQueryFilterBuilder& operator=(FPhysicsObjectQueryFilterBuilder&&) = default;
+
+	ENGINE_API FPhysicsObjectQueryFilterBuilder& SetObjectTypes(const uint64 ObjectTypes);
+	ENGINE_API FPhysicsObjectQueryFilterBuilder& SetMultiQuery(const bool bMultiQuery);
+	ENGINE_API FPhysicsObjectQueryFilterBuilder& SetMaskFilter(const uint8 InMaskFilter);
+	ENGINE_API FPhysicsObjectQueryFilterBuilder& SetFlags(const Chaos::EFilterFlags InFlags, bool bEnabled);
+
+	ENGINE_API Chaos::Filter::FQueryFilterData Build() const;
+
+private:
+	PRAGMA_DISABLE_INTERNAL_WARNINGS
+	Chaos::Filter::FQueryObjectFilterBuilder Builder;
+	PRAGMA_ENABLE_INTERNAL_WARNINGS
+};
+
+struct FPhysicsTraceQueryFilterBuilder
+{
+	UE_INTERNAL FPhysicsTraceQueryFilterBuilder() = default;
+	FPhysicsTraceQueryFilterBuilder(const FPhysicsTraceQueryFilterBuilder&) = default;
+	FPhysicsTraceQueryFilterBuilder(FPhysicsTraceQueryFilterBuilder&&) = default;
+	ENGINE_API FPhysicsTraceQueryFilterBuilder(const FChaosScene* InScene);
+
+	FPhysicsTraceQueryFilterBuilder& operator=(const FPhysicsTraceQueryFilterBuilder&) = default;
+	FPhysicsTraceQueryFilterBuilder& operator=(FPhysicsTraceQueryFilterBuilder&&) = default;
+
+	ENGINE_API FPhysicsTraceQueryFilterBuilder& SetCollisionChannelIndex(const TEnumAsByte<enum ECollisionChannel> InObjectType);
+	ENGINE_API FPhysicsTraceQueryFilterBuilder& SetResponses(const FCollisionResponseContainer& InResponses);
+	ENGINE_API FPhysicsTraceQueryFilterBuilder& SetMaskFilter(const FMaskFilter InMaskFilter);
+	ENGINE_API FPhysicsTraceQueryFilterBuilder& SetFlags(const Chaos::EFilterFlags InFlags, bool bEnabled);
+
+	ENGINE_API Chaos::Filter::FQueryFilterData Build() const;
+
+private:
+	PRAGMA_DISABLE_INTERNAL_WARNINGS
+	Chaos::Filter::FQueryTraceFilterBuilder Builder;
+	PRAGMA_ENABLE_INTERNAL_WARNINGS
 };
 
 /** Utility for creating a FCollisionFilterData for filtering query (trace) and sim (physics) from the Unreal filtering info. */
+
+UE_DEPRECATED(5.8, "Use FPhysicsFilterBuilder directly instead")
 inline void CreateShapeFilterData(
 	const uint8 MyChannel,
 	const FMaskFilter MaskFilter,
@@ -105,18 +185,24 @@ inline void CreateShapeFilterData(
 	bool bStaticShape,
 	bool bModifyContacts = false)
 {
-	FPhysicsFilterBuilder Builder((ECollisionChannel)MyChannel, MaskFilter, ResponseToChannels);
-	Builder.ConditionalSetFlags(EPDF_CCD, bEnableCCD);
-	Builder.ConditionalSetFlags(EPDF_ContactNotify, bEnableContactNotify);
-	Builder.ConditionalSetFlags(EPDF_StaticShape, bStaticShape);
-	Builder.ConditionalSetFlags(EPDF_ModifyContacts, bModifyContacts);
-
-	OutQueryData = FCollisionFilterData();
-	OutSimData = FCollisionFilterData();
-	Builder.GetQueryData(SourceObjectID, OutQueryData.Word0, OutQueryData.Word1, OutQueryData.Word2, OutQueryData.Word3);
-	Builder.GetSimData(BodyIndex, ComponentID, OutSimData.Word0, OutSimData.Word1, OutSimData.Word2, OutSimData.Word3);
+	PRAGMA_DISABLE_INTERNAL_WARNINGS
+	FPhysicsFilterBuilder Builder;
+	Builder.SetCollisionChannelIndex((ECollisionChannel)MyChannel);
+	Builder.SetMaskFilter(MaskFilter);
+	Builder.SetResponses(ResponseToChannels);
+	Builder.SetFlags(Chaos::EFilterFlags::CCD, bEnableCCD);
+	Builder.SetFlags(Chaos::EFilterFlags::ContactNotify, bEnableContactNotify);
+	Builder.SetFlags(Chaos::EFilterFlags::StaticShape, bStaticShape);
+	Builder.SetFlags(Chaos::EFilterFlags::ModifyContacts, bModifyContacts);
+
+	const Chaos::Filter::FCombinedShapeFilterData CombinedShapeFilter(Builder.BuildShapeFilterData(), Builder.BuildInstanceData());
+	Chaos::Filter::FShapeFilterBuilder::GetLegacyShapeFilter(CombinedShapeFilter, OutQueryData, OutSimData);
+	PRAGMA_ENABLE_INTERNAL_WARNINGS
 }
 
+ENGINE_API ECollisionChannel GetCollisionChannel(const Chaos::Filter::FShapeFilterData& ShapeFilterData);
+ENGINE_API ECollisionChannel GetCollisionChannel(const Chaos::Filter::FQueryFilterData& QueryFilterData);
+
 inline ECollisionChannel GetCollisionChannel(uint32 Word3)
 {
 	uint32 ChannelMask = (Word3 << NumExtraFilterBits) >> (32 - NumCollisionChannelBits);
diff --git a/Engine/Source/Runtime/Experimental/Chaos/Private/Chaos/CollisionFilterData.cpp b/Engine/Source/Runtime/Experimental/Chaos/Private/Chaos/CollisionFilterData.cpp
--- a/Engine/Source/Runtime/Experimental/Chaos/Private/Chaos/CollisionFilterData.cpp
+++ b/Engine/Source/Runtime/Experimental/Chaos/Private/Chaos/CollisionFilterData.cpp
@@ -36,6 +36,18 @@ namespace Chaos::Filter
 		Word3 = Word3Cleared | (Flags & FilterFlagsMask);
 	}
 
+	void SetFlags(EFilterFlags& CurrentFlags, const EFilterFlags InFilterFlags, bool bEnabled)
+	{
+		if (bEnabled)
+		{
+			EnumAddFlags(CurrentFlags, InFilterFlags);
+		}
+		else
+		{
+			EnumRemoveFlags(CurrentFlags, InFilterFlags);
+		}
+	}
+
 	bool HasFlag(uint32 Word3, EFilterFlags InFlag)
 	{
 		const uint32 FilterFlags = GetFlags(Word3);
@@ -98,7 +110,7 @@ namespace Chaos::Filter
 	{
 		SetOwnerId(InActorId);
 	}
-	
+
 	uint32 FInstanceData::GetOwnerId() const
 	{
 		return OwnerId;
@@ -176,6 +188,11 @@ namespace Chaos::Filter
 		return SimBlockChannels;
 	}
 
+	uint32 FShapeFilterData::GetLegacyWord3() const
+	{
+		return Word3;
+	}
+
 	bool FQueryFilterData::IsValid() const
 	{
 		return Word0 != 0 || Word1 != 0 || Word2 != 0 || Word3 != 0;
@@ -325,6 +342,13 @@ namespace Chaos::Filter
 		return *this;
 	}
 
+	FShapeFilterBuilder& FShapeFilterBuilder::SetFilterFlags(const EFilterFlags InFilterFlags, bool bEnabled)
+	{
+		EFilterFlags NewFlags = ShapeFilterData.GetFlags();
+		SetFlags(NewFlags, InFilterFlags, bEnabled);
+		return SetFilterFlags(NewFlags);
+	}
+
 	FShapeFilterData FShapeFilterBuilder::Build() const
 	{
 		return ShapeFilterData;
@@ -573,6 +597,19 @@ namespace Chaos::Filter
 		return *this;
 	}
 
+	FQueryObjectFilterBuilder& FQueryObjectFilterBuilder::SetFilterFlags(const EFilterFlags InFilterFlags)
+	{
+		QueryFilterData.SetFlags(InFilterFlags);
+		return *this;
+	}
+
+	FQueryObjectFilterBuilder& FQueryObjectFilterBuilder::SetFilterFlags(const EFilterFlags InFilterFlags, bool bEnabled)
+	{
+		EFilterFlags NewFlags = QueryFilterData.GetFlags();
+		SetFlags(NewFlags, InFilterFlags, bEnabled);
+		return SetFilterFlags(NewFlags);
+	}
+
 	FQueryFilterData FQueryObjectFilterBuilder::Build() const
 	{
 		return QueryFilterData;
@@ -626,6 +663,13 @@ namespace Chaos::Filter
 		return *this;
 	}
 
+	FQueryTraceFilterBuilder& FQueryTraceFilterBuilder::SetFilterFlags(const EFilterFlags InFilterFlags, bool bEnabled)
+	{
+		EFilterFlags NewFlags = QueryFilterData.GetFlags();
+		SetFlags(NewFlags, InFilterFlags, bEnabled);
+		return SetFilterFlags(NewFlags);
+	}
+
 	FQueryFilterData FQueryTraceFilterBuilder::Build() const
 	{
 		return QueryFilterData;
diff --git a/Engine/Source/Runtime/Experimental/Chaos/Public/Chaos/CollisionFilterData.h b/Engine/Source/Runtime/Experimental/Chaos/Public/Chaos/CollisionFilterData.h
--- a/Engine/Source/Runtime/Experimental/Chaos/Public/Chaos/CollisionFilterData.h
+++ b/Engine/Source/Runtime/Experimental/Chaos/Public/Chaos/CollisionFilterData.h
@@ -128,6 +128,8 @@ namespace Chaos::Filter
 		CHAOS_API uint64 GetQueryOverlapChannels() const;
 		CHAOS_API uint64 GetSimBlockChannels() const;
 
+		UE_INTERNAL uint32 GetLegacyWord3() const;
+
 	private:
 		friend class FShapeFilterBuilder;
 		friend FQueryFilterData;
@@ -223,6 +225,7 @@ namespace Chaos::Filter
 		CHAOS_API FShapeFilterBuilder& SetOverlapChannelMask(const uint64 OverlapMask);
 		CHAOS_API FShapeFilterBuilder& SetMaskFilter(const uint8 InMaskFilter);
 		CHAOS_API FShapeFilterBuilder& SetFilterFlags(const EFilterFlags InFilterFlags);
+		CHAOS_API FShapeFilterBuilder& SetFilterFlags(const EFilterFlags InFilterFlags, bool bEnabled);
 
 		CHAOS_API FShapeFilterData Build() const;
 		CHAOS_API static FShapeFilterData BuildBlockAll(const EFilterFlags InFilterFlags = EFilterFlags::None);
@@ -271,6 +274,8 @@ namespace Chaos::Filter
 		CHAOS_API FQueryObjectFilterBuilder& SetObjectTypes(uint64 ObjectTypes);
 		CHAOS_API FQueryObjectFilterBuilder& SetMultiQuery(bool bMultiQuery);
 		CHAOS_API FQueryObjectFilterBuilder& SetMaskFilter(const uint8 InMaskFilter);
+		CHAOS_API FQueryObjectFilterBuilder& SetFilterFlags(const EFilterFlags InFilterFlags);
+		CHAOS_API FQueryObjectFilterBuilder& SetFilterFlags(const EFilterFlags InFilterFlags, bool bEnabled);
 
 		CHAOS_API FQueryFilterData Build() const;
 		CHAOS_API static FQueryFilterData BuildBlockAll();
@@ -291,6 +296,7 @@ namespace Chaos::Filter
 		CHAOS_API FQueryTraceFilterBuilder& SetOverlapChannelMask(const uint64 OverlapChannelMask);
 		CHAOS_API FQueryTraceFilterBuilder& SetMaskFilter(const uint8 InMaskFilter);
 		CHAOS_API FQueryTraceFilterBuilder& SetFilterFlags(const EFilterFlags InFilterFlags);
+		CHAOS_API FQueryTraceFilterBuilder& SetFilterFlags(const EFilterFlags InFilterFlags, bool bEnabled);
 
 		CHAOS_API FQueryFilterData Build() const;
 		CHAOS_API static FQueryFilterData BuildOverlapAll(const EFilterFlags InFilterFlags = EFilterFlags::None);
diff --git a/Engine/Source/Runtime/Experimental/GeometryCollectionEngine/Private/GeometryCollection/GeometryCollectionComponent.cpp b/Engine/Source/Runtime/Experimental/GeometryCollectionEngine/Private/GeometryCollection/GeometryCollectionComponent.cpp
--- a/Engine/Source/Runtime/Experimental/GeometryCollectionEngine/Private/GeometryCollection/GeometryCollectionComponent.cpp
+++ b/Engine/Source/Runtime/Experimental/GeometryCollectionEngine/Private/GeometryCollection/GeometryCollectionComponent.cpp
@@ -4817,13 +4817,19 @@ void UGeometryCollectionComponent::LoadCollisionProfiles()
 			const uint32 ActorID = Owner ? Owner->GetUniqueID() : 0;
 			const uint32 CompID = GetUniqueID();
 
-			Cache.QueryFilter = InitialQueryFilter;
-			Cache.SimFilter = InitialSimFilter;
-			CreateShapeFilterData(Template.ObjectType, BodyInstance.GetMaskFilter(), ActorID, Template.ResponseToChannels, CompID, INDEX_NONE, Cache.QueryFilter, Cache.SimFilter, BodyInstance.bUseCCD, bNotifyCollisions || bNotifyGlobalCollisions, false, false);
-
 			// Maintain parity with the rest of the geometry collection filters.
-			Cache.QueryFilter.Word3 |= (EPDF_SimpleCollision | EPDF_ComplexCollision);
-			Cache.SimFilter.Word3 |= (EPDF_SimpleCollision | EPDF_ComplexCollision);
+			FPhysicsFilterBuilder FilterBuilder;
+			FilterBuilder.SetOwnerID(ActorID);
+			FilterBuilder.SetComponentID(CompID);
+			FilterBuilder.SetCollisionChannelIndex(Template.ObjectType);
+			FilterBuilder.SetMaskFilter(BodyInstance.GetMaskFilter());
+			FilterBuilder.SetResponses(Template.ResponseToChannels);
+			FilterBuilder.SetFlags(Chaos::EFilterFlags::SimpleCollision | Chaos::EFilterFlags::ComplexCollision, true);
+			FilterBuilder.SetFlags(Chaos::EFilterFlags::CCD, BodyInstance.bUseCCD);
+			FilterBuilder.SetFlags(Chaos::EFilterFlags::ContactNotify, bNotifyCollisions || bNotifyGlobalCollisions);
+
+			const Chaos::Filter::FCombinedShapeFilterData CombinedFilter(FilterBuilder.BuildShapeFilterData(), FilterBuilder.BuildInstanceData());
+			Chaos::Filter::FShapeFilterBuilder::GetLegacyShapeFilter(CombinedFilter, Cache.QueryFilter, Cache.SimFilter);
 		
 			Cache.bQueryEnabled = CollisionEnabledHasQuery(Template.CollisionEnabled);
 			Cache.bSimEnabled = CollisionEnabledHasPhysics(Template.CollisionEnabled);
diff --git a/Engine/Source/Runtime/Landscape/Private/LandscapeCollision.cpp b/Engine/Source/Runtime/Landscape/Private/LandscapeCollision.cpp
--- a/Engine/Source/Runtime/Landscape/Private/LandscapeCollision.cpp
+++ b/Engine/Source/Runtime/Landscape/Private/LandscapeCollision.cpp
@@ -401,16 +401,22 @@ void ULandscapeHeightfieldCollisionComponent::OnCreatePhysicsState()
 				TUniquePtr<Chaos::FPerShapeData> NewShape = Chaos::FShapeInstanceProxy::Make(ShapeArray.Num(), ChaosHeightFieldFromCooked);
 
 				// Setup filtering
-				FCollisionFilterData QueryFilterData, SimFilterData;
-				CreateShapeFilterData(static_cast<uint8>(GetCollisionObjectType()), FMaskFilter(0), GetOwner()->GetUniqueID(), GetCollisionResponseToChannels(),
-					GetUniqueID(), 0, QueryFilterData, SimFilterData, true, false, true);
-
-				// Heightfield is used for simple and complex collision
-				QueryFilterData.Word3 |= bCreateSimpleCollision ? EPDF_ComplexCollision : (EPDF_SimpleCollision | EPDF_ComplexCollision);
-				SimFilterData.Word3 |= bCreateSimpleCollision ? EPDF_ComplexCollision : (EPDF_SimpleCollision | EPDF_ComplexCollision);
-
-				NewShape->SetQueryData(QueryFilterData);
-				NewShape->SetSimData(SimFilterData);
+				FPhysicsFilterBuilder Builder;
+				Builder.SetOwnerID(GetOwner()->GetUniqueID());
+				Builder.SetComponentID(GetUniqueID());
+				Builder.SetCollisionChannelIndex(GetCollisionObjectType());
+				Builder.SetResponses(GetCollisionResponseToChannels());
+				// If we don't create separate collision for simple, this heightfield is used for both simple and complex.
+				Builder.SetFlags(Chaos::EFilterFlags::SimpleCollision, !bCreateSimpleCollision);
+				Builder.SetFlags(Chaos::EFilterFlags::ComplexCollision, true);
+				Builder.SetFlags(Chaos::EFilterFlags::CCD, true);
+				Builder.SetFlags(Chaos::EFilterFlags::StaticShape, true);
+
+				const Chaos::Filter::FShapeFilterData ShapeFilterData = Builder.BuildShapeFilterData();
+				const Chaos::Filter::FInstanceData FilterInstanceData = Builder.BuildInstanceData();
+
+				NewShape->SetFilterInstanceData(FilterInstanceData);
+				NewShape->SetShapeFilterData(ShapeFilterData);
 				NewShape->SetMaterials(HeightfieldRef->UsedChaosMaterials);
 
 				Geoms.Emplace(MoveTemp(ChaosHeightFieldFromCooked));
@@ -426,13 +432,12 @@ void ULandscapeHeightfieldCollisionComponent::OnCreatePhysicsState()
 
 					TUniquePtr<Chaos::FPerShapeData> NewSimpleShape = Chaos::FShapeInstanceProxy::Make(ShapeArray.Num(), ChaosSimpleHeightFieldFromCooked);
 
-					FCollisionFilterData QueryFilterDataSimple = QueryFilterData;
-					FCollisionFilterData SimFilterDataSimple = SimFilterData;
-					QueryFilterDataSimple.Word3 = (QueryFilterDataSimple.Word3 & ~EPDF_ComplexCollision) | EPDF_SimpleCollision;
-					SimFilterDataSimple.Word3 = (SimFilterDataSimple.Word3 & ~EPDF_ComplexCollision) | EPDF_SimpleCollision;
+					Builder.SetFlags(Chaos::EFilterFlags::ComplexCollision, false);
+					Builder.SetFlags(Chaos::EFilterFlags::SimpleCollision, true);
+					const Chaos::Filter::FShapeFilterData ShapeFilterDataSimple = Builder.BuildShapeFilterData();
 
-					NewSimpleShape->SetQueryData(QueryFilterDataSimple);
-					NewSimpleShape->SetSimData(SimFilterDataSimple);
+					NewSimpleShape->SetFilterInstanceData(FilterInstanceData);
+					NewSimpleShape->SetShapeFilterData(Builder.BuildShapeFilterData());
 					NewSimpleShape->SetMaterials(HeightfieldRef->UsedChaosMaterials);
 
 					Geoms.Emplace(MoveTemp(ChaosSimpleHeightFieldFromCooked));
@@ -452,13 +457,17 @@ void ULandscapeHeightfieldCollisionComponent::OnCreatePhysicsState()
 					FCollisionResponseContainer CollisionResponse;
 					CollisionResponse.SetAllChannels(ECollisionResponse::ECR_Ignore);
 					CollisionResponse.SetResponse(ECollisionChannel::ECC_Visibility, ECR_Block);
-					FCollisionFilterData QueryFilterDataEd, SimFilterDataEd;
-					CreateShapeFilterData(ECollisionChannel::ECC_Visibility, FMaskFilter(0), GetOwner()->GetUniqueID(), CollisionResponse, GetUniqueID(), 0, QueryFilterDataEd, SimFilterDataEd, true, false, true);
 
-					QueryFilterDataEd.Word3 |= (EPDF_SimpleCollision | EPDF_ComplexCollision);
+					FPhysicsFilterBuilder EdBuilder;
+					EdBuilder.SetCollisionChannelIndex(ECollisionChannel::ECC_Visibility);
+					EdBuilder.SetResponses(CollisionResponse);
+					EdBuilder.SetFlags(Chaos::EFilterFlags::SimpleCollision | Chaos::EFilterFlags::ComplexCollision, true);
+					EdBuilder.SetFlags(Chaos::EFilterFlags::CCD, true);
+					EdBuilder.SetFlags(Chaos::EFilterFlags::StaticShape, true);
+					const Chaos::Filter::FShapeFilterData ShapeFilterDataEd = EdBuilder.BuildShapeFilterData();
 
-					NewEditorShape->SetQueryData(QueryFilterDataEd);
-					NewEditorShape->SetSimData(SimFilterDataEd);
+					NewEditorShape->SetShapeFilterData(ShapeFilterDataEd);
+					NewEditorShape->SetFilterInstanceData(FilterInstanceData);
 					NewEditorShape->SetMaterials(HeightfieldRef->UsedChaosMaterials);
 
 					Geoms.Emplace(MoveTemp(ChaosEditorHeightFieldFromCooked));
