Removing more implicit command list usage from streaming code.

diff --git a/Engine/Source/Runtime/Engine/Private/Streaming/Texture2DMipAllocator_AsyncReallocate.cpp b/Engine/Source/Runtime/Engine/Private/Streaming/Texture2DMipAllocator_AsyncReallocate.cpp
--- a/Engine/Source/Runtime/Engine/Private/Streaming/Texture2DMipAllocator_AsyncReallocate.cpp
+++ b/Engine/Source/Runtime/Engine/Private/Streaming/Texture2DMipAllocator_AsyncReallocate.cpp
@@ -38,7 +38,9 @@ bool FTexture2DMipAllocator_AsyncReallocate::AllocateMips(
 		return false;
 	}
 
-	// Step (1) : Create the texture on the renderthread using RHIAsyncReallocateTexture2D. Wait for the RHI to signal the operation is completed through the FThreadSafeCounter.
+	FRHICommandListImmediate& RHICmdList = FRHICommandListImmediate::Get();
+
+	// Step (1) : Create the texture on the render thread using AsyncReallocateTexture2D. Wait for the RHI to signal the operation is completed through the FThreadSafeCounter.
 	if (!IntermediateTextureRHI)
 	{
 		const FTexture2DMipMap& OwnerMip = *Context.MipsView[PendingFirstLODIdx];
@@ -46,7 +48,7 @@ bool FTexture2DMipAllocator_AsyncReallocate::AllocateMips(
 
 		SyncOptions.Counter->Increment();
 
-		IntermediateTextureRHI = RHIAsyncReallocateTexture2D(
+		IntermediateTextureRHI = RHICmdList.AsyncReallocateTexture2D(
 			Texture2DRHI,
 			ResourceState.NumRequestedLODs,
 			OwnerMip.SizeX,
@@ -57,7 +59,7 @@ bool FTexture2DMipAllocator_AsyncReallocate::AllocateMips(
 		AdvanceTo(ETickState::AllocateMips, ETickThread::Render);
 		return true;
 	}
-	// Step (2) : Finalize the texture state through RHIFinalizeAsyncReallocateTexture2D and lock the new mips.
+	// Step (2) : Finalize the texture state through FinalizeAsyncReallocateTexture2D and lock the new mips.
 	else
 	{
 		const bool bFlushRHIThread = CVarFlushRHIThreadOnSTreamingTextureLocks.GetValueOnAnyThread() > 0;
@@ -77,7 +79,7 @@ bool FTexture2DMipAllocator_AsyncReallocate::AllocateMips(
 #else // Hasn't really been used on console. To investigate!
 			MipInfo.DataSize = 0;
 #endif
-			MipInfo.DestData = RHILockTexture2D(IntermediateTextureRHI, MipIndex - PendingFirstLODIdx, RLM_WriteOnly, MipInfo.RowPitch, false, bFlushRHIThread);
+			MipInfo.DestData = RHICmdList.LockTexture2D(IntermediateTextureRHI, MipIndex - PendingFirstLODIdx, RLM_WriteOnly, MipInfo.RowPitch, false, bFlushRHIThread);
 
 			// Add this mip in the locked list of mips so that it can safely be unlocked when needed.
 			LockedMipIndices.Add(MipIndex - PendingFirstLODIdx);
diff --git a/Engine/Source/Runtime/Engine/Private/Streaming/Texture2DMipAllocator_Virtual.cpp b/Engine/Source/Runtime/Engine/Private/Streaming/Texture2DMipAllocator_Virtual.cpp
--- a/Engine/Source/Runtime/Engine/Private/Streaming/Texture2DMipAllocator_Virtual.cpp
+++ b/Engine/Source/Runtime/Engine/Private/Streaming/Texture2DMipAllocator_Virtual.cpp
@@ -39,10 +39,12 @@ bool FTexture2DMipAllocator_Virtual::AllocateMips(
 		return false;
 	}
 
+	FRHICommandListImmediate& RHICmdList = FRHICommandListImmediate::Get();
+
 	// Step (1) : Create the texture on the renderthread .......
 	if (!IntermediateTextureRHI)
 	{
-		DoConvertToVirtualWithNewMips(Context);
+		DoConvertToVirtualWithNewMips(RHICmdList, Context);
 		// Run the next step, when IntermediateTextureRHI will be somewhat valid (after synchronization).
 		AdvanceTo(ETickState::AllocateMips, ETickThread::Render);
 		return true;
@@ -52,7 +54,7 @@ bool FTexture2DMipAllocator_Virtual::AllocateMips(
 	{
 		const bool bFlushRHIThread = CVarFlushRHIThreadOnSTreamingTextureLocks.GetValueOnAnyThread() > 0;
 
-		RHIVirtualTextureSetFirstMipInMemory(IntermediateTextureRHI, PendingFirstLODIdx);
+		RHICmdList.VirtualTextureSetFirstMipInMemory(IntermediateTextureRHI, PendingFirstLODIdx);
 
 		OutMipInfos.AddDefaulted(CurrentFirstLODIdx);
 
@@ -70,7 +72,7 @@ bool FTexture2DMipAllocator_Virtual::AllocateMips(
 			MipInfo.DataSize = 0;
 #endif
 
-			MipInfo.DestData = RHILockTexture2D(IntermediateTextureRHI, MipIndex, RLM_WriteOnly, MipInfo.RowPitch, false, bFlushRHIThread);
+			MipInfo.DestData = RHICmdList.LockTexture2D(IntermediateTextureRHI, MipIndex, RLM_WriteOnly, MipInfo.RowPitch, false, bFlushRHIThread);
 
 			// Add this mip in the locked list of mips so that it can safely be unlocked when needed.
 			LockedMipIndices.Add(MipIndex);
@@ -89,12 +91,14 @@ bool FTexture2DMipAllocator_Virtual::UploadMips(const FTextureUpdateContext& Con
 		return false;
 	}
 
+	FRHICommandListImmediate& RHICmdList = FRHICommandListImmediate::Get();
+
 	// Unlock the mips so that the texture can be updated.
-	UnlockNewMips();
+	UnlockNewMips(RHICmdList);
 
 	if(IntermediateTextureRHI)
 	{
-		RHIVirtualTextureSetFirstMipVisible(IntermediateTextureRHI, PendingFirstLODIdx);
+		RHICmdList.VirtualTextureSetFirstMipVisible(IntermediateTextureRHI, PendingFirstLODIdx);
 	}
 
 	AdvanceTo(ETickState::Done, ETickThread::None);
@@ -103,12 +107,14 @@ bool FTexture2DMipAllocator_Virtual::UploadMips(const FTextureUpdateContext& Con
 
 void FTexture2DMipAllocator_Virtual::Cancel(const FTextureUpdateSyncOptions& SyncOptions)
 {
-	// Unlock any locked mips.
-	UnlockNewMips();
-
 	if (IntermediateTextureRHI)
 	{
-		RHIVirtualTextureSetFirstMipInMemory(IntermediateTextureRHI, CurrentFirstLODIdx);
+		FRHICommandListImmediate& RHICmdList = FRHICommandListImmediate::Get();
+
+		// Unlock any locked mips.
+		UnlockNewMips(RHICmdList);
+
+		RHICmdList.VirtualTextureSetFirstMipInMemory(IntermediateTextureRHI, CurrentFirstLODIdx);
 	}
 }
 
@@ -122,7 +128,7 @@ FTextureMipAllocator::ETickThread FTexture2DMipAllocator_Virtual::GetCancelThrea
 // ********* Helpers **********
 // ****************************
 
-void FTexture2DMipAllocator_Virtual::UnlockNewMips()
+void FTexture2DMipAllocator_Virtual::UnlockNewMips(FRHICommandListImmediate& RHICmdList)
 {
 	// Unlock any locked mips.
 	if (IntermediateTextureRHI)
@@ -130,21 +136,19 @@ void FTexture2DMipAllocator_Virtual::UnlockNewMips()
 		const bool bFlushRHIThread = CVarFlushRHIThreadOnSTreamingTextureLocks.GetValueOnAnyThread() > 0;
 		for (int32 MipIndex : LockedMipIndices)
 		{
-			RHIUnlockTexture2D(IntermediateTextureRHI, MipIndex, false, CVarFlushRHIThreadOnSTreamingTextureLocks.GetValueOnAnyThread() > 0 );
+			RHICmdList.UnlockTexture2D(IntermediateTextureRHI, MipIndex, false, CVarFlushRHIThreadOnSTreamingTextureLocks.GetValueOnAnyThread() > 0 );
 		}
 		LockedMipIndices.Empty();
 	}
 }
 
-void FTexture2DMipAllocator_Virtual::DoConvertToVirtualWithNewMips(const FTextureUpdateContext& Context)
+void FTexture2DMipAllocator_Virtual::DoConvertToVirtualWithNewMips(FRHICommandListImmediate& RHICmdList, const FTextureUpdateContext& Context)
 {
 	if (Context.Resource)
 	{
 		// If the texture is not virtual, then make it virtual immediately.
 		if (!Context.Resource->GetTexture2DRHI() || (Context.Resource->GetTexture2DRHI()->GetFlags() & ETextureCreateFlags::Virtual) != ETextureCreateFlags::Virtual)
 		{
-			FRHICommandListImmediate& RHICmdList = FRHICommandListImmediate::Get();
-
 			const FTexture2DMipMap& MipMap0 = *Context.MipsView[0];
 
 			FTexture2DResource* Resource2D = Context.Resource->GetTexture2DResource();
diff --git a/Engine/Source/Runtime/Engine/Private/Streaming/Texture2DMipAllocator_Virtual.h b/Engine/Source/Runtime/Engine/Private/Streaming/Texture2DMipAllocator_Virtual.h
--- a/Engine/Source/Runtime/Engine/Private/Streaming/Texture2DMipAllocator_Virtual.h
+++ b/Engine/Source/Runtime/Engine/Private/Streaming/Texture2DMipAllocator_Virtual.h
@@ -1,9 +1,5 @@
 // Copyright Epic Games, Inc. All Rights Reserved.
 
-/*=============================================================================
-Texture2DMipAllocator_AsyncReallocate.h: Implementation of FTextureMipAllocator using RHIAsyncReallocateTexture2D
-=============================================================================*/
-
 #pragma once
 
 #include "CoreMinimal.h"
@@ -35,9 +31,9 @@ public:
 protected:
 
 	// Unlock the mips referenced in LockedMipIndices.
-	void UnlockNewMips();
+	void UnlockNewMips(FRHICommandListImmediate& RHICmdList);
 
-	void DoConvertToVirtualWithNewMips(const FTextureUpdateContext& Context);
+	void DoConvertToVirtualWithNewMips(FRHICommandListImmediate& RHICmdList, const FTextureUpdateContext& Context);
 
 	// The list of mips that are currently locked.
 	TArray<int32, TInlineAllocator<MAX_TEXTURE_MIP_COUNT>> LockedMipIndices;
diff --git a/Engine/Source/Runtime/Engine/Private/Streaming/Texture2DStreamIn.cpp b/Engine/Source/Runtime/Engine/Private/Streaming/Texture2DStreamIn.cpp
--- a/Engine/Source/Runtime/Engine/Private/Streaming/Texture2DStreamIn.cpp
+++ b/Engine/Source/Runtime/Engine/Private/Streaming/Texture2DStreamIn.cpp
@@ -59,7 +59,7 @@ void FTexture2DStreamIn::DoFreeNewMips(const FContext& Context)
 	}
 }
 
-void FTexture2DStreamIn::DoLockNewMips(const FContext& Context)
+void FTexture2DStreamIn::DoLockNewMips(FRHICommandListImmediate& RHICmdList, const FContext& Context)
 {
 	check(Context.CurrentThread == TT_Render);
 
@@ -72,9 +72,9 @@ void FTexture2DStreamIn::DoLockNewMips(const FContext& Context)
 		{
 			check(MipData[MipIndex].Data == nullptr);
 			uint32 DestPitch = -1;
-			MipData[MipIndex].Data = RHILockTexture2D(IntermediateTextureRHI, MipIndex - MipOffset, RLM_WriteOnly, DestPitch, false, CVarFlushRHIThreadOnSTreamingTextureLocks.GetValueOnAnyThread() > 0);
+			MipData[MipIndex].Data = RHICmdList.LockTexture2D(IntermediateTextureRHI, MipIndex - MipOffset, RLM_WriteOnly, DestPitch, false, CVarFlushRHIThreadOnSTreamingTextureLocks.GetValueOnAnyThread() > 0);
 			MipData[MipIndex].Pitch = DestPitch;
-			// note: should store Size but RHILockTexture2D doesn't tell us size
+			// note: should store Size but LockTexture2D doesn't tell us size
 			//MipData[MipIndex].Size
 
 			UE_LOG(LogTextureUpload,Verbose,TEXT("FTexture2DStreamIn::DoLockNewMips( : Lock Mip %d Pitch=%d"),MipIndex,DestPitch);
@@ -83,7 +83,7 @@ void FTexture2DStreamIn::DoLockNewMips(const FContext& Context)
 }
 
 
-void FTexture2DStreamIn::DoUnlockNewMips(const FContext& Context)
+void FTexture2DStreamIn::DoUnlockNewMips(FRHICommandListImmediate& RHICmdList, const FContext& Context)
 {
 	check(Context.CurrentThread == TT_Render);
 
@@ -96,7 +96,7 @@ void FTexture2DStreamIn::DoUnlockNewMips(const FContext& Context)
 		{
 			if (MipData[MipIndex].Data != nullptr)
 			{
-				RHIUnlockTexture2D(IntermediateTextureRHI, MipIndex - MipOffset, false, CVarFlushRHIThreadOnSTreamingTextureLocks.GetValueOnAnyThread() > 0 );
+				RHICmdList.UnlockTexture2D(IntermediateTextureRHI, MipIndex - MipOffset, false, CVarFlushRHIThreadOnSTreamingTextureLocks.GetValueOnAnyThread() > 0 );
 				MipData[MipIndex].Data = nullptr;
 				MipData[MipIndex].Pitch = -1;
 			}
@@ -104,14 +104,14 @@ void FTexture2DStreamIn::DoUnlockNewMips(const FContext& Context)
 	}
 }
 
-void FTexture2DStreamIn::DoCopySharedMips(const FContext& Context)
+void FTexture2DStreamIn::DoCopySharedMips(FRHICommandListImmediate& RHICmdList, const FContext& Context)
 {
 	check(Context.CurrentThread == TT_Render);
 
 	if (!IsCancelled() && IntermediateTextureRHI && Context.Resource)
 	{
 		UE::RHI::CopySharedMips_AssumeSRVMaskState(
-			FRHICommandListExecutor::GetImmediateCommandList(),
+			RHICmdList,
 			Context.Resource->GetTexture2DRHI(),
 			IntermediateTextureRHI);
 	}
diff --git a/Engine/Source/Runtime/Engine/Private/Streaming/Texture2DStreamIn.h b/Engine/Source/Runtime/Engine/Private/Streaming/Texture2DStreamIn.h
--- a/Engine/Source/Runtime/Engine/Private/Streaming/Texture2DStreamIn.h
+++ b/Engine/Source/Runtime/Engine/Private/Streaming/Texture2DStreamIn.h
@@ -39,12 +39,12 @@ protected:
 	void DoFreeNewMips(const FContext& Context);
 
 	// Lock each streamed mips into MipData.
-	void DoLockNewMips(const FContext& Context);
+	void DoLockNewMips(FRHICommandListImmediate& RHICmdList, const FContext& Context);
 	// Unlock each streamed mips from MipData.
-	void DoUnlockNewMips(const FContext& Context);
+	void DoUnlockNewMips(FRHICommandListImmediate& RHICmdList, const FContext& Context);
 
 	// Copy each shared mip to the intermediate texture.
-	void DoCopySharedMips(const FContext& Context);
+	void DoCopySharedMips(FRHICommandListImmediate& RHICmdList, const FContext& Context);
 
 	// Async create the texture to the requested size.
 	void DoAsyncCreateWithNewMips(const FContext& Context);
diff --git a/Engine/Source/Runtime/Engine/Private/Streaming/Texture2DStreamIn_DDC_AsyncCreate.cpp b/Engine/Source/Runtime/Engine/Private/Streaming/Texture2DStreamIn_DDC_AsyncCreate.cpp
--- a/Engine/Source/Runtime/Engine/Private/Streaming/Texture2DStreamIn_DDC_AsyncCreate.cpp
+++ b/Engine/Source/Runtime/Engine/Private/Streaming/Texture2DStreamIn_DDC_AsyncCreate.cpp
@@ -81,8 +81,10 @@ void FTexture2DStreamIn_DDC_AsyncCreate::Finalize(const FContext& Context)
 	DECLARE_SCOPE_CYCLE_COUNTER(TEXT("FTexture2DStreamIn_DDC_AsyncCreate::Finalize"), STAT_Texture2DStreamInDDCAsyncCreate_Finalize, STATGROUP_StreamingDetails);
 	check(Context.CurrentThread == TT_Render);
 
-	DoCopySharedMips(Context);
-	DoFinishUpdate(Context);
+	FRHICommandListImmediate& RHICmdList = FRHICommandListImmediate::Get();
+
+	DoCopySharedMips(RHICmdList, Context);
+	DoFinishUpdate(RHICmdList, Context);
 }
 
 // ****************************
@@ -95,8 +97,10 @@ void FTexture2DStreamIn_DDC_AsyncCreate::Cancel(const FContext& Context)
 	DECLARE_SCOPE_CYCLE_COUNTER(TEXT("FTexture2DStreamIn_DDC_AsyncCreate::Cancel"), STAT_Texture2DStreamInDDCAsyncCreate_Cancel, STATGROUP_StreamingDetails);
 	check(Context.CurrentThread == TT_Render);
 
+	FRHICommandListImmediate& RHICmdList = FRHICommandListImmediate::Get();
+
 	DoFreeNewMips(Context);
-	DoFinishUpdate(Context);
+	DoFinishUpdate(RHICmdList, Context);
 }
 
 #endif // WITH_EDITORONLY_DATA
diff --git a/Engine/Source/Runtime/Engine/Private/Streaming/Texture2DStreamIn_DDC_AsyncReallocate.cpp b/Engine/Source/Runtime/Engine/Private/Streaming/Texture2DStreamIn_DDC_AsyncReallocate.cpp
--- a/Engine/Source/Runtime/Engine/Private/Streaming/Texture2DStreamIn_DDC_AsyncReallocate.cpp
+++ b/Engine/Source/Runtime/Engine/Private/Streaming/Texture2DStreamIn_DDC_AsyncReallocate.cpp
@@ -25,7 +25,9 @@ void FTexture2DStreamIn_DDC_AsyncReallocate::AsyncReallocate(const FContext& Con
 	DECLARE_SCOPE_CYCLE_COUNTER(TEXT("FTexture2DStreamIn_DDC_AsyncReallocate::AsyncReallocate"), STAT_Texture2DStreamInDDCAsyncReallocate_AsyncReallocate, STATGROUP_StreamingDetails);
 	check(Context.CurrentThread == TT_Render);
 
-	DoAsyncReallocate(Context);
+	FRHICommandListImmediate& RHICmdList = FRHICommandListImmediate::Get();
+
+	DoAsyncReallocate(RHICmdList, Context);
 
 	PushTask(Context, TT_Render, SRA_UPDATE_CALLBACK(LockMips), TT_Render, SRA_UPDATE_CALLBACK(Cancel));
 }
@@ -35,7 +37,9 @@ void FTexture2DStreamIn_DDC_AsyncReallocate::LockMips(const FContext& Context)
 	DECLARE_SCOPE_CYCLE_COUNTER(TEXT("FTexture2DStreamIn_DDC_AsyncReallocate::LockMips"), STAT_Texture2DStreamInDDCAsyncReallocate_LockMips, STATGROUP_StreamingDetails);
 	check(Context.CurrentThread == TT_Render);
 
-	DoLockNewMips(Context);
+	FRHICommandListImmediate& RHICmdList = FRHICommandListImmediate::Get();
+
+	DoLockNewMips(RHICmdList, Context);
 
 	PushTask(Context, TT_Async, SRA_UPDATE_CALLBACK(LoadMips), TT_Render, SRA_UPDATE_CALLBACK(Cancel));
 }
@@ -55,8 +59,10 @@ void FTexture2DStreamIn_DDC_AsyncReallocate::Finalize(const FContext& Context)
 	DECLARE_SCOPE_CYCLE_COUNTER(TEXT("FTexture2DStreamIn_DDC_AsyncReallocate::Finalize"), STAT_Texture2DStreamInDDCAsyncReallocate_Finalize, STATGROUP_StreamingDetails);
 	check(Context.CurrentThread == TT_Render);
 
-	DoUnlockNewMips(Context);
-	DoFinishUpdate(Context);
+	FRHICommandListImmediate& RHICmdList = FRHICommandListImmediate::Get();
+
+	DoUnlockNewMips(RHICmdList, Context);
+	DoFinishUpdate(RHICmdList, Context);
 }
 
 // ****************************
@@ -68,8 +74,10 @@ void FTexture2DStreamIn_DDC_AsyncReallocate::Cancel(const FContext& Context)
 	DECLARE_SCOPE_CYCLE_COUNTER(TEXT("FTexture2DStreamIn_DDC_AsyncReallocate::Cancel"), STAT_Texture2DStreamInDDCAsyncReallocate_Cancel, STATGROUP_StreamingDetails);
 	check(Context.CurrentThread == TT_Render);
 
-	DoUnlockNewMips(Context);
-	DoFinishUpdate(Context);
+	FRHICommandListImmediate& RHICmdList = FRHICommandListImmediate::Get();
+
+	DoUnlockNewMips(RHICmdList, Context);
+	DoFinishUpdate(RHICmdList, Context);
 }
 
 #endif // WITH_EDITORONLY_DATA
diff --git a/Engine/Source/Runtime/Engine/Private/Streaming/Texture2DStreamIn_DerivedData.cpp b/Engine/Source/Runtime/Engine/Private/Streaming/Texture2DStreamIn_DerivedData.cpp
--- a/Engine/Source/Runtime/Engine/Private/Streaming/Texture2DStreamIn_DerivedData.cpp
+++ b/Engine/Source/Runtime/Engine/Private/Streaming/Texture2DStreamIn_DerivedData.cpp
@@ -173,8 +173,10 @@ void FTexture2DStreamIn_DerivedData_AsyncCreate::Finalize(const FContext& Contex
 		STAT_Texture2DStreamIn_DerivedData_AsyncCreate_Finalize, STATGROUP_StreamingDetails);
 	check(Context.CurrentThread == TT_Render);
 
-	DoCopySharedMips(Context);
-	DoFinishUpdate(Context);
+	FRHICommandListImmediate& RHICmdList = FRHICommandListImmediate::Get();
+
+	DoCopySharedMips(RHICmdList, Context);
+	DoFinishUpdate(RHICmdList, Context);
 }
 
 void FTexture2DStreamIn_DerivedData_AsyncCreate::Cancel(const FContext& Context)
@@ -183,11 +185,13 @@ void FTexture2DStreamIn_DerivedData_AsyncCreate::Cancel(const FContext& Context)
 		STAT_Texture2DStreamIn_DerivedData_AsyncCreate_Cancel, STATGROUP_StreamingDetails);
 	check(Context.CurrentThread == TT_Render);
 
+	FRHICommandListImmediate& RHICmdList = FRHICommandListImmediate::Get();
+
 	DoCancelIoRequests();
 	DoEndIoRequests(Context);
 
 	DoFreeNewMips(Context);
-	DoFinishUpdate(Context);
+	DoFinishUpdate(RHICmdList, Context);
 }
 
 ///////////////////////////////////////////////////////////////////////////////////////////////////////////////////////
@@ -205,7 +209,9 @@ void FTexture2DStreamIn_DerivedData_AsyncReallocate::AsyncReallocate(const FCont
 		STAT_Texture2DStreamIn_DerivedData_AsyncReallocate_AsyncReallocate, STATGROUP_StreamingDetails);
 	check(Context.CurrentThread == TT_Render);
 
-	DoAsyncReallocate(Context);
+	FRHICommandListImmediate& RHICmdList = FRHICommandListImmediate::Get();
+
+	DoAsyncReallocate(RHICmdList, Context);
 
 	PushTask(Context, TT_Render, SRA_UPDATE_CALLBACK(LockMips), TT_Render, SRA_UPDATE_CALLBACK(Cancel));
 }
@@ -216,7 +222,9 @@ void FTexture2DStreamIn_DerivedData_AsyncReallocate::LockMips(const FContext& Co
 		STAT_Texture2DStreamIn_DerivedData_AsyncReallocate_LockMips, STATGROUP_StreamingDetails);
 	check(Context.CurrentThread == TT_Render);
 
-	DoLockNewMips(Context);
+	FRHICommandListImmediate& RHICmdList = FRHICommandListImmediate::Get();
+
+	DoLockNewMips(RHICmdList, Context);
 
 	PushTask(Context, TT_Async, SRA_UPDATE_CALLBACK(LoadMips), TT_Render, SRA_UPDATE_CALLBACK(Cancel));
 }
@@ -239,8 +247,10 @@ void FTexture2DStreamIn_DerivedData_AsyncReallocate::Finalize(const FContext& Co
 		STAT_Texture2DStreamIn_DerivedData_AsyncReallocate_Finalize, STATGROUP_StreamingDetails);
 	check(Context.CurrentThread == TT_Render);
 
-	DoUnlockNewMips(Context);
-	DoFinishUpdate(Context);
+	FRHICommandListImmediate& RHICmdList = FRHICommandListImmediate::Get();
+
+	DoUnlockNewMips(RHICmdList, Context);
+	DoFinishUpdate(RHICmdList, Context);
 }
 
 void FTexture2DStreamIn_DerivedData_AsyncReallocate::Cancel(const FContext& Context)
@@ -249,11 +259,13 @@ void FTexture2DStreamIn_DerivedData_AsyncReallocate::Cancel(const FContext& Cont
 		STAT_Texture2DStreamIn_DerivedData_AsyncReallocate_Cancel, STATGROUP_StreamingDetails);
 	check(Context.CurrentThread == TT_Render);
 
+	FRHICommandListImmediate& RHICmdList = FRHICommandListImmediate::Get();
+
 	DoCancelIoRequests();
 	DoEndIoRequests(Context);
 
-	DoUnlockNewMips(Context);
-	DoFinishUpdate(Context);
+	DoUnlockNewMips(RHICmdList, Context);
+	DoFinishUpdate(RHICmdList, Context);
 }
 
 ///////////////////////////////////////////////////////////////////////////////////////////////////////////////////////
@@ -269,8 +281,10 @@ void FTexture2DStreamIn_DerivedData_Virtual::LockMips(const FContext& Context)
 	DECLARE_SCOPE_CYCLE_COUNTER(TEXT("FTexture2DStreamIn_DerivedData_Virtual::LockMips"), STAT_Texture2DStreamIn_DerivedData_Virtual_LockMips, STATGROUP_StreamingDetails);
 	check(Context.CurrentThread == TT_Render);
 
-	DoConvertToVirtualWithNewMips(Context);
-	DoLockNewMips(Context);
+	FRHICommandListImmediate& RHICmdList = FRHICommandListImmediate::Get();
+
+	DoConvertToVirtualWithNewMips(RHICmdList, Context);
+	DoLockNewMips(RHICmdList, Context);
 
 	PushTask(Context, TT_Async, SRA_UPDATE_CALLBACK(LoadMips), TT_Render, SRA_UPDATE_CALLBACK(Cancel));
 }
@@ -301,12 +315,14 @@ void FTexture2DStreamIn_DerivedData_Virtual::Finalize(const FContext& Context)
 	DECLARE_SCOPE_CYCLE_COUNTER(TEXT("FTexture2DStreamIn_DerivedData_Virtual::Finalize"), STAT_Texture2DStreamIn_DerivedData_Virtual_Finalize, STATGROUP_StreamingDetails);
 	check(Context.CurrentThread == TT_Render);
 
-	DoUnlockNewMips(Context);
+	FRHICommandListImmediate& RHICmdList = FRHICommandListImmediate::Get();
+
+	DoUnlockNewMips(RHICmdList, Context);
 	if (IntermediateTextureRHI)
 	{
-		RHIVirtualTextureSetFirstMipVisible(IntermediateTextureRHI, PendingFirstLODIdx);
+		RHICmdList.VirtualTextureSetFirstMipVisible(IntermediateTextureRHI, PendingFirstLODIdx);
 	}
-	DoFinishUpdate(Context);
+	DoFinishUpdate(RHICmdList, Context);
 }
 
 void FTexture2DStreamIn_DerivedData_Virtual::Cancel(const FContext& Context)
@@ -314,15 +330,17 @@ void FTexture2DStreamIn_DerivedData_Virtual::Cancel(const FContext& Context)
 	DECLARE_SCOPE_CYCLE_COUNTER(TEXT("FTexture2DStreamIn_DerivedData_Virtual::Cancel"), STAT_Texture2DStreamIn_DerivedData_Virtual_Cancel, STATGROUP_StreamingDetails);
 	check(Context.CurrentThread == TT_Render);
 
+	FRHICommandListImmediate& RHICmdList = FRHICommandListImmediate::Get();
+
 	DoCancelIoRequests();
 	DoEndIoRequests(Context);
 
-	DoUnlockNewMips(Context);
+	DoUnlockNewMips(RHICmdList, Context);
 	if (IntermediateTextureRHI)
 	{
-		RHIVirtualTextureSetFirstMipInMemory(IntermediateTextureRHI, CurrentFirstLODIdx);
+		RHICmdList.VirtualTextureSetFirstMipInMemory(IntermediateTextureRHI, CurrentFirstLODIdx);
 	}
-	DoFinishUpdate(Context);
+	DoFinishUpdate(RHICmdList, Context);
 }
 
 ///////////////////////////////////////////////////////////////////////////////////////////////////////////////////////
diff --git a/Engine/Source/Runtime/Engine/Private/Streaming/Texture2DStreamIn_IO_AsyncCreate.cpp b/Engine/Source/Runtime/Engine/Private/Streaming/Texture2DStreamIn_IO_AsyncCreate.cpp
--- a/Engine/Source/Runtime/Engine/Private/Streaming/Texture2DStreamIn_IO_AsyncCreate.cpp
+++ b/Engine/Source/Runtime/Engine/Private/Streaming/Texture2DStreamIn_IO_AsyncCreate.cpp
@@ -47,8 +47,10 @@ void FTexture2DStreamIn_IO_AsyncCreate::Finalize(const FContext& Context)
 	DECLARE_SCOPE_CYCLE_COUNTER(TEXT("FTexture2DStreamIn_IO_AsyncCreate::Finalize"), STAT_Texture2DStreamInIOAsyncCreate_Finalize, STATGROUP_StreamingDetails);
 	check(Context.CurrentThread == TT_Render);
 
-	DoCopySharedMips(Context);
-	DoFinishUpdate(Context);
+	FRHICommandListImmediate& RHICmdList = FRHICommandListImmediate::Get();
+
+	DoCopySharedMips(RHICmdList, Context);
+	DoFinishUpdate(RHICmdList, Context);
 }
 
 // ****************************
@@ -71,7 +73,9 @@ void FTexture2DStreamIn_IO_AsyncCreate::Cancel(const FContext& Context)
 	DECLARE_SCOPE_CYCLE_COUNTER(TEXT("FTexture2DStreamIn_IO_AsyncCreate::Cancel"), STAT_Texture2DStreamInIOAsyncCreate_Cancel, STATGROUP_StreamingDetails);
 	check(Context.CurrentThread == TT_Render);
 
+	FRHICommandListImmediate& RHICmdList = FRHICommandListImmediate::Get();
+
 	DoFreeNewMips(Context);
-	DoFinishUpdate(Context);
+	DoFinishUpdate(RHICmdList, Context);
 	ReportIOError(Context);
 }
diff --git a/Engine/Source/Runtime/Engine/Private/Streaming/Texture2DStreamIn_IO_AsyncReallocate.cpp b/Engine/Source/Runtime/Engine/Private/Streaming/Texture2DStreamIn_IO_AsyncReallocate.cpp
--- a/Engine/Source/Runtime/Engine/Private/Streaming/Texture2DStreamIn_IO_AsyncReallocate.cpp
+++ b/Engine/Source/Runtime/Engine/Private/Streaming/Texture2DStreamIn_IO_AsyncReallocate.cpp
@@ -23,7 +23,9 @@ void FTexture2DStreamIn_IO_AsyncReallocate::AsyncReallocate(const FContext& Cont
 	DECLARE_SCOPE_CYCLE_COUNTER(TEXT("FTexture2DStreamIn_IO_AsyncReallocate::AsyncReallocate"), STAT_Texture2DStreamInIOAsyncReallocate_AsyncReallocate, STATGROUP_StreamingDetails);
 	check(Context.CurrentThread == TT_Render);
 
-	DoAsyncReallocate(Context);
+	FRHICommandListImmediate& RHICmdList = FRHICommandListImmediate::Get();
+
+	DoAsyncReallocate(RHICmdList, Context);
 
 	PushTask(Context, TT_Render, SRA_UPDATE_CALLBACK(LockMips), TT_Render, SRA_UPDATE_CALLBACK(Cancel));
 }
@@ -33,7 +35,9 @@ void FTexture2DStreamIn_IO_AsyncReallocate::LockMips(const FContext& Context)
 	DECLARE_SCOPE_CYCLE_COUNTER(TEXT("FTexture2DStreamIn_IO_AsyncReallocate::LockMips"), STAT_Texture2DStreamInIOAsyncReallocate_LockMips, STATGROUP_StreamingDetails);
 	check(Context.CurrentThread == TT_Render);
 
-	DoLockNewMips(Context);
+	FRHICommandListImmediate& RHICmdList = FRHICommandListImmediate::Get();
+
+	DoLockNewMips(RHICmdList, Context);
 
 	PushTask(Context, TT_Async, SRA_UPDATE_CALLBACK(LoadMips), TT_Render, SRA_UPDATE_CALLBACK(Cancel));
 }
@@ -53,9 +57,11 @@ void FTexture2DStreamIn_IO_AsyncReallocate::Finalize(const FContext& Context)
 	DECLARE_SCOPE_CYCLE_COUNTER(TEXT("FTexture2DStreamIn_IO_AsyncReallocate::Finalize"), STAT_Texture2DStreamInIOAsyncReallocate_Finalize, STATGROUP_StreamingDetails);
 	check(Context.CurrentThread == TT_Render);
 
+	FRHICommandListImmediate& RHICmdList = FRHICommandListImmediate::Get();
+
 	ClearIORequests(Context);
-	DoUnlockNewMips(Context);
-	DoFinishUpdate(Context);
+	DoUnlockNewMips(RHICmdList, Context);
+	DoFinishUpdate(RHICmdList, Context);
 }
 
 // ****************************
@@ -77,7 +83,9 @@ void FTexture2DStreamIn_IO_AsyncReallocate::Cancel(const FContext& Context)
 	DECLARE_SCOPE_CYCLE_COUNTER(TEXT("FTexture2DStreamIn_IO_AsyncReallocate::Cancel"), STAT_Texture2DStreamInIOAsyncReallocate_Cancel, STATGROUP_StreamingDetails);
 	check(Context.CurrentThread == TT_Render);
 
-	DoUnlockNewMips(Context);
-	DoFinishUpdate(Context);
+	FRHICommandListImmediate& RHICmdList = FRHICommandListImmediate::Get();
+
+	DoUnlockNewMips(RHICmdList, Context);
+	DoFinishUpdate(RHICmdList, Context);
 	ReportIOError(Context);
 }
diff --git a/Engine/Source/Runtime/Engine/Private/Streaming/Texture2DStreamIn_IO_Virtual.cpp b/Engine/Source/Runtime/Engine/Private/Streaming/Texture2DStreamIn_IO_Virtual.cpp
--- a/Engine/Source/Runtime/Engine/Private/Streaming/Texture2DStreamIn_IO_Virtual.cpp
+++ b/Engine/Source/Runtime/Engine/Private/Streaming/Texture2DStreamIn_IO_Virtual.cpp
@@ -19,8 +19,10 @@ void FTexture2DStreamIn_IO_Virtual::LockMips(const FContext& Context)
 	DECLARE_SCOPE_CYCLE_COUNTER(TEXT("FTexture2DStreamIn_IO_Virtual::LockMips"), STAT_Texture2DStreamInIOVirtual_LockMips, STATGROUP_StreamingDetails);
 	check(Context.CurrentThread == TT_Render);
 
-	DoConvertToVirtualWithNewMips(Context);
-	DoLockNewMips(Context);
+	FRHICommandListImmediate& RHICmdList = FRHICommandListImmediate::Get();
+
+	DoConvertToVirtualWithNewMips(RHICmdList, Context);
+	DoLockNewMips(RHICmdList, Context);
 
 	PushTask(Context, TT_Async, SRA_UPDATE_CALLBACK(LoadMips), TT_Render, SRA_UPDATE_CALLBACK(Cancel));
 }
@@ -51,12 +53,14 @@ void FTexture2DStreamIn_IO_Virtual::Finalize(const FContext& Context)
 	DECLARE_SCOPE_CYCLE_COUNTER(TEXT("FTexture2DStreamIn_IO_Virtual::Finalize"), STAT_Texture2DStreamInIOVirtual_Finalize, STATGROUP_StreamingDetails);
 	check(Context.CurrentThread == TT_Render);
 
-	DoUnlockNewMips(Context);
+	FRHICommandListImmediate& RHICmdList = FRHICommandListImmediate::Get();
+
+	DoUnlockNewMips(RHICmdList, Context);
 	if (IntermediateTextureRHI)
 	{
-		RHIVirtualTextureSetFirstMipVisible(IntermediateTextureRHI, PendingFirstLODIdx);
+		RHICmdList.VirtualTextureSetFirstMipVisible(IntermediateTextureRHI, PendingFirstLODIdx);
 	}
-	DoFinishUpdate(Context);
+	DoFinishUpdate(RHICmdList, Context);
 }
 
 // ****************************
@@ -78,11 +82,13 @@ void FTexture2DStreamIn_IO_Virtual::Cancel(const FContext& Context)
 	DECLARE_SCOPE_CYCLE_COUNTER(TEXT("FTexture2DStreamIn_IO_Virtual::Cancel"), STAT_Texture2DStreamInIOVirtual_Cancel, STATGROUP_StreamingDetails);
 	check(Context.CurrentThread == TT_Render);
 
-	DoUnlockNewMips(Context);
+	FRHICommandListImmediate& RHICmdList = FRHICommandListImmediate::Get();
+
+	DoUnlockNewMips(RHICmdList, Context);
 	if (IntermediateTextureRHI)
 	{
-		RHIVirtualTextureSetFirstMipInMemory(IntermediateTextureRHI, CurrentFirstLODIdx);
+		RHICmdList.VirtualTextureSetFirstMipInMemory(IntermediateTextureRHI, CurrentFirstLODIdx);
 	}
-	DoFinishUpdate(Context);
+	DoFinishUpdate(RHICmdList, Context);
 	ReportIOError(Context);
 }
diff --git a/Engine/Source/Runtime/Engine/Private/Streaming/Texture2DStreamOut_AsyncCreate.cpp b/Engine/Source/Runtime/Engine/Private/Streaming/Texture2DStreamOut_AsyncCreate.cpp
--- a/Engine/Source/Runtime/Engine/Private/Streaming/Texture2DStreamOut_AsyncCreate.cpp
+++ b/Engine/Source/Runtime/Engine/Private/Streaming/Texture2DStreamOut_AsyncCreate.cpp
@@ -68,15 +68,17 @@ void FTexture2DStreamOut_AsyncCreate::Finalize(const FContext& Context)
 	DECLARE_SCOPE_CYCLE_COUNTER(TEXT("FTexture2DStreamOut_AsyncReallocate::Finalize"), STAT_Texture2DStreamOutAsyncCreate_Finalize, STATGROUP_StreamingDetails);
 	check(Context.CurrentThread == TT_Render);
 
+	FRHICommandListImmediate& RHICmdList = FRHICommandListImmediate::Get();
+
 	if (!IsCancelled() && IntermediateTextureRHI && Context.Resource)
 	{
 		UE::RHI::CopySharedMips_AssumeSRVMaskState(
-			FRHICommandListExecutor::GetImmediateCommandList(),
+			RHICmdList,
 			Context.Resource->GetTexture2DRHI(),
 			IntermediateTextureRHI);
 	}
 
-	DoFinishUpdate(Context);
+	DoFinishUpdate(RHICmdList, Context);
 }
 
 // ****************************
@@ -89,5 +91,7 @@ void FTexture2DStreamOut_AsyncCreate::Cancel(const FContext& Context)
 	DECLARE_SCOPE_CYCLE_COUNTER(TEXT("FTexture2DStreamOut_AsyncCreate::Cancel"), STAT_Texture2DStreamOutAsyncCreate_Cancel, STATGROUP_StreamingDetails);
 	check(Context.CurrentThread == TT_Render);
 
-	DoFinishUpdate(Context);
+	FRHICommandListImmediate& RHICmdList = FRHICommandListImmediate::Get();
+
+	DoFinishUpdate(RHICmdList, Context);
 }
diff --git a/Engine/Source/Runtime/Engine/Private/Streaming/Texture2DStreamOut_AsyncReallocate.cpp b/Engine/Source/Runtime/Engine/Private/Streaming/Texture2DStreamOut_AsyncReallocate.cpp
--- a/Engine/Source/Runtime/Engine/Private/Streaming/Texture2DStreamOut_AsyncReallocate.cpp
+++ b/Engine/Source/Runtime/Engine/Private/Streaming/Texture2DStreamOut_AsyncReallocate.cpp
@@ -27,7 +27,9 @@ void FTexture2DStreamOut_AsyncReallocate::AsyncReallocate(const FContext& Contex
 	DECLARE_SCOPE_CYCLE_COUNTER(TEXT("FTexture2DStreamOut_AsyncReallocate::AsyncReallocate"), STAT_Texture2DStreamOutAsyncReallocate_AsyncReallocate, STATGROUP_StreamingDetails);
 	check(Context.CurrentThread == TT_Render);
 
-	DoAsyncReallocate(Context);
+	FRHICommandListImmediate& RHICmdList = FRHICommandListImmediate::Get();
+
+	DoAsyncReallocate(RHICmdList, Context);
 	
 	PushTask(Context, TT_Render, SRA_UPDATE_CALLBACK(Finalize), TT_Render, SRA_UPDATE_CALLBACK(Cancel));
 }
@@ -37,7 +39,9 @@ void FTexture2DStreamOut_AsyncReallocate::Finalize(const FContext& Context)
 	DECLARE_SCOPE_CYCLE_COUNTER(TEXT("FTexture2DStreamOut_AsyncReallocate::Finalize"), STAT_Texture2DStreamOutAsyncReallocate_Finalize, STATGROUP_StreamingDetails);
 	check(Context.CurrentThread == TT_Render);
 
-	DoFinishUpdate(Context);
+	FRHICommandListImmediate& RHICmdList = FRHICommandListImmediate::Get();
+
+	DoFinishUpdate(RHICmdList, Context);
 }
 
 // ****************************
@@ -50,5 +54,7 @@ void FTexture2DStreamOut_AsyncReallocate::Cancel(const FContext& Context)
 	DECLARE_SCOPE_CYCLE_COUNTER(TEXT("FTexture2DStreamOut_AsyncReallocate::Cancel"), STAT_Texture2DStreamOutAsyncReallocate_Cancel, STATGROUP_StreamingDetails);
 	check(Context.CurrentThread == TT_Render);
 
-	DoFinishUpdate(Context);
+	FRHICommandListImmediate& RHICmdList = FRHICommandListImmediate::Get();
+
+	DoFinishUpdate(RHICmdList, Context);
 }
diff --git a/Engine/Source/Runtime/Engine/Private/Streaming/Texture2DStreamOut_Virtual.cpp b/Engine/Source/Runtime/Engine/Private/Streaming/Texture2DStreamOut_Virtual.cpp
--- a/Engine/Source/Runtime/Engine/Private/Streaming/Texture2DStreamOut_Virtual.cpp
+++ b/Engine/Source/Runtime/Engine/Private/Streaming/Texture2DStreamOut_Virtual.cpp
@@ -28,20 +28,22 @@ void FTexture2DStreamOut_Virtual::Finalize(const FContext& Context)
 	DECLARE_SCOPE_CYCLE_COUNTER(TEXT("FTexture2DStreamOut_Virtual::Finalize"), STAT_Texture2DStreamOutVirtual_Finalize, STATGROUP_StreamingDetails);
 	check(Context.CurrentThread == TT_Render);
 
+	FRHICommandListImmediate& RHICmdList = FRHICommandListImmediate::Get();
+
 	static TConsoleVariableData<int32>* CVarReducedMode = IConsoleManager::Get().FindTConsoleVariableDataInt(TEXT("r.PRT.ReducedMemory"));
 	check(CVarReducedMode);
 
 	if (!CVarReducedMode->GetValueOnRenderThread() || ResourceState.NumRequestedLODs > ResourceState.NumNonStreamingLODs)
 	{
 		IntermediateTextureRHI = Context.Resource->GetTexture2DRHI();
-		RHIVirtualTextureSetFirstMipVisible(IntermediateTextureRHI, PendingFirstLODIdx);
-		RHIVirtualTextureSetFirstMipInMemory(IntermediateTextureRHI, PendingFirstLODIdx);
+		RHICmdList.VirtualTextureSetFirstMipVisible(IntermediateTextureRHI, PendingFirstLODIdx);
+		RHICmdList.VirtualTextureSetFirstMipInMemory(IntermediateTextureRHI, PendingFirstLODIdx);
 	}
 	else
 	{
-		 DoConvertToNonVirtual(Context);
+		 DoConvertToNonVirtual(RHICmdList, Context);
 	}
-	DoFinishUpdate(Context);
+	DoFinishUpdate(RHICmdList, Context);
 }
 
 // ****************************
@@ -53,5 +55,7 @@ void FTexture2DStreamOut_Virtual::Cancel(const FContext& Context)
 	DECLARE_SCOPE_CYCLE_COUNTER(TEXT("FTexture2DStreamOut_Virtual::Cancel"), STAT_Texture2DStreamOutVirtual_Cancel, STATGROUP_StreamingDetails);
 	check(Context.CurrentThread == TT_Render);
 
-	DoFinishUpdate(Context);
+	FRHICommandListImmediate& RHICmdList = FRHICommandListImmediate::Get();
+
+	DoFinishUpdate(RHICmdList, Context);
 }
diff --git a/Engine/Source/Runtime/Engine/Private/Streaming/Texture2DUpdate.cpp b/Engine/Source/Runtime/Engine/Private/Streaming/Texture2DUpdate.cpp
--- a/Engine/Source/Runtime/Engine/Private/Streaming/Texture2DUpdate.cpp
+++ b/Engine/Source/Runtime/Engine/Private/Streaming/Texture2DUpdate.cpp
@@ -57,7 +57,7 @@ FTexture2DUpdate::~FTexture2DUpdate()
 // ********* Helpers **********
 // ****************************
 
-void FTexture2DUpdate::DoAsyncReallocate(const FContext& Context)
+void FTexture2DUpdate::DoAsyncReallocate(FRHICommandListImmediate& RHICmdList, const FContext& Context)
 {
 	check(Context.CurrentThread == TT_Render);
 
@@ -69,7 +69,7 @@ void FTexture2DUpdate::DoAsyncReallocate(const FContext& Context)
 
 		ensure(!IntermediateTextureRHI);
 
-		IntermediateTextureRHI = RHIAsyncReallocateTexture2D(
+		IntermediateTextureRHI = RHICmdList.AsyncReallocateTexture2D(
 			Context.Resource->GetTexture2DRHI(),
 			ResourceState.NumRequestedLODs,
 			RequestedMipMap.SizeX,
@@ -80,14 +80,12 @@ void FTexture2DUpdate::DoAsyncReallocate(const FContext& Context)
 
 
 //  Transform the texture into a virtual texture.
-void FTexture2DUpdate::DoConvertToVirtualWithNewMips(const FContext& Context)
+void FTexture2DUpdate::DoConvertToVirtualWithNewMips(FRHICommandListImmediate& RHICmdList, const FContext& Context)
 {
 	check(Context.CurrentThread == TT_Render);
 
 	if (!IsCancelled() && Context.Resource)
 	{
-		FRHICommandListImmediate& RHICmdList = FRHICommandListImmediate::Get();
-
 		// If the texture is not virtual, then make it virtual immediately.
 		if (!Context.Resource->IsTextureRHIPartiallyResident())
 		{
@@ -127,7 +125,7 @@ void FTexture2DUpdate::DoConvertToVirtualWithNewMips(const FContext& Context)
 	}
 }
 
-bool FTexture2DUpdate::DoConvertToNonVirtual(const FContext& Context)
+bool FTexture2DUpdate::DoConvertToNonVirtual(FRHICommandListImmediate& RHICmdList, const FContext& Context)
 {
 	check(Context.CurrentThread == TT_Render);
 
@@ -136,8 +134,6 @@ bool FTexture2DUpdate::DoConvertToNonVirtual(const FContext& Context)
 	{
 		if (Context.Resource->IsTextureRHIPartiallyResident())
 		{
-			FRHICommandListImmediate& RHICmdList = FRHICommandListImmediate::Get();
-
 			const FTexture2DMipMap& PendingFirstMipMap = *Context.MipsView[PendingFirstLODIdx];
 
 			ensure(!IntermediateTextureRHI);
@@ -167,7 +163,7 @@ bool FTexture2DUpdate::DoConvertToNonVirtual(const FContext& Context)
 	return false;
 }
 
-void FTexture2DUpdate::DoFinishUpdate(const FContext& Context)
+void FTexture2DUpdate::DoFinishUpdate(FRHICommandListImmediate& RHICmdList, const FContext& Context)
 {
 	check(Context.CurrentThread == TT_Render);
 
@@ -175,7 +171,7 @@ void FTexture2DUpdate::DoFinishUpdate(const FContext& Context)
 	{
 		if (!IsCancelled())
 		{
-			Context.Resource->FinalizeStreaming(FRHICommandListImmediate::Get(), IntermediateTextureRHI);
+			Context.Resource->FinalizeStreaming(RHICmdList, IntermediateTextureRHI);
 			MarkAsSuccessfullyFinished();
 		}
 		IntermediateTextureRHI.SafeRelease();
diff --git a/Engine/Source/Runtime/Engine/Private/Streaming/Texture2DUpdate.h b/Engine/Source/Runtime/Engine/Private/Streaming/Texture2DUpdate.h
--- a/Engine/Source/Runtime/Engine/Private/Streaming/Texture2DUpdate.h
+++ b/Engine/Source/Runtime/Engine/Private/Streaming/Texture2DUpdate.h
@@ -64,13 +64,13 @@ protected:
 	// ****************************
 
 	/** Async Reallocate the texture to the requested size. */
-	void DoAsyncReallocate(const FContext& Context);
+	void DoAsyncReallocate(FRHICommandListImmediate& RHICmdList, const FContext& Context);
 	/** Transform the texture into a virtual texture. The virtual texture. */
-	void DoConvertToVirtualWithNewMips(const FContext& Context);
+	void DoConvertToVirtualWithNewMips(FRHICommandListImmediate& RHICmdList, const FContext& Context);
 	/** Transform the texture into a non virtual texture. The new texture will have the size of the requested mips to avoid later reallocs. */
-	bool DoConvertToNonVirtual(const FContext& Context);
+	bool DoConvertToNonVirtual(FRHICommandListImmediate& RHICmdList, const FContext& Context);
 	/** Apply the new texture (if not cancelled) and finish the update process. When cancelled, the intermediate texture simply gets discarded. */
-	void DoFinishUpdate(const FContext& Context);
+	void DoFinishUpdate(FRHICommandListImmediate& RHICmdList, const FContext& Context);
 
 	/** The intermediate texture created in the update process. In the virtual path, this can exceptionally end update being the same as the original texture. */
 	FTextureRHIRef IntermediateTextureRHI;
