Setup the TickContext with the proper TickGroup.  This change causes the FTickFunctionTask::DoTask and FBatchTickFunctionTask::DoTask functions' TickContext parameter to have the proper TickGroup set (instead of always being TG_PrePhysics).

diff --git a/Engine/Source/Runtime/Engine/Private/TickTaskManager.cpp b/Engine/Source/Runtime/Engine/Private/TickTaskManager.cpp
--- a/Engine/Source/Runtime/Engine/Private/TickTaskManager.cpp
+++ b/Engine/Source/Runtime/Engine/Private/TickTaskManager.cpp
@@ -627,6 +627,7 @@ public:
 		checkSlow(TickFunction->InternalData->ActualStartTickGroup >=0 && TickFunction->InternalData->ActualStartTickGroup < TG_MAX);
 
 		FTickContext UseContext = TickContext;
+		UseContext.TickGroup = TickFunction->GetActualTickGroup();
 		UseContext.bLogTick = bLogTicks;
 		UseContext.bLogTicksShowPrerequistes = bLogTicksShowPrerequistes;
 
