World Partition: make sure to trash non-OFPA packages when moving actors to streaming levels to guarantee that any potential future load of this actor won't find the original package where this actor was initially.

#FailedValidations Preflight

diff --git a/Engine/Source/Runtime/Engine/Private/WorldPartition/WorldPartitionLevelHelper.cpp b/Engine/Source/Runtime/Engine/Private/WorldPartition/WorldPartitionLevelHelper.cpp
--- a/Engine/Source/Runtime/Engine/Private/WorldPartition/WorldPartitionLevelHelper.cpp
+++ b/Engine/Source/Runtime/Engine/Private/WorldPartition/WorldPartitionLevelHelper.cpp
@@ -227,6 +227,7 @@ void FWorldPartitionLevelHelper::MoveExternalActorsToLevel(const TArray<FWorldPa
 
 	// Move all actors to Cell level
 	TSet<AActor*> LoadedActors;
+	TSet<UPackage*> ModifiedPackages;
 
 	for (const FWorldPartitionRuntimeCellObjectMapping& PackageObjectMapping : InChildPackages)
 	{
@@ -269,6 +270,8 @@ void FWorldPartitionLevelHelper::MoveExternalActorsToLevel(const TArray<FWorldPa
 					Actor->RerunConstructionScripts();
 					ApplyConstructionScriptPropertyOverridesFromAnnotation(Actor);
 				}
+
+				ModifiedPackages.Add(ActorPackage);
 			}
 			else if (!InLevel->Actors.Contains(Actor))
 			{
@@ -326,10 +329,6 @@ void FWorldPartitionLevelHelper::MoveExternalActorsToLevel(const TArray<FWorldPa
 						}
 					}
 				}
-
-				// Trash this package to guarantee that any potential future load of this actor won't find the old empty package
-				// @todo_ow: Decide if we want to support actor reloads during cook. If not, remove this code, detect the reload and report an error.
-				FLevelStreamingGCHelper::TrashPackage(ActorPackage);
 			}
 
 			OutModifiedPackages.Add(ActorPackage);
@@ -350,6 +349,12 @@ void FWorldPartitionLevelHelper::MoveExternalActorsToLevel(const TArray<FWorldPa
 			checkf(LevelActors.Contains(Actor->GetFName()), TEXT("Actor %s(%s) was unexpectedly loaded when moving actors to streaming cell"), *Actor->GetActorNameOrLabel(), *Actor->GetName());
 		}
 	}
+
+	// Trash the modified packages to guarantee that any potential future load of this actor won't find the old empty package
+	for (UPackage* ModifiedPackage : ModifiedPackages)
+	{
+		FLevelStreamingGCHelper::TrashPackage(ModifiedPackage);
+	}
 }
 
 void FWorldPartitionLevelHelper::RemapLevelSoftObjectPaths(ULevel* InLevel, UWorldPartition* InWorldPartition)
