Dstm: Add support for constraint component migration when the constrained components are in different actors.

diff --git a/Engine/Source/Runtime/CoreUObject/Private/UObject/RemoteObject.cpp b/Engine/Source/Runtime/CoreUObject/Private/UObject/RemoteObject.cpp
--- a/Engine/Source/Runtime/CoreUObject/Private/UObject/RemoteObject.cpp
+++ b/Engine/Source/Runtime/CoreUObject/Private/UObject/RemoteObject.cpp
@@ -650,14 +650,13 @@ void ChangeOwnerServerId(const UObject* Object, FRemoteServerId NewOwnerServerId
 #endif
 }
 
-FPhysicsIslandId GetPhysicsIslandId(const UObject* Object)
+FPhysicsIslandId GetPhysicsIslandId(const FRemoteObjectId ObjectId)
 {
 	using namespace UE::CoreUObject::Private;
 	using namespace UE::RemoteObject::Private;
-
 	FPhysicsIslandId Result;
 #if UE_WITH_REMOTE_OBJECT_HANDLE
-	UE::RemoteObject::Handle::FRemoteObjectStub* RemoteStub = FindRemoteObjectStub(FObjectHandleUtils::GetRemoteId(Object));
+	UE::RemoteObject::Handle::FRemoteObjectStub* RemoteStub = FindRemoteObjectStub(ObjectId);
 	if (RemoteStub)
 	{
 		Result = RemoteStub->PhysicsIslandId;
@@ -666,6 +665,17 @@ FPhysicsIslandId GetPhysicsIslandId(const UObject* Object)
 	return Result;
 }
 
+
+FPhysicsIslandId GetPhysicsIslandId(const UObject* Object)
+{
+	using namespace UE::CoreUObject::Private;
+	using namespace UE::RemoteObject::Private;
+#if UE_WITH_REMOTE_OBJECT_HANDLE
+	return GetPhysicsIslandId(FObjectHandleUtils::GetRemoteId(Object));
+#endif
+	return FPhysicsIslandId();
+}
+
 void ChangePhysicsIslandId(const UObject* Object, FPhysicsIslandId NewPhysicsIslandId)
 {
 	using namespace UE::CoreUObject::Private;
diff --git a/Engine/Source/Runtime/CoreUObject/Public/UObject/RemoteObject.h b/Engine/Source/Runtime/CoreUObject/Public/UObject/RemoteObject.h
--- a/Engine/Source/Runtime/CoreUObject/Public/UObject/RemoteObject.h
+++ b/Engine/Source/Runtime/CoreUObject/Public/UObject/RemoteObject.h
@@ -56,7 +56,7 @@ namespace UE::RemoteObject::Handle
 
 		bool IsValid() const
 		{
-			return PhysicsServerId.IsValid() && PhysicsLocalIslandId != INDEX_NONE;
+			return PhysicsServerId.IsValid() && PhysicsLocalIslandId != 0;
 		}
 
 		bool operator<(const FPhysicsIslandId& Other) const
@@ -234,6 +234,8 @@ namespace UE::RemoteObject::Handle
 	*/
 	COREUOBJECT_API FPhysicsIslandId GetPhysicsIslandId(const UObject* Object);
 
+	COREUOBJECT_API FPhysicsIslandId GetPhysicsIslandId(const FRemoteObjectId ObjectId);
+
 	COREUOBJECT_API void ChangePhysicsIslandId(const UObject* Object, FPhysicsIslandId NewPhysicsIslandId);
 
 	COREUOBJECT_API void ClearAllPhysicsServerId();
diff --git a/Engine/Source/Runtime/Engine/Classes/PhysicsEngine/PhysicsConstraintComponent.h b/Engine/Source/Runtime/Engine/Classes/PhysicsEngine/PhysicsConstraintComponent.h
--- a/Engine/Source/Runtime/Engine/Classes/PhysicsEngine/PhysicsConstraintComponent.h
+++ b/Engine/Source/Runtime/Engine/Classes/PhysicsEngine/PhysicsConstraintComponent.h
@@ -5,6 +5,7 @@
 #include "UObject/ObjectMacros.h"
 #include "Engine/EngineTypes.h"
 #include "Components/SceneComponent.h"
+#include "Physics/Experimental/PhysScene_Chaos.h"
 #include "PhysicsEngine/ConstraintInstance.h"
 #include "PhysicsConstraintComponent.generated.h"
 
@@ -93,11 +94,19 @@ public:
 	ENGINE_API virtual void PostMigrate(const struct FUObjectMigrationContext& MigrationContext) override;
 	ENGINE_API void ApplyMigratedPhysicsState(Chaos::FSerializedDataBufferPtr MigratedStateData = nullptr);
 	ENGINE_API bool HasValidPhysicsState() const;
+	ENGINE_API void GetConstrainedBodyIds(TPair<FRemoteServerId, FRemoteServerId>& OwnerServerIds, TPair<FRemoteObjectId, FRemoteObjectId>& ObjectIds) const;
+	ENGINE_API bool AreConstrainedComponentsInSameActor() const;
+	ENGINE_API void AddReInitDelegatesToComponents();
+	ENGINE_API void InitComponentConstraintWithoutUpdatingFrames();
+	ENGINE_API void CheckAndUpdateConstraintPT();
+	ENGINE_API bool GetIsBrokenLocal() const { return bIsBrokenLocal; };
 #endif
 
 	UFUNCTION()
 	ENGINE_API void ReInitConstraintAndApplyPhysicsState(UPrimitiveComponent* ChangedComponent, EComponentPhysicsStateChange StateChange);
 
+	ENGINE_API void ReInitConstraintAndApplyPhysicsState(FPhysScene* PhysScene, float DeltaTime);
+
 	//~ Begin SceneComponent Interface
 #if WITH_EDITOR
 	ENGINE_API virtual void PostEditComponentMove(bool bFinished) override;
@@ -439,7 +448,11 @@ private:
 	
 	ENGINE_API bool CanSendMigratedPhysicsState(const struct FUObjectMigrationContext* MigrationContext, EObjectMigrationSendType ObjMigrationSendType) const;
 
-	void CheckAndUnbindPhysicsStateDelegates();
+	void CheckAndUnbindPhysicsStateDelegates(bool bCheckIfBound = true);
+
+	void UnBindPreTickDelegate();
+
+	FDelegateHandle PreTickDelegateHandle;
 #endif
 };
 
diff --git a/Engine/Source/Runtime/Engine/Private/Components/PrimitiveComponent.cpp b/Engine/Source/Runtime/Engine/Private/Components/PrimitiveComponent.cpp
--- a/Engine/Source/Runtime/Engine/Private/Components/PrimitiveComponent.cpp
+++ b/Engine/Source/Runtime/Engine/Private/Components/PrimitiveComponent.cpp
@@ -1384,6 +1384,14 @@ void UPrimitiveComponent::MigratePhysicsState(FArchive& Ar)
 			Chaos::FSerializedDataBuffer EmptyPhysicsState;
 			Ar << EmptyPhysicsState;
 		}
+
+		if (ObjMigrationSendType == EObjectMigrationSendType::ReassignOwnership)
+		{
+			if (FPhysScene* PhysicsScene = FPhysicsInterface::GetCurrentSceneForComponent(this))
+			{
+				PhysicsScene->RemoveFromComponentOwnerMigratedSet(TObjectPtr<AActor>(GetOwner()).GetRemoteId());
+			}
+		}
 	}
 }
 
@@ -1451,6 +1459,14 @@ void UPrimitiveComponent::PostMigrate(const struct FUObjectMigrationContext& Mig
 			}
 		}
 
+		if (RecvType == EObjectMigrationRecvType::AssignedOwnership)
+		{
+			if (FPhysScene* PhysicsScene = FPhysicsInterface::GetCurrentSceneForComponent(this))
+			{
+				PhysicsScene->AddToComponentOwnerMigratedSet(TObjectPtr<AActor>(GetOwner()).GetRemoteId());
+			}
+		}
+
 		if (HasValidPhysicsState())
 		{
 			ApplyMigratedPhysicsState();
diff --git a/Engine/Source/Runtime/Engine/Private/PhysicsEngine/Experimental/PhysScene_Chaos.cpp b/Engine/Source/Runtime/Engine/Private/PhysicsEngine/Experimental/PhysScene_Chaos.cpp
--- a/Engine/Source/Runtime/Engine/Private/PhysicsEngine/Experimental/PhysScene_Chaos.cpp
+++ b/Engine/Source/Runtime/Engine/Private/PhysicsEngine/Experimental/PhysScene_Chaos.cpp
@@ -863,6 +863,31 @@ Chaos::FSerializedDataBufferPtr FPhysScene_Chaos::PopPendingSerializedStateForOb
 	return MoveTemp(FoundStateRef);
 }
 
+bool FPhysScene_Chaos::IsComponentMigrated(UPrimitiveComponent* PrimComp) const
+{
+#if UE_WITH_REMOTE_OBJECT_HANDLE
+	return MigratedPrimitiveComponentOwnerIds.Contains(TObjectPtr<AActor>(PrimComp->GetOwner()).GetRemoteId());
+#else
+	return false;
+#endif
+}
+
+void FPhysScene_Chaos::ResetComponentOwnerMigratedSet()
+{
+	MigratedPrimitiveComponentOwnerIds.Reset();
+}
+
+void FPhysScene_Chaos::AddToComponentOwnerMigratedSet(const FRemoteObjectId CompOwnerId)
+{
+	MigratedPrimitiveComponentOwnerIds.Add(CompOwnerId);
+}
+
+void FPhysScene_Chaos::RemoveFromComponentOwnerMigratedSet(const FRemoteObjectId CompOwnerId)
+{
+	MigratedPrimitiveComponentOwnerIds.Remove(CompOwnerId);
+}
+
+
 IPhysicsReplication* FPhysScene_Chaos::GetPhysicsReplication()
 {
 	return PhysicsReplication.Get();
diff --git a/Engine/Source/Runtime/Engine/Private/PhysicsEngine/PhysicsConstraintComponent.cpp b/Engine/Source/Runtime/Engine/Private/PhysicsEngine/PhysicsConstraintComponent.cpp
--- a/Engine/Source/Runtime/Engine/Private/PhysicsEngine/PhysicsConstraintComponent.cpp
+++ b/Engine/Source/Runtime/Engine/Private/PhysicsEngine/PhysicsConstraintComponent.cpp
@@ -10,6 +10,10 @@
 #include "PhysicsEngine/ConstraintUtils.h"
 #include "PhysicsEngine/PhysicsObjectExternalInterface.h"
 #include "PhysicsEngine/SkeletalBodySetup.h"
+#include "PhysicsProxy/JointConstraintProxy.h"
+#include "Chaos/PBDJointConstraintData.h"
+#include "Chaos/PhysicsObjectInterface.h"
+#include "PBDRigidsSolver.h"
 #include "PhysicsProxy/SingleParticlePhysicsProxy.h"
 #include "Components/BillboardComponent.h"
 #include "UObject/ICookInfo.h"
@@ -20,6 +24,7 @@
 #include "PhysicsInterfaceDeclaresCore.h"
 #include "Physics/Experimental/PhysScene_Chaos.h"
 #include "Chaos/Serialization/SerializedMultiPhysicsState.h"
+#include "UObject/UObjectMigrationContext.h"
 #endif
 
 #include UE_INLINE_GENERATED_CPP_BY_NAME(PhysicsConstraintComponent)
@@ -54,10 +59,24 @@ UPrimitiveComponent* UPhysicsConstraintComponent::GetComponentInternal(EConstrai
 		// Use override component if specified
 		if(OverrideComponent1.IsValid())
 		{
+#if UE_WITH_REMOTE_OBJECT_HANDLE
+			if (OverrideComponent1.IsRemote() && !AutoRTFM::IsClosed())
+			{
+				ensureMsgf(false, TEXT("UPhysicsConstraintComponent::GetComponentInternal: Frame1 Component is remote. Simulation result is not guaranteed."));
+				return NULL;
+			}
+#endif
 			return OverrideComponent1.Get();
 		}
 
 		ComponentName = ComponentName1.ComponentName;
+#if UE_WITH_REMOTE_OBJECT_HANDLE
+		if (ConstraintActor1.IsRemote() && !AutoRTFM::IsClosed())
+		{
+			ensureMsgf(false, TEXT("UPhysicsConstraintComponent::GetComponentInternal: Frame1 Actor is remote. Simulation result is not guaranteed."));
+			return NULL;
+		}
+#endif
 		Actor = ConstraintActor1;
 	}
 	// Frame 2
@@ -66,10 +85,24 @@ UPrimitiveComponent* UPhysicsConstraintComponent::GetComponentInternal(EConstrai
 		// Use override component if specified
 		if(OverrideComponent2.IsValid())
 		{
+#if UE_WITH_REMOTE_OBJECT_HANDLE
+			if (OverrideComponent2.IsRemote() && !AutoRTFM::IsClosed())
+			{
+				ensureMsgf(false, TEXT("UPhysicsConstraintComponent::GetComponentInternal: Frame2 Component is remote. Simulation result is not guaranteed."));
+				return NULL;
+			}
+#endif
 			return OverrideComponent2.Get();
 		}
 
 		ComponentName = ComponentName2.ComponentName;
+#if UE_WITH_REMOTE_OBJECT_HANDLE
+		if (ConstraintActor2.IsRemote() && !AutoRTFM::IsClosed())
+		{
+			ensureMsgf(false, TEXT("UPhysicsConstraintComponent::GetComponentInternal: Frame2 Actor is remote. Simulation result is not guaranteed."));
+			return NULL;
+		}
+#endif
 		Actor = ConstraintActor2;
 	}
 
@@ -315,6 +348,28 @@ void UPhysicsConstraintComponent::InitComponentConstraint()
 	bIsBrokenLocal = false;
 }
 
+#if UE_WITH_REMOTE_OBJECT_HANDLE
+
+void UPhysicsConstraintComponent::InitComponentConstraintWithoutUpdatingFrames()
+{
+	// Only reinitialize constraint on migration when the constraint component is not locally broken 
+	if (!bIsBrokenLocal)
+	{
+		FBodyInstance* Body1 = GetBodyInstance(EConstraintFrame::Frame1);
+		FBodyInstance* Body2 = GetBodyInstance(EConstraintFrame::Frame2);
+
+		Chaos::FPhysicsObject* Object1 = (Body1 && Body1->IsValidBodyInstance()) ? Body1->GetPhysicsActor()->GetPhysicsObject() : GetPhysicsObject(EConstraintFrame::Frame1);
+		Chaos::FPhysicsObject* Object2 = (Body2 && Body2->IsValidBodyInstance()) ? Body2->GetPhysicsActor()->GetPhysicsObject() : GetPhysicsObject(EConstraintFrame::Frame2);
+
+		if (Object1 || Object2)
+		{
+			ConstraintInstance.InitConstraint(Object1, Object2, GetConstraintScale(), this, FOnConstraintBroken::CreateUObject(this, &UPhysicsConstraintComponent::OnConstraintBrokenWrapper));
+		}
+	}
+}
+
+#endif
+
 void UPhysicsConstraintComponent::TermComponentConstraint()
 {
 	ConstraintInstance.TermConstraint();
@@ -409,6 +464,7 @@ void UPhysicsConstraintComponent::OnUnregister()
 	TermComponentConstraint();
 #if UE_WITH_REMOTE_OBJECT_HANDLE
 	CheckAndUnbindPhysicsStateDelegates();
+	UnBindPreTickDelegate();
 #endif
 }
 
@@ -701,13 +757,18 @@ void UPhysicsConstraintComponent::Serialize(FArchive& Ar)
 	{
 		Ar << OverrideComponent1;
 		Ar << OverrideComponent2;
-
-		MigratePhysicsState(Ar);
+		Ar << bIsBrokenLocal;
+		if (!bIsBrokenLocal)
+		{
+			MigratePhysicsState(Ar);
+		}
 	}
 
+	// Clear out all bound delegates if the object is being migrated away. 
 	if (Ar.IsSaving())
 	{
 		CheckAndUnbindPhysicsStateDelegates();
+		UnBindPreTickDelegate();
 	}
 #endif
 }
@@ -773,7 +834,7 @@ void UPhysicsConstraintComponent::MigratePhysicsState(FArchive& Ar)
 
 			if (RawMigratedData)
 			{
-				UE_LOG(LogPhysicsConstraints, Warning, TEXT("Sending a pending physics state | [%s]. It is sent away before processed. | Migration Context [%s]."), *GetPathNameSafe(this), ToString(ObjMigrationSendType));
+				UE_LOG(LogPhysicsConstraints, Verbose, TEXT("Sending a pending physics state | [%s]. It is sent away before processed. | Migration Context [%s]."), *GetPathNameSafe(this), ToString(ObjMigrationSendType));
 				Ar << *RawMigratedData;
 			}
 			else
@@ -786,6 +847,93 @@ void UPhysicsConstraintComponent::MigratePhysicsState(FArchive& Ar)
 	}
 }
 
+bool UPhysicsConstraintComponent::AreConstrainedComponentsInSameActor() const
+{
+	if (OverrideComponent1.IsRemote() || OverrideComponent2.IsRemote())
+	{
+		return false;
+	}
+
+	if (UPrimitiveComponent* Comp1 = OverrideComponent1.Get())
+	{
+		if (Comp1 == NULL || Comp1->GetOwner() != GetOwner())
+		{
+			return false;
+		}
+	}
+
+	if (UPrimitiveComponent* Comp2 = OverrideComponent2.Get())
+	{
+		if (Comp2 == NULL || Comp2->GetOwner() != GetOwner())
+		{
+			return false;
+		}
+	}
+
+	if (ConstraintActor1 != NULL)
+	{
+		if (ConstraintActor1.IsRemote())
+		{
+			return false;
+		}
+		if (ConstraintActor1.Get() != GetOwner())
+		{
+			return false;
+		}
+	}
+
+	if (ConstraintActor2 != NULL)
+	{
+		if (ConstraintActor2.IsRemote())
+		{
+			return false;
+		}
+		if (ConstraintActor2.Get() != GetOwner())
+		{
+			return false;
+		}
+	}
+
+	return true;
+}
+
+void UPhysicsConstraintComponent::CheckAndUpdateConstraintPT()
+{
+	if (HasValidPhysicsState())
+	{
+		UE_LOG(LogPhysicsConstraints, Verbose, TEXT("UPhysicsConstraintComponent::CheckAndUpdateConstraintPT(): Constraint [%s] has a valid physics state."), *GetPathNameSafe(this));
+
+		bool bDependentBodiesAreNotMigrated = true;
+		bool bIsComp1Migrated = false;
+		bool bIsComp2Migrated = false;
+
+		if (UWorld* World = GetWorld())
+		{
+			if (FPhysScene* PhysScene = World->GetPhysicsScene())
+			{
+				if (UPrimitiveComponent* Comp1 = GetComponentInternal(EConstraintFrame::Frame1))
+				{
+					bIsComp1Migrated = PhysScene->IsComponentMigrated(Comp1);
+				}
+				if (UPrimitiveComponent* Comp2 = GetComponentInternal(EConstraintFrame::Frame2))
+				{
+					bIsComp2Migrated = PhysScene->IsComponentMigrated(Comp2);
+				}
+			}
+		}
+
+		if (bIsComp1Migrated || bIsComp2Migrated)
+		{
+			UE_LOG(LogPhysicsConstraints, Verbose, TEXT("UPhysicsConstraintComponent::CheckAndUpdateConstraintPT(): Constraint [%s] needs to be reinitialized because its PT is no longer valid."), *GetPathNameSafe(this));
+			InitComponentConstraintWithoutUpdatingFrames();
+		}
+		else
+		{
+			UE_LOG(LogPhysicsConstraints, Verbose, TEXT("UPhysicsConstraintComponent::CheckAndUpdateConstraintPT(): Constraint [%s] does not need to be reinitialized."), *GetPathNameSafe(this));
+		}
+	}
+}
+
 bool UPhysicsConstraintComponent::CanReInitConstraintAndApplyPhysicsState() const
 {
 	FBodyInstance* Body1 = GetBodyInstance(EConstraintFrame::Frame1);
@@ -805,32 +953,67 @@ void UPhysicsConstraintComponent::PostMigrate(const struct FUObjectMigrationCont
 	{
 		EObjectMigrationRecvType RecvType = MigrationContext.GetObjectMigrationRecvType(this);
 
-		if (HasValidPhysicsState())
+		if (!bIsBrokenLocal)
 		{
-			ApplyMigratedPhysicsState();
-		}
-		else if (IsValid(this))
-		{
-			if (CanReInitConstraintAndApplyPhysicsState())
+			if (HasValidPhysicsState())
 			{
-				InitComponentConstraint();
 				ApplyMigratedPhysicsState();
 			}
-			else
+			else if (IsValid(this))
 			{
-				if (UPrimitiveComponent* Comp1 = GetComponentInternal(EConstraintFrame::Frame1))
+				if (AreConstrainedComponentsInSameActor())
 				{
-					Comp1->OnComponentPhysicsStateChanged.AddUniqueDynamic(this, &UPhysicsConstraintComponent::ReInitConstraintAndApplyPhysicsState);
+					if (CanReInitConstraintAndApplyPhysicsState())
+					{
+						InitComponentConstraintWithoutUpdatingFrames();
+						ApplyMigratedPhysicsState();
+					}
+					else
+					{
+						AddReInitDelegatesToComponents();
+					}
 				}
-				if (UPrimitiveComponent* Comp2 = GetComponentInternal(EConstraintFrame::Frame2))
+				else
 				{
-					Comp2->OnComponentPhysicsStateChanged.AddUniqueDynamic(this, &UPhysicsConstraintComponent::ReInitConstraintAndApplyPhysicsState);
+					if (UWorld* World = GetWorld())
+					{
+						if (FPhysScene* PhysScene = World->GetPhysicsScene())
+						{
+							PreTickDelegateHandle = PhysScene->OnPhysScenePreTick.AddUObject(this, &UPhysicsConstraintComponent::ReInitConstraintAndApplyPhysicsState);
+						}
+					}
 				}
 			}
 		}
 	}
 }
 
+void UPhysicsConstraintComponent::ReInitConstraintAndApplyPhysicsState(FPhysScene* PhysScene, float DeltaTime)
+{
+#if UE_WITH_REMOTE_OBJECT_HANDLE
+	if (CanReInitConstraintAndApplyPhysicsState())
+	{
+		UE_LOG(LogPhysicsConstraints, Verbose, TEXT("UPhysicsConstraintComponent::ReInitConstraintAndApplyPhysicsState: Reinitializing constraint and apply migrated physics state "));
+		InitComponentConstraintWithoutUpdatingFrames();
+		ApplyMigratedPhysicsState();
+		UnBindPreTickDelegate();
+	}
+#endif
+}
+
+
+void UPhysicsConstraintComponent::AddReInitDelegatesToComponents()
+{
+	if (UPrimitiveComponent* Comp1 = GetComponentInternal(EConstraintFrame::Frame1))
+	{
+		Comp1->OnComponentPhysicsStateChanged.AddUniqueDynamic(this, &UPhysicsConstraintComponent::ReInitConstraintAndApplyPhysicsState);
+	}
+	if (UPrimitiveComponent* Comp2 = GetComponentInternal(EConstraintFrame::Frame2))
+	{
+		Comp2->OnComponentPhysicsStateChanged.AddUniqueDynamic(this, &UPhysicsConstraintComponent::ReInitConstraintAndApplyPhysicsState);
+	}
+}
+
 void UPhysicsConstraintComponent::ApplyMigratedPhysicsState(Chaos::FSerializedDataBufferPtr MigratedStateData)
 {
 	Chaos::FSerializedDataBufferPtr RawMigratedData;
@@ -887,26 +1070,65 @@ bool UPhysicsConstraintComponent::HasValidPhysicsState() const
 	return ConstraintInstance.IsValidConstraintInstance();
 }
 
-void UPhysicsConstraintComponent::CheckAndUnbindPhysicsStateDelegates()
+void UPhysicsConstraintComponent::CheckAndUnbindPhysicsStateDelegates(bool bCheckIfBound)
 {
 	TFunction<void(UPrimitiveComponent*, const FString&)> CheckAndUnbind =
-		[this](UPrimitiveComponent* Comp, const FString& FrameLabel)
+		[this, bCheckIfBound](UPrimitiveComponent* Comp, const FString& FrameLabel)
 		{
 			if (Comp)
 			{
 				if (Comp->OnComponentPhysicsStateChanged.IsAlreadyBound(this, &UPhysicsConstraintComponent::ReInitConstraintAndApplyPhysicsState))
 				{
-					UE_LOG(LogPhysics, Warning,
-						TEXT("UPhysicsConstraintComponent::CheckAndUnbindPhysicsStateDelegates: PhysicsConstraintComponent [%s] still had physics state delegate bound for %s. Cleaning up."),
-						*GetNameSafe(this), *FrameLabel);
+					if (bCheckIfBound)
+					{
+						UE_LOG(LogPhysics, Warning,
+							TEXT("UPhysicsConstraintComponent::CheckAndUnbindPhysicsStateDelegates: PhysicsConstraintComponent [%s] still had physics state delegate bound for %s. Cleaning up."),
+							*GetNameSafe(this), *FrameLabel);
+					}
 
 					Comp->OnComponentPhysicsStateChanged.RemoveDynamic(this, &UPhysicsConstraintComponent::ReInitConstraintAndApplyPhysicsState);
 				}
 			}
 		};
+	// This function is called inside Serialize as well, which is inside the migration code
+	// We shouldn't borrow when migrating objects, so the check will skip the unbind if the dependent bodies may need to be borrowed. 
+	if (AreConstrainedComponentsInSameActor())
+	{
+		CheckAndUnbind(GetComponentInternal(EConstraintFrame::Frame1), TEXT("Frame1"));
+		CheckAndUnbind(GetComponentInternal(EConstraintFrame::Frame2), TEXT("Frame2"));
+	}
+}
 
-	CheckAndUnbind(GetComponentInternal(EConstraintFrame::Frame1), TEXT("Frame1"));
-	CheckAndUnbind(GetComponentInternal(EConstraintFrame::Frame2), TEXT("Frame2"));
+void UPhysicsConstraintComponent::UnBindPreTickDelegate()
+{
+	if (PreTickDelegateHandle.IsValid())
+	{
+		if (const UWorld* World = GetWorld())
+		{
+			if (FPhysScene* PhysScene = World->GetPhysicsScene())
+			{
+				PhysScene->OnPhysScenePreTick.Remove(PreTickDelegateHandle);
+				PreTickDelegateHandle.Reset();
+			}
+		}
+	}
+}
+
+void UPhysicsConstraintComponent::GetConstrainedBodyIds(TPair<FRemoteServerId, FRemoteServerId>& OwnerServerIds, TPair<FRemoteObjectId, FRemoteObjectId>& ObjectIds) const
+{
+	UPrimitiveComponent* Comp1 = GetComponentInternal(EConstraintFrame::Frame1);
+	UPrimitiveComponent* Comp2 = GetComponentInternal(EConstraintFrame::Frame2);
+	if (Comp1 != NULL)
+	{
+		OwnerServerIds.Key = UE::RemoteObject::Handle::GetOwnerServerId(Comp1);
+		ObjectIds.Key = TObjectPtr<AActor>(Comp1->GetOwner()).GetRemoteId();
+	}
+
+	if (Comp2 != NULL)
+	{
+		OwnerServerIds.Value = UE::RemoteObject::Handle::GetOwnerServerId(Comp2);
+		ObjectIds.Value = TObjectPtr<AActor>(Comp2->GetOwner()).GetRemoteId();
+	}
 }
 #endif
 
@@ -915,9 +1137,9 @@ void UPhysicsConstraintComponent::ReInitConstraintAndApplyPhysicsState(UPrimitiv
 #if UE_WITH_REMOTE_OBJECT_HANDLE
 	if (StateChange == EComponentPhysicsStateChange::Created && CanReInitConstraintAndApplyPhysicsState())
 	{
-		InitComponentConstraint();
+		InitComponentConstraintWithoutUpdatingFrames();
 		ApplyMigratedPhysicsState();
-		CheckAndUnbindPhysicsStateDelegates();
+		CheckAndUnbindPhysicsStateDelegates(false);
 	}
 #endif
 }
diff --git a/Engine/Source/Runtime/Engine/Public/Physics/Experimental/PhysScene_Chaos.h b/Engine/Source/Runtime/Engine/Public/Physics/Experimental/PhysScene_Chaos.h
--- a/Engine/Source/Runtime/Engine/Public/Physics/Experimental/PhysScene_Chaos.h
+++ b/Engine/Source/Runtime/Engine/Public/Physics/Experimental/PhysScene_Chaos.h
@@ -172,6 +172,11 @@ public:
 	void PushPendingSerializedStateForObject(UObject* InComponent, Chaos::FSerializedDataBufferPtr&& InState);
 	Chaos::FSerializedDataBufferPtr PopPendingSerializedStateForObject(UObject* InComponent);
 
+	bool IsComponentMigrated(UPrimitiveComponent* PrimComp) const;
+	void ResetComponentOwnerMigratedSet();
+	void AddToComponentOwnerMigratedSet(const FRemoteObjectId CompOwnerId);
+	void RemoveFromComponentOwnerMigratedSet(const FRemoteObjectId CompOwnerId);
+
 	ENGINE_API IPhysicsReplication* GetPhysicsReplication();
 
 	ENGINE_API IPhysicsReplication* CreatePhysicsReplication();
@@ -358,6 +363,8 @@ private:
 	
 	TMap<UObject*, Chaos::FSerializedDataBufferPtr> PendingMigratedPhysicsStateByComponent;
 
+	TSet<FRemoteObjectId> MigratedPrimitiveComponentOwnerIds;
+
 	TUniquePtr<FPhysicsReplicationCache> PhysicsReplicationCache;
 
 	TSet<UPrimitiveComponent*> CollisionEventRegistrations;
