[Engine]
* Fixed header unit compile errors

diff --git a/Engine/Source/Runtime/Engine/Public/Rendering/MorphTargetVertexCodec.h b/Engine/Source/Runtime/Engine/Public/Rendering/MorphTargetVertexCodec.h
--- a/Engine/Source/Runtime/Engine/Public/Rendering/MorphTargetVertexCodec.h
+++ b/Engine/Source/Runtime/Engine/Public/Rendering/MorphTargetVertexCodec.h
@@ -14,10 +14,10 @@ class FDefaultBitArrayAllocator;
 namespace UE::MorphTargetVertexCodec
 {
 
-constexpr uint32 BatchSizeBits = 6u;
-constexpr uint32 BatchSize = 1u << BatchSizeBits;
+inline constexpr uint32 BatchSizeBits = 6u;
+inline constexpr uint32 BatchSize = 1u << BatchSizeBits;
 
-constexpr uint32 NumBatchHeaderDwords = 10u;
+inline constexpr uint32 NumBatchHeaderDwords = 10u;
 
 
 struct FQuantizedDelta
