BugHawk issue: Unconditionally adding ActorPackage to ModifiedPackages risks including the destination level's package (LevelPackage) when the actor is already outered to the level (bSameOuter == true for non-OFPA scenarios). This can cause FLevelStreamingGCHelper::TrashPackage to trash the active level package, leading to severe instability.

diff --git a/Engine/Source/Runtime/Engine/Private/WorldPartition/WorldPartitionLevelHelper.cpp b/Engine/Source/Runtime/Engine/Private/WorldPartition/WorldPartitionLevelHelper.cpp
--- a/Engine/Source/Runtime/Engine/Private/WorldPartition/WorldPartitionLevelHelper.cpp
+++ b/Engine/Source/Runtime/Engine/Private/WorldPartition/WorldPartitionLevelHelper.cpp
@@ -348,10 +348,13 @@ void FWorldPartitionLevelHelper::MoveExternalActorsToLevel(const TArray<FWorldPa
 		}
 	}
 
-	// Trash the modified packages to guarantee that any potential future load of this actor won't find the old empty package
+	// Trash the modified packages (except the destination level) to guarantee that any potential future load of this actor won't find the old empty package
 	for (UPackage* ModifiedPackage : ModifiedPackages)
 	{
-		FLevelStreamingGCHelper::TrashPackage(ModifiedPackage);
+		if (ModifiedPackage != LevelPackage)
+		{
+			FLevelStreamingGCHelper::TrashPackage(ModifiedPackage);
+		}
 	}
 
 	OutModifiedPackages.Append(ModifiedPackages.Array());
