Change how we move Objects from ExternalActor pacakges to worldpartition cell packages to maximally preserve the object's name in case there's significant information in it (even though it's not recommend to store such information in the name).

diff --git a/Engine/Source/Runtime/Engine/Private/WorldPartition/WorldPartitionLevelHelper.cpp b/Engine/Source/Runtime/Engine/Private/WorldPartition/WorldPartitionLevelHelper.cpp
--- a/Engine/Source/Runtime/Engine/Private/WorldPartition/WorldPartitionLevelHelper.cpp
+++ b/Engine/Source/Runtime/Engine/Private/WorldPartition/WorldPartitionLevelHelper.cpp
@@ -316,8 +316,17 @@ void FWorldPartitionLevelHelper::MoveExternalActorsToLevel(const TArray<FWorldPa
 						}
 						else
 						{
-							// Move objects in the destination level package
-							Object->Rename(nullptr, LevelPackage, REN_NonTransactional | REN_DoNotDirty | REN_DontCreateRedirectors);
+							// Move objects in the destination level package if not outerred to the Actor (and thus wouldn't follow)
+							// We handle name clashes in a custom way by increasing a suffix on the already existing base name instead
+							// of using the default behavior of having the class name be the base name, this is to maximally preserve
+							// information present in the object name in case somebody is using the name as significant information 
+							FName NewName = Object->GetFName();
+							if (StaticFindObject(nullptr, LevelPackage, *NewName.ToString()))
+							{
+								NewName = MakeUniqueObjectName(LevelPackage, Object->GetClass(), *Object->GetName());
+							}
+
+							Object->Rename(*NewName.ToString(), LevelPackage, REN_NonTransactional | REN_DoNotDirty | REN_DontCreateRedirectors);
 						}
 
 						if (PackageObjectMapping.bIsEditorOnly)
