[RandomStream] Add RandBool to FRandomStream. Update KismetMathLibary to use this function instead.

diff --git a/Engine/Source/Runtime/Core/Public/Math/RandomStream.h b/Engine/Source/Runtime/Core/Public/Math/RandomStream.h
--- a/Engine/Source/Runtime/Core/Public/Math/RandomStream.h
+++ b/Engine/Source/Runtime/Core/Public/Math/RandomStream.h
@@ -211,6 +211,16 @@ public:
 		return InMin + (InMax - InMin) * FRand();
 	}
 
+	/**
+	 * Helper function for rand implementations.
+	 *
+	 * @return A random boolean value.
+	 */
+	UE_FORCEINLINE_HINT bool RandBool() const
+	{
+		return RandHelper(2) == 1;
+	}
+
 	/**
 	 * Returns a random vector of unit size.
 	 *
diff --git a/Engine/Source/Runtime/Engine/Private/KismetMathLibrary.cpp b/Engine/Source/Runtime/Engine/Private/KismetMathLibrary.cpp
--- a/Engine/Source/Runtime/Engine/Private/KismetMathLibrary.cpp
+++ b/Engine/Source/Runtime/Engine/Private/KismetMathLibrary.cpp
@@ -1073,7 +1073,7 @@ bool UKismetMathLibrary::InRange_Int64Int64(int64 Value, int64 Min, int64 Max, b
 
 bool UKismetMathLibrary::RandomBoolFromStream(const FRandomStream& Stream)
 {
-	return (Stream.RandRange(0, 1) == 1) ? true : false;
+	return Stream.RandBool();
 }
 
 float UKismetMathLibrary::RandomFloatFromStream(const FRandomStream& Stream)
