[WorldPartition] Fixes in IsStreamingCompleted.
- Fixed IsStreamingCompleted with QuerySources which include Data Layers: the old implementation was ignoring the QueryState.
- Fixed IsStreamingCompleted with StreamingSource: If the source TargetState = EStreamingSourceTargetState::Loaded, queried data layers didn't include active data layers.
- Added operators to compare EWorldPartitionRuntimeCellState with EDataLayerRuntimeState.
- Small optimization in UWorldPartitionRuntimeHashSet::ForEachStreamingCellsQuery and FSpatialHashStreamingGrid::GetCells

diff --git a/Engine/Source/Runtime/Engine/Private/WorldPartition/RuntimeHashSet/WorldPartitionRuntimeHashSet.cpp b/Engine/Source/Runtime/Engine/Private/WorldPartition/RuntimeHashSet/WorldPartitionRuntimeHashSet.cpp
--- a/Engine/Source/Runtime/Engine/Private/WorldPartition/RuntimeHashSet/WorldPartitionRuntimeHashSet.cpp
+++ b/Engine/Source/Runtime/Engine/Private/WorldPartition/RuntimeHashSet/WorldPartitionRuntimeHashSet.cpp
@@ -559,7 +559,8 @@ void UWorldPartitionRuntimeHashSet::ForEachStreamingCellsQuery(const FWorldParti
 		{
 			if (Cell->HasDataLayers())
 			{
-				if (Cell->GetDataLayers().FindByPredicate([&](const FName& DataLayerName) { return QuerySource.DataLayers.Contains(DataLayerName); }))
+				if (!QuerySource.DataLayers.IsEmpty() && 
+					Cell->GetDataLayers().FindByPredicate([&](const FName& DataLayerName) { return QuerySource.DataLayers.Contains(DataLayerName); }))
 				{
 					return true;
 				}
@@ -627,7 +628,7 @@ void UWorldPartitionRuntimeHashSet::ForEachStreamingCellsQuery(const FWorldParti
 	{
         if (FStreamingSourceShapeHelper::IsSourceAffectingGrid(QuerySource.TargetGrids, QuerySource.TargetBehavior, StreamingData.Name))
         {
-			return ForEachSpatiallyLoadedCells(StreamingData.SpatialIndex.Get(), StreamingData.GetLoadingRange()) && 
+				return ForEachSpatiallyLoadedCells(StreamingData.SpatialIndex.Get(), StreamingData.GetLoadingRange()) &&
 				   ForEachSpatiallyLoadedCells(StreamingData.SpatialIndex2D.Get(), StreamingData.GetLoadingRange()) &&
 				   ForEachNonSpatiallyLoadedCells(StreamingData.NonSpatiallyLoadedCells);
 		}
diff --git a/Engine/Source/Runtime/Engine/Private/WorldPartition/WorldPartitionRuntimeCell.cpp b/Engine/Source/Runtime/Engine/Private/WorldPartition/WorldPartitionRuntimeCell.cpp
--- a/Engine/Source/Runtime/Engine/Private/WorldPartition/WorldPartitionRuntimeCell.cpp
+++ b/Engine/Source/Runtime/Engine/Private/WorldPartition/WorldPartitionRuntimeCell.cpp
@@ -19,6 +19,13 @@
 
 #include UE_INLINE_GENERATED_CPP_BY_NAME(WorldPartitionRuntimeCell)
 
+static_assert(std::is_same<std::underlying_type<EDataLayerRuntimeState>::type, uint8>::value);
+static_assert(std::is_same<std::underlying_type<EWorldPartitionRuntimeCellState>::type, uint8>::value);
+
+static_assert(static_cast<uint8>(EWorldPartitionRuntimeCellState::Unloaded) == static_cast<uint8>(EDataLayerRuntimeState::Unloaded));
+static_assert(static_cast<uint8>(EWorldPartitionRuntimeCellState::Loaded) == static_cast<uint8>(EDataLayerRuntimeState::Loaded));
+static_assert(static_cast<uint8>(EWorldPartitionRuntimeCellState::Activated) == static_cast<uint8>(EDataLayerRuntimeState::Activated));
+
 UWorldPartitionRuntimeCell::UWorldPartitionRuntimeCell(const FObjectInitializer& ObjectInitializer)
 	: Super(ObjectInitializer)
 	, bIsAlwaysLoaded(false)
diff --git a/Engine/Source/Runtime/Engine/Private/WorldPartition/WorldPartitionRuntimeHash.cpp b/Engine/Source/Runtime/Engine/Private/WorldPartition/WorldPartitionRuntimeHash.cpp
--- a/Engine/Source/Runtime/Engine/Private/WorldPartition/WorldPartitionRuntimeHash.cpp
+++ b/Engine/Source/Runtime/Engine/Private/WorldPartition/WorldPartitionRuntimeHash.cpp
@@ -809,6 +809,14 @@ FWorldPartitionStreamingContext::FWorldPartitionStreamingContext(EWorldPartition
 	, UpdateStreamingStateEpoch(InUpdateStreamingStateEpoch)
 {}
 
+EDataLayerRuntimeState FWorldPartitionStreamingContext::GetDataLayerEffectiveRuntimeState(FName DataLayerName) const
+{
+	check(IsValid());
+	check(DataLayerEffectiveStates);
+	check(!DataLayerName.IsNone());
+	return !DataLayerName.IsNone() ? DataLayerEffectiveStates->GetDataLayerEffectiveRuntimeStateByName(DataLayerName) : EDataLayerRuntimeState::Unloaded;
+}
+
 EDataLayerRuntimeState FWorldPartitionStreamingContext::ResolveDataLayerRuntimeState(const FDataLayerInstanceNames& InDataLayers) const
 {
 	if (InDataLayers.IsEmpty())
diff --git a/Engine/Source/Runtime/Engine/Private/WorldPartition/WorldPartitionRuntimeSpatialHash.cpp b/Engine/Source/Runtime/Engine/Private/WorldPartition/WorldPartitionRuntimeSpatialHash.cpp
--- a/Engine/Source/Runtime/Engine/Private/WorldPartition/WorldPartitionRuntimeSpatialHash.cpp
+++ b/Engine/Source/Runtime/Engine/Private/WorldPartition/WorldPartitionRuntimeSpatialHash.cpp
@@ -301,7 +301,8 @@ void FSpatialHashStreamingGrid::GetCells(const FWorldPartitionStreamingQuerySour
 	{
 		if (Cell->HasDataLayers())
 		{
-			if (Cell->GetDataLayers().FindByPredicate([&](const FName& DataLayerName) { return QuerySource.DataLayers.Contains(DataLayerName); }))
+			if (!QuerySource.DataLayers.IsEmpty() &&
+				Cell->GetDataLayers().FindByPredicate([&](const FName& DataLayerName) { return QuerySource.DataLayers.Contains(DataLayerName); }))
 			{
 				return true;
 			}
diff --git a/Engine/Source/Runtime/Engine/Private/WorldPartition/WorldPartitionStreamingPolicy.cpp b/Engine/Source/Runtime/Engine/Private/WorldPartition/WorldPartitionStreamingPolicy.cpp
--- a/Engine/Source/Runtime/Engine/Private/WorldPartition/WorldPartitionStreamingPolicy.cpp
+++ b/Engine/Source/Runtime/Engine/Private/WorldPartition/WorldPartitionStreamingPolicy.cpp
@@ -1103,7 +1103,12 @@ bool UWorldPartitionStreamingPolicy::IsStreamingCompleted(const TArray<FWorldPar
 		QuerySource.bUseGridLoadingRange = true;
 		QuerySource.Radius = 0.f;
 		QuerySource.bDataLayersOnly = false;
-		QuerySource.DataLayers = (StreamingSource.TargetState == EStreamingSourceTargetState::Loaded) ? DataLayerManager->GetEffectiveLoadedDataLayerNames().Array() : DataLayerManager->GetEffectiveActiveDataLayerNames().Array();
+		// Always include active Data Layers even when testing loaded because we are querrying in non-exact state
+		QuerySource.DataLayers = DataLayerManager->GetEffectiveActiveDataLayerNames().Array();
+		if (StreamingSource.TargetState == EStreamingSourceTargetState::Loaded)
+		{
+			QuerySource.DataLayers.Append(DataLayerManager->GetEffectiveLoadedDataLayerNames().Array());
+		}
 
 		// Execute query
 		const EWorldPartitionRuntimeCellState QueryState = (StreamingSource.TargetState == EStreamingSourceTargetState::Loaded) ? EWorldPartitionRuntimeCellState::Loaded : EWorldPartitionRuntimeCellState::Activated;
@@ -1125,7 +1130,19 @@ bool UWorldPartitionStreamingPolicy::IsStreamingCompleted(EWorldPartitionRuntime
 	bool bResult = true;
 	for (const FWorldPartitionStreamingQuerySource& QuerySource : QuerySources)
 	{
-		WorldPartition->RuntimeHash->ForEachStreamingCellsQuery(QuerySource, [QuerySource, QueryState, bExactState, bIsHLODEnabled, DataLayerManager, &StreamingContext, &bResult](const UWorldPartitionRuntimeCell* Cell)
+		// Verify that the state of the query data layers are in the query target state (or greater if not querying for exact state)
+		bool bAreQueryDLsInTargetState = true;
+		for (const FName& QueryDataLayer : QuerySource.DataLayers)
+		{
+			EDataLayerRuntimeState QueryDataLayerState = StreamingContext.GetDataLayerEffectiveRuntimeState(QueryDataLayer);
+			if ((QueryDataLayerState != QueryState) && (bExactState || QueryDataLayerState < QueryState))
+			{
+				bAreQueryDLsInTargetState = false;
+				break;
+			}
+		}
+
+		WorldPartition->RuntimeHash->ForEachStreamingCellsQuery(QuerySource, [QuerySource, QueryState, bExactState, bIsHLODEnabled, DataLayerManager, bAreQueryDLsInTargetState, &StreamingContext, &bResult](const UWorldPartitionRuntimeCell* Cell)
 		{
 			const EWorldPartitionRuntimeCellState CellState = Cell->GetCurrentState();
 
@@ -1139,30 +1156,21 @@ bool UWorldPartitionStreamingPolicy::IsStreamingCompleted(EWorldPartitionRuntime
 					bSkipCell = Cell->GetIsHLOD();
 				}
 
-				// Test if cell is already in the effective wanted state.
-				// Note that GetCellEffectiveWantedState always return Activated when the cell has no data layers.
-				// In this case, continue testing the QueryState with the CellState to respect bExactState.
 				if (!bSkipCell && Cell->HasDataLayers())
 				{
-					const EDataLayerRuntimeState CellWantedState = Cell->GetCellEffectiveWantedState(StreamingContext);
-					bSkipCell = (CellState == EWorldPartitionRuntimeCellState::Unloaded && CellWantedState == EDataLayerRuntimeState::Unloaded) || 
-								(CellState == EWorldPartitionRuntimeCellState::Loaded && CellWantedState == EDataLayerRuntimeState::Loaded) || 
-								(CellState == EWorldPartitionRuntimeCellState::Activated && CellWantedState == EDataLayerRuntimeState::Activated);
-				}
-
-				// If we are querying for Unloaded/Loaded but a Cell is part of a data layer outside of the query that is activated do not consider it
-				if (!bSkipCell && QueryState < CellState)
-				{
-					for (const FName& CellDataLayer : Cell->GetDataLayers())
+					if (!bAreQueryDLsInTargetState)
 					{
-						if (!QuerySource.DataLayers.Contains(CellDataLayer) && DataLayerManager->GetDataLayerInstanceEffectiveRuntimeState(DataLayerManager->GetDataLayerInstanceFromName(CellDataLayer)) > EDataLayerRuntimeState::Unloaded)
-						{
-							bSkipCell = true;
-							break;
-						}
+						bResult = false;
+						return false;
 					}
+
+					// Because we already validated that the queried data layers are in the proper effective state
+					// if the cell wanted state is not in the proper queried state, it's because of unquerried cell data layer states.
+					// In this case, we must ignore the cell.
+					const EDataLayerRuntimeState CellWantedState = Cell->GetCellEffectiveWantedState(StreamingContext);
+					bSkipCell = ((CellWantedState != QueryState) && (bExactState || CellWantedState < QueryState));
 				}
-								
+
 				if (!bSkipCell && (bExactState || CellState < QueryState))
 				{
 					bResult = false;
diff --git a/Engine/Source/Runtime/Engine/Public/WorldPartition/WorldPartitionRuntimeCell.h b/Engine/Source/Runtime/Engine/Public/WorldPartition/WorldPartitionRuntimeCell.h
--- a/Engine/Source/Runtime/Engine/Public/WorldPartition/WorldPartitionRuntimeCell.h
+++ b/Engine/Source/Runtime/Engine/Public/WorldPartition/WorldPartitionRuntimeCell.h
@@ -467,3 +467,43 @@ public:
 	friend class UWorldPartitionStreamingPolicy;
 	friend struct FSpatialHashStreamingGrid;
 };
+
+[[nodiscard]] constexpr bool operator == (EWorldPartitionRuntimeCellState a, EDataLayerRuntimeState b)
+{
+	return static_cast<uint8>(a) == static_cast<uint8>(b);
+}
+
+[[nodiscard]] constexpr bool operator == (EDataLayerRuntimeState a, EWorldPartitionRuntimeCellState b)
+{
+	return static_cast<uint8>(a) == static_cast<uint8>(b);
+}
+
+[[nodiscard]] constexpr bool operator != (EWorldPartitionRuntimeCellState a, EDataLayerRuntimeState b)
+{
+	return !(a == b);
+}
+
+[[nodiscard]] constexpr bool operator != (EDataLayerRuntimeState a, EWorldPartitionRuntimeCellState b)
+{
+	return !(a == b);
+}
+
+[[nodiscard]] constexpr bool operator < (EWorldPartitionRuntimeCellState a, EDataLayerRuntimeState b)
+{
+	return static_cast<uint8>(a) < static_cast<uint8>(b);
+}
+
+[[nodiscard]] constexpr bool operator < (EDataLayerRuntimeState a, EWorldPartitionRuntimeCellState b)
+{
+	return static_cast<uint8>(a) < static_cast<uint8>(b);
+}
+
+[[nodiscard]] constexpr bool operator > (EWorldPartitionRuntimeCellState a, EDataLayerRuntimeState b)
+{
+	return static_cast<uint8>(a) > static_cast<uint8>(b);
+}
+
+[[nodiscard]] constexpr bool operator > (EDataLayerRuntimeState a, EWorldPartitionRuntimeCellState b)
+{
+	return static_cast<uint8>(a) > static_cast<uint8>(b);
+}
\ No newline at end of file
diff --git a/Engine/Source/Runtime/Engine/Public/WorldPartition/WorldPartitionRuntimeHash.h b/Engine/Source/Runtime/Engine/Public/WorldPartition/WorldPartitionRuntimeHash.h
--- a/Engine/Source/Runtime/Engine/Public/WorldPartition/WorldPartitionRuntimeHash.h
+++ b/Engine/Source/Runtime/Engine/Public/WorldPartition/WorldPartitionRuntimeHash.h
@@ -162,6 +162,7 @@ public:
 	static ENGINE_API FWorldPartitionStreamingContext Create(const UWorld* InWorld);
 	ENGINE_API FWorldPartitionStreamingContext();
 	bool IsValid() const { return bIsValid; }
+	ENGINE_API EDataLayerRuntimeState GetDataLayerEffectiveRuntimeState(FName DataLayerName) const;
 
 private:
 	ENGINE_API FWorldPartitionStreamingContext(const UWorld* InWorld);
