Shader SDCE, only enable for Adaptive

diff --git a/Engine/Source/Runtime/Engine/Private/Materials/HLSLMaterialTranslator.cpp b/Engine/Source/Runtime/Engine/Private/Materials/HLSLMaterialTranslator.cpp
--- a/Engine/Source/Runtime/Engine/Private/Materials/HLSLMaterialTranslator.cpp
+++ b/Engine/Source/Runtime/Engine/Private/Materials/HLSLMaterialTranslator.cpp
@@ -3087,23 +3087,26 @@ void FHLSLMaterialTranslator::GetMaterialEnvironment(EShaderPlatform InPlatform,
 		OutEnvironment.SetDefine(TEXT("SUBSTRATE_CLAMPED_CLOSURE_COUNT"), EnvironmentDefines->SubstrateClampedClosureCount);
 		OutEnvironment.SetDefine(TEXT("SUBSTRATE_COMPLEXSPECIALPATH"), EnvironmentDefines->bSubstrateComplexSpecialPath);
 
-		static const TCHAR* SDCESymbols[] = {
-			TEXT("FMaterialPixelParameters"),
-			TEXT("FMaterialLWCData"),
-			TEXT("FPrimitiveSceneData"),
-			TEXT("FPixelMaterialInputs"),
-			TEXT("FCluster"),
-			TEXT("FNaniteView"),
-			TEXT("FInstanceDynamicData"),
-			TEXT("FNanitePostDeformVertex"),
-			TEXT("FMaterialVertexParameters"),
-			TEXT("FNanitePixelAttributes"),
-			TEXT("ViewState"),
-			TEXT("FSceneData"),
-			TEXT("FPrimitiveSceneData")
-		};
+		if (!Substrate::IsSubstrateBlendableGBufferEnabled(InPlatform))
+		{
+			static const TCHAR* SDCESymbols[] = {
+				TEXT("FMaterialPixelParameters"),
+				TEXT("FMaterialLWCData"),
+				TEXT("FPrimitiveSceneData"),
+				TEXT("FPixelMaterialInputs"),
+				TEXT("FCluster"),
+				TEXT("FNaniteView"),
+				TEXT("FInstanceDynamicData"),
+				TEXT("FNanitePostDeformVertex"),
+				TEXT("FMaterialVertexParameters"),
+				TEXT("FNanitePixelAttributes"),
+				TEXT("ViewState"),
+				TEXT("FSceneData"),
+				TEXT("FPrimitiveSceneData")
+			};
 		
-		OutEnvironment.SetCompileArgument(TEXT("UESHADERMETADATA_SDCE"), FString::Join(SDCESymbols, TEXT(";")));
+			OutEnvironment.SetCompileArgument(TEXT("UESHADERMETADATA_SDCE"), FString::Join(SDCESymbols, TEXT(";")));
+		}
 	}
 
 	for (int i = 0; i < EnvironmentDefines->SubstrateDefines.Num(); ++i)
