Fix potential crash with null UWorld dereference in ConvertQueryImpactHitImp()

diff --git a/Engine/Source/Runtime/Engine/Private/Collision/CollisionConversions.cpp b/Engine/Source/Runtime/Engine/Private/Collision/CollisionConversions.cpp
--- a/Engine/Source/Runtime/Engine/Private/Collision/CollisionConversions.cpp
+++ b/Engine/Source/Runtime/Engine/Private/Collision/CollisionConversions.cpp
@@ -286,7 +286,7 @@ EConvertQueryResult ConvertQueryImpactHitImp(const UWorld* World, const THitLoca
 	SCOPE_CYCLE_COUNTER(STAT_ConvertQueryImpactHit);
 
 #if WITH_EDITOR
-	if(bReturnFaceIndex && World->IsGameWorld())
+	if(bReturnFaceIndex && World && World->IsGameWorld())
 	{
 		if(UPhysicsSettings::Get()->bSuppressFaceRemapTable)
 		{
