Fixing a log category (LogPhysics instead of LogTemp)

diff --git a/Engine/Source/Runtime/Engine/Private/PhysicsEngine/PhysicsReplication.cpp b/Engine/Source/Runtime/Engine/Private/PhysicsEngine/PhysicsReplication.cpp
--- a/Engine/Source/Runtime/Engine/Private/PhysicsEngine/PhysicsReplication.cpp
+++ b/Engine/Source/Runtime/Engine/Private/PhysicsEngine/PhysicsReplication.cpp
@@ -2288,12 +2288,12 @@ bool FPhysicsReplicationAsync::ResimulationReplication(Chaos::FPBDRigidParticleH
 
 	if (Chaos::FPhysicsSolverBase::CanDebugNetworkPhysicsPrediction())
 	{
-		UE_LOG(LogTemp, Log, TEXT("Apply Rigid body state at local frame %d with offset = %d"), LocalFrame, Target.FrameOffset);
-		UE_LOG(LogTemp, Log, TEXT("Should Trigger Resim = %s | Server Frame = %d | Client Frame = %d"), (bShouldTriggerResim ? TEXT("True") : TEXT("False")), Target.ServerFrame, LocalFrame);
-		UE_LOG(LogTemp, Log, TEXT("Particle Target Position = %s | Current Position = %s"), *Target.TargetState.Position.ToString(), *PastState.GetX().ToString());
-		UE_LOG(LogTemp, Log, TEXT("Particle Target Velocity = %s | Current Velocity = %s"), *Target.TargetState.LinVel.ToString(), *PastState.GetV().ToString());
-		UE_LOG(LogTemp, Log, TEXT("Particle Target Quaternion = %s | Current Quaternion = %s"), *Target.TargetState.Quaternion.ToString(), *PastState.GetR().ToString());
-		UE_LOG(LogTemp, Log, TEXT("Particle Target Omega = %s | Current Omega= %s"), *Target.TargetState.AngVel.ToString(), *PastState.GetW().ToString());
+		UE_LOG(LogPhysics, Log, TEXT("Apply Rigid body state at local frame %d with offset = %d"), LocalFrame, Target.FrameOffset);
+		UE_LOG(LogPhysics, Log, TEXT("Should Trigger Resim = %s | Server Frame = %d | Client Frame = %d"), (bShouldTriggerResim ? TEXT("True") : TEXT("False")), Target.ServerFrame, LocalFrame);
+		UE_LOG(LogPhysics, Log, TEXT("Particle Target Position = %s | Current Position = %s"), *Target.TargetState.Position.ToString(), *PastState.GetX().ToString());
+		UE_LOG(LogPhysics, Log, TEXT("Particle Target Velocity = %s | Current Velocity = %s"), *Target.TargetState.LinVel.ToString(), *PastState.GetV().ToString());
+		UE_LOG(LogPhysics, Log, TEXT("Particle Target Quaternion = %s | Current Quaternion = %s"), *Target.TargetState.Quaternion.ToString(), *PastState.GetR().ToString());
+		UE_LOG(LogPhysics, Log, TEXT("Particle Target Omega = %s | Current Omega= %s"), *Target.TargetState.AngVel.ToString(), *PastState.GetW().ToString());
 	}
 
 	if (PhysicsReplicationCVars::ResimulationCVars::bDrawDebug || RenderInterpolationCVars::bRenderInterpDebugDrawResimTrigger)
