[NetPhysics] Fix issue where Predictive Interpolation tick alignment would get out of sync.
- Calculate the TickOffset even if we are not allowed to perform extrapolation / alignment of the target state

diff --git a/Engine/Source/Runtime/Engine/Private/PhysicsEngine/PhysicsReplication.cpp b/Engine/Source/Runtime/Engine/Private/PhysicsEngine/PhysicsReplication.cpp
--- a/Engine/Source/Runtime/Engine/Private/PhysicsEngine/PhysicsReplication.cpp
+++ b/Engine/Source/Runtime/Engine/Private/PhysicsEngine/PhysicsReplication.cpp
@@ -1196,19 +1196,19 @@ void FPhysicsReplicationAsync::UpdateAsyncTarget(const FPhysicsRepAsyncInputData
 				* If the clients latency is raised or lowered since replication started there might be a consistent offset in the TickCount which is handled by TimeDilation of client physics through APlayerController::UpdateServerAsyncPhysicsTickOffset()
 				*/
 
-				// Run target alignment if we have been allowed to alter the target during the last two target updates
-				if (!bFirstTarget && bPrevAllowTargetAltering && Target->bAllowTargetAltering && !bFrameOffsetCorrected)
+				if (!bFirstTarget && !bFrameOffsetCorrected)
 				{
-					const int32 AdjustedAverageReceiveInterval = FMath::CeilToInt(Target->AverageReceiveInterval) * PhysicsReplicationCVars::PredictiveInterpolationCVars::TargetTickAlignmentClampMultiplier;
+					// Calculate the current tick-offset, leave at 0 for the first target received and if the FrameOffset between client and server has been corrected
+					Target->TickCount = PrevTickCount - Target->ReceiveInterval;
 
-					// Set the TickCount to the physics tick offset value from where we expected this target to arrive.
-					// If the client has ticked 2 times ahead from the last target and this target is 3 ticks in front of the previous target then the TickOffset should be -1
-					Target->TickCount = FMath::Clamp(PrevTickCount - Target->ReceiveInterval, -AdjustedAverageReceiveInterval, AdjustedAverageReceiveInterval);
-
-					// Apply target alignment if we aren't waiting for a newer state from the server
-					if (!Target->IsWaiting())
+					// Run target alignment if we have been allowed to alter the target during the last two target updates
+					if (bPrevAllowTargetAltering && Target->bAllowTargetAltering && !Target->IsWaiting())
 					{
-						FPhysicsReplicationAsync::ExtrapolateTarget(*Target, Target->TickCount, GetDeltaTime_Internal());
+						// Target alignment is based on TickCount but clamped to not extrapolate forwards or backwards too far
+						const int32 AdjustedAverageReceiveInterval = FMath::CeilToInt(Target->AverageReceiveInterval) * PhysicsReplicationCVars::PredictiveInterpolationCVars::TargetTickAlignmentClampMultiplier;
+						const int32 TickAlignment = FMath::Clamp(Target->TickCount, -AdjustedAverageReceiveInterval, AdjustedAverageReceiveInterval);
+
+						FPhysicsReplicationAsync::ExtrapolateTarget(*Target, TickAlignment, GetDeltaTime_Internal());
 					}
 				}
 			}
