New Material Translator:  Preshader2, a register-based VM designed to be more efficient than the original stack-based VM.  Advantages include fewer opcodes (stack push ops disappear), a single bytecode stream (avoids context switch between header loop and Evaluate), and theoretically able to support more advanced features like branching or reusing subexpressions, which are complicated in a stack design.  Currently only texture parameter opcodes are generated in Preshader2 in the new translator -- general math opcodes will be the next step.

Also took the opportunity to dead strip jump and struct support from the original Preshader VM.  The "old new translator" supported structs, but that itself was dead stripped, so there's no longer any code that can generate a struct in a preshader.  And jumps were never implemented.  Reduces the size of the preshader header structure from 28 to 12 bytes each.  Saves around 600K in a large project -- memory savings is trivial, but more compact structures will be better for cache performance.  The goal is to allow a fair comparison of memory and perf between the preshader implementations, with the low hanging fruit cleaned up in the original translator.

#rnx

diff --git a/Engine/Source/Runtime/Engine/Private/Materials/MaterialExpressionsIR.cpp b/Engine/Source/Runtime/Engine/Private/Materials/MaterialExpressionsIR.cpp
--- a/Engine/Source/Runtime/Engine/Private/Materials/MaterialExpressionsIR.cpp
+++ b/Engine/Source/Runtime/Engine/Private/Materials/MaterialExpressionsIR.cpp
@@ -255,6 +255,7 @@
 #include "MaterialShared.h"
 #include "Misc/MemStackUtility.h"
 #include "RenderUtils.h"
+#include "Shader/Preshader2.h"
 #include "TextureResource.h"
 #include "ColorManagement/ColorSpace.h"
 
@@ -1572,10 +1573,10 @@ void UMaterialExpressionTextureProperty::Build(MIR::FEmitter& Em)
 
 	const bool bTexelSizeInUVSpace = Property == TMTM_TexelSize;
 
-	const UE::Shader::EPreshaderOpcode PreshaderOpcode =
+	const EPreshader2Opcode PreshaderOpcode =
 		bTexelSizeInUVSpace
-			? UE::Shader::EPreshaderOpcode::TexelSize
-			: UE::Shader::EPreshaderOpcode::TextureSize;
+			? EPreshader2Opcode::TexelSize
+			: EPreshader2Opcode::TextureSize;
 
 	UObject* SourceParameterTexture = Em.GetTextureFromValue(TextureValue);
 	if (!SourceParameterTexture)
@@ -1781,7 +1782,7 @@ void UMaterialExpressionRuntimeVirtualTextureSample::Build(MIR::FEmitter& Em)
 	{
 		const UE::Shader::EValueType UniformType = URuntimeVirtualTexture::GetUniformParameterType(UniformIndex);
 		Uniforms[UniformIndex] = Em.PreshaderParameter(
-			MIR::FType::FromShaderType(UniformType), UE::Shader::EPreshaderOpcode::RuntimeVirtualTextureUniform,
+			MIR::FType::FromShaderType(UniformType), EPreshader2Opcode::RuntimeVirtualTextureUniform,
 			TextureObjects[0], MIR::FPreshaderParameterPayload{ .UniformIndex = UniformIndex });
 	}
 
diff --git a/Engine/Source/Runtime/Engine/Private/Materials/MaterialIREmitter.cpp b/Engine/Source/Runtime/Engine/Private/Materials/MaterialIREmitter.cpp
--- a/Engine/Source/Runtime/Engine/Private/Materials/MaterialIREmitter.cpp
+++ b/Engine/Source/Runtime/Engine/Private/Materials/MaterialIREmitter.cpp
@@ -8,6 +8,7 @@
 #include "Materials/MaterialExpressionTextureBase.h"
 #include "Materials/MaterialExternalCodeRegistry.h"
 #include "Materials/MaterialAggregate.h"
+#include "Shader/Preshader2.h"
 #include "Shader/ShaderTypes.h"
 #include "MaterialShared.h"
 #include "MaterialSharedPrivate.h"
@@ -1253,12 +1254,12 @@ UObject* FEmitter::GetTextureFromValue(MIR::FValueRef Texture)
 }
 
 // For now, only a small subset of opcodes are supported: TextureSize, TexelSize, and RuntimeVirtualTextureUniform.
-FValueRef FEmitter::PreshaderParameter(FType Type, UE::Shader::EPreshaderOpcode Opcode, FValueRef SourceParameter, FPreshaderParameterPayload Payload)
+FValueRef FEmitter::PreshaderParameter(FType Type, EPreshader2Opcode Opcode, FValueRef SourceParameter, FPreshaderParameterPayload Payload)
 {
 	checkf(
-		Opcode == UE::Shader::EPreshaderOpcode::TextureSize ||
-		Opcode == UE::Shader::EPreshaderOpcode::TexelSize ||
-		Opcode == UE::Shader::EPreshaderOpcode::RuntimeVirtualTextureUniform,
+		Opcode == EPreshader2Opcode::TextureSize ||
+		Opcode == EPreshader2Opcode::TexelSize ||
+		Opcode == EPreshader2Opcode::RuntimeVirtualTextureUniform,
 		TEXT("Preshader opcode (0x%X) not supported for parameters in new material translator"), (int32)Opcode
 	);
 
diff --git a/Engine/Source/Runtime/Engine/Private/Materials/MaterialIRValueAnalyzer.cpp b/Engine/Source/Runtime/Engine/Private/Materials/MaterialIRValueAnalyzer.cpp
--- a/Engine/Source/Runtime/Engine/Private/Materials/MaterialIRValueAnalyzer.cpp
+++ b/Engine/Source/Runtime/Engine/Private/Materials/MaterialIRValueAnalyzer.cpp
@@ -24,6 +24,7 @@
 #include "VT/RuntimeVirtualTexture.h"
 #include "UObject/Package.h"
 #include "PostProcess/PostProcessMaterialInputs.h"
+#include "Shader/Preshader2.h"
 
 static void ErrorUnlessFeatureLevelSupported(FMaterialIRModule* Module, ERHIFeatureLevel::Type RequiredFeatureLevel, const TCHAR* Message)
 {
@@ -759,28 +760,11 @@ static void AnalyzePreshaderParameter(FMaterialIRValueAnalyzer& Analyzer, MIR::F
 
 	const uint32 GlobalComponentOffset = NextGlobalComponentOffset(Analyzer, PrimitiveType);
 
-	UniformExpressionSet.WriteUniformPreshaderEntry(
-		GlobalComponentOffset, UE::Shader::MakeValueType(UE::Shader::EValueComponentType::Float, PrimitiveType.NumComponents()),
-		[Parameter, SourceParameterName](UE::Shader::FPreshaderData& UniformPreshaderData) -> void
-		{
-			const FHashedMaterialParameterInfo HashedSourceParameterInfo{ SourceParameterName, EMaterialParameterAssociation::GlobalParameter, INDEX_NONE };
-
-			switch (Parameter->Opcode)
-			{
-			case UE::Shader::EPreshaderOpcode::TextureSize:
-			case UE::Shader::EPreshaderOpcode::TexelSize:
-				UniformPreshaderData.WriteOpcode(Parameter->Opcode).Write(HashedSourceParameterInfo).Write(Parameter->TextureIndex);
-				break;
+	UE::Shader::FPreshaderData& PreshaderData = Analyzer.GetPreshaderData();
+	const FHashedMaterialParameterInfo HashedSourceParameterInfo{ SourceParameterName, EMaterialParameterAssociation::GlobalParameter, INDEX_NONE };
 
-			case UE::Shader::EPreshaderOpcode::RuntimeVirtualTextureUniform:
-				UniformPreshaderData.WriteOpcode(Parameter->Opcode).Write(HashedSourceParameterInfo).Write(Parameter->TextureIndex).Write(Parameter->Payload.UniformIndex);
-				break;
-
-			default:
-				UE_MIR_UNREACHABLE();
-			}
-		}
-	);
+	// UniformIndex argument only applies for RuntimeVirtualTextureUniform, but it's harmless to pass it in regardless.
+	UE::Preshader2::EmitTextureOp(PreshaderData, Parameter->Opcode, (uint16)GlobalComponentOffset, HashedSourceParameterInfo, Parameter->TextureIndex, Parameter->Payload.UniformIndex);
 
 	Parameter->Analysis_PreshaderOffset = GlobalComponentOffset;
 }
diff --git a/Engine/Source/Runtime/Engine/Private/Materials/MaterialIRValueAnalyzer.h b/Engine/Source/Runtime/Engine/Private/Materials/MaterialIRValueAnalyzer.h
--- a/Engine/Source/Runtime/Engine/Private/Materials/MaterialIRValueAnalyzer.h
+++ b/Engine/Source/Runtime/Engine/Private/Materials/MaterialIRValueAnalyzer.h
@@ -63,6 +63,9 @@ struct FMaterialIRValueAnalyzer
 	// Performs stage-specific analysis on the given value, if the value kind requires some.
 	void AnalyzeInStage(MIR::FValue* Value, MIR::EStage Stage);
 
+	// Return writable preshader data.
+	UE::Shader::FPreshaderData& GetPreshaderData() { return CompilationOutput->UniformExpressionSet.UniformPreshaderData; }
+
 	// The material being built.
 	UMaterial* Material{};
 	
diff --git a/Engine/Source/Runtime/Engine/Private/Materials/MaterialShared.cpp b/Engine/Source/Runtime/Engine/Private/Materials/MaterialShared.cpp
--- a/Engine/Source/Runtime/Engine/Private/Materials/MaterialShared.cpp
+++ b/Engine/Source/Runtime/Engine/Private/Materials/MaterialShared.cpp
@@ -135,7 +135,6 @@ IMPLEMENT_TYPE_LAYOUT(FMaterialProcessedSource);
 IMPLEMENT_TYPE_LAYOUT(FMaterialShaderMapContent);
 IMPLEMENT_TYPE_LAYOUT(FMaterialUniformParameterEvaluation);
 IMPLEMENT_TYPE_LAYOUT(FMaterialUniformPreshaderHeader);
-IMPLEMENT_TYPE_LAYOUT(FMaterialUniformPreshaderField);
 IMPLEMENT_TYPE_LAYOUT(FMaterialNumericParameterInfo);
 IMPLEMENT_TYPE_LAYOUT(FMaterialTextureParameterInfo);
 IMPLEMENT_TYPE_LAYOUT(FMaterialTextureCollectionParameterInfo);
@@ -4597,18 +4596,17 @@ bool FMaterial::GetMaterialExpressionSource(FString& OutSource)
 				for (int32 PreshaderIndex = 0; PreshaderIndex < NewCompilationOutput.UniformExpressionSet.UniformPreshaders.Num(); PreshaderIndex++)
 				{
 					const FMaterialUniformPreshaderHeader& PreshaderHeader = NewCompilationOutput.UniformExpressionSet.UniformPreshaders[PreshaderIndex];
-					const FMaterialUniformPreshaderField& PreshaderField = NewCompilationOutput.UniformExpressionSet.UniformPreshaderFields[PreshaderIndex];
 
 					UE::Shader::FPreshaderDataContext PreshaderContext(PreshaderContextBase, PreshaderHeader.OpcodeOffset, PreshaderHeader.OpcodeSize);
 					FString PreshaderDebug = PreshaderGenerateDebugString(NewCompilationOutput.UniformExpressionSet, MaterialContext, PreshaderContext, &ParameterReferences);
 
 					// If this is a numeric field, add a swizzle for it
 					const TCHAR* SwizzleSuffix = TEXT("");
-					if (((uint8)PreshaderField.Type >= (uint8)UE::Shader::EValueType::Float1 && (uint8)PreshaderField.Type <= (uint8)UE::Shader::EValueType::Float4) ||
-						((uint8)PreshaderField.Type >= (uint8)UE::Shader::EValueType::Int1 && (uint8)PreshaderField.Type <= (uint8)UE::Shader::EValueType::Int4) ||
-						((uint8)PreshaderField.Type >= (uint8)UE::Shader::EValueType::Bool1 && (uint8)PreshaderField.Type <= (uint8)UE::Shader::EValueType::Bool4))
+					if (((uint8)PreshaderHeader.Type >= (uint8)UE::Shader::EValueType::Float1 && (uint8)PreshaderHeader.Type <= (uint8)UE::Shader::EValueType::Float4) ||
+						((uint8)PreshaderHeader.Type >= (uint8)UE::Shader::EValueType::Int1 && (uint8)PreshaderHeader.Type <= (uint8)UE::Shader::EValueType::Int4) ||
+						((uint8)PreshaderHeader.Type >= (uint8)UE::Shader::EValueType::Bool1 && (uint8)PreshaderHeader.Type <= (uint8)UE::Shader::EValueType::Bool4))
 					{
-						UE::Shader::FType ShaderType(PreshaderField.Type);
+						UE::Shader::FType ShaderType(PreshaderHeader.Type);
 
 						// First axis is offset, second axis is number of components (minus one)
 						static const TCHAR* SwizzleTable[4][4] =
@@ -4618,10 +4616,10 @@ bool FMaterial::GetMaterialExpressionSource(FString& OutSource)
 							{ TEXT(".z"), TEXT(".zw"), TEXT(".zw?"), TEXT(".zw??") },
 							{ TEXT(".w"), TEXT(".w?"), TEXT(".w??"), TEXT(".w???") },
 						};
-						SwizzleSuffix = SwizzleTable[PreshaderField.BufferOffset % 4][ShaderType.GetNumComponents() - 1];
+						SwizzleSuffix = SwizzleTable[PreshaderHeader.BufferOffset % 4][ShaderType.GetNumComponents() - 1];
 					}
 
-					OutSource.Appendf(TEXT("// PreshaderBuffer[%d]%s = %s\n"), PreshaderField.BufferOffset / 4, SwizzleSuffix, *PreshaderDebug);
+					OutSource.Appendf(TEXT("// PreshaderBuffer[%d]%s = %s\n"), PreshaderHeader.BufferOffset / 4, SwizzleSuffix, *PreshaderDebug);
 				}
 
 				// Sort parameter references by frequency
diff --git a/Engine/Source/Runtime/Engine/Private/Materials/MaterialUniformExpressions.cpp b/Engine/Source/Runtime/Engine/Private/Materials/MaterialUniformExpressions.cpp
--- a/Engine/Source/Runtime/Engine/Private/Materials/MaterialUniformExpressions.cpp
+++ b/Engine/Source/Runtime/Engine/Private/Materials/MaterialUniformExpressions.cpp
@@ -15,6 +15,7 @@
 #include "Materials/MaterialRenderProxy.h"
 #include "RHIStaticStates.h"
 #include "Shader/PreshaderEvaluate.h"
+#include "Shader/Preshader2.h"
 #include "ExternalTexture.h"
 
 #include "GlobalRenderResources.h"
@@ -168,7 +169,7 @@ FMaterialUniformExpressionType::FMaterialUniformExpressionType(const TCHAR* InNa
 void FMaterialUniformExpression::WriteNumberOpcodes(UE::Shader::FPreshaderData& OutData) const
 {
 	UE_LOG(LogMaterial, Warning, TEXT("Missing WriteNumberOpcodes impl for %s"), GetType()->GetName());
-	OutData.WriteOpcode(UE::Shader::EPreshaderOpcode::ConstantZero).Write(UE::Shader::FType(UE::Shader::EValueType::Float1));
+	OutData.WriteOpcode(UE::Shader::EPreshaderOpcode::ConstantZero).Write(UE::Shader::EValueType::Float1);
 }
 
 void FUniformParameterOverrides::SetNumericOverride(EMaterialParameterType Type, const FHashedMaterialParameterInfo& ParameterInfo, const UE::Shader::FValue& Value)
@@ -839,8 +840,6 @@ void FUniformExpressionSet::SetVTLayer(int32 StackIndex, int32 VTLayerIndex, int
 void FUniformExpressionSet::WriteUniformPreshaderEntry(uint32 BufferOffset, UE::Shader::EValueType EntryType, const TFunction<void(UE::Shader::FPreshaderData&)>& EntryGenerator)
 {
 	FMaterialUniformPreshaderHeader& PreshaderHeader = UniformPreshaders.AddDefaulted_GetRef();
-	PreshaderHeader.FieldIndex = UniformPreshaderFields.Num();
-	PreshaderHeader.NumFields = 1;
 
 	PreshaderHeader.OpcodeOffset = UniformPreshaderData.Num();
 	{
@@ -848,10 +847,8 @@ void FUniformExpressionSet::WriteUniformPreshaderEntry(uint32 BufferOffset, UE::
 	}
 	PreshaderHeader.OpcodeSize = UniformPreshaderData.Num() - PreshaderHeader.OpcodeOffset;
 
-	FMaterialUniformPreshaderField& PreshaderField = UniformPreshaderFields.AddDefaulted_GetRef();
-	PreshaderField.BufferOffset = BufferOffset;
-	PreshaderField.Type = EntryType;
-	PreshaderField.ComponentIndex = 0;
+	PreshaderHeader.BufferOffset = BufferOffset;
+	PreshaderHeader.Type = EntryType;
 }
 
 void FUniformExpressionSet::GetGameThreadTextureValue(EMaterialTextureParameterType Type, int32 Index, const UMaterialInterface* MaterialInterface, const FMaterial& Material, UTexture*& OutValue, bool bAllowOverride) const
@@ -1205,31 +1202,37 @@ void FUniformExpressionSet::FillUniformBuffer(const FMaterialRenderContext& Mate
 
 		// Dump preshader results into buffer.
 		float* const PreshaderBuffer = (float*)BufferCursor;
-		FPreshaderStack PreshaderStack;
-		FPreshaderDataContext PreshaderBaseContext(UniformPreshaderData);
-		for (const FMaterialUniformPreshaderHeader& Preshader : UniformPreshaders)
+		if (UniformPreshaderData.bPreshader2)
 		{
-			FPreshaderDataContext PreshaderContext(PreshaderBaseContext, Preshader.OpcodeOffset, Preshader.OpcodeSize);
-			const FPreshaderValue Result = EvaluatePreshader(this, MaterialRenderContext, PreshaderStack, PreshaderContext);
-
-			// Fast path for non-structure float1 to float4 fields, which represents the vast majority of cases
-			EValueType FirstFieldType = UniformPreshaderFields[Preshader.FieldIndex].Type;
-			EValueType ResultType = Result.Type.ValueType;
-
-			// This min-max logic assumes that Float1 through Float4 are sequential in the EValueType enum, so we can do just two comparisons
-			// to detect if both Result and the destination field are float types.
-			static_assert((int32)EValueType::Float2 == (int32)EValueType::Float1 + 1);
-			static_assert((int32)EValueType::Float3 == (int32)EValueType::Float2 + 1);
-			static_assert((int32)EValueType::Float4 == (int32)EValueType::Float3 + 1);
-			
-			int32 MinValueType = FMath::Min((int32)FirstFieldType, (int32)ResultType);
-			int32 MaxValueType = FMath::Max((int32)FirstFieldType, (int32)ResultType);
-
-			for (uint32 FieldIndex = 0u; FieldIndex < Preshader.NumFields; ++FieldIndex)
+			// Preshader2 executes a single bytecode stream to write all outputs to the buffer.  UniformPreshaderBufferSize is a count of float4.
+			FPreshaderDataContext Preshader2Context(UniformPreshaderData, TArrayView<float>(PreshaderBuffer, UniformPreshaderBufferSize*4));
+			UE::Preshader2::Evaluate(Preshader2Context, &MaterialRenderContext);
+		}
+		else
+		{
+			FPreshaderStack PreshaderStack;
+			FPreshaderDataContext PreshaderBaseContext(UniformPreshaderData);
+			for (const FMaterialUniformPreshaderHeader& Preshader : UniformPreshaders)
 			{
+				FPreshaderDataContext PreshaderContext(PreshaderBaseContext, Preshader.OpcodeOffset, Preshader.OpcodeSize);
+				const FPreshaderValue Result = EvaluatePreshader(this, MaterialRenderContext, PreshaderStack, PreshaderContext);
+
+				// Fast path for non-structure float1 to float4 fields, which represents the vast majority of cases
+				EValueType FieldType = Preshader.Type;
+				EValueType ResultType = Result.Type.ValueType;
+
+				// This min-max logic assumes that Float1 through Float4 are sequential in the EValueType enum, so we can do just two comparisons
+				// to detect if both Result and the destination field are float types.
+				static_assert((int32)EValueType::Float2 == (int32)EValueType::Float1 + 1);
+				static_assert((int32)EValueType::Float3 == (int32)EValueType::Float2 + 1);
+				static_assert((int32)EValueType::Float4 == (int32)EValueType::Float3 + 1);
+
+				int32 MinValueType = FMath::Min((int32)FieldType, (int32)ResultType);
+				int32 MaxValueType = FMath::Max((int32)FieldType, (int32)ResultType);
+
 				if (MinValueType >= (int32)EValueType::Float1 && MaxValueType <= (int32)EValueType::Float4)
 				{
-					float* DestAddress = PreshaderBuffer + UniformPreshaderFields[Preshader.FieldIndex].BufferOffset;
+					float* DestAddress = PreshaderBuffer + Preshader.BufferOffset;
 
 					// Copy components that exist in both result and destination buffer
 					int32 NumCopyComponents = MinValueType - ((int32)EValueType::Float1 - 1);
@@ -1239,7 +1242,7 @@ void FUniformExpressionSet::FillUniformBuffer(const FMaterialRenderContext& Mate
 					}
 
 					// Zero out any additional components that exist in the destination buffer
-					int32 NumDestComponents = (int32)FirstFieldType - ((int32)EValueType::Float1 - 1);
+					int32 NumDestComponents = (int32)FieldType - ((int32)EValueType::Float1 - 1);
 					for (int32 i = NumCopyComponents; i < NumDestComponents; ++i)
 					{
 						*DestAddress++ = 0.0f;
@@ -1247,21 +1250,20 @@ void FUniformExpressionSet::FillUniformBuffer(const FMaterialRenderContext& Mate
 				}
 				else
 				{
-					const FMaterialUniformPreshaderField& PreshaderField = UniformPreshaderFields[Preshader.FieldIndex + FieldIndex];
-					const FValueTypeDescription UniformTypeDesc = GetValueTypeDescription(PreshaderField.Type);
+					const FValueTypeDescription UniformTypeDesc = GetValueTypeDescription(Preshader.Type);
 					const int32 NumFieldComponents = UniformTypeDesc.NumComponents;
 
-					const int32 NumResultComponents = FMath::Min<int32>(NumFieldComponents, Result.Component.Num() - PreshaderField.ComponentIndex);
+					const int32 NumResultComponents = FMath::Min<int32>(NumFieldComponents, Result.Component.Num());
 					check(NumResultComponents > 0);
 
 					// The type generated by the preshader might not match the expected type
-					FValue FieldValue(Result.Type.GetComponentType(PreshaderField.ComponentIndex), NumResultComponents);
+					FValue FieldValue(Result.Type.GetComponentType(0), NumResultComponents);
 					for (int32 i = 0; i < NumResultComponents; ++i)
 					{
-						FieldValue.Component[i] = Result.Component[PreshaderField.ComponentIndex + i];
+						FieldValue.Component[i] = Result.Component[i];
 					}
-					
-					CopyValueToUniformBuffer(FieldValue, PreshaderBuffer, PreshaderField.BufferOffset);
+
+					CopyValueToUniformBuffer(FieldValue, PreshaderBuffer, Preshader.BufferOffset);
 				}
 			}
 		}
diff --git a/Engine/Source/Runtime/Engine/Private/Shader/Preshader.cpp b/Engine/Source/Runtime/Engine/Private/Shader/Preshader.cpp
--- a/Engine/Source/Runtime/Engine/Private/Shader/Preshader.cpp
+++ b/Engine/Source/Runtime/Engine/Private/Shader/Preshader.cpp
@@ -91,30 +91,14 @@ struct FPreshaderDebugStack
 	}
 };
 
-FPreshaderType::FPreshaderType(const FType& InType) : ValueType(InType.ValueType)
-{
-	if (InType.IsStruct())
-	{
-		StructTypeHash = InType.StructType->Hash;
-		StructComponentTypes = InType.StructType->ComponentTypes;
-	}
-}
-
 FPreshaderType::FPreshaderType(EValueType InType) : ValueType(InType)
 {
 }
 
 EValueComponentType FPreshaderType::GetComponentType(int32 Index) const
 {
-	if (IsStruct())
-	{
-		return StructComponentTypes.IsValidIndex(Index) ? StructComponentTypes[Index] : EValueComponentType::Void;
-	}
-	else
-	{
-		const FValueTypeDescription& TypeDesc = GetValueTypeDescription(ValueType);
-		return (Index >= 0 && Index < TypeDesc.NumComponents) ? TypeDesc.ComponentType : EValueComponentType::Void;
-	}
+	const FValueTypeDescription& TypeDesc = GetValueTypeDescription(ValueType);
+	return (Index >= 0 && Index < TypeDesc.NumComponents) ? TypeDesc.ComponentType : EValueComponentType::Void;
 }
 
 void FPreshaderStack::PushValue(const FValue& InValue)
@@ -205,21 +189,8 @@ FPreshaderValue FPreshaderStack::PeekValue(int32 Offset)
 FValue FPreshaderValue::AsShaderValue(const FStructTypeRegistry* TypeRegistry) const
 {
 	FValue Result;
-	if (!Type.IsStruct())
-	{
-		Result.Type = Type.ValueType;
-		Result.Component.Append(Component.GetData(), Component.Num());
-	}
-	else if (ensure(TypeRegistry))
-	{
-		Result.Type = TypeRegistry->FindType(Type.StructTypeHash);
-		if (Result.Type.IsStruct())
-		{
-			check(Result.Type.GetNumComponents() == Component.Num());
-			Result.Component.Append(Component.GetData(), Component.Num());
-		}
-	}
-
+	Result.Type = Type.ValueType;
+	Result.Component.Append(Component.GetData(), Component.Num());
 	return Result;
 }
 
@@ -240,43 +211,13 @@ void FPreshaderData::WriteName(const FScriptName& Name)
 	Write((uint16)Index);
 }
 
-void FPreshaderData::WriteType(const FType& Type)
-{
-	Write(Type.ValueType);
-	if (Type.IsStruct())
-	{
-		const uint64 Hash = Type.StructType->Hash;
-		int32 Index = INDEX_NONE;
-		for (int32 PrevIndex = 0; PrevIndex < StructTypes.Num(); ++PrevIndex)
-		{
-			if (StructTypes[PrevIndex].Hash == Hash)
-			{
-				Index = PrevIndex;
-				break;
-			}
-		}
-		if (Index == INDEX_NONE)
-		{
-			Index = StructTypes.Num();
-			FPreshaderStructType& PreshaderStructType = StructTypes.AddDefaulted_GetRef();
-			PreshaderStructType.Hash = Hash;
-			PreshaderStructType.ComponentTypeIndex = StructComponentTypes.Num();
-			PreshaderStructType.NumComponents = Type.StructType->ComponentTypes.Num();
-			StructComponentTypes.Append(Type.StructType->ComponentTypes.GetData(), Type.StructType->ComponentTypes.Num());
-		}
-
-		check(Index >= 0 && Index <= 0xffff);
-		Write((uint16)Index);
-	}
-}
-
 void FPreshaderData::WriteValue(const FValue& Value)
 {
 	const int32 NumComponents = Value.Type.GetNumComponents();
 
-	check(!Value.Type.IsObject());
+	check(!Value.Type.IsObject() && !Value.Type.IsStruct());
 
-	WriteType(Value.Type);
+	Write(Value.Type.ValueType);
 	for (int32 Index = 0; Index < NumComponents; ++Index)
 	{
 		const EValueComponentType ComponentType = Value.Type.GetComponentType(Index);
@@ -286,54 +227,21 @@ void FPreshaderData::WriteValue(const FValue& Value)
 	}
 }
 
-FPreshaderLabel FPreshaderData::WriteJump(EPreshaderOpcode Op)
-{
-	WriteOpcode(Op);
-	const int32 Offset = Data.Num();
-	Write((uint32)0xffffffff); // Write a placeholder for the jump offset
-	return FPreshaderLabel(Offset);
-}
-
-void FPreshaderData::WriteJump(EPreshaderOpcode Op, FPreshaderLabel Label)
-{
-	WriteOpcode(Op);
-	const int32 Offset = Data.Num();
-	const int32 JumpOffset = Label.Offset - Offset - 4; // Compute the offset to jump
-	Write(JumpOffset);
-}
-
-void FPreshaderData::SetLabel(FPreshaderLabel InLabel)
-{
-	const int32 TargetOffset = Data.Num();
-	const int32 BaseOffset = InLabel.Offset;
-	const int32 JumpOffset = TargetOffset - BaseOffset - 4; // Compute the offset to jump
-	check(JumpOffset >= 0);
-
-	uint8* Dst = &Data[BaseOffset];
-	check(Dst[0] == 0xff && Dst[1] == 0xff && Dst[2] == 0xff && Dst[3] == 0xff);
-	FMemory::Memcpy(Dst, &JumpOffset, 4); // Patch the offset into the jump opcode
-}
-
-FPreshaderLabel FPreshaderData::GetLabel()
-{
-	return FPreshaderLabel(Data.Num());
-}
-
-FPreshaderDataContext::FPreshaderDataContext(const FPreshaderData& InData)
+FPreshaderDataContext::FPreshaderDataContext(const FPreshaderData& InData, TArrayView<float> InResultsBuffer)
 	: Ptr(InData.Data.GetData())
 	, EndPtr(Ptr + InData.Data.Num())
 	, Names(InData.Names)
-	, StructTypes(InData.StructTypes)
-	, StructComponentTypes(InData.StructComponentTypes)
+	, ResultsBuffer(InResultsBuffer)
 {}
 
 FPreshaderDataContext::FPreshaderDataContext(const FPreshaderDataContext& InContext, uint32 InOffset, uint32 InSize)
 	: Ptr(InContext.Ptr + InOffset)
 	, EndPtr(Ptr + InSize)
 	, Names(InContext.Names)
-	, StructTypes(InContext.StructTypes)
-	, StructComponentTypes(InContext.StructComponentTypes)
-{}
+{
+	// Preshader2 doesn't include individual preshaders, just a single bytecode stream, so this function shouldn't be called on those.
+	check(InContext.IsPreshader2() == false);
+}
 
 static inline void ReadPreshaderData(FPreshaderDataContext& RESTRICT Data, int32 Size, void* Result)
 {
@@ -369,13 +277,6 @@ FPreshaderType ReadPreshaderValue<FPreshaderType>(FPreshaderDataContext& RESTRIC
 {
 	FPreshaderType Result;
 	Result.ValueType = ReadPreshaderValue<EValueType>(Data);
-	if (Result.ValueType == EValueType::Struct)
-	{
-		const uint16 Index = ReadPreshaderValue<uint16>(Data);
-		const FPreshaderStructType& PreshaderStruct = Data.StructTypes[Index];
-		Result.StructTypeHash = PreshaderStruct.Hash;
-		Result.StructComponentTypes = MakeArrayView(Data.StructComponentTypes.GetData() + PreshaderStruct.ComponentTypeIndex, PreshaderStruct.NumComponents);
-	}
 	return Result;
 }
 
@@ -417,24 +318,12 @@ static void EvaluateConstant(FPreshaderStack& Stack, FPreshaderDataContext& REST
 {
 	const FPreshaderType Type = ReadPreshaderValue<FPreshaderType>(Data);
 	TArrayView<FValueComponent> Component = Stack.PushEmptyValue(Type);
-	if (!Type.IsStruct())
-	{
-		// Common case:  not a struct
-		const FValueTypeDescription& TypeDesc = GetValueTypeDescription(Type.ValueType);
-		const int32 ComponmentSizeInBytes = TypeDesc.ComponentSizeInBytes;
-		for (int32 Index = 0; Index < Component.Num(); ++Index)
-		{
-			ReadPreshaderData(Data, ComponmentSizeInBytes, &Component[Index].Packed);
-		}
-	}
-	else
+
+	const FValueTypeDescription& TypeDesc = GetValueTypeDescription(Type.ValueType);
+	const int32 ComponmentSizeInBytes = TypeDesc.ComponentSizeInBytes;
+	for (int32 Index = 0; Index < Component.Num(); ++Index)
 	{
-		for (int32 Index = 0; Index < Component.Num(); ++Index)
-		{
-			const EValueComponentType ComponentType = Type.GetComponentType(Index);
-			const int32 ComponmentSizeInBytes = GetComponentTypeSizeInBytes(ComponentType);
-			ReadPreshaderData(Data, ComponmentSizeInBytes, &Component[Index].Packed);
-		}
+		ReadPreshaderData(Data, ComponmentSizeInBytes, &Component[Index].Packed);
 	}
 }
 
@@ -454,93 +343,6 @@ static void EvaluateConstant(FPreshaderDebugStack& Stack, FPreshaderDataContext&
 	}
 }
 
-static void EvaluateSetField(FPreshaderStack& Stack, FPreshaderDataContext& RESTRICT Data)
-{
-	const FPreshaderValue Value = Stack.PopValue();
-	const FPreshaderValue StructValue = Stack.PeekValue();
-
-	const int32 ComponentIndex = ReadPreshaderValue<int32>(Data);
-	const int32 ComponentNum = ReadPreshaderValue<int32>(Data);
-	
-	// Modify the struct value in-place
-	if (Value.Component.Num() == 1)
-	{
-		// Splat scalar
-		const FValueComponent Component = Value.Component[0];
-		for (int32 Index = 0; Index < ComponentNum; ++Index)
-		{
-			StructValue.Component[ComponentIndex + Index] = Component;
-		}
-	}
-	else
-	{
-		const int32 NumComponentsToCopy = FMath::Min(ComponentNum, Value.Component.Num());
-		for (int32 Index = 0; Index < NumComponentsToCopy; ++Index)
-		{
-			StructValue.Component[ComponentIndex + Index] = Value.Component[Index];
-		}
-		for (int32 Index = NumComponentsToCopy; Index < ComponentNum; ++Index)
-		{
-			StructValue.Component[ComponentIndex + Index] = FValueComponent();
-		}
-	}
-}
-
-static void EvaluateSetField(FPreshaderDebugStack& Stack, FPreshaderDataContext& RESTRICT Data)
-{
-	Stack.OpCount++;
-
-	const int32 ComponentIndex = ReadPreshaderValue<int32>(Data);
-	const int32 ComponentNum = ReadPreshaderValue<int32>(Data);
-
-	if (Stack.bGenerateString)
-	{
-		FString Value = Stack.Elements.Pop().Text;
-
-		Stack.Elements.Last().Text = FString::Printf(TEXT("SetField(%s[%d,%d], %s)"), *Stack.Elements.Last().Text, ComponentIndex, ComponentNum, *Value);
-		Stack.Elements.Last().DebugStackType = EPreshaderDebugStackType::Singular;
-	}
-}
-
-static void EvaluateGetField(FPreshaderStack& Stack, FPreshaderDataContext& RESTRICT Data)
-{
-	const FPreshaderValue StructValue = Stack.PopValue();
-	const FPreshaderType FieldType = ReadPreshaderValue<FPreshaderType>(Data);
-	const int32 ComponentIndex = ReadPreshaderValue<int32>(Data);
-	const int32 ComponentNum = FieldType.GetNumComponents();
-
-	// Need to make a local copy of components, since StructValue.Component is only valid until we push the result
-	TArray<FValueComponent, TInlineAllocator<64>> FieldComponents;
-	FieldComponents.Empty(ComponentNum);
-	for (int32 Index = 0; Index < ComponentNum; ++Index)
-	{
-		FieldComponents.Add(StructValue.Component[ComponentIndex + Index]);
-	}
-	Stack.PushValue(FieldType, FieldComponents);
-}
-
-static void EvaluateGetField(FPreshaderDebugStack& Stack, FPreshaderDataContext& RESTRICT Data)
-{
-	Stack.OpCount++;
-
-	const FPreshaderType FieldType = ReadPreshaderValue<FPreshaderType>(Data);
-	const int32 ComponentIndex = ReadPreshaderValue<int32>(Data);
-	const int32 ComponentNum = FieldType.GetNumComponents();
-
-	if (Stack.bGenerateString)
-	{
-		if (Stack.Elements.Last().DebugStackType == EPreshaderDebugStackType::Singular)
-		{
-			Stack.Elements.Last().Text.Append(FString::Printf(TEXT("[%d]"), ComponentIndex));
-		}
-		else
-		{
-			Stack.Elements.Last().Text = FString::Printf(TEXT("(%s)[%d]"), *Stack.Elements.Last().Text, ComponentIndex);
-			Stack.Elements.Last().DebugStackType = EPreshaderDebugStackType::Singular;
-		}
-	}
-}
-
 static void EvaluatePushValue(FPreshaderStack& Stack, FPreshaderDataContext& RESTRICT Data)
 {
 	const int32 StackOffset = ReadPreshaderValue<uint16>(Data);
@@ -816,10 +618,7 @@ static inline void EvaluateUnaryOpInPlace(FPreshaderStack& Stack, const Operatio
 	// to a loose EValueType and array of FValueComponent, as opposed to requiring those two things to be copied into an FValue first.
 	// It's possible for any given operation to modify the type, so we override the type of the top item based on the return value.
 	FPreshaderValue Top = Stack.PeekValue();
-	if (!Top.Type.IsStruct())
-	{
-		Stack.OverrideTopType(Op(Top.Type.ValueType, Top.Component));
-	}
+	Stack.OverrideTopType(Op(Top.Type.ValueType, Top.Component));
 }
 
 static inline void EvaluateUnaryOpInPlaceDebug(FPreshaderDebugStack& Stack, const TCHAR* Op)
@@ -913,18 +712,11 @@ static inline void EvaluateBinaryOpInPlace(FPreshaderStack& Stack, const Operati
 	// Lhs is one below top of stack, Rhs is top
 	FPreshaderValue Lhs = Stack.PeekValue(1);
 	FPreshaderValue Rhs = Stack.PeekValue(0);
-	if (Lhs.Type.IsStruct() || Rhs.Type.IsStruct())
-	{
-		// Throw up our hands if one of the inputs is a struct, and run the old logic.  Probably a bug for a struct to be passed to a BinaryOp.
-		EvaluateBinaryOp(Stack, FallbackOp);
-	}
-	else
-	{
-		// Initialize to avoid static analysis warning -- all code paths do fill in ComponentsConsumed
-		int32 ComponentsConsumed = 0;
-		EValueType ResultType = Op(Lhs.Type.ValueType, Rhs.Type.ValueType, TArrayView<FValueComponent>(Lhs.Component.GetData(), Lhs.Component.Num() + Rhs.Component.Num()), ComponentsConsumed);
-		Stack.MergeTopTwoValues(ResultType, ComponentsConsumed);
-	}
+
+	// Initialize to avoid static analysis warning -- all code paths do fill in ComponentsConsumed
+	int32 ComponentsConsumed = 0;
+	EValueType ResultType = Op(Lhs.Type.ValueType, Rhs.Type.ValueType, TArrayView<FValueComponent>(Lhs.Component.GetData(), Lhs.Component.Num() + Rhs.Component.Num()), ComponentsConsumed);
+	Stack.MergeTopTwoValues(ResultType, ComponentsConsumed);
 }
 
 template<typename Operation>
@@ -956,7 +748,6 @@ static void EvaluateComponentSwizzle(FPreshaderStack& Stack, FPreshaderDataConte
 	// the elements, so we can still access them, but saves time over having to re-fetch the component pointer
 	// in case the component array gets resized between the read and write.
 	const FPreshaderType& OriginalType = Stack.PeekType();
-	check(OriginalType.IsStruct() == false);
 
 	const int32 OriginalNumComponents = GetValueTypeDescription(OriginalType.ValueType).NumComponents;
 	const int32 NumElements = (int32)ReadPreshaderValue<uint8>(Data);
@@ -1282,40 +1073,6 @@ static void EvaluateSparseTextureUniform(const FMaterialRenderContext& Context,
 	}
 }
 
-static void EvaluateJump(FPreshaderDataContext& RESTRICT Data)
-{
-	const int32 JumpOffset = ReadPreshaderValue<int32>(Data);
-	check(Data.Ptr + JumpOffset <= Data.EndPtr);
-	Data.Ptr += JumpOffset;
-}
-
-static void EvaluateJumpIfFalse(FPreshaderStack& Stack, FPreshaderDataContext& RESTRICT Data)
-{
-	const int32 JumpOffset = ReadPreshaderValue<int32>(Data);
-	check(Data.Ptr + JumpOffset <= Data.EndPtr);
-
-	const FValue ConditionValue = Stack.PopValue().AsShaderValue();
-	if (!ConditionValue.AsBoolScalar())
-	{
-		Data.Ptr += JumpOffset;
-	}
-}
-
-static void EvaluateJumpIfFalse(FPreshaderDebugStack& Stack, FPreshaderDataContext& RESTRICT Data)
-{
-	Stack.OpCount++;
-
-	const int32 JumpOffset = ReadPreshaderValue<int32>(Data);
-	check(Data.Ptr + JumpOffset <= Data.EndPtr);
-
-	// It's complicated to include conditional jumps in a stack evaluator.  Ideally we'd want to evaluate and display both
-	// branches, but that requires evaluating both branches forward to where the branches reach the same instruction
-	// again, then converting stack entries that diverge between the two code paths to include a conditional expression.
-
-	// The new translator can generate conditionals, until as above is implemented, ignore the conditional.
-	// This may elevate the reported count as actually exclusive branches may both be debug evaluated.
-}
-
 FPreshaderValue EvaluatePreshader(const FUniformExpressionSet* UniformExpressionSet, const FMaterialRenderContext& Context, FPreshaderStack& Stack, FPreshaderDataContext& RESTRICT Data)
 {
 	uint8 const* const DataEnd = Data.EndPtr;
@@ -1328,8 +1085,6 @@ FPreshaderValue EvaluatePreshader(const FUniformExpressionSet* UniformExpression
 		{
 		case EPreshaderOpcode::ConstantZero: EvaluateConstantZero(Stack, Data); break;
 		case EPreshaderOpcode::Constant: EvaluateConstant(Stack, Data); break;
-		case EPreshaderOpcode::GetField: EvaluateGetField(Stack, Data); break;
-		case EPreshaderOpcode::SetField: EvaluateSetField(Stack, Data); break;
 		case EPreshaderOpcode::Parameter:
 			EvaluateParameter(Stack, UniformExpressionSet, ReadPreshaderValue<uint16>(Data), Context);
 			break;
@@ -1384,8 +1139,6 @@ FPreshaderValue EvaluatePreshader(const FUniformExpressionSet* UniformExpression
 		case EPreshaderOpcode::ExternalTextureCoordinateOffset: EvaluateExternalTextureCoordinateOffset(Context, Stack, Data); break;
 		case EPreshaderOpcode::RuntimeVirtualTextureUniform: EvaluateRuntimeVirtualTextureUniform(Context, Stack, Data); break;
 		case EPreshaderOpcode::SparseVolumeTextureUniform: EvaluateSparseTextureUniform(Context, Stack, Data); break;
-		case EPreshaderOpcode::Jump: EvaluateJump(Data); break;
-		case EPreshaderOpcode::JumpIfFalse: EvaluateJumpIfFalse(Stack, Data); break;
 		default:
 			UE_LOG(LogMaterial, Fatal, TEXT("Unknown preshader opcode %d"), (uint8)Opcode);
 			break;
@@ -1419,8 +1172,6 @@ static void EvaluatePreshaderDebug(const FUniformExpressionSet* UniformExpressio
 		{
 		case EPreshaderOpcode::ConstantZero: EvaluateConstantZero(Stack, Data); break;
 		case EPreshaderOpcode::Constant: EvaluateConstant(Stack, Data); break;
-		case EPreshaderOpcode::GetField: EvaluateGetField(Stack, Data); break;
-		case EPreshaderOpcode::SetField: EvaluateSetField(Stack, Data); break;
 		case EPreshaderOpcode::Parameter:
 			EvaluateParameter(Stack, UniformExpressionSet, ReadPreshaderValue<uint16>(Data), Context);
 			break;
@@ -1475,8 +1226,6 @@ static void EvaluatePreshaderDebug(const FUniformExpressionSet* UniformExpressio
 		case EPreshaderOpcode::ExternalTextureCoordinateOffset: EvaluateExternalTextureCoordinateOffset(Context, Stack, Data); break;
 		case EPreshaderOpcode::RuntimeVirtualTextureUniform: EvaluateRuntimeVirtualTextureUniform(Context, Stack, Data); break;
 		case EPreshaderOpcode::SparseVolumeTextureUniform: EvaluateSparseTextureUniform(Context, Stack, Data); break;
-		case EPreshaderOpcode::Jump: EvaluateJump(Data); break;
-		case EPreshaderOpcode::JumpIfFalse: EvaluateJumpIfFalse(Stack, Data); break;
 		default:
 			UE_LOG(LogMaterial, Fatal, TEXT("Unknown preshader opcode %d"), (uint8)Opcode);
 			break;
@@ -1524,9 +1273,8 @@ FPreshaderValue FPreshaderData::EvaluateConstant(const FMaterial& Material, FPre
 
 void FPreshaderData::AppendHash(FXxHash64Builder& OutHasher) const
 {
+	OutHasher.Update(&bPreshader2, sizeof(bPreshader2));
 	OutHasher.Update(Names.GetData(), Names.Num() * Names.GetTypeSize());
-	OutHasher.Update(StructTypes.GetData(), StructTypes.Num() * StructTypes.GetTypeSize());
-	OutHasher.Update(StructComponentTypes.GetData(), StructComponentTypes.Num() * StructComponentTypes.GetTypeSize());
 	OutHasher.Update(Data.GetData(), Data.Num() * Data.GetTypeSize());
 }
 
diff --git a/Engine/Source/Runtime/Engine/Private/Shader/Preshader2.cpp b/Engine/Source/Runtime/Engine/Private/Shader/Preshader2.cpp
new file mode 100644
--- /dev/null
+++ b/Engine/Source/Runtime/Engine/Private/Shader/Preshader2.cpp
@@ -0,0 +1,2016 @@
+// Copyright Epic Games, Inc. All Rights Reserved.
+
+#include "Shader/Preshader2.h"
+#include "Materials/MaterialRenderProxy.h"
+#include "Materials/MaterialUniformExpressions.h"
+#include "Math/DoubleFloat.h"
+#include "Misc/LargeWorldRenderPosition.h"
+#include "SparseVolumeTexture/SparseVolumeTexture.h"
+#include "VT/RuntimeVirtualTexture.h"
+
+namespace UE::Preshader2
+{
+
+#define PRESHADER2_UNARY_OP(OP_NAME)	\
+	TEXT(#OP_NAME),						\
+	TEXT(#OP_NAME "_InPlace"),			\
+	TEXT(#OP_NAME "_Reuse"),
+
+#define PRESHADER2_BINARY_OP(OP_NAME)	\
+	TEXT(#OP_NAME),						\
+	TEXT(#OP_NAME "_InPlace"),			\
+	TEXT(#OP_NAME "_Reuse"),			\
+	TEXT(#OP_NAME "_Const"),			\
+	TEXT(#OP_NAME "_ConstInPlace"),		\
+	TEXT(#OP_NAME "_ConstReuse"),		\
+	TEXT(#OP_NAME "_ConstReuseScalar"),
+
+#define PRESHADER2_UNARY_OP_CAST_REUSE(OP_NAME) TEXT(#OP_NAME),
+#define PRESHADER2_UNARY_OP_CUSTOM(OP_NAME) TEXT(#OP_NAME),
+#define PRESHADER2_BINARY_OP_CUSTOM(OP_NAME) TEXT(#OP_NAME),
+#define PRESHADER2_OP(OP_NAME) TEXT(#OP_NAME),
+
+static const TCHAR* GPreshader2OpcodeNames[] =
+{
+	TEXT("Invalid"),
+	#include "Shader/Preshader2Opcodes.inl"
+};
+
+#undef PRESHADER2_UNARY_OP
+#undef PRESHADER2_BINARY_OP
+#undef PRESHADER2_UNARY_OP_CAST_REUSE
+#undef PRESHADER2_UNARY_OP_CUSTOM
+#undef PRESHADER2_BINARY_OP_CUSTOM
+#undef PRESHADER2_OP
+
+
+struct FOpAbs { template<typename T> static T Call(T Value) { return FMath::Abs(Value); } };
+struct FOpACos { template<typename T> static T Call(T Value) { return FMath::Acos(Value); } };
+struct FOpASin { template<typename T> static T Call(T Value) { return FMath::Asin(Value); } };
+struct FOpATan { template<typename T> static T Call(T Value) { return FMath::Atan(Value); } };
+struct FOpCeil { template<typename T> static T Call(T Value) { return FMath::CeilToFloat(Value); } };
+struct FOpCos { template<typename T> static T Call(T Value) { return FMath::Cos(Value); } };
+struct FOpExponential { template<typename T> static T Call(T Value) { return FMath::Exp(Value); } };
+struct FOpExponential2 { template<typename T> static T Call(T Value) { return FMath::Exp2(Value); } };
+struct FOpFloor { template<typename T> static T Call(T Value) { return FMath::FloorToFloat(Value); } };
+struct FOpFrac { template<typename T> static T Call(T Value) { return FMath::Frac(Value); } };
+struct FOpFractional { template<typename T> static T Call(T Value) { return FMath::Fractional(Value); } };
+struct FOpLogarithm { template<typename T> static T Call(T Value) { return FMath::Loge(Value); } };
+struct FOpLogarithm2 { template<typename T> static T Call(T Value) { return FMath::Log2(Value); } };
+struct FOpLogarithm10 { template<typename T> static T Call(T Value)
+{
+	static const T LogToLog10 = (T)1.0 / FMath::Loge((T)10);
+	return FMath::Loge(Value) * LogToLog10;
+}};
+struct FOpNegate { template<typename T> static T Call(T Value) { return -Value; } };
+struct FOpReciprocal { template<typename T> static T Call(T Value)
+{
+	// Generate a safe divisor.
+	if (FMath::Abs(Value) < (T)UE_DELTA)
+	{
+		Value = (Value < 0) ? -(T)UE_DELTA : (T)UE_DELTA;
+	}
+	return (T)1 / Value;
+}};
+struct FOpRound { template<typename T> static T Call(T Value) { return FMath::RoundToFloat(Value); } };
+struct FOpSaturate { template<typename T> static T Call(T Value) { return FMath::Clamp(Value, (T)0, (T)1); } };
+struct FOpSign { template<typename T> static T Call(T Value) { return FMath::Sign(Value); } };
+struct FOpSin { template<typename T> static T Call(T Value) { return FMath::Sin(Value); } };
+struct FOpSqrt { template<typename T> static T Call(T Value) { return FMath::Sqrt(Value); } };
+struct FOpTan { template<typename T> static T Call(T Value) { return FMath::Tan(Value); } };
+struct FOpTruncate { template<typename T> static T Call(T Value) { return FMath::TruncToFloat(Value); } };
+
+struct FOpAdd { template<typename T> static T Call(T A, T B) { return A + B; } };
+struct FOpSubtract { template<typename T> static T Call(T A, T B) { return A - B; } };
+struct FOpMultiply { template<typename T> static T Call(T A, T B) { return A * B; } };
+struct FOpDivide { template<typename T> static T Call(T A, T B) { return A / B; } };
+struct FOpMin { template<typename T> static T Call(T A, T B) { return FMath::Min(A, B); } };
+struct FOpMax { template<typename T> static T Call(T A, T B) { return FMath::Max(A, B); } };
+struct FOpATan2 { template<typename T> static T Call(T A, T B) { return FMath::Atan2(A, B); } };
+
+// Note:  we are not defining a separate opcode for Fmod and Modulo.  Operations on integer will do Modulo, while operations on float will do Fmod.
+struct FOpFmod
+{
+	static float Call(float A, float B) { return FMath::Fmod(A, B); }
+	static double Call(double A, double B) { return FMath::Fmod(A, B); }
+	static int32 Call(int32 A, int32 B) { return A % B; }
+};
+
+// Ops that aren't per-component.  These are dummy structures for template specialization or overloading, and don't need an operator implementation.
+// Unary.
+struct FOpCast {};
+struct FOpLength {};
+struct FOpNormalize {};
+
+// Binary.
+struct FOpCross {};
+struct FOpDot {};
+
+// Custom.
+struct FOpConstant {};
+struct FOpMove {};
+
+// External Inputs.
+struct FOpExternalTextureCoordinateOffset {};
+struct FOpExternalTextureCoordinateScaleRotation {};
+struct FOpRuntimeVirtualTextureUniform {};
+struct FOpSparseVolumeTextureUniform {};
+struct FOpTexelSize {};
+struct FOpTextureSize {};
+
+static float& AsFloat(void* Buffer, int32 Offset, int32 Index)
+{
+	return ((float*)Buffer)[Offset + Index];
+}
+
+static double& AsDouble(void* Buffer, int32 Offset, int32 Index)
+{
+	// Offset is always in 32-bit words, while Index applies to the double type, so the Index gets doubled.
+	return (double&)((int32*)Buffer)[Offset + Index*2];
+}
+
+static int32& AsInt(void* Buffer, int32 Offset, int32 Index)
+{
+	return ((int32*)Buffer)[Offset + Index];
+}
+
+static uint16 FetchOffset(const uint8* Ptr)
+{
+	return (uint16)Ptr[0] | ((uint16)Ptr[1] << 8);
+}
+
+// Increments OpcodePtr by the size of T.
+template<typename T>
+void FetchDataValue(const uint8*& OpcodePtr, T& Value)
+{
+	FMemory::Memcpy(&Value, OpcodePtr, sizeof(T));
+	OpcodePtr += sizeof(T);
+}
+
+// Bounds checking utility functions -- we want to bounds check once, rather than per component in vector operations, to reduce overhead in non-shipping builds.
+static FORCEINLINE void BoundsCheck(TArrayView<float> Buffer, int32 Offset, int32 Dimension, bool bDouble = false)
+{
+	if (bDouble)
+	{
+		Dimension *= 2;
+	}
+	check(Offset + Dimension <= Buffer.Num());
+}
+
+static FORCEINLINE void BoundsCheck(TArrayView<float> Buffer, int32 ResultOffset, int32 AOffset, int32 Dimension, bool bDouble = false)
+{
+	if (bDouble)
+	{
+		Dimension *= 2;
+	}
+	check(FMath::Max(AOffset, ResultOffset) + Dimension <= Buffer.Num());
+}
+
+static FORCEINLINE void BoundsCheck(TArrayView<float> Buffer, int32 ResultOffset, int32 AOffset, int32 BOffset, int32 Dimension, bool bAScalar = false, bool bBScalar = false, bool bDouble = false)
+{
+	if (bDouble)
+	{
+		Dimension *= 2;
+	}
+	check(FMath::Max3(ResultOffset + Dimension, AOffset + (bAScalar ? 1 : Dimension), BOffset + (bBScalar ? 1 : Dimension)) <= Buffer.Num());
+}
+
+// UnaryEvaluate accepts Type as a reference, as the Normalize op changes the type of its output to float.
+template<typename TOp, bool bAllowInt>
+void UnaryEvaluate(TArrayView<float> Buffer, EPreshader2Type& Type, int32 Dimension, int32 ResultOffset, int32 AOffset)
+{
+	// By design, the macros that implement unary ops pass the Dimension byte unmodified, to allow clients to store data in it.
+	// Mask the low bits to get the actual dimension.
+	int32 LocalDimension = Dimension & 0x7;
+	float* BufferData = Buffer.GetData();
+
+	if (LIKELY(Type == EPreshader2Type::Float))
+	{
+		BoundsCheck(Buffer, ResultOffset, AOffset, LocalDimension);
+
+		// This binary tree based on the dimension generates more compact code than a "for" loop for float / double types.
+		float A = TOp::Call(AsFloat(BufferData, AOffset, 0));
+		if (LocalDimension >= 3)
+		{
+			float B = TOp::Call(AsFloat(BufferData, AOffset, 1));
+			float C = TOp::Call(AsFloat(BufferData, AOffset, 2));
+			if (LocalDimension == 4)
+			{
+				AsFloat(BufferData, ResultOffset, 3) = TOp::Call(AsFloat(BufferData, AOffset, 3));
+			}
+			AsFloat(BufferData, ResultOffset, 1) = B;
+			AsFloat(BufferData, ResultOffset, 2) = C;
+		}
+		else if (LocalDimension == 2)
+		{
+			AsFloat(BufferData, ResultOffset, 1) = TOp::Call(AsFloat(BufferData, AOffset, 1));
+		}
+		AsFloat(BufferData, ResultOffset, 0) = A;
+	}
+	else if (!bAllowInt || Type == EPreshader2Type::Double)
+	{
+		BoundsCheck(Buffer, ResultOffset, AOffset, LocalDimension, /*bDouble=*/ true);
+
+		double A = TOp::Call(AsDouble(BufferData, AOffset, 0));
+		if (LocalDimension >= 3)
+		{
+			double B = TOp::Call(AsDouble(BufferData, AOffset, 1));
+			double C = TOp::Call(AsDouble(BufferData, AOffset, 2));
+			if (LocalDimension == 4)
+			{
+				AsDouble(BufferData, ResultOffset, 3) = TOp::Call(AsDouble(BufferData, AOffset, 3));
+			}
+			AsDouble(BufferData, ResultOffset, 1) = B;
+			AsDouble(BufferData, ResultOffset, 2) = C;
+		}
+		else if (LocalDimension == 2)
+		{
+			AsDouble(BufferData, ResultOffset, 1) = TOp::Call(AsDouble(BufferData, AOffset, 1));
+		}
+		AsDouble(BufferData, ResultOffset, 0) = A;
+	}
+	else if constexpr (bAllowInt)
+	{
+		BoundsCheck(Buffer, ResultOffset, AOffset, LocalDimension);
+
+		for (int32 Index = 0; Index < LocalDimension; ++Index)
+		{
+			AsInt(BufferData, ResultOffset, Index) = TOp::Call(AsInt(BufferData, AOffset, Index));
+		}
+	}
+}
+
+// Specialization of UnaryEvaluate for Normalize
+template <>
+void UnaryEvaluate<FOpNormalize, false>(TArrayView<float> Buffer, EPreshader2Type& Type, int32 Dimension, int32 ResultOffset, int32 AOffset)
+{
+	// By design, the macros that implement unary ops pass the Dimension byte unmodified, to allow clients to store data in it.
+	// Mask the low bits to get the actual dimension.
+	int32 LocalDimension = Dimension & 0x7;
+	float* BufferData = Buffer.GetData();
+
+	if (LIKELY(Type == EPreshader2Type::Float))
+	{
+		BoundsCheck(Buffer, ResultOffset, AOffset, LocalDimension);
+
+		float MagnitudeScale = 0.0f;
+		for (int32 Index = 0; Index < LocalDimension; ++Index)
+		{
+			MagnitudeScale += FMath::Square(AsFloat(BufferData, AOffset, Index));
+		}
+		MagnitudeScale = FMath::InvSqrt(FMath::Max(MagnitudeScale, (float)UE_DELTA));
+
+		for (int32 Index = 0; Index < LocalDimension; ++Index)
+		{
+			AsFloat(BufferData, ResultOffset, Index) = AsFloat(BufferData, AOffset, Index) * MagnitudeScale;
+		}
+	}
+	else
+	{
+		// Result is float, A is double
+		BoundsCheck(Buffer, ResultOffset, LocalDimension);
+		BoundsCheck(Buffer, AOffset, LocalDimension, /*bDouble=*/ true);
+
+		// Note that normalize automatically downcasts double inputs to float!
+		double MagnitudeScale = 0.0f;
+		for (int32 Index = 0; Index < LocalDimension; ++Index)
+		{
+			MagnitudeScale += FMath::Square(AsDouble(BufferData, AOffset, Index));
+		}
+		MagnitudeScale = FMath::InvSqrt(FMath::Max(MagnitudeScale, (double)UE_DELTA));
+
+		for (int32 Index = 0; Index < LocalDimension; ++Index)
+		{
+			AsFloat(BufferData, ResultOffset, Index) = (float)(AsDouble(BufferData, AOffset, Index) * MagnitudeScale);
+		}
+
+		// Tell the caller we changed the type
+		Type = EPreshader2Type::Float;
+	}
+}
+
+// Specialization of UnaryEvaluate for Cast.
+template <>
+void UnaryEvaluate<FOpCast, true>(TArrayView<float> Buffer, EPreshader2Type& Type, int32 Dimension, int32 ResultOffset, int32 AOffset)
+{
+	// By design, the macros that implement unary ops pass the Dimension byte unmodified, to allow clients to store data in it.
+	// Mask the low bits to get the actual dimension.
+	int32 LocalDimension = Dimension & 0x7;
+	float* BufferData = Buffer.GetData();
+
+	// The cast op includes the source type in the Dimension byte.
+	EPreshader2Type DestType = EPreshader2Type((Dimension >> 4) & 0x3);
+
+	if (DestType == EPreshader2Type::Float)
+	{
+		if (Type == EPreshader2Type::Double)
+		{
+			// Result is float, A is double
+			BoundsCheck(Buffer, ResultOffset, LocalDimension);
+			BoundsCheck(Buffer, AOffset, LocalDimension, /*bDouble=*/ true);
+
+			for (int32 Index = 0; Index < LocalDimension; ++Index)
+			{
+				AsFloat(BufferData, ResultOffset, Index) = (float)AsDouble(BufferData, AOffset, Index);
+			}
+		}
+		else
+		{
+			BoundsCheck(Buffer, ResultOffset, AOffset, LocalDimension);
+
+			// EPreshader2Type::Integer
+			for (int32 Index = 0; Index < LocalDimension; ++Index)
+			{
+				AsFloat(BufferData, ResultOffset, Index) = (float)AsInt(BufferData, AOffset, Index);
+			}
+		}
+	}
+	else if (DestType == EPreshader2Type::LWC || DestType == EPreshader2Type::Double)
+	{
+		// We have to consider the possibility of in-place conversion (writing LWC or Double in place to a location that contains
+		// the source data of a smaller type we are reading from), so store values to a temporary array in a first pass.
+		double Temp[4];
+		if (Type == EPreshader2Type::Float)
+		{
+			BoundsCheck(Buffer, AOffset, LocalDimension);
+
+			for (int32 Index = 0; Index < LocalDimension; ++Index)
+			{
+				Temp[Index] = AsFloat(BufferData, AOffset, Index);
+			}
+		}
+		else if (Type == EPreshader2Type::Double)
+		{
+			BoundsCheck(Buffer, AOffset, LocalDimension, /*bDouble=*/ true);
+
+			for (int32 Index = 0; Index < LocalDimension; ++Index)
+			{
+				Temp[Index] = AsDouble(BufferData, AOffset, Index);
+			}
+		}
+		else
+		{
+			// EPreshader2Type::Integer
+			BoundsCheck(Buffer, AOffset, LocalDimension);
+
+			for (int32 Index = 0; Index < LocalDimension; ++Index)
+			{
+				Temp[Index] = AsInt(BufferData, AOffset, Index);
+			}
+		}
+
+		BoundsCheck(Buffer, ResultOffset, LocalDimension, /*bDouble=*/ true);
+
+		if (DestType == EPreshader2Type::LWC)
+		{
+			for (int32 Index = 0; Index < LocalDimension; ++Index)
+			{
+				FLargeWorldRenderScalar LWCValue(Temp[Index]);
+				AsFloat(BufferData, ResultOffset, Index*2 + 0) = LWCValue.GetTile();
+				AsFloat(BufferData, ResultOffset, Index*2 + 1) = LWCValue.GetOffset();
+			}
+		}
+		else
+		{
+			for (int32 Index = 0; Index < LocalDimension; ++Index)
+			{
+				AsDouble(BufferData, ResultOffset, Index) = Temp[Index];
+			}
+		}
+	}
+	else if (DestType == EPreshader2Type::Integer)
+	{
+		if (Type == EPreshader2Type::Float)
+		{
+			BoundsCheck(Buffer, ResultOffset, AOffset, LocalDimension);
+
+			for (int32 Index = 0; Index < LocalDimension; ++Index)
+			{
+				AsInt(BufferData, ResultOffset, Index) = (int32)AsFloat(BufferData, AOffset, Index);
+			}
+		}
+		else
+		{
+			// EPreshader2Type::Double
+			// Result is integer, A is double
+			BoundsCheck(Buffer, ResultOffset, LocalDimension);
+			BoundsCheck(Buffer, AOffset, LocalDimension, /*bDouble=*/ true);
+
+			for (int32 Index = 0; Index < LocalDimension; ++Index)
+			{
+				AsInt(BufferData, ResultOffset, Index) = (int32)AsDouble(BufferData, AOffset, Index);
+			}
+		}
+	}
+
+	// Update Type to DestType, as Type is preserved across loop iterations for Reuse opcode variants.
+	Type = DestType;
+}
+
+// Specialization of UnaryEvaluate for Length.
+static void UnaryEvaluate(const FOpLength&, TArrayView<float> Buffer, EPreshader2Type Type, int32& Dimension, int32 ResultOffset, int32 AOffset)
+{
+	// By design, the macros that implement unary ops pass the Dimension byte unmodified, to allow clients to store data in it.
+	// Mask the low bits to get the actual dimension.
+	int32 LocalDimension = Dimension & 0x7;
+	float* BufferData = Buffer.GetData();
+
+	// Override "Dimension" to 1, indicating output is a scalar!
+	Dimension = 1;
+
+	if (LIKELY(Type == EPreshader2Type::Float))
+	{
+		BoundsCheck(Buffer, AOffset, LocalDimension);
+		BoundsCheck(Buffer, ResultOffset, 1);
+
+		float MagnitudeScale = 0.0f;
+		for (int32 Index = 0; Index < LocalDimension; ++Index)
+		{
+			MagnitudeScale += FMath::Square(AsFloat(BufferData, AOffset, Index));
+		}
+		AsFloat(BufferData, ResultOffset, 0) = FMath::Sqrt(MagnitudeScale);
+	}
+	else
+	{
+		BoundsCheck(Buffer, AOffset, LocalDimension, /*bDouble=*/ true);
+		BoundsCheck(Buffer, ResultOffset, 1, /*bDouble=*/ true);
+
+		double MagnitudeScale = 0.0f;
+		for (int32 Index = 0; Index < LocalDimension; ++Index)
+		{
+			MagnitudeScale += FMath::Square(AsDouble(BufferData, AOffset, Index));
+		}
+		AsDouble(BufferData, ResultOffset, 0) = FMath::Sqrt(MagnitudeScale);
+	}
+}
+
+// Evaluate a standard per component binary op.  AInc and BInc are array increments which allow for either input argument to be a scalar.
+// A zero increment will use the same input component for each item.
+template<typename TOp, bool bAllowInt>
+void BinaryEvaluate(TArrayView<float> Buffer, EPreshader2Type Type, int32 Dimension, int32 ResultOffset, int32 AOffset, int32 BOffset, int32 AInc, int32 BInc)
+{
+	float* BufferData = Buffer.GetData();
+
+	if (LIKELY(Type == EPreshader2Type::Float))
+	{
+		BoundsCheck(Buffer, ResultOffset, AOffset, BOffset, Dimension, AInc == 0, BInc == 0);
+
+		// We do index 0 last, to handle in-place overwriting of result scalar with vector
+		for (int32 Index = 1; Index < Dimension; ++Index)
+		{
+			AsFloat(BufferData, ResultOffset, Index) = TOp::Call(AsFloat(BufferData, AOffset, AInc*Index), AsFloat(BufferData, BOffset, BInc*Index));
+		}
+		AsFloat(BufferData, ResultOffset, 0) = TOp::Call(AsFloat(BufferData, AOffset, 0), AsFloat(BufferData, BOffset, 0));
+	}
+	else if (!bAllowInt || Type == EPreshader2Type::Double)
+	{
+		BoundsCheck(Buffer, ResultOffset, AOffset, BOffset, Dimension, AInc == 0, BInc == 0, /*bDouble=*/ true);
+
+		for (int32 Index = 1; Index < Dimension; ++Index)
+		{
+			AsDouble(BufferData, ResultOffset, Index) = TOp::Call(AsDouble(BufferData, AOffset, AInc*Index), AsDouble(BufferData, BOffset, BInc*Index));
+		}
+		AsDouble(BufferData, ResultOffset, 0) = TOp::Call(AsDouble(BufferData, AOffset, 0), AsDouble(BufferData, BOffset, 0));
+	}
+	else if constexpr (bAllowInt)
+	{
+		BoundsCheck(Buffer, ResultOffset, AOffset, BOffset, Dimension, AInc == 0, BInc == 0);
+
+		for (int32 Index = 1; Index < Dimension; ++Index)
+		{
+			AsInt(BufferData, ResultOffset, Index) = TOp::Call(AsInt(BufferData, AOffset, AInc*Index), AsInt(BufferData, BOffset, BInc*Index));
+		}
+		AsInt(BufferData, ResultOffset, 0) = TOp::Call(AsInt(BufferData, AOffset, 0), AsInt(BufferData, BOffset, 0));
+	}
+}
+
+// Variant of binary evaluate that accepts an immediate constant.
+template<typename TOp, bool bAllowInt>
+void BinaryEvaluateConst(TArrayView<float> Buffer, EPreshader2Type Type, int32 Dimension, int32 ResultOffset, int32 AOffset, const uint8*& BData, int32 AInc, int32 BInc)
+{
+	float* BufferData = Buffer.GetData();
+
+	// Copy data to an aligned buffer
+	double BDataCopy[4];
+	int32 CopyNum = BInc ? Dimension : 1;
+	int32 CopySize = CopyNum * (Type == EPreshader2Type::Double ? 8 : 4);
+	FMemory::Memcpy(BDataCopy, BData, CopySize);
+
+	// Increment input pointer by copy size
+	BData += CopySize;
+
+	if (LIKELY(Type == EPreshader2Type::Float))
+	{
+		BoundsCheck(Buffer, AOffset, AInc ? Dimension : 1);
+		BoundsCheck(Buffer, ResultOffset, Dimension);
+
+		// We do index 0 last, to handle in-place overwriting of result scalar with vector
+		for (int32 Index = 1; Index < Dimension; ++Index)
+		{
+			AsFloat(BufferData, ResultOffset, Index) = TOp::Call(AsFloat(BufferData, AOffset, AInc*Index), AsFloat(BDataCopy, 0, BInc*Index));
+		}
+		AsFloat(BufferData, ResultOffset, 0) = TOp::Call(AsFloat(BufferData, AOffset, 0), AsFloat(BDataCopy, 0, 0));
+	}
+	else if (!bAllowInt || Type == EPreshader2Type::Double)
+	{
+		BoundsCheck(Buffer, AOffset, AInc ? Dimension : 1, /*bDouble=*/ true);
+		BoundsCheck(Buffer, ResultOffset, Dimension, /*bDouble=*/ true);
+
+		for (int32 Index = 1; Index < Dimension; ++Index)
+		{
+			AsDouble(BufferData, ResultOffset, Index) = TOp::Call(AsDouble(BufferData, AOffset, AInc*Index), AsDouble(BDataCopy, 0, BInc*Index));
+		}
+		AsDouble(BufferData, ResultOffset, 0) = TOp::Call(AsDouble(BufferData, AOffset, 0), AsDouble(BDataCopy, 0, 0));
+	}
+	else if constexpr (bAllowInt)
+	{
+		BoundsCheck(Buffer, AOffset, AInc ? Dimension : 1);
+		BoundsCheck(Buffer, ResultOffset, Dimension);
+
+		for (int32 Index = 1; Index < Dimension; ++Index)
+		{
+			AsInt(BufferData, ResultOffset, Index) = TOp::Call(AsInt(BufferData, AOffset, AInc*Index), AsInt(BDataCopy, 0, BInc*Index));
+		}
+		AsInt(BufferData, ResultOffset, 0) = TOp::Call(AsInt(BufferData, AOffset, 0), AsInt(BDataCopy, 0, 0));
+	}
+}
+
+template<typename T>
+void CrossCompute(T* Result, T* A, T* B)
+{
+	T Temp[3];
+	Temp[0] = A[1] * B[2] - A[2] * B[1];
+	Temp[1] = A[2] * B[0] - A[0] * B[2];
+	Temp[2] = A[0] * B[1] - A[1] * B[0];
+	Result[0] = Temp[0];
+	Result[1] = Temp[1];
+	Result[2] = Temp[2];
+}
+
+// Specialization of BinaryEvaluate for Cross.
+static void BinaryEvaluate(const FOpCross&, TArrayView<float> Buffer, EPreshader2Type Type, int32 Dimension, int32 ResultOffset, int32 AOffset, int32 BOffset)
+{
+	float* BufferData = Buffer.GetData();
+
+	// Ignores dimension, assuming dimension of 3 (we check this at creation of the op, not here at runtime)
+	if (LIKELY(Type == EPreshader2Type::Float))
+	{
+		BoundsCheck(Buffer, ResultOffset, AOffset, BOffset, 3);
+		CrossCompute(&AsFloat(BufferData, ResultOffset, 0), &AsFloat(BufferData, AOffset, 0), &AsFloat(BufferData, BOffset, 0));
+	}
+	else if (Type == EPreshader2Type::Double)
+	{
+		BoundsCheck(Buffer, ResultOffset, AOffset, BOffset, 3, false, false, /*bDouble=*/ true);
+		CrossCompute(&AsDouble(BufferData, ResultOffset, 0), &AsDouble(BufferData, AOffset, 0), &AsDouble(BufferData, BOffset, 0));
+	}
+	else
+	{
+		BoundsCheck(Buffer, ResultOffset, AOffset, BOffset, 3);
+		CrossCompute(&AsInt(BufferData, ResultOffset, 0), &AsInt(BufferData, AOffset, 0), &AsInt(BufferData, BOffset, 0));
+	}
+}
+
+// Specialization of BinaryEvaluate for Dot.  Accepts Dimension as a reference by design!
+static void BinaryEvaluate(const FOpDot&, TArrayView<float> Buffer, EPreshader2Type Type, int32& Dimension, int32 ResultOffset, int32 AOffset, int32 BOffset)
+{
+	float* BufferData = Buffer.GetData();
+
+	if (LIKELY(Type == EPreshader2Type::Float))
+	{
+		BoundsCheck(Buffer, AOffset, BOffset, Dimension);
+		BoundsCheck(Buffer, ResultOffset, 1);
+
+		float DotSum = 0.0f;
+		for (int32 Index = 0; Index < Dimension; ++Index)
+		{
+			DotSum += AsFloat(BufferData, AOffset, Index) * AsFloat(BufferData, BOffset, Index);
+		}
+		AsFloat(BufferData, ResultOffset, 0) = DotSum;
+	}
+	else
+	{
+		BoundsCheck(Buffer, AOffset, BOffset, Dimension, /*bDouble=*/ true);
+		BoundsCheck(Buffer, ResultOffset, 1, /*bDouble=*/ true);
+
+		double DotSum = 0.0;
+		for (int32 Index = 0; Index < Dimension; ++Index)
+		{
+			DotSum += AsDouble(BufferData, AOffset, Index) * AsDouble(BufferData, BOffset, Index);
+		}
+		AsDouble(BufferData, ResultOffset, 0) = DotSum;
+	}
+
+	// Override "Dimension" to 1, indicating output is a scalar.
+	Dimension = 1;
+}
+
+// OpEvaluate functions accept references to and modify Type, Dimension, and ResultOffset, as the subsequent op may reuse these.
+static void OpEvaluate(const FOpConstant&, const FMaterialRenderContext* Context, TArrayView<const FScriptName> Names, TArrayView<float> Buffer, const uint8*& OpcodePtr, EPreshader2Type& Type, int32& Dimension, int32& ResultOffset)
+{
+	// Opcode(1), ResultOffset(2), Dimension+Type(1), Data(N)
+	ResultOffset = FetchOffset(OpcodePtr);
+	Dimension = OpcodePtr[2];
+	Type = (EPreshader2Type)(Dimension >> 6);
+	Dimension = Dimension & 0x7;
+
+	const uint8* DataPtr = OpcodePtr + 3;
+	int32 ElementSize = Type == EPreshader2Type::Double ? 8 : 4;
+	int32 CopySize = Dimension * ElementSize;
+	FMemory::Memcpy(&Buffer[ResultOffset], DataPtr, CopySize);
+
+	OpcodePtr = DataPtr + CopySize;
+}
+
+static void OpEvaluate(const FOpMove&, const FMaterialRenderContext* Context, TArrayView<const FScriptName> Names, TArrayView<float> Buffer, const uint8*& OpcodePtr, EPreshader2Type& Type, int32& Dimension, int32& ResultOffset)
+{
+	float* BufferData = Buffer.GetData();
+
+	uint8 CountData = OpcodePtr[0];
+	ResultOffset = FetchOffset(OpcodePtr + 1);
+	OpcodePtr += 3;
+
+	// MoveNum and MoveSizes here are all one less than their actual count
+	uint8 MoveNum = (CountData & 0x3);
+	uint8 MoveSizes[4] =
+	{
+		((CountData >> 2) & 0x3),			// 2 bit first move size
+		((CountData >> 4) & 0x3),			// 2 bit second move size
+		((ResultOffset >> 15) & 0x1),		// 1 bit third move size
+		0									// 0 bit implicit fourth move size
+	};
+	Type = (EPreshader2Type)(CountData >> 6);
+	ResultOffset &= 0x7fff;
+	Dimension = 0;
+
+	if (UNLIKELY(Type == EPreshader2Type::Double))
+	{
+		// Need to copy to temporary array first to allow for in-place moves
+		double Temp[4];
+
+		// Less than or equals compares because MoveNum and MoveSizes are one less than the actual count.
+		for (int32 MoveIndex = 0; MoveIndex <= MoveNum; ++MoveIndex, OpcodePtr += 2)
+		{
+			int32 SourceOffset = FetchOffset(OpcodePtr);
+			BoundsCheck(Buffer, SourceOffset, MoveSizes[MoveIndex] + 1, /*bDouble=*/ true);
+
+			for (int32 MoveWord = 0; MoveWord <= MoveSizes[MoveIndex]; ++MoveWord)
+			{
+				Temp[Dimension] = AsDouble(BufferData, SourceOffset, MoveWord);
+				Dimension++;
+			}
+		}
+
+		BoundsCheck(Buffer, ResultOffset, Dimension, /*bDouble=*/ true);
+		FMemory::Memcpy(BufferData + ResultOffset, Temp, Dimension * sizeof(double));
+	}
+	else
+	{
+		// Need to copy to temporary array first to allow for in-place moves
+		float Temp[4];
+
+		// Less than or equals compares because MoveNum and MoveSizes are one less than the actual count.
+		for (int32 MoveIndex = 0; MoveIndex <= MoveNum; ++MoveIndex, OpcodePtr += 2)
+		{
+			int32 SourceOffset = FetchOffset(OpcodePtr);
+			BoundsCheck(Buffer, SourceOffset, MoveSizes[MoveIndex] + 1);
+
+			for (int32 MoveWord = 0; MoveWord <= MoveSizes[MoveIndex]; ++MoveWord)
+			{
+				Temp[Dimension] = AsFloat(BufferData, SourceOffset, MoveWord);
+				Dimension++;
+			}
+		}
+
+		BoundsCheck(Buffer, ResultOffset, Dimension);
+		FMemory::Memcpy(Buffer.GetData() + ResultOffset, Temp, Dimension * sizeof(float));
+	}
+}
+
+static const uint8* FetchHashedMaterialParameterInfo(TArrayView<const FScriptName> Names, const uint8* OpcodePtr, FHashedMaterialParameterInfo& OutParameterInfo)
+{
+	// Parameter name
+	uint16 NameIndex16;
+	FetchDataValue(OpcodePtr, NameIndex16);
+	OutParameterInfo.Name = Names[NameIndex16];
+
+	// Other fields of FHashedMaterialParameterInfo
+	FetchDataValue(OpcodePtr, OutParameterInfo.Index);
+	FetchDataValue(OpcodePtr, OutParameterInfo.Association);
+
+	return OpcodePtr;
+}
+
+static const uint8* GetTextureParameter(const FMaterialRenderContext* Context, TArrayView<const FScriptName> Names, const uint8* OpcodePtr, const UTexture*& OutTexture, const USparseVolumeTexture*& OutSparseVolumeTexture)
+{
+	// Parameter info.
+	FHashedMaterialParameterInfo ParameterInfo;
+	OpcodePtr = FetchHashedMaterialParameterInfo(Names, OpcodePtr, ParameterInfo);
+
+	// Other data.
+	int32 TextureIndex;
+	FetchDataValue(OpcodePtr, TextureIndex);
+
+	OutTexture = nullptr;
+	OutSparseVolumeTexture = nullptr;
+
+	Context->GetTextureParameterValue(ParameterInfo, TextureIndex, OutTexture);
+	if (!OutTexture)
+	{
+		Context->GetTextureParameterValue(ParameterInfo, TextureIndex, OutSparseVolumeTexture);
+	}
+
+	// Return new OpcodePtr, advanced during calls to FetchDataValue.
+	return OpcodePtr;
+}
+
+static void OpEvaluate(const FOpRuntimeVirtualTextureUniform&, const FMaterialRenderContext* Context, TArrayView<const FScriptName> Names, TArrayView<float> Buffer, const uint8*& OpcodePtr, EPreshader2Type& Type, int32& Dimension, int32& ResultOffset)
+{
+	check(Context);
+
+	ResultOffset = FetchOffset(OpcodePtr);
+
+	// Parameter info.
+	FHashedMaterialParameterInfo ParameterInfo;
+	OpcodePtr = FetchHashedMaterialParameterInfo(Names, OpcodePtr + 2, ParameterInfo);
+
+	// Other data.
+	int32 TextureIndex;
+	int32 VectorIndex;
+	FetchDataValue(OpcodePtr, TextureIndex);
+	FetchDataValue(OpcodePtr, VectorIndex);
+
+	const URuntimeVirtualTexture* Texture = nullptr;
+	if (ParameterInfo.Name.IsNone() || !Context->MaterialRenderProxy || !Context->MaterialRenderProxy->GetTextureValue(ParameterInfo, &Texture, *Context))
+	{
+		Texture = GetIndexedTexture<URuntimeVirtualTexture>(Context->Material, TextureIndex);
+	}
+	if (Texture && VectorIndex != INDEX_NONE)
+	{
+		FVector4 Value = Texture->GetUniformParameter(VectorIndex);
+		FDFVector4 DFValue(Value);
+		BoundsCheck(Buffer, ResultOffset, 8);
+		FMemory::Memcpy(Buffer.GetData() + ResultOffset, &DFValue.High[0], 8 * sizeof(float));
+	}
+	else
+	{
+		BoundsCheck(Buffer, ResultOffset, 8);
+		FMemory::Memset(Buffer.GetData() + ResultOffset, 0, 8 * sizeof(float));
+	}
+}
+
+void OpEvaluate(const FOpTexelSize&, const FMaterialRenderContext* Context, TArrayView<const FScriptName> Names, TArrayView<float> Buffer, const uint8*& OpcodePtr, EPreshader2Type& Type, int32& Dimension, int32& ResultOffset)
+{
+	check(Context);
+
+	float* BufferData = Buffer.GetData();
+
+	ResultOffset = FetchOffset(OpcodePtr);
+	BoundsCheck(Buffer, ResultOffset, 3);		// Writes 3 floats to output
+
+	const UTexture* Texture;
+	const USparseVolumeTexture* SparseVolumeTexture;
+	OpcodePtr = GetTextureParameter(Context, Names, OpcodePtr + 2, Texture, SparseVolumeTexture);
+
+	uint32 SizeX;
+	uint32 SizeY;
+	uint32 SizeZ;
+	if (Texture && Texture->GetResource())
+	{
+		SizeX = Texture->GetResource()->GetSizeX();
+		SizeY = Texture->GetResource()->GetSizeY();
+		SizeZ = Texture->GetResource()->GetSizeZ();
+	}
+	else if (SparseVolumeTexture)
+	{
+		SizeX = SparseVolumeTexture->GetSizeX();
+		SizeY = SparseVolumeTexture->GetSizeY();
+		SizeZ = SparseVolumeTexture->GetSizeZ();
+	}
+	else
+	{
+		AsFloat(BufferData, ResultOffset, 0) = 0.0f;
+		AsFloat(BufferData, ResultOffset, 1) = 0.0f;
+		AsFloat(BufferData, ResultOffset, 2) = 0.0f;
+		return;
+	}
+
+	AsFloat(BufferData, ResultOffset, 0) = 1.0f / (float)SizeX;
+	AsFloat(BufferData, ResultOffset, 1) = 1.0f / (float)SizeY;
+	AsFloat(BufferData, ResultOffset, 2) = 1.0f / (float)SizeZ;
+}
+
+void OpEvaluate(const FOpTextureSize&, const FMaterialRenderContext* Context, TArrayView<const FScriptName> Names, TArrayView<float> Buffer, const uint8*& OpcodePtr, EPreshader2Type& Type, int32& Dimension, int32& ResultOffset)
+{
+	check(Context);
+
+	float* BufferData = Buffer.GetData();
+
+	ResultOffset = FetchOffset(OpcodePtr);
+	BoundsCheck(Buffer, ResultOffset, 3);		// Writes 3 floats to output
+
+	const UTexture* Texture;
+	const USparseVolumeTexture* SparseVolumeTexture;
+	OpcodePtr = GetTextureParameter(Context, Names, OpcodePtr + 2, Texture, SparseVolumeTexture);
+
+	uint32 SizeX;
+	uint32 SizeY;
+	uint32 SizeZ;
+	if (Texture && Texture->GetResource())
+	{
+		SizeX = Texture->GetResource()->GetSizeX();
+		SizeY = Texture->GetResource()->GetSizeY();
+		SizeZ = Texture->GetResource()->GetSizeZ();
+	}
+	else if (SparseVolumeTexture)
+	{
+		SizeX = SparseVolumeTexture->GetSizeX();
+		SizeY = SparseVolumeTexture->GetSizeY();
+		SizeZ = SparseVolumeTexture->GetSizeZ();
+	}
+	else
+	{
+		AsFloat(BufferData, ResultOffset, 0) = 0.0f;
+		AsFloat(BufferData, ResultOffset, 1) = 0.0f;
+		AsFloat(BufferData, ResultOffset, 2) = 0.0f;
+		return;
+	}
+
+	AsFloat(BufferData, ResultOffset, 0) = (float)SizeX;
+	AsFloat(BufferData, ResultOffset, 1) = (float)SizeY;
+	AsFloat(BufferData, ResultOffset, 2) = (float)SizeZ;
+}
+
+// TODO
+void OpEvaluate(const FOpExternalTextureCoordinateOffset&, const FMaterialRenderContext* Context, TArrayView<const FScriptName> Names, TArrayView<float> Buffer, const uint8*& OpcodePtr, EPreshader2Type& Type, int32& Dimension, int32& ResultOffset) { checkNoEntry(); }
+void OpEvaluate(const FOpExternalTextureCoordinateScaleRotation&, const FMaterialRenderContext* Context, TArrayView<const FScriptName> Names, TArrayView<float> Buffer, const uint8*& OpcodePtr, EPreshader2Type& Type, int32& Dimension, int32& ResultOffset) { checkNoEntry(); }
+void OpEvaluate(const FOpSparseVolumeTextureUniform&, const FMaterialRenderContext* Context, TArrayView<const FScriptName> Names, TArrayView<float> Buffer, const uint8*& OpcodePtr, EPreshader2Type& Type, int32& Dimension, int32& ResultOffset) { checkNoEntry(); }
+
+// Unary opcode variants.
+// Reuse:				Opcode(1)  													ResultOffset/Dimension/Type used from previous op, AOffset same as ResultOffset
+// InPlace:				Opcode(1)  ResultOffset(2)  Dimension+Type(1)				AOffset same as ResultOffset
+// Default:				Opcode(1)  ResultOffset(2)  AOffset(2)  Dimension+Type(1)
+// Does not sanitize Dimension in case custom implementation stores extra data in that byte!  Client evaluate function should remove junk high bits from Dimension.
+#define PRESHADER2_UNARY_OP(OP_NAME) \
+	case EPreshader2Opcode::OP_NAME:																																	\
+		OpcodePtr += 2;									/* Skip to AOffset */																							\
+	case EPreshader2Opcode::OP_NAME ## _InPlace:		/* Otherwise AOffset pulled from same address as ResultOffset */												\
+		ResultOffset = FetchOffset(ResultSrc);																															\
+		AOffset = FetchOffset(OpcodePtr);																																\
+		Dimension = OpcodePtr[2];						/* Dimension value which may include extra payload */															\
+		Type = (EPreshader2Type)(Dimension >> 6);																														\
+		OpcodePtr += 3;																																					\
+	case EPreshader2Opcode::OP_NAME ## _Reuse:																															\
+		UnaryEvaluate<FOp ## OP_NAME,EPreshader2Opcode::OP_NAME < EPreshader2Opcode::UnaryNoInt>(DataContext.ResultsBuffer, Type, Dimension, ResultOffset, AOffset);	\
+		Dimension = Dimension & 7;						/* Dimension that's preserved across loop iterations, with payload removed */									\
+		break;
+
+// The Cast op is the only unary op that includes a separate destination type, which the source type is converted to, but there's no payload to store the type for
+// "Reuse" variations, so it's pulled from the opcode.  We subtract the opcode from the first "Cast_Reuse" opcode to extract the type and add it to the Dimension,
+// where it is stored in the other versions of the opcode.  To reduce code size bloat, we technically only need a case for the last version of the opcode, so we
+// fall through in the prior cases (i.e., actual code is only generated for Cast_ReuseLWC, and the other two opcodes only produce the "case" statement).
+#define PRESHADER2_UNARY_OP_CAST_REUSE(OP_NAME) \
+	case EPreshader2Opcode::OP_NAME:																																			\
+		if constexpr (EPreshader2Opcode::OP_NAME == EPreshader2Opcode::Cast_ReuseLWC)																							\
+		{																																										\
+			UnaryEvaluate<FOpCast, true>(DataContext.ResultsBuffer, Type, Dimension | (((uint8)Opcode - (uint8)EPreshader2Opcode::Cast_Reuse) << 4), ResultOffset, AOffset);		\
+			break;																																								\
+		}
+
+// Opcode(1)  ResultOffset(2)  AOffset(2)  Dimension+Type(1).
+// Does not sanitize Dimension in case custom implementation stores extra data in that byte!  Client evaluate function should remove junk high bits from Dimension.
+#define PRESHADER2_UNARY_OP_CUSTOM(OP_NAME) \
+	case EPreshader2Opcode::OP_NAME:																			\
+		ResultOffset = FetchOffset(OpcodePtr);																	\
+		AOffset = FetchOffset(OpcodePtr + 2);																	\
+		Dimension = OpcodePtr[4];																				\
+		OpcodePtr += 5;																							\
+		Type = (EPreshader2Type)(Dimension >> 6);																\
+		UnaryEvaluate(FOp ## OP_NAME(), DataContext.ResultsBuffer, Type, Dimension, ResultOffset, AOffset);		\
+		break;
+
+// Binary opcode variants.
+// Reuse:				Opcode(1)  BOffset+BInc(2)														ResultOffset/Dimension/Type used from previous op, AOffset same as ResultOffset, AInc 1
+// InPlace:				Opcode(1)  ResultOffset(2)  BOffset+BInc(2)  Dimension+Type+AInc(1)				AOffset same as ResultOffset
+// Default:				Opcode(1)  ResultOffset(2)  AOffset(2)  BOffset+BInc(2)  Dimension+Type+AInc(1)
+// ConstReuse:			Opcode(1)  BData(N)																ResultOffset/Dimension/Type used from previous op, AOffset same as ResultOffset, BInc 1, AInc 1
+// ConstReuseScalar:	Opcode(1)  BData(N)																ResultOffset/Dimension/Type used from previous op, AOffset same as ResultOffset, BInc 0, AInc 1
+// ConstInPlace:		Opcode(1)  ResultOffset(2)  Dimension+Type+AInc+BInc(1)  BData(N)				AOffset same as ResultOffset
+// ConstDefault:		Opcode(1)  ResultOffset(2)  AOffset(2)  Dimension+Type+AInc+BInc(1)  BData(N)
+#define PRESHADER2_BINARY_OP(OP_NAME) \
+	case EPreshader2Opcode::OP_NAME ## _Reuse:																				\
+		BOffset = FetchOffset(OpcodePtr);																					\
+		OpcodePtr += 2;									/* Skip over BOffset */												\
+		AInc = 1;																											\
+		goto OP_NAME ## _Evaluate;																							\
+	case EPreshader2Opcode::OP_NAME:																						\
+		OpcodePtr += 2;									/* Skip to AOffset */												\
+	case EPreshader2Opcode::OP_NAME ## _InPlace:		/* Otherwise AOffset pulled from same address as ResultOffset */	\
+		ResultOffset = FetchOffset(ResultSrc);																				\
+		AOffset = FetchOffset(OpcodePtr);																					\
+		BOffset = FetchOffset(OpcodePtr + 2);																				\
+		Dimension = OpcodePtr[4];																							\
+		Type = (EPreshader2Type)(Dimension >> 6);																			\
+		AInc = (Dimension >> 5) & 1;																						\
+		Dimension = Dimension & 7;																							\
+		OpcodePtr += 5;																										\
+	OP_NAME ## _Evaluate:																									\
+		BInc = (BOffset >> 15) & 1;																							\
+		BOffset &= 0x7fff;																									\
+		BinaryEvaluate<FOp ## OP_NAME, EPreshader2Opcode::OP_NAME < EPreshader2Opcode::BinaryNoInt>(DataContext.ResultsBuffer, Type, Dimension, ResultOffset, AOffset, BOffset, AInc, BInc);	\
+		break;																												\
+	case EPreshader2Opcode::OP_NAME ## _ConstReuse:																			\
+		AInc = 1;																											\
+		BInc = 1;																											\
+		goto OP_NAME ## _ConstEvaluate;																						\
+	case EPreshader2Opcode::OP_NAME ## _ConstReuseScalar:																	\
+		AInc = 1;																											\
+		BInc = 0;																											\
+		goto OP_NAME ## _ConstEvaluate;																						\
+	case EPreshader2Opcode::OP_NAME ## _Const:																				\
+		OpcodePtr += 2;									/* Skip to AOffset */												\
+	case EPreshader2Opcode::OP_NAME ## _ConstInPlace:	/* Otherwise AOffset pulled from same address as ResultOffset */	\
+		ResultOffset = FetchOffset(ResultSrc);																				\
+		AOffset = FetchOffset(OpcodePtr);																					\
+		Dimension = OpcodePtr[2];																							\
+		Type = (EPreshader2Type)(Dimension >> 6);																			\
+		AInc = (Dimension >> 5) & 1;																						\
+		BInc = (Dimension >> 4) & 1;																						\
+		Dimension = Dimension & 7;																							\
+		OpcodePtr += 3;																										\
+	OP_NAME ## _ConstEvaluate:																								\
+		BinaryEvaluateConst<FOp ## OP_NAME, EPreshader2Opcode::OP_NAME < EPreshader2Opcode::BinaryNoInt>(DataContext.ResultsBuffer, Type, Dimension, ResultOffset, AOffset, OpcodePtr, AInc, BInc);	\
+		break;
+
+// Opcode(1)  ResultOffset(2)  AOffset(2)  BOffset(2)  Dimension+Type(1)
+#define PRESHADER2_BINARY_OP_CUSTOM(OP_NAME) \
+	case EPreshader2Opcode::OP_NAME:																						\
+		ResultOffset = FetchOffset(OpcodePtr);																				\
+		AOffset = FetchOffset(OpcodePtr + 2);																				\
+		BOffset = FetchOffset(OpcodePtr + 4);																				\
+		Dimension = OpcodePtr[6];																							\
+		Type = (EPreshader2Type)(Dimension >> 6);																			\
+		Dimension = Dimension & 7;																							\
+		OpcodePtr += 7;																										\
+		BinaryEvaluate(FOp ## OP_NAME(), DataContext.ResultsBuffer, Type, Dimension, ResultOffset, AOffset, BOffset);		\
+		break;
+
+// Fully custom evaluation
+#define PRESHADER2_OP(OP_NAME) \
+	case EPreshader2Opcode::OP_NAME:																												\
+		OpEvaluate(FOp ## OP_NAME(), MaterialContext, DataContext.Names, DataContext.ResultsBuffer, OpcodePtr, Type, Dimension, ResultOffset);		\
+		break;
+
+
+void Evaluate(UE::Shader::FPreshaderDataContext& DataContext, const FMaterialRenderContext* MaterialContext)
+{
+	const uint8* OpcodePtr = DataContext.Ptr;
+	const uint8* OpcodeEnd = DataContext.EndPtr;
+
+	// Variables that are preserved between loop iterations, for "Reuse" opcode variants
+	EPreshader2Type Type = EPreshader2Type::Float;
+	int32 Dimension = 0;
+	int32 ResultOffset = 0;
+
+	while (OpcodePtr < OpcodeEnd)
+	{
+		const EPreshader2Opcode Opcode = (EPreshader2Opcode)*OpcodePtr;
+		OpcodePtr++;
+
+		// When an operation is set to "reuse", the result offset from the previous loop iteration is used as both input and output.
+		// Initialize these to the same, in case we hit that code path.
+		int32 AOffset = ResultOffset;
+
+		// When an operation is set to "in place", the result and first argument will both be fetched from the same address.
+		// Store a copy of the original address before it gets advanced, in case we hit that code path.  These two statements
+		// outside the switch statement avoid the need for extra branches for every switch case where fallthrough is used,
+		// reducing code size.
+		const uint8* ResultSrc = OpcodePtr;
+
+		// Temporaries, to reduce clutter in the case statement macros
+		int32 BOffset;
+		int32 AInc;
+		int32 BInc;
+
+		switch (Opcode)
+		{
+			#include "Shader/Preshader2Opcodes.inl"
+		}
+	}
+}
+
+#undef PRESHADER2_UNARY_OP
+#undef PRESHADER2_BINARY_OP
+#undef PRESHADER2_UNARY_OP_CAST_REUSE
+#undef PRESHADER2_UNARY_OP_CUSTOM
+#undef PRESHADER2_BINARY_OP_CUSTOM
+#undef PRESHADER2_OP
+
+
+// --------------------------------------------------------------------------------------------------------------------------------------------------------
+//
+//	Opcode construction logic
+//
+// --------------------------------------------------------------------------------------------------------------------------------------------------------
+static void AppendOffset(TMemoryImageArray<uint8>& Opcodes, uint16 Offset)
+{
+	Opcodes.Add((uint8)Offset);
+	Opcodes.Add((uint8)(Offset >> 8));
+}
+
+static void AppendImmediateConstant(TMemoryImageArray<uint8>& Opcodes, EPreshader2Type Type, uint8 Dimension, void* Data)
+{
+	int32 TypeSize = Type == EPreshader2Type::Double ? 8 : 4;
+	Opcodes.Append((const uint8*)Data, TypeSize * Dimension);
+}
+
+static void EnablePreshader2(UE::Shader::FPreshaderData& PreshaderData)
+{
+	if (!PreshaderData.bPreshader2)
+	{
+		// Make sure no one previously attempted to add non-preshader2 data to the array
+		check(PreshaderData.Data.IsEmpty());
+
+		PreshaderData.bPreshader2 = true;
+	}
+}
+
+// "DestType" only matters for the Cast op.
+void EmitUnaryOp(UE::Shader::FPreshaderData& PreshaderData, EPreshader2Opcode Op, EPreshader2Type Type, uint8 Dimension, uint16 ResultOffset, uint16 AOffset, bool bReuse, EPreshader2Type DestType)
+{
+	EnablePreshader2(PreshaderData);
+
+	// Make sure this is a Unary opcode, and isn't a variant.  The first opcode in each group of 3 is the default version of the opcode.
+	// Also ensure integer types aren't used for opcodes that don't support integer, and LWC type is only supported for Cast.
+	check(Op >= EPreshader2Opcode::UnaryFirst && Op <= EPreshader2Opcode::UnaryLast);
+	check(Op >= EPreshader2Opcode::UnaryCustom || (((int32)Op - (int32)EPreshader2Opcode::UnaryFirst) % 3 == 0));
+	check((Op < EPreshader2Opcode::UnaryNoInt || Type != EPreshader2Type::Integer));
+	check(Type != EPreshader2Type::LWC);
+
+	// Length of a single component "vector" is equivalent to Abs, just swap in that opcode.
+	if (Op == EPreshader2Opcode::Length && Dimension == 1)
+	{
+		Op = EPreshader2Opcode::Abs;
+	}
+
+	if (Op == EPreshader2Opcode::Length)
+	{
+		PreshaderData.Data.Add((uint8)Op);
+		AppendOffset(PreshaderData.Data, ResultOffset);
+		AppendOffset(PreshaderData.Data, AOffset);
+		PreshaderData.Data.Add(Dimension | ((uint8)Type << 6));
+	}
+	else if (bReuse)
+	{
+		if (Op == EPreshader2Opcode::Cast)
+		{
+			PreshaderData.Data.Add((uint8)Op + 2 + (uint8)DestType);	// Reuse previous result variant, with DestType embedded into opcode
+		}
+		else
+		{
+			PreshaderData.Data.Add((uint8)Op + 2);						// Reuse previous result variant
+		}
+	}
+	else
+	{
+		if (ResultOffset != AOffset)
+		{
+			PreshaderData.Data.Add((uint8)Op);							// Default variant
+			AppendOffset(PreshaderData.Data, ResultOffset);
+		}
+		else
+		{
+			PreshaderData.Data.Add((uint8)Op + 1);						// In place variant
+		}
+		AppendOffset(PreshaderData.Data, AOffset);
+
+		if (Op == EPreshader2Opcode::Cast)
+		{
+			// Shouldn't emit a cast to the same type
+			check(Type != DestType);
+			PreshaderData.Data.Add((uint8)Dimension | ((uint8)Type << 6) | ((uint8)DestType << 4));
+		}
+		else
+		{
+			PreshaderData.Data.Add((uint8)Dimension | ((uint8)Type << 6));
+		}
+	}
+}
+
+void EmitBinaryOp(UE::Shader::FPreshaderData& PreshaderData, EPreshader2Opcode Op, EPreshader2Type Type, uint8 Dimension, uint16 ResultOffset, uint16 AOffset, uint16 BOffset, bool bAScalar, bool bBScalar, bool bReuse)
+{
+	EnablePreshader2(PreshaderData);
+
+	// Make sure this is a Binary opcode, and isn't a variant.  The first opcode in each group of 7 is the default version of the opcode.
+	// Also ensure integer types aren't used for opcodes that don't support integer, and LWC type is only supported for Cast.
+	check(Op >= EPreshader2Opcode::BinaryFirst && Op <= EPreshader2Opcode::BinaryLast);
+	check(Op >= EPreshader2Opcode::BinaryCustom || (((int32)Op - (int32)EPreshader2Opcode::BinaryFirst) % 7 == 0));
+	check((Op < EPreshader2Opcode::BinaryNoInt || Type != EPreshader2Type::Integer));
+	check(Type != EPreshader2Type::LWC);
+
+	if (Op == EPreshader2Opcode::Dot || Op == EPreshader2Opcode::Cross)
+	{
+		// Dot and Cross don't support scalars for A or B.  Cross only works with dimension 3.
+		check(bAScalar == false && bBScalar == false);
+		check(Op != EPreshader2Opcode::Cross || Dimension == 3);
+
+		PreshaderData.Data.Add((uint8)Op);
+		AppendOffset(PreshaderData.Data, ResultOffset);
+		AppendOffset(PreshaderData.Data, AOffset);
+		AppendOffset(PreshaderData.Data, BOffset);
+		PreshaderData.Data.Add(Dimension | ((uint8)Type << 6));
+	}
+	else if (bReuse)
+	{
+		// Reuse previous result doesn't support scalar for the first argument, since the first argument's type is taken from the previous op.
+		check(bAScalar == false);
+
+		PreshaderData.Data.Add((uint8)Op + 2);										// Reuse previous result variant
+		AppendOffset(PreshaderData.Data, BOffset | (bBScalar ? 0 : 1u << 15));		// Increment per component is zero for scalar
+	}
+	else
+	{
+		if (ResultOffset != AOffset)
+		{
+			PreshaderData.Data.Add((uint8)Op);										// Default variant
+			AppendOffset(PreshaderData.Data, ResultOffset);
+		}
+		else
+		{
+			PreshaderData.Data.Add((uint8)Op + 1);									// In place variant
+		}
+		AppendOffset(PreshaderData.Data, AOffset);
+		AppendOffset(PreshaderData.Data, BOffset | (bBScalar ? 0 : 1u << 15));
+		PreshaderData.Data.Add(Dimension | ((uint8)Type << 6) | (bAScalar ? 0 : 1u << 5));
+	}
+}
+
+void EmitBinaryConstOp(UE::Shader::FPreshaderData& PreshaderData, EPreshader2Opcode Op, EPreshader2Type Type, uint8 Dimension, uint16 ResultOffset, uint16 AOffset, void* BData, bool bAScalar, bool bBScalar, bool bReuse)
+{
+	EnablePreshader2(PreshaderData);
+
+	// Make sure this is a Binary opcode, and isn't a variant.  The first opcode in each group of 7 is the default version of the opcode.
+	// Also ensure integer types aren't used for opcodes that don't support integer, and LWC type is only supported for Cast.
+	check(Op >= EPreshader2Opcode::BinaryFirst && Op < EPreshader2Opcode::BinaryCustom);
+	check(Op >= EPreshader2Opcode::BinaryCustom || (((int32)Op - (int32)EPreshader2Opcode::BinaryFirst) % 7 == 0));
+	check((Op < EPreshader2Opcode::BinaryNoInt || Type != EPreshader2Type::Integer));
+	check(Type != EPreshader2Type::LWC);
+
+	if (bReuse)
+	{
+		// Reuse previous result doesn't support scalar for the first argument, since the first argument's type is taken from the previous op.
+		check(bAScalar == false);
+
+		PreshaderData.Data.Add((uint8)Op + (bBScalar ? 6 : 5));					// Const reuse previous variant, vector or scalar constant
+		AppendImmediateConstant(PreshaderData.Data, Type, bBScalar ? 1 : Dimension, BData);
+	}
+	else
+	{
+		if (ResultOffset != AOffset)
+		{
+			PreshaderData.Data.Add((uint8)Op + 3);									// Const default variant
+			AppendOffset(PreshaderData.Data, ResultOffset);
+		}
+		else
+		{
+			PreshaderData.Data.Add((uint8)Op + 4);									// Const in place variant
+		}
+		AppendOffset(PreshaderData.Data, AOffset);
+
+		// Emit an increment of zero for scalar, one for vector.
+		PreshaderData.Data.Add((uint8)Dimension | ((uint8)Type << 6) | (bAScalar ? 0 : 1u << 5) | (bBScalar ? 0 : 1u << 4));
+
+		AppendImmediateConstant(PreshaderData.Data, Type, bBScalar ? 1 : Dimension, BData);
+	}
+}
+
+void EmitConstantOp(UE::Shader::FPreshaderData& PreshaderData, EPreshader2Type Type, uint8 Dimension, uint16 ResultOffset, void* Data)
+{
+	EnablePreshader2(PreshaderData);
+
+	check(Dimension <= 4 && Dimension >= 1);
+	check(Type != EPreshader2Type::LWC);
+
+	PreshaderData.Data.Add((uint8)EPreshader2Opcode::Constant);
+	AppendOffset(PreshaderData.Data, ResultOffset);
+	PreshaderData.Data.Add((uint8)Dimension | ((uint8)Type << 6));
+
+	AppendImmediateConstant(PreshaderData.Data, Type, Dimension, Data);
+}
+
+// SourceOffsets contains the offsets of individual components to be moved.  If any are adjacent, the moves can be combined.
+void EmitMoveOp(UE::Shader::FPreshaderData& PreshaderData, EPreshader2Type Type, uint16 ResultOffset, TConstArrayView<uint16> SourceOffsets)
+{
+	EnablePreshader2(PreshaderData);
+
+	check(SourceOffsets.Num() <= 4 && SourceOffsets.Num() >= 1);
+	check(Type != EPreshader2Type::LWC);
+
+	int32 ElementSizeInWords = Type == EPreshader2Type::Double ? 2 : 1;
+
+	// Validate offsets for double -- all offsets must be even!
+	if (Type == EPreshader2Type::Double)
+	{
+		check((ResultOffset & 1) == 0);
+
+		for (int32 Index = 0; Index < SourceOffsets.Num(); ++Index)
+		{
+			check((SourceOffsets[Index] & 1) == 0);
+		}
+	}
+
+	// Move Sizes default to 1, so unused move sizes can have one subtracted and produce zero.
+	uint8 MoveSizes[4] = { 1, 1, 1, 1 };
+	uint16 MoveOffsets[4];
+	uint8 MoveNum = 0;
+
+	// Add first move
+	MoveSizes[0] = 1;
+	MoveOffsets[0] = SourceOffsets[0];
+	MoveNum++;
+
+	for (int32 Index = 1; Index < SourceOffsets.Num(); ++Index)
+	{
+		if (SourceOffsets[Index] == MoveOffsets[MoveNum - 1] + ElementSizeInWords * MoveSizes[MoveNum - 1])
+		{
+			// Add to existing move
+			MoveSizes[MoveNum - 1]++;
+		}
+		else
+		{
+			// Start a new move
+			MoveSizes[MoveNum] = 1;
+			MoveOffsets[MoveNum] = SourceOffsets[Index];
+			MoveNum++;
+		}
+	}
+
+	PreshaderData.Data.Add((uint8)EPreshader2Opcode::Move);
+
+	// For compression, MoveNum and move sizes have one subtracted.  This allows the first and second move size to fit in 2 bits, and the
+	// third move size (which can be 1 or 2 elements) to fit in 1 bit.  No fourth move size is required, as it can at most be 1 element.
+	PreshaderData.Data.Add(((uint8)Type << 6) | (MoveNum - 1) | ((MoveSizes[0] - 1) << 2) | ((MoveSizes[1] - 1) << 4));
+
+	AppendOffset(PreshaderData.Data, ResultOffset | ((MoveSizes[2] - 1u) << 15));
+
+	for (int32 Index = 0; Index < MoveNum; ++Index)
+	{
+		AppendOffset(PreshaderData.Data, MoveOffsets[Index]);
+	}
+}
+
+void EmitTextureOp(UE::Shader::FPreshaderData& PreshaderData, EPreshader2Opcode Op, uint16 ResultOffset, const FHashedMaterialParameterInfo& HashedSourceParameterInfo, int32 TextureIndex, int32 UniformIndex)
+{
+	EnablePreshader2(PreshaderData);
+
+	check(Op >= EPreshader2Opcode::TextureFirst && Op <= EPreshader2Opcode::TextureLast);
+	PreshaderData.Data.Add((uint8)Op);
+	AppendOffset(PreshaderData.Data, ResultOffset);
+
+	// Material parameter info
+	PreshaderData.Write(HashedSourceParameterInfo);
+
+	// Other data
+	PreshaderData.Write(TextureIndex);
+
+	if (Op == EPreshader2Opcode::RuntimeVirtualTextureUniform)
+	{
+		PreshaderData.Write(UniformIndex);
+	}
+}
+
+
+// --------------------------------------------------------------------------------------------------------------------------------------------------------
+//
+//	Unit tests
+//
+// --------------------------------------------------------------------------------------------------------------------------------------------------------
+struct FUnitTestValue
+{
+	FVector4f Float;
+	FVector4 Double;
+	int32 Integer[4];
+	float LWC[8];			// Tile / Offset pairs
+
+	void* GetData(EPreshader2Type Type)
+	{
+		switch (Type)
+		{
+		case EPreshader2Type::Float:  return &Float;
+		case EPreshader2Type::Double:  return &Double;
+		case EPreshader2Type::Integer:  return &Integer[0];
+		default:  return &LWC[0];
+		}
+	}
+
+	const void* GetData(EPreshader2Type Type) const
+	{
+		return const_cast<FUnitTestValue*>(this)->GetData(Type);
+	}
+
+	static int32 GetElementSize(EPreshader2Type Type)
+	{
+		return ((Type == EPreshader2Type::Double) || (Type == EPreshader2Type::LWC)) ? 8 : 4;
+	}
+
+	// Initialize LWC data from the double version of the data
+	void InitLWCData()
+	{
+		for (int32 Index = 0; Index < 4; ++Index)
+		{
+			FLargeWorldRenderScalar LWCValue(Double[Index]);
+
+			LWC[Index*2 + 0] = LWCValue.GetTile();
+			LWC[Index*2 + 1] = LWCValue.GetOffset();
+		}
+	}
+};
+
+struct FUnitTestContext
+{
+	const uint8 UnusedByte = 0xeb;
+
+	EPreshader2Opcode Op;
+
+	TStaticArray<float, 32> ResultsBuffer;
+	UE::Shader::FPreshaderData PreshaderData;
+	uint32 UnusedMemoryPattern[8];
+
+	FUnitTestContext(EPreshader2Opcode InOp)
+		: Op(InOp)
+	{
+		FMemory::Memset(UnusedMemoryPattern, UnusedByte);
+		FMemory::Memset(ResultsBuffer, UnusedByte);
+	}
+
+	void SetInput(uint16 Offset, EPreshader2Type Type, int32 InputDimension, const FUnitTestValue& Input)
+	{
+		FMemory::Memcpy(&ResultsBuffer[Offset], Input.GetData(Type), InputDimension * FUnitTestValue::GetElementSize(Type));
+	}
+
+	void EvaluateTest(FUnitTestValue& ExpectedOutput, EPreshader2Type Type, int32 InputDimension, int32 OutputDimension = INDEX_NONE, bool bCheckUnused = true)
+	{
+		UE::Shader::FPreshaderDataContext DataContext(PreshaderData, ResultsBuffer);
+
+		// If output dimension is not set, it is assumed to equal the input dimension
+		OutputDimension = (OutputDimension == INDEX_NONE) ? InputDimension : OutputDimension;
+
+		uint32 DataSize = OutputDimension * FUnitTestValue::GetElementSize(Type);
+
+		Evaluate(DataContext);
+
+		checkf(FMemory::Memcmp(ResultsBuffer.GetData(), ExpectedOutput.GetData(Type), DataSize) == 0, TEXT("Unit Test Fail:  %s  %d  %d"), GPreshader2OpcodeNames[(int32)Op], (int32)Type, InputDimension);
+
+		// In-place ops that write to smaller types may not be able to unit test for unused bytes after the expected output.  In that case, we'll skip the test.
+		if (bCheckUnused)
+		{
+			checkf(FMemory::Memcmp((uint8*)ResultsBuffer.GetData() + DataSize, UnusedMemoryPattern, 32 - DataSize) == 0, TEXT("Unit Test Failed:  %s  %d  %d"), GPreshader2OpcodeNames[(int32)Op], (int32)Type, InputDimension);
+		}
+
+		// After a successful test, reset the buffer to unused and the opcode array.
+		FMemory::Memset(ResultsBuffer, UnusedByte);
+		PreshaderData.Data.Empty();
+	}
+};
+
+static void UnitTestUnaryOp(EPreshader2Opcode Op, FUnitTestValue& Input, FUnitTestValue& ExpectedOutput)
+{
+	FUnitTestContext Context(Op);
+	int32 TypesToTest = Op < EPreshader2Opcode::UnaryNoInt ? 3 : 2;
+
+	for (int32 Dimension = 1; Dimension <= 4; ++Dimension)
+	{
+		for (int32 TypeIndex = 0; TypeIndex < TypesToTest; ++TypeIndex)
+		{
+			EPreshader2Type Type = (EPreshader2Type)TypeIndex;
+			uint32 DataSize = Dimension * FUnitTestValue::GetElementSize(Type);
+
+			// Test of separate source and destination offset (data starts at offset 8, result written to offset 0)
+			EmitUnaryOp(Context.PreshaderData, Op, Type, Dimension, 0, 8, false);
+
+			Context.SetInput(8, Type, Dimension, Input);
+			Context.EvaluateTest(ExpectedOutput, Type, Dimension);
+
+			// Test of in-place result offset (data starts at offset 0, result written to offset 0)
+			EmitUnaryOp(Context.PreshaderData, Op, Type, Dimension, 0, 0, false);
+
+			Context.SetInput(0, Type, Dimension, Input);
+			Context.EvaluateTest(ExpectedOutput, Type, Dimension);
+
+			// Test of reused result offset (previous op writes a constant of the given type to offset 0, then that value is used as the input AND result of the next op)
+			EmitConstantOp(Context.PreshaderData, Type, Dimension, 0, Input.GetData(Type));
+			EmitUnaryOp(Context.PreshaderData, Op, Type, Dimension, 0, 0, true);
+
+			Context.EvaluateTest(ExpectedOutput, Type, Dimension);
+		}
+	}
+}
+
+static void UnitTestBinaryOp(EPreshader2Opcode Op, FUnitTestValue& InputA, FUnitTestValue& InputB, bool bScalarA, bool bScalarB, FUnitTestValue& ExpectedOutput)
+{
+	FUnitTestContext Context(Op);
+	int32 TypesToTest = Op < EPreshader2Opcode::BinaryNoInt ? 3 : 2;
+
+	for (int32 Dimension = 1; Dimension <= 4; ++Dimension)
+	{
+		for (int32 TypeIndex = 0; TypeIndex < TypesToTest; ++TypeIndex)
+		{
+			EPreshader2Type Type = (EPreshader2Type)TypeIndex;
+			uint32 DataSize = Dimension * FUnitTestValue::GetElementSize(Type);
+
+			// Test of separate source and destination offset (data starts at offset 8 and 16, result written to offset 0)
+			EmitBinaryOp(Context.PreshaderData, Op, Type, Dimension, 0, 8, 16, bScalarA, bScalarB);
+
+			Context.SetInput(8, Type, bScalarA ? 1 : Dimension, InputA);
+			Context.SetInput(16, Type, bScalarB ? 1 : Dimension, InputB);
+			Context.EvaluateTest(ExpectedOutput, Type, Dimension);
+
+			// Test of in-place result offset (data starts at offset 0 and 16, result written to offset 0)
+			EmitBinaryOp(Context.PreshaderData, Op, Type, Dimension, 0, 0, 16, bScalarA, bScalarB);
+
+			Context.SetInput(0, Type, bScalarA ? 1 : Dimension, InputA);
+			Context.SetInput(16, Type, bScalarB ? 1 : Dimension, InputB);
+			Context.EvaluateTest(ExpectedOutput, Type, Dimension);
+
+			// Reuse result doesn't work if first argument is scalar
+			if (bScalarA == false)
+			{
+				// Test of reused result offset (previous op writes a constant of the given type to offset 0, then that value is used as the input AND result of the next op)
+				EmitConstantOp(Context.PreshaderData, Type, bScalarA ? 1 : Dimension, 0, InputA.GetData(Type));
+				EmitBinaryOp(Context.PreshaderData, Op, Type, Dimension, 0, 0, 16, bScalarA, bScalarB, true);
+
+				Context.SetInput(16, Type, bScalarB ? 1 : Dimension, InputB);
+				Context.EvaluateTest(ExpectedOutput, Type, Dimension);
+			}
+
+			//-------------------------------------------------------------------------------------------------------------
+			// Same tests repeated with immediate constant for second op
+			//-------------------------------------------------------------------------------------------------------------
+			EmitBinaryConstOp(Context.PreshaderData, Op, Type, Dimension, 0, 8, InputB.GetData(Type), bScalarA, bScalarB);
+
+			Context.SetInput(8, Type, bScalarA ? 1 : Dimension, InputA);
+			Context.EvaluateTest(ExpectedOutput, Type, Dimension);
+
+			// Test of in-place result offset (data starts at offset 0, result written to offset 0)
+			EmitBinaryConstOp(Context.PreshaderData, Op, Type, Dimension, 0, 0, InputB.GetData(Type), bScalarA, bScalarB);
+
+			Context.SetInput(0, Type, bScalarA ? 1 : Dimension, InputA);
+			Context.EvaluateTest(ExpectedOutput, Type, Dimension);
+
+			// Reuse result doesn't work if first argument is scalar
+			if (bScalarA == false)
+			{
+				// Test of reused result offset (previous op writes a constant of the given type to offset 0, then that value is used as the input AND result of the next op)
+				EmitConstantOp(Context.PreshaderData, Type, bScalarA ? 1 : Dimension, 0, InputA.GetData(Type));
+				EmitBinaryConstOp(Context.PreshaderData, Op, Type, Dimension, 0, 0, InputB.GetData(Type), bScalarA, bScalarB, true);
+
+				Context.EvaluateTest(ExpectedOutput, Type, Dimension);
+			}
+		}
+	}
+}
+
+static void UnitTestCast()
+{
+	FUnitTestContext Context(EPreshader2Opcode::Cast);
+	FUnitTestValue Input
+	{
+		FVector4f(10000000, 10, 100, 1000),
+		FVector4(10000000, 10, 100, 1000),
+		{ 10000000, 10, 100, 1000 },
+	};
+	Input.InitLWCData();
+
+	for (int32 Dimension = 1; Dimension <= 4; ++Dimension)
+	{
+		for (int32 DestTypeIndex = 0; DestTypeIndex <= (int32)EPreshader2Type::LWC; ++DestTypeIndex)
+		{
+			EPreshader2Type DestType = (EPreshader2Type)DestTypeIndex;
+			int32 DestDataSize = Dimension * (DestType == EPreshader2Type::LWC ? 8 : FUnitTestValue::GetElementSize(DestType));
+
+			for (int32 SrcTypeIndex = 0; SrcTypeIndex <= (int32)EPreshader2Type::Integer; ++SrcTypeIndex)
+			{
+				EPreshader2Type SrcType = (EPreshader2Type)SrcTypeIndex;
+
+				// Casts to the same type aren't valid.
+				if (SrcType == DestType)
+				{
+					continue;
+				}
+
+				EmitUnaryOp(Context.PreshaderData, EPreshader2Opcode::Cast, SrcType, Dimension, 0, 8, false, DestType);
+
+				Context.SetInput(8, SrcType, Dimension, Input);
+				Context.EvaluateTest(Input, DestType, Dimension);
+
+				// Test of in-place result offset (data starts at offset 0, result written to offset 0)
+				EmitUnaryOp(Context.PreshaderData, EPreshader2Opcode::Cast, SrcType, Dimension, 0, 0, false, DestType);
+
+				// In place conversion to smaller type won't leave unused pattern in remaining bytes, so only do unused test if dest type isn't larger than source type.
+				bool bCheckUnused = FUnitTestValue::GetElementSize(SrcType) <= FUnitTestValue::GetElementSize(DestType);
+				Context.SetInput(0, SrcType, Dimension, Input);
+				Context.EvaluateTest(Input, DestType, Dimension, Dimension, bCheckUnused);
+
+				// Test of reused result offset (previous op writes a constant of the given type to offset 0, then that value is used as the input AND result of the next op)
+				EmitConstantOp(Context.PreshaderData, SrcType, Dimension, 0, Input.GetData(SrcType));
+				EmitUnaryOp(Context.PreshaderData, EPreshader2Opcode::Cast, SrcType, Dimension, 0, 0, true, DestType);
+
+				Context.EvaluateTest(Input, DestType, Dimension, Dimension, bCheckUnused);
+			}
+		}
+	}
+}
+
+static void UnitTestLength()
+{
+	FUnitTestContext Context(EPreshader2Opcode::Length);
+	FUnitTestValue Input
+	{
+		FVector4f(10, 20, 30, 40),
+		FVector4(10, 20, 30, 40),
+	};
+	FUnitTestValue ExpectedResult;
+
+	// Output of length is a scalar.
+	const int32 OutputDimension = 1;
+
+	for (int32 Dimension = 2; Dimension <= 4; ++Dimension)
+	{
+		switch (Dimension)
+		{
+		case 2:
+			ExpectedResult.Float[0] = FVector2f(Input.Float).Length();
+			ExpectedResult.Double[0] = FVector2d(Input.Double).Length();
+			break;
+		case 3:
+			ExpectedResult.Float[0] = FVector3f(Input.Float).Length();
+			ExpectedResult.Double[0] = FVector3d(Input.Double).Length();
+			break;
+		case 4:
+			ExpectedResult.Float[0] = FMath::Sqrt(Dot4(Input.Float, Input.Float));
+			ExpectedResult.Double[0] = FMath::Sqrt(Dot4(Input.Double, Input.Double));
+			break;
+		}
+
+		for (int32 TypeIndex = 0; TypeIndex <= 1; ++TypeIndex)
+		{
+			EPreshader2Type Type = (EPreshader2Type)TypeIndex;
+
+			EmitUnaryOp(Context.PreshaderData, EPreshader2Opcode::Length, Type, Dimension, 0, 8, false);
+
+			Context.SetInput(8, Type, Dimension, Input);
+			Context.EvaluateTest(ExpectedResult, Type, Dimension, OutputDimension);
+
+			// Test of in-place result offset (data starts at offset 0, result written to offset 0)
+			EmitUnaryOp(Context.PreshaderData, EPreshader2Opcode::Length, Type, Dimension, 0, 0, false);
+
+			Context.SetInput(0, Type, Dimension, Input);
+			Context.EvaluateTest(ExpectedResult, Type, Dimension, OutputDimension, /*bCheckUnused=*/ false);		// Output type is smaller, don't check unused bytes
+		}
+	}
+}
+
+static void UnitTestNormalize()
+{
+	FUnitTestContext Context(EPreshader2Opcode::Normalize);
+	FUnitTestValue Input
+	{
+		FVector4f(10, 20, 30, 40),
+		FVector4(10, 20, 30, 40),
+	};
+
+	// Normalize downcasts to float automatically, so output type is always float
+	const EPreshader2Type OutputType = EPreshader2Type::Float;
+
+	for (int32 Dimension = 2; Dimension <= 4; ++Dimension)
+	{
+		for (int32 TypeIndex = 0; TypeIndex <= 1; ++TypeIndex)
+		{
+			EPreshader2Type Type = (EPreshader2Type)TypeIndex;
+
+			// Generate output
+			FUnitTestValue ExpectedOutput;
+			if (Type == EPreshader2Type::Float)
+			{
+				switch (Dimension)
+				{
+				case 2:
+					{
+						FVector2f Input2(Input.Float);
+						Input2.Normalize();
+						ExpectedOutput.Float = FVector4f(Input2, FVector2f{});
+						break;
+					}
+				case 3:
+					{
+						FVector3f Input3(Input.Float);
+						Input3.Normalize();
+						ExpectedOutput.Float = FVector4f(Input3, 0.0f);
+						break;
+					}
+				default:
+					{
+						// No normalize implementation on vector4?  Use the quaternion version.
+						FQuat4f TempQuat = FQuat4f(Input.Float.X, Input.Float.Y, Input.Float.Z, Input.Float.W);
+						TempQuat.Normalize();
+						ExpectedOutput.Float = FVector4f(TempQuat.X, TempQuat.Y, TempQuat.Z, TempQuat.W);
+						break;
+					}
+				}
+			}
+			else
+			{
+				// Calculation is done with doubles, then downcasted
+				switch (Dimension)
+				{
+				case 2:
+					{
+						FVector2d Input2(Input.Double);
+						Input2.Normalize();
+						ExpectedOutput.Float = FVector4f(FVector4(Input2, FVector2d{}));
+						break;
+					}
+				case 3:
+					{
+						FVector3d Input3(Input.Double);
+						Input3.Normalize();
+						ExpectedOutput.Float = FVector4f(FVector4(Input3, 0.0));
+						break;
+					}
+				default:
+					{
+						// No normalize implementation on vector4?  Use the quaternion version.
+						FQuat4d TempQuat = FQuat4d(Input.Double.X, Input.Double.Y, Input.Double.Z, Input.Double.W);
+						TempQuat.Normalize();
+						ExpectedOutput.Float = FVector4f(FVector4(TempQuat.X, TempQuat.Y, TempQuat.Z, TempQuat.W));
+						break;
+					}
+				}
+			}
+
+			// Test of separate source and destination offset (data starts at offset 8, result written to offset 0)
+			EmitUnaryOp(Context.PreshaderData, EPreshader2Opcode::Normalize, Type, Dimension, 0, 8, false);
+
+			Context.SetInput(8, Type, Dimension, Input);
+			Context.EvaluateTest(ExpectedOutput, OutputType, Dimension);
+
+			// Test of in-place result offset (data starts at offset 0, result written to offset 0)
+			EmitUnaryOp(Context.PreshaderData, EPreshader2Opcode::Normalize, Type, Dimension, 0, 0, false);
+
+			Context.SetInput(0, Type, Dimension, Input);
+			Context.EvaluateTest(ExpectedOutput, OutputType, Dimension, Dimension, /*bCheckUnused=*/ Type == OutputType);
+		}
+	}
+}
+
+static void UnitTestCross(FUnitTestValue& InputA, FUnitTestValue& InputB)
+{
+	FUnitTestContext Context(EPreshader2Opcode::Cross);
+	FUnitTestValue ExpectedOutput;
+	ExpectedOutput.Float = InputA.Float ^ InputB.Float;
+	ExpectedOutput.Double = InputA.Double ^ InputB.Double;
+
+	for (int32 TypeIndex = 0; TypeIndex < 2; ++TypeIndex)
+	{
+		// Cross only operates on 3 component vectors
+		const int32 Dimension = 3;
+		EPreshader2Type Type = (EPreshader2Type)TypeIndex;
+
+		// Test of separate source and destination offset (data starts at offset 8 and 16, result written to offset 0)
+		EmitBinaryOp(Context.PreshaderData, EPreshader2Opcode::Cross, Type, Dimension, 0, 8, 16, false, false);
+
+		Context.SetInput(8, Type, Dimension, InputA);
+		Context.SetInput(16, Type, Dimension, InputB);
+		Context.EvaluateTest(ExpectedOutput, Type, Dimension);
+
+		// Test of in-place result offset (data starts at offset 0 and 16, result written to offset 0)
+		EmitBinaryOp(Context.PreshaderData, EPreshader2Opcode::Cross, Type, Dimension, 0, 0, 16, false, false);
+
+		Context.SetInput(0, Type, Dimension, InputA);
+		Context.SetInput(16, Type, Dimension, InputB);
+		Context.EvaluateTest(ExpectedOutput, Type, Dimension);
+	}
+}
+
+static void UnitTestDot(FUnitTestValue& InputA, FUnitTestValue& InputB)
+{
+	FUnitTestContext Context(EPreshader2Opcode::Dot);
+
+	// Output of dot is a scalar.
+	const int32 OutputDimension = 1;
+
+	for (int32 Dimension = 1; Dimension <= 4; ++Dimension)
+	{
+		FUnitTestValue ExpectedOutput;
+		switch (Dimension)
+		{
+		case 1:
+			ExpectedOutput.Float = FVector4f(InputA.Float[0] * InputB.Float[0]);
+			ExpectedOutput.Double = FVector4(InputA.Double[0] * InputB.Double[0]);
+			break;
+		case 2:
+			ExpectedOutput.Float = FVector4f(InputA.Float[0] * InputB.Float[0] + InputA.Float[1] * InputB.Float[1]);
+			ExpectedOutput.Double = FVector4(InputA.Double[0] * InputB.Double[0] + InputA.Double[1] * InputB.Double[1]);
+			break;
+		case 3:
+			ExpectedOutput.Float = FVector4f(Dot3(InputA.Float, InputB.Float));
+			ExpectedOutput.Double = FVector4(Dot3(InputA.Double, InputB.Double));
+			break;
+		default:
+			ExpectedOutput.Float = FVector4f(Dot4(InputA.Float, InputB.Float));
+			ExpectedOutput.Double = FVector4(Dot4(InputA.Double, InputB.Double));
+			break;
+		}
+
+		for (int32 TypeIndex = 0; TypeIndex < 2; ++TypeIndex)
+		{
+			EPreshader2Type Type = (EPreshader2Type)TypeIndex;
+
+			// Test of separate source and destination offset (data starts at offset 8 and 16, result written to offset 0)
+			EmitBinaryOp(Context.PreshaderData, EPreshader2Opcode::Dot, Type, Dimension, 0, 8, 16, false, false);
+
+			Context.SetInput(8, Type, Dimension, InputA);
+			Context.SetInput(16, Type, Dimension, InputB);
+			Context.EvaluateTest(ExpectedOutput, Type, Dimension, OutputDimension);
+
+			// Test of in-place result offset (data starts at offset 0 and 16, result written to offset 0)
+			EmitBinaryOp(Context.PreshaderData, EPreshader2Opcode::Dot, Type, Dimension, 0, 0, 16, false, false);
+
+			Context.SetInput(0, Type, Dimension, InputA);
+			Context.SetInput(16, Type, Dimension, InputB);
+			Context.EvaluateTest(ExpectedOutput, Type, Dimension, OutputDimension, /*bCheckUnused=*/ false);
+		}
+	}
+}
+
+// Move doesn't use FUnitTestContext, as it requires custom input initialization and validation.
+static void UnitTestMove(TConstArrayView<uint16> Pattern)
+{
+	int32 Dimension = Pattern.Num();
+
+	TStaticArray<float, 32> ResultsBuffer;
+	UE::Shader::FPreshaderData PreshaderData;
+
+	const uint8 UnusedByte = 0xeb;
+	uint64 UnusedMemoryPattern;
+	FMemory::Memset(UnusedMemoryPattern, UnusedByte);
+
+	int32 DataSize;
+	FUnitTestValue ExpectedOutput;
+	FMemory::Memset(ExpectedOutput, UnusedByte);
+
+	EPreshader2Type Type = EPreshader2Type::Float;
+
+	FMemory::Memset(ResultsBuffer, UnusedByte);
+
+	// Write the value in the pattern as float at the given offset.  This will produce an output equal to the input pattern
+	// if the correct words are moved.
+	for (int32 Index = 0; Index < Pattern.Num(); ++Index)
+	{
+		AsFloat(ResultsBuffer.GetData(), Pattern[Index], 0) = Pattern[Index];
+		ExpectedOutput.Float[Index] = Pattern[Index];
+	}
+
+	{
+		PreshaderData.Data.Empty();
+		EmitMoveOp(PreshaderData, Type, 0, Pattern);
+
+		UE::Shader::FPreshaderDataContext DataContext(PreshaderData, ResultsBuffer);
+		Evaluate(DataContext);
+
+		DataSize = Pattern.Num() * FUnitTestValue::GetElementSize(Type);
+		check(FMemory::Memcmp(ResultsBuffer.GetData(), ExpectedOutput.GetData(Type), DataSize) == 0);
+
+		// Check that memory after the destination of the move wasn't touched
+		for (uint16 Index = Pattern.Num(); Index < 8; ++Index)
+		{
+			if (Pattern.Contains(Index))
+			{
+				// Memory was source data from the pattern, make sure it still has a value equal to its offset
+				check(AsFloat(ResultsBuffer.GetData(), Index, 0) == (float)Index);
+			}
+			else
+			{
+				// Memory should be unused
+				check(FMemory::Memcmp(&ResultsBuffer[Index], &UnusedMemoryPattern, sizeof(float)) == 0);
+			}
+		}
+	}
+
+	{
+		// Repeat test for double format.  We'll double the offsets in the pattern since double format requires even offsets.
+		TArray<uint16> PatternDouble;
+		for (uint16 PatternOffset : Pattern)
+		{
+			PatternDouble.Add(PatternOffset * 2);
+		}
+
+		Type = EPreshader2Type::Double;
+
+		FMemory::Memset(ResultsBuffer, UnusedByte);
+
+		for (int32 Index = 0; Index < PatternDouble.Num(); ++Index)
+		{
+			AsDouble(ResultsBuffer.GetData(), PatternDouble[Index], 0) = PatternDouble[Index];
+			ExpectedOutput.Double[Index] = PatternDouble[Index];
+		}
+
+		PreshaderData.Data.Empty();
+		EmitMoveOp(PreshaderData, Type, 0, PatternDouble);
+
+		UE::Shader::FPreshaderDataContext DataContext(PreshaderData, ResultsBuffer);
+		Evaluate(DataContext);
+
+		DataSize = PatternDouble.Num() * FUnitTestValue::GetElementSize(Type);
+		check(FMemory::Memcmp(ResultsBuffer.GetData(), ExpectedOutput.GetData(Type), DataSize) == 0);
+
+		// Check that memory after the destination of the move wasn't touched.  Check doubles at every other offset.
+		for (uint16 Index = PatternDouble.Num() * 2; Index < 16; Index += 2)
+		{
+			if (PatternDouble.Contains(Index))
+			{
+				// Memory was source data from the pattern, make sure it still has a value equal to its offset
+				check(AsDouble(ResultsBuffer.GetData(), Index, 0) == (double)Index);
+			}
+			else
+			{
+				// Memory should be unused
+				check(FMemory::Memcmp(&ResultsBuffer[Index], &UnusedMemoryPattern, sizeof(double)) == 0);
+			}
+		}
+	}
+}
+
+#define UNIT_TEST_UNARY(X,Y,Z,W,OP,FUNC) do {												\
+	FUnitTestValue Input { FVector4f(X, Y, Z, W), FVector4(X, Y, Z, W), { X, Y, Z, W } };	\
+	FUnitTestValue ExpectedOutput{															\
+		FVector4f(FUNC((float)X), FUNC((float)Y), FUNC((float)Z), FUNC((float)W)),			\
+		FVector4(FUNC((double)X), FUNC((double)Y), FUNC((double)Z), FUNC((double)W)),		\
+		{ FUNC((int32)X), FUNC((int32)Y), FUNC((int32)Z), FUNC((int32)W) } };				\
+	UnitTestUnaryOp(EPreshader2Opcode::OP, Input, ExpectedOutput); } while(0)
+
+// Need to use volatile temporary variables to prevent the compiler from precomputing math function results, which can return
+// different values in the last bit or two from runtime versions of the same functions, breaking our binary comparison.
+#define UNIT_TEST_UNARY_NOINT(X,Y,Z,W,OP,FUNC) do { \
+	volatile float XF=(float)X, YF=(float)Y, ZF=(float)Z, WF=(float)W;						\
+	volatile double XD=(double)X, YD=(double)Y, ZD=(double)Z, WD=(double)W;					\
+	FUnitTestValue Input { FVector4f(X, Y, Z, W), FVector4(X, Y, Z, W) };					\
+	FUnitTestValue ExpectedOutput{															\
+		FVector4f(FUNC(XF), FUNC(YF), FUNC(ZF), FUNC(WF)),									\
+		FVector4(FUNC(XD), FUNC(YD), FUNC(ZD), FUNC(WD)) };									\
+	UnitTestUnaryOp(EPreshader2Opcode::OP, Input, ExpectedOutput); } while (0)
+
+#define UNIT_TEST_BINARY_SCALAR_PERMUTATION(A,B,C,D,X,Y,Z,W,OP,FUNC,SCALARA,SCALARB) do {					\
+	volatile float AF=(float)A, BF=(float)B, CF=(float)C, DF=(float)D;										\
+	volatile double AD=(double)A, BD=(double)B, CD=(double)C, DD=(double)D;									\
+	volatile float XF=(float)X, YF=(float)Y, ZF=(float)Z, WF=(float)W;										\
+	volatile double XD=(double)X, YD=(double)Y, ZD=(double)Z, WD=(double)W;									\
+	FUnitTestValue InputA { FVector4f(A, B, C, D), FVector4(A, B, C, D), { A, B, C, D } };					\
+	FUnitTestValue InputB { FVector4f(X, Y, Z, W), FVector4(X, Y, Z, W), { X, Y, Z, W } };					\
+	FUnitTestValue ExpectedOutput{																			\
+		FVector4f(FUNC(AF,XF),																				\
+			FUNC(SCALARA ? AF : BF, SCALARB ? XF : YF),														\
+			FUNC(SCALARA ? AF : CF, SCALARB ? XF : ZF),														\
+			FUNC(SCALARA ? AF : DF, SCALARB ? XF : WF)),													\
+		FVector4(FUNC(AD,XD),																				\
+			FUNC(SCALARA ? AD : BD, SCALARB ? XD : YD),														\
+			FUNC(SCALARA ? AD : CD, SCALARB ? XD : ZD),														\
+			FUNC(SCALARA ? AD : DD, SCALARB ? XD : WD)),													\
+		{ FUNC((int32)A,(int32)X),																			\
+			FUNC((int32)(SCALARA ? A : B), (int32)(SCALARB ? X : Y)),										\
+			FUNC((int32)(SCALARA ? A : C), (int32)(SCALARB ? X : Z)),										\
+			FUNC((int32)(SCALARA ? A : D), (int32)(SCALARB ? X : W)) } };									\
+	UnitTestBinaryOp(EPreshader2Opcode::OP, InputA, InputB, SCALARA, SCALARB, ExpectedOutput); } while(0)
+
+#define UNIT_TEST_BINARY(A,B,C,D,X,Y,Z,W,OP,FUNC) \
+	UNIT_TEST_BINARY_SCALAR_PERMUTATION(A,B,C,D,X,Y,Z,W,OP,FUNC, false, false);								\
+	UNIT_TEST_BINARY_SCALAR_PERMUTATION(A,B,C,D,X,Y,Z,W,OP,FUNC, true, false);								\
+	UNIT_TEST_BINARY_SCALAR_PERMUTATION(A,B,C,D,X,Y,Z,W,OP,FUNC, false, true);								\
+	UNIT_TEST_BINARY_SCALAR_PERMUTATION(A,B,C,D,X,Y,Z,W,OP,FUNC, true, true)
+
+#define UNIT_TEST_BINARY_NOINT_SCALAR_PERMUTATION(A,B,C,D,X,Y,Z,W,OP,FUNC,SCALARA,SCALARB) do {				\
+	volatile float AF=(float)A, BF=(float)B, CF=(float)C, DF=(float)D;										\
+	volatile double AD=(double)A, BD=(double)B, CD=(double)C, DD=(double)D;									\
+	volatile float XF=(float)X, YF=(float)Y, ZF=(float)Z, WF=(float)W;										\
+	volatile double XD=(double)X, YD=(double)Y, ZD=(double)Z, WD=(double)W;									\
+	FUnitTestValue InputA { FVector4f(A, B, C, D), FVector4(A, B, C, D) };									\
+	FUnitTestValue InputB { FVector4f(X, Y, Z, W), FVector4(X, Y, Z, W) };									\
+	FUnitTestValue ExpectedOutput{																			\
+		FVector4f(FUNC(AF,XF),																				\
+			FUNC(SCALARA ? AF : BF, SCALARB ? XF : YF),														\
+			FUNC(SCALARA ? AF : CF, SCALARB ? XF : ZF),														\
+			FUNC(SCALARA ? AF : DF, SCALARB ? XF : WF)),													\
+		FVector4(FUNC(AD,XD),																				\
+			FUNC(SCALARA ? AD : BD, SCALARB ? XD : YD),														\
+			FUNC(SCALARA ? AD : CD, SCALARB ? XD : ZD),														\
+			FUNC(SCALARA ? AD : DD, SCALARB ? XD : WD)) };													\
+	UnitTestBinaryOp(EPreshader2Opcode::OP, InputA, InputB, SCALARA, SCALARB, ExpectedOutput); } while(0)
+
+#define UNIT_TEST_BINARY_NOINT(A,B,C,D,X,Y,Z,W,OP,FUNC) \
+	UNIT_TEST_BINARY_NOINT_SCALAR_PERMUTATION(A,B,C,D,X,Y,Z,W,OP,FUNC, false, false);						\
+	UNIT_TEST_BINARY_NOINT_SCALAR_PERMUTATION(A,B,C,D,X,Y,Z,W,OP,FUNC, true, false);						\
+	UNIT_TEST_BINARY_NOINT_SCALAR_PERMUTATION(A,B,C,D,X,Y,Z,W,OP,FUNC, false, true);						\
+	UNIT_TEST_BINARY_NOINT_SCALAR_PERMUTATION(A,B,C,D,X,Y,Z,W,OP,FUNC, true, true)
+
+#define UNIT_TEST_BINARY_CROSS(A,B,C,X,Y,Z) do {												\
+	FUnitTestValue InputA { FVector4f(A, B, C, 0.0f), FVector4(A, B, C, 0.0) };					\
+	FUnitTestValue InputB { FVector4f(X, Y, Z, 0.0f), FVector4(X, Y, Z, 0.0) };					\
+	UnitTestCross(InputA, InputB); } while(0)
+
+#define UNIT_TEST_BINARY_DOT(A,B,C,D,X,Y,Z,W) do {												\
+	FUnitTestValue InputA { FVector4f(A, B, C, D), FVector4(A, B, C, D) };						\
+	FUnitTestValue InputB { FVector4f(X, Y, Z, W), FVector4(X, Y, Z, W) };						\
+	UnitTestDot(InputA, InputB); } while(0)
+
+void UnitTest()
+{
+	UNIT_TEST_UNARY(-1, 2, -3, 4, Abs, FMath::Abs);
+	UNIT_TEST_UNARY(5, -6, 7, -8, Abs, FMath::Abs);
+	UNIT_TEST_UNARY(-1, 2, -3, 4, Negate, FOpNegate::Call);
+	UNIT_TEST_UNARY(5, -6, 7, -8, Negate, FOpNegate::Call);
+	UNIT_TEST_UNARY(-1, 0, 1, 2, Saturate, FOpSaturate::Call);
+	UNIT_TEST_UNARY(2, 1, 0, -1, Saturate, FOpSaturate::Call);
+	UNIT_TEST_UNARY_NOINT(-0.125, 0.125, -0.875, 0.875, Saturate, FOpSaturate::Call);
+	UNIT_TEST_UNARY_NOINT(0.75, 1.25, 0.75, 1.25, Saturate, FOpSaturate::Call);
+	UNIT_TEST_UNARY(-1, 2, -3, 4, Sign, FMath::Sign);
+	UNIT_TEST_UNARY(5, -6, 7, -8, Sign, FMath::Sign);
+
+	UNIT_TEST_UNARY_NOINT(-1.25, -1.0, -0.75, -0.5, ACos, FMath::Acos);
+	UNIT_TEST_UNARY_NOINT(-0.25, 0.0, 0.25, 0.5, ACos, FMath::Acos);
+	UNIT_TEST_UNARY_NOINT(0.5, 0.75, 1.0, 1.25, ACos, FMath::Acos);
+	UNIT_TEST_UNARY_NOINT(-1.25, -1.0, -0.75, -0.5, ASin, FMath::Asin);
+	UNIT_TEST_UNARY_NOINT(-0.25, 0.0, 0.25, 0.5, ASin, FMath::Asin);
+	UNIT_TEST_UNARY_NOINT(0.5, 0.75, 1.0, 1.25, ASin, FMath::Asin);
+	UNIT_TEST_UNARY_NOINT(-1.25, -1.0, -0.75, -0.5, ATan, FMath::Atan);
+	UNIT_TEST_UNARY_NOINT(-0.25, 0.0, 0.25, 0.5, ATan, FMath::Atan);
+	UNIT_TEST_UNARY_NOINT(0.5, 0.75, 1.0, 1.25, ATan, FMath::Atan);
+	UNIT_TEST_UNARY_NOINT(-0.125, 0.125, -0.875, 0.875, Ceil, FMath::CeilToFloat);
+	UNIT_TEST_UNARY_NOINT(0.75, 1.25, 0.75, 1.25, Ceil, FMath::CeilToFloat);
+	UNIT_TEST_UNARY_NOINT(1, 2, 3, 4, Cos, FMath::Cos);
+	UNIT_TEST_UNARY_NOINT(1, 2, 3, 4, Exponential, FMath::Exp);
+	UNIT_TEST_UNARY_NOINT(1, 2, 3, 4, Exponential2, FMath::Exp2);
+	UNIT_TEST_UNARY_NOINT(-0.125, 0.125, -0.875, 0.875, Floor, FMath::FloorToFloat);
+	UNIT_TEST_UNARY_NOINT(0.75, 1.25, 0.75, 1.25, Floor, FMath::FloorToFloat);
+	UNIT_TEST_UNARY_NOINT(-0.125, 0.125, -0.875, 0.875, Frac, FMath::Frac);
+	UNIT_TEST_UNARY_NOINT(0.75, 1.25, 0.75, 1.25, Frac, FMath::Frac);
+	UNIT_TEST_UNARY_NOINT(-0.125, 0.125, -0.875, 0.875, Fractional, FMath::Fractional);
+	UNIT_TEST_UNARY_NOINT(0.75, 1.25, 0.75, 1.25, Fractional, FMath::Fractional);
+	UNIT_TEST_UNARY_NOINT(0.5, 1.0, 1.5, 2.0, Logarithm, FMath::Loge);
+	UNIT_TEST_UNARY_NOINT(0.5, 1.0, 1.5, 2.0, Logarithm10, FOpLogarithm10::Call);
+	UNIT_TEST_UNARY_NOINT(0.5, 1.0, 1.5, 2.0, Logarithm2, FMath::Log2);
+	UNIT_TEST_UNARY_NOINT(-0.125, 0.125, -0.875, 0.875, Reciprocal, FOpReciprocal::Call);
+	UNIT_TEST_UNARY_NOINT(-0.000000000000001, 0.0, 0.000000000000001, 2.0, Reciprocal, FOpReciprocal::Call);
+	UNIT_TEST_UNARY_NOINT(-0.125, 0.125, -0.875, 0.875, Round, FMath::RoundToFloat);
+	UNIT_TEST_UNARY_NOINT(0.75, 1.25, 0.75, 1.25, Round, FMath::RoundToFloat);
+	UNIT_TEST_UNARY_NOINT(1, 2, 3, 4, Sin, FMath::Sin);
+	UNIT_TEST_UNARY_NOINT(1, 2, 3, 4, Sqrt, FMath::Sqrt);
+	UNIT_TEST_UNARY_NOINT(1, 2, 3, 4, Tan, FMath::Tan);
+	UNIT_TEST_UNARY_NOINT(-0.125, 0.125, -0.875, 0.875, Truncate, FMath::TruncToFloat);
+	UNIT_TEST_UNARY_NOINT(0.75, 1.25, 0.75, 1.25, Truncate, FMath::TruncToFloat);
+
+	UNIT_TEST_BINARY(2, 3, 40, 50, 60, 70, 8, 9, Add, FOpAdd::Call);
+	UNIT_TEST_BINARY(2, 3, 40, 50, 60, 70, 8, 9, Divide, FOpDivide::Call);
+	UNIT_TEST_BINARY(2, 3, 40, 50, 60, 70, 8, 9, Fmod, FOpFmod::Call);
+	UNIT_TEST_BINARY(2, 3, 40, 50, 60, 70, 8, 9, Max, FMath::Max);
+	UNIT_TEST_BINARY(2, 3, 40, 50, 60, 70, 8, 9, Min, FMath::Min);
+	UNIT_TEST_BINARY(2, 3, 40, 50, 60, 70, 8, 9, Multiply, FOpMultiply::Call);
+	UNIT_TEST_BINARY(2, 3, 40, 50, 60, 70, 8, 9, Subtract, FOpSubtract::Call);
+
+	UNIT_TEST_BINARY_NOINT(2, 3, 40, 50, 60, 70, 8, 9, ATan2, FMath::Atan2);
+
+	UNIT_TEST_BINARY_CROSS(1, -2, 3, -4, 5, -6);
+	UNIT_TEST_BINARY_DOT(0.5, 1.0, 1.5, 2.0, 2.5, 3.0, 3.5, 4.0);
+
+	UnitTestCast();
+	UnitTestLength();
+	UnitTestNormalize();
+
+	// Exhaustively test a bunch of move permutations
+	for (uint16 IndexA = 0; IndexA < 8; ++IndexA)
+	{
+		// Permutations with one input
+		UnitTestMove({ IndexA });
+
+		for (uint16 IndexB = 0; IndexB < 8; ++IndexB)
+		{
+			// Permutations with two inputs
+			UnitTestMove({ IndexA, IndexB });
+
+			for (uint16 IndexC = 0; IndexC < 8; ++IndexC)
+			{
+				// Permutations with three inputs
+				UnitTestMove({ IndexA, IndexB, IndexC });
+
+				for (uint16 IndexD = 0; IndexD < 8; ++IndexD)
+				{
+					// Permutations with four inputs
+					UnitTestMove({ IndexA, IndexB, IndexC, IndexD });
+				}
+			}
+		}
+	}
+
+	// Note that "Constant" op is unit tested by unary and binary ops when testing "Reuse" opcode variants,
+	// so we don't need a separate unit test for that.
+}
+
+}  // namespace UE::Preshader2
diff --git a/Engine/Source/Runtime/Engine/Private/Shader/PreshaderEvaluate.h b/Engine/Source/Runtime/Engine/Private/Shader/PreshaderEvaluate.h
--- a/Engine/Source/Runtime/Engine/Private/Shader/PreshaderEvaluate.h
+++ b/Engine/Source/Runtime/Engine/Private/Shader/PreshaderEvaluate.h
@@ -16,14 +16,15 @@ class FPreshaderData;
 
 struct FPreshaderDataContext
 {
-	explicit FPreshaderDataContext(const FPreshaderData& InData);
+	explicit FPreshaderDataContext(const FPreshaderData& InData, TArrayView<float> InResultsBuffer = {});
 	FPreshaderDataContext(const FPreshaderDataContext& InContext, uint32 InOffset, uint32 InSize);
 
 	const uint8* RESTRICT Ptr;
 	const uint8* RESTRICT EndPtr;
 	TArrayView<const FScriptName> Names;
-	TArrayView<const FPreshaderStructType> StructTypes;
-	TArrayView<const EValueComponentType> StructComponentTypes;
+	TArrayView<float> ResultsBuffer;				// Used by Preshader2 for results and intermediates
+
+	bool IsPreshader2() const { return !ResultsBuffer.IsEmpty(); }
 };
 
 ENGINE_API FPreshaderValue EvaluatePreshader(const FUniformExpressionSet* UniformExpressionSet, const FMaterialRenderContext& Context, FPreshaderStack& Stack, FPreshaderDataContext& RESTRICT Data);
diff --git a/Engine/Source/Runtime/Engine/Public/MaterialShared.h b/Engine/Source/Runtime/Engine/Public/MaterialShared.h
--- a/Engine/Source/Runtime/Engine/Public/MaterialShared.h
+++ b/Engine/Source/Runtime/Engine/Public/MaterialShared.h
@@ -475,35 +475,26 @@ public:
 	LAYOUT_FIELD(uint16, BufferOffset); 	/** Offset in the preshader uniform buffer (in bytes) where result will be written */
 };
 
-class FMaterialUniformPreshaderField
-{
-	DECLARE_TYPE_LAYOUT(FMaterialUniformPreshaderField, NonVirtual);
-public:
-	LAYOUT_FIELD(uint32, BufferOffset); /** Offset in the uniform buffer where result will be written */
-	LAYOUT_FIELD(uint32, ComponentIndex);
-	LAYOUT_FIELD(UE::Shader::EValueType, Type);
-};
-
 class FMaterialUniformPreshaderHeader
 {
 	DECLARE_TYPE_LAYOUT(FMaterialUniformPreshaderHeader, NonVirtual);
 public:
 	friend inline bool operator==(const FMaterialUniformPreshaderHeader& Lhs, const FMaterialUniformPreshaderHeader& Rhs)
 	{
-		return Lhs.OpcodeOffset == Rhs.OpcodeOffset &&
-			Lhs.OpcodeSize == Rhs.OpcodeSize &&
-			Lhs.FieldIndex == Rhs.FieldIndex &&
-			Lhs.NumFields == Rhs.NumFields;
+		return Lhs.Type == Rhs.Type &&
+			Lhs.BufferOffset == Rhs.BufferOffset &&
+			Lhs.OpcodeOffset == Rhs.OpcodeOffset &&
+			Lhs.OpcodeSize == Rhs.OpcodeSize;
 	}
 	friend inline bool operator!=(const FMaterialUniformPreshaderHeader& Lhs, const FMaterialUniformPreshaderHeader& Rhs)
 	{
 		return !operator==(Lhs, Rhs);
 	}
 
+	LAYOUT_FIELD(UE::Shader::EValueType, Type);
+	LAYOUT_FIELD(uint16, BufferOffset); /** Offset in the uniform buffer where result will be written */
 	LAYOUT_FIELD(uint32, OpcodeOffset); /** Offset of the preshader opcodes, within the material's buffer */
 	LAYOUT_FIELD(uint32, OpcodeSize); /** Size of the preshader opcodes */
-	LAYOUT_FIELD(uint32, FieldIndex);
-	LAYOUT_FIELD(uint32, NumFields);
 };
 
 class FMaterialNumericParameterInfo
@@ -744,8 +735,8 @@ public:
 	void SetVTLayer(int32 StackIndex, int32 VTLayerIndex, int32 TextureIndex);
 
 	/// Writes a new uniform preshader entry.
-	/// @param BufferOffset Specifies the location in the preshader buffer where the result of the preshader instruction will be written to. See FMaterialUniformPreshaderField::BufferOffset.
-	/// @param EntryType Specifies the value type of the preshader instruction output. See FMaterialUniformPreshaderField::Type.
+	/// @param BufferOffset Specifies the location in the preshader buffer where the result of the preshader instruction will be written to. See FMaterialUniformPreshaderHeader::BufferOffset.
+	/// @param EntryType Specifies the value type of the preshader instruction output. See FMaterialUniformPreshaderHeader::Type.
 	/// @param EntryGenerator Callback that writes the preshader opcodes.
 	void WriteUniformPreshaderEntry(uint32 BufferOffset, UE::Shader::EValueType EntryType, const TFunction<void(UE::Shader::FPreshaderData&)>& EntryGenerator);
 
@@ -773,10 +764,10 @@ protected:
 	friend class FMaterialRenderProxy;
 	friend class FMaterialVirtualTextureStack;
 	friend class FDebugUniformExpressionSet;
+	friend struct FMaterialIRValueAnalyzer;
 
 	LAYOUT_FIELD(TMemoryImageArray<FMaterialUniformParameterEvaluation>, UniformParameterEvaluations);
 	LAYOUT_FIELD(TMemoryImageArray<FMaterialUniformPreshaderHeader>, UniformPreshaders);
-	LAYOUT_FIELD(TMemoryImageArray<FMaterialUniformPreshaderField>, UniformPreshaderFields);
 	LAYOUT_FIELD(TMemoryImageArray<FMaterialNumericParameterInfo>, UniformNumericParameters);
 	LAYOUT_ARRAY(TMemoryImageArray<FMaterialTextureParameterInfo>, UniformTextureParameters, NumMaterialTextureParameterTypes);
 	LAYOUT_FIELD(TMemoryImageArray<FMaterialExternalTextureParameterInfo>, UniformExternalTextureParameters);
diff --git a/Engine/Source/Runtime/Engine/Public/Materials/MaterialIR.h b/Engine/Source/Runtime/Engine/Public/Materials/MaterialIR.h
--- a/Engine/Source/Runtime/Engine/Public/Materials/MaterialIR.h
+++ b/Engine/Source/Runtime/Engine/Public/Materials/MaterialIR.h
@@ -8,6 +8,7 @@
 #include "MaterialSceneTextureId.h"
 #include "Engine/TextureDefines.h"
 #include "Shader/Preshader.h"
+#include "Shader/Preshader2.h"
 
 #if WITH_EDITOR
 
@@ -1197,7 +1198,7 @@ struct FPreshaderParameter : TInstruction<VK_PreshaderParameter, 1>
 	FValue* SourceParameter;
 
 	// Specifies the opcode that needs to be encoded for this parameter. For now, only a small subset of opcodes are supported for material preshader parameters.
-	UE::Shader::EPreshaderOpcode Opcode;
+	EPreshader2Opcode Opcode;
 
 	// Index into UMaterialInterface::GetReferencedTextures() pointing to the source parameter's texture.
 	int32 TextureIndex;
diff --git a/Engine/Source/Runtime/Engine/Public/Materials/MaterialIREmitter.h b/Engine/Source/Runtime/Engine/Public/Materials/MaterialIREmitter.h
--- a/Engine/Source/Runtime/Engine/Public/Materials/MaterialIREmitter.h
+++ b/Engine/Source/Runtime/Engine/Public/Materials/MaterialIREmitter.h
@@ -439,7 +439,7 @@ public:
 	FValueRef Parameter(FName Name, FMaterialParameterMetadata& Metadata, EMaterialSamplerType SamplerType = SAMPLERTYPE_Color, int32 VTLayerIndex = INDEX_NONE, int32 VTPageTableIndex = INDEX_NONE);
 
 	// Emits a preshader parameter value with given opcode.
-	FValueRef PreshaderParameter(FType Type, UE::Shader::EPreshaderOpcode Opcode, FValueRef SourceParameter, FPreshaderParameterPayload Payload = {});
+	FValueRef PreshaderParameter(FType Type, EPreshader2Opcode Opcode, FValueRef SourceParameter, FPreshaderParameterPayload Payload = {});
 
 	// Emits a custom primitive data value with given index.
 	FValueRef CustomPrimitiveData(uint32 PrimitiveDataIndex);
diff --git a/Engine/Source/Runtime/Engine/Public/Shader/Preshader.h b/Engine/Source/Runtime/Engine/Public/Shader/Preshader.h
--- a/Engine/Source/Runtime/Engine/Public/Shader/Preshader.h
+++ b/Engine/Source/Runtime/Engine/Public/Shader/Preshader.h
@@ -63,11 +63,7 @@ enum class EPreshaderOpcode : uint8
 	ExternalTextureCoordinateOffset,
 	RuntimeVirtualTextureUniform,
 	SparseVolumeTextureUniform,
-	GetField,
-	SetField,
 	Neg,
-	Jump,
-	JumpIfFalse,
 	PushValue,
 	Less,
 	Assign,
@@ -100,7 +96,7 @@ class FPreshaderData
 public:
 	friend inline bool operator==(const FPreshaderData& Lhs, const FPreshaderData& Rhs)
 	{
-		return Lhs.Names == Rhs.Names && Lhs.Data == Rhs.Data;
+		return Lhs.bPreshader2 == Rhs.bPreshader2 && Lhs.Names == Rhs.Names && Lhs.Data == Rhs.Data;
 	}
 
 	friend inline bool operator!=(const FPreshaderData& Lhs, const FPreshaderData& Rhs)
@@ -117,15 +113,8 @@ public:
 
 	void WriteData(const void* Value, uint32 Size);
 	void WriteName(const FScriptName& Name);
-	void WriteType(const FType& Type);
 	void WriteValue(const FValue& Value);
 
-	FPreshaderLabel WriteJump(EPreshaderOpcode Op);
-	void WriteJump(EPreshaderOpcode Op, FPreshaderLabel Label);
-
-	FPreshaderLabel GetLabel();
-	void SetLabel(FPreshaderLabel InLabel);
-
 	template<typename T>
 	FPreshaderData& Write(const T& Value) { WriteData(&Value, sizeof(T)); return *this; }
 
@@ -133,10 +122,11 @@ public:
 	FPreshaderData& Write<FScriptName>(const FScriptName& Value) { WriteName(Value); return *this; }
 
 	template<>
-	FPreshaderData& Write<FType>(const FType& Value) { WriteType(Value); return *this; }
+	FPreshaderData& Write<FValue>(const FValue& Value) { WriteValue(Value); return *this; }
 
+	/** Writing FType isn't supported -- use its ValueType enum member instead */
 	template<>
-	FPreshaderData& Write<FValue>(const FValue& Value) { WriteValue(Value); return *this; }
+	FPreshaderData& Write<FType>(const FType& Value) = delete;
 
 	/** Can't write FName, use FScriptName instead */
 	template<>
@@ -145,11 +135,10 @@ public:
 	template<>
 	FPreshaderData& Write<FHashedMaterialParameterInfo>(const FHashedMaterialParameterInfo& Value) { return Write(Value.Name).Write(Value.Index).Write(Value.Association); }
 
-	inline FPreshaderData& WriteOpcode(EPreshaderOpcode Op) { ensure(Op != EPreshaderOpcode::Nop); return Write<uint8>((uint8)Op); }
+	inline FPreshaderData& WriteOpcode(EPreshaderOpcode Op) { ensure(Op != EPreshaderOpcode::Nop && bPreshader2 == false); return Write<uint8>((uint8)Op); }
 
+	LAYOUT_FIELD_INITIALIZED(bool, bPreshader2, false);
 	LAYOUT_FIELD(TMemoryImageArray<FScriptName>, Names);
-	LAYOUT_FIELD(TMemoryImageArray<FPreshaderStructType>, StructTypes);
-	LAYOUT_FIELD(TMemoryImageArray<EValueComponentType>, StructComponentTypes);
 	LAYOUT_FIELD(TMemoryImageArray<uint8>, Data);
 };
 
diff --git a/Engine/Source/Runtime/Engine/Public/Shader/Preshader2.h b/Engine/Source/Runtime/Engine/Public/Shader/Preshader2.h
new file mode 100644
--- /dev/null
+++ b/Engine/Source/Runtime/Engine/Public/Shader/Preshader2.h
@@ -0,0 +1,101 @@
+// Copyright Epic Games, Inc. All Rights Reserved.
+
+/*=============================================================================
+	Preshader2.h: New Preshader interface.
+=============================================================================*/
+
+#pragma once
+
+#include "Materials/MaterialParameters.h"
+#include "Shader/Preshader.h"
+#include "Shader/PreshaderEvaluate.h"
+
+#define PRESHADER2_UNARY_OP(OP_NAME)	\
+	OP_NAME,									\
+	OP_NAME ## _InPlace,						\
+	OP_NAME ## _Reuse,
+
+#define PRESHADER2_BINARY_OP(OP_NAME)	\
+	OP_NAME,									\
+	OP_NAME ## _InPlace,						\
+	OP_NAME ## _Reuse,							\
+	OP_NAME ## _Const,							\
+	OP_NAME ## _ConstInPlace,					\
+	OP_NAME ## _ConstReuse,						\
+	OP_NAME ## _ConstReuseScalar,
+
+#define PRESHADER2_UNARY_OP_CAST_REUSE(OP_NAME) OP_NAME,
+#define PRESHADER2_UNARY_OP_CUSTOM(OP_NAME) OP_NAME,
+#define PRESHADER2_BINARY_OP_CUSTOM(OP_NAME) OP_NAME,
+#define PRESHADER2_OP(OP_NAME) OP_NAME,
+
+enum class EPreshader2Opcode : uint8
+{
+	Invalid,
+
+	#include "Preshader2Opcodes.inl"
+
+	// Standard unary ops have three variations:  normal, in place (input == output), and reuse (input == output == last output)
+	UnaryFirst = Abs,
+	UnaryNoInt = ACos,			// First unary op that doesn't support integer types
+	UnaryCustom = Length,		// First custom unary op with a single opcode variant instead of 3
+	UnaryLast = Length,
+
+	// Standard binary ops have seven variations:  normal, in place (first input == output), reuse (first input == output == last output),
+	// plus four variations with inline constant second arguments (normal, in place, reuse, reuse with scalar constant).
+	BinaryFirst = Add,
+	BinaryNoInt = ATan2,		// First binary op that doesn't support integer types
+	BinaryCustom = Cross,		// First binary op with a single opcode variant instead of 7
+	BinaryLast = Dot,
+
+	TextureFirst = RuntimeVirtualTextureUniform,
+	TextureLast = TextureSize,
+};
+
+#undef PRESHADER2_UNARY_OP
+#undef PRESHADER2_BINARY_OP
+#undef PRESHADER2_UNARY_OP_CAST_REUSE
+#undef PRESHADER2_UNARY_OP_CUSTOM
+#undef PRESHADER2_BINARY_OP_CUSTOM
+#undef PRESHADER2_OP
+
+enum class EPreshader2Type : uint8
+{
+	Float,
+	Double,
+	Integer,
+	LWC				// Large World Coordinates.  Not supported by math ops, only as an output format for the Cast op.
+};
+
+namespace UE::Preshader2
+{
+
+// Preshader2 uses buffer offsets specified in 4-byte words.  Operations on Double types must use even aligned word offsets, so the 8-byte type ends up aligned.
+
+// "DestType" only applies to the Cast op.  Reuse option uses the result offset of the previous instruction for both ResultOffset and AOffset, but some opcodes
+// may not support this, so a valid ResultOffset and AOffset should still be specified.
+void EmitUnaryOp(UE::Shader::FPreshaderData& PreshaderData, EPreshader2Opcode Op, EPreshader2Type Type, uint8 Dimension, uint16 ResultOffset, uint16 AOffset, bool bReuse = false, EPreshader2Type DestType = EPreshader2Type::Float);
+
+// Binary ops support either or both operands being scalars instead of vectors.
+void EmitBinaryOp(UE::Shader::FPreshaderData& PreshaderData, EPreshader2Opcode Op, EPreshader2Type Type, uint8 Dimension, uint16 ResultOffset, uint16 AOffset, uint16 BOffset, bool bAScalar, bool bBScalar, bool bReuse = false);
+
+// Similar to above, but second argument is a constant inlined into the bytecode.  Note that "custom" binary opcodes (ones in the numeric range BinaryCustom to BinaryLast)
+// don't work with this function.  Currently this includes Cross and Dot opcodes.
+void EmitBinaryConstOp(UE::Shader::FPreshaderData& PreshaderData, EPreshader2Opcode Op, EPreshader2Type Type, uint8 Dimension, uint16 ResultOffset, uint16 AOffset, void* BData, bool bAScalar, bool bBScalar, bool bReuse = false);
+
+// Loads an inlined constant to the given ResultOffset.
+void EmitConstantOp(UE::Shader::FPreshaderData& PreshaderData, EPreshader2Type Type, uint8 Dimension, uint16 ResultOffset, void* Data);
+
+// SourceOffsets contains the offsets of individual components to be moved.  If any are adjacent, the moves can be combined.
+void EmitMoveOp(UE::Shader::FPreshaderData& PreshaderData, EPreshader2Type Type, uint16 ResultOffset, TConstArrayView<uint16> SourceOffsets);
+
+// Emit texture related opcode (TexelSize, TextureSize, etc).
+void EmitTextureOp(UE::Shader::FPreshaderData& PreshaderData, EPreshader2Opcode Op, uint16 ResultOffset, const FHashedMaterialParameterInfo& HashedSourceParameterInfo, int32 TextureIndex, int32 UniformIndex = INDEX_NONE);
+
+// Evaluate opcodes.  FMaterialRenderContext is required if any texture ops are used (TexelSize, TextureSize, etc).
+void Evaluate(UE::Shader::FPreshaderDataContext& DataContext, const FMaterialRenderContext* MaterialContext = nullptr);
+
+// Unit test of system.
+void UnitTest();
+
+}  // namespace UE::Preshader2
\ No newline at end of file
diff --git a/Engine/Source/Runtime/Engine/Public/Shader/Preshader2Opcodes.inl b/Engine/Source/Runtime/Engine/Public/Shader/Preshader2Opcodes.inl
new file mode 100644
--- /dev/null
+++ b/Engine/Source/Runtime/Engine/Public/Shader/Preshader2Opcodes.inl
@@ -0,0 +1,71 @@
+// Copyright Epic Games, Inc. All Rights Reserved.
+
+// Unary ops with integer support -- each has 3 opcode variants, with Cast having 3 extra custom "Reuse" opcode variants for different destination types.
+// The three extra opcodes conveniently maintain the 3-element stride for surrounding standard unary opcode variants.
+PRESHADER2_UNARY_OP(Abs)
+PRESHADER2_UNARY_OP(Cast)
+PRESHADER2_UNARY_OP_CAST_REUSE(Cast_ReuseDouble)
+PRESHADER2_UNARY_OP_CAST_REUSE(Cast_ReuseInteger)
+PRESHADER2_UNARY_OP_CAST_REUSE(Cast_ReuseLWC)
+PRESHADER2_UNARY_OP(Negate)
+PRESHADER2_UNARY_OP(Saturate)
+PRESHADER2_UNARY_OP(Sign)
+
+// Unary ops without integer support
+PRESHADER2_UNARY_OP(ACos)
+PRESHADER2_UNARY_OP(ASin)
+PRESHADER2_UNARY_OP(ATan)
+PRESHADER2_UNARY_OP(Ceil)
+PRESHADER2_UNARY_OP(Cos)
+PRESHADER2_UNARY_OP(Exponential)
+PRESHADER2_UNARY_OP(Exponential2)
+PRESHADER2_UNARY_OP(Floor)
+PRESHADER2_UNARY_OP(Frac)
+PRESHADER2_UNARY_OP(Fractional)
+PRESHADER2_UNARY_OP(Logarithm)
+PRESHADER2_UNARY_OP(Logarithm10)
+PRESHADER2_UNARY_OP(Logarithm2)
+PRESHADER2_UNARY_OP(Normalize)
+PRESHADER2_UNARY_OP(Reciprocal)
+PRESHADER2_UNARY_OP(Round)
+PRESHADER2_UNARY_OP(Sin)
+PRESHADER2_UNARY_OP(Sqrt)
+PRESHADER2_UNARY_OP(Tan)
+PRESHADER2_UNARY_OP(Truncate)
+
+// Non-standard unary ops -- single opcode variant
+PRESHADER2_UNARY_OP_CUSTOM(Length)
+
+// Binary ops with integer support -- each has 7 opcode variants
+PRESHADER2_BINARY_OP(Add)
+PRESHADER2_BINARY_OP(Divide)
+PRESHADER2_BINARY_OP(Fmod)
+PRESHADER2_BINARY_OP(Max)
+PRESHADER2_BINARY_OP(Min)
+PRESHADER2_BINARY_OP(Multiply)
+PRESHADER2_BINARY_OP(Subtract)
+
+// Binary ops without integer support
+PRESHADER2_BINARY_OP(ATan2)
+
+// Non-standard binary ops -- single opcode variant
+PRESHADER2_BINARY_OP_CUSTOM(Cross)
+PRESHADER2_BINARY_OP_CUSTOM(Dot)
+
+// Ops that don't fall under unary or binary categories
+PRESHADER2_OP(Move)			// Opcode(1), MoveCount+Type+FirstMoveSize+SecondMoveSize(1), ResultOffset+ThirdMoveSize(2), SourceOffset(2*N)
+PRESHADER2_OP(Constant)		// Opcode(1), ResultOffset(2), Dimension+Type(1), Data(N)
+
+// Texture inputs
+PRESHADER2_OP(RuntimeVirtualTextureUniform)
+PRESHADER2_OP(TexelSize)
+PRESHADER2_OP(TextureSize)
+
+// External inputs (TODO:  Unimplemented)
+PRESHADER2_OP(ExternalTextureCoordinateScaleRotation)
+PRESHADER2_OP(ExternalTextureCoordinateOffset)
+PRESHADER2_OP(SparseVolumeTextureUniform)
+
+// NOTE:  Clamp op is not implemented, caller should generate Min+Max.  Challenge of implementing Clamp as a single op is that there are many
+// permutations of register and constant input arguments that need to be dealt with, so it's easier to split in two.  It should be plenty
+// efficient with reuse of previous op's result, taking only one extra byte in storage, and one extra op.  Clamp is 0.01% of ops.
\ No newline at end of file
diff --git a/Engine/Source/Runtime/Engine/Public/Shader/PreshaderTypes.h b/Engine/Source/Runtime/Engine/Public/Shader/PreshaderTypes.h
--- a/Engine/Source/Runtime/Engine/Public/Shader/PreshaderTypes.h
+++ b/Engine/Source/Runtime/Engine/Public/Shader/PreshaderTypes.h
@@ -16,15 +16,11 @@ namespace Shader
 struct FPreshaderType
 {
 	FPreshaderType() = default;
-	FPreshaderType(const FType& InType);
 	FPreshaderType(EValueType InType);
 
-	bool IsStruct() const { return StructTypeHash != 0u; }
-	int32 GetNumComponents() const { return StructTypeHash != 0u ? StructComponentTypes.Num() : GetValueTypeDescription(ValueType).NumComponents; }
+	int32 GetNumComponents() const { return GetValueTypeDescription(ValueType).NumComponents; }
 	EValueComponentType GetComponentType(int32 Index) const;
 
-	uint64 StructTypeHash = 0u;
-	TArrayView<const EValueComponentType> StructComponentTypes;
 	EValueType ValueType = EValueType::Void;
 };
 
@@ -75,7 +71,6 @@ public:
 	/** Override top value type (after an in place unary operation) */
 	inline void OverrideTopType(EValueType ValueType)
 	{
-		check(!Values.Last().IsStruct());
 		Values.Last().ValueType = ValueType;
 	}
 
