Preshader2 static analysis warning suppression.  Both warnings are spurious (false positive on the commas in a template argument and an always true expression in a macro by design).

#rnx

diff --git a/Engine/Source/Runtime/Engine/Private/Shader/Preshader2.cpp b/Engine/Source/Runtime/Engine/Private/Shader/Preshader2.cpp
--- a/Engine/Source/Runtime/Engine/Private/Shader/Preshader2.cpp
+++ b/Engine/Source/Runtime/Engine/Private/Shader/Preshader2.cpp
@@ -991,7 +991,7 @@ void Evaluate(UE::Shader::FPreshaderDataContext& DataContext, const FMaterialRen
 		int32 AInc;
 		int32 BInc;
 
-		switch (Opcode)
+		switch (Opcode)		//-V501 //-V521
 		{
 			#include "Shader/Preshader2Opcodes.inl"
 		}
