Fixed documentation error in SpringMath

diff --git a/Engine/Source/Runtime/Engine/Classes/Kismet/BlueprintSpringMathLibrary.h b/Engine/Source/Runtime/Engine/Classes/Kismet/BlueprintSpringMathLibrary.h
--- a/Engine/Source/Runtime/Engine/Classes/Kismet/BlueprintSpringMathLibrary.h
+++ b/Engine/Source/Runtime/Engine/Classes/Kismet/BlueprintSpringMathLibrary.h
@@ -32,7 +32,7 @@ public:
 	 * @param InOutAngle The value to be damped in degrees
 	 * @param InOutAngularVelocity The speed of the value to be damped in deg/s
 	 * @param TargetAngle The goal to damp towards in degrees
-	 * @param SmoothingTime The smoothing time to use for the spring. Longer times lead to more damped behaviour. A time of 0 will snap InOutAngle to TargetX
+	 * @param SmoothingTime The smoothing time to use for the spring. Longer times lead to more damped behaviour. A time of 0 will snap InOutAngle to TargetAngle
 	 * @param DeltaTime Timestep in seconds
 	 */
 	UFUNCTION(BlueprintCallable, meta = (Category = "Math|Spring|Experimental"))
