[NetPhysics] Improvements to predictive interpolation when used with resimulation.
- Option to apply input decay for simulated proxies that apply inputs at runtime (outside of resimulation)
- Hook up settings for Predictive Interpolation replication and correcton versions (bVelocityBased and bCorrectionAsVelocity)
- CVar to stop resimulations from happening, for debugging: p.Resim.ForceSkipResim

diff --git a/Engine/Source/Runtime/Engine/Private/PhysicsEngine/NetworkPhysicsComponent.cpp b/Engine/Source/Runtime/Engine/Private/PhysicsEngine/NetworkPhysicsComponent.cpp
--- a/Engine/Source/Runtime/Engine/Private/PhysicsEngine/NetworkPhysicsComponent.cpp
+++ b/Engine/Source/Runtime/Engine/Private/PhysicsEngine/NetworkPhysicsComponent.cpp
@@ -71,6 +71,10 @@ namespace PhysicsReplicationCVars
 		static FAutoConsoleVariableRef CVarApplyInputDecayOverSetTime(TEXT("np2.Resim.ApplyInputDecayOverSetTime"), bApplyInputDecayOverSetTime, TEXT("When true, apply the Input Decay Curve over a set amount of time instead of over the start of input prediction and end of resim which is variable each resimulation"));
 		float InputDecaySetTime = 0.15f;
 		static FAutoConsoleVariableRef CVarInputDecaySetTime(TEXT("np2.Resim.InputDecaySetTime"), InputDecaySetTime, TEXT("Applied when np2.Resim.ApplyInputDecayOverSetTime is true, read there for more info. Set time to apply Input Decay Curve over while predicting inputs during resimulation"));
+		bool bApplyInputDecaySimProxyInputAtRuntime = false;
+		static FAutoConsoleVariableRef CVarApplyInputDecaySimProxyInputAtRuntime(TEXT("np2.Resim.ApplyInputDecaySimProxyInputAtRuntime"), bApplyInputDecaySimProxyInputAtRuntime, TEXT("When true, apply input decay on inputs applied on simulated proxies at runtime (outside of resimulation)"));
+		float InputDecaySimProxyInputAtRuntime = 0.25f;
+		static FAutoConsoleVariableRef CVarInputDecaySimProxyInputAtRuntime(TEXT("np2.Resim.InputDecaySimProxyInputAtRuntime"), InputDecaySimProxyInputAtRuntime, TEXT("The amount of input decay to apply on simulated proxy inputs at runtime (outside of resim), if np2.Resim.ApplyInputDecaySimProxyInputAtRuntime is true"));
 
 		bool bEnableStatefulDeltaSerialization = true;
 		static FAutoConsoleVariableRef CVarResimEnableStatefulDeltaSerialization(TEXT("np2.Resim.StatefulDeltaSerialization.Enable"), bEnableStatefulDeltaSerialization, TEXT("Enables stateful delta serialization for FNetworkPhysicsData derived inputs and states. During FNetworkPhysicsData::NetSerialize there will be a valid pointer to a previous FNetworkPhysicsData which can be used for delta serialization, FNetworkPhysicsData::DeltaSourceData. NOTE: Switching this during gameplay might cause disconnections."));
@@ -2135,6 +2139,12 @@ void FAsyncNetworkPhysicsComponent::OnPreProcessInputs_Internal(const int32 Phys
 						const int32 LastFrame = PhysicsData->LocalFrame;
 						for (; NextExpectedLocalFrame <= LastFrame; NextExpectedLocalFrame++)
 						{
+							// Perform input decay on sim-proxy runtime (non-resim) input
+							if (bApplySimProxyInput && ComponentSettings.GetApplyInputDecaySimProxyInputAtRuntime())
+							{
+								PhysicsData->DecayData(ComponentSettings.GetInputDecaySimProxyInputAtRuntime());
+							}
+
 							if (InputHistory->ExtractData(NextExpectedLocalFrame, bIsSolverReset, PhysicsData, true) && NextExpectedLocalFrame < LastFrame)
 							{
 								if (bIsUsingLegacyData)
@@ -2174,6 +2184,12 @@ void FAsyncNetworkPhysicsComponent::OnPreProcessInputs_Internal(const int32 Phys
 					InputHistory->RecordData(PhysicsStep, PhysicsData);
 				}
 
+				// Perform input decay on sim-proxy runtime (non-resim) input
+				if (!bIsSolverResim && bApplySimProxyInput && ComponentSettings.GetApplyInputDecaySimProxyInputAtRuntime())
+				{
+					PhysicsData->DecayData(ComponentSettings.GetInputDecaySimProxyInputAtRuntime());
+				}
+
 				if (bIsUsingLegacyData)
 				{
 					if (ensure(ActorComponent.IsValid() && ActorComponent.Get()->IsBeingDestroyed() == false))
diff --git a/Engine/Source/Runtime/Engine/Private/PhysicsEngine/PhysicsReplication.cpp b/Engine/Source/Runtime/Engine/Private/PhysicsEngine/PhysicsReplication.cpp
--- a/Engine/Source/Runtime/Engine/Private/PhysicsEngine/PhysicsReplication.cpp
+++ b/Engine/Source/Runtime/Engine/Private/PhysicsEngine/PhysicsReplication.cpp
@@ -229,7 +229,7 @@ namespace PhysicsReplicationCVars
 		static FAutoConsoleVariableRef CVarPostResimWaitForUpdate(TEXT("np2.PredictiveInterpolation.PostResimWaitForUpdate"), bPostResimWaitForUpdate, TEXT("After a resimulation, wait for replicated states that correspond to post-resim state before processing replication again."));
 		
 		bool bVelocityBased = true;
-		static FAutoConsoleVariableRef CVarVelocityBased(TEXT("np2.PredictiveInterpolation.VelocityBased"), bVelocityBased, TEXT("When true, predictive interpolation replication mode will only apply linear velocity and angular velocity"));
+		static FAutoConsoleVariableRef CVarVelocityBased(TEXT("np2.PredictiveInterpolation.VelocityBased"), bVelocityBased, TEXT("When true, predictive interpolation replication mode will move objects replicate by applying linear velocity and angular velocity. When false, objects will move by transform shifting towards the target state."));
 		
 		bool bCorrectionAsVelocity = false;
 		static FAutoConsoleVariableRef CVarCorrectionAsVelocity(TEXT("np2.PredictiveInterpolation.CorrectionAsVelocity"), bCorrectionAsVelocity, TEXT("When true, predictive interpolation will apply positional and rotational offset correction as a velocity instead of as a transform shift."));
@@ -1917,7 +1917,7 @@ bool FPhysicsReplicationAsync::PredictiveInterpolation(Chaos::FPBDRigidParticleH
 	/** If target velocities are low enough, check the traveled direction and distance from previous frame and compare with replicated linear velocity.
 	* If the object isn't moving enough along the replicated velocity it's considered stuck and needs reconciliation.
 	* SoftSnap is performed each tick while there is a registered error, if enough time pass HardSnap forces the object into the correct state. */
-	bool bSoftSnap = !PhysicsReplicationCVars::PredictiveInterpolationCVars::bVelocityBased;
+	bool bSoftSnap = false;
 
 	if (PhysicsReplicationCVars::PredictiveInterpolationCVars::bDisableErrorVelocityLimits ||
 		(TargetLinVel.Size() < PhysicsReplicationCVars::PredictiveInterpolationCVars::ErrorAccLinVelMaxLimit && TargetAngVel.Size() < PhysicsReplicationCVars::PredictiveInterpolationCVars::ErrorAccAngVelMaxLimit))
@@ -1945,7 +1945,7 @@ bool FPhysicsReplicationAsync::PredictiveInterpolation(Chaos::FPBDRigidParticleH
 		Target.AccumulatedErrorSeconds = 0;
 	}
 
-	if (SettingsCurrent.PredictiveInterpolationSettings.GetDisableSoftSnap() && PhysicsReplicationCVars::PredictiveInterpolationCVars::bVelocityBased)
+	if (SettingsCurrent.PredictiveInterpolationSettings.GetDisableSoftSnap())
 	{
 		bSoftSnap = false;
 	}
@@ -2003,7 +2003,7 @@ bool FPhysicsReplicationAsync::PredictiveInterpolation(Chaos::FPBDRigidParticleH
 			return EndReplicationHelper(Target, true);
 		}
 	}
-	else // Velocity-based Replication
+	else // Dynamic Replication
 	{
 		// Wake up if sleeping
 		if (bIsSleeping)
@@ -2011,137 +2011,142 @@ bool FPhysicsReplicationAsync::PredictiveInterpolation(Chaos::FPBDRigidParticleH
 			RigidsSolver->GetEvolution()->SetParticleObjectState(Handle, Chaos::EObjectStateType::Dynamic);
 		}
 
-		// Calculate interpolation time based on current average receive rate
-		const float AverageReceiveIntervalSeconds = Target.AverageReceiveInterval * DeltaSeconds;
-		const float InterpolationTime = AverageReceiveIntervalSeconds * SettingsCurrent.PredictiveInterpolationSettings.GetPosInterpolationTimeMultiplier();
+		// Velocity-based Replication
+		if (SettingsCurrent.PredictiveInterpolationSettings.GetVelocityBased())
+		{
+			// Calculate interpolation time based on current average receive rate
+			const float AverageReceiveIntervalSeconds = Target.AverageReceiveInterval * DeltaSeconds;
+			const float InterpolationTime = AverageReceiveIntervalSeconds * SettingsCurrent.PredictiveInterpolationSettings.GetPosInterpolationTimeMultiplier();
 
-		// Calculate position correction time based on current Round Trip Time
-		const float RTT = LatencyOneWay * 2.f;
-		const float PosCorrectionTime = FMath::Max(SettingsCurrent.PredictiveInterpolationSettings.GetPosCorrectionTimeBase() + AverageReceiveIntervalSeconds + RTT * SettingsCurrent.PredictiveInterpolationSettings.GetPosCorrectionTimeMultiplier(),
-			DeltaSeconds + SettingsCurrent.PredictiveInterpolationSettings.GetPosCorrectionTimeMin());
-		const float RotCorrectionTime = FMath::Max(SettingsCurrent.PredictiveInterpolationSettings.GetRotCorrectionTimeBase() + AverageReceiveIntervalSeconds + RTT * SettingsCurrent.PredictiveInterpolationSettings.GetRotCorrectionTimeMultiplier(),
-			DeltaSeconds + SettingsCurrent.PredictiveInterpolationSettings.GetRotCorrectionTimeMin());
+			// Calculate position correction time based on current Round Trip Time
+			const float RTT = LatencyOneWay * 2.f;
+			const float PosCorrectionTime = FMath::Max(SettingsCurrent.PredictiveInterpolationSettings.GetPosCorrectionTimeBase() + AverageReceiveIntervalSeconds + RTT * SettingsCurrent.PredictiveInterpolationSettings.GetPosCorrectionTimeMultiplier(),
+				DeltaSeconds + SettingsCurrent.PredictiveInterpolationSettings.GetPosCorrectionTimeMin());
+			const float RotCorrectionTime = FMath::Max(SettingsCurrent.PredictiveInterpolationSettings.GetRotCorrectionTimeBase() + AverageReceiveIntervalSeconds + RTT * SettingsCurrent.PredictiveInterpolationSettings.GetRotCorrectionTimeMultiplier(),
+				DeltaSeconds + SettingsCurrent.PredictiveInterpolationSettings.GetRotCorrectionTimeMin());
 
-		FVector CorrectionX = CurrentState.Position;
-		if ((bXCanEarlyOut && SettingsCurrent.PredictiveInterpolationSettings.GetSkipVelocityRepOnPosEarlyOut()) == false)
-		{	// --- Velocity Replication ---
+			FVector CorrectionX = CurrentState.Position;
+			if ((bXCanEarlyOut && SettingsCurrent.PredictiveInterpolationSettings.GetSkipVelocityRepOnPosEarlyOut()) == false)
+			{	// --- Velocity Replication ---
 
-			// Get PosDiff
-			const FVector PosDiff = TargetPos - CurrentState.Position;
+				// Get PosDiff
+				const FVector PosDiff = TargetPos - CurrentState.Position;
 
-			// Get LinVelDiff by adding inverted CurrentState.LinVel to TargetLinVel
-			const FVector LinVelDiff = -CurrentState.LinVel + TargetLinVel;
+				// Get LinVelDiff by adding inverted CurrentState.LinVel to TargetLinVel
+				const FVector LinVelDiff = -CurrentState.LinVel + TargetLinVel;
 
-			// Calculate velocity blend amount for this tick as an alpha value
-			const float VelocityAlpha = FMath::Clamp(DeltaSeconds / InterpolationTime, 0.0f, 1.0f);
+				// Calculate velocity blend amount for this tick as an alpha value
+				const float VelocityAlpha = FMath::Clamp(DeltaSeconds / InterpolationTime, 0.0f, 1.0f);
 
-			FVector RepLinVel;
-			if (PhysicsReplicationCVars::PredictiveInterpolationCVars::bCorrectionAsVelocity)
-			{
-				// Convert PosDiff to a velocity
-				const FVector PosDiffVelocity = PosDiff / PosCorrectionTime;
+				FVector RepLinVel;
+				if (SettingsCurrent.PredictiveInterpolationSettings.GetCorrectionAsVelocity())
+				{
+					// Convert PosDiff to a velocity
+					const FVector PosDiffVelocity = PosDiff / PosCorrectionTime;
 
-				// Add PosDiffVelocity to LinVelDiff to get BlendedTargetVelocity
-				const FVector BlendedTargetVelocity = LinVelDiff + PosDiffVelocity;
+					// Add PosDiffVelocity to LinVelDiff to get BlendedTargetVelocity
+					const FVector BlendedTargetVelocity = LinVelDiff + PosDiffVelocity;
 
-				// Add BlendedTargetVelocity onto current velocity
-				RepLinVel = CurrentState.LinVel + (BlendedTargetVelocity * VelocityAlpha); // Same as (BlendedTargetVelocity / InterpolationTime) * DeltaSeconds
-			}
-			else // Positional correction as transform shift
-			{
-				// Add velocity diff onto current velocity
-				RepLinVel = CurrentState.LinVel + (LinVelDiff * VelocityAlpha); // Same as (LinVelDiff / InterpolationTime) * DeltaSeconds
+					// Add BlendedTargetVelocity onto current velocity
+					RepLinVel = CurrentState.LinVel + (BlendedTargetVelocity * VelocityAlpha); // Same as (BlendedTargetVelocity / InterpolationTime) * DeltaSeconds
+				}
+				else // Positional correction as transform shift
+				{
+					// Add velocity diff onto current velocity
+					RepLinVel = CurrentState.LinVel + (LinVelDiff * VelocityAlpha); // Same as (LinVelDiff / InterpolationTime) * DeltaSeconds
 
-				// Calculate correction blend amount for this tick as an alpha value
-				const float CorrectionAlpha = FMath::Clamp(DeltaSeconds / PosCorrectionTime, 0.0f, 1.0f);
+					// Calculate correction blend amount for this tick as an alpha value
+					const float CorrectionAlpha = FMath::Clamp(DeltaSeconds / PosCorrectionTime, 0.0f, 1.0f);
 
-				// Calculate the PosDiff amount to correct this tick
-				const FVector PosDiffVelocityDelta = PosDiff * CorrectionAlpha; // Same as (PosDiff / PosCorrectionTime) * DeltaSeconds
+					// Calculate the PosDiff amount to correct this tick
+					const FVector PosDiffVelocityDelta = PosDiff * CorrectionAlpha; // Same as (PosDiff / PosCorrectionTime) * DeltaSeconds
 
-				// The new position after correction
-				CorrectionX = Handle->GetX() + PosDiffVelocityDelta;
-			}
+					// The new position after correction
+					CorrectionX = Handle->GetX() + PosDiffVelocityDelta;
+				}
 
-			// Apply velocity replication
-			Handle->SetV(RepLinVel);
+				// Apply velocity replication
+				Handle->SetV(RepLinVel);
 
-#if !(UE_BUILD_SHIPPING || UE_BUILD_TEST)
-			if (PhysicsReplicationCVars::PredictiveInterpolationCVars::bDrawDebugVectors)
-			{
-				const FVector Offset = FVector(0.0f, 0.0f, PhysicsReplicationCVars::PredictiveInterpolationCVars::DrawDebugZOffset);
-				const FVector OffsetAdd = FVector(0.0f, 0.0f, 10.0f);
-				const FVector StartPos = TargetPos + Offset;
-				FVector Direction = TargetLinVel;
-				Direction.Normalize();
-				Chaos::FDebugDrawQueue::GetInstance().DrawDebugDirectionalArrow((StartPos + OffsetAdd * 0), (StartPos + OffsetAdd * 0) + TargetLinVel * 0.5f, 5.0f, FColor::Green, false, -1.0f, 0, 2.0f);
-				Chaos::FDebugDrawQueue::GetInstance().DrawDebugDirectionalArrow((StartPos + OffsetAdd * 1), (StartPos + OffsetAdd * 1) + CurrentState.LinVel * 0.5f, 5.0f, FColor::Blue, false, -1.0f, 0, 2.0f);
-				Chaos::FDebugDrawQueue::GetInstance().DrawDebugDirectionalArrow((StartPos + OffsetAdd * 2), (StartPos + OffsetAdd * 2) + (Target.PrevLinVel - CurrentState.LinVel) * 0.5f, 5.0f, FColor::Red, false, -1.0f, 0, 2.0f);
-				Chaos::FDebugDrawQueue::GetInstance().DrawDebugDirectionalArrow((StartPos + OffsetAdd * 3), (StartPos + OffsetAdd * 3) + RepLinVel * 0.5f, 5.0f, FColor::Magenta, false, -1.0f, 0, 2.0f);
-				Chaos::FDebugDrawQueue::GetInstance().DrawDebugDirectionalArrow((StartPos + OffsetAdd * 4), (StartPos + OffsetAdd * 4) + (Target.PrevLinVel - RepLinVel) * 0.5f, 5.0f, FColor::Orange, false, -1.0f, 0, 2.0f);
-				Chaos::FDebugDrawQueue::GetInstance().DrawDebugDirectionalArrow((StartPos + OffsetAdd * 5), (StartPos + OffsetAdd * 5) + Direction * RTT, 5.0f, FColor::White, false, -1.0f, 0, 2.0f);
-				Chaos::FDebugDrawQueue::GetInstance().DrawDebugDirectionalArrow((StartPos + OffsetAdd * 6), (StartPos + OffsetAdd * 6) + Direction * InterpolationTime, 5.0f, FColor::Yellow, false, -1.0f, 0, 2.0f);
+	#if !(UE_BUILD_SHIPPING || UE_BUILD_TEST)
+				if (PhysicsReplicationCVars::PredictiveInterpolationCVars::bDrawDebugVectors)
+				{
+					const FVector Offset = FVector(0.0f, 0.0f, PhysicsReplicationCVars::PredictiveInterpolationCVars::DrawDebugZOffset);
+					const FVector OffsetAdd = FVector(0.0f, 0.0f, 10.0f);
+					const FVector StartPos = TargetPos + Offset;
+					FVector Direction = TargetLinVel;
+					Direction.Normalize();
+					Chaos::FDebugDrawQueue::GetInstance().DrawDebugDirectionalArrow((StartPos + OffsetAdd * 0), (StartPos + OffsetAdd * 0) + TargetLinVel * 0.5f, 5.0f, FColor::Green, false, -1.0f, 0, 2.0f);
+					Chaos::FDebugDrawQueue::GetInstance().DrawDebugDirectionalArrow((StartPos + OffsetAdd * 1), (StartPos + OffsetAdd * 1) + CurrentState.LinVel * 0.5f, 5.0f, FColor::Blue, false, -1.0f, 0, 2.0f);
+					Chaos::FDebugDrawQueue::GetInstance().DrawDebugDirectionalArrow((StartPos + OffsetAdd * 2), (StartPos + OffsetAdd * 2) + (Target.PrevLinVel - CurrentState.LinVel) * 0.5f, 5.0f, FColor::Red, false, -1.0f, 0, 2.0f);
+					Chaos::FDebugDrawQueue::GetInstance().DrawDebugDirectionalArrow((StartPos + OffsetAdd * 3), (StartPos + OffsetAdd * 3) + RepLinVel * 0.5f, 5.0f, FColor::Magenta, false, -1.0f, 0, 2.0f);
+					Chaos::FDebugDrawQueue::GetInstance().DrawDebugDirectionalArrow((StartPos + OffsetAdd * 4), (StartPos + OffsetAdd * 4) + (Target.PrevLinVel - RepLinVel) * 0.5f, 5.0f, FColor::Orange, false, -1.0f, 0, 2.0f);
+					Chaos::FDebugDrawQueue::GetInstance().DrawDebugDirectionalArrow((StartPos + OffsetAdd * 5), (StartPos + OffsetAdd * 5) + Direction * RTT, 5.0f, FColor::White, false, -1.0f, 0, 2.0f);
+					Chaos::FDebugDrawQueue::GetInstance().DrawDebugDirectionalArrow((StartPos + OffsetAdd * 6), (StartPos + OffsetAdd * 6) + Direction * InterpolationTime, 5.0f, FColor::Yellow, false, -1.0f, 0, 2.0f);
+				}
+	#endif
+				// Cache data for next replication
+				Target.PrevLinVel = FVector(RepLinVel);
 			}
-#endif
-			// Cache data for next replication
-			Target.PrevLinVel = FVector(RepLinVel);
-		}
 
-		FQuat CorrectionR = CurrentState.Quaternion;
-		{	// --- Angular Velocity Replication ---
+			FQuat CorrectionR = CurrentState.Quaternion;
+			{	// --- Angular Velocity Replication ---
 
-			// Get AngVelDiff by adding inverted CurrentState.AngVel to TargetAngVel
-			const FVector AngVelDiff = -CurrentState.AngVel + TargetAngVel;
+				// Get AngVelDiff by adding inverted CurrentState.AngVel to TargetAngVel
+				const FVector AngVelDiff = -CurrentState.AngVel + TargetAngVel;
 
-			// Calculate velocity blend amount for this tick as an alpha value
-			const float VelocityAlpha = FMath::Clamp(DeltaSeconds / InterpolationTime, 0.0f, 1.0f);
+				// Calculate velocity blend amount for this tick as an alpha value
+				const float VelocityAlpha = FMath::Clamp(DeltaSeconds / InterpolationTime, 0.0f, 1.0f);
 
-			FVector RepAngVel;
-			if (PhysicsReplicationCVars::PredictiveInterpolationCVars::bCorrectionAsVelocity)
-			{
-				// Get RotDiff
-				const FQuat RotDiff = TargetRot * CurrentState.Quaternion.Inverse();
+				FVector RepAngVel;
+				if (SettingsCurrent.PredictiveInterpolationSettings.GetCorrectionAsVelocity())
+				{
+					// Get RotDiff
+					const FQuat RotDiff = TargetRot * CurrentState.Quaternion.Inverse();
 
-				// Convert RotDiff to a velocity
-				float WAngle;
-				FVector WAxis;
-				RotDiff.ToAxisAndAngle(WAxis, WAngle);
-				WAngle = FMath::UnwindRadians(WAngle);
-				const FVector RotDiffVelocity = FVector(WAxis * (WAngle / RotCorrectionTime));
+					// Convert RotDiff to a velocity
+					float WAngle;
+					FVector WAxis;
+					RotDiff.ToAxisAndAngle(WAxis, WAngle);
+					WAngle = FMath::UnwindRadians(WAngle);
+					const FVector RotDiffVelocity = FVector(WAxis * (WAngle / RotCorrectionTime));
 
-				// Add RotDiffVelocity to AngVelDiff to get BlendedTargetVelocity
-				const FVector BlendedTargetVelocity = AngVelDiff + RotDiffVelocity;
+					// Add RotDiffVelocity to AngVelDiff to get BlendedTargetVelocity
+					const FVector BlendedTargetVelocity = AngVelDiff + RotDiffVelocity;
 
-				// Add BlendedTargetVelocity to CurrentState.AngVel
-				RepAngVel = CurrentState.AngVel + (BlendedTargetVelocity * VelocityAlpha); // Same as (BlendedTargetVelocity / InterpolationTime) * DeltaSeconds
-			}
-			else // Positional correction as transform shift
-			{
-				// Add velocity diff onto current velocity
-				RepAngVel = CurrentState.AngVel + (AngVelDiff * VelocityAlpha); // Same as (AngVelDiff / InterpolationTime) * DeltaSeconds
+					// Add BlendedTargetVelocity to CurrentState.AngVel
+					RepAngVel = CurrentState.AngVel + (BlendedTargetVelocity * VelocityAlpha); // Same as (BlendedTargetVelocity / InterpolationTime) * DeltaSeconds
+				}
+				else // Rotational correction as transform shift
+				{
+					// Add velocity diff onto current velocity
+					RepAngVel = CurrentState.AngVel + (AngVelDiff * VelocityAlpha); // Same as (AngVelDiff / InterpolationTime) * DeltaSeconds
 
-				// Calculate correction blend amount for this tick as an alpha value
-				const float CorrectionAlpha = FMath::Clamp(DeltaSeconds / RotCorrectionTime, 0.0f, 1.0f);
+					// Calculate correction blend amount for this tick as an alpha value
+					const float CorrectionAlpha = FMath::Clamp(DeltaSeconds / RotCorrectionTime, 0.0f, 1.0f);
 
-				// The new position after correction
-				CorrectionR = FQuat::Slerp(Handle->GetR(), TargetRot, CorrectionAlpha);
-			}
+					// The new position after correction
+					CorrectionR = FQuat::Slerp(Handle->GetR(), TargetRot, CorrectionAlpha);
+				}
 
-			// Apply velocity replication
-			Handle->SetW(RepAngVel);
-		}
+				// Apply velocity replication
+				Handle->SetW(RepAngVel);
+			}
 
-		// Cache data for next replication
-		Target.PrevPos = FVector(CurrentState.Position);
+			// Cache data for next replication
+			Target.PrevPos = FVector(CurrentState.Position);
 
-		// Apply correction as a transform shift
-		if (!PhysicsReplicationCVars::PredictiveInterpolationCVars::bCorrectionAsVelocity)
-		{	
-			const bool bCorrectConnectedBodies = SettingsCurrent.PredictiveInterpolationSettings.GetCorrectConnectedBodies();
-			const bool bCorrectConnectedBodiesFriction = SettingsCurrent.PredictiveInterpolationSettings.GetCorrectConnectedBodiesFriction();
-			RigidsSolver->GetEvolution()->ApplyParticleTransformCorrection(Handle, CorrectionX, CorrectionR, bCorrectConnectedBodies, bCorrectConnectedBodiesFriction, ReplicatedParticleIDs);
+			// Apply correction as a transform shift
+			if (SettingsCurrent.PredictiveInterpolationSettings.GetCorrectionAsVelocity() == false)
+			{	
+				const bool bCorrectConnectedBodies = SettingsCurrent.PredictiveInterpolationSettings.GetCorrectConnectedBodies();
+				const bool bCorrectConnectedBodiesFriction = SettingsCurrent.PredictiveInterpolationSettings.GetCorrectConnectedBodiesFriction();
+				RigidsSolver->GetEvolution()->ApplyParticleTransformCorrection(Handle, CorrectionX, CorrectionR, bCorrectConnectedBodies, bCorrectConnectedBodiesFriction, ReplicatedParticleIDs);
+			}
 		}
 
-		if (bSoftSnap)
+		// Transform shift replication, if either soft-snap is needed or velocity-based replication is false
+		if (bSoftSnap || SettingsCurrent.PredictiveInterpolationSettings.GetVelocityBased() == false)
 		{
 			const FVector SoftSnapPos = FMath::Lerp(FVector(CurrentState.Position),
 				SettingsCurrent.PredictiveInterpolationSettings.GetSoftSnapToSource() ? Target.PrevPosTarget : Target.TargetState.Position,
diff --git a/Engine/Source/Runtime/Engine/Public/Physics/NetworkPhysicsSettingsComponent.h b/Engine/Source/Runtime/Engine/Public/Physics/NetworkPhysicsSettingsComponent.h
--- a/Engine/Source/Runtime/Engine/Public/Physics/NetworkPhysicsSettingsComponent.h
+++ b/Engine/Source/Runtime/Engine/Public/Physics/NetworkPhysicsSettingsComponent.h
@@ -34,6 +34,8 @@ namespace PhysicsReplicationCVars
 	
 	namespace PredictiveInterpolationCVars
 	{
+		extern bool bVelocityBased;
+		extern bool bCorrectionAsVelocity;
 		extern float PosCorrectionTimeBase;
 		extern float PosCorrectionTimeMin;
 		extern float PosCorrectionTimeMultiplier;
@@ -75,6 +77,8 @@ namespace PhysicsReplicationCVars
 		extern bool bApplySimProxyStateAtRuntime;
 		extern bool bApplySimProxyInputAtRuntime;
 		extern bool bTriggerResimOnInputReceive;
+		extern bool bApplyInputDecaySimProxyInputAtRuntime;
+		extern float InputDecaySimProxyInputAtRuntime;
 		extern bool bApplyInputDecayOverSetTime;
 		extern float InputDecaySetTime;
 	}
@@ -155,7 +159,9 @@ struct FNetworkPhysicsSettingsPredictiveInterpolation
 	GENERATED_BODY()
 
 	FNetworkPhysicsSettingsPredictiveInterpolation() 
-		: bOverridePosCorrectionTimeBase(0)
+		: bOverrideVelocityBased(0)
+		, bOverrideCorrectionAsVelocity(0)
+		, bOverridePosCorrectionTimeBase(0)
 		, bOverridePosCorrectionTimeMin(0)
 		, bOverridePosCorrectionTimeMultiplier(0)
 		, bOverrideRotCorrectionTimeBase(0)
@@ -175,6 +181,10 @@ struct FNetworkPhysicsSettingsPredictiveInterpolation
 
 	// Override properties
 	UPROPERTY(config, EditDefaultsOnly, Category = "Overrides", Meta = (InlineEditConditionToggle))
+	uint32 bOverrideVelocityBased : 1;
+	UPROPERTY(config, EditDefaultsOnly, Category = "Overrides", Meta = (InlineEditConditionToggle))
+	uint32 bOverrideCorrectionAsVelocity : 1;
+	UPROPERTY(config, EditDefaultsOnly, Category = "Overrides", Meta = (InlineEditConditionToggle))
 	uint32 bOverridePosCorrectionTimeBase : 1;
 	UPROPERTY(config, EditDefaultsOnly, Category = "Overrides", Meta = (InlineEditConditionToggle))
 	uint32 bOverridePosCorrectionTimeMin : 1;
@@ -207,6 +217,16 @@ struct FNetworkPhysicsSettingsPredictiveInterpolation
 	UPROPERTY(config, EditDefaultsOnly, Category = "Overrides", Meta = (InlineEditConditionToggle))
 	uint32 bOverrideCorrectConnectedBodiesFriction : 1;
 
+	// Overrides CVar: np2.PredictiveInterpolation.VelocityBased -- When true, predictive interpolation replication mode will move objects replicate by applying linear velocity and angular velocity. When false, objects will move by transform shifting towards the target state.
+	UPROPERTY(config, EditDefaultsOnly, Category = "Overrides", Meta = (EditCondition = "bOverrideVelocityBased"))
+	bool bVelocityBased = PhysicsReplicationCVars::PredictiveInterpolationCVars::bVelocityBased;
+	bool GetVelocityBased() const { return bOverrideVelocityBased ? bVelocityBased : PhysicsReplicationCVars::PredictiveInterpolationCVars::bVelocityBased; }
+
+	// Overrides CVar: np2.PredictiveInterpolation.CorrectionAsVelocity -- When true, predictive interpolation will apply positional and rotational offset correction as a velocity instead of as a transform shift, while np2.PredictiveInterpolation.VelocityBased is true.
+	UPROPERTY(config, EditDefaultsOnly, Category = "Overrides", Meta = (EditCondition = "bOverrideCorrectionAsVelocity"))
+	bool bCorrectionAsVelocity = PhysicsReplicationCVars::PredictiveInterpolationCVars::bCorrectionAsVelocity;
+	bool GetCorrectionAsVelocity() const { return bOverrideCorrectionAsVelocity ? bCorrectionAsVelocity : PhysicsReplicationCVars::PredictiveInterpolationCVars::bCorrectionAsVelocity; }
+
 	// Overrides CVar: np2.PredictiveInterpolation.PosCorrectionTimeBase -- Base time to correct positional offset over. RoundTripTime * PosCorrectionTimeMultiplier is added on top of this.
 	UPROPERTY(config, EditDefaultsOnly, Category = "Overrides", Meta = (EditCondition = "bOverridePosCorrectionTimeBase"))
 	float PosCorrectionTimeBase = PhysicsReplicationCVars::PredictiveInterpolationCVars::PosCorrectionTimeBase;
@@ -483,9 +503,11 @@ struct FNetworkPhysicsSettingsNetworkPhysicsComponent
 		, bOverrideApplyDataInsteadOfMergeData(0)
 		, bOverrideAllowInputExtrapolation(0)
 		, bOverrideValidateDataOnGameThread(0)
+		, bOverrideTriggerResimOnInputReceive(0)
 		, bOverrideApplySimProxyStateAtRuntime(0)
 		, bOverrideApplySimProxyInputAtRuntime(0)
-		, bOverrideTriggerResimOnInputReceive(0)
+		, bOverrideApplyInputDecaySimProxyInputAtRuntime(0)
+		, bOverrideInputDecaySimProxyInputAtRuntime(0)
 		, bOverrideApplyInputDecayOverSetTime(0)
 		, bOverrideInputDecaySetTime(0)
 	{};
@@ -526,11 +548,15 @@ struct FNetworkPhysicsSettingsNetworkPhysicsComponent
 	UPROPERTY(config, EditDefaultsOnly, Category = "Overrides", Meta = (InlineEditConditionToggle))
 	uint32 bOverrideValidateDataOnGameThread : 1;
 	UPROPERTY(config, EditDefaultsOnly, Category = "Overrides", Meta = (InlineEditConditionToggle))
+	uint32 bOverrideTriggerResimOnInputReceive : 1;
+	UPROPERTY(config, EditDefaultsOnly, Category = "Overrides", Meta = (InlineEditConditionToggle))
 	uint32 bOverrideApplySimProxyStateAtRuntime : 1;
 	UPROPERTY(config, EditDefaultsOnly, Category = "Overrides", Meta = (InlineEditConditionToggle))
 	uint32 bOverrideApplySimProxyInputAtRuntime : 1;
 	UPROPERTY(config, EditDefaultsOnly, Category = "Overrides", Meta = (InlineEditConditionToggle))
-	uint32 bOverrideTriggerResimOnInputReceive : 1;
+	uint32 bOverrideApplyInputDecaySimProxyInputAtRuntime : 1;
+	UPROPERTY(config, EditDefaultsOnly, Category = "Overrides", Meta = (InlineEditConditionToggle))
+	uint32 bOverrideInputDecaySimProxyInputAtRuntime : 1;
 	UPROPERTY(config, EditDefaultsOnly, Category = "Overrides", Meta = (InlineEditConditionToggle))
 	uint32 bOverrideApplyInputDecayOverSetTime : 1;
 	UPROPERTY(config, EditDefaultsOnly, Category = "Overrides", Meta = (InlineEditConditionToggle))
@@ -596,6 +622,11 @@ struct FNetworkPhysicsSettingsNetworkPhysicsComponent
 	bool bValidateDataOnGameThread = PhysicsReplicationCVars::ResimulationCVars::bValidateDataOnGameThread;
 	const bool GetValidateDataOnGameThread() const { return bOverrideValidateDataOnGameThread ? bValidateDataOnGameThread : PhysicsReplicationCVars::ResimulationCVars::bValidateDataOnGameThread; }
 
+	// Overrides CVar: np2.Resim.TriggerResimOnInputReceive -- When true, a resimulation will be requested to the frame of the latest frame of received inputs this frame
+	UPROPERTY(config, EditDefaultsOnly, Category = "Overrides", Meta = (EditCondition = "bOverrideTriggerResimOnInputReceive"))
+	bool bTriggerResimOnInputReceive = PhysicsReplicationCVars::ResimulationCVars::bTriggerResimOnInputReceive;
+	const bool GetTriggerResimOnInputReceive() const { return bOverrideTriggerResimOnInputReceive ? bTriggerResimOnInputReceive : PhysicsReplicationCVars::ResimulationCVars::bTriggerResimOnInputReceive; }
+
 	// Overrides CVar: np2.Resim.ApplySimProxyStateAtRuntime -- When true, call ApplyData on received states for simulated proxies at runtime.
 	UPROPERTY(config, EditDefaultsOnly, Category = "Overrides", Meta = (EditCondition = "bOverrideApplySimProxyStateAtRuntime"))
 	bool bApplySimProxyStateAtRuntime = PhysicsReplicationCVars::ResimulationCVars::bApplySimProxyStateAtRuntime;
@@ -606,10 +637,15 @@ struct FNetworkPhysicsSettingsNetworkPhysicsComponent
 	bool bApplySimProxyInputAtRuntime = PhysicsReplicationCVars::ResimulationCVars::bApplySimProxyInputAtRuntime;
 	const bool GetApplySimProxyInputAtRuntime() const { return bOverrideApplySimProxyInputAtRuntime ? bApplySimProxyInputAtRuntime : PhysicsReplicationCVars::ResimulationCVars::bApplySimProxyInputAtRuntime; }
 
-	// Overrides CVar: np2.Resim.TriggerResimOnInputReceive -- When true, a resimulation will be requested to the frame of the latest frame of received inputs this frame
-	UPROPERTY(config, EditDefaultsOnly, Category = "Overrides", Meta = (EditCondition = "bOverrideTriggerResimOnInputReceive"))
-	bool bTriggerResimOnInputReceive = PhysicsReplicationCVars::ResimulationCVars::bTriggerResimOnInputReceive;
-	const bool GetTriggerResimOnInputReceive() const { return bOverrideTriggerResimOnInputReceive ? bTriggerResimOnInputReceive : PhysicsReplicationCVars::ResimulationCVars::bTriggerResimOnInputReceive; }
+	// Overrides CVar: np2.Resim.ApplyInputDecaySimProxyInputAtRuntime -- When true, apply input decay on inputs applied on simulated proxies at runtime (outside of resimulation)
+	UPROPERTY(config, EditDefaultsOnly, Category = "Overrides", Meta = (EditCondition = "bOverrideApplyInputDecaySimProxyInputAtRuntime"))
+	bool bApplyInputDecaySimProxyInputAtRuntime = PhysicsReplicationCVars::ResimulationCVars::bApplyInputDecaySimProxyInputAtRuntime;
+	const bool GetApplyInputDecaySimProxyInputAtRuntime() const { return bOverrideApplyInputDecaySimProxyInputAtRuntime ? bApplyInputDecaySimProxyInputAtRuntime : PhysicsReplicationCVars::ResimulationCVars::bApplyInputDecaySimProxyInputAtRuntime; }
+
+	// Overrides CVar: np2.Resim.InputDecaySimProxyInputAtRuntime -- The amount of input decay to apply on simulated proxy inputs at runtime (outside of resim), if np2.Resim.ApplyInputDecaySimProxyInputAtRuntime is true
+	UPROPERTY(config, EditDefaultsOnly, Category = "Overrides", Meta = (EditCondition = "bOverrideInputDecaySimProxyInputAtRuntime"))
+	float InputDecaySimProxyInputAtRuntime = PhysicsReplicationCVars::ResimulationCVars::InputDecaySimProxyInputAtRuntime;
+	const float GetInputDecaySimProxyInputAtRuntime() const { return bOverrideInputDecaySimProxyInputAtRuntime ? InputDecaySimProxyInputAtRuntime : PhysicsReplicationCVars::ResimulationCVars::InputDecaySimProxyInputAtRuntime; }
 
 	// Overrides CVar: np2.Resim.ApplyInputDecayOverSetTime -- When true, apply the Input Decay Curve over a set amount of time instead of over the start of input prediction and end of resim which is variable each resimulation
 	UPROPERTY(config, EditDefaultsOnly, Category = "Overrides", Meta = (EditCondition = "bOverrideApplyInputDecayOverSetTime"))
diff --git a/Engine/Source/Runtime/Experimental/Chaos/Private/RewindData.cpp b/Engine/Source/Runtime/Experimental/Chaos/Private/RewindData.cpp
--- a/Engine/Source/Runtime/Experimental/Chaos/Private/RewindData.cpp
+++ b/Engine/Source/Runtime/Experimental/Chaos/Private/RewindData.cpp
@@ -408,6 +408,9 @@ bool FRewindData::CheckQuaternionThreshold(FQuat A, FQuat B, float ThresholdDegr
 }
 
 
+bool bForceSkipResim = false;
+FAutoConsoleVariableRef CVarForceSkipResim(TEXT("p.Resim.ForceSkipResim"), bForceSkipResim, TEXT("Don't perform resimulations while this is true"));
+
 CHAOS_API bool bResimAllowRewindToResimulatedFrames = true;
 FAutoConsoleVariableRef CVarResimAllowRewindToResimulatedFrames(TEXT("p.Resim.AllowRewindToResimulatedFrames"), bResimAllowRewindToResimulatedFrames, TEXT("Allow rewinding back to a frame that was previously part of a resimulation. If a resimulation is performed between frame 100-110, allow a new resim from 105-115 if needed, else next resim will be able to start from frame 111."));
 
@@ -421,6 +424,12 @@ bool FRewindData::RewindToFrame(int32 Frame)
 	CVD_SCOPE_TRACE_SOLVER_STEP(CVDDC_Default, TEXT("Rewind To Frame"));
 
 	EnsureIsInPhysicsThreadContext();
+
+	if (bForceSkipResim)
+	{
+		return false;
+	}
+
 	//Can't go too far back
 	const int32 EarliestFrame = GetEarliestFrame_Internal();
 	if (Frame < EarliestFrame)
