[PostProcessing] Add r.PostProcessing.Debug.Property. When set to the name of a property in PostProcessSettings, the source of the value for that property will be shown in ShowFlag.VisualizePostProcessStack

diff --git a/Engine/Source/Runtime/Engine/Classes/Engine/Scene.h b/Engine/Source/Runtime/Engine/Classes/Engine/Scene.h
--- a/Engine/Source/Runtime/Engine/Classes/Engine/Scene.h
+++ b/Engine/Source/Runtime/Engine/Classes/Engine/Scene.h
@@ -682,6 +682,8 @@ struct FPostProcessSettingsDebugInfo
 	float CurrentBlendWeight;
 	bool bIsEnabled;
 	bool bIsUnbound;
+
+	FString DebugPropertyValue;
 };
 
 /** To be able to use struct PostProcessSettings. */
diff --git a/Engine/Source/Runtime/Engine/Private/World.cpp b/Engine/Source/Runtime/Engine/Private/World.cpp
--- a/Engine/Source/Runtime/Engine/Private/World.cpp
+++ b/Engine/Source/Runtime/Engine/Private/World.cpp
@@ -10363,6 +10363,98 @@ void UWorld::PostInitializeSubsystems()
 	});
 }
 
+#if DEBUG_POST_PROCESS_VOLUME_ENABLE
+static FProperty* GPostProcessDebugProperty = nullptr;
+static FBoolProperty* GPostProcessDebugOverrideProperty = nullptr;
+static TAutoConsoleVariable<FString> CVarPostProcessDebugProperty(
+	TEXT("r.PostProcessing.Debug.Property"),
+	TEXT(""),
+	TEXT("When set to the name of a property in PostProcessSettings, the source of the value for that property will be shown in ShowFlag.VisualizePostProcessStack"),
+	FConsoleVariableDelegate::CreateLambda(
+		[](IConsoleVariable* CVar)
+		{
+			FString NewPropertyName = CVar->GetString();
+
+			if (NewPropertyName.IsEmpty())
+			{
+				GPostProcessDebugProperty = nullptr;
+			}
+			else
+			{
+				GPostProcessDebugProperty = FindFProperty<FProperty>(FPostProcessSettings::StaticStruct(), FName(NewPropertyName));
+				GPostProcessDebugOverrideProperty = FindFProperty<FBoolProperty>(FPostProcessSettings::StaticStruct(), FName(TEXT("bOverride_") + NewPropertyName));
+				if (GPostProcessDebugProperty)
+				{
+					FString FullName = GPostProcessDebugProperty->GetFullName();
+					UE_LOG(LogWorld, Display, TEXT("Visualizing %s"), *FullName);
+				}
+				else
+				{
+					TStringBuilder<1024> StringBuilder {};
+					for (TFieldIterator<FProperty>It(FPostProcessSettings::StaticStruct(), EFieldIterationFlags::Default); It; ++It)
+					{
+						if (!It->GetFName().ToString().StartsWith(TEXT("bOverride_")))
+						{
+							StringBuilder.Append(It->GetFName().ToString());
+							StringBuilder.Append(TEXT("\n"));
+						}
+					}
+					FString ValidPropertyNames = StringBuilder.ToString();
+
+					UE_LOG(LogWorld, Error, 
+						   TEXT("Cannot find %s\n")
+						   TEXT("Valid property names are:\n%s"),
+						   *NewPropertyName,
+						   *ValidPropertyNames);
+				}
+			}
+		}
+	)
+);
+
+static void GatherPostProcessVolumeDebugInfo(IInterface_PostProcessVolume* Volume, float LocalWeight, FPostProcessSettingsDebugInfo& Output)
+{
+	Output.Name = Volume->GetDebugName();
+
+	FPostProcessVolumeProperties VProperties = Volume->GetProperties();
+	Output.bIsEnabled = VProperties.bIsEnabled;
+	Output.bIsUnbound = VProperties.bIsUnbound;
+	Output.Priority = VProperties.Priority;
+	Output.CurrentBlendWeight = LocalWeight;
+
+	if (GPostProcessDebugProperty != nullptr)
+	{
+		bool bOverrideEnabled = true; // if we couldn't find the bOverride_ property, assume it's true
+		if (GPostProcessDebugOverrideProperty)
+		{
+			const bool* bOverridePtr = GPostProcessDebugOverrideProperty->ContainerPtrToValuePtr<bool>(VProperties.Settings);
+			if (bOverridePtr)
+			{
+				bOverrideEnabled = *bOverridePtr;
+			}
+		}
+
+		if (bOverrideEnabled)
+		{
+			const void* PropData = GPostProcessDebugProperty->ContainerPtrToValuePtr<void>(VProperties.Settings);
+			FString PropValueAsString;
+			if (GPostProcessDebugProperty->ExportText_Direct(PropValueAsString, PropData, PropData, (UObject*)VProperties.Settings, PPF_DebugDump))
+			{
+				Output.DebugPropertyValue = PropValueAsString;
+			}
+			else
+			{
+				Output.DebugPropertyValue = TEXT("<failed to print value>");
+			}
+		}
+		else
+		{
+			Output.DebugPropertyValue = TEXT("N/A");
+		}
+	}
+}
+#endif
+
 static void DoPostProcessVolume(IInterface_PostProcessVolume* Volume, FVector ViewLocation, FSceneView* SceneView)
 {
 	const FPostProcessVolumeProperties VolumeProperties = Volume->GetProperties();
@@ -10410,25 +10502,18 @@ static void DoPostProcessVolume(IInterface_PostProcessVolume* Volume, FVector Vi
 		}
 	}
 
+	if (LocalWeight > 0)
+	{
+		SceneView->OverridePostProcessSettings(*VolumeProperties.Settings, LocalWeight);
+	}
+
 #if DEBUG_POST_PROCESS_VOLUME_ENABLE
 	if (SceneView->Family && SceneView->Family->EngineShowFlags.VisualizePostProcessStack)
 	{
 		FPostProcessSettingsDebugInfo& PPDebug = SceneView->FinalPostProcessDebugInfo.AddDefaulted_GetRef();
-
-		PPDebug.Name = Volume->GetDebugName();
-
-		FPostProcessVolumeProperties VProperties = Volume->GetProperties();
-		PPDebug.bIsEnabled = VProperties.bIsEnabled;
-		PPDebug.bIsUnbound = VProperties.bIsUnbound;
-		PPDebug.Priority = VProperties.Priority;
-		PPDebug.CurrentBlendWeight = LocalWeight;
+		GatherPostProcessVolumeDebugInfo(Volume, LocalWeight, PPDebug);
 	}
 #endif
-
-	if (LocalWeight > 0)
-	{
-		SceneView->OverridePostProcessSettings(*VolumeProperties.Settings, LocalWeight);
-	}
 }
 
 void UWorld::AddPostProcessingSettings(FVector ViewLocation, FSceneView* SceneView)
@@ -10436,6 +10521,7 @@ void UWorld::AddPostProcessingSettings(FVector ViewLocation, FSceneView* SceneVi
 	OnBeginPostProcessSettings.Broadcast(ViewLocation, SceneView);
 
 #if DEBUG_POST_PROCESS_VOLUME_ENABLE
+	SceneView->PostProcessDebugPropertyName = GPostProcessDebugProperty ? GPostProcessDebugProperty->GetFName() : NAME_None;
 	SceneView->FinalPostProcessDebugInfo.Reset();
 #endif
 
diff --git a/Engine/Source/Runtime/Engine/Public/SceneView.h b/Engine/Source/Runtime/Engine/Public/SceneView.h
--- a/Engine/Source/Runtime/Engine/Public/SceneView.h
+++ b/Engine/Source/Runtime/Engine/Public/SceneView.h
@@ -1781,6 +1781,7 @@ public:
 	 */
 	FFinalPostProcessSettings FinalPostProcessSettings;
 #if DEBUG_POST_PROCESS_VOLUME_ENABLE
+	FName PostProcessDebugPropertyName;
 	TArray<FPostProcessSettingsDebugInfo> FinalPostProcessDebugInfo;
 #endif
 
diff --git a/Engine/Source/Runtime/Renderer/Private/PostProcess/PostProcessing.cpp b/Engine/Source/Runtime/Renderer/Private/PostProcess/PostProcessing.cpp
--- a/Engine/Source/Runtime/Renderer/Private/PostProcess/PostProcessing.cpp
+++ b/Engine/Source/Runtime/Renderer/Private/PostProcess/PostProcessing.cpp
@@ -3525,18 +3525,19 @@ FScreenPassTexture AddFinalPostProcessDebugInfoPasses(FRDGBuilder& GraphBuilder,
 {
 	RDG_EVENT_SCOPE(GraphBuilder, "FinalPostProcessDebugInfo");
 
-	FRDGTextureRef SceneColor = ScreenPassSceneColor.Texture;
-
-	AddDrawCanvasPass(GraphBuilder, RDG_EVENT_NAME("PostProcessDebug"), View, FScreenPassRenderTarget(SceneColor, View.ViewRect, ERenderTargetLoadAction::ELoad),
+	AddDrawCanvasPass(GraphBuilder, RDG_EVENT_NAME("PostProcessDebug"), View, FScreenPassRenderTarget(ScreenPassSceneColor, ERenderTargetLoadAction::ELoad),
 		[&View](FCanvas& Canvas)
 		{
+			const float DPIScale = Canvas.GetDPIScale();
+			Canvas.SetBaseTransform(FMatrix(FScaleMatrix(DPIScale) * Canvas.CalcBaseTransform2D(Canvas.GetViewRect().Width(), Canvas.GetViewRect().Height())));
+
 			FLinearColor TextColor(FLinearColor::White);
 			FLinearColor GrayTextColor(FLinearColor::Gray);
 			FLinearColor GreenTextColor(FLinearColor::Green);
 			FString Text;
 
-			const float ViewPortWidth = float(View.ViewRect.Width());
-			const float ViewPortHeight = float(View.ViewRect.Height());
+			const float ViewPortWidth = float(Canvas.GetViewRect().Width());
+			const float ViewPortHeight = float(Canvas.GetViewRect().Height());
 
 			const float CRHeight = 20.0f;
 			const float PrintX_CR = ViewPortWidth * 0.1f;
@@ -3552,6 +3553,12 @@ FScreenPassTexture AddFinalPostProcessDebugInfoPasses(FRDGBuilder& GraphBuilder,
 			Canvas.DrawShadowedString(PrintX, PrintY, *FString("Priority"), GetStatsFont(), GrayTextColor);			PrintX += 96.0f;
 			Canvas.DrawShadowedString(PrintX, PrintY, *FString("CurrentWeight"), GetStatsFont(), GrayTextColor);	PrintX += 96.0f;
 			Canvas.DrawShadowedString(PrintX, PrintY, *FString("bIsUnbound"), GetStatsFont(), GrayTextColor);		PrintX += 96.0f;
+
+			if (View.PostProcessDebugPropertyName != NAME_None)
+			{
+				FString Name = View.PostProcessDebugPropertyName.ToString();
+				Canvas.DrawShadowedString(PrintX, PrintY, *Name, GetStatsFont(), GrayTextColor);		PrintX += 96.0f;
+			}
 			
 			PrintY += CRHeight;
 			PrintX = PrintX_CR;
@@ -3576,11 +3583,23 @@ FScreenPassTexture AddFinalPostProcessDebugInfoPasses(FRDGBuilder& GraphBuilder,
 				Text = FString::Printf(TEXT("%d"), PPDebugInfo.bIsUnbound ? 1 : 0);
 				Canvas.DrawShadowedString(PrintX+32.0f, PrintY, *Text, GetStatsFont(), TextColor); PrintX += 96.0f;
 
+				if (View.PostProcessDebugPropertyName != NAME_None)
+				{
+					Text = FString::Printf(TEXT("%s"), *PPDebugInfo.DebugPropertyValue.Left(256));
+					Canvas.DrawShadowedString(PrintX, PrintY, *Text, GetStatsFont(), TextColor); PrintX += 256.0f;
+				}
+
 				Canvas.DrawShadowedString(PrintX_CR, PrintY+3.0f, *FString("______________________________________________________________________________________________________________"), GetStatsFont(), TextColor);
 
 				PrintX = PrintX_CR;
 				PrintY += CRHeight;
 			}
+
+			PrintY += CRHeight;
+			if (View.PostProcessDebugPropertyName == NAME_None)
+			{
+				Canvas.DrawShadowedString(PrintX, PrintY, TEXT("Use r.PostProcessing.Debug.Property <SomePropertyName> to visualize property values"), GetStatsFont(), GrayTextColor);
+			}
 		});
 
 	return MoveTemp(ScreenPassSceneColor);
