Preshader2 compile error fix:  forward declare the preshader data context and remove access to private PreshaderEvaluate.h header.

#rnx

diff --git a/Engine/Source/Runtime/Engine/Private/Shader/Preshader2.cpp b/Engine/Source/Runtime/Engine/Private/Shader/Preshader2.cpp
--- a/Engine/Source/Runtime/Engine/Private/Shader/Preshader2.cpp
+++ b/Engine/Source/Runtime/Engine/Private/Shader/Preshader2.cpp
@@ -1,6 +1,7 @@
 // Copyright Epic Games, Inc. All Rights Reserved.
 
 #include "Shader/Preshader2.h"
+#include "Shader/PreshaderEvaluate.h"
 #include "Materials/MaterialRenderProxy.h"
 #include "Materials/MaterialUniformExpressions.h"
 #include "Math/DoubleFloat.h"
diff --git a/Engine/Source/Runtime/Engine/Public/Shader/Preshader2.h b/Engine/Source/Runtime/Engine/Public/Shader/Preshader2.h
--- a/Engine/Source/Runtime/Engine/Public/Shader/Preshader2.h
+++ b/Engine/Source/Runtime/Engine/Public/Shader/Preshader2.h
@@ -8,7 +8,11 @@
 
 #include "Materials/MaterialParameters.h"
 #include "Shader/Preshader.h"
-#include "Shader/PreshaderEvaluate.h"
+
+namespace UE::Shader
+{
+struct FPreshaderDataContext;
+}
 
 #define PRESHADER2_UNARY_OP(OP_NAME)	\
 	OP_NAME,									\
