SkyLight - fix for max sky resolution in fact being a set resolution.

#tests pc editor dx12
#rnx

diff --git a/Engine/Source/Runtime/Engine/Private/Components/SkyLightComponent.cpp b/Engine/Source/Runtime/Engine/Private/Components/SkyLightComponent.cpp
--- a/Engine/Source/Runtime/Engine/Private/Components/SkyLightComponent.cpp
+++ b/Engine/Source/Runtime/Engine/Private/Components/SkyLightComponent.cpp
@@ -363,13 +363,8 @@ void USkyLightComponent::SetCaptureIsDirty()
 
 void USkyLightComponent::SanitizeCubemapSize()
 {
-	const int32 MaxCubemapResolution = GetMaxCubeTextureDimension();
 	const int32 MinCubemapResolution = 8;
-
-	if (GSkylightCubemapMaxResolution > 0)
-	{
-		CubemapResolution = GSkylightCubemapMaxResolution;
-	}
+	const int32 MaxCubemapResolution = GSkylightCubemapMaxResolution > 0 ? FMath::Clamp(uint32(GSkylightCubemapMaxResolution), uint32(MinCubemapResolution), GetMaxCubeTextureDimension()) : GetMaxCubeTextureDimension();
 
 	CubemapResolution = FMath::Clamp(int32(FMath::RoundUpToPowerOfTwo(CubemapResolution)), MinCubemapResolution, MaxCubemapResolution);
 
