Single File Compilation Fixes

diff --git a/Engine/Source/Runtime/Engine/Classes/PhysicsEngine/PhysicsConstraintComponent.h b/Engine/Source/Runtime/Engine/Classes/PhysicsEngine/PhysicsConstraintComponent.h
--- a/Engine/Source/Runtime/Engine/Classes/PhysicsEngine/PhysicsConstraintComponent.h
+++ b/Engine/Source/Runtime/Engine/Classes/PhysicsEngine/PhysicsConstraintComponent.h
@@ -13,6 +13,10 @@ class AActor;
 class UPrimitiveComponent;
 struct FBodyInstance;
 
+#if UE_WITH_REMOTE_OBJECT_HANDLE
+enum class EObjectMigrationSendType : uint8;
+#endif
+
 /**
  *	This is effectively a joint that allows you to connect 2 rigid bodies together. You can create different types of joints using the various parameters of this component.
  */
diff --git a/Engine/Source/Runtime/Engine/Private/PhysicsEngine/PhysicsConstraintComponent.cpp b/Engine/Source/Runtime/Engine/Private/PhysicsEngine/PhysicsConstraintComponent.cpp
--- a/Engine/Source/Runtime/Engine/Private/PhysicsEngine/PhysicsConstraintComponent.cpp
+++ b/Engine/Source/Runtime/Engine/Private/PhysicsEngine/PhysicsConstraintComponent.cpp
@@ -3,6 +3,7 @@
 #include "PhysicsEngine/PhysicsConstraintComponent.h"
 #include "Components/SkeletalMeshComponent.h"
 #include "Engine/Texture2D.h"
+#include "Engine/World.h"
 #include "Logging/MessageLog.h"
 #include "Misc/UObjectToken.h"
 #include "PhysicsEngine/PhysicsConstraintTemplate.h"
