Fixed typos in IPhysicsBodyInstanceOwner : Both GetPhysicsBodyInstandeOwnerFromHitResult and GetPhysicsBodyInstandeOwnerFromOverlapResult were deprecated and renamed to GetPhysicsBodyInstanceOwnerFromHitResult & GetPhysicsBodyInstanceOwnerFromOverlapResult.

diff --git a/Engine/Plugins/Experimental/Mover/Source/Mover/Private/DefaultMovementSet/Modes/AsyncNavWalkingMode.cpp b/Engine/Plugins/Experimental/Mover/Source/Mover/Private/DefaultMovementSet/Modes/AsyncNavWalkingMode.cpp
--- a/Engine/Plugins/Experimental/Mover/Source/Mover/Private/DefaultMovementSet/Modes/AsyncNavWalkingMode.cpp
+++ b/Engine/Plugins/Experimental/Mover/Source/Mover/Private/DefaultMovementSet/Modes/AsyncNavWalkingMode.cpp
@@ -470,7 +470,7 @@ void UAsyncNavWalkingMode::FindBestNavMeshLocation(const FVector& TraceStart, co
 		{
 			UPrimitiveComponent* PrimComp = TestHit.GetComponent();
 			// Prefer using primitive component if valid
-			if (IPhysicsBodyInstanceOwner* BodyInstanceOwner = !PrimComp ? IPhysicsBodyInstanceOwner::GetPhysicsBodyInstandeOwnerFromHitResult(TestHit) : nullptr)
+			if (IPhysicsBodyInstanceOwner* BodyInstanceOwner = !PrimComp ? IPhysicsBodyInstanceOwner::GetPhysicsBodyInstanceOwnerFromHitResult(TestHit) : nullptr)
 			{
 				const bool bBlockOnWorldStatic = (BodyInstanceOwner->GetCollisionResponseToChannel(ECC_WorldStatic) == ECR_Block);
 				const bool bBlockOnWorldDynamic = (BodyInstanceOwner->GetCollisionResponseToChannel(ECC_WorldDynamic) == ECR_Block);
diff --git a/Engine/Source/Runtime/Engine/Classes/PhysicsEngine/PhysicsBodyInstanceOwnerInterface.h b/Engine/Source/Runtime/Engine/Classes/PhysicsEngine/PhysicsBodyInstanceOwnerInterface.h
--- a/Engine/Source/Runtime/Engine/Classes/PhysicsEngine/PhysicsBodyInstanceOwnerInterface.h
+++ b/Engine/Source/Runtime/Engine/Classes/PhysicsEngine/PhysicsBodyInstanceOwnerInterface.h
@@ -36,11 +36,23 @@ class IPhysicsBodyInstanceOwner
 public:
 	virtual ~IPhysicsBodyInstanceOwner() = default;
 
+	UE_DEPRECATED(5.8, "Use GetPhysicsBodyInstanceOwnerFromHitResult instead.")
+	static IPhysicsBodyInstanceOwner* GetPhysicsBodyInstandeOwnerFromHitResult(const FHitResult& Result)
+	{
+		return GetPhysicsBodyInstanceOwnerFromHitResult(Result);
+	}
+
+	UE_DEPRECATED(5.8, "Use GetPhysicsBodyInstanceOwnerFromOverlapResult instead.")
+	static IPhysicsBodyInstanceOwner* GetPhysicsBodyInstandeOwnerFromOverlapResult(const FOverlapResult& OverlapResult)
+	{
+		return GetPhysicsBodyInstanceOwnerFromOverlapResult(OverlapResult);
+	}
+
 	/** Returns the IPhysicsBodyInstanceOwner based on a given hit result. */
-	static ENGINE_API IPhysicsBodyInstanceOwner* GetPhysicsBodyInstandeOwnerFromHitResult(const FHitResult& Result);
+	static ENGINE_API IPhysicsBodyInstanceOwner* GetPhysicsBodyInstanceOwnerFromHitResult(const FHitResult& Result);
 
 	/** Returns the IPhysicsBodyInstanceOwner based on a given overlap result. */
-	static ENGINE_API IPhysicsBodyInstanceOwner* GetPhysicsBodyInstandeOwnerFromOverlapResult(const FOverlapResult& OverlapResult);
+	static ENGINE_API IPhysicsBodyInstanceOwner* GetPhysicsBodyInstanceOwnerFromOverlapResult(const FOverlapResult& OverlapResult);
 
 	/** Whether the physics is static. */
 	virtual bool IsStaticPhysics() const = 0;
diff --git a/Engine/Source/Runtime/Engine/Private/PhysicsEngine/PhysicsBodyInstanceOwnerInterface.cpp b/Engine/Source/Runtime/Engine/Private/PhysicsEngine/PhysicsBodyInstanceOwnerInterface.cpp
--- a/Engine/Source/Runtime/Engine/Private/PhysicsEngine/PhysicsBodyInstanceOwnerInterface.cpp
+++ b/Engine/Source/Runtime/Engine/Private/PhysicsEngine/PhysicsBodyInstanceOwnerInterface.cpp
@@ -6,7 +6,7 @@
 
 #include UE_INLINE_GENERATED_CPP_BY_NAME(PhysicsBodyInstanceOwnerInterface)
 
-IPhysicsBodyInstanceOwner* IPhysicsBodyInstanceOwner::GetPhysicsBodyInstandeOwnerFromHitResult(const FHitResult& Result)
+IPhysicsBodyInstanceOwner* IPhysicsBodyInstanceOwner::GetPhysicsBodyInstanceOwnerFromHitResult(const FHitResult& Result)
 {
 	if (IPhysicsBodyInstanceOwnerResolver* PhysicsObjectOwnerResolver = Result.PhysicsObject ? Cast<IPhysicsBodyInstanceOwnerResolver>(Result.PhysicsObjectOwner.Get()) : nullptr)
 	{
@@ -15,7 +15,7 @@ IPhysicsBodyInstanceOwner* IPhysicsBodyInstanceOwner::GetPhysicsBodyInstandeOwne
 	return nullptr;
 }
 
-IPhysicsBodyInstanceOwner* IPhysicsBodyInstanceOwner::GetPhysicsBodyInstandeOwnerFromOverlapResult(const FOverlapResult& Result)
+IPhysicsBodyInstanceOwner* IPhysicsBodyInstanceOwner::GetPhysicsBodyInstanceOwnerFromOverlapResult(const FOverlapResult& Result)
 {
 	if (IPhysicsBodyInstanceOwnerResolver* PhysicsObjectOwnerResolver = Result.PhysicsObject ? Cast<IPhysicsBodyInstanceOwnerResolver>(Result.PhysicsObjectOwner.Get()) : nullptr)
 	{
