Follow-up minor fixes and disabled the skinning extension by default to investigate crash.

diff --git a/Engine/Source/Runtime/Engine/Internal/SkeletalMeshUpdater.h b/Engine/Source/Runtime/Engine/Internal/SkeletalMeshUpdater.h
--- a/Engine/Source/Runtime/Engine/Internal/SkeletalMeshUpdater.h
+++ b/Engine/Source/Runtime/Engine/Internal/SkeletalMeshUpdater.h
@@ -306,7 +306,7 @@ public:
 	void InvalidatePathTracedOutput()
 	{
 #if RHI_RAYTRACING
-		bInvalidatePathTracedOutput |= bInvalidatePathTracedOutput;
+		bInvalidatePathTracedOutput = true;
 #endif
 	}
 
@@ -400,7 +400,7 @@ public:
 	void Release();
 
 private:
-	FSkeletalMeshUpdateChannel* Channel;
+	FSkeletalMeshUpdateChannel* Channel = nullptr;
 	uint32 Index = INDEX_NONE;
 
 	friend FSkeletalMeshUpdateChannel;
@@ -458,7 +458,7 @@ public:
 		FSubmitData() = default;
 		FSubmitData(const FSubmitData&) = delete;
 
-		// Returns set of LOD changes this update. Must wait on the Filter task before calling.
+		// Returns set of events that occurred during this update. Must wait on the Filter task before calling.
 		TConstArrayView<FSkeletalMeshUpdateEvent> GetUpdateEvents() const
 		{
 			check(Tasks.Filter.IsCompleted());
diff --git a/Engine/Source/Runtime/Engine/Private/SkeletalRenderGPUSkin.cpp b/Engine/Source/Runtime/Engine/Private/SkeletalRenderGPUSkin.cpp
--- a/Engine/Source/Runtime/Engine/Private/SkeletalRenderGPUSkin.cpp
+++ b/Engine/Source/Runtime/Engine/Private/SkeletalRenderGPUSkin.cpp
@@ -130,7 +130,7 @@ EBoneTransformStorageMode GetBoneTransformStorageMode()
 #endif
 }
 
-bool GGPUSkinSceneExtensionEnabled = true;
+bool GGPUSkinSceneExtensionEnabled = false;
 static FAutoConsoleVariableRef CVarGPUSkinUseSceneExtension(
 	TEXT("r.GPUSkin.UseSceneExtension"),
 	GGPUSkinSceneExtensionEnabled,
