For networking EngineTests, move setting of net driver MaxNetTickRate into FScopedTestSettings.
This covers the multi-server proxy tests too and fixes sporadic failures due to inconsistent timing.
Use FScopedTestSettings in FTestMultiServerWorlds, like FTestWorlds.

diff --git a/Engine/Source/Runtime/Engine/Private/Tests/NetTestHelpers.cpp b/Engine/Source/Runtime/Engine/Private/Tests/NetTestHelpers.cpp
--- a/Engine/Source/Runtime/Engine/Private/Tests/NetTestHelpers.cpp
+++ b/Engine/Source/Runtime/Engine/Private/Tests/NetTestHelpers.cpp
@@ -334,16 +334,12 @@ void FTestWorldInstance::UnloadStreamingLevel(FName LevelName)
 FTestWorlds::FTestWorlds(const EReplicationSystem ReplicationSystem)
 	: Server(true)
 {
-	InitDelegates();
-
 	Server = FTestWorldInstance::CreateServer(TEXT("/Engine/Maps/Entry?listen?Game=/Script/Engine.GameMode"), ReplicationSystem);
 }
 
 FTestWorlds::FTestWorlds(const FString& MapName, const FString& GameModeName, const EReplicationSystem ReplicationSystem)
 	: Server(true)
 {
-	InitDelegates();
-
 	FString URL = FString::Printf(TEXT("%s?listen?Game=%s"), *MapName, *GameModeName);
 
 	Server = FTestWorldInstance::CreateServer(*URL, ReplicationSystem);
@@ -352,14 +348,11 @@ FTestWorlds::FTestWorlds(const FString& MapName, const FString& GameModeName, co
 FTestWorlds::FTestWorlds(const TCHAR* ServerURL, const EReplicationSystem ReplicationSystem)
 	: Server(true)
 {
-	InitDelegates();
-
 	Server = FTestWorldInstance::CreateServer(ServerURL, ReplicationSystem);
 }
 
 FTestWorlds::~FTestWorlds()
 {
-	FWorldDelegates::OnNetDriverCreated.Remove(NetDriverCreatedHandle);
 }
 
 void FTestWorlds::SetTickInSeconds(float TickInSeconds)
@@ -367,17 +360,6 @@ void FTestWorlds::SetTickInSeconds(float TickInSeconds)
 	TickDeltaSeconds = TickInSeconds;
 }
 
-void FTestWorlds::InitDelegates()
-{
-	NetDriverCreatedHandle = FWorldDelegates::OnNetDriverCreated.AddRaw(this, &FTestWorlds::OnNetDriverCreated);
-}
-
-void FTestWorlds::OnNetDriverCreated(UWorld* InWorld, UNetDriver* InNetDriver)
-{
-	// Make sure NetDriver will tick every engine frame
-	InNetDriver->MaxNetTickRate = 0;
-}
-
 bool FTestWorlds::CreateAndConnectClient()
 {
 	Clients.Emplace(FTestWorldInstance::CreateClient(Server.GetPort()));
@@ -667,10 +649,16 @@ FScopedTestSettings::FScopedTestSettings()
 	, RandomNetUpdateDelayDisabled(TEXT("net.DisableRandomNetUpdateDelay"), 1)
 	, GameplayDebuggerDisabled(TEXT("GameplayDebugger.AutoCreateGameplayDebuggerManager"), 0)
 {
+	NetDriverCreatedHandle = FWorldDelegates::OnNetDriverCreated.AddLambda([](UWorld* InWorld, UNetDriver* InNetDriver)
+		{
+			// Make sure NetDriver will tick every engine frame
+			InNetDriver->MaxNetTickRate = 0;
+		});
 }
 
 FScopedTestSettings::~FScopedTestSettings()
 {
+	FWorldDelegates::OnNetDriverCreated.Remove(NetDriverCreatedHandle);
 }
 
 //------------------------------------------------------------------------
diff --git a/Engine/Source/Runtime/Engine/Public/Tests/NetTestHelpers.h b/Engine/Source/Runtime/Engine/Public/Tests/NetTestHelpers.h
--- a/Engine/Source/Runtime/Engine/Public/Tests/NetTestHelpers.h
+++ b/Engine/Source/Runtime/Engine/Public/Tests/NetTestHelpers.h
@@ -85,6 +85,8 @@ private:
 	FScopedCVarOverrideInt RepGraphBandwidthThrottlingDisabled;
 	FScopedCVarOverrideInt RandomNetUpdateDelayDisabled;
 	FScopedCVarOverrideInt GameplayDebuggerDisabled;
+
+	FDelegateHandle NetDriverCreatedHandle;
 };
 
 /**
@@ -275,13 +277,6 @@ public:
 	FTestWorldInstance Server;
 	TArray<FTestWorldInstance> Clients;
 
-private:
-
-	void InitDelegates();
-
-	void OnNetDriverCreated(UWorld* InWorld, UNetDriver* InNetDriver);
-	FDelegateHandle NetDriverCreatedHandle;
-
 private:
 
 	float TickDeltaSeconds = 0.0166f;
