Engine change to have player specific anim profiling
Make  USkeletalMeshComponentBudgeted::TickComponent & USkeletalMeshComponentBudgeted::CompleteParallelAnimationEvaluation protected instead of private
Make USkeletalMeshComponent::ParallelAnimationEvaluation virtual

diff --git a/Engine/Source/Runtime/Engine/Classes/Components/SkeletalMeshComponent.h b/Engine/Source/Runtime/Engine/Classes/Components/SkeletalMeshComponent.h
--- a/Engine/Source/Runtime/Engine/Classes/Components/SkeletalMeshComponent.h
+++ b/Engine/Source/Runtime/Engine/Classes/Components/SkeletalMeshComponent.h
@@ -2555,7 +2555,7 @@ private:
 
 public:
 	// Parallel evaluation wrappers
-	ENGINE_API void ParallelAnimationEvaluation();
+	virtual ENGINE_API void ParallelAnimationEvaluation();
 	ENGINE_API virtual void CompleteParallelAnimationEvaluation(bool bDoPostAnimEvaluation);
 
 
