Optimized linear array search in UMaterial::GetAllReferencedExpressions to avoid N^2 behavior. Observed speed improvement for switching preview platform of up to 50s (58s -> 8s).

diff --git a/Engine/Source/Runtime/Engine/Private/Materials/Material.cpp b/Engine/Source/Runtime/Engine/Private/Materials/Material.cpp
--- a/Engine/Source/Runtime/Engine/Private/Materials/Material.cpp
+++ b/Engine/Source/Runtime/Engine/Private/Materials/Material.cpp
@@ -6621,6 +6621,7 @@ bool UMaterial::GetAllReferencedExpressions(TArray<UMaterialExpression*>& OutExp
 {
 	using namespace MaterialImpl;
 
+	TSet<UMaterialExpression*> TmpExpressions;
 	OutExpressions.Empty();
 
 	// For mobile only consider nodes connected to material properties that affect mobile, and exclude any custom outputs
@@ -6633,7 +6634,7 @@ bool UMaterial::GetAllReferencedExpressions(TArray<UMaterialExpression*>& OutExp
 			{
 				for (int32 AddIdx = 0; AddIdx < MPRefdExpressions.Num(); AddIdx++)
 				{
-					OutExpressions.AddUnique(MPRefdExpressions[AddIdx]);
+					TmpExpressions.Add(MPRefdExpressions[AddIdx]);
 				}
 			}
 		}
@@ -6647,7 +6648,7 @@ bool UMaterial::GetAllReferencedExpressions(TArray<UMaterialExpression*>& OutExp
 				{
 					for (int32 AddIdx = 0; AddIdx < MPRefdExpressions.Num(); AddIdx++)
 					{
-						OutExpressions.AddUnique(MPRefdExpressions[AddIdx]);
+						TmpExpressions.Add(MPRefdExpressions[AddIdx]);
 					}
 				}
 			}
@@ -6666,8 +6667,8 @@ bool UMaterial::GetAllReferencedExpressions(TArray<UMaterialExpression*>& OutExp
 				if ((bMobileUseVirtualTexturing && Expression->IsA<UMaterialExpressionRuntimeVirtualTextureOutput>()) ||
 					(InMobileCustomOutputExpressionTypesToQuery && InMobileCustomOutputExpressionTypesToQuery->Contains(Expression->GetClass())))
 				{
-					TArray<FExpressionInput*> ProcessedInputs;
-					RecursiveGetExpressionChain(Expression, ProcessedInputs, OutExpressions, InStaticParameterSet, InFeatureLevel, InQuality, InShadingPath, SF_NumFrequencies, MP_MAX, bInRecurseIntoMaterialFunctions);
+					TSet<FExpressionInput*> ProcessedInputs;
+					RecursiveGetExpressionChain(Expression, ProcessedInputs, TmpExpressions, InStaticParameterSet, InFeatureLevel, InQuality, InShadingPath, SF_NumFrequencies, MP_MAX, bInRecurseIntoMaterialFunctions);
 				}
 			}
 		}
@@ -6682,7 +6683,7 @@ bool UMaterial::GetAllReferencedExpressions(TArray<UMaterialExpression*>& OutExp
 			{
 			    for (int32 AddIdx = 0; AddIdx < MPRefdExpressions.Num(); AddIdx++)
 			    {
-				    OutExpressions.AddUnique(MPRefdExpressions[AddIdx]);
+				    TmpExpressions.Add(MPRefdExpressions[AddIdx]);
 			    }
 		    }
 	    }
@@ -6691,8 +6692,8 @@ bool UMaterial::GetAllReferencedExpressions(TArray<UMaterialExpression*>& OutExp
 	    GetAllCustomOutputExpressions(CustomOutputExpressions);
 	    for (UMaterialExpressionCustomOutput* Expression : CustomOutputExpressions)
 	    {
-		    TArray<FExpressionInput*> ProcessedInputs;
-			RecursiveGetExpressionChain(Expression, ProcessedInputs, OutExpressions, InStaticParameterSet, InFeatureLevel, InQuality, InShadingPath, SF_NumFrequencies, MP_MAX, bInRecurseIntoMaterialFunctions);
+			TSet<FExpressionInput*> ProcessedInputs;
+			RecursiveGetExpressionChain(Expression, ProcessedInputs, TmpExpressions, InStaticParameterSet, InFeatureLevel, InQuality, InShadingPath, SF_NumFrequencies, MP_MAX, bInRecurseIntoMaterialFunctions);
 		}
 
 		// If this is a material function, we want to also trace function outputs
@@ -6700,11 +6701,14 @@ bool UMaterial::GetAllReferencedExpressions(TArray<UMaterialExpression*>& OutExp
 		GetAllFunctionOutputExpressions(FunctionOutputExpressions);
 		for (UMaterialExpressionFunctionOutput* Expression : FunctionOutputExpressions)
 		{
-			TArray<FExpressionInput*> ProcessedInputs;
-			RecursiveGetExpressionChain(Expression, ProcessedInputs, OutExpressions, InStaticParameterSet, InFeatureLevel, InQuality, InShadingPath, SF_NumFrequencies, MP_MAX, bInRecurseIntoMaterialFunctions);
+			TSet<FExpressionInput*> ProcessedInputs;
+			RecursiveGetExpressionChain(Expression, ProcessedInputs, TmpExpressions, InStaticParameterSet, InFeatureLevel, InQuality, InShadingPath, SF_NumFrequencies, MP_MAX, bInRecurseIntoMaterialFunctions);
 		}
 	}
 
+
+	OutExpressions = TmpExpressions.Array();
+
 	return true;
 }
 
@@ -6713,6 +6717,7 @@ bool UMaterial::GetExpressionsInPropertyChain(EMaterialProperty InProperty,
 	TArray<UMaterialExpression*>& OutExpressions, struct FStaticParameterSet* InStaticParameterSet,
 	ERHIFeatureLevel::Type InFeatureLevel, EMaterialQualityLevel::Type InQuality, ERHIShadingPath::Type InShadingPath, const bool bInRecurseIntoMaterialFunctions)
 {
+	TSet<UMaterialExpression*> TmpExpressions;
 	OutExpressions.Empty();
 	FExpressionInput* StartingExpression = GetExpressionInputForProperty(InProperty);
 
@@ -6722,10 +6727,10 @@ bool UMaterial::GetExpressionsInPropertyChain(EMaterialProperty InProperty,
 		return false;
 	}
 
-	TArray<FExpressionInput*> ProcessedInputs;
+	TSet<FExpressionInput*> ProcessedInputs;
 	if (StartingExpression->Expression)
 	{
-		ProcessedInputs.AddUnique(StartingExpression);
+		ProcessedInputs.Add(StartingExpression);
 		
 		EShaderFrequency ShaderFrequency = SF_NumFrequencies;
 		// These properties are "special", attempting to pass them to FMaterialAttributeDefinitionMap::GetShaderFrequency() will generate log spam
@@ -6734,8 +6739,9 @@ bool UMaterial::GetExpressionsInPropertyChain(EMaterialProperty InProperty,
 			ShaderFrequency = FMaterialAttributeDefinitionMap::GetShaderFrequency(InProperty);
 		}
 
-		RecursiveGetExpressionChain(StartingExpression->Expression, ProcessedInputs, OutExpressions, InStaticParameterSet, InFeatureLevel, InQuality, InShadingPath, ShaderFrequency, InProperty, bInRecurseIntoMaterialFunctions);
+		RecursiveGetExpressionChain(StartingExpression->Expression, ProcessedInputs, TmpExpressions, InStaticParameterSet, InFeatureLevel, InQuality, InShadingPath, ShaderFrequency, InProperty, bInRecurseIntoMaterialFunctions);
 	}
+	OutExpressions = TmpExpressions.Array();
 	return true;
 }
 
@@ -6800,15 +6806,15 @@ bool UMaterial::GetTexturesInPropertyChain(EMaterialProperty InProperty, TArray<
 }
 
 bool UMaterial::RecursiveGetExpressionChain(
-	UMaterialExpression* InExpression, TArray<FExpressionInput*>& InOutProcessedInputs, 
-	TArray<UMaterialExpression*>& OutExpressions, struct FStaticParameterSet* InStaticParameterSet,
+	UMaterialExpression* InExpression, TSet<FExpressionInput*>& InOutProcessedInputs, 
+	TSet<UMaterialExpression*>& OutExpressions, struct FStaticParameterSet* InStaticParameterSet,
 	ERHIFeatureLevel::Type InFeatureLevel, EMaterialQualityLevel::Type InQuality, ERHIShadingPath::Type InShadingPath, EShaderFrequency InShaderFrequency, EMaterialProperty InProperty, const bool bInRecurseIntoMaterialFunctions)
 {
 	using namespace MaterialImpl;
 
-	OutExpressions.AddUnique(InExpression);
-	TArray<FExpressionInput*> Inputs;
-	TArray<EShaderFrequency> InputsFrequency;
+	OutExpressions.Add(InExpression);
+	TArray<FExpressionInput*, TInlineAllocator<8>> Inputs;
+	TArray<EShaderFrequency, TInlineAllocator<8>> InputsFrequency;
 	
 	UMaterialExpressionFeatureLevelSwitch* FeatureLevelSwitchExp;
 	UMaterialExpressionQualitySwitch* QualitySwitchExp;
@@ -7032,8 +7038,7 @@ bool UMaterial::RecursiveGetExpressionChain(
 		FExpressionInput* InnerInput = Inputs[InputIdx];
 		if (InnerInput != NULL)
 		{
-			int32 DummyIdx;
-			if (InOutProcessedInputs.Find(InnerInput,DummyIdx) == false)
+			if (!InOutProcessedInputs.Contains(InnerInput))
 			{
 				if (InnerInput->Expression)
 				{
diff --git a/Engine/Source/Runtime/Engine/Public/Materials/Material.h b/Engine/Source/Runtime/Engine/Public/Materials/Material.h
--- a/Engine/Source/Runtime/Engine/Public/Materials/Material.h
+++ b/Engine/Source/Runtime/Engine/Public/Materials/Material.h
@@ -1895,8 +1895,8 @@ protected:
 	 *	Recursively retrieve the expressions contained in the chain of the given expression.
 	 *
 	 *	@param	InExpression						The expression to start at.
-	 *	@param	InOutProcessedInputs				An array of processed expression inputs. (To avoid circular loops causing infinite recursion)
-	 *	@param	OutExpressions						The array to fill in all of the expressions.
+	 *	@param	InOutProcessedInputs				A set of processed expression inputs. (To avoid circular loops causing infinite recursion)
+	 *	@param	OutExpressions						The set to fill in all of the expressions.
 	 *	@param	InStaticParameterSet				Optional static parameter set - if supplied only walk the StaticSwitch branches according to it.
 	 *	@Param	InFeatureLevel						Optional feature level - if supplied, only walk FeatureLevelSwitch branches according to it.
 	 *	@Param	InQuality							Optional quality switch - if supplied, only walk QualitySwitch branches according to it.
@@ -1906,8 +1906,8 @@ protected:
 	 *
 	 *	@return	bool					true if successful, false if not.
 	 */
-	ENGINE_API virtual bool RecursiveGetExpressionChain(UMaterialExpression* InExpression, TArray<FExpressionInput*>& InOutProcessedInputs, 
-		TArray<UMaterialExpression*>& OutExpressions, struct FStaticParameterSet* InStaticParameterSet,
+	ENGINE_API virtual bool RecursiveGetExpressionChain(UMaterialExpression* InExpression, TSet<FExpressionInput*>& InOutProcessedInputs, 
+		TSet<UMaterialExpression*>& OutExpressions, struct FStaticParameterSet* InStaticParameterSet,
 		ERHIFeatureLevel::Type InFeatureLevel = ERHIFeatureLevel::Num,
 		EMaterialQualityLevel::Type InQuality = EMaterialQualityLevel::Num,
 		ERHIShadingPath::Type InShadingPath = ERHIShadingPath::Num,
