Cloth - Added simulation config refresh for whenever the Anim Blueprint is recompiled.

#rnx

diff --git a/Engine/Source/Runtime/Engine/Private/Animation/AnimInstance.cpp b/Engine/Source/Runtime/Engine/Private/Animation/AnimInstance.cpp
--- a/Engine/Source/Runtime/Engine/Private/Animation/AnimInstance.cpp
+++ b/Engine/Source/Runtime/Engine/Private/Animation/AnimInstance.cpp
@@ -32,6 +32,8 @@
 #if WITH_EDITOR
 #include "Engine/Blueprint.h"
 #include "BlueprintEditorSettings.h"
+#include "ClothingSimulationInstance.h"
+#include "ClothingSimulationInteractor.h"
 #endif
 
 #include UE_INLINE_GENERATED_CPP_BY_NAME(AnimInstance)
@@ -4578,6 +4580,15 @@ void UAnimInstance::HandleObjectsReinstanced(const TMap<UObject*, UObject*>& Old
 				}
 
 				MeshComponent->ClearMotionVector();
+
+				// Reset interactors which might have been used in the blueprint
+				for (FClothingSimulationInstance& ClothingSimulationInstance : MeshComponent->GetClothingSimulationInstances())
+				{
+					if (UClothingSimulationInteractor* const ClothingSimulationInteractor = ClothingSimulationInstance.GetClothingSimulationInteractor())
+					{
+						ClothingSimulationInteractor->ClothConfigUpdated();
+					}
+				}
 			}
 		}
 
