PR #11614: Fix finding maximum in BodySetup.cpp

diff --git a/Engine/Source/Runtime/Engine/Private/PhysicsEngine/BodySetup.cpp b/Engine/Source/Runtime/Engine/Private/PhysicsEngine/BodySetup.cpp
--- a/Engine/Source/Runtime/Engine/Private/PhysicsEngine/BodySetup.cpp
+++ b/Engine/Source/Runtime/Engine/Private/PhysicsEngine/BodySetup.cpp
@@ -1929,10 +1929,16 @@ void FKSphereElem::ScaleElem(FVector DeltaSize, float MinSize)
 {
 	// Find element with largest magnitude, btu preserve sign.
 	float DeltaRadius = DeltaSize.X;
+
 	if (FMath::Abs(DeltaSize.Y) > FMath::Abs(DeltaRadius))
+	{
 		DeltaRadius = DeltaSize.Y;
-	else if (FMath::Abs(DeltaSize.Z) > FMath::Abs(DeltaRadius))
+	}
+
+	if (FMath::Abs(DeltaSize.Z) > FMath::Abs(DeltaRadius))
+	{
 		DeltaRadius = DeltaSize.Z;
+	}
 
 	Radius = FMath::Max(Radius + DeltaRadius, MinSize);
 }
