[MaterialExpression]
- Replace FApp::IsGame() which can dynamically change when starting or exiting PIE to a more viable and deterministic per-object approach

diff --git a/Engine/Source/Runtime/Engine/Private/Materials/MaterialExpressions.cpp b/Engine/Source/Runtime/Engine/Private/Materials/MaterialExpressions.cpp
--- a/Engine/Source/Runtime/Engine/Private/Materials/MaterialExpressions.cpp
+++ b/Engine/Source/Runtime/Engine/Private/Materials/MaterialExpressions.cpp
@@ -18112,7 +18112,7 @@ bool UMaterialExpressionNamedRerouteDeclaration::GetRerouteInput(FExpressionInpu
 void UMaterialExpressionNamedRerouteDeclaration::UpdateVariableGuid(bool bForceGeneration, bool bAllowMarkingPackageDirty)
 {
 	// If we are in the editor, and we don't have a valid GUID yet, generate one.
-	if (GIsEditor && !FApp::IsGame() && !IsTemplate())
+	if (GIsEditor && !IsRunningGame() && !GetPackage()->HasAnyPackageFlags(PKG_PlayInEditor) && !IsTemplate())
 	{
 		if (bForceGeneration || !VariableGuid.IsValid())
 		{
