Added nullptr check to AddAssetUserDataOfClass as Shipping build will compile out check()

diff --git a/Engine/Source/Runtime/Engine/Private/Interfaces/Interface_AssetUserData.cpp b/Engine/Source/Runtime/Engine/Private/Interfaces/Interface_AssetUserData.cpp
--- a/Engine/Source/Runtime/Engine/Private/Interfaces/Interface_AssetUserData.cpp
+++ b/Engine/Source/Runtime/Engine/Private/Interfaces/Interface_AssetUserData.cpp
@@ -23,9 +23,12 @@ bool IInterface_AssetUserData::HasAssetUserDataOfClass(TSubclassOf<UAssetUserDat
 
 bool IInterface_AssetUserData::AddAssetUserDataOfClass(TSubclassOf<UAssetUserData> InUserDataClass)
 {
-	check(InUserDataClass.Get() != nullptr);
-	AddAssetUserData(NewObject<UAssetUserData>(Cast<UObject>(this), InUserDataClass.Get()));
-	return GetAssetUserDataOfClass(InUserDataClass.Get()) != nullptr;
+	const UClass* DataClass = InUserDataClass.Get();
+	if (DataClass != nullptr)
+	{
+		AddAssetUserData(NewObject<UAssetUserData>(Cast<UObject>(this), DataClass));
+	}
+	return GetAssetUserDataOfClass(InUserDataClass) != nullptr;
 }
 
 const TArray<UAssetUserData*>* IInterface_AssetUserData::GetAssetUserDataArray() const
