[UE][TaskSyncManager] Initialize batch for newly created world so one can access tick function from BeginPlay callback and establish dependencies

diff --git a/Engine/Source/Runtime/Engine/Private/TaskSyncManager.cpp b/Engine/Source/Runtime/Engine/Private/TaskSyncManager.cpp
--- a/Engine/Source/Runtime/Engine/Private/TaskSyncManager.cpp
+++ b/Engine/Source/Runtime/Engine/Private/TaskSyncManager.cpp
@@ -968,6 +968,7 @@ FTaskSyncManager::FTaskSyncManager()
 	{
 		GEngine->OnWorldContextDestroyed().AddRaw(this, &FTaskSyncManager::OnWorldContextRemove);
 	}
+	FWorldDelegates::OnPostWorldInitialization.AddRaw(this, &FTaskSyncManager::OnPostWorldInitialization);
 }
 
 FTaskSyncManager::~FTaskSyncManager()
@@ -976,6 +977,7 @@ FTaskSyncManager::~FTaskSyncManager()
 	{
 		GEngine->OnWorldContextDestroyed().RemoveAll(this);
 	}
+	FWorldDelegates::OnPostWorldInitialization.RemoveAll(this);
 }
 
 void FTaskSyncManager::RegisterSyncPointInternal(const FSyncPointDescription& InDescription)
@@ -1394,6 +1396,35 @@ void FTaskSyncManager::EndFrame(const UWorld* InWorld)
 	}
 }
 
+void FTaskSyncManager::InitializeForWorld(FWorldContextId WorldContext)
+{
+	check(IsInGameThread() && WorldContext.IsValid());
+
+	FBatchContextId BatchContext = CreateNewBatch(WorldContext);
+	check(BatchList.IsValidIndex(BatchContext.BatchId));
+	InitializeBatchForNewWorld(BatchList[BatchContext.BatchId]);
+}
+
+void FTaskSyncManager::InitializeBatchForNewWorld(FBatchData& BatchData)
+{
+	// Make sure all registered sync points are created.
+	for (TPair<FSyncPointId::FInternalId, FRegisteredSyncPointData>& DataPair : RegisteredDataMap)
+	{
+		GetOrCreateSyncPoint(BatchData, DataPair.Value);
+	}
+}
+
+void FTaskSyncManager::OnPostWorldInitialization(UWorld* World, const UWorld::InitializationValues IVS)
+{
+	check(World != nullptr);
+
+	FWorldContextId WorldContext(World);
+	if (WorldContext.IsValid())
+	{
+		InitializeForWorld(World);
+	}
+}
+
 void FTaskSyncManager::ReleaseWorldContext(FWorldContextId WorldContext)
 {
 	// We're completely destroying a world context so free the data
diff --git a/Engine/Source/Runtime/Engine/Public/TaskSyncManager.h b/Engine/Source/Runtime/Engine/Public/TaskSyncManager.h
--- a/Engine/Source/Runtime/Engine/Public/TaskSyncManager.h
+++ b/Engine/Source/Runtime/Engine/Public/TaskSyncManager.h
@@ -5,6 +5,7 @@
 #include "Engine/EngineBaseTypes.h"
 #include "Misc/TransactionallySafeCriticalSection.h"
 #include "Engine/DeveloperSettings.h"
+#include "Engine/World.h"
 #include "Async/TaskGraphFwd.h"
 #include "Delegates/Delegate.h"
 
@@ -965,6 +966,10 @@ private:
 	void RegisterSyncPointInternal(const FSyncPointDescription& InDescription);
 	void OnWorldContextRemove(FWorldContext& InWorldContext);
 	void InitializeBatchForFrame(FBatchData& BatchData, ULevel* PersistentLevel);
+	void InitializeForWorld(FWorldContextId WorldContext);
+	void OnPostWorldInitialization(UWorld* World, const UWorld::InitializationValues IVS);
+	void InitializeBatchForNewWorld(FBatchData& BatchData);
+
 	FActiveSyncPoint* GetOrCreateSyncPoint(FBatchData& BatchData, FRegisteredSyncPointData& SyncData);
 	FActiveSyncPoint* GetOrCreateTickGroupWork(FBatchData& BatchData, ETickingGroup TickGroup);
 
