Fix off by one in check statement noticed by audit tool

diff --git a/Engine/Source/Runtime/Engine/Public/TaskSyncManager.h b/Engine/Source/Runtime/Engine/Public/TaskSyncManager.h
--- a/Engine/Source/Runtime/Engine/Public/TaskSyncManager.h
+++ b/Engine/Source/Runtime/Engine/Public/TaskSyncManager.h
@@ -472,7 +472,7 @@ private:
 	uint32 WorkIndex : 24 = 0;
 
 	/** Max amount of work per sync point */
-	static constexpr uint32 MaxWorkIndex = (2 << 24) - 1;
+	static constexpr uint32 MaxWorkIndex = (1 << 24) - 1;
 
 	/** True if this may represent active reserved work */
 	uint32 bWorkReserved : 1 = 0;
