[Engine]
* Removed includes from Engine to UnrealEd that was part of shared pch
- Removed ChangeTransactor.h include from IAnimationDataController.h
- Removed Kismet2/Kismet2NameValidators.h from EdGraphSchema.h
- Kismet2/CompilerResultsLog.h from Blueprint.h

diff --git a/Engine/Plugins/Animation/AnimationData/Source/AnimationData/Private/AnimSequencerController.h b/Engine/Plugins/Animation/AnimationData/Source/AnimationData/Private/AnimSequencerController.h
--- a/Engine/Plugins/Animation/AnimationData/Source/AnimationData/Private/AnimSequencerController.h
+++ b/Engine/Plugins/Animation/AnimationData/Source/AnimationData/Private/AnimSequencerController.h
@@ -6,6 +6,7 @@
 #include "Animation/AnimData/AnimDataModel.h"
 #include "Animation/AnimCurveTypes.h"
 #include "Animation/AnimData/IAnimationDataController.h"
+#include "ChangeTransactor.h"
 
 #include "AnimSequencerController.generated.h"
 
@@ -133,4 +134,7 @@ private:
 	friend class FAnimDataControllerTestBase;
 	friend UE::Anim::FOpenBracketAction;
 	friend UE::Anim::FCloseBracketAction;
+
+	DECLARE_TRANSACTION_CONTROLLER
+
 };
diff --git a/Engine/Source/Developer/AnimationDataController/Private/AnimDataController.cpp b/Engine/Source/Developer/AnimationDataController/Private/AnimDataController.cpp
--- a/Engine/Source/Developer/AnimationDataController/Private/AnimDataController.cpp
+++ b/Engine/Source/Developer/AnimationDataController/Private/AnimDataController.cpp
@@ -3,6 +3,7 @@
 #include "AnimDataController.h"
 #include "Animation/Skeleton.h"
 #include "AnimDataControllerActions.h"
+#include "ChangeTransactor.h"
 #include "MovieSceneTimeHelpers.h"
 
 #include "Animation/AnimData/IAnimationDataModel.h"
diff --git a/Engine/Source/Developer/AnimationDataController/Public/AnimDataController.h b/Engine/Source/Developer/AnimationDataController/Public/AnimDataController.h
--- a/Engine/Source/Developer/AnimationDataController/Public/AnimDataController.h
+++ b/Engine/Source/Developer/AnimationDataController/Public/AnimDataController.h
@@ -118,4 +118,6 @@ private:
 	friend class FAnimDataControllerTestBase;
 	friend UE::Anim::FOpenBracketAction;
 	friend UE::Anim::FCloseBracketAction;
+
+	DECLARE_TRANSACTION_CONTROLLER
 };
diff --git a/Engine/Source/Editor/BlueprintGraph/Classes/EdGraphSchema_K2.h b/Engine/Source/Editor/BlueprintGraph/Classes/EdGraphSchema_K2.h
--- a/Engine/Source/Editor/BlueprintGraph/Classes/EdGraphSchema_K2.h
+++ b/Engine/Source/Editor/BlueprintGraph/Classes/EdGraphSchema_K2.h
@@ -15,6 +15,7 @@
 #include "EdGraph/EdGraphSchema.h"
 #include "HAL/PlatformMath.h"
 #include "Internationalization/Text.h"
+#include "Kismet2/BlueprintImpl.h"
 #include "Math/Color.h"
 #include "Math/Vector2D.h"
 #include "Misc/EnumClassFlags.h"
diff --git a/Engine/Source/Editor/UnrealEd/Public/Kismet2/BlueprintImpl.h b/Engine/Source/Editor/UnrealEd/Public/Kismet2/BlueprintImpl.h
new file mode 100644
--- /dev/null
+++ b/Engine/Source/Editor/UnrealEd/Public/Kismet2/BlueprintImpl.h
@@ -0,0 +1,45 @@
+// Copyright Epic Games, Inc. All Rights Reserved.
+
+#pragma once
+
+#include "Engine/Blueprint.h"
+#include "Kismet2/CompilerResultsLog.h"
+
+template<typename... ArgTypes>
+void UBlueprint::Message_Note(const FString& MessageToLog, ArgTypes... Args)
+{
+	if (CurrentMessageLog)
+	{
+		CurrentMessageLog->Note(*MessageToLog, Forward<ArgTypes>(Args)...);
+	}
+	else
+	{
+		UE_LOG(LogBlueprint, Log, TEXT("[%s] %s"), *GetName(), *MessageToLog);
+	}
+}
+
+template<typename... ArgTypes>
+void UBlueprint::Message_Warn(const FString& MessageToLog, ArgTypes... Args)
+{
+	if (CurrentMessageLog)
+	{
+		CurrentMessageLog->Warning(*MessageToLog, Forward<ArgTypes>(Args)...);
+	}
+	else
+	{
+		UE_LOG(LogBlueprint, Warning, TEXT("[%s] %s"), *GetName(), *MessageToLog);
+	}
+}
+
+template<typename... ArgTypes>
+void UBlueprint::Message_Error(const FString& MessageToLog, ArgTypes... Args)
+{
+	if (CurrentMessageLog)
+	{
+		CurrentMessageLog->Error(*MessageToLog, Forward<ArgTypes>(Args)...);
+	}
+	else
+	{
+		UE_LOG(LogBlueprint, Error, TEXT("[%s] %s"), *GetName(), *MessageToLog);
+	}
+}
diff --git a/Engine/Source/Runtime/Engine/Classes/Animation/AnimData/IAnimationDataController.h b/Engine/Source/Runtime/Engine/Classes/Animation/AnimData/IAnimationDataController.h
--- a/Engine/Source/Runtime/Engine/Classes/Animation/AnimData/IAnimationDataController.h
+++ b/Engine/Source/Runtime/Engine/Classes/Animation/AnimData/IAnimationDataController.h
@@ -12,10 +12,6 @@
 #include "Algo/Transform.h"
 #include "Animation/AnimCurveTypes.h"
 
-#if WITH_EDITOR
-#include "ChangeTransactor.h"
-#endif // WITH_EDITOR
-
 #include "IAnimationDataController.generated.h"
 
 #define LOCTEXT_NAMESPACE "IAnimationDataController"
@@ -809,36 +805,6 @@ protected:
 	
 	static constexpr int32 DefaultCurveFlags = EAnimAssetCurveFlags::AACF_Editable;
 
-#if WITH_EDITOR
-	typedef TUniquePtr<UE::FScopedCompoundTransaction> FTransaction;
-	FTransaction ConditionalTransaction(const FText& Description, bool bCondition)
-	{
-		if (UE::FChangeTransactor::CanTransactChanges() && bCondition)
-		{ 
-			return MakeUnique<UE::FScopedCompoundTransaction>(ChangeTransactor, Description);
-		}
-
-		return nullptr;
-	}
-
-	typedef TUniquePtr<IAnimationDataController::FScopedBracket> FBracket; 
-	FBracket ConditionalBracket(const FText& Description, bool bCondition)
-	{	
-		return MakeUnique<IAnimationDataController::FScopedBracket>(this, Description, UE::FChangeTransactor::CanTransactChanges() && bCondition);	
-	}
-
-	template<typename ActionClass, class... ActionArguments>
-	void ConditionalAction(bool bCondition, ActionArguments&&... Arguments)
-	{
-		if (UE::FChangeTransactor::CanTransactChanges() && bCondition) \
-		{ 
-			ChangeTransactor.AddTransactionChange<ActionClass>(Forward<ActionArguments>(Arguments)...); \
-		}
-	}
-protected:
-	UE::FChangeTransactor ChangeTransactor;
-#endif // WITH_EDITOR
-
 public:
 	template <typename FmtType, typename... Types>
 	static void ReportObjectWarningf(const UObject* ErrorObject, const FmtType& Fmt, Types... Args)
@@ -865,8 +831,38 @@ private:
 	friend UE::Anim::FCloseBracketAction;
 };
 
+// Add this to subclasses for change transactor. Note, this code uses UnrealEd types and this file is in Engine
+// Needs to include "ChangeTransactor.h"
+#define DECLARE_TRANSACTION_CONTROLLER \
+		typedef TUniquePtr<UE::FScopedCompoundTransaction> FTransaction; \
+		FTransaction ConditionalTransaction(const FText& Description, bool bCondition) \
+		{ \
+			if (UE::FChangeTransactor::CanTransactChanges() && bCondition) \
+			{  \
+				return MakeUnique<UE::FScopedCompoundTransaction>(ChangeTransactor, Description); \
+			} \
+			return nullptr; \
+		} \
+		typedef TUniquePtr<IAnimationDataController::FScopedBracket> FBracket;  \
+		FBracket ConditionalBracket(const FText& Description, bool bCondition) \
+		{	 \
+			return MakeUnique<IAnimationDataController::FScopedBracket>(this, Description, UE::FChangeTransactor::CanTransactChanges() && bCondition);	 \
+		} \
+		template<typename ActionClass, class... ActionArguments> \
+		void ConditionalAction(bool bCondition, ActionArguments&&... Arguments) \
+		{ \
+			if (UE::FChangeTransactor::CanTransactChanges() && bCondition) \
+			{  \
+				ChangeTransactor.AddTransactionChange<ActionClass>(Forward<ActionArguments>(Arguments)...); \
+			} \
+		} \
+	protected: \
+		UE::FChangeTransactor ChangeTransactor; \
+
+
 #define CONDITIONAL_TRANSACTION(Text) FTransaction Transaction = ConditionalTransaction(Test, bShouldTransact);
 #define CONDITIONAL_BRACKET(Text) FBracket Bracket = ConditionalBracket(Text, bShouldTransact);
 #define CONDITIONAL_ACTION(ActionClass, ...) ConditionalAction<ActionClass>(bShouldTransact, __VA_ARGS__);
 
 #undef LOCTEXT_NAMESPACE // "IAnimationDataController"
+
diff --git a/Engine/Source/Runtime/Engine/Classes/EdGraph/EdGraphSchema.h b/Engine/Source/Runtime/Engine/Classes/EdGraph/EdGraphSchema.h
--- a/Engine/Source/Runtime/Engine/Classes/EdGraph/EdGraphSchema.h
+++ b/Engine/Source/Runtime/Engine/Classes/EdGraph/EdGraphSchema.h
@@ -10,9 +10,6 @@
 #include "AssetRegistry/AssetData.h"
 #include "UObject/ObjectKey.h"
 #include "Input/Reply.h"
-#if WITH_EDITOR
-#include "Kismet2/Kismet2NameValidators.h"
-#endif
 #include "EdGraphSchema.generated.h"
 
 class FSlateRect;
diff --git a/Engine/Source/Runtime/Engine/Classes/Engine/Blueprint.h b/Engine/Source/Runtime/Engine/Classes/Engine/Blueprint.h
--- a/Engine/Source/Runtime/Engine/Classes/Engine/Blueprint.h
+++ b/Engine/Source/Runtime/Engine/Classes/Engine/Blueprint.h
@@ -17,12 +17,12 @@
 
 #if WITH_EDITOR
 #include "EngineLogs.h"
-#include "Kismet2/CompilerResultsLog.h"
 #endif
 
 #include "Blueprint.generated.h"
 
 enum class EBlueprintCompileOptions;
+class FCompilerResultsLog;
 class ITargetPlatform;
 class UActorComponent;
 class UEdGraph;
@@ -1146,45 +1146,16 @@ public:
 	/** 
 	 * Sends a message to the CurrentMessageLog, if there is one available.  Otherwise, defaults to logging to the normal channels.
 	 * Should use this for node and blueprint actions that happen during compilation!
+	 * Note, needs to include BlueprintImpl.h in UnrealEd since code is using UnrealEd and this class is in Engine
 	 */
 	template<typename... ArgTypes>
-	void Message_Note(const FString& MessageToLog, ArgTypes... Args)
-	{
-		if (CurrentMessageLog)
-		{
-			CurrentMessageLog->Note(*MessageToLog, Forward<ArgTypes>(Args)...);
-		}
-		else
-		{
-			UE_LOG(LogBlueprint, Log, TEXT("[%s] %s"), *GetName(), *MessageToLog);
-		}
-	}
+	void Message_Note(const FString& MessageToLog, ArgTypes... Args);
 
 	template<typename... ArgTypes>
-	void Message_Warn(const FString& MessageToLog, ArgTypes... Args)
-	{
-		if (CurrentMessageLog)
-		{
-			CurrentMessageLog->Warning(*MessageToLog, Forward<ArgTypes>(Args)...);
-		}
-		else
-		{
-			UE_LOG(LogBlueprint, Warning, TEXT("[%s] %s"), *GetName(), *MessageToLog);
-		}
-	}
+	void Message_Warn(const FString& MessageToLog, ArgTypes... Args);
 
 	template<typename... ArgTypes>
-	void Message_Error(const FString& MessageToLog, ArgTypes... Args)
-	{
-		if (CurrentMessageLog)
-		{
-			CurrentMessageLog->Error(*MessageToLog, Forward<ArgTypes>(Args)...);
-		}
-		else
-		{
-			UE_LOG(LogBlueprint, Error, TEXT("[%s] %s"), *GetName(), *MessageToLog);
-		}
-	}
+	void Message_Error(const FString& MessageToLog, ArgTypes... Args);
 #endif
 	
 #if WITH_EDITORONLY_DATA
