[VisualLogger] updated to use UE_API

diff --git a/Engine/Source/Runtime/Engine/Public/VisualLogger/VisualLogger.h b/Engine/Source/Runtime/Engine/Public/VisualLogger/VisualLogger.h
--- a/Engine/Source/Runtime/Engine/Public/VisualLogger/VisualLogger.h
+++ b/Engine/Source/Runtime/Engine/Public/VisualLogger/VisualLogger.h
@@ -18,6 +18,8 @@
 #include "UObject/ObjectKey.h"
 #include "VisualLogger/VisualLoggerTypes.h"
 
+#define UE_API ENGINE_API
+
 #if ENABLE_VISUAL_LOG
 
 #define REDIRECT_TO_VLOG(Dest) FVisualLogger::Redirect(this, Dest)
@@ -206,7 +208,7 @@
 #define TEXT_FALSE TEXT("FALSE")
 #define TEXT_CONDITION(Condition) ((Condition) ? TEXT_TRUE : TEXT_FALSE)
 
-ENGINE_API DECLARE_LOG_CATEGORY_EXTERN(LogVisual, Display, All);
+UE_API DECLARE_LOG_CATEGORY_EXTERN(LogVisual, Display, All);
 
 #if ENABLE_VISUAL_LOG
 
@@ -220,21 +222,21 @@ DECLARE_DELEGATE_RetVal(FString, FVisualLogFilenameGetterDelegate);
 
 class FVisualLogger : public FOutputDevice
 {
-	static ENGINE_API void CategorizedLogfImpl(const UObject* LogOwner, const FName& CategoryName, ELogVerbosity::Type Verbosity, const TCHAR* Fmt, ...);
-	static ENGINE_API void SegmentLogfImpl(const UObject* LogOwner, const FName& CategoryName, ELogVerbosity::Type Verbosity, const FVector& Start, const FVector& End, const FColor& Color, const uint16 Thickness, const TCHAR* Fmt, ...);
-	static ENGINE_API void LocationLogfImpl(const UObject* LogOwner, const FName& CategoryName, ELogVerbosity::Type Verbosity, const FVector& Center, uint16 Thickness, const FColor& Color, const TCHAR* Fmt, ...);
-	static ENGINE_API void SphereLogfImpl(const UObject* LogOwner, const FName& CategoryName, ELogVerbosity::Type Verbosity, const FVector& Center, float Radius, const FColor& Color, bool bWireframe, const TCHAR* Fmt, ...);
-	static ENGINE_API void BoxLogfImpl(const UObject* LogOwner, const FName& CategoryName, ELogVerbosity::Type Verbosity, const FBox& Box, const FMatrix& Matrix, const FColor& Color, bool bWireframe, const TCHAR* Fmt, ...);
-	static ENGINE_API void ConeLogfImpl(const UObject* LogOwner, const FName& CategoryName, ELogVerbosity::Type Verbosity, const FVector& Origin, const FVector& Direction, const float Length, const float Angle, const FColor& Color, bool bWireframe, const TCHAR* Fmt, ...);
-	static ENGINE_API void CylinderLogfImpl(const UObject* LogOwner, const FName& CategoryName, ELogVerbosity::Type Verbosity, const FVector& Start, const FVector& End, const float Radius, const FColor& Color, bool bWireframe, const TCHAR* Fmt, ...);
-	static ENGINE_API void CapsuleLogfImpl(const UObject* LogOwner, const FName& CategoryName, ELogVerbosity::Type Verbosity, const FVector& Base, float HalfHeight, float Radius, const FQuat& Rotation, const FColor& Color, bool bWireframe, const TCHAR* Fmt, ...);
-	static ENGINE_API void PulledConvexLogfImpl(const UObject* LogOwner, const FName& CategoryName, ELogVerbosity::Type Verbosity, const TArray<FVector>& ConvexPoints, float MinZ, float MaxZ, const FColor& Color, const TCHAR* Fmt, ...);
-	static ENGINE_API void MeshLogfImpl(const UObject* LogOwner, const FName& CategoryName, ELogVerbosity::Type Verbosity, const TArray<FVector>& Vertices, const TArray<int32>& Indices, const FColor& Color, const TCHAR* Fmt, ...);
-	static ENGINE_API void ConvexLogfImpl(const UObject* LogOwner, const FName& CategoryName, ELogVerbosity::Type Verbosity, const TArray<FVector>& Points, const FColor& Color, const TCHAR* Fmt, ...);
-	static ENGINE_API void HistogramDataLogfImpl(const UObject* LogOwner, const FName& CategoryName, ELogVerbosity::Type Verbosity, FName GraphName, FName DataName, const FVector2D& Data, const FColor& Color, const TCHAR* Fmt, ...);
-	static ENGINE_API void ArrowLineLogfImpl(const UObject* LogOwner, const FName& CategoryName, ELogVerbosity::Type Verbosity, const FVector& Start, const FVector& End, const FColor& Color, const uint16 Mag, const TCHAR* Fmt, ...);
-	static ENGINE_API void DiscLogfImpl(const UObject* LogOwner, const FName& CategoryName, ELogVerbosity::Type Verbosity, const FVector& Center, const FVector& UpAxis, const float Radius, const FColor& Color, const uint16 Thickness, bool bWireframe, const TCHAR* Fmt, ...);
-	static ENGINE_API void CoordinateSystemLogfImpl(const UObject* LogOwner, const FName& CategoryName, ELogVerbosity::Type Verbosity, const FVector& AxisLoc, const FRotator& AxisRot, const float Scale, const FColor& Color, const uint16 Thickness, const TCHAR* Fmt, ...);
+	static UE_API void CategorizedLogfImpl(const UObject* LogOwner, const FName& CategoryName, ELogVerbosity::Type Verbosity, const TCHAR* Fmt, ...);
+	static UE_API void SegmentLogfImpl(const UObject* LogOwner, const FName& CategoryName, ELogVerbosity::Type Verbosity, const FVector& Start, const FVector& End, const FColor& Color, const uint16 Thickness, const TCHAR* Fmt, ...);
+	static UE_API void LocationLogfImpl(const UObject* LogOwner, const FName& CategoryName, ELogVerbosity::Type Verbosity, const FVector& Center, uint16 Thickness, const FColor& Color, const TCHAR* Fmt, ...);
+	static UE_API void SphereLogfImpl(const UObject* LogOwner, const FName& CategoryName, ELogVerbosity::Type Verbosity, const FVector& Center, float Radius, const FColor& Color, bool bWireframe, const TCHAR* Fmt, ...);
+	static UE_API void BoxLogfImpl(const UObject* LogOwner, const FName& CategoryName, ELogVerbosity::Type Verbosity, const FBox& Box, const FMatrix& Matrix, const FColor& Color, bool bWireframe, const TCHAR* Fmt, ...);
+	static UE_API void ConeLogfImpl(const UObject* LogOwner, const FName& CategoryName, ELogVerbosity::Type Verbosity, const FVector& Origin, const FVector& Direction, const float Length, const float Angle, const FColor& Color, bool bWireframe, const TCHAR* Fmt, ...);
+	static UE_API void CylinderLogfImpl(const UObject* LogOwner, const FName& CategoryName, ELogVerbosity::Type Verbosity, const FVector& Start, const FVector& End, const float Radius, const FColor& Color, bool bWireframe, const TCHAR* Fmt, ...);
+	static UE_API void CapsuleLogfImpl(const UObject* LogOwner, const FName& CategoryName, ELogVerbosity::Type Verbosity, const FVector& Base, float HalfHeight, float Radius, const FQuat& Rotation, const FColor& Color, bool bWireframe, const TCHAR* Fmt, ...);
+	static UE_API void PulledConvexLogfImpl(const UObject* LogOwner, const FName& CategoryName, ELogVerbosity::Type Verbosity, const TArray<FVector>& ConvexPoints, float MinZ, float MaxZ, const FColor& Color, const TCHAR* Fmt, ...);
+	static UE_API void MeshLogfImpl(const UObject* LogOwner, const FName& CategoryName, ELogVerbosity::Type Verbosity, const TArray<FVector>& Vertices, const TArray<int32>& Indices, const FColor& Color, const TCHAR* Fmt, ...);
+	static UE_API void ConvexLogfImpl(const UObject* LogOwner, const FName& CategoryName, ELogVerbosity::Type Verbosity, const TArray<FVector>& Points, const FColor& Color, const TCHAR* Fmt, ...);
+	static UE_API void HistogramDataLogfImpl(const UObject* LogOwner, const FName& CategoryName, ELogVerbosity::Type Verbosity, FName GraphName, FName DataName, const FVector2D& Data, const FColor& Color, const TCHAR* Fmt, ...);
+	static UE_API void ArrowLineLogfImpl(const UObject* LogOwner, const FName& CategoryName, ELogVerbosity::Type Verbosity, const FVector& Start, const FVector& End, const FColor& Color, const uint16 Mag, const TCHAR* Fmt, ...);
+	static UE_API void DiscLogfImpl(const UObject* LogOwner, const FName& CategoryName, ELogVerbosity::Type Verbosity, const FVector& Center, const FVector& UpAxis, const float Radius, const FColor& Color, const uint16 Thickness, bool bWireframe, const TCHAR* Fmt, ...);
+	static UE_API void CoordinateSystemLogfImpl(const UObject* LogOwner, const FName& CategoryName, ELogVerbosity::Type Verbosity, const FVector& AxisLoc, const FRotator& AxisRot, const float Scale, const FColor& Color, const uint16 Thickness, const TCHAR* Fmt, ...);
 
 	static bool IsFilteredOut(const FLogCategoryBase& Category, ELogVerbosity::Type Verbosity)
 	{
@@ -244,7 +246,7 @@ class FVisualLogger : public FOutputDevice
 	}
 
 public:
-	ENGINE_API static bool bUseVerbosityFilterWhenRecording;
+	UE_API static bool bUseVerbosityFilterWhenRecording;
 
 	// Regular text log
 	template <typename FmtType, typename... Types>
@@ -562,34 +564,34 @@ public:
 	}
 
 	// Navigation data debug snapshot
-	static ENGINE_API void NavigationDataDump(const UObject* LogOwner, const FLogCategoryBase& Category, const ELogVerbosity::Type Verbosity, const FBox& Box);
-	static ENGINE_API void NavigationDataDump(const UObject* LogOwner, const FName& CategoryName, const ELogVerbosity::Type Verbosity, const FBox& Box);
+	static UE_API void NavigationDataDump(const UObject* LogOwner, const FLogCategoryBase& Category, const ELogVerbosity::Type Verbosity, const FBox& Box);
+	static UE_API void NavigationDataDump(const UObject* LogOwner, const FName& CategoryName, const ELogVerbosity::Type Verbosity, const FBox& Box);
 
 	DECLARE_MULTICAST_DELEGATE_SixParams(FNavigationDataDump, const UObject* /*Object*/, const FName& /*CategoryName*/, const ELogVerbosity::Type /*Verbosity*/, const FBox& /*Box*/, const UWorld& /*World*/, FVisualLogEntry& /*CurrentEntry*/);
-	static ENGINE_API FNavigationDataDump NavigationDataDumpDelegate;
+	static UE_API FNavigationDataDump NavigationDataDumpDelegate;
 
 	/** Log events */
-	static ENGINE_API void EventLog(const UObject* LogOwner, const FName EventTag1, const FVisualLogEventBase& Event1, const FName EventTag2 = NAME_None, const FName EventTag3 = NAME_None, const FName EventTag4 = NAME_None, const FName EventTag5 = NAME_None, const FName EventTag6 = NAME_None);
-	static ENGINE_API void EventLog(const UObject* LogOwner, const FName EventTag1, const FVisualLogEventBase& Event1, const FVisualLogEventBase& Event2);
-	static ENGINE_API void EventLog(const UObject* LogOwner, const FName EventTag1, const FVisualLogEventBase& Event1, const FVisualLogEventBase& Event2, const FVisualLogEventBase& Event3);
-	static ENGINE_API void EventLog(const UObject* LogOwner, const FName EventTag1, const FVisualLogEventBase& Event1, const FVisualLogEventBase& Event2, const FVisualLogEventBase& Event3, const FVisualLogEventBase& Event4);
-	static ENGINE_API void EventLog(const UObject* LogOwner, const FName EventTag1, const FVisualLogEventBase& Event1, const FVisualLogEventBase& Event2, const FVisualLogEventBase& Event3, const FVisualLogEventBase& Event4, const FVisualLogEventBase& Event5);
-	static ENGINE_API void EventLog(const UObject* LogOwner, const FName EventTag1, const FVisualLogEventBase& Event1, const FVisualLogEventBase& Event2, const FVisualLogEventBase& Event3, const FVisualLogEventBase& Event4, const FVisualLogEventBase& Event5, const FVisualLogEventBase& Event6);
-	static ENGINE_API void EventLog(const UObject* LogOwner, const FVisualLogEventBase& Event1, const FName EventTag1 = NAME_None, const FName EventTag2 = NAME_None, const FName EventTag3 = NAME_None, const FName EventTag4 = NAME_None, const FName EventTag5 = NAME_None, const FName EventTag6 = NAME_None);
+	static UE_API void EventLog(const UObject* LogOwner, const FName EventTag1, const FVisualLogEventBase& Event1, const FName EventTag2 = NAME_None, const FName EventTag3 = NAME_None, const FName EventTag4 = NAME_None, const FName EventTag5 = NAME_None, const FName EventTag6 = NAME_None);
+	static UE_API void EventLog(const UObject* LogOwner, const FName EventTag1, const FVisualLogEventBase& Event1, const FVisualLogEventBase& Event2);
+	static UE_API void EventLog(const UObject* LogOwner, const FName EventTag1, const FVisualLogEventBase& Event1, const FVisualLogEventBase& Event2, const FVisualLogEventBase& Event3);
+	static UE_API void EventLog(const UObject* LogOwner, const FName EventTag1, const FVisualLogEventBase& Event1, const FVisualLogEventBase& Event2, const FVisualLogEventBase& Event3, const FVisualLogEventBase& Event4);
+	static UE_API void EventLog(const UObject* LogOwner, const FName EventTag1, const FVisualLogEventBase& Event1, const FVisualLogEventBase& Event2, const FVisualLogEventBase& Event3, const FVisualLogEventBase& Event4, const FVisualLogEventBase& Event5);
+	static UE_API void EventLog(const UObject* LogOwner, const FName EventTag1, const FVisualLogEventBase& Event1, const FVisualLogEventBase& Event2, const FVisualLogEventBase& Event3, const FVisualLogEventBase& Event4, const FVisualLogEventBase& Event5, const FVisualLogEventBase& Event6);
+	static UE_API void EventLog(const UObject* LogOwner, const FVisualLogEventBase& Event1, const FName EventTag1 = NAME_None, const FName EventTag2 = NAME_None, const FName EventTag3 = NAME_None, const FName EventTag4 = NAME_None, const FName EventTag5 = NAME_None, const FName EventTag6 = NAME_None);
 
 	// static getter
-	static ENGINE_API FVisualLogger& Get();
+	static UE_API FVisualLogger& Get();
 
 	virtual ~FVisualLogger() {}
 
 	// called on engine shutdown to flush all, etc.
-	ENGINE_API virtual void TearDown() override;
+	UE_API virtual void TearDown() override;
 
 	// Removes all logged data 
-	ENGINE_API void Cleanup(UWorld* OldWorld, bool bReleaseMemory = false);
+	UE_API void Cleanup(UWorld* OldWorld, bool bReleaseMemory = false);
 
 	// Use when a visual logger device has discarded all of its data, waiting for new data
-	ENGINE_API void OnDataReset();
+	UE_API void OnDataReset();
 
 	/** Set log owner redirection from one object to another, to combine logs */
 	static void Redirect(const UObject* FromObject, const UObject* ToObject)
@@ -627,28 +629,28 @@ public:
 	void ClearCategoryAllowList() { CategoryAllowList.Reset(); }
 
 	/** Generates and returns Id unique for given timestamp - used to connect different logs between (ex. text log with geometry shape) */
-	ENGINE_API int32 GetUniqueId(double Timestamp);
+	UE_API int32 GetUniqueId(double Timestamp);
 
 	/** Starts visual log collecting and recording */
-	ENGINE_API void SetIsRecording(const bool bInIsRecording);
+	UE_API void SetIsRecording(const bool bInIsRecording);
 	/** return information is vlog recording is enabled or not */
 	inline static bool IsRecording() { return !!bIsRecording; }
 
 	/** Starts visual log collecting and recording */
-	ENGINE_API void SetIsRecordingToFile(bool InIsRecording);
+	UE_API void SetIsRecordingToFile(bool InIsRecording);
 	/** return information is vlog recording is enabled or not */
 	bool IsRecordingToFile() const { return !!bIsRecordingToFile; }
 	/** disables recording to file and discards all data without saving it to file */
-	ENGINE_API void DiscardRecordingToFile();
+	UE_API void DiscardRecordingToFile();
 
 	/** Starts visual log collecting and recording to insights traces (for Rewind Debugger)*/
-	ENGINE_API void SetIsRecordingToTrace(const bool bInIsRecording);
+	UE_API void SetIsRecordingToTrace(const bool bInIsRecording);
 
 	void SetIsRecordingOnServer(const bool bInIsRecording) { bIsRecordingOnServer = bInIsRecording; }
 	bool IsRecordingOnServer() const { return !!bIsRecordingOnServer; }
 
 	/** Configure whether VisLog should be using decorated, unique names */
-	ENGINE_API void SetUseUniqueNames(const bool bEnable);
+	UE_API void SetUseUniqueNames(const bool bEnable);
 
 	/** Add visual logger output device */
 	void AddDevice(FVisualLogDevice* InDevice) { OutputDevices.AddUnique(InDevice); }
@@ -657,19 +659,19 @@ public:
 	/** Remove visual logger output device */
 	const TArray<FVisualLogDevice*>& GetDevices() const { return OutputDevices; }
 	/** Check if log category can be recorded, verify before using GetEntryToWrite! */
-	ENGINE_API bool IsCategoryLogged(const FLogCategoryBase& Category) const;
+	UE_API bool IsCategoryLogged(const FLogCategoryBase& Category) const;
 
 	/** Returns current entry for given TimeStamp or creates another one, but first it serialize previous 
 	 *	entry as completed to vislog devices. Use VisualLogger::DontCreate to get current entry without serialization
 	 *	@note this function can return null */
 	UE_DEPRECATED_FORGAME(5.4, "Use the static GetEntryToWrite instead because this TimeStamp is inconsistent across multiple instances (or threads in Editor).  This function will be made private/protected.")
-	ENGINE_API FVisualLogEntry* GetEntryToWrite(const UObject* Object, double TimeStamp, ECreateIfNeeded ShouldCreate = ECreateIfNeeded::Create);
+	UE_API FVisualLogEntry* GetEntryToWrite(const UObject* Object, double TimeStamp, ECreateIfNeeded ShouldCreate = ECreateIfNeeded::Create);
 
 	/** Returns  the current (or new) entry for the given object; alternatively nullptr if we aren't allowed to vlog with the given parameters.
 	 * @param LogOwner - The UObject (typically an AActor) we are going to write log entries about.  This becomes the row in the Visual Logger timeline.
 	 * @param LogCategory - The LogCategory we are logging about.  This function will only return a valid log entry if the visual logging is enabled for the category.
 	 */
-	[[nodiscard]] static ENGINE_API FVisualLogEntry* GetEntryToWrite(const UObject* LogOwner, const FLogCategoryBase& LogCategory);
+	[[nodiscard]] static UE_API FVisualLogEntry* GetEntryToWrite(const UObject* LogOwner, const FLogCategoryBase& LogCategory);
 
 	/**
 	 * Executes the provided function using the last used entry for given UObject if any.
@@ -678,12 +680,12 @@ public:
 	 * @param Function Function to execute if a valid view can be created.
 	 * @return Whether a LogEntry was found and the function was called
 	 */
-	ENGINE_API bool ExecuteOnLastEntryForObject(const UObject* Object, TFunctionRef<void(FVisualLogEntry&)> Function);
+	UE_API bool ExecuteOnLastEntryForObject(const UObject* Object, TFunctionRef<void(FVisualLogEntry&)> Function);
 
 	/** flush and serialize data if timestamp allows it */
-	ENGINE_API virtual void Flush() override;
+	UE_API virtual void Flush() override;
 	/** Moves all threads entries into the global entry map */
-	ENGINE_API void FlushThreadsEntries();
+	UE_API void FlushThreadsEntries();
 
 	/** FileName getter to set project specific file name for vlogs - highly encouraged to use FVisualLogFilenameGetterDelegate::CreateUObject with this */
 	void SetLogFileNameGetter(const FVisualLogFilenameGetterDelegate& InLogFileNameGetter) { LogFileNameGetter = InLogFileNameGetter; }
@@ -698,16 +700,16 @@ public:
 	const TMap<FName, FVisualLogExtensionInterface*>& GetAllExtensions() const { return AllExtensions; }
 
 	/** internal check for each usage of visual logger */
-	static ENGINE_API bool CheckVisualLogInputInternal(const UObject* Object, const FName& CategoryName, ELogVerbosity::Type Verbosity, UWorld** OutWorld, FVisualLogEntry** OutCurrentEntry);
+	static UE_API bool CheckVisualLogInputInternal(const UObject* Object, const FName& CategoryName, ELogVerbosity::Type Verbosity, UWorld** OutWorld, FVisualLogEntry** OutCurrentEntry);
 
 	/** Returns a current time stamp to associate with a recorded event that occurred on Object (used for ordering events on a timeline) */
-	ENGINE_API double GetTimeStampForObject(const UObject* Object) const;
+	UE_API double GetTimeStampForObject(const UObject* Object) const;
 
 	/** Sets function to call to get a timestamp instead of the default implementation (e.g. using a network synchronized clock instead of the local world time) */
-	ENGINE_API void SetGetTimeStampFunc(TFunction<double(const UObject*)> Function);
+	UE_API void SetGetTimeStampFunc(TFunction<double(const UObject*)> Function);
 
 	typedef TMap<FObjectKey, TArray<TWeakObjectPtr<const UObject>>> FOwnerToChildrenRedirectionMap;
-	static ENGINE_API FOwnerToChildrenRedirectionMap& GetRedirectionMap(const UObject* InObject);
+	static UE_API FOwnerToChildrenRedirectionMap& GetRedirectionMap(const UObject* InObject);
 
 	typedef TMap<FObjectKey, TWeakObjectPtr<const UObject>> FChildToOwnerRedirectionMap;
 	FChildToOwnerRedirectionMap& GetChildToOwnerRedirectionMap() { return ChildToOwnerMap; }
@@ -715,12 +717,12 @@ public:
 	typedef TMap<FObjectKey, TWeakObjectPtr<const UWorld>> FObjectToWorldMapType;
 	FObjectToWorldMapType& GetObjectToWorldMap() { return ObjectToWorldMap; }
 
-	ENGINE_API void AddClassToAllowList(UClass& InClass);
-	ENGINE_API bool IsClassAllowed(const UClass& InClass) const;
+	UE_API void AddClassToAllowList(UClass& InClass);
+	UE_API bool IsClassAllowed(const UClass& InClass) const;
 
-	ENGINE_API void AddObjectToAllowList(const UObject& InObject);
-	ENGINE_API void ClearObjectAllowList();
-	ENGINE_API bool IsObjectAllowed(const UObject* InObject) const;
+	UE_API void AddObjectToAllowList(const UObject& InObject);
+	UE_API void ClearObjectAllowList();
+	UE_API bool IsObjectAllowed(const UObject* InObject) const;
 
 	struct FVisualLoggerObjectEntryMap : TMap<FObjectKey, FVisualLogEntry>
 	{
@@ -729,37 +731,37 @@ public:
 	};
 
 private:
-	ENGINE_API FVisualLogger();
+	UE_API FVisualLogger();
 	virtual void Serialize(const TCHAR* V, ELogVerbosity::Type Verbosity, const FName& Category) override { ensureMsgf(0, TEXT("Regular serialize is forbidden for visual logs")); }
 
-	ENGINE_API FVisualLoggerObjectEntryMap& GetThreadCurrentEntryMap();
+	UE_API FVisualLoggerObjectEntryMap& GetThreadCurrentEntryMap();
 
 	/** Retrieves last used entry for given UObject
 	 * @note this function can return null
 	 */
-	ENGINE_API FVisualLogEntry* GetLastEntryForObject(const UObject* Object);
+	UE_API FVisualLogEntry* GetLastEntryForObject(const UObject* Object);
 
 	/** Get global entry to write where all logs are combined, not thread safe */
-	ENGINE_API FVisualLogEntry* GetEntryToWriteInternal(const UObject* Object, double TimeStamp, ECreateIfNeeded ShouldCreate);
+	UE_API FVisualLogEntry* GetEntryToWriteInternal(const UObject* Object, double TimeStamp, ECreateIfNeeded ShouldCreate);
 
 	/** Redirect internal implementation, not thread safe */
-	ENGINE_API UObject* RedirectInternal(const UObject* FromObject, const UObject* ToObject);
+	UE_API UObject* RedirectInternal(const UObject* FromObject, const UObject* ToObject);
 	/** Find redirects internal implementation, not thread safe */
-	ENGINE_API UObject* FindRedirectionInternal(const UObject* Object) const;
+	UE_API UObject* FindRedirectionInternal(const UObject* Object) const;
 	/** Cleanup invalid redirects */
-	ENGINE_API void CleanupRedirects();
+	UE_API void CleanupRedirects();
 	/** Figure out all conditions if this entry is allowed to log */
-	ENGINE_API void CalculateEntryAllowLogging(FVisualLogEntry* CurrentEntry, const UObject* LogOwner, const UObject* Object) const;
+	UE_API void CalculateEntryAllowLogging(FVisualLogEntry* CurrentEntry, const UObject* LogOwner, const UObject* Object) const;
 
 protected:
 	/** Flushes entries recorded in the frame */
-	ENGINE_API void Tick(float DeltaTime);
+	UE_API void Tick(float DeltaTime);
 
 	/**
 	 * Serializes a single entry and resets it.
 	 * Method expects an initialized entry and will ensure otherwise. 
 	 */
-	ENGINE_API void FlushEntry(FVisualLogEntry& Entry, const FObjectKey& ObjectKey);
+	UE_API void FlushEntry(FVisualLogEntry& Entry, const FObjectKey& ObjectKey);
 
 	/** Handle to the registered ticker to flush entries */
 	FTSTicker::FDelegateHandle TickerHandle;
@@ -767,7 +769,7 @@ protected:
 	/** Array of output devices to redirect to */
 	TArray<FVisualLogDevice*> OutputDevices;
 	// Map for inter-objects redirections
-	static ENGINE_API TMap<const UWorld*, FOwnerToChildrenRedirectionMap> WorldToRedirectionMap;
+	static UE_API TMap<const UWorld*, FOwnerToChildrenRedirectionMap> WorldToRedirectionMap;
 
 	// allowed classes - only instances of these classes will be logged. 
 	// if ClassAllowList is empty (default) everything will log
@@ -843,7 +845,9 @@ protected:
 #endif
 
 	// if set we are recording and collecting all vlog data
-	static ENGINE_API int32 bIsRecording;
+	static UE_API int32 bIsRecording;
 };
 
 #endif //ENABLE_VISUAL_LOG
+
+#undef UE_API
