[Engine]
* Moved FBodySetupObjectTextFactory to its own file
* Removed some includes in Chaos/ParticleHandle and reduced module dependencies

diff --git a/Engine/Source/Editor/PhysicsAssetEditor/Private/PhysicsAssetEditorSharedData.cpp b/Engine/Source/Editor/PhysicsAssetEditor/Private/PhysicsAssetEditorSharedData.cpp
--- a/Engine/Source/Editor/PhysicsAssetEditor/Private/PhysicsAssetEditorSharedData.cpp
+++ b/Engine/Source/Editor/PhysicsAssetEditor/Private/PhysicsAssetEditorSharedData.cpp
@@ -9,6 +9,7 @@
 #include "PhysicsAssetEditorPhysicsHandleComponent.h"
 #include "PhysicsAssetRenderUtils.h"
 #include "PhysicsAssetEditorSelection.h"
+#include "PhysicsEngine/BodySetupObjectTextFactory.h"
 #include "PhysicsEngine/RigidBodyIndexPair.h"
 #include "Math/Axis.h"
 #include "Misc/MessageDialog.h"
diff --git a/Engine/Source/Editor/StaticMeshEditor/Private/StaticMeshEditor.cpp b/Engine/Source/Editor/StaticMeshEditor/Private/StaticMeshEditor.cpp
--- a/Engine/Source/Editor/StaticMeshEditor/Private/StaticMeshEditor.cpp
+++ b/Engine/Source/Editor/StaticMeshEditor/Private/StaticMeshEditor.cpp
@@ -8,6 +8,7 @@
 #include "Misc/StringOutputDevice.h"
 #include "Modules/ModuleManager.h"
 #include "Framework/Application/SlateApplication.h"
+#include "PhysicsEngine/BodySetupObjectTextFactory.h"
 #include "Styling/AppStyle.h"
 #include "EditorReimportHandler.h"
 #include "Editor/UnrealEdEngine.h"
diff --git a/Engine/Source/Runtime/Engine/Classes/PhysicsEngine/BodySetup.h b/Engine/Source/Runtime/Engine/Classes/PhysicsEngine/BodySetup.h
--- a/Engine/Source/Runtime/Engine/Classes/PhysicsEngine/BodySetup.h
+++ b/Engine/Source/Runtime/Engine/Classes/PhysicsEngine/BodySetup.h
@@ -15,9 +15,6 @@
 #include "HAL/ThreadSafeBool.h"
 #include "Async/TaskGraphInterfaces.h"
 #include "BodySetupCore.h"
-#if WITH_EDITOR
-#include "Factories.h"
-#endif
 #include "BodySetup.generated.h"
 
 class ITargetPlatform;
@@ -473,18 +470,3 @@ public:
 	friend struct FIterateBodySetupHelper;
 
 };
-
-#if WITH_EDITOR
-
-class FBodySetupObjectTextFactory : public FCustomizableTextObjectFactory
-{
-public:
-	FBodySetupObjectTextFactory() : FCustomizableTextObjectFactory(GWarn) { }
-	ENGINE_API virtual bool CanCreateClass(UClass* InObjectClass, bool& bOmitSubObjs) const override;
-	ENGINE_API virtual void ProcessConstructedObject(UObject* NewObject) override;
-
-public:
-	TArray<UBodySetup*> NewBodySetups;
-};
-
-#endif // WITH_EDITOR
diff --git a/Engine/Source/Runtime/Engine/Classes/PhysicsEngine/BodySetupObjectTextFactory.h b/Engine/Source/Runtime/Engine/Classes/PhysicsEngine/BodySetupObjectTextFactory.h
new file mode 100644
--- /dev/null
+++ b/Engine/Source/Runtime/Engine/Classes/PhysicsEngine/BodySetupObjectTextFactory.h
@@ -0,0 +1,21 @@
+// Copyright Epic Games, Inc. All Rights Reserved.
+
+#pragma once
+
+#if WITH_EDITOR
+
+#include "Factories.h"
+#include "PhysicsEngine/BodySetup.h"
+
+class FBodySetupObjectTextFactory : public FCustomizableTextObjectFactory
+{
+public:
+	FBodySetupObjectTextFactory() : FCustomizableTextObjectFactory(GWarn) { }
+	ENGINE_API virtual bool CanCreateClass(UClass* InObjectClass, bool& bOmitSubObjs) const override;
+	ENGINE_API virtual void ProcessConstructedObject(UObject* NewObject) override;
+
+public:
+	TArray<UBodySetup*> NewBodySetups;
+};
+
+#endif // WITH_EDITOR
diff --git a/Engine/Source/Runtime/Engine/Private/Collision/CollisionConversions.cpp b/Engine/Source/Runtime/Engine/Private/Collision/CollisionConversions.cpp
--- a/Engine/Source/Runtime/Engine/Private/Collision/CollisionConversions.cpp
+++ b/Engine/Source/Runtime/Engine/Private/Collision/CollisionConversions.cpp
@@ -2,6 +2,7 @@
 
 #include "Collision/CollisionConversions.h"
 #include "BodySetupCore.h"
+#include "Chaos/ChaosUserEntity.h"
 #include "Engine/OverlapResult.h"
 #include "Engine/World.h"
 #include "Chaos/Capsule.h"
diff --git a/Engine/Source/Runtime/Engine/Private/PhysicsEngine/BodySetup.cpp b/Engine/Source/Runtime/Engine/Private/PhysicsEngine/BodySetup.cpp
--- a/Engine/Source/Runtime/Engine/Private/PhysicsEngine/BodySetup.cpp
+++ b/Engine/Source/Runtime/Engine/Private/PhysicsEngine/BodySetup.cpp
@@ -44,6 +44,7 @@
 
 #if WITH_EDITOR
 #include "Misc/DataValidation.h"
+#include "PhysicsEngine/BodySetupObjectTextFactory.h"
 #endif
 
 #include UE_INLINE_GENERATED_CPP_BY_NAME(BodySetup)
diff --git a/Engine/Source/Runtime/Engine/Private/PhysicsEngine/Experimental/ChaosInterfaceWrapper.cpp b/Engine/Source/Runtime/Engine/Private/PhysicsEngine/Experimental/ChaosInterfaceWrapper.cpp
--- a/Engine/Source/Runtime/Engine/Private/PhysicsEngine/Experimental/ChaosInterfaceWrapper.cpp
+++ b/Engine/Source/Runtime/Engine/Private/PhysicsEngine/Experimental/ChaosInterfaceWrapper.cpp
@@ -1,6 +1,7 @@
 // Copyright Epic Games, Inc. All Rights Reserved.
 
 #include "Physics/Experimental/ChaosInterfaceWrapper.h"
+#include "Chaos/ChaosUserEntity.h"
 #include "Physics/Experimental/PhysScene_Chaos.h"
 #include "PBDRigidsSolver.h"
 
diff --git a/Engine/Source/Runtime/Experimental/Chaos/Private/Chaos/PhysicsObjectInterface.cpp b/Engine/Source/Runtime/Experimental/Chaos/Private/Chaos/PhysicsObjectInterface.cpp
--- a/Engine/Source/Runtime/Experimental/Chaos/Private/Chaos/PhysicsObjectInterface.cpp
+++ b/Engine/Source/Runtime/Experimental/Chaos/Private/Chaos/PhysicsObjectInterface.cpp
@@ -2,6 +2,7 @@
 #include "Chaos/PhysicsObjectInterface.h"
 
 #include "Chaos/CastingUtilities.h"
+#include "Chaos/ChaosUserEntity.h"
 #include "Chaos/CollisionFilterData.h"
 #include "Chaos/PhysicsObjectCollisionInterface.h"
 #include "Chaos/PhysicsObjectInternalInterface.h"
diff --git a/Engine/Source/Runtime/Experimental/Chaos/Public/Chaos/ParticleHandle.h b/Engine/Source/Runtime/Experimental/Chaos/Public/Chaos/ParticleHandle.h
--- a/Engine/Source/Runtime/Experimental/Chaos/Public/Chaos/ParticleHandle.h
+++ b/Engine/Source/Runtime/Experimental/Chaos/Public/Chaos/ParticleHandle.h
@@ -2,7 +2,6 @@
 #pragma once
 
 #include "Math/NumericLimits.h"
-#include "Chaos/ChaosUserEntity.h"
 #include "Chaos/ISpatialAcceleration.h"
 #include "Chaos/PBDRigidClusteredParticles.h"
 #include "Chaos/PBDGeometryCollectionParticles.h"
@@ -14,7 +13,6 @@
 #include "Chaos/ChaosDebugDrawDeclares.h"
 #include "Chaos/AsyncInitBodyHelper.h"
 #include "ChaosVisualDebugger/ChaosVisualDebuggerTrace.h"
-#include "PhysicsInterfaceTypesCore.h"
 #if CHAOS_DEBUG_DRAW
 #include "Chaos/ChaosDebugDraw.h"
 #endif
