Preshader2:  static analysis false alarm suppression, take two.

#rnx

diff --git a/Engine/Source/Runtime/Engine/Private/Shader/Preshader2.cpp b/Engine/Source/Runtime/Engine/Private/Shader/Preshader2.cpp
--- a/Engine/Source/Runtime/Engine/Private/Shader/Preshader2.cpp
+++ b/Engine/Source/Runtime/Engine/Private/Shader/Preshader2.cpp
@@ -851,6 +851,8 @@ void OpEvaluate(const FOpSparseVolumeTextureUniform&, const FMaterialRenderConte
 // InPlace:				Opcode(1)  ResultOffset(2)  Dimension+Type(1)				AOffset same as ResultOffset
 // Default:				Opcode(1)  ResultOffset(2)  AOffset(2)  Dimension+Type(1)
 // Does not sanitize Dimension in case custom implementation stores extra data in that byte!  Client evaluate function should remove junk high bits from Dimension.
+
+//-V:PRESHADER2_UNARY_OP:521
 #define PRESHADER2_UNARY_OP(OP_NAME) \
 	case EPreshader2Opcode::OP_NAME:																																	\
 		OpcodePtr += 2;									/* Skip to AOffset */																							\
@@ -869,6 +871,8 @@ void OpEvaluate(const FOpSparseVolumeTextureUniform&, const FMaterialRenderConte
 // "Reuse" variations, so it's pulled from the opcode.  We subtract the opcode from the first "Cast_Reuse" opcode to extract the type and add it to the Dimension,
 // where it is stored in the other versions of the opcode.  To reduce code size bloat, we technically only need a case for the last version of the opcode, so we
 // fall through in the prior cases (i.e., actual code is only generated for Cast_ReuseLWC, and the other two opcodes only produce the "case" statement).
+
+//-V:PRESHADER2_UNARY_OP_CAST_REUSE:501
 #define PRESHADER2_UNARY_OP_CAST_REUSE(OP_NAME) \
 	case EPreshader2Opcode::OP_NAME:																																			\
 		if constexpr (EPreshader2Opcode::OP_NAME == EPreshader2Opcode::Cast_ReuseLWC)																							\
@@ -992,7 +996,7 @@ void Evaluate(UE::Shader::FPreshaderDataContext& DataContext, const FMaterialRen
 		int32 AInc;
 		int32 BInc;
 
-		switch (Opcode)		//-V501 //-V521
+		switch (Opcode)
 		{
 			#include "Shader/Preshader2Opcodes.inl"
 		}
