Fixed unreachable code warning.

#rnx

diff --git a/Engine/Source/Runtime/Engine/Private/Collision/SceneQuery.cpp b/Engine/Source/Runtime/Engine/Private/Collision/SceneQuery.cpp
--- a/Engine/Source/Runtime/Engine/Private/Collision/SceneQuery.cpp
+++ b/Engine/Source/Runtime/Engine/Private/Collision/SceneQuery.cpp
@@ -41,8 +41,9 @@ bool GetIgnoreQueryHandler()
 {
 #if UE_WITH_REMOTE_OBJECT_HANDLE
 return bIgnoreQueryHandler;
-#endif
+#else
 return false;
+#endif
 }
 
 enum class ESingleMultiOrTest : uint8
