PR #10494: Fix AdjustCoord() call in UPhysicalMaterialMask::GetPhysMatIndex()

diff --git a/Engine/Source/Runtime/Engine/Private/PhysicsEngine/PhysicalMaterialMask.cpp b/Engine/Source/Runtime/Engine/Private/PhysicsEngine/PhysicalMaterialMask.cpp
--- a/Engine/Source/Runtime/Engine/Private/PhysicsEngine/PhysicalMaterialMask.cpp
+++ b/Engine/Source/Runtime/Engine/Private/PhysicsEngine/PhysicalMaterialMask.cpp
@@ -405,7 +405,7 @@ uint32 UPhysicalMaterialMask::GetPhysMatIndex(const TArray<uint32>& MaskData, in
 	};
 
 	AdjustCoord(U, static_cast<TextureAddress>(AddressX));
-	AdjustCoord(V, static_cast<TextureAddress>(AddressX));
+	AdjustCoord(V, static_cast<TextureAddress>(AddressY));
 
 	uint32 X = FMath::FloorToInt((float)SizeX * U);
 	uint32 Y = FMath::FloorToInt((float)SizeY * V);
