CIS Fix

diff --git a/Engine/Source/Runtime/Engine/Classes/PhysicsEngine/BodySetup.h b/Engine/Source/Runtime/Engine/Classes/PhysicsEngine/BodySetup.h
--- a/Engine/Source/Runtime/Engine/Classes/PhysicsEngine/BodySetup.h
+++ b/Engine/Source/Runtime/Engine/Classes/PhysicsEngine/BodySetup.h
@@ -17,7 +17,11 @@
 #include "BodySetupCore.h"
 
 #if UE_ENABLE_INCLUDE_ORDER_DEPRECATED_IN_5_8
+
+#if WITH_EDITOR
 #include "Factories.h"
+#endif
+
 #include "PhysicsEngine/BodySetupObjectTextFactory.h"
 #endif // UE_ENABLE_INCLUDE_ORDER_DEPRECATED_IN_5_8
 
