[ObjectTrace] updated most events to the UE_TRACE_MINIMAL API
OBJECT_TRACE_ENABLED can now be enabled by UE_TRACE_MINIMAL_ENABLED
OBJECT_PROPERTY_TRACE_ENABLED stays enabled only in non TEST or SHIPPING builds

diff --git a/Engine/Source/Runtime/Engine/Private/ObjectTrace.cpp b/Engine/Source/Runtime/Engine/Private/ObjectTrace.cpp
--- a/Engine/Source/Runtime/Engine/Private/ObjectTrace.cpp
+++ b/Engine/Source/Runtime/Engine/Private/ObjectTrace.cpp
@@ -12,74 +12,75 @@
 #else
 #include "UObject/Package.h"
 #include "UObject/UObjectAnnotation.h"
-#endif
-
-UE_TRACE_CHANNEL(ObjectChannel)
-
-UE_TRACE_EVENT_BEGIN(Object, Type)
-	UE_TRACE_EVENT_FIELD(uint64, Id)
-	UE_TRACE_EVENT_FIELD(uint64, SuperId)
-	UE_TRACE_EVENT_FIELD(UE::Trace::WideString, Name)
-	UE_TRACE_EVENT_FIELD(UE::Trace::WideString, Path)
-UE_TRACE_EVENT_END()
-
-UE_TRACE_EVENT_BEGIN(Object, Object)
-	UE_TRACE_EVENT_FIELD(uint64, Id)
-	UE_TRACE_EVENT_FIELD(uint64, ClassId)
-	UE_TRACE_EVENT_FIELD(uint64, OuterId)
-	UE_TRACE_EVENT_FIELD(UE::Trace::WideString, Name)
-	UE_TRACE_EVENT_FIELD(UE::Trace::WideString, Path)
-	UE_TRACE_EVENT_FIELD(uint8, Flags)
-UE_TRACE_EVENT_END()
-
-UE_TRACE_EVENT_BEGIN(Object, ObjectLifetimeBegin2)
-	UE_TRACE_EVENT_FIELD(uint64, Cycle)
-	UE_TRACE_EVENT_FIELD(double, RecordingTime)
-	UE_TRACE_EVENT_FIELD(uint64, Id)
-UE_TRACE_EVENT_END()
-
-UE_TRACE_EVENT_BEGIN(Object, ObjectTransform)
-	UE_TRACE_EVENT_FIELD(uint64, Cycle)
-	UE_TRACE_EVENT_FIELD(double, RecordingTime)
-	UE_TRACE_EVENT_FIELD(uint64, Id)
-	UE_TRACE_EVENT_FIELD(uint8, TeleportType)
-	UE_TRACE_EVENT_FIELD(double[], Transform)
-UE_TRACE_EVENT_END()
-
-UE_TRACE_EVENT_BEGIN(Object, ObjectLifetimeEnd2)
-	UE_TRACE_EVENT_FIELD(uint64, Cycle)
-	UE_TRACE_EVENT_FIELD(double, RecordingTime)
-	UE_TRACE_EVENT_FIELD(uint64, Id)
-UE_TRACE_EVENT_END()
-
-UE_TRACE_EVENT_BEGIN(Object, ObjectEvent)
-	UE_TRACE_EVENT_FIELD(uint64, Cycle)
-	UE_TRACE_EVENT_FIELD(uint64, Id)
-	UE_TRACE_EVENT_FIELD(UE::Trace::WideString, Event)
-UE_TRACE_EVENT_END()
-
-UE_TRACE_EVENT_BEGIN(Object, PawnPossess)
-	UE_TRACE_EVENT_FIELD(uint64, Cycle)
-	UE_TRACE_EVENT_FIELD(uint64, ControllerId)
-	UE_TRACE_EVENT_FIELD(uint64, PawnId)
-UE_TRACE_EVENT_END()
-
-UE_TRACE_EVENT_BEGIN(Object, World)
-	UE_TRACE_EVENT_FIELD(uint64, Id)
-	UE_TRACE_EVENT_FIELD(int32, PIEInstanceId)
-	UE_TRACE_EVENT_FIELD(uint8, Type)
-	UE_TRACE_EVENT_FIELD(uint8, NetMode)
-	UE_TRACE_EVENT_FIELD(bool, IsSimulating)
-UE_TRACE_EVENT_END()
-
-UE_TRACE_EVENT_BEGIN(Object, RecordingInfo)
-	UE_TRACE_EVENT_FIELD(uint64, WorldId)
-	UE_TRACE_EVENT_FIELD(uint64, Cycle)
-	UE_TRACE_EVENT_FIELD(uint32, RecordingIndex)
-	UE_TRACE_EVENT_FIELD(uint32, FrameIndex)
-	UE_TRACE_EVENT_FIELD(double, ElapsedTime)
-UE_TRACE_EVENT_END()
-
+#endif // WITH_EDITOR
+
+UE_TRACE_MINIMAL_CHANNEL(ObjectChannel)
+
+UE_TRACE_MINIMAL_EVENT_BEGIN(Object, Type)
+	UE_TRACE_MINIMAL_EVENT_FIELD(uint64, Id)
+	UE_TRACE_MINIMAL_EVENT_FIELD(uint64, SuperId)
+	UE_TRACE_MINIMAL_EVENT_FIELD(UE::Trace::WideString, Name)
+	UE_TRACE_MINIMAL_EVENT_FIELD(UE::Trace::WideString, Path)
+UE_TRACE_MINIMAL_EVENT_END()
+
+UE_TRACE_MINIMAL_EVENT_BEGIN(Object, Object)
+	UE_TRACE_MINIMAL_EVENT_FIELD(uint64, Id)
+	UE_TRACE_MINIMAL_EVENT_FIELD(uint64, ClassId)
+	UE_TRACE_MINIMAL_EVENT_FIELD(uint64, OuterId)
+	UE_TRACE_MINIMAL_EVENT_FIELD(UE::Trace::WideString, Name)
+	UE_TRACE_MINIMAL_EVENT_FIELD(UE::Trace::WideString, Path)
+	UE_TRACE_MINIMAL_EVENT_FIELD(uint8, Flags)
+UE_TRACE_MINIMAL_EVENT_END()
+
+UE_TRACE_MINIMAL_EVENT_BEGIN(Object, ObjectLifetimeBegin2)
+	UE_TRACE_MINIMAL_EVENT_FIELD(uint64, Cycle)
+	UE_TRACE_MINIMAL_EVENT_FIELD(double, RecordingTime)
+	UE_TRACE_MINIMAL_EVENT_FIELD(uint64, Id)
+UE_TRACE_MINIMAL_EVENT_END()
+
+UE_TRACE_MINIMAL_EVENT_BEGIN(Object, ObjectTransform)
+	UE_TRACE_MINIMAL_EVENT_FIELD(uint64, Cycle)
+	UE_TRACE_MINIMAL_EVENT_FIELD(double, RecordingTime)
+	UE_TRACE_MINIMAL_EVENT_FIELD(uint64, Id)
+	UE_TRACE_MINIMAL_EVENT_FIELD(uint8, TeleportType)
+	UE_TRACE_MINIMAL_EVENT_FIELD(double[], Transform)
+UE_TRACE_MINIMAL_EVENT_END()
+
+UE_TRACE_MINIMAL_EVENT_BEGIN(Object, ObjectLifetimeEnd2)
+	UE_TRACE_MINIMAL_EVENT_FIELD(uint64, Cycle)
+	UE_TRACE_MINIMAL_EVENT_FIELD(double, RecordingTime)
+	UE_TRACE_MINIMAL_EVENT_FIELD(uint64, Id)
+UE_TRACE_MINIMAL_EVENT_END()
+
+UE_TRACE_MINIMAL_EVENT_BEGIN(Object, ObjectEvent)
+	UE_TRACE_MINIMAL_EVENT_FIELD(uint64, Cycle)
+	UE_TRACE_MINIMAL_EVENT_FIELD(uint64, Id)
+	UE_TRACE_MINIMAL_EVENT_FIELD(UE::Trace::WideString, Event)
+UE_TRACE_MINIMAL_EVENT_END()
+
+UE_TRACE_MINIMAL_EVENT_BEGIN(Object, PawnPossess)
+	UE_TRACE_MINIMAL_EVENT_FIELD(uint64, Cycle)
+	UE_TRACE_MINIMAL_EVENT_FIELD(uint64, ControllerId)
+	UE_TRACE_MINIMAL_EVENT_FIELD(uint64, PawnId)
+UE_TRACE_MINIMAL_EVENT_END()
+
+UE_TRACE_MINIMAL_EVENT_BEGIN(Object, World)
+	UE_TRACE_MINIMAL_EVENT_FIELD(uint64, Id)
+	UE_TRACE_MINIMAL_EVENT_FIELD(int32, PIEInstanceId)
+	UE_TRACE_MINIMAL_EVENT_FIELD(uint8, Type)
+	UE_TRACE_MINIMAL_EVENT_FIELD(uint8, NetMode)
+	UE_TRACE_MINIMAL_EVENT_FIELD(bool, IsSimulating)
+UE_TRACE_MINIMAL_EVENT_END()
+
+UE_TRACE_MINIMAL_EVENT_BEGIN(Object, RecordingInfo)
+	UE_TRACE_MINIMAL_EVENT_FIELD(uint64, WorldId)
+	UE_TRACE_MINIMAL_EVENT_FIELD(uint64, Cycle)
+	UE_TRACE_MINIMAL_EVENT_FIELD(uint32, RecordingIndex)
+	UE_TRACE_MINIMAL_EVENT_FIELD(uint32, FrameIndex)
+	UE_TRACE_MINIMAL_EVENT_FIELD(double, ElapsedTime)
+UE_TRACE_MINIMAL_EVENT_END()
+
+// Not in the minimal trace set of events since not required on shipping servers
 UE_TRACE_EVENT_BEGIN(Object, View)
 	UE_TRACE_EVENT_FIELD(uint64, PlayerId)
 	UE_TRACE_EVENT_FIELD(uint64, Cycle)
@@ -98,12 +99,11 @@ UE_TRACE_EVENT_END()
 
 namespace UE::ObjectTrace
 {
-	
+
 // duplicate ObjectInfoFlags must match FObjectInfo::Flag_
 constexpr uint8 ObjectInfoFlag_Transient = 0x1;
 constexpr uint8 ObjectInfoFlag_StructInstance = 0x2;
 
-	
 bool IsTracingDisabledForObject(const UObject* InObject)
 {
 	if (InObject == nullptr)
@@ -111,7 +111,7 @@ bool IsTracingDisabledForObject(const UObject* InObject)
 		return true;
 	}
 
-	if (!UE_TRACE_CHANNELEXPR_IS_ENABLED(ObjectChannel))
+	if (!UE_TRACE_MINIMAL_CHANNELEXPR_IS_ENABLED(ObjectChannel))
 	{
 		return true;
 	}
@@ -195,7 +195,7 @@ static void TickObjectTraceWorldSubsystem(UWorld* InWorld, ELevelTick InTickType
 				Subsystem->FrameIndex++;
 				Subsystem->ElapsedTime += InDeltaSeconds;
 
-				UE_TRACE_LOG(Object, RecordingInfo, ObjectChannel)
+				UE_TRACE_MINIMAL_LOG(Object, RecordingInfo, ObjectChannel)
 					<< RecordingInfo.WorldId(FObjectTrace::GetObjectId(InWorld))
 					<< RecordingInfo.Cycle(FPlatformTime::Cycles64())
 					<< RecordingInfo.ElapsedTime(Subsystem->ElapsedTime)
@@ -390,7 +390,7 @@ uint16 FObjectTrace::GetObjectWorldTickCounter(const UObject* InObject)
 
 void FObjectTrace::OutputType(const UStruct* InType)
 {
-	const bool bChannelEnabled = UE_TRACE_CHANNELEXPR_IS_ENABLED(ObjectChannel);
+	const bool bChannelEnabled = UE_TRACE_MINIMAL_CHANNELEXPR_IS_ENABLED(ObjectChannel);
 	if (!bChannelEnabled || InType == nullptr)
 	{
 		return;
@@ -407,7 +407,7 @@ void FObjectTrace::OutputType(const UStruct* InType)
 	TCHAR TypeName[FName::StringBufferSize];
 	const uint32 TypeNameLength = InType->GetFName().ToString(TypeName);
 
-	UE_TRACE_LOG(Object, Type, ObjectChannel)
+	UE_TRACE_MINIMAL_LOG(Object, Type, ObjectChannel)
 		<< Type.Id(ObjectId)
 		<< Type.SuperId(GetObjectId(InType->GetSuperStruct()))
 		<< Type.Name(TypeName, TypeNameLength)
@@ -445,7 +445,7 @@ void OutputInstanceInternal(const UObject* InOuterObject, uint64 InInstanceId, u
 	// Trace the object's class first
 	TRACE_TYPE(InType);
 
-	UE_TRACE_LOG(Object, Object, ObjectChannel)
+	UE_TRACE_MINIMAL_LOG(Object, Object, ObjectChannel)
 		<< Object.Id(InInstanceId)
 		<< Object.ClassId(FObjectTrace::GetObjectId(InType))
 		<< Object.OuterId(InOuterId)
@@ -453,7 +453,7 @@ void OutputInstanceInternal(const UObject* InOuterObject, uint64 InInstanceId, u
 		<< Object.Path(*InPathName, InPathName.Len())
 		<< Object.Flags(InFlags);
 		
-	UE_TRACE_LOG(Object, ObjectLifetimeBegin2, ObjectChannel)
+	UE_TRACE_MINIMAL_LOG(Object, ObjectLifetimeBegin2, ObjectChannel)
 		<< ObjectLifetimeBegin2.Cycle(FPlatformTime::Cycles64())
 		<< ObjectLifetimeBegin2.RecordingTime(FObjectTrace::GetWorldElapsedTime(InOuterObject->GetWorld()))
 		<< ObjectLifetimeBegin2.Id(InInstanceId);
@@ -465,38 +465,35 @@ void FObjectTrace::OutputInstance(const UObject* InOuterObject, uint64 InInstanc
 	{
 		return;
 	}
-	
+
 	if (!AddObjectToTracedSet(UE::ObjectTrace::FObjectId{InInstanceId}))
 	{
 		return;
 	}
 
-	uint8 Flags = UE::ObjectTrace::ObjectInfoFlag_StructInstance; 
-	
+	constexpr uint8 Flags = UE::ObjectTrace::ObjectInfoFlag_StructInstance;
 	OutputInstanceInternal(InOuterObject, InInstanceId, InOuterId, InType, InName, InPathName, Flags);
 }
 
-
 bool FObjectTrace::OutputObject(const UObject* InObject)
 {
 	if (UE::ObjectTrace::IsTracingDisabledForObject(InObject))
- 	{
- 		return false;
- 	}
+	{
+		return false;
+	}
 
-	uint64 InstanceId = GetObjectId(InObject);
-	
-	if (!AddObjectToTracedSet(UE::ObjectTrace::FObjectId{InstanceId}))
+	const uint64 InstanceId = GetObjectId(InObject);
+	if (!AddObjectToTracedSet(UE::ObjectTrace::FObjectId{ InstanceId }))
 	{
 		return false;
 	}
-	
+
 	uint8 Flags = 0;
 	if (InObject->HasAnyFlags(RF_Transient))
 	{
-		Flags |=  UE::ObjectTrace::ObjectInfoFlag_Transient;
+		Flags |= UE::ObjectTrace::ObjectInfoFlag_Transient;
 	}
- 	
+
 	OutputObject(InObject->GetOuter());
 	OutputInstanceInternal(InObject, InstanceId, GetObjectId(InObject->GetOuter()), InObject->GetClass(), InObject->GetFName().ToString(), InObject->GetPathName(), Flags);
 	return true;
@@ -511,7 +508,7 @@ void FObjectTrace::OutputObjectTransform(const UObject* InObject, const FTransfo
 
 	TRACE_OBJECT(InObject);
 
-	UE_TRACE_LOG(Object, ObjectTransform, ObjectChannel)
+	UE_TRACE_MINIMAL_LOG(Object, ObjectTransform, ObjectChannel)
 		<< ObjectTransform.Cycle(FPlatformTime::Cycles64())
 		<< ObjectTransform.RecordingTime(FObjectTrace::GetWorldElapsedTime(InObject->GetWorld()))
 		<< ObjectTransform.Id(GetObjectId(InObject))
@@ -528,7 +525,7 @@ void FObjectTrace::OutputObjectEvent(const UObject* InObject, const TCHAR* InEve
 
 	TRACE_OBJECT(InObject);
 
-	UE_TRACE_LOG(Object, ObjectEvent, ObjectChannel)
+	UE_TRACE_MINIMAL_LOG(Object, ObjectEvent, ObjectChannel)
 		<< ObjectEvent.Cycle(FPlatformTime::Cycles64())
 		<< ObjectEvent.Id(GetObjectId(InObject))
 		<< ObjectEvent.Event(InEvent);
@@ -547,7 +544,7 @@ void FObjectTrace::OutputInstanceLifetimeEnd(const UObject* InOuterObject, uint6
 	}
 	if (UWorld* World = InOuterObject->GetWorld())
 	{
-		UE_TRACE_LOG(Object, ObjectLifetimeEnd2, ObjectChannel)
+		UE_TRACE_MINIMAL_LOG(Object, ObjectLifetimeEnd2, ObjectChannel)
 			<< ObjectLifetimeEnd2.Cycle(FPlatformTime::Cycles64())
 			<< ObjectLifetimeEnd2.RecordingTime(FObjectTrace::GetWorldElapsedTime(World))
 			<< ObjectLifetimeEnd2.Id(InInstanceId);
@@ -573,7 +570,7 @@ void FObjectTrace::OutputPawnPossess(const UObject* InController, const UObject*
 		TRACE_OBJECT(InPawn);
 	}
 
-	UE_TRACE_LOG(Object, PawnPossess, ObjectChannel)
+	UE_TRACE_MINIMAL_LOG(Object, PawnPossess, ObjectChannel)
 		<< PawnPossess.Cycle(FPlatformTime::Cycles64())
 		<< PawnPossess.ControllerId(GetObjectId(InController))
 		<< PawnPossess.PawnId(GetObjectId(InPawn));
@@ -581,7 +578,7 @@ void FObjectTrace::OutputPawnPossess(const UObject* InController, const UObject*
 
 void FObjectTrace::OutputWorld(const UWorld* InWorld)
 {
-	const bool bChannelEnabled = UE_TRACE_CHANNELEXPR_IS_ENABLED(ObjectChannel);
+	const bool bChannelEnabled = UE_TRACE_MINIMAL_CHANNELEXPR_IS_ENABLED(ObjectChannel);
 	if (!bChannelEnabled || InWorld == nullptr)
 	{
 		return;
@@ -598,7 +595,7 @@ void FObjectTrace::OutputWorld(const UWorld* InWorld)
 	bool bIsSimulating = false;
 #endif
 
-	UE_TRACE_LOG(Object, World, ObjectChannel)
+	UE_TRACE_MINIMAL_LOG(Object, World, ObjectChannel)
 		<< World.Id(GetObjectId(InWorld))
 		<< World.PIEInstanceId(InWorld->GetOutermost()->GetPIEInstanceID())
 		<< World.Type(static_cast<uint8>(InWorld->WorldType))
@@ -609,4 +606,4 @@ void FObjectTrace::OutputWorld(const UWorld* InWorld)
 	TRACE_OBJECT(InWorld);
 }
 
-#endif
+#endif // OBJECT_TRACE_ENABLED
diff --git a/Engine/Source/Runtime/Engine/Public/ObjectPropertyTrace.h b/Engine/Source/Runtime/Engine/Public/ObjectPropertyTrace.h
--- a/Engine/Source/Runtime/Engine/Public/ObjectPropertyTrace.h
+++ b/Engine/Source/Runtime/Engine/Public/ObjectPropertyTrace.h
@@ -9,11 +9,7 @@
 
 #include "ObjectTraceDefines.h"
 
-#if WITH_ENGINE
-#define OBJECT_PROPERTY_TRACE_ENABLED OBJECT_TRACE_ENABLED
-#else
-#define OBJECT_PROPERTY_TRACE_ENABLED 0
-#endif
+#define OBJECT_PROPERTY_TRACE_ENABLED (WITH_ENGINE && OBJECT_TRACE_ENABLED && !(UE_BUILD_SHIPPING || UE_BUILD_TEST))
 
 #if OBJECT_PROPERTY_TRACE_ENABLED
 
@@ -47,9 +43,9 @@ struct FObjectPropertyTrace
 #define TRACE_OBJECT_PROPERTIES_BEGIN(Object) FObjectPropertyTrace::RegisterObject(Object);
 #define TRACE_OBJECT_PROPERTIES_END(Object) FObjectPropertyTrace::UnregisterObject(Object);
 
-#else
+#else // OBJECT_PROPERTY_TRACE_ENABLED
 
 #define TRACE_OBJECT_PROPERTIES_BEGIN(Object)
 #define TRACE_OBJECT_PROPERTIES_END(Object)
 
-#endif
+#endif // OBJECT_PROPERTY_TRACE_ENABLED
diff --git a/Engine/Source/Runtime/Engine/Public/ObjectTrace.h b/Engine/Source/Runtime/Engine/Public/ObjectTrace.h
--- a/Engine/Source/Runtime/Engine/Public/ObjectTrace.h
+++ b/Engine/Source/Runtime/Engine/Public/ObjectTrace.h
@@ -162,7 +162,7 @@ struct FObjectTrace
 #define TRACE_PAWN_POSSESS(Controller, Pawn)\
 	if (CAN_TRACE_OBJECT(Controller) && (Pawn==nullptr || (CAN_TRACE_OBJECT(Pawn)))) { FObjectTrace::OutputPawnPossess(Controller, Pawn); }
 
-#else
+#else // TRACE_FILTERING_ENABLED
 
 #define TRACE_OBJECT_EVENT(Object, Event) \
 	UNCONDITIONAL_TRACE_OBJECT_EVENT(Object, Event);
@@ -176,10 +176,13 @@ struct FObjectTrace
 #define TRACE_OBJECT_LIFETIME_END(Object) \
 	FObjectTrace::OutputObjectLifetimeEnd(Object);
 
+#define TRACE_INSTANCE_LIFETIME_END(OuterObject, InstanceId) \
+	FObjectTrace::OutputInstanceLifetimeEnd(OuterObject, InstanceId);
+
 #define TRACE_PAWN_POSSESS(Controller, Pawn)\
 	FObjectTrace::OutputPawnPossess(Controller, Pawn);
 
-#endif
+#endif // TRACE_FILTERING_ENABLED
 	
 #define UNCONDITIONAL_TRACE_OBJECT_EVENT(Object, Event) \
 	FObjectTrace::OutputObjectEvent(Object, TEXT(#Event));
@@ -188,12 +191,15 @@ struct FObjectTrace
 	FObjectTrace::OutputWorld(World);
 
 
-#else
+#else // OBJECT_TRACE_ENABLED
 
 struct FObjectTrace
 {
 	/** Helper function to get the UObject from an ObjectId, if the UObject still exists */
-	static UObject* GetObjectFromId(uint64 id) { return nullptr; }
+	static UObject* GetObjectFromId(uint64 id)
+	{
+		return nullptr;
+	}
 };
 
 #define TRACE_CLASS(...)
@@ -215,4 +221,4 @@ struct FObjectTrace
 #define TRACE_CHILD_ELEMENT_LIFETIME_END(...)
 #define TRACE_CHILD_ELEMENT_EVENT(...)
 
-#endif
+#endif // OBJECT_TRACE_ENABLED
diff --git a/Engine/Source/Runtime/Engine/Public/ObjectTraceDefines.h b/Engine/Source/Runtime/Engine/Public/ObjectTraceDefines.h
--- a/Engine/Source/Runtime/Engine/Public/ObjectTraceDefines.h
+++ b/Engine/Source/Runtime/Engine/Public/ObjectTraceDefines.h
@@ -5,11 +5,5 @@
 #include "Trace/Config.h"
 
 #ifndef OBJECT_TRACE_ENABLED
-
-#if UE_TRACE_ENABLED && !IS_PROGRAM && !UE_BUILD_SHIPPING
-#define OBJECT_TRACE_ENABLED 1
-#else
-#define OBJECT_TRACE_ENABLED 0
-#endif
-
+#define OBJECT_TRACE_ENABLED ((UE_TRACE_ENABLED && !IS_PROGRAM && !UE_BUILD_SHIPPING) || UE_TRACE_MINIMAL_ENABLED)
 #endif
diff --git a/Engine/Source/Runtime/Engine/Public/TraceFilter.h b/Engine/Source/Runtime/Engine/Public/TraceFilter.h
--- a/Engine/Source/Runtime/Engine/Public/TraceFilter.h
+++ b/Engine/Source/Runtime/Engine/Public/TraceFilter.h
@@ -2,22 +2,21 @@
 
 #pragma once
 
+#if UE_ENABLE_INCLUDE_ORDER_DEPRECATED_IN_5_8
 #include "CoreTypes.h"
-#include "Templates/EnableIf.h"
-#include "Templates/PointerIsConvertibleFromTo.h"
-#include "GameFramework/Actor.h"
-#include "Engine/World.h"
+#endif // UE_ENABLE_INCLUDE_ORDER_DEPRECATED_IN_5_8
+
+#include "Animation/AnimInstance.h"
 #include "Components/ActorComponent.h"
 #include "Components/SkeletalMeshComponent.h" // Needed because AnimInstance::GetOwningComponent. (this only fails on clang compilers)
-#include "Animation/AnimInstance.h"
+#include "Engine/World.h"
+#include "GameFramework/Actor.h"
+#include "Templates/EnableIf.h"
+#include "Templates/PointerIsConvertibleFromTo.h"
 #include "Templates/UnrealTemplate.h"
 #include "UObject/Object.h"
 
-#if UE_TRACE_FORCE_ENABLED || (UE_TRACE_ENABLED && !IS_PROGRAM && !UE_BUILD_SHIPPING)
-#define TRACE_FILTERING_ENABLED 1
-#else
-#define TRACE_FILTERING_ENABLED 0
-#endif
+#define TRACE_FILTERING_ENABLED (UE_TRACE_FORCE_ENABLED || UE_TRACE_MINIMAL_ENABLED || (UE_TRACE_ENABLED && !IS_PROGRAM && !UE_BUILD_SHIPPING))
 
 #if TRACE_FILTERING_ENABLED
 
@@ -25,49 +24,51 @@ struct FTraceFilter
 {
 	struct TObjectFilter
 	{
-	public:
 		template<typename T>
-		static typename TEnableIf<TPointerIsConvertibleFromTo<typename TRemovePointer<T>::Type, UWorld>::Value, bool>::Type FORCEINLINE CanTrace(const T* Object)
+		static TEnableIf<TPointerIsConvertibleFromTo<typename TRemovePointer<T>::Type, UWorld>::Value, bool>::Type FORCEINLINE CanTrace(const T* Object)
 		{
 			const UWorld* World = (const UWorld*)Object;
-			return FTraceFilter::IsObjectTraceable(World);
+			return IsObjectTraceable(World);
 		}
 
 		template<typename T>
-		static typename TEnableIf<TPointerIsConvertibleFromTo<typename TRemovePointer<T>::Type, AActor>::Value, bool>::Type FORCEINLINE CanTrace(const T* Object)
+		static TEnableIf<TPointerIsConvertibleFromTo<typename TRemovePointer<T>::Type, AActor>::Value, bool>::Type FORCEINLINE CanTrace(const T* Object)
 		{
 			/** For an individual Actor, we expect it and the owning world to both be marked traceable */
 			const AActor* Actor = (const AActor*)Object;
-			return (FTraceFilter::IsObjectTraceable(Cast<UObject>(Actor)) && (Actor && CanTrace(Actor->GetWorld())));
+			return (IsObjectTraceable(Cast<UObject>(Actor)) && (Actor && CanTrace(Actor->GetWorld())));
 		}
 
 		template<typename T>
-		static typename TEnableIf<TPointerIsConvertibleFromTo<typename TRemovePointer<T>::Type, UActorComponent>::Value, bool>::Type FORCEINLINE CanTrace(const T* Object)
+		static TEnableIf<TPointerIsConvertibleFromTo<typename TRemovePointer<T>::Type, UActorComponent>::Value, bool>::Type FORCEINLINE CanTrace(const T* Object)
 		{
 			/** For an individual Actor Component, we expect it or the owning actor to be marked traceable */
 			const UActorComponent* ActorComponent = (const UActorComponent*)Object;
-			return FTraceFilter::IsObjectTraceable(ActorComponent) || (ActorComponent && TObjectFilter::CanTrace(ActorComponent->GetOwner()));
+			return IsObjectTraceable(ActorComponent) || (ActorComponent && CanTrace(ActorComponent->GetOwner()));
 		}
 
 		template<typename T>
-		static typename TEnableIf<TPointerIsConvertibleFromTo<typename TRemovePointer<T>::Type, UAnimInstance>::Value, bool>::Type FORCEINLINE CanTrace(const T* Object)
+		static TEnableIf<TPointerIsConvertibleFromTo<typename TRemovePointer<T>::Type, UAnimInstance>::Value, bool>::Type FORCEINLINE CanTrace(const T* Object)
 		{
 			/** For an AnimInstance object, we expect it or the owning component to be marked traceable */
 			const UAnimInstance* AnimInstance = (const UAnimInstance*)Object;
-			return FTraceFilter::IsObjectTraceable(AnimInstance) || (AnimInstance && TObjectFilter::CanTrace(AnimInstance->GetOwningComponent()));
+			return IsObjectTraceable(AnimInstance) || (AnimInstance && CanTrace(AnimInstance->GetOwningComponent()));
 		}
 
 		template<typename T>
-		static typename TEnableIf<(!TPointerIsConvertibleFromTo<typename TRemovePointer<T>::Type, UActorComponent>::Value && !TPointerIsConvertibleFromTo<typename TRemovePointer<T>::Type, AActor>::Value
-			&& !TPointerIsConvertibleFromTo<typename TRemovePointer<T>::Type, UWorld>::Value && !TPointerIsConvertibleFromTo<typename TRemovePointer<T>::Type, UAnimInstance>::Value && TPointerIsConvertibleFromTo<typename TRemovePointer<T>::Type, UObject>::Value), bool>::Type FORCEINLINE CanTrace(const T* Object)
+		static TEnableIf<(!TPointerIsConvertibleFromTo<typename TRemovePointer<T>::Type, UActorComponent>::Value
+			&& !TPointerIsConvertibleFromTo<typename TRemovePointer<T>::Type, AActor>::Value
+			&& !TPointerIsConvertibleFromTo<typename TRemovePointer<T>::Type, UWorld>::Value
+			&& !TPointerIsConvertibleFromTo<typename TRemovePointer<T>::Type, UAnimInstance>::Value
+			&& TPointerIsConvertibleFromTo<typename TRemovePointer<T>::Type, UObject>::Value), bool>::Type FORCEINLINE CanTrace(const T* Object)
 		{
 			/** For an individual UObject, we expect it or the owning world to be marked traceable */
 			const UObject* BaseObject = (const UObject*)Object;
-			return FTraceFilter::IsObjectTraceable(Object) || (BaseObject && FTraceFilter::IsObjectTraceable(BaseObject->GetWorld()));
+			return FTraceFilter::IsObjectTraceable(Object) || (BaseObject && IsObjectTraceable(BaseObject->GetWorld()));
 		}
 	};
 
-	/** Set whether or not an Object is Traceable, and eligible for outputting Trace Data */
+	/** Set whether an Object is Traceable, and eligible for outputting Trace Data */
 	template<bool bForceThreadSafe = true>
 	static void SetObjectIsTraceable(const UObject* InObject, bool bIsTraceable);
 
@@ -81,7 +82,7 @@ struct FTraceFilter
 	/** Destroys any of the Filters which are part of the Engine */
 	static ENGINE_API void Destroy();
 
-	/** Returns whether or not an object is eligible to be outputted (output trace data) */
+	/** Returns whether an object is eligible to be outputted (output trace data) */
 	template<bool bForceThreadSafe = true>
 	static bool IsObjectTraceable(const UObject* InObject);
 
@@ -119,9 +120,9 @@ extern template ENGINE_API bool FTraceFilter::IsObjectTraceable<false>(const UOb
 
 #define MARK_OBJECT_TRACEABLE(Object)
 
-#define SET_OBJECT_TRACEABLE(Object, bIsTraceable) 
+#define SET_OBJECT_TRACEABLE(Object, bIsTraceable)
 
 #define GET_TRACE_OBJECT_VALUE(Object) \
 	false
 
-#endif
+#endif // TRACE_FILTERING_ENABLED
