Fix issue with skeletal mesh component curve propagation to followers where curves are only propagated when there are both material or morph curves, not just one type.

diff --git a/Engine/Source/Runtime/Engine/Private/Components/SkeletalMeshComponent.cpp b/Engine/Source/Runtime/Engine/Private/Components/SkeletalMeshComponent.cpp
--- a/Engine/Source/Runtime/Engine/Private/Components/SkeletalMeshComponent.cpp
+++ b/Engine/Source/Runtime/Engine/Private/Components/SkeletalMeshComponent.cpp
@@ -3210,7 +3210,7 @@ void USkeletalMeshComponent::ApplyAnimationCurvesToComponent(const TMap<FName, f
 	}
 
 	/** Push through curves to follower components */
-	if (bPropagateCurvesToFollowers && bContainsMorphCurves && bContainsMaterialCurves && FollowerPoseComponents.Num() > 0)
+	if (bPropagateCurvesToFollowers && (bContainsMorphCurves || bContainsMaterialCurves) && FollowerPoseComponents.Num() > 0)
 	{
 		for (TWeakObjectPtr<USkinnedMeshComponent> MeshComponent : FollowerPoseComponents)
 		{
