Fixing compile errors and crashes when UE_WITH_REMOTE_OBJECT_HANDLE=1 in configs without AutoRTFM enabled

diff --git a/Engine/Source/Runtime/CoreUObject/Private/UObject/RemoteExecutor.cpp b/Engine/Source/Runtime/CoreUObject/Private/UObject/RemoteExecutor.cpp
--- a/Engine/Source/Runtime/CoreUObject/Private/UObject/RemoteExecutor.cpp
+++ b/Engine/Source/Runtime/CoreUObject/Private/UObject/RemoteExecutor.cpp
@@ -1045,7 +1045,7 @@ void ExecuteTransactional(FName WorkName, const TFunctionRef<void(void)>& Work)
 
 void EnqueueWorkWithExplicitPriority(FName WorkName, FRemoteWorkPriority WorkPriority, bool bIsTransactional, const TFunction<void(void)>& InWork)
 {
-#if UE_WITH_REMOTE_OBJECT_HANDLE
+#if UE_WITH_REMOTE_OBJECT_HANDLE && UE_AUTORTFM
 	GRemoteExecutor.EnqueueWork(WorkName, WorkPriority, bIsTransactional, INDEX_NONE, InWork);
 #else
 	InWork();
diff --git a/Engine/Source/Runtime/CoreUObject/Private/UObject/RemoteObjectTransfer.cpp b/Engine/Source/Runtime/CoreUObject/Private/UObject/RemoteObjectTransfer.cpp
--- a/Engine/Source/Runtime/CoreUObject/Private/UObject/RemoteObjectTransfer.cpp
+++ b/Engine/Source/Runtime/CoreUObject/Private/UObject/RemoteObjectTransfer.cpp
@@ -968,6 +968,16 @@ void MigrateObjectFromRemoteServer(FRemoteObjectId ObjectId, FRemoteServerId Cur
 		*ObjectId.ToString(),
 		*CurrentOwnerServerId.ToString());
 
+#if !UE_AUTORTFM
+	{
+		FRemoteServerId LastKnownResidentServerId;
+		if (UE::RemoteObject::Handle::FRemoteObjectStub* Stub = UE::RemoteObject::Private::FindRemoteObjectStub(ObjectId))
+		{
+			LastKnownResidentServerId = Stub->ResidentServerId;
+		}
+		GTransferQueue.SendRemoteObjectRequest(FRemoteWorkPriority(), ObjectId, LastKnownResidentServerId, FRemoteServerId::GetLocalServerId());
+	}
+#else
 	checkf(GTransferQueue.ActiveRequest, TEXT("Attempting to access remote object %s but we are outside of a transaction"), *ObjectId.ToString());
 
 	// ensure the request is added in the Open so after we abort it is preserved
@@ -1024,6 +1034,7 @@ void MigrateObjectFromRemoteServer(FRemoteObjectId ObjectId, FRemoteServerId Cur
 		FString::Format(TEXT("object {0} required from server {1}"),
 			{*ObjectId.ToString(),
 			*CurrentOwnerServerId.ToString()}));
+#endif // UE_AUTORTFM
 }
 
 void TouchResidentObject(UObject* Object)
diff --git a/Engine/Source/Runtime/Engine/Classes/Engine/NetDriver.h b/Engine/Source/Runtime/Engine/Classes/Engine/NetDriver.h
--- a/Engine/Source/Runtime/Engine/Classes/Engine/NetDriver.h
+++ b/Engine/Source/Runtime/Engine/Classes/Engine/NetDriver.h
@@ -2351,7 +2351,7 @@ public:
 		bReplicateTransactionally = bInReplicateTransactionally;
 	}
 
-#if UE_WITH_REMOTE_OBJECT_HANDLE
+#if UE_WITH_REMOTE_OBJECT_HANDLE && UE_AUTORTFM
 	/** Returns true if this driver should run certain operations via UE::RemoteExecutor::ExecuteTransactional */
 	bool ShouldReplicateTransactionally() const
 	{
diff --git a/Engine/Source/Runtime/Engine/Private/Collision/WorldCollisionAsync.cpp b/Engine/Source/Runtime/Engine/Private/Collision/WorldCollisionAsync.cpp
--- a/Engine/Source/Runtime/Engine/Private/Collision/WorldCollisionAsync.cpp
+++ b/Engine/Source/Runtime/Engine/Private/Collision/WorldCollisionAsync.cpp
@@ -793,7 +793,7 @@ void UWorld::FinishAsyncTrace()
 	ExecuteAsyncTraceIfAvailable<FTraceDatum>  (AsyncTraceState, true);
 	ExecuteAsyncTraceIfAvailable<FOverlapDatum>(AsyncTraceState, true);
 
-#if UE_WITH_REMOTE_OBJECT_HANDLE
+#if UE_WITH_REMOTE_OBJECT_HANDLE && UE_AUTORTFM
 	auto DoWorkCallback = [this]()
 	{
 		// We need to add all requests inside of a transaction, otherwise when a transaction is 
diff --git a/Engine/Source/Runtime/Engine/Private/DataChannel.cpp b/Engine/Source/Runtime/Engine/Private/DataChannel.cpp
--- a/Engine/Source/Runtime/Engine/Private/DataChannel.cpp
+++ b/Engine/Source/Runtime/Engine/Private/DataChannel.cpp
@@ -1995,7 +1995,7 @@ void UControlChannel::ReceivedBunch( FInBunch& Bunch )
 		else if (Connection->Driver->Notify != nullptr)
 		{
 			// Process control message on client/server connection
-#if UE_WITH_REMOTE_OBJECT_HANDLE
+#if UE_WITH_REMOTE_OBJECT_HANDLE && UE_AUTORTFM
 			if (Connection->Driver->ShouldReplicateTransactionally())
 			{
 				static FName TransactionalWorkName("UControlChannel::ReceivedBunch NotifyControlMessage");
@@ -3506,7 +3506,7 @@ void UActorChannel::ProcessBunchInternal(FInBunch & Bunch)
 
 void UActorChannel::ProcessBunch(FInBunch& Bunch)
 {
-#if UE_WITH_REMOTE_OBJECT_HANDLE
+#if UE_WITH_REMOTE_OBJECT_HANDLE && UE_AUTORTFM
 	if (Connection && Connection->Driver && Connection->Driver->ShouldReplicateTransactionally())
 	{
 		static FName TransactionalWorkName("UActorChannel::ProcessBunch");
diff --git a/Engine/Source/Runtime/Engine/Private/NetDriver.cpp b/Engine/Source/Runtime/Engine/Private/NetDriver.cpp
--- a/Engine/Source/Runtime/Engine/Private/NetDriver.cpp
+++ b/Engine/Source/Runtime/Engine/Private/NetDriver.cpp
@@ -1139,7 +1139,7 @@ void UNetDriver::TickFlush(float DeltaSeconds)
 		}
 		else if (ClientConnections.Num() > 0)
 		{
-#if UE_WITH_REMOTE_OBJECT_HANDLE
+#if UE_WITH_REMOTE_OBJECT_HANDLE && UE_AUTORTFM
 			if (bReplicateTransactionally && !UE::Net::Private::bServerReplicateActorsGranularTransactions)
 			{
 				static FName TransactionalWorkName("ServerReplicateActors");
@@ -6125,8 +6125,8 @@ int32 UNetDriver::ServerReplicateActors(float DeltaSeconds)
 
 	int32 Updated = 0;
 
-#if UE_WITH_REMOTE_OBJECT_HANDLE
-	const bool bUseGranularTransactions = (bReplicateTransactionally && UE::Net::Private::bServerReplicateActorsGranularTransactions);
+#if UE_WITH_REMOTE_OBJECT_HANDLE && UE_AUTORTFM
+	const bool bUseGranularTransactions = (ShouldReplicateTransactionally() && UE::Net::Private::bServerReplicateActorsGranularTransactions);
 #else
 	constexpr bool bUseGranularTransactions = false;
 #endif
diff --git a/Engine/Source/Runtime/Engine/Private/PhysicsEngine/Experimental/PhysScene_Chaos.cpp b/Engine/Source/Runtime/Engine/Private/PhysicsEngine/Experimental/PhysScene_Chaos.cpp
--- a/Engine/Source/Runtime/Engine/Private/PhysicsEngine/Experimental/PhysScene_Chaos.cpp
+++ b/Engine/Source/Runtime/Engine/Private/PhysicsEngine/Experimental/PhysScene_Chaos.cpp
@@ -2551,7 +2551,7 @@ void FPhysScene_Chaos::OnSyncBodies(Chaos::FPhysicsSolverBase* Solver)
 		});
 	};
 
-#if UE_WITH_REMOTE_OBJECT_HANDLE
+#if UE_WITH_REMOTE_OBJECT_HANDLE && UE_AUTORTFM
 	static FName TransactionalDispatcherWorkName("FPhysScene_Chaos::OnSyncBodies-Dispatcher");
 	UE::RemoteExecutor::ExecuteTransactional(TransactionalDispatcherWorkName, ExecuteCallback);
 #else
@@ -2586,7 +2586,7 @@ void FPhysScene_Chaos::OnSyncBodies(Chaos::FPhysicsSolverBase* Solver)
 			}
 		};
 
-#if UE_WITH_REMOTE_OBJECT_HANDLE
+#if UE_WITH_REMOTE_OBJECT_HANDLE && UE_AUTORTFM
 		static FName TransactionalWorkName("FPhysScene_Chaos::OnSyncBodies");
 		UE::RemoteExecutor::ExecuteTransactional(TransactionalWorkName, DoWorkCallback);
 #else
diff --git a/Engine/Source/Runtime/Engine/Private/PlayerController.cpp b/Engine/Source/Runtime/Engine/Private/PlayerController.cpp
--- a/Engine/Source/Runtime/Engine/Private/PlayerController.cpp
+++ b/Engine/Source/Runtime/Engine/Private/PlayerController.cpp
@@ -538,7 +538,7 @@ FString APlayerController::ConsoleCommand(const FString& Cmd, bool bWriteToLog)
 	if (Player != nullptr)
 	{
 		FString Result;
-#if UE_WITH_REMOTE_OBJECT_HANDLE
+#if UE_WITH_REMOTE_OBJECT_HANDLE && UE_AUTORTFM
 		if (!AutoRTFM::IsTransactional())
 		{
 			static FName TransactionalWorkName("APlayerController::ConsoleComand");
diff --git a/Engine/Source/Runtime/Engine/Private/TickTaskManager.cpp b/Engine/Source/Runtime/Engine/Private/TickTaskManager.cpp
--- a/Engine/Source/Runtime/Engine/Private/TickTaskManager.cpp
+++ b/Engine/Source/Runtime/Engine/Private/TickTaskManager.cpp
@@ -324,7 +324,7 @@ public:
 #endif
 			LIGHTWEIGHT_TIME_GUARD_BEGIN(FTickFunctionTask, GTimeguardThresholdMS);
 
-#if UE_WITH_REMOTE_OBJECT_HANDLE
+#if UE_WITH_REMOTE_OBJECT_HANDLE && UE_AUTORTFM
 			auto ExecuteTickWork = [this, CurrentThread, &MyCompletionGraphEvent]()
 			{
 				// !IsCompletionHandleValid is an indication we had previously been ticked this frame and then migrated back
@@ -332,7 +332,7 @@ public:
 				{
 #endif
 				Target->ExecuteTick(Target->CalculateDeltaTime(Context.DeltaSeconds, Context.World), Context.TickType, CurrentThread, MyCompletionGraphEvent);
-#if UE_WITH_REMOTE_OBJECT_HANDLE
+#if UE_WITH_REMOTE_OBJECT_HANDLE && UE_AUTORTFM
 				}
 			};
 
diff --git a/Engine/Source/Runtime/Engine/Private/TimerManager.cpp b/Engine/Source/Runtime/Engine/Private/TimerManager.cpp
--- a/Engine/Source/Runtime/Engine/Private/TimerManager.cpp
+++ b/Engine/Source/Runtime/Engine/Private/TimerManager.cpp
@@ -344,7 +344,7 @@ void FTimerUnifiedDelegate::Execute() const
 		{
 			const FTimerDelegate& FuncDelegate = VariantDelegate.Get<FTimerDelegate>();
 
-#if UE_WITH_REMOTE_OBJECT_HANDLE
+#if UE_WITH_REMOTE_OBJECT_HANDLE && UE_AUTORTFM
 			const FName WorkName{ TEXT("TimerManager::Execute FTimerDelegate") };
 			UE::RemoteExecutor::ExecuteTransactional(WorkName, [&FuncDelegate]()
 			{
@@ -354,7 +354,7 @@ void FTimerUnifiedDelegate::Execute() const
 					FScopeCycleCounterUObject Context(FuncDelegate.GetUObject());
 					FuncDelegate.Execute();
 				}
-#if UE_WITH_REMOTE_OBJECT_HANDLE
+#if UE_WITH_REMOTE_OBJECT_HANDLE && UE_AUTORTFM
 			});
 #endif
 			break;
@@ -362,7 +362,7 @@ void FTimerUnifiedDelegate::Execute() const
 	case FTimerDelegateVariant::IndexOfType<FTimerDynamicDelegate>():
 		{
 			const FTimerDynamicDelegate& FuncDynDelegate = VariantDelegate.Get<FTimerDynamicDelegate>();
-#if UE_WITH_REMOTE_OBJECT_HANDLE
+#if UE_WITH_REMOTE_OBJECT_HANDLE && UE_AUTORTFM
 			const FName WorkName { TEXT("TimerManager::Execute FTimerDynamicDelegate") };
 			UE::RemoteExecutor::ExecuteTransactional(WorkName, [&FuncDynDelegate]()
 			{
@@ -372,7 +372,7 @@ void FTimerUnifiedDelegate::Execute() const
 					// stat scope is handled by UObject::ProcessEvent for the UFunction.
 					FuncDynDelegate.ProcessDelegate<UObject>(nullptr);
 				}
-#if UE_WITH_REMOTE_OBJECT_HANDLE
+#if UE_WITH_REMOTE_OBJECT_HANDLE && UE_AUTORTFM
 			});
 #endif
 			break;
diff --git a/Engine/Source/Runtime/Engine/Private/World.cpp b/Engine/Source/Runtime/Engine/Private/World.cpp
--- a/Engine/Source/Runtime/Engine/Private/World.cpp
+++ b/Engine/Source/Runtime/Engine/Private/World.cpp
@@ -4812,7 +4812,7 @@ void UWorld::InternalUpdateStreamingState()
 	// Update streaming levels state using streaming volumes.
 	// Issues level streaming load/unload requests based on local players being inside/outside level streaming volumes.
 
-#if UE_WITH_REMOTE_OBJECT_HANDLE
+#if UE_WITH_REMOTE_OBJECT_HANDLE && UE_AUTORTFM
 	// Only process transactionally if the world has begun play. This allows time for the transactional infrastructure to be set up.
 	if (HasBegunPlay())
 	{
