Skeletal Mesh Render - Deprecated FSkelMeshRenderSection::ClothMappingData_DEPRECATED so that the array can be removed since it doesn't contribute in serialization and is no longer used.

#rnx

diff --git a/Engine/Source/Runtime/Engine/Public/Rendering/SkeletalMeshLODRenderData.h b/Engine/Source/Runtime/Engine/Public/Rendering/SkeletalMeshLODRenderData.h
--- a/Engine/Source/Runtime/Engine/Public/Rendering/SkeletalMeshLODRenderData.h
+++ b/Engine/Source/Runtime/Engine/Public/Rendering/SkeletalMeshLODRenderData.h
@@ -50,7 +50,7 @@ struct FSkelMeshRenderSection
 	/** The offset into the LOD's vertex buffer of this section's vertices. */
 	uint32 BaseVertexIndex;
 
-	/** The extra vertex data for mapping to an APEX clothing simulation mesh. */
+	UE_DEPRECATED(5.8, "Use ClothMappingDataLODs instead.")
 	TArray<FMeshToMeshVertData> ClothMappingData_DEPRECATED;
 
 	/**
