Add a PreferredReplicationSystem setting to FIrisNetDriverConfig that can override the net.Iris.UseIrisReplication cvar for specific NetDriverDefinitions.

diff --git a/Engine/Source/Runtime/Engine/Classes/Engine/Engine.h b/Engine/Source/Runtime/Engine/Classes/Engine/Engine.h
--- a/Engine/Source/Runtime/Engine/Classes/Engine/Engine.h
+++ b/Engine/Source/Runtime/Engine/Classes/Engine/Engine.h
@@ -274,6 +274,13 @@ struct FIrisNetDriverConfig
 	 */
 	UPROPERTY()
 	bool bCanUseIris = false;
+
+	/**
+	 * Sets the preferred replication system for this NetDriver. This is the one that will be used unless it's overridden
+	 * by bCanUseIris, UGameInstance::GetDesiredReplicationSystem, or AGameModeBase::GameNetDriverReplicationSystem.
+	 */
+	UPROPERTY()
+	EReplicationSystem PreferredReplicationSystem = EReplicationSystem::Default;
 };
 
 
diff --git a/Engine/Source/Runtime/Engine/Private/UnrealEngine.cpp b/Engine/Source/Runtime/Engine/Private/UnrealEngine.cpp
--- a/Engine/Source/Runtime/Engine/Private/UnrealEngine.cpp
+++ b/Engine/Source/Runtime/Engine/Private/UnrealEngine.cpp
@@ -14505,8 +14505,24 @@ bool UEngine::WillNetDriverUseIris(const FWorldContext& Context, FName InNetDriv
 	const bool bConfigCanUseIris = IrisConfig && IrisConfig->bCanUseIris;
 	const bool bIsEngineDefaultIris = UE::Net::ShouldUseIrisReplication();
 
+	// Default preference: use Iris if the config can & the engine cvar is enabled.
 	bool bUseIrisRepSystem = bConfigCanUseIris && bIsEngineDefaultIris;
 
+	if (IrisConfig)
+	{
+		if (IrisConfig->PreferredReplicationSystem == EReplicationSystem::Iris)
+		{
+			// If Iris is preferred and can be used, use it (unless overridden below)
+			UE_CLOG(!bConfigCanUseIris, LogNet, Log, TEXT("NetDriver %s (NetDefinition %s) is configured to prefer the Iris replication system, but bCanUseIris is false."), *InNetDriverName.ToString(), *InNetDriverDefinition.ToString());
+			bUseIrisRepSystem = bConfigCanUseIris;
+		}
+		else if (IrisConfig->PreferredReplicationSystem == EReplicationSystem::Generic)
+		{
+			// If Generic is preferred, use it (unless overridden below)
+			bUseIrisRepSystem = false;
+		}
+	}
+
 	if (UGameInstance* ContextGameInstance = Context.OwningGameInstance)
 	{
 		EReplicationSystem GameInstanceDesiredRepSystem = ContextGameInstance->GetDesiredReplicationSystem(InNetDriverDefinition);
