Preshader2 static analysis third attempt at warning suppression, now in the INL file.

#rnx

diff --git a/Engine/Source/Runtime/Engine/Private/Shader/Preshader2.cpp b/Engine/Source/Runtime/Engine/Private/Shader/Preshader2.cpp
--- a/Engine/Source/Runtime/Engine/Private/Shader/Preshader2.cpp
+++ b/Engine/Source/Runtime/Engine/Private/Shader/Preshader2.cpp
@@ -851,8 +851,6 @@ void OpEvaluate(const FOpSparseVolumeTextureUniform&, const FMaterialRenderConte
 // InPlace:				Opcode(1)  ResultOffset(2)  Dimension+Type(1)				AOffset same as ResultOffset
 // Default:				Opcode(1)  ResultOffset(2)  AOffset(2)  Dimension+Type(1)
 // Does not sanitize Dimension in case custom implementation stores extra data in that byte!  Client evaluate function should remove junk high bits from Dimension.
-
-//-V:PRESHADER2_UNARY_OP:521
 #define PRESHADER2_UNARY_OP(OP_NAME) \
 	case EPreshader2Opcode::OP_NAME:																																	\
 		OpcodePtr += 2;									/* Skip to AOffset */																							\
@@ -871,8 +869,6 @@ void OpEvaluate(const FOpSparseVolumeTextureUniform&, const FMaterialRenderConte
 // "Reuse" variations, so it's pulled from the opcode.  We subtract the opcode from the first "Cast_Reuse" opcode to extract the type and add it to the Dimension,
 // where it is stored in the other versions of the opcode.  To reduce code size bloat, we technically only need a case for the last version of the opcode, so we
 // fall through in the prior cases (i.e., actual code is only generated for Cast_ReuseLWC, and the other two opcodes only produce the "case" statement).
-
-//-V:PRESHADER2_UNARY_OP_CAST_REUSE:501
 #define PRESHADER2_UNARY_OP_CAST_REUSE(OP_NAME) \
 	case EPreshader2Opcode::OP_NAME:																																			\
 		if constexpr (EPreshader2Opcode::OP_NAME == EPreshader2Opcode::Cast_ReuseLWC)																							\
diff --git a/Engine/Source/Runtime/Engine/Public/Shader/Preshader2Opcodes.inl b/Engine/Source/Runtime/Engine/Public/Shader/Preshader2Opcodes.inl
--- a/Engine/Source/Runtime/Engine/Public/Shader/Preshader2Opcodes.inl
+++ b/Engine/Source/Runtime/Engine/Public/Shader/Preshader2Opcodes.inl
@@ -1,5 +1,8 @@
 // Copyright Epic Games, Inc. All Rights Reserved.
 
+//-V501
+//-V521
+
 // Unary ops with integer support -- each has 3 opcode variants, with Cast having 3 extra custom "Reuse" opcode variants for different destination types.
 // The three extra opcodes conveniently maintain the 3-element stride for surrounding standard unary opcode variants.
 PRESHADER2_UNARY_OP(Abs)
@@ -68,4 +71,4 @@ PRESHADER2_OP(SparseVolumeTextureUniform)
 
 // NOTE:  Clamp op is not implemented, caller should generate Min+Max.  Challenge of implementing Clamp as a single op is that there are many
 // permutations of register and constant input arguments that need to be dealt with, so it's easier to split in two.  It should be plenty
-// efficient with reuse of previous op's result, taking only one extra byte in storage, and one extra op.  Clamp is 0.01% of ops.
\ No newline at end of file
+// efficient with reuse of previous op's result, taking only one extra byte in storage, and one extra op.  Clamp is 0.01% of ops.
