Improve tooltip to represent all parameter override sources.
Additionally, clean up the code to remove the unnecessary copies of the AssetPath everywhere FMaterialParameterMetadata is used.

[RN] minor | Rendering Architecture

diff --git a/Engine/Source/Editor/MaterialEditor/Private/MaterialPropertyHelpers.cpp b/Engine/Source/Editor/MaterialEditor/Private/MaterialPropertyHelpers.cpp
--- a/Engine/Source/Editor/MaterialEditor/Private/MaterialPropertyHelpers.cpp
+++ b/Engine/Source/Editor/MaterialEditor/Private/MaterialPropertyHelpers.cpp
@@ -57,6 +57,7 @@
 #include "Factories/MaterialFunctionMaterialLayerBlendFactory.h"
 #include "Curves/CurveLinearColorAtlas.h"
 #include "Curves/CurveLinearColor.h"
+#include "MaterialCachedData.h"
 
 #include UE_INLINE_GENERATED_CPP_BY_NAME(MaterialPropertyHelpers)
 
@@ -1153,11 +1154,29 @@ FText FMaterialPropertyHelpers::GetParameterExpressionDescription(UDEditorParame
 
 FText FMaterialPropertyHelpers::GetParameterTooltip(UDEditorParameterValue* Parameter, UObject* MaterialEditorInstance)
 {
-	const FText AssetPath = FText::FromString(Parameter->AssetPath);
-	FText TooltipText;
+	return GetParameterTooltip(Parameter, static_cast<UMaterialEditorInstanceConstant*>(MaterialEditorInstance));
+}
+
+FText FMaterialPropertyHelpers::GetParameterTooltip(UDEditorParameterValue* Parameter, UMaterialEditorInstanceConstant* MaterialEditorInstance)
+{
+	const FMaterialCachedExpressionData* CachedExpressionData = MaterialEditorInstance->Parent ? &MaterialEditorInstance->Parent->GetCachedExpressionData() : nullptr;
 
+	FText AssetPath;
+	if (CachedExpressionData)
+	{
+		const TSharedPtr<FMaterialCachedExpressionEditorOnlyData>& EditorOnlyData = CachedExpressionData->EditorOnlyData;
+		if (EditorOnlyData)
+		{
+			if (const FParameterAssetPathsEntry* Entry = EditorOnlyData->AssetPathsMap.Find(Parameter->ParameterInfo.Name))
+			{
+				AssetPath = FText::FromString(Entry->ToString());
+			}
+		}
+	}
+
+	FText TooltipText;
 	// If the material instance disallows the creation of new shader permutations, prevent overriding the static parameter.
-	if (Parameter->IsStaticParameter() && static_cast<UMaterialEditorInstanceConstant*>(MaterialEditorInstance)->SourceInstance->bDisallowStaticParameterPermutations)
+	if (Parameter->IsStaticParameter() && MaterialEditorInstance->SourceInstance->bDisallowStaticParameterPermutations)
 	{
 		return FText::FromString(TEXT("This material instance parent restricts the creation of new shader permutations. Overriding this parameter would result in the generation of additional shader permutations."));
 	}
@@ -1169,6 +1188,7 @@ FText FMaterialPropertyHelpers::GetParameterTooltip(UDEditorParameterValue* Para
 	{
 		TooltipText = FText::Format(LOCTEXT("ParameterInfoLocationOnly", "Found in: {0}"), AssetPath);
 	}
+
 	return TooltipText;
 }
 
diff --git a/Engine/Source/Editor/MaterialEditor/Public/MaterialPropertyHelpers.h b/Engine/Source/Editor/MaterialEditor/Public/MaterialPropertyHelpers.h
--- a/Engine/Source/Editor/MaterialEditor/Public/MaterialPropertyHelpers.h
+++ b/Engine/Source/Editor/MaterialEditor/Public/MaterialPropertyHelpers.h
@@ -208,6 +208,9 @@ public:
 	/** Gets the expression description of this parameter from the base material */
 	static	UE_API FText GetParameterExpressionDescription(UDEditorParameterValue* Parameter, UObject* MaterialEditorInstance);
 	
+	static UE_API FText GetParameterTooltip(UDEditorParameterValue* Parameter, UMaterialEditorInstanceConstant* MaterialEditorInstance);
+
+	UE_DEPRECATED(5.8, "Please use the fully typed method instead.")
 	static UE_API FText GetParameterTooltip(UDEditorParameterValue* Parameter, UObject* MaterialEditorInstance);
 	/**
 	 * Called when a parameter is overridden;
diff --git a/Engine/Source/Editor/UnrealEd/Classes/MaterialEditor/DEditorParameterValue.h b/Engine/Source/Editor/UnrealEd/Classes/MaterialEditor/DEditorParameterValue.h
--- a/Engine/Source/Editor/UnrealEd/Classes/MaterialEditor/DEditorParameterValue.h
+++ b/Engine/Source/Editor/UnrealEd/Classes/MaterialEditor/DEditorParameterValue.h
@@ -33,9 +33,6 @@ class UDEditorParameterValue : public UObject
 	UPROPERTY()
 	FString Description;
 
-	UPROPERTY()
-	FString AssetPath;
-
 #if WITH_EDITORONLY_DATA
 	/** Controls where this parameter is displayed in a material instance parameter list.  The lower the number the higher up in the parameter list. */
 	UPROPERTY()
@@ -48,7 +45,6 @@ class UDEditorParameterValue : public UObject
 	virtual bool GetValue(FMaterialParameterMetadata& OutResult) const
 	{
 		OutResult.Description = Description;
-		OutResult.AssetPath = AssetPath;
 		OutResult.ExpressionGuid = ExpressionId;
 		OutResult.SortPriority = SortPriority;
 		OutResult.bOverride = bOverride;
diff --git a/Engine/Source/Editor/UnrealEd/Private/EditorParameters.cpp b/Engine/Source/Editor/UnrealEd/Private/EditorParameters.cpp
--- a/Engine/Source/Editor/UnrealEd/Private/EditorParameters.cpp
+++ b/Engine/Source/Editor/UnrealEd/Private/EditorParameters.cpp
@@ -241,7 +241,6 @@ UDEditorParameterValue* UDEditorParameterValue::Create(UObject* Owner,
 	Parameter->ParameterInfo = ParameterInfo;
 	Parameter->ExpressionId = Meta.ExpressionGuid;
 	Parameter->Description = Meta.Description;
-	Parameter->AssetPath = Meta.AssetPath;
 	Parameter->SortPriority = Meta.SortPriority;
 	Parameter->bOverride = Meta.bOverride;
 	return Parameter;
diff --git a/Engine/Source/Runtime/Engine/Private/Materials/MaterialCachedData.cpp b/Engine/Source/Runtime/Engine/Private/Materials/MaterialCachedData.cpp
--- a/Engine/Source/Runtime/Engine/Private/Materials/MaterialCachedData.cpp
+++ b/Engine/Source/Runtime/Engine/Private/Materials/MaterialCachedData.cpp
@@ -62,6 +62,27 @@ const FMaterialCachedExpressionEditorOnlyData FMaterialCachedExpressionEditorOnl
 
 static_assert((uint64)(EMaterialProperty::MP_MaterialAttributes)-1 < (8 * sizeof(FMaterialCachedExpressionData::PropertyConnectedMask)), "PropertyConnectedMask cannot contain entire EMaterialProperty enumeration.");
 
+FString FParameterAssetPathsEntry::ToString() const
+{
+	FString Result;
+	for (const auto& Element : AssetPaths)
+	{
+		const FString& AssetPath = Element.Key;
+		const int32 PathCounter = Element.Value;
+
+		if (PathCounter > 1)
+		{
+			Result += FString::Printf(TEXT("\n %dx %s"), PathCounter, *AssetPath);
+		}
+		else
+		{
+			Result += FString::Printf(TEXT("\n %s"), *AssetPath);
+		}
+	}
+
+	return Result;
+}
+
 FMaterialCachedExpressionData::FMaterialCachedExpressionData()
 	: FunctionInfosStateCRC(0xffffffff)
 	, bHasMaterialLayers(false)
@@ -192,16 +213,20 @@ static bool TryAddParameter(FMaterialCachedExpressionData& CachedData,
 	return false;
 }
 
-bool FMaterialCachedExpressionData::AddParameter(const FMaterialParameterInfo& ParameterInfo, const FMaterialParameterMetadata& ParameterMeta, UObject*& OutReferencedTexture, UTextureCollection*& OutReferencedTextureCollection, FText& OutErrorMessage)
+bool FMaterialCachedExpressionData::AddParameter(const FMaterialParameterInfo& ParameterInfo, const FMaterialParameterMetadata& ParameterMeta, const FString& InAssetPath, UObject*& OutReferencedTexture, UTextureCollection*& OutReferencedTextureCollection, FText& OutErrorMessage)
 {
 	check(EditorOnlyData);
-	int32 AssetIndex = INDEX_NONE;
-	if (!ParameterMeta.AssetPath.IsEmpty())
+
+	if (!InAssetPath.IsEmpty())
 	{
-		AssetIndex = EditorOnlyData->AssetPaths.AddUnique(ParameterMeta.AssetPath);
+		FParameterAssetPathsEntry& Entry = EditorOnlyData->AssetPathsMap.FindOrAdd(ParameterInfo.Name);
+		TMap<FString, int32>& AssetPaths = Entry.AssetPaths;
+
+		int32& PathsCounter = AssetPaths.FindOrAdd(InAssetPath);
+		++PathsCounter;
 	}
 
-	const FMaterialCachedParameterEditorInfo EditorInfo(ParameterMeta.ExpressionGuid, ParameterMeta.Description, ParameterMeta.Group, ParameterMeta.SortPriority, AssetIndex);
+	const FMaterialCachedParameterEditorInfo EditorInfo(ParameterMeta.ExpressionGuid, ParameterMeta.Description, ParameterMeta.Group, ParameterMeta.SortPriority);
 	int32 Index = INDEX_NONE;
 	TOptional<FMaterialCachedParameterEditorInfo> PreviousEditorInfo;
 	if (TryAddParameter(*this, ParameterMeta.Value.Type, ParameterInfo, EditorInfo, Index, PreviousEditorInfo))
@@ -288,18 +313,16 @@ bool FMaterialCachedExpressionData::AddParameter(const FMaterialParameterInfo& P
 	}
 	else
 	{
-		auto GetEditorInfoAsText = [EditorOnlyData = EditorOnlyData](const FMaterialCachedParameterEditorInfo& InEditorInfo) -> FText
+		auto GetEditorInfoAsText = [AssetPath = InAssetPath](const FMaterialCachedParameterEditorInfo& InEditorInfo) -> FText
 			{
-				if (EditorOnlyData->AssetPaths.IsValidIndex(InEditorInfo.AssetIndex))
+				if(!AssetPath.IsEmpty())
 				{
-					FString AssetPath = EditorOnlyData->AssetPaths[InEditorInfo.AssetIndex];
 					return FText::Format(LOCTEXT("ReportParameterMetaData", "group:'{0}', asset:'{1}'"), FText::FromName(InEditorInfo.Group), FText::FromString(AssetPath));
 				}
 				else
 				{
 					return FText::Format(LOCTEXT("ReportInvalidatedParameterMetaData", "group:'{0}', description:'{1}'"), FText::FromName(InEditorInfo.Group), FText::FromString(InEditorInfo.Description));
 				}
-
 			};
 
 		auto ReportError = [&GetEditorInfoAsText](const FString& InCachedValue, const FString& InValue, const FMaterialCachedParameterEditorInfo& InCachedEditorInfo, const FMaterialCachedParameterEditorInfo& InEditorInfo) -> FText
@@ -519,8 +542,10 @@ void FMaterialCachedExpressionData::UpdateForExpressions(const FMaterialCachedEx
 
 			const FMaterialParameterInfo ParameterInfo(ParameterName, Association, ParameterIndex);
 
+			const FString AssetPath = Expression->GetAssetPathName();
+
 			// Try add the parameter. If this fails, the parameter is being added twice with different values. Report it as error.
-			if (!AddParameter(ParameterInfo, ParameterMeta, ReferencedTexture, ReferencedTextureCollection, ErrorContext))
+			if (!AddParameter(ParameterInfo, ParameterMeta, AssetPath, ReferencedTexture, ReferencedTextureCollection, ErrorContext))
 			{
 				FText ErrorMessage = FText::Format(LOCTEXT("DuplicateParameterError", "Parameter '{0}' is set multiple times to different values : {1}. Make sure each parameter is set once or always to the same value."),
 					 FText::FromName(ParameterName), ErrorContext);
@@ -891,10 +916,6 @@ void FMaterialCachedExpressionData::GetParameterValueByIndex(EMaterialParameterT
 			OutResult.Description = EditorInfo.Description;
 			OutResult.Group = EditorInfo.Group;
 			OutResult.SortPriority = EditorInfo.SortPriority;
-			if (EditorInfo.AssetIndex != INDEX_NONE)
-			{
-				OutResult.AssetPath = EditorOnlyData->AssetPaths[EditorInfo.AssetIndex];
-			}
 		}
 	}
 #endif // WITH_EDITORONLY_DATA
diff --git a/Engine/Source/Runtime/Engine/Private/Materials/MaterialExpressions.cpp b/Engine/Source/Runtime/Engine/Private/Materials/MaterialExpressions.cpp
--- a/Engine/Source/Runtime/Engine/Private/Materials/MaterialExpressions.cpp
+++ b/Engine/Source/Runtime/Engine/Private/Materials/MaterialExpressions.cpp
@@ -9539,7 +9539,6 @@ bool UMaterialExpressionTextureCollectionParameter::GetParameterValue(FMaterialP
 	OutMeta.ExpressionGuid = ExpressionGUID;
 	OutMeta.Group = Group;
 	OutMeta.SortPriority = SortPriority;
-	OutMeta.AssetPath = GetAssetPathName();
 	return true;
 }
 
@@ -17034,7 +17033,6 @@ bool UMaterialExpressionCollectionParameter::GetParameterValue(FMaterialParamete
 	OutMeta.ExpressionGuid = ExpressionGUID;
 	OutMeta.Group = Group;
 	OutMeta.SortPriority = SortPriority;
-	OutMeta.AssetPath = GetAssetPathName();
 	return true;
 }
 
@@ -20597,7 +20595,6 @@ bool UMaterialExpressionCurveAtlasRowParameter::GetParameterValue(FMaterialParam
 	OutMeta.ExpressionGuid = ExpressionGUID;
 	OutMeta.Group = Group;
 	OutMeta.SortPriority = SortPriority;
-	OutMeta.AssetPath = GetAssetPathName();
 	return true;
 }
 
@@ -21954,7 +21951,6 @@ bool UMaterialExpressionSparseVolumeTextureSampleParameter::GetParameterValue(FM
 	OutMeta.ExpressionGuid = ExpressionGUID;
 	OutMeta.Group = Group;
 	OutMeta.SortPriority = SortPriority;
-	OutMeta.AssetPath = GetAssetPathName();
 	return true;
 }
 
diff --git a/Engine/Source/Runtime/Engine/Public/MaterialCachedData.h b/Engine/Source/Runtime/Engine/Public/MaterialCachedData.h
--- a/Engine/Source/Runtime/Engine/Public/MaterialCachedData.h
+++ b/Engine/Source/Runtime/Engine/Public/MaterialCachedData.h
@@ -79,11 +79,10 @@ struct FMaterialCachedParameterEditorInfo
 
 	FMaterialCachedParameterEditorInfo() = default;
 	FMaterialCachedParameterEditorInfo(const FGuid& InGuid) : ExpressionGuid(InGuid) {}
-	FMaterialCachedParameterEditorInfo(const FGuid& InGuid, const FString& InDescription, const FName& InGroup, int32 InSortPriority, int32 InAssetIndex)
+	FMaterialCachedParameterEditorInfo(const FGuid& InGuid, const FString& InDescription, const FName& InGroup, int32 InSortPriority)
 		: Description(InDescription)
 		, Group(InGroup)
 		, SortPriority(InSortPriority)
-		, AssetIndex(InAssetIndex)
 		, ExpressionGuid(InGuid)
 	{}
 
@@ -96,9 +95,6 @@ struct FMaterialCachedParameterEditorInfo
 	UPROPERTY()
 	int32 SortPriority = 0;
 
-	UPROPERTY()
-	int32 AssetIndex = INDEX_NONE;
-
 	UPROPERTY()
 	FGuid ExpressionGuid;
 };
@@ -132,6 +128,17 @@ struct FMaterialCachedExpressionContext
 	bool bUpdateFunctionExpressions = true;
 };
 
+USTRUCT()
+struct FParameterAssetPathsEntry
+{
+	GENERATED_USTRUCT_BODY()
+
+	UPROPERTY()
+	TMap<FString, int32> AssetPaths;
+
+	ENGINE_API FString ToString() const;
+};
+
 USTRUCT()
 struct FMaterialCachedExpressionEditorOnlyData
 {
@@ -179,7 +186,7 @@ struct FMaterialCachedExpressionEditorOnlyData
 	FMaterialLayersFunctionsEditorOnlyData MaterialLayers;
 
 	UPROPERTY()
-	TArray<FString> AssetPaths;
+	TMap<FName, FParameterAssetPathsEntry> AssetPathsMap;
 
 	UPROPERTY()
 	TArray<FName> LandscapeLayerNames;
@@ -212,12 +219,12 @@ struct FMaterialCachedExpressionData
 	void Validate(const UMaterialInterface& Material);
 
 	/** Adds a parameter. If this returns false, a parameter with identical name has already been added but it was set to a different value. */
-	bool AddParameter(const FMaterialParameterInfo& ParameterInfo, const FMaterialParameterMetadata& ParameterMeta, UObject*& OutReferencedTexture, UTextureCollection*& OutReferencedTextureCollection, FText& OutErrorMessage);
+	bool AddParameter(const FMaterialParameterInfo& ParameterInfo, const FMaterialParameterMetadata& ParameterMeta, const FString& InAssetPath ,UObject*& OutReferencedTexture, UTextureCollection*& OutReferencedTextureCollection, FText& OutErrorMessage);
 
-	bool AddParameter(const FMaterialParameterInfo& ParameterInfo, const FMaterialParameterMetadata& ParameterMeta, UObject*& OutReferencedTexture, FText& OutErrorMessage)
+	bool AddParameter(const FMaterialParameterInfo& ParameterInfo, const FMaterialParameterMetadata& ParameterMeta, const FString& InAssetPath, UObject*& OutReferencedTexture, FText& OutErrorMessage)
 	{
 		UTextureCollection* UnusedReferencedTextureCollection = nullptr;
-		return AddParameter(ParameterInfo, ParameterMeta, OutReferencedTexture, UnusedReferencedTextureCollection, OutErrorMessage);
+		return AddParameter(ParameterInfo, ParameterMeta, InAssetPath, OutReferencedTexture, UnusedReferencedTextureCollection, OutErrorMessage);
 	}
 #endif // WITH_EDITOR
 
diff --git a/Engine/Source/Runtime/Engine/Public/Materials/MaterialExpressionFontSampleParameter.h b/Engine/Source/Runtime/Engine/Public/Materials/MaterialExpressionFontSampleParameter.h
--- a/Engine/Source/Runtime/Engine/Public/Materials/MaterialExpressionFontSampleParameter.h
+++ b/Engine/Source/Runtime/Engine/Public/Materials/MaterialExpressionFontSampleParameter.h
@@ -55,7 +55,6 @@ class UMaterialExpressionFontSampleParameter : public UMaterialExpressionFontSam
 		OutMeta.ExpressionGuid = ExpressionGUID;
 		OutMeta.Group = Group;
 		OutMeta.SortPriority = SortPriority;
-		OutMeta.AssetPath = GetAssetPathName();
 		return true;
 	}
 	virtual bool SetParameterValue(const FName& Name, const FMaterialParameterMetadata& Meta, EMaterialExpressionSetParameterValueFlags Flags) override
diff --git a/Engine/Source/Runtime/Engine/Public/Materials/MaterialExpressionParameter.h b/Engine/Source/Runtime/Engine/Public/Materials/MaterialExpressionParameter.h
--- a/Engine/Source/Runtime/Engine/Public/Materials/MaterialExpressionParameter.h
+++ b/Engine/Source/Runtime/Engine/Public/Materials/MaterialExpressionParameter.h
@@ -55,7 +55,6 @@ class UMaterialExpressionParameter : public UMaterialExpression
 		OutMeta.ExpressionGuid = ExpressionGUID;
 		OutMeta.Group = Group;
 		OutMeta.SortPriority = SortPriority;
-		OutMeta.AssetPath = GetAssetPathName();
 		return true;
 	}
 
diff --git a/Engine/Source/Runtime/Engine/Public/Materials/MaterialExpressionRuntimeVirtualTextureSampleParameter.h b/Engine/Source/Runtime/Engine/Public/Materials/MaterialExpressionRuntimeVirtualTextureSampleParameter.h
--- a/Engine/Source/Runtime/Engine/Public/Materials/MaterialExpressionRuntimeVirtualTextureSampleParameter.h
+++ b/Engine/Source/Runtime/Engine/Public/Materials/MaterialExpressionRuntimeVirtualTextureSampleParameter.h
@@ -51,7 +51,6 @@ class UMaterialExpressionRuntimeVirtualTextureSampleParameter : public UMaterial
 		OutMeta.ExpressionGuid = ExpressionGUID;
 		OutMeta.Group = Group;
 		OutMeta.SortPriority = SortPriority;
-		OutMeta.AssetPath = GetAssetPathName();
 		return true;
 	}
 	virtual bool SetParameterValue(const FName& Name, const FMaterialParameterMetadata& Meta, EMaterialExpressionSetParameterValueFlags Flags) override
diff --git a/Engine/Source/Runtime/Engine/Public/Materials/MaterialExpressionTextureSampleParameter.h b/Engine/Source/Runtime/Engine/Public/Materials/MaterialExpressionTextureSampleParameter.h
--- a/Engine/Source/Runtime/Engine/Public/Materials/MaterialExpressionTextureSampleParameter.h
+++ b/Engine/Source/Runtime/Engine/Public/Materials/MaterialExpressionTextureSampleParameter.h
@@ -53,7 +53,6 @@ class UMaterialExpressionTextureSampleParameter : public UMaterialExpressionText
 		OutMeta.ExpressionGuid = ExpressionGUID;
 		OutMeta.Group = Group;
 		OutMeta.SortPriority = SortPriority;
-		OutMeta.AssetPath = GetAssetPathName();
 		OutMeta.ChannelNames = GetTextureChannelNames();
 		return true;
 	}
diff --git a/Engine/Source/Runtime/Engine/Public/Materials/MaterialParameters.h b/Engine/Source/Runtime/Engine/Public/Materials/MaterialParameters.h
--- a/Engine/Source/Runtime/Engine/Public/Materials/MaterialParameters.h
+++ b/Engine/Source/Runtime/Engine/Public/Materials/MaterialParameters.h
@@ -462,9 +462,6 @@ struct FMaterialParameterMetadata
 	/** Description of the parameter, typically taken from the 'Desc' field of the parameter's UMaterialExpression */
 	FString Description;
 
-	/** Path to the asset (UMaterial or UMaterialFunction) that introduced this parameter */
-	FString AssetPath;
-
 	/** Name of the parameter's group */
 	FName Group;
 
