Preshader2:  Static analysis warning fixes -- all comment syntax I attempted to suppress the errors failed, so I refactored the code, making it a tiny bit less readable to fix the errors.  Require explicit "ALLOW_INTEGER" flag per opcode, rather than using an opcode range comparison in the template argument, which static analysis falsely believes is a comma separated expression (which it considers inherently unsafe).  Make separate versions of the "CAST_REUSE" macro for fall through versus implementation code, to avoid an always true expression.

#rnx

diff --git a/Engine/Source/Runtime/Engine/Private/Shader/Preshader2.cpp b/Engine/Source/Runtime/Engine/Private/Shader/Preshader2.cpp
--- a/Engine/Source/Runtime/Engine/Private/Shader/Preshader2.cpp
+++ b/Engine/Source/Runtime/Engine/Private/Shader/Preshader2.cpp
@@ -12,21 +12,22 @@
 namespace UE::Preshader2
 {
 
-#define PRESHADER2_UNARY_OP(OP_NAME)	\
-	TEXT(#OP_NAME),						\
-	TEXT(#OP_NAME "_InPlace"),			\
+#define PRESHADER2_UNARY_OP(OP_NAME, ALLOW_INTEGER)	\
+	TEXT(#OP_NAME),									\
+	TEXT(#OP_NAME "_InPlace"),						\
 	TEXT(#OP_NAME "_Reuse"),
 
-#define PRESHADER2_BINARY_OP(OP_NAME)	\
-	TEXT(#OP_NAME),						\
-	TEXT(#OP_NAME "_InPlace"),			\
-	TEXT(#OP_NAME "_Reuse"),			\
-	TEXT(#OP_NAME "_Const"),			\
-	TEXT(#OP_NAME "_ConstInPlace"),		\
-	TEXT(#OP_NAME "_ConstReuse"),		\
+#define PRESHADER2_BINARY_OP(OP_NAME, ALLOW_INTEGER)	\
+	TEXT(#OP_NAME),										\
+	TEXT(#OP_NAME "_InPlace"),							\
+	TEXT(#OP_NAME "_Reuse"),							\
+	TEXT(#OP_NAME "_Const"),							\
+	TEXT(#OP_NAME "_ConstInPlace"),						\
+	TEXT(#OP_NAME "_ConstReuse"),						\
 	TEXT(#OP_NAME "_ConstReuseScalar"),
 
 #define PRESHADER2_UNARY_OP_CAST_REUSE(OP_NAME) TEXT(#OP_NAME),
+#define PRESHADER2_UNARY_OP_CAST_REUSE_LAST(OP_NAME) TEXT(#OP_NAME),
 #define PRESHADER2_UNARY_OP_CUSTOM(OP_NAME) TEXT(#OP_NAME),
 #define PRESHADER2_BINARY_OP_CUSTOM(OP_NAME) TEXT(#OP_NAME),
 #define PRESHADER2_OP(OP_NAME) TEXT(#OP_NAME),
@@ -40,6 +41,7 @@ static const TCHAR* GPreshader2OpcodeNames[] =
 #undef PRESHADER2_UNARY_OP
 #undef PRESHADER2_BINARY_OP
 #undef PRESHADER2_UNARY_OP_CAST_REUSE
+#undef PRESHADER2_UNARY_OP_CAST_REUSE_LAST
 #undef PRESHADER2_UNARY_OP_CUSTOM
 #undef PRESHADER2_BINARY_OP_CUSTOM
 #undef PRESHADER2_OP
@@ -851,7 +853,7 @@ void OpEvaluate(const FOpSparseVolumeTextureUniform&, const FMaterialRenderConte
 // InPlace:				Opcode(1)  ResultOffset(2)  Dimension+Type(1)				AOffset same as ResultOffset
 // Default:				Opcode(1)  ResultOffset(2)  AOffset(2)  Dimension+Type(1)
 // Does not sanitize Dimension in case custom implementation stores extra data in that byte!  Client evaluate function should remove junk high bits from Dimension.
-#define PRESHADER2_UNARY_OP(OP_NAME) \
+#define PRESHADER2_UNARY_OP(OP_NAME, ALLOW_INTEGER) \
 	case EPreshader2Opcode::OP_NAME:																																	\
 		OpcodePtr += 2;									/* Skip to AOffset */																							\
 	case EPreshader2Opcode::OP_NAME ## _InPlace:		/* Otherwise AOffset pulled from same address as ResultOffset */												\
@@ -861,21 +863,21 @@ void OpEvaluate(const FOpSparseVolumeTextureUniform&, const FMaterialRenderConte
 		Type = (EPreshader2Type)(Dimension >> 6);																														\
 		OpcodePtr += 3;																																					\
 	case EPreshader2Opcode::OP_NAME ## _Reuse:																															\
-		UnaryEvaluate<FOp ## OP_NAME,EPreshader2Opcode::OP_NAME < EPreshader2Opcode::UnaryNoInt>(DataContext.ResultsBuffer, Type, Dimension, ResultOffset, AOffset);	\
+		UnaryEvaluate<FOp ## OP_NAME, ALLOW_INTEGER>(DataContext.ResultsBuffer, Type, Dimension, ResultOffset, AOffset);												\
 		Dimension = Dimension & 7;						/* Dimension that's preserved across loop iterations, with payload removed */									\
 		break;
 
 // The Cast op is the only unary op that includes a separate destination type, which the source type is converted to, but there's no payload to store the type for
 // "Reuse" variations, so it's pulled from the opcode.  We subtract the opcode from the first "Cast_Reuse" opcode to extract the type and add it to the Dimension,
 // where it is stored in the other versions of the opcode.  To reduce code size bloat, we technically only need a case for the last version of the opcode, so we
-// fall through in the prior cases (i.e., actual code is only generated for Cast_ReuseLWC, and the other two opcodes only produce the "case" statement).
+// fall through in the prior cases, and the other two opcodes only produce the "case" statement.
 #define PRESHADER2_UNARY_OP_CAST_REUSE(OP_NAME) \
-	case EPreshader2Opcode::OP_NAME:																																			\
-		if constexpr (EPreshader2Opcode::OP_NAME == EPreshader2Opcode::Cast_ReuseLWC)																							\
-		{																																										\
-			UnaryEvaluate<FOpCast, true>(DataContext.ResultsBuffer, Type, Dimension | (((uint8)Opcode - (uint8)EPreshader2Opcode::Cast_Reuse) << 4), ResultOffset, AOffset);		\
-			break;																																								\
-		}
+	case EPreshader2Opcode::OP_NAME: [[fallthrough]];
+
+#define PRESHADER2_UNARY_OP_CAST_REUSE_LAST(OP_NAME) \
+	case EPreshader2Opcode::OP_NAME:																																		\
+		UnaryEvaluate<FOpCast, true>(DataContext.ResultsBuffer, Type, Dimension | (((uint8)Opcode - (uint8)EPreshader2Opcode::Cast_Reuse) << 4), ResultOffset, AOffset);	\
+		break;
 
 // Opcode(1)  ResultOffset(2)  AOffset(2)  Dimension+Type(1).
 // Does not sanitize Dimension in case custom implementation stores extra data in that byte!  Client evaluate function should remove junk high bits from Dimension.
@@ -897,7 +899,7 @@ void OpEvaluate(const FOpSparseVolumeTextureUniform&, const FMaterialRenderConte
 // ConstReuseScalar:	Opcode(1)  BData(N)																ResultOffset/Dimension/Type used from previous op, AOffset same as ResultOffset, BInc 0, AInc 1
 // ConstInPlace:		Opcode(1)  ResultOffset(2)  Dimension+Type+AInc+BInc(1)  BData(N)				AOffset same as ResultOffset
 // ConstDefault:		Opcode(1)  ResultOffset(2)  AOffset(2)  Dimension+Type+AInc+BInc(1)  BData(N)
-#define PRESHADER2_BINARY_OP(OP_NAME) \
+#define PRESHADER2_BINARY_OP(OP_NAME, ALLOW_INTEGER) \
 	case EPreshader2Opcode::OP_NAME ## _Reuse:																				\
 		BOffset = FetchOffset(OpcodePtr);																					\
 		OpcodePtr += 2;									/* Skip over BOffset */												\
@@ -917,7 +919,7 @@ void OpEvaluate(const FOpSparseVolumeTextureUniform&, const FMaterialRenderConte
 	OP_NAME ## _Evaluate:																									\
 		BInc = (BOffset >> 15) & 1;																							\
 		BOffset &= 0x7fff;																									\
-		BinaryEvaluate<FOp ## OP_NAME, EPreshader2Opcode::OP_NAME < EPreshader2Opcode::BinaryNoInt>(DataContext.ResultsBuffer, Type, Dimension, ResultOffset, AOffset, BOffset, AInc, BInc);	\
+		BinaryEvaluate<FOp ## OP_NAME, ALLOW_INTEGER>(DataContext.ResultsBuffer, Type, Dimension, ResultOffset, AOffset, BOffset, AInc, BInc);	\
 		break;																												\
 	case EPreshader2Opcode::OP_NAME ## _ConstReuse:																			\
 		AInc = 1;																											\
@@ -939,7 +941,7 @@ void OpEvaluate(const FOpSparseVolumeTextureUniform&, const FMaterialRenderConte
 		Dimension = Dimension & 7;																							\
 		OpcodePtr += 3;																										\
 	OP_NAME ## _ConstEvaluate:																								\
-		BinaryEvaluateConst<FOp ## OP_NAME, EPreshader2Opcode::OP_NAME < EPreshader2Opcode::BinaryNoInt>(DataContext.ResultsBuffer, Type, Dimension, ResultOffset, AOffset, OpcodePtr, AInc, BInc);	\
+		BinaryEvaluateConst<FOp ## OP_NAME, ALLOW_INTEGER>(DataContext.ResultsBuffer, Type, Dimension, ResultOffset, AOffset, OpcodePtr, AInc, BInc);	\
 		break;
 
 // Opcode(1)  ResultOffset(2)  AOffset(2)  BOffset(2)  Dimension+Type(1)
@@ -1002,6 +1004,7 @@ void Evaluate(UE::Shader::FPreshaderDataContext& DataContext, const FMaterialRen
 #undef PRESHADER2_UNARY_OP
 #undef PRESHADER2_BINARY_OP
 #undef PRESHADER2_UNARY_OP_CAST_REUSE
+#undef PRESHADER2_UNARY_OP_CAST_REUSE_LAST
 #undef PRESHADER2_UNARY_OP_CUSTOM
 #undef PRESHADER2_BINARY_OP_CUSTOM
 #undef PRESHADER2_OP
diff --git a/Engine/Source/Runtime/Engine/Public/Shader/Preshader2.h b/Engine/Source/Runtime/Engine/Public/Shader/Preshader2.h
--- a/Engine/Source/Runtime/Engine/Public/Shader/Preshader2.h
+++ b/Engine/Source/Runtime/Engine/Public/Shader/Preshader2.h
@@ -14,21 +14,22 @@ namespace UE::Shader
 struct FPreshaderDataContext;
 }
 
-#define PRESHADER2_UNARY_OP(OP_NAME)	\
-	OP_NAME,									\
-	OP_NAME ## _InPlace,						\
+#define PRESHADER2_UNARY_OP(OP_NAME, ALLOW_INTEGER)		\
+	OP_NAME,											\
+	OP_NAME ## _InPlace,								\
 	OP_NAME ## _Reuse,
 
-#define PRESHADER2_BINARY_OP(OP_NAME)	\
-	OP_NAME,									\
-	OP_NAME ## _InPlace,						\
-	OP_NAME ## _Reuse,							\
-	OP_NAME ## _Const,							\
-	OP_NAME ## _ConstInPlace,					\
-	OP_NAME ## _ConstReuse,						\
+#define PRESHADER2_BINARY_OP(OP_NAME, ALLOW_INTEGER)	\
+	OP_NAME,											\
+	OP_NAME ## _InPlace,								\
+	OP_NAME ## _Reuse,									\
+	OP_NAME ## _Const,									\
+	OP_NAME ## _ConstInPlace,							\
+	OP_NAME ## _ConstReuse,								\
 	OP_NAME ## _ConstReuseScalar,
 
 #define PRESHADER2_UNARY_OP_CAST_REUSE(OP_NAME) OP_NAME,
+#define PRESHADER2_UNARY_OP_CAST_REUSE_LAST(OP_NAME) OP_NAME,
 #define PRESHADER2_UNARY_OP_CUSTOM(OP_NAME) OP_NAME,
 #define PRESHADER2_BINARY_OP_CUSTOM(OP_NAME) OP_NAME,
 #define PRESHADER2_OP(OP_NAME) OP_NAME,
@@ -59,6 +60,7 @@ enum class EPreshader2Opcode : uint8
 #undef PRESHADER2_UNARY_OP
 #undef PRESHADER2_BINARY_OP
 #undef PRESHADER2_UNARY_OP_CAST_REUSE
+#undef PRESHADER2_UNARY_OP_CAST_REUSE_LAST
 #undef PRESHADER2_UNARY_OP_CUSTOM
 #undef PRESHADER2_BINARY_OP_CUSTOM
 #undef PRESHADER2_OP
diff --git a/Engine/Source/Runtime/Engine/Public/Shader/Preshader2Opcodes.inl b/Engine/Source/Runtime/Engine/Public/Shader/Preshader2Opcodes.inl
--- a/Engine/Source/Runtime/Engine/Public/Shader/Preshader2Opcodes.inl
+++ b/Engine/Source/Runtime/Engine/Public/Shader/Preshader2Opcodes.inl
@@ -1,55 +1,52 @@
 // Copyright Epic Games, Inc. All Rights Reserved.
 
-//-V501
-//-V521
-
 // Unary ops with integer support -- each has 3 opcode variants, with Cast having 3 extra custom "Reuse" opcode variants for different destination types.
 // The three extra opcodes conveniently maintain the 3-element stride for surrounding standard unary opcode variants.
-PRESHADER2_UNARY_OP(Abs)
-PRESHADER2_UNARY_OP(Cast)
+PRESHADER2_UNARY_OP(Abs, true)
+PRESHADER2_UNARY_OP(Cast, true)
 PRESHADER2_UNARY_OP_CAST_REUSE(Cast_ReuseDouble)
 PRESHADER2_UNARY_OP_CAST_REUSE(Cast_ReuseInteger)
-PRESHADER2_UNARY_OP_CAST_REUSE(Cast_ReuseLWC)
-PRESHADER2_UNARY_OP(Negate)
-PRESHADER2_UNARY_OP(Saturate)
-PRESHADER2_UNARY_OP(Sign)
+PRESHADER2_UNARY_OP_CAST_REUSE_LAST(Cast_ReuseLWC)
+PRESHADER2_UNARY_OP(Negate, true)
+PRESHADER2_UNARY_OP(Saturate, true)
+PRESHADER2_UNARY_OP(Sign, true)
 
 // Unary ops without integer support
-PRESHADER2_UNARY_OP(ACos)
-PRESHADER2_UNARY_OP(ASin)
-PRESHADER2_UNARY_OP(ATan)
-PRESHADER2_UNARY_OP(Ceil)
-PRESHADER2_UNARY_OP(Cos)
-PRESHADER2_UNARY_OP(Exponential)
-PRESHADER2_UNARY_OP(Exponential2)
-PRESHADER2_UNARY_OP(Floor)
-PRESHADER2_UNARY_OP(Frac)
-PRESHADER2_UNARY_OP(Fractional)
-PRESHADER2_UNARY_OP(Logarithm)
-PRESHADER2_UNARY_OP(Logarithm10)
-PRESHADER2_UNARY_OP(Logarithm2)
-PRESHADER2_UNARY_OP(Normalize)
-PRESHADER2_UNARY_OP(Reciprocal)
-PRESHADER2_UNARY_OP(Round)
-PRESHADER2_UNARY_OP(Sin)
-PRESHADER2_UNARY_OP(Sqrt)
-PRESHADER2_UNARY_OP(Tan)
-PRESHADER2_UNARY_OP(Truncate)
+PRESHADER2_UNARY_OP(ACos, false)
+PRESHADER2_UNARY_OP(ASin, false)
+PRESHADER2_UNARY_OP(ATan, false)
+PRESHADER2_UNARY_OP(Ceil, false)
+PRESHADER2_UNARY_OP(Cos, false)
+PRESHADER2_UNARY_OP(Exponential, false)
+PRESHADER2_UNARY_OP(Exponential2, false)
+PRESHADER2_UNARY_OP(Floor, false)
+PRESHADER2_UNARY_OP(Frac, false)
+PRESHADER2_UNARY_OP(Fractional, false)
+PRESHADER2_UNARY_OP(Logarithm, false)
+PRESHADER2_UNARY_OP(Logarithm10, false)
+PRESHADER2_UNARY_OP(Logarithm2, false)
+PRESHADER2_UNARY_OP(Normalize, false)
+PRESHADER2_UNARY_OP(Reciprocal, false)
+PRESHADER2_UNARY_OP(Round, false)
+PRESHADER2_UNARY_OP(Sin, false)
+PRESHADER2_UNARY_OP(Sqrt, false)
+PRESHADER2_UNARY_OP(Tan, false)
+PRESHADER2_UNARY_OP(Truncate, false)
 
 // Non-standard unary ops -- single opcode variant
 PRESHADER2_UNARY_OP_CUSTOM(Length)
 
 // Binary ops with integer support -- each has 7 opcode variants
-PRESHADER2_BINARY_OP(Add)
-PRESHADER2_BINARY_OP(Divide)
-PRESHADER2_BINARY_OP(Fmod)
-PRESHADER2_BINARY_OP(Max)
-PRESHADER2_BINARY_OP(Min)
-PRESHADER2_BINARY_OP(Multiply)
-PRESHADER2_BINARY_OP(Subtract)
+PRESHADER2_BINARY_OP(Add, true)
+PRESHADER2_BINARY_OP(Divide, true)
+PRESHADER2_BINARY_OP(Fmod, true)
+PRESHADER2_BINARY_OP(Max, true)
+PRESHADER2_BINARY_OP(Min, true)
+PRESHADER2_BINARY_OP(Multiply, true)
+PRESHADER2_BINARY_OP(Subtract, true)
 
 // Binary ops without integer support
-PRESHADER2_BINARY_OP(ATan2)
+PRESHADER2_BINARY_OP(ATan2, false)
 
 // Non-standard binary ops -- single opcode variant
 PRESHADER2_BINARY_OP_CUSTOM(Cross)
