Material translator: Handle case when null/invalid chunk is passed through TruncateLWC().
#rnx

diff --git a/Engine/Source/Runtime/Engine/Private/Materials/HLSLMaterialTranslator.cpp b/Engine/Source/Runtime/Engine/Private/Materials/HLSLMaterialTranslator.cpp
--- a/Engine/Source/Runtime/Engine/Private/Materials/HLSLMaterialTranslator.cpp
+++ b/Engine/Source/Runtime/Engine/Private/Materials/HLSLMaterialTranslator.cpp
@@ -5242,7 +5242,7 @@ int32 FHLSLMaterialTranslator::TruncateLWC(int32 Code)
 	const int32 LWCTruncateMode = UE::MaterialTranslatorUtils::GetLWCTruncateMode();
 	const bool bAllowLWCTruncate = LWCTruncateMode == 1 || LWCTruncateMode == 2;
 
-	if (!bAllowLWCTruncate)
+	if (Code == INDEX_NONE || !bAllowLWCTruncate)
 	{
 		return Code;
 	}
