Collision Filter Refactor Part 3.7: Deprecating FCollisionStructure as it was never used.

diff --git a/Engine/Source/Runtime/Engine/Private/PhysicsEngine/CollisionStructure.h b/Engine/Source/Runtime/Engine/Private/PhysicsEngine/CollisionStructure.h
--- a/Engine/Source/Runtime/Engine/Private/PhysicsEngine/CollisionStructure.h
+++ b/Engine/Source/Runtime/Engine/Private/PhysicsEngine/CollisionStructure.h
@@ -7,7 +7,7 @@
 
 struct FHitResult;
 
-class FCollisionStructure
+class UE_DEPRECATED(5.8, "This was unused and will be removed") FCollisionStructure
 {
 	TArray<bool>					ValidFlag;
 	TArray<FBox>					Bounds;
