Substrate - fixed wrong assert. We care about the new shading model not the existing one.

#rnx

diff --git a/Engine/Source/Runtime/Engine/Private/Materials/MaterialExpressions.cpp b/Engine/Source/Runtime/Engine/Private/Materials/MaterialExpressions.cpp
--- a/Engine/Source/Runtime/Engine/Private/Materials/MaterialExpressions.cpp
+++ b/Engine/Source/Runtime/Engine/Private/Materials/MaterialExpressions.cpp
@@ -2140,7 +2140,11 @@ void UMaterialExpression::ConnectToPreviewMaterial(UMaterial* InMaterial, int32
 			else if (IsResultMaterialAttributes(0))
 			{
 				// Propagate material attributes to MaterialAttributes input
-				InMaterial->SetShadingModel(MSM_DefaultLit);
+				if (InMaterial->GetShadingModels()!=MSM_FromMaterialExpression)
+				{
+					// Only set DefaultLit if we are not using the special FromMaterialExpression already as this is forbidden by SetShadingModel and Substrate deal with that case.
+					InMaterial->SetShadingModel(MSM_DefaultLit);
+				}
 				InMaterial->bUseMaterialAttributes = true;
 				FExpressionInput* MaterialInput = InMaterial->GetExpressionInputForProperty(MP_MaterialAttributes);
 				check(MaterialInput);
