Substrate - follow up to 46756517. Better to use built in function even though when MSM_FromMaterialExpression is used, it must be the single one.

#tests PC editor
#rnx

diff --git a/Engine/Source/Runtime/Engine/Private/Materials/MaterialExpressions.cpp b/Engine/Source/Runtime/Engine/Private/Materials/MaterialExpressions.cpp
--- a/Engine/Source/Runtime/Engine/Private/Materials/MaterialExpressions.cpp
+++ b/Engine/Source/Runtime/Engine/Private/Materials/MaterialExpressions.cpp
@@ -2140,7 +2140,7 @@ void UMaterialExpression::ConnectToPreviewMaterial(UMaterial* InMaterial, int32
 			else if (IsResultMaterialAttributes(0))
 			{
 				// Propagate material attributes to MaterialAttributes input
-				if (InMaterial->GetShadingModels()!=MSM_FromMaterialExpression)
+				if (InMaterial->GetShadingModels().HasShadingModel(MSM_FromMaterialExpression))
 				{
 					// Only set DefaultLit if we are not using the special FromMaterialExpression already as this is forbidden by SetShadingModel and Substrate deal with that case.
 					InMaterial->SetShadingModel(MSM_DefaultLit);
