Substrate - follow up to 46756604. Using the same logic as 46756517...

#tests PC editor
#rnx

diff --git a/Engine/Source/Runtime/Engine/Private/Materials/MaterialExpressions.cpp b/Engine/Source/Runtime/Engine/Private/Materials/MaterialExpressions.cpp
--- a/Engine/Source/Runtime/Engine/Private/Materials/MaterialExpressions.cpp
+++ b/Engine/Source/Runtime/Engine/Private/Materials/MaterialExpressions.cpp
@@ -2140,12 +2140,16 @@ void UMaterialExpression::ConnectToPreviewMaterial(UMaterial* InMaterial, int32
 			else if (IsResultMaterialAttributes(0))
 			{
 				// Propagate material attributes to MaterialAttributes input
-				if (InMaterial->GetShadingModels().HasShadingModel(MSM_FromMaterialExpression))
+				if (!InMaterial->IsShadingModelFromMaterialExpression())
 				{
 					// Only set DefaultLit if we are not using the special FromMaterialExpression already as this is forbidden by SetShadingModel and Substrate deal with that case.
 					InMaterial->SetShadingModel(MSM_DefaultLit);
+					InMaterial->bUseMaterialAttributes = false;
+				}
+				else
+				{
+					InMaterial->bUseMaterialAttributes = true;
 				}
-				InMaterial->bUseMaterialAttributes = true;
 				FExpressionInput* MaterialInput = InMaterial->GetExpressionInputForProperty(MP_MaterialAttributes);
 				check(MaterialInput);
 				ConnectExpression( MaterialInput, OutputIndex );
