PR #13884: Mark UAnimInstance::GetCurrentStateName as const and add const correctness in function body

diff --git a/Engine/Source/Runtime/Engine/Classes/Animation/AnimInstance.h b/Engine/Source/Runtime/Engine/Classes/Animation/AnimInstance.h
--- a/Engine/Source/Runtime/Engine/Classes/Animation/AnimInstance.h
+++ b/Engine/Source/Runtime/Engine/Classes/Animation/AnimInstance.h
@@ -1230,7 +1230,7 @@ public:
 
 	/** Returns the name of a currently active state in a state machine. */
 	UFUNCTION(BlueprintPure, Category="Animation|State Machines", meta=(BlueprintInternalUseOnly = "true", AnimGetter = "true", BlueprintThreadSafe))
-	ENGINE_API FName GetCurrentStateName(int32 MachineIndex);
+	ENGINE_API FName GetCurrentStateName(int32 MachineIndex) const;
 
 	/** Sets a morph target to a certain weight. */
 	UFUNCTION(BlueprintCallable, Category="Animation|Morph Targets")
diff --git a/Engine/Source/Runtime/Engine/Private/Animation/AnimInstance.cpp b/Engine/Source/Runtime/Engine/Private/Animation/AnimInstance.cpp
--- a/Engine/Source/Runtime/Engine/Private/Animation/AnimInstance.cpp
+++ b/Engine/Source/Runtime/Engine/Private/Animation/AnimInstance.cpp
@@ -1946,19 +1946,19 @@ void UAnimInstance::SetRootMotionMode(TEnumAsByte<ERootMotionMode::Type> Value)
 }
 
 
-FName UAnimInstance::GetCurrentStateName(int32 MachineIndex)
+FName UAnimInstance::GetCurrentStateName(int32 MachineIndex) const
 {
-	if (IAnimClassInterface* AnimBlueprintClass = IAnimClassInterface::GetFromClass(GetClass()))
+	if (const IAnimClassInterface* AnimBlueprintClass = IAnimClassInterface::GetFromClass(GetClass()))
 	{
 		const TArray<FStructProperty*>& AnimNodeProperties = AnimBlueprintClass->GetAnimNodeProperties();
 		if ((MachineIndex >= 0) && (MachineIndex < AnimNodeProperties.Num()))
 		{
 			const int32 InstancePropertyIndex = AnimNodeProperties.Num() - 1 - MachineIndex; //@TODO: ANIMREFACTOR: Reverse indexing
 
-			FStructProperty* MachineInstanceProperty = AnimNodeProperties[InstancePropertyIndex];
+			const FStructProperty* MachineInstanceProperty = AnimNodeProperties[InstancePropertyIndex];
 			checkSlow(MachineInstanceProperty->Struct->IsChildOf(FAnimNode_StateMachine::StaticStruct()));
 
-			FAnimNode_StateMachine* MachineInstance = MachineInstanceProperty->ContainerPtrToValuePtr<FAnimNode_StateMachine>(this);
+			const FAnimNode_StateMachine* MachineInstance = MachineInstanceProperty->ContainerPtrToValuePtr<FAnimNode_StateMachine>(this);
 
 			return MachineInstance->GetCurrentStateName();
 		}
