Fixed a physics filtering deprecated function that forgot to set some values.

diff --git a/Engine/Source/Runtime/Engine/Public/Physics/PhysicsFiltering.h b/Engine/Source/Runtime/Engine/Public/Physics/PhysicsFiltering.h
--- a/Engine/Source/Runtime/Engine/Public/Physics/PhysicsFiltering.h
+++ b/Engine/Source/Runtime/Engine/Public/Physics/PhysicsFiltering.h
@@ -194,6 +194,8 @@ inline void CreateShapeFilterData(
 	Builder.SetFlags(Chaos::EFilterFlags::ContactNotify, bEnableContactNotify);
 	Builder.SetFlags(Chaos::EFilterFlags::StaticShape, bStaticShape);
 	Builder.SetFlags(Chaos::EFilterFlags::ModifyContacts, bModifyContacts);
+	Builder.SetOwnerID(SourceObjectID);
+	Builder.SetComponentID(ComponentID);
 
 	const Chaos::Filter::FCombinedShapeFilterData CombinedShapeFilter(Builder.BuildShapeFilterData(), Builder.BuildInstanceData());
 	Chaos::Filter::FShapeFilterBuilder::GetLegacyShapeFilter(CombinedShapeFilter, OutQueryData, OutSimData);
