[MetaSoundsEditor] Add descriptions to UE.* metasounds interfaces

diff --git a/Engine/Plugins/Runtime/Metasound/Source/MetasoundFrontend/Private/Interfaces/MetasoundFrontendSourceInterface.cpp b/Engine/Plugins/Runtime/Metasound/Source/MetasoundFrontend/Private/Interfaces/MetasoundFrontendSourceInterface.cpp
--- a/Engine/Plugins/Runtime/Metasound/Source/MetasoundFrontend/Private/Interfaces/MetasoundFrontendSourceInterface.cpp
+++ b/Engine/Plugins/Runtime/Metasound/Source/MetasoundFrontend/Private/Interfaces/MetasoundFrontendSourceInterface.cpp
@@ -58,6 +58,10 @@ namespace Metasound::Frontend
 							LOCTEXT("OnFinishedWarning", "\"On Finished\" should be connected for OneShot MetaSound sources. For sources with undefined duration (e.g. looping), remove the OneShot interface and use an audio component to avoid leaking the source."),
 						}
 					};
+#if WITH_EDITOR
+					Description = LOCTEXT(AUDIO_PARAMETER_INTERFACE_NAMESPACE ".Description",
+						"For MetaSounds that are not looping nor infinite. Provides an output trigger to tell when the MetaSound is finished.");
+#endif
 				}
 			};
 
@@ -99,11 +103,17 @@ namespace Metasound::Frontend
 					{
 						{
 							LOCTEXT("StartTime", "Start Time"),
-							LOCTEXT("StartTime Description", "The StartTime of the source passed into the MetaSound from the Gameplay API."),
+							LOCTEXT("StartTime Description", "[0, Duration] The StartTime of the source passed into the MetaSound from the Gameplay API."),
 							GetMetasoundDataTypeName<FTime>(),
 							Inputs::StartTime
 						}
 					};
+#if WITH_EDITOR
+					Description = LOCTEXT(AUDIO_PARAMETER_INTERFACE_NAMESPACE ".Description",
+						"Provides the start time (from 0) to seek into the sound as its played. Works from Gameplay API and Sequencer.\n"
+						"The MetaSound itself will need to determine what to do with this input.\n"
+						"In the simplest case, a MetaSound with a single Wave Player can route StartTime into the Wave Player's StartTime.");
+#endif
 				}
 			};
 
diff --git a/Engine/Source/Runtime/Engine/Private/SoundAttenuation.cpp b/Engine/Source/Runtime/Engine/Private/SoundAttenuation.cpp
--- a/Engine/Source/Runtime/Engine/Private/SoundAttenuation.cpp
+++ b/Engine/Source/Runtime/Engine/Private/SoundAttenuation.cpp
@@ -287,11 +287,14 @@ namespace Audio
 					{
 						{
 							FText(),
-							NSLOCTEXT("AudioGeneratorInterface_Attenuation", "DistanceDescription", "Distance between listener and sound location in game units."),
+							NSLOCTEXT("AudioGeneratorInterface_Attenuation", "DistanceDescription", "Distance between listener and sound locations in Unreal Units."),
 							FName(),
 							{ Inputs::Distance, 0.0f }
 						}
 					};
+#if WITH_EDITOR
+					Description = LOCTEXT(AUDIO_PARAMETER_INTERFACE_NAMESPACE ".Description", "Provides the distance between listener and sound locations in Unreal Units.");
+#endif
 				}
 			};
 
@@ -328,17 +331,22 @@ namespace Audio
 					{
 						{
 							FText(),
-							NSLOCTEXT("Spatialization", "AzimuthDescription", "Horizontal angle between listener forward and sound location in degrees."),
+							NSLOCTEXT("Spatialization", "AzimuthDescription", "[-180, 180] Horizontal angle between listener forward and sound location in degrees."),
 							FName(),
 							{ Inputs::Azimuth, 0.0f }
 						},
 						{
 							FText(),
-							NSLOCTEXT("Spatialization", "ElevationDescription", "Vertical angle between listener forward and sound location in degrees."),
+							NSLOCTEXT("Spatialization", "ElevationDescription", "[-90, 90] Vertical angle between listener forward and sound location in degrees."),
 							FName(),
 							{ Inputs::Elevation, 0.0f }
 						}
 					};
+#if WITH_EDITOR
+					Description = LOCTEXT(AUDIO_PARAMETER_INTERFACE_NAMESPACE ".Description",
+						"From the perspective of the listener, the horizontal angle (azimuth) and vertical angle (elevation) to the sound.\n"
+						"Useful for focus-based mixing, like that of the Focus feature of Attenuation Settings.");
+#endif
 				}
 			};
 
@@ -375,17 +383,22 @@ namespace Audio
 					{
 						{
 							FText(),
-							NSLOCTEXT("SourceOrientation", "AzimuthDescription", "Horizontal angle between emitter forward and listener location in degrees."),
+							NSLOCTEXT("SourceOrientation", "AzimuthDescription", "[-180, 180] Horizontal angle between sound forward and listener location in degrees."),
 							FName(),
 							{ Inputs::Azimuth, 0.0f }
 						},
 						{
 							FText(),
-							NSLOCTEXT("SourceOrientation", "ElevationDescription", "Vertical angle between emitter forward and listener location in degrees."),
+							NSLOCTEXT("SourceOrientation", "ElevationDescription", "[-90, 90] Vertical angle between sound forward and listener location in degrees."),
 							FName(),
 							{ Inputs::Elevation, 0.0f }
 						}
 					};
+#if WITH_EDITOR
+					Description = LOCTEXT(AUDIO_PARAMETER_INTERFACE_NAMESPACE ".Description",
+						"From the perspective of the sound, provides the horizontal angle (azimuth) and vertical angle (elevation) to the listener.\n"
+						"Useful for directional sources like a megaphone or speaker, where the sound is louder pointed at you and quieter pointed away.");
+#endif
 				}
 			};
 
@@ -422,17 +435,22 @@ namespace Audio
 					{
 						{
 							FText(),
-							NSLOCTEXT("ListenerOrientation", "AzimuthDescription", "Horizontal viewing angle of the current listener in world."),
+							NSLOCTEXT("ListenerOrientation", "AzimuthDescription", "[-180, 180] Horizontal viewing angle of the current listener in world, in degrees."),
 							FName(),
 							{ Inputs::Azimuth, 0.0f }
 						},
 						{
 							FText(),
-							NSLOCTEXT("ListenerOrientation", "ElevationDescription", "Vertical viewing angle of the current listener in world."),
+							NSLOCTEXT("ListenerOrientation", "ElevationDescription", "[-90, 90] Vertical viewing angle of the current listener in world, in degrees."),
 							FName(),
 							{ Inputs::Elevation, 0.0f }
 						}
 					};
+#if WITH_EDITOR
+					Description = LOCTEXT(AUDIO_PARAMETER_INTERFACE_NAMESPACE ".Description",
+						"Relative to the world, provides the listener orientation, as horizontal angle (azimuth) and vertical angle (elevation).\n"
+						"Useful for world-locked quad beds for example.");
+#endif
 				}
 			};
 
