Added support in ISKM for GPU-driven instance culling. It currently relies on the Direct storage mode which is disabled by default due to shader branching cost.

diff --git a/Engine/Source/Runtime/Engine/Classes/Components/InstancedSkinnedMeshComponent.h b/Engine/Source/Runtime/Engine/Classes/Components/InstancedSkinnedMeshComponent.h
--- a/Engine/Source/Runtime/Engine/Classes/Components/InstancedSkinnedMeshComponent.h
+++ b/Engine/Source/Runtime/Engine/Classes/Components/InstancedSkinnedMeshComponent.h
@@ -103,6 +103,10 @@ protected:
 	UPROPERTY()
 	bool bIsInstanceDataGPUOnly = false;
 
+	/** If true, this component will use GPU LOD selection. */
+	UPROPERTY(EditAnywhere, BlueprintReadOnly, Category = Culling)
+	uint8 bUseGpuLodSelection : 1;
+
 	UPROPERTY()
 	int32 NumInstancesGPUOnly = 0;
 
diff --git a/Engine/Source/Runtime/Engine/Private/Components/InstancedSkinnedMeshComponent.cpp b/Engine/Source/Runtime/Engine/Private/Components/InstancedSkinnedMeshComponent.cpp
--- a/Engine/Source/Runtime/Engine/Private/Components/InstancedSkinnedMeshComponent.cpp
+++ b/Engine/Source/Runtime/Engine/Private/Components/InstancedSkinnedMeshComponent.cpp
@@ -62,6 +62,7 @@ static FSkeletalMeshObject* CreateInstancedSkeletalMeshObjectFunction(void* User
 UInstancedSkinnedMeshComponent::UInstancedSkinnedMeshComponent(FVTableHelper& Helper)
 : Super(Helper)
 , bInheritPerInstanceData(false)
+, bUseGpuLodSelection(true)
 , InstanceDataManager(this)
 {
 }
@@ -69,6 +70,7 @@ UInstancedSkinnedMeshComponent::UInstancedSkinnedMeshComponent(FVTableHelper& He
 UInstancedSkinnedMeshComponent::UInstancedSkinnedMeshComponent(const FObjectInitializer& ObjectInitializer) 
 : Super(ObjectInitializer) 
 , bInheritPerInstanceData(false)
+, bUseGpuLodSelection(true)
 , InstanceDataManager(this)
 {
 }
diff --git a/Engine/Source/Runtime/Engine/Private/InstancedSkeletalMeshObject.cpp b/Engine/Source/Runtime/Engine/Private/InstancedSkeletalMeshObject.cpp
--- a/Engine/Source/Runtime/Engine/Private/InstancedSkeletalMeshObject.cpp
+++ b/Engine/Source/Runtime/Engine/Private/InstancedSkeletalMeshObject.cpp
@@ -19,12 +19,15 @@ FInstancedSkeletalMeshObject::FInstancedSkeletalMeshObject(
 	, TransformProvider(InMeshDesc.TransformProvider)
 	, Type(InType)
 {
-	for (int32 LODIndex = 0; LODIndex < InRenderData->LODRenderData.Num(); ++LODIndex)
+	const int32 NumLODs = InRenderData->LODRenderData.Num();
+
+	for (int32 LODIndex = 0; LODIndex < NumLODs; ++LODIndex)
 	{
 		LODs.Emplace(InRenderData, LODIndex, InFeatureLevel);
 	}
 
-	BoneTransformStorageMode = Type == EType::GPUSkin ? ::GetBoneTransformStorageMode() : EBoneTransformStorageMode::Direct;
+	bUseGpuLodSelection = NumLODs > 1 && InMeshDesc.bUseGpuLodSelection && GPUSKIN_TRANSFORM_STORAGE_DIRECT_ENABLED;
+	BoneTransformStorageMode = Type == EType::GPUSkin && !bUseGpuLodSelection ? ::GetBoneTransformStorageMode() : EBoneTransformStorageMode::Direct;
 
 	InitResources(InMeshDesc);
 
@@ -105,15 +108,29 @@ void FInstancedSkeletalMeshObject::UpdateSceneExtensionHeader(FRHICommandListBas
 	if (Type == EType::GPUSkin)
 	{
 		check(DynamicData);
-		const FLOD& LOD = LODs[DynamicData->LODIndex];
 
-		for (const TUniquePtr<FGPUBaseSkinVertexFactory>& VertexFactory : LOD.VertexFactories)
+		const auto UpdateVertexFactory = [&] (int32 LODIndex)
 		{
-			if (VertexFactory)
+			for (const TUniquePtr<FGPUBaseSkinVertexFactory>& VertexFactory : LODs[LODIndex].VertexFactories)
+			{
+				if (VertexFactory)
+				{
+					VertexFactory->UpdateUniformBuffer(RHICmdList, &SkinningHeader, TransformDataBufferSRV);
+				}
+			}
+		};
+
+		if (bUseGpuLodSelection)
+		{
+			for (int32 LODIndex = 0; LODIndex < LODs.Num(); ++LODIndex)
 			{
-				VertexFactory->UpdateUniformBuffer(RHICmdList, &SkinningHeader, TransformDataBufferSRV);
+				UpdateVertexFactory(LODIndex);
 			}
 		}
+		else
+		{
+			UpdateVertexFactory(DynamicData->LODIndex);
+		}
 	}
 }
 
diff --git a/Engine/Source/Runtime/Engine/Private/InstancedSkeletalMeshObject.h b/Engine/Source/Runtime/Engine/Private/InstancedSkeletalMeshObject.h
--- a/Engine/Source/Runtime/Engine/Private/InstancedSkeletalMeshObject.h
+++ b/Engine/Source/Runtime/Engine/Private/InstancedSkeletalMeshObject.h
@@ -106,6 +106,7 @@ protected:
 	FSkeletalMeshUpdateHandle UpdateHandle;
 	FDynamicData* DynamicData = nullptr;
 	EType Type;
+	bool bUseGpuLodSelection = false;
 
 	friend class FInstancedSkeletalMeshUpdatePacket;
 };
\ No newline at end of file
diff --git a/Engine/Source/Runtime/Engine/Private/InstancedSkinnedMeshSceneProxy.cpp b/Engine/Source/Runtime/Engine/Private/InstancedSkinnedMeshSceneProxy.cpp
--- a/Engine/Source/Runtime/Engine/Private/InstancedSkinnedMeshSceneProxy.cpp
+++ b/Engine/Source/Runtime/Engine/Private/InstancedSkinnedMeshSceneProxy.cpp
@@ -3,6 +3,7 @@
 #include "InstancedSkinnedMeshSceneProxy.h"
 #include "ComponentRecreateRenderStateContext.h"
 #include "Elements/SMInstance/SMInstanceElementData.h"
+#include "Rendering/SkeletalMeshRenderData.h"
 
 FInstancedSkinnedMeshData::FInstancedSkinnedMeshData(const FInstancedSkinnedMeshSceneProxyDesc& InMeshDesc)
 {
@@ -73,6 +74,7 @@ FInstancedSkinnedMeshSceneProxy::FInstancedSkinnedMeshSceneProxy(const FInstance
 	bAlwaysHasVelocity = true;
 	bInstancedSkinnedMesh = true;
 	bDoesMeshBatchesUseSceneInstanceCount = true;
+	bUseGpuLodSelection = NumLODs > 1 && InMeshDesc.bUseGpuLodSelection && GPUSKIN_TRANSFORM_STORAGE_DIRECT_ENABLED;
 
 	SetupInstanceSceneDataBuffers(Data.InstanceDataSceneProxy->GeInstanceSceneDataBuffers());
 }
\ No newline at end of file
diff --git a/Engine/Source/Runtime/Engine/Private/InstancedSkinnedMeshSceneProxy.h b/Engine/Source/Runtime/Engine/Private/InstancedSkinnedMeshSceneProxy.h
--- a/Engine/Source/Runtime/Engine/Private/InstancedSkinnedMeshSceneProxy.h
+++ b/Engine/Source/Runtime/Engine/Private/InstancedSkinnedMeshSceneProxy.h
@@ -136,6 +136,23 @@ public:
 		return Data.AnimationMinScreenSize;
 	}
 
+	virtual float GetGpuLodInstanceRadius() const override
+	{
+		// Note that StaticMeshBounds.SphereRadius is a better fit, but doesn't match the value on the GPU used for LOD culling.
+		// That's because GPUScene drops the bounds sphere radius and uses the box. So we end up with the sphere encompassing the box, encompassing the sphere :(
+		return bUseGpuLodSelection ? PreSkinnedLocalBounds.BoxExtent.Length() : 0.0f;
+	}
+
+	virtual uint8 GetCurrentFirstLODIdx_RenderThread() const override
+	{
+		return GetCurrentFirstLODIdx_Internal();
+	}
+
+	virtual FDesiredLODLevel GetDesiredLODLevel_RenderThread(const FSceneView* View) const override
+	{
+		return FDesiredLODLevel::CreateFirst(GetCurrentFirstLODIdx_Internal());
+	}
+
 private:
 	FInstancedSkinnedMeshData Data;
 };
\ No newline at end of file
diff --git a/Engine/Source/Runtime/Engine/Private/InstancedSkinnedMeshSceneProxyDesc.cpp b/Engine/Source/Runtime/Engine/Private/InstancedSkinnedMeshSceneProxyDesc.cpp
--- a/Engine/Source/Runtime/Engine/Private/InstancedSkinnedMeshSceneProxyDesc.cpp
+++ b/Engine/Source/Runtime/Engine/Private/InstancedSkinnedMeshSceneProxyDesc.cpp
@@ -92,4 +92,6 @@ void FInstancedSkinnedMeshSceneProxyDesc::InitializeFromInstancedSkinnedMeshComp
 	bAllowAlwaysVisible = true;
 
 	InstanceDataSceneProxy = InComponent->GetInstanceDataSceneProxy();
+
+	bUseGpuLodSelection = InComponent->bUseGpuLodSelection;
 }
\ No newline at end of file
diff --git a/Engine/Source/Runtime/Engine/Private/SkeletalMeshSceneProxy.cpp b/Engine/Source/Runtime/Engine/Private/SkeletalMeshSceneProxy.cpp
--- a/Engine/Source/Runtime/Engine/Private/SkeletalMeshSceneProxy.cpp
+++ b/Engine/Source/Runtime/Engine/Private/SkeletalMeshSceneProxy.cpp
@@ -115,6 +115,7 @@ FSkeletalMeshSceneProxy::FSkeletalMeshSceneProxy(const FSkinnedMeshSceneProxyDes
 		,	bDrawDebugSkeleton(InMeshDesc.ShouldDrawDebugSkeleton())
 #endif
 		,	bAllowDynamicMeshBounds(FMath::IsNearlyEqual(InMeshDesc.BoundsScale, 1.0f))
+		,	bUseGpuLodSelection(false)
 		,	FeatureLevel(GetScene().GetFeatureLevel())
 		,	bMaterialsNeedMorphUsage_GameThread(false)
 		,	MaterialRelevance(InMeshDesc.GetMaterialRelevance(GetScene().GetShaderPlatform()))
@@ -734,6 +735,16 @@ void FSkeletalMeshSceneProxy::GetStaticMeshBatch(
 	BatchElement.IndexBuffer = LODData.MultiSizeIndexContainer.GetIndexBuffer();
 	BatchElement.bFetchInstanceCountFromScene = bInstancedSkinnedMesh;
 	BatchElement.NumInstances = GetInstanceDataHeader().NumInstances;
+
+	if (bUseGpuLodSelection)
+	{
+		BatchElement.MaxScreenSize = MeshObject->GetScreenSize(LODIndex);
+		BatchElement.MinScreenSize = 0.0f;
+		if (LODIndex < MAX_STATIC_MESH_LODS - 1)
+		{
+			BatchElement.MinScreenSize = MeshObject->GetScreenSize(LODIndex + 1);
+		}
+	}
 }
 
 void FSkeletalMeshSceneProxy::GetOverlayMeshBatch(const FSkeletalMeshLODRenderData& LODData, UMaterialInterface* SectionOverlayMaterial, const FMeshBatch& MeshBatch, FMeshBatch& OutOverlayMeshBatch) const
diff --git a/Engine/Source/Runtime/Engine/Public/InstancedSkinnedMeshSceneProxyDesc.h b/Engine/Source/Runtime/Engine/Public/InstancedSkinnedMeshSceneProxyDesc.h
--- a/Engine/Source/Runtime/Engine/Public/InstancedSkinnedMeshSceneProxyDesc.h
+++ b/Engine/Source/Runtime/Engine/Public/InstancedSkinnedMeshSceneProxyDesc.h
@@ -31,4 +31,6 @@ struct FInstancedSkinnedMeshSceneProxyDesc : public FSkinnedMeshSceneProxyDesc
 	/** One bit per instance if the instance is selected. */
 	TBitArray<> SelectedInstances;
 #endif
+
+	bool bUseGpuLodSelection = false;
 };
diff --git a/Engine/Source/Runtime/Engine/Public/SkeletalMeshSceneProxy.h b/Engine/Source/Runtime/Engine/Public/SkeletalMeshSceneProxy.h
--- a/Engine/Source/Runtime/Engine/Public/SkeletalMeshSceneProxy.h
+++ b/Engine/Source/Runtime/Engine/Public/SkeletalMeshSceneProxy.h
@@ -139,12 +139,12 @@ public:
 
 	ENGINE_API virtual void OnTransformChanged(FRHICommandListBase& RHICmdList) override;
 
-	virtual uint8 GetCurrentFirstLODIdx_RenderThread() const final override
+	virtual uint8 GetCurrentFirstLODIdx_RenderThread() const override
 	{
 		return GetCurrentFirstLODIdx_Internal();
 	}
 
-	ENGINE_API virtual FDesiredLODLevel GetDesiredLODLevel_RenderThread(const FSceneView* View) const final override;
+	ENGINE_API virtual FDesiredLODLevel GetDesiredLODLevel_RenderThread(const FSceneView* View) const override;
 
 	ENGINE_API bool GetCachedGeometry(FRDGBuilder& GraphBuilder, struct FCachedGeometry& OutCachedGeometry) const;
 
@@ -158,7 +158,7 @@ protected:
 	TUniquePtr<FSkinningSceneExtensionProxy> SceneExtensionProxy;
 	FSkeletalMeshRenderData* RenderData;
 	int32 ClampedLODIndex;
-	uint8  NumLODs = 1u;
+	uint8 NumLODs = 1u;
 
 #if !(UE_BUILD_SHIPPING || UE_BUILD_TEST)
 	class UPhysicsAsset* PhysicsAssetForDebug;
@@ -186,6 +186,7 @@ protected:
 	uint8 bDrawDebugSkeleton:1;
 #endif
 	uint8 bAllowDynamicMeshBounds:1;
+	uint8 bUseGpuLodSelection:1;
 #if WITH_EDITOR
 	bool bHasSelectedInstances = false;
 #else
