[Misc] Remove commentted out code from previous CL 46771686.

#rnx

diff --git a/Engine/Source/Runtime/Engine/Private/ShaderCompiler/ShaderCompilerThreadRunnable.cpp b/Engine/Source/Runtime/Engine/Private/ShaderCompiler/ShaderCompilerThreadRunnable.cpp
--- a/Engine/Source/Runtime/Engine/Private/ShaderCompiler/ShaderCompilerThreadRunnable.cpp
+++ b/Engine/Source/Runtime/Engine/Private/ShaderCompiler/ShaderCompilerThreadRunnable.cpp
@@ -336,9 +336,6 @@ uint32 FShaderCompileThreadRunnableBase::Run()
 	}
 	Manager->OnShaderCompilingThreadStopped(Thread);
 
-	// Shader compiling threads own their own memory
-	//delete this;
-
 	return 0;
 }
 
