[StaticMesh] Disable fast build on meshes when Nanite is enabled on them in the editor.

Fix for being unable to enable Nanite on meshes that were created with BuildFromMeshDescriptions with bFastBuild=true.

diff --git a/Engine/Source/Runtime/Engine/Private/StaticMesh.cpp b/Engine/Source/Runtime/Engine/Private/StaticMesh.cpp
--- a/Engine/Source/Runtime/Engine/Private/StaticMesh.cpp
+++ b/Engine/Source/Runtime/Engine/Private/StaticMesh.cpp
@@ -5336,6 +5336,12 @@ void UStaticMesh::PostEditChangeProperty(FPropertyChangedEvent& PropertyChangedE
 	if (PropertyName == UStaticMesh::GetNaniteSettingsMemberName())
 	{
 		CheckForMissingShaderModels();
+
+		if (GetNaniteSettings().bEnabled)
+		{
+			// If Nanite was enabled we can no longer do fast build
+			bDoFastBuild = false;
+		}
 	}
 
 	UStaticMesh::FBuildParameters BuildParameters;
