FWorldPartitionLevelHelper::LoadActors() - Always perform an async load under the hood
* This allows us to avoid any potential issue with circular dependencies between actors. In the case of circular dependencies, it was previously possible for CreatePackage() to clash with an existing package that wasn't done loading.

diff --git a/Engine/Source/Runtime/Engine/Private/WorldPartition/WorldPartitionLevelHelper.cpp b/Engine/Source/Runtime/Engine/Private/WorldPartition/WorldPartitionLevelHelper.cpp
--- a/Engine/Source/Runtime/Engine/Private/WorldPartition/WorldPartitionLevelHelper.cpp
+++ b/Engine/Source/Runtime/Engine/Private/WorldPartition/WorldPartitionLevelHelper.cpp
@@ -939,28 +939,32 @@ bool FWorldPartitionLevelHelper::LoadActorsInternal(FLoadActorsParams&& InParams
 		{
 			CompletionCallback.Execute(PackageToLoad, ExistingPackage, EAsyncLoadingResult::Succeeded);
 		}
-		else if (InParams.bLoadAsync)
+		else
 		{
-			check(InParams.DestLevel);
-			const UPackage* DestPackage = InParams.DestLevel->GetPackage();
-			const EPackageFlags PackageFlags = InParams.DestLevel->GetPackage()->HasAnyPackageFlags(PKG_PlayInEditor) ? PKG_PlayInEditor : PKG_None;
+			EPackageFlags PackageFlags = PKG_None;
+			int32 PIEInstanceID = INDEX_NONE;
+
+			// Compute PIE flags/instance only if a DestLevel is provided
+			if (InParams.DestLevel)
+			{
+				const UPackage* DestPackage = InParams.DestLevel->GetPackage();
+				PackageFlags = DestPackage->HasAnyPackageFlags(PKG_PlayInEditor) ? PKG_PlayInEditor : PKG_None;
+				PIEInstanceID = DestPackage->GetPIEInstanceID();
+			}
+
+			// Enqueue async load
 			const FPackagePath PackagePath = FPackagePath::FromPackageNameChecked(PackageToLoad);
-			int32 RequestID = ::LoadPackageAsync(PackagePath, PackageName, CompletionCallback, PackageFlags, DestPackage->GetPIEInstanceID(), 0, &ContainerInstancingContext);
+			const int32 RequestID = ::LoadPackageAsync(PackagePath, PackageName, CompletionCallback, PackageFlags, PIEInstanceID, 0, &ContainerInstancingContext);
 			if (InParams.AsyncRequestIDs)
 			{
 				InParams.AsyncRequestIDs->Add(RequestID);
 			}
-		}
-		else
-		{
-			UPackage* InstancingPackage = nullptr;
-			if (PackageName != PackageToLoad)
+
+			// If a sync load was requested, wait for it
+			if (!InParams.bLoadAsync)
 			{
-				InstancingPackage = CreatePackage(*PackageName.ToString());
+				FlushAsyncLoading(RequestID);
 			}
-
-			UPackage* Package = LoadPackage(InstancingPackage, *PackageToLoad.ToString(), LOAD_None, nullptr, &ContainerInstancingContext);
-			CompletionCallback.Execute(PackageToLoad, Package, Package ? EAsyncLoadingResult::Succeeded : EAsyncLoadingResult::Failed);
 		}
 	}
 
