Fix to allow notifies that are sampled via AnimSingleNodeInstance to store the current playback time.  This means that UAnimNotifyLibrary functions that rely on the current time can be used on these notifies.

diff --git a/Engine/Source/Runtime/Engine/Private/Animation/AnimSingleNodeInstance.cpp b/Engine/Source/Runtime/Engine/Private/Animation/AnimSingleNodeInstance.cpp
--- a/Engine/Source/Runtime/Engine/Private/Animation/AnimSingleNodeInstance.cpp
+++ b/Engine/Source/Runtime/Engine/Private/Animation/AnimSingleNodeInstance.cpp
@@ -371,6 +371,7 @@ void UAnimSingleNodeInstance::SetPositionWithPreviousTime(float InPosition, floa
 		if (SequenceBase)
 		{
 			FAnimTickRecord TickRecord;
+			TickRecord.TimeAccumulator = &InPosition;
 			FAnimNotifyContext NotifyContext(TickRecord);
 			NotifyQueue.Reset(GetSkelMeshComponent());
 
