[Shaders] Fix potential race condition in shader compiling thread handling.

#test Verified potential race condition with RaceDetector in Lyra, which no longer reports this bug with this fix.
#rnx

#FailedValidations Preflight

diff --git a/Engine/Source/Runtime/Engine/Private/ShaderCompiler/ShaderCompiler.cpp b/Engine/Source/Runtime/Engine/Private/ShaderCompiler/ShaderCompiler.cpp
--- a/Engine/Source/Runtime/Engine/Private/ShaderCompiler/ShaderCompiler.cpp
+++ b/Engine/Source/Runtime/Engine/Private/ShaderCompiler/ShaderCompiler.cpp
@@ -1092,13 +1092,13 @@ FShaderCompileThreadRunnableBase* FShaderCompilingManager::FindShaderCompilingTh
 	return nullptr;
 }
 
-void FShaderCompilingManager::OnShaderCompilingThreadStarted(FRunnableThread* /*Thread*/)
+void FShaderCompilingManager::OnShaderCompilingThreadStarted(FShaderCompileThreadRunnableBase* /*Thread*/)
 {
 	// Update assigned priority ranges for all threads
 	AssignShaderCompilingThreadJobPriorities();
 }
 
-void FShaderCompilingManager::OnShaderCompilingThreadStopped(FRunnableThread* Thread)
+void FShaderCompilingManager::OnShaderCompilingThreadStopped(FShaderCompileThreadRunnableBase* Thread)
 {
 	// If the manager is shutting down, ignore this notification
 	if (bIsManagerShuttingDown)
@@ -1109,13 +1109,13 @@ void FShaderCompilingManager::OnShaderCompilingThreadStopped(FRunnableThread* Th
 	// Remove shader compiling thread from list of active threads
 	for (auto Iter = Threads.CreateIterator(); Iter; ++Iter)
 	{
-		if ((*Iter)->Thread == Thread)
+		if (Iter->Get() == Thread)
 		{
 			// Move thread ownership to shutdown task. Its memory needs to be kept alive until the thread has finished.
 			// Let the task system handle this deferred deletion, since this notification callback is called within the shader compiling threads and waiting for its compiltion in this function would deadlock.
 			checkf((*Iter)->bForceFinish, TEXT("%s is signaled to have stopped but is not put into the right state"), (*Iter)->GetThreadName());
 
-			const FString TaskName = FString::Printf(TEXT("ShutdownTask-%s"), *Thread->GetThreadName());
+			const FString TaskName = FString::Printf(TEXT("ShutdownTask-%s"), Thread->GetThreadName());
 			UE::Tasks::Launch(
 				*TaskName,
 				[Thread = MoveTemp(*Iter)](){
diff --git a/Engine/Source/Runtime/Engine/Private/ShaderCompiler/ShaderCompilerThreadRunnable.cpp b/Engine/Source/Runtime/Engine/Private/ShaderCompiler/ShaderCompilerThreadRunnable.cpp
--- a/Engine/Source/Runtime/Engine/Private/ShaderCompiler/ShaderCompilerThreadRunnable.cpp
+++ b/Engine/Source/Runtime/Engine/Private/ShaderCompiler/ShaderCompilerThreadRunnable.cpp
@@ -324,7 +324,7 @@ void FShaderCompileThreadRunnable::ForEachPendingJob(const FShaderCompileJobCall
 /** Entry point for the shader compiling thread. */
 uint32 FShaderCompileThreadRunnableBase::Run()
 {
-	Manager->OnShaderCompilingThreadStarted(Thread);
+	Manager->OnShaderCompilingThreadStarted(this);
 	{
 		LLM_SCOPE_BYTAG(ShaderCompiler);
 		check(Manager->bAllowAsynchronousShaderCompiling);
@@ -334,7 +334,7 @@ uint32 FShaderCompileThreadRunnableBase::Run()
 		}
 		UE_LOG(LogShaderCompilers, Display, TEXT("Exiting %s"), GetThreadName());
 	}
-	Manager->OnShaderCompilingThreadStopped(Thread);
+	Manager->OnShaderCompilingThreadStopped(this);
 
 	return 0;
 }
diff --git a/Engine/Source/Runtime/Engine/Public/ShaderCompiler.h b/Engine/Source/Runtime/Engine/Public/ShaderCompiler.h
--- a/Engine/Source/Runtime/Engine/Public/ShaderCompiler.h
+++ b/Engine/Source/Runtime/Engine/Public/ShaderCompiler.h
@@ -1069,13 +1069,13 @@ private:
 	 * Notifies the manager that one of the shader compiling thread has started.
 	 * This updates the job priority range for all threads depending on how many threads the manager handles.
 	 */
-	void OnShaderCompilingThreadStarted(FRunnableThread* Thread);
+	void OnShaderCompilingThreadStarted(FShaderCompileThreadRunnableBase* Thread);
 
 	/**
 	 * Notifies the manager that one of the shader compiling thread has stopped and must be removed from the list.
 	 * This updates the job priority range for all threads depending on how many threads the manager handles.
 	 */
-	void OnShaderCompilingThreadStopped(FRunnableThread* Thread);
+	void OnShaderCompilingThreadStopped(FShaderCompileThreadRunnableBase* Thread);
 
 	/** Assigns the job priorities each shader compiling thread is responsible for depending on how many threads there are. */
 	void AssignShaderCompilingThreadJobPriorities();
