[Engine]
* Moved Engine/Selection.h include in Engine.h to deprecated scope because it is an include to UnrealEd

diff --git a/Engine/Source/Runtime/Engine/Public/Engine.h b/Engine/Source/Runtime/Engine/Public/Engine.h
--- a/Engine/Source/Runtime/Engine/Public/Engine.h
+++ b/Engine/Source/Runtime/Engine/Public/Engine.h
@@ -285,7 +285,6 @@ MONOLITHIC_HEADER_BOILERPLATE()
 #include "Sound/ReverbEffect.h"
 #include "GameFramework/SaveGame.h"
 #include "Engine/SCS_Node.h"
-#include "Engine/Selection.h"
 #include "Engine/SimpleConstructionScript.h"
 #include "Animation/PreviewAssetAttachComponent.h"
 #include "Engine/SkeletalMeshSocket.h"
@@ -335,4 +334,8 @@ MONOLITHIC_HEADER_BOILERPLATE()
 #include "ParticleHelper.h"
 #endif
 
+#if UE_ENABLE_INCLUDE_ORDER_DEPRECATED_IN_5_7
+#include "Engine/Selection.h"
+#endif
+
 // IWYU pragma: end_keep
