[NetPhysics] Fix issue where sim-proxies could merge inputs incorrectly on the first step after a resim.

diff --git a/Engine/Source/Runtime/Engine/Private/PhysicsEngine/NetworkPhysicsComponent.cpp b/Engine/Source/Runtime/Engine/Private/PhysicsEngine/NetworkPhysicsComponent.cpp
--- a/Engine/Source/Runtime/Engine/Private/PhysicsEngine/NetworkPhysicsComponent.cpp
+++ b/Engine/Source/Runtime/Engine/Private/PhysicsEngine/NetworkPhysicsComponent.cpp
@@ -2060,10 +2060,10 @@ void FAsyncNetworkPhysicsComponent::OnPreProcessInputs_Internal(const int32 Phys
 		if ((bApplyServerInput || bApplySimProxyInput || bIsSolverResim) && InputHistory && InputData)
 		{
 			FNetworkPhysicsPayload* PhysicsData = InputData.Get();
-			int32 NextExpectedLocalFrame = PhysicsData->LocalFrame + 1;
+			int32 NextExpectedLocalFrame = bIsSolverResim ? PhysicsStep : PhysicsData->LocalFrame + 1;
 
 			// There are important inputs earlier than upcoming input to apply
-			if (NewImportantInputFrame < NextExpectedLocalFrame && !bIsSolverResim)
+			if (!bIsSolverResim && NewImportantInputFrame < NextExpectedLocalFrame)
 			{
 				if (ComponentSettings.GetApplyDataInsteadOfMergeData())
 				{
@@ -2120,12 +2120,9 @@ void FAsyncNetworkPhysicsComponent::OnPreProcessInputs_Internal(const int32 Phys
 				// Calculate input decay if we are resimulating and we don't have up to date inputs
 				if (bIsSolverResim)
 				{
-					// If the input was extrapolated, set it back to the input frame it was extrapolated from (the last in the input history) so that we can calculate input decay during resimulation
-					PhysicsData->LocalFrame = FMath::Min(PhysicsData->LocalFrame, InputHistory->GetLatestFrame());
-
-					if (PhysicsData->LocalFrame < PhysicsStep)
+					if (InputHistory->GetLatestFrame() < PhysicsStep)
 					{
-						const float InputDecay = GetCurrentInputDecay(PhysicsData);
+						const float InputDecay = GetCurrentInputDecay();
 						PhysicsData->DecayData(InputDecay);
 					}
 				}
@@ -2211,7 +2208,7 @@ void FAsyncNetworkPhysicsComponent::OnPreProcessInputs_Internal(const int32 Phys
 #if DEBUG_NETWORK_PHYSICS
 				{
 					UE_LOG(LogChaos, Log, TEXT("			Applying extracted input from history | LocalFrame = %d | ServerFrame = %d | bDataAltered = %d | IsResim = %d | IsLocallyControlled = %d | InputDecay = %f | Data: %s")
-						, PhysicsData->LocalFrame, PhysicsData->ServerFrame, PhysicsData->bDataAltered, bIsSolverResim, IsLocallyControlled(), GetCurrentInputDecay(PhysicsData), *PhysicsData->DebugData());
+						, PhysicsData->LocalFrame, PhysicsData->ServerFrame, PhysicsData->bDataAltered, bIsSolverResim, IsLocallyControlled(), GetCurrentInputDecay(), *PhysicsData->DebugData());
 				}
 #endif
 			}
@@ -2503,29 +2500,36 @@ void FAsyncNetworkPhysicsComponent::TriggerResimulation(int32 ResimFrame)
 	}
 }
 
-const float FAsyncNetworkPhysicsComponent::GetCurrentInputDecay(const FNetworkPhysicsPayload* PhysicsData)
+const float FAsyncNetworkPhysicsComponent::GetCurrentInputDecay() const
 {
-	if (!PhysicsData)
+	if (!InputHistory)
 	{
 		return 0.0f;
 	}
-	
-	Chaos::FPhysicsSolverBase* BaseSolver = GetSolver();
-	if (!BaseSolver)
+
+	const Chaos::FPhysicsSolverBase* BaseSolver = GetSolver();
+	if (!BaseSolver || !BaseSolver->IsResimming())
 	{
 		return 0.0f;
 	}
 
-	Chaos::FRewindData* RewindData = BaseSolver->GetRewindData();
+	const Chaos::FRewindData* RewindData = BaseSolver->GetRewindData();
 	if (!RewindData)
 	{
 		return 0.0f;
 	}
 
+	const int32 LatestReceivedFrame = InputHistory->GetLatestFrame();
+
+	if (RewindData->CurrentFrame() <= LatestReceivedFrame)
+	{
+		return 0.0f;
+	}
+
 	const FNetworkPhysicsSettingsNetworkPhysicsComponent& ComponentSettings = GetComponentSettings();
 	const FRuntimeFloatCurve& InputDecayCurve = ComponentSettings.GetInputDecayCurve();
 
-	const float NumPredictedInputs = RewindData->CurrentFrame() - PhysicsData->LocalFrame; // Number of frames we have used the same PhysicsData for during resim
+	const float NumPredictedInputs = RewindData->CurrentFrame() - LatestReceivedFrame; // Number of frames we have used the same PhysicsData for during resim
 	float MaxPredictedInputs = 0;
 
 	if (ComponentSettings.GetApplyInputDecayOverSetTime())
@@ -2536,9 +2540,9 @@ const float FAsyncNetworkPhysicsComponent::GetCurrentInputDecay(const FNetworkPh
 	else
 	{
 		// Number of frames the input will predict and decay over until resimulation has completed
-		MaxPredictedInputs = RewindData->GetLatestFrame() - 1 - PhysicsData->LocalFrame;
+		MaxPredictedInputs = RewindData->GetLatestFrame() - 1 - LatestReceivedFrame;
 	}
-	
+
 	// Linear decay
 	const float PredictionAlpha = MaxPredictedInputs > 0 ? FMath::Clamp(NumPredictedInputs / MaxPredictedInputs, 0.0f, 1.0f) : 0.0f;
 
diff --git a/Engine/Source/Runtime/Engine/Public/Physics/NetworkPhysicsComponent.h b/Engine/Source/Runtime/Engine/Public/Physics/NetworkPhysicsComponent.h
--- a/Engine/Source/Runtime/Engine/Public/Physics/NetworkPhysicsComponent.h
+++ b/Engine/Source/Runtime/Engine/Public/Physics/NetworkPhysicsComponent.h
@@ -1956,7 +1956,11 @@ private:
 	void TriggerResimulation(int32 ResimFrame);
 
 	/** Returns the current amount of input decay during resimulation as a magnitude from 0.0 to 1.0. Returns 0 if not currently resimulating. */
-	const float GetCurrentInputDecay(const FNetworkPhysicsPayload* PhysicsData);
+	UE_DEPRECATED(5.8, "Use function that doesn't take a parameter, the decay curve will be based on the current frame number compared with the latest input frame number available")
+	const float GetCurrentInputDecay(const FNetworkPhysicsPayload* PhysicsData) { return GetCurrentInputDecay(); }
+	
+	/** Returns the current amount of input decay during resimulation as a magnitude from 0.0 to 1.0. Returns 0 if not currently resimulating. */
+	const float GetCurrentInputDecay() const;
 
 	/** Populate input data in AsyncOutput to send over the network */
 	void SendInputData_Internal(FAsyncNetworkPhysicsComponentOutput& AsyncOutput, const int32 PhysicsStep);
diff --git a/Engine/Source/Runtime/Experimental/Chaos/Public/Chaos/Framework/PhysicsSolverBase.h b/Engine/Source/Runtime/Experimental/Chaos/Public/Chaos/Framework/PhysicsSolverBase.h
--- a/Engine/Source/Runtime/Experimental/Chaos/Public/Chaos/Framework/PhysicsSolverBase.h
+++ b/Engine/Source/Runtime/Experimental/Chaos/Public/Chaos/Framework/PhysicsSolverBase.h
@@ -452,11 +452,21 @@ namespace Chaos
 			return MRewindData.Get();
 		}
 
+		const FRewindData* GetRewindData() const
+		{
+			return MRewindData.Get();
+		}
+
 		IRewindCallback* GetRewindCallback()
 		{
 			return MRewindCallback.Get();
 		}
 
+		const IRewindCallback* GetRewindCallback() const
+		{
+			return MRewindCallback.Get();
+		}
+
 		bool ShouldApplyRewindCallbacks()
 		{
 			return MRewindCallback.IsValid() && MRewindData.IsValid();
