Disable r.GPUSkin.UseSceneExtension to avoid a missing uniform buffer crash.

diff --git a/Engine/Source/Runtime/Engine/Private/SkeletalRenderGPUSkin.cpp b/Engine/Source/Runtime/Engine/Private/SkeletalRenderGPUSkin.cpp
--- a/Engine/Source/Runtime/Engine/Private/SkeletalRenderGPUSkin.cpp
+++ b/Engine/Source/Runtime/Engine/Private/SkeletalRenderGPUSkin.cpp
@@ -130,7 +130,7 @@ EBoneTransformStorageMode GetBoneTransformStorageMode()
 #endif
 }
 
-bool GGPUSkinSceneExtensionEnabled = true;
+bool GGPUSkinSceneExtensionEnabled = false;
 static FAutoConsoleVariableRef CVarGPUSkinUseSceneExtension(
 	TEXT("r.GPUSkin.UseSceneExtension"),
 	GGPUSkinSceneExtensionEnabled,
