[NetPhysics] Compare ServerFrame in FRepMovement when physics is used since the same state for different frames should still be considered different.

diff --git a/Engine/Source/Runtime/Engine/Classes/Engine/ReplicatedState.h b/Engine/Source/Runtime/Engine/Classes/Engine/ReplicatedState.h
--- a/Engine/Source/Runtime/Engine/Classes/Engine/ReplicatedState.h
+++ b/Engine/Source/Runtime/Engine/Classes/Engine/ReplicatedState.h
@@ -237,6 +237,11 @@ struct FRepMovement
 			return false;
 		}
 
+		if ( bRepPhysics && (ServerFrame != Other.ServerFrame))
+		{
+			return false;
+		}
+
 		if ( bRepAcceleration != Other.bRepAcceleration )
 		{
 			return false;
