Fix MetaSound interface descriptions localization keys as they are not expanding their macros in the loc tool

#rnx

diff --git a/Engine/Plugins/Runtime/Metasound/Source/MetasoundFrontend/Private/Interfaces/MetasoundFrontendSourceInterface.cpp b/Engine/Plugins/Runtime/Metasound/Source/MetasoundFrontend/Private/Interfaces/MetasoundFrontendSourceInterface.cpp
--- a/Engine/Plugins/Runtime/Metasound/Source/MetasoundFrontend/Private/Interfaces/MetasoundFrontendSourceInterface.cpp
+++ b/Engine/Plugins/Runtime/Metasound/Source/MetasoundFrontend/Private/Interfaces/MetasoundFrontendSourceInterface.cpp
@@ -59,7 +59,7 @@ namespace Metasound::Frontend
 						}
 					};
 #if WITH_EDITOR
-					Description = LOCTEXT(AUDIO_PARAMETER_INTERFACE_NAMESPACE ".Description",
+					Description = LOCTEXT("UE.Source.OneShot.Description",
 						"For MetaSounds that are not looping nor infinite. Provides an output trigger to tell when the MetaSound is finished.");
 #endif
 				}
@@ -109,7 +109,7 @@ namespace Metasound::Frontend
 						}
 					};
 #if WITH_EDITOR
-					Description = LOCTEXT(AUDIO_PARAMETER_INTERFACE_NAMESPACE ".Description",
+					Description = LOCTEXT("UE.Source.StartTime.Description",
 						"Provides the start time (from 0) to seek into the sound as its played. Works from Gameplay API and Sequencer.\n"
 						"The MetaSound itself will need to determine what to do with this input.\n"
 						"In the simplest case, a MetaSound with a single Wave Player can route StartTime into the Wave Player's StartTime.");
diff --git a/Engine/Source/Runtime/Engine/Private/SoundAttenuation.cpp b/Engine/Source/Runtime/Engine/Private/SoundAttenuation.cpp
--- a/Engine/Source/Runtime/Engine/Private/SoundAttenuation.cpp
+++ b/Engine/Source/Runtime/Engine/Private/SoundAttenuation.cpp
@@ -293,7 +293,7 @@ namespace Audio
 						}
 					};
 #if WITH_EDITOR
-					Description = LOCTEXT(AUDIO_PARAMETER_INTERFACE_NAMESPACE ".Description", "Provides the distance between listener and sound locations in Unreal Units.");
+					Description = LOCTEXT("UE.Attenuation.Description", "Provides the distance between listener and sound locations in Unreal Units.");
 #endif
 				}
 			};
@@ -343,7 +343,7 @@ namespace Audio
 						}
 					};
 #if WITH_EDITOR
-					Description = LOCTEXT(AUDIO_PARAMETER_INTERFACE_NAMESPACE ".Description",
+					Description = LOCTEXT("UE.Spatialization.Description",
 						"From the perspective of the listener, the horizontal angle (azimuth) and vertical angle (elevation) to the sound.\n"
 						"Useful for focus-based mixing, like that of the Focus feature of Attenuation Settings.");
 #endif
@@ -395,7 +395,7 @@ namespace Audio
 						}
 					};
 #if WITH_EDITOR
-					Description = LOCTEXT(AUDIO_PARAMETER_INTERFACE_NAMESPACE ".Description",
+					Description = LOCTEXT("UE.Source.Orientation.Description",
 						"From the perspective of the sound, provides the horizontal angle (azimuth) and vertical angle (elevation) to the listener.\n"
 						"Useful for directional sources like a megaphone or speaker, where the sound is louder pointed at you and quieter pointed away.");
 #endif
@@ -447,7 +447,7 @@ namespace Audio
 						}
 					};
 #if WITH_EDITOR
-					Description = LOCTEXT(AUDIO_PARAMETER_INTERFACE_NAMESPACE ".Description",
+					Description = LOCTEXT("UE.Listener.Orientation.Description",
 						"Relative to the world, provides the listener orientation, as horizontal angle (azimuth) and vertical angle (elevation).\n"
 						"Useful for world-locked quad beds for example.");
 #endif
