Collision Filter Refactor Part 3.9: Updating BodyInstance and FBodyCollisionFilterData to use new API and structures.

diff --git a/Engine/Plugins/Experimental/ControlRigPhysics/Source/ControlRigPhysics/Private/RigPhysicsSimulation.cpp b/Engine/Plugins/Experimental/ControlRigPhysics/Source/ControlRigPhysics/Private/RigPhysicsSimulation.cpp
--- a/Engine/Plugins/Experimental/ControlRigPhysics/Source/ControlRigPhysics/Private/RigPhysicsSimulation.cpp
+++ b/Engine/Plugins/Experimental/ControlRigPhysics/Source/ControlRigPhysics/Private/RigPhysicsSimulation.cpp
@@ -406,9 +406,8 @@ static bool CreateGeometry(
 	FBodyCollisionData BodyCollisionData;
 
 	const Chaos::Filter::FShapeFilterData ShapeFilter = Chaos::Filter::FShapeFilterBuilder::BuildSimBlockAll(Chaos::EFilterFlags::All);
-	Chaos::Filter::FCombinedShapeFilterData CombinedShapeFilter(ShapeFilter, Chaos::Filter::FInstanceData());
-	Chaos::Filter::FShapeFilterBuilder::GetLegacyShapeFilter(CombinedShapeFilter, BodyCollisionData.CollisionFilterData.QuerySimpleFilter, BodyCollisionData.CollisionFilterData.SimFilter);
-	BodyCollisionData.CollisionFilterData.QueryComplexFilter = BodyCollisionData.CollisionFilterData.QuerySimpleFilter;
+	BodyCollisionData.CollisionFilterData.ComplexShapeFilterData = BodyCollisionData.CollisionFilterData.SimpleShapeFilterData = ShapeFilter;
+	BodyCollisionData.CollisionFilterData.FilterInstanceData = Chaos::Filter::FInstanceData();
 
 	// See FBodyInstance::BuildBodyCollisionFlags
 	BodyCollisionData.CollisionFlags.bEnableQueryCollision = false;
diff --git a/Engine/Source/Runtime/Engine/Private/Collision/PhysicsFiltering.cpp b/Engine/Source/Runtime/Engine/Private/Collision/PhysicsFiltering.cpp
--- a/Engine/Source/Runtime/Engine/Private/Collision/PhysicsFiltering.cpp
+++ b/Engine/Source/Runtime/Engine/Private/Collision/PhysicsFiltering.cpp
@@ -29,6 +29,12 @@ FPhysicsFilterBuilder::FPhysicsFilterBuilder(const FChaosScene* InScene)
 {
 }
 
+FPhysicsFilterBuilder::FPhysicsFilterBuilder(const FChaosScene* InScene, const Chaos::Filter::FShapeFilterData& InShapeFilter, const Chaos::Filter::FInstanceData& InInstanceData)
+	: InstanceData(InInstanceData)
+	, ShapeFilterBuilder(InShapeFilter)
+{
+}
+
 FPhysicsFilterBuilder::FPhysicsFilterBuilder(TEnumAsByte<enum ECollisionChannel> InObjectType, FMaskFilter MaskFilter, const struct FCollisionResponseContainer& ResponseToChannels)
 {
 	SetCollisionChannelIndex(InObjectType);
diff --git a/Engine/Source/Runtime/Engine/Private/PhysicsEngine/BodyInstance.cpp b/Engine/Source/Runtime/Engine/Private/PhysicsEngine/BodyInstance.cpp
--- a/Engine/Source/Runtime/Engine/Private/PhysicsEngine/BodyInstance.cpp
+++ b/Engine/Source/Runtime/Engine/Private/PhysicsEngine/BodyInstance.cpp
@@ -1106,9 +1106,9 @@ void FBodyInstance::UpdatePhysicsFilterData()
 				}
 
 				// Apply new collision settings to this shape
-				const FCollisionFilterData& QueryFilter = bIsTrimesh ? FilterData.QueryComplexFilter : FilterData.QuerySimpleFilter;
-				Chaos::Filter::FCombinedShapeFilterData CombinedShapeFilter = Chaos::Filter::FShapeFilterBuilder::BuildFromLegacyShapeFilter(QueryFilter, FilterData.SimFilter);
-				FPhysicsInterface::SetCombinedShapeFilterData(InnerShape, CombinedShapeFilter);
+				const Chaos::Filter::FShapeFilterData& ShapeFilterData = bIsTrimesh ? FilterData.ComplexShapeFilterData : FilterData.SimpleShapeFilterData;
+				FPhysicsInterface::SetShapeFilterData(InnerShape, ShapeFilterData);
+				FPhysicsInterface::SetFilterInstanceData(InnerShape, FilterData.FilterInstanceData);
 			});
 		}
 
@@ -4761,16 +4761,16 @@ void FBodyInstance::BuildBodyFilterData(FBodyCollisionFilterData& OutFilterData,
 		Builder.SetFlags(Chaos::EFilterFlags::StaticShape, bPhysicsStatic);
 		Builder.SetFlags(Chaos::EFilterFlags::ModifyContacts, bUseContactModification);
 
+		OutFilterData.FilterInstanceData = Builder.BuildInstanceData();
+
 		// Build shape filter variations for complex and simple
 		Builder.SetFlags(Chaos::EFilterFlags::SimpleCollision, true);
 		Builder.SetFlags(Chaos::EFilterFlags::ComplexCollision, bUseSimpleAsComplex);
-		const Chaos::Filter::FCombinedShapeFilterData SimpleCombinedFilter(Builder.BuildShapeFilterData(), Builder.BuildInstanceData());
-		Chaos::Filter::FShapeFilterBuilder::GetLegacyShapeFilter(SimpleCombinedFilter, OutFilterData.QuerySimpleFilter, OutFilterData.SimFilter);
+		OutFilterData.SimpleShapeFilterData = Builder.BuildShapeFilterData();
 
 		Builder.SetFlags(Chaos::EFilterFlags::ComplexCollision, true);
 		Builder.SetFlags(Chaos::EFilterFlags::SimpleCollision, bUseComplexAsSimple);
-		const Chaos::Filter::FCombinedShapeFilterData ComplexCombinedFilter(Builder.BuildShapeFilterData(), Builder.BuildInstanceData());
-		Chaos::Filter::FShapeFilterBuilder::GetLegacyShapeFilter(ComplexCombinedFilter, OutFilterData.QueryComplexFilter, OutFilterData.SimFilter);
+		OutFilterData.ComplexShapeFilterData = Builder.BuildShapeFilterData();
 	}
 }
 
diff --git a/Engine/Source/Runtime/Engine/Private/PhysicsEngine/Experimental/ChaosInterfaceUtils.cpp b/Engine/Source/Runtime/Engine/Private/PhysicsEngine/Experimental/ChaosInterfaceUtils.cpp
--- a/Engine/Source/Runtime/Engine/Private/PhysicsEngine/Experimental/ChaosInterfaceUtils.cpp
+++ b/Engine/Source/Runtime/Engine/Private/PhysicsEngine/Experimental/ChaosInterfaceUtils.cpp
@@ -103,9 +103,10 @@ namespace ChaosInterface
 		auto NewShapeHelper = [&InParams, &CollisionTraceType](const Chaos::FImplicitObjectPtr& InGeom, int32 ShapeIdx, void* UserData, ECollisionEnabled::Type ShapeCollisionEnabled, bool bComplexShape = false)
 		{
 			TUniquePtr<Chaos::FPerShapeData> NewShape = Chaos::FShapeInstanceProxy::Make(ShapeIdx, InGeom);
-			const FCollisionFilterData& QueryFilter = bComplexShape ? InParams.CollisionData.CollisionFilterData.QueryComplexFilter : InParams.CollisionData.CollisionFilterData.QuerySimpleFilter;
-			Chaos::Filter::FCombinedShapeFilterData CombinedShapeFilter = Chaos::Filter::FShapeFilterBuilder::BuildFromLegacyShapeFilter(QueryFilter, InParams.CollisionData.CollisionFilterData.SimFilter);
-			NewShape->SetCombinedShapeFilterData(CombinedShapeFilter);
+			const FBodyCollisionFilterData& CollisionFilterData = InParams.CollisionData.CollisionFilterData;
+			const Chaos::Filter::FShapeFilterData& ShapeFilterData = bComplexShape ? CollisionFilterData.ComplexShapeFilterData : CollisionFilterData.SimpleShapeFilterData;
+			NewShape->SetShapeFilterData(ShapeFilterData);
+			NewShape->SetFilterInstanceData(CollisionFilterData.FilterInstanceData);
 			NewShape->SetCollisionTraceType(ConvertCollisionTraceFlag(CollisionTraceType));
 			NewShape->UpdateShapeBounds(InParams.WorldTransform);
 			NewShape->SetUserData(UserData);
diff --git a/Engine/Source/Runtime/Engine/Private/PhysicsEngine/ImmediatePhysics/ImmediatePhysicsAdapters.cpp b/Engine/Source/Runtime/Engine/Private/PhysicsEngine/ImmediatePhysics/ImmediatePhysicsAdapters.cpp
--- a/Engine/Source/Runtime/Engine/Private/PhysicsEngine/ImmediatePhysics/ImmediatePhysicsAdapters.cpp
+++ b/Engine/Source/Runtime/Engine/Private/PhysicsEngine/ImmediatePhysics/ImmediatePhysicsAdapters.cpp
@@ -125,9 +125,8 @@ namespace ImmediatePhysics
 		FBodyCollisionData BodyCollisionData;
 
 		const Chaos::Filter::FShapeFilterData ShapeFilterData = Chaos::Filter::FShapeFilterBuilder::BuildSimBlockAll(EFilterFlags::All);
-		const Chaos::Filter::FCombinedShapeFilterData CombinedShapeFilter(ShapeFilterData, Chaos::Filter::FInstanceData());
-		Chaos::Filter::FShapeFilterBuilder::GetLegacyShapeFilter(CombinedShapeFilter, BodyCollisionData.CollisionFilterData.QuerySimpleFilter, BodyCollisionData.CollisionFilterData.SimFilter);
-		BodyCollisionData.CollisionFilterData.QueryComplexFilter = BodyCollisionData.CollisionFilterData.QuerySimpleFilter;
+		BodyCollisionData.CollisionFilterData.SimpleShapeFilterData = ShapeFilterData;
+		BodyCollisionData.CollisionFilterData.ComplexShapeFilterData = ShapeFilterData;
 		BodyCollisionData.CollisionFlags.bEnableSimCollisionSimple = true;
 
 		BodyInstance->BuildBodyCollisionFlags(BodyCollisionData.CollisionFlags, BodyInstance->GetCollisionEnabled(), BodyInstance->BodySetup->GetCollisionTraceFlag() == CTF_UseComplexAsSimple);
diff --git a/Engine/Source/Runtime/Engine/Public/Physics/PhysicsFiltering.h b/Engine/Source/Runtime/Engine/Public/Physics/PhysicsFiltering.h
--- a/Engine/Source/Runtime/Engine/Public/Physics/PhysicsFiltering.h
+++ b/Engine/Source/Runtime/Engine/Public/Physics/PhysicsFiltering.h
@@ -41,6 +41,7 @@ struct FPhysicsFilterBuilder
 	FPhysicsFilterBuilder(const FPhysicsFilterBuilder&) = default;
 	FPhysicsFilterBuilder(FPhysicsFilterBuilder&&) = default;
 	ENGINE_API FPhysicsFilterBuilder(const FChaosScene* InScene);
+	ENGINE_API FPhysicsFilterBuilder(const FChaosScene* InScene, const Chaos::Filter::FShapeFilterData& InShapeFilter, const Chaos::Filter::FInstanceData& InInstanceData);
 	UE_DEPRECATED(5.8, "Use The constructor that takes a scene instead")
 	ENGINE_API FPhysicsFilterBuilder(TEnumAsByte<enum ECollisionChannel> InObjectType, FMaskFilter MaskFilter, const struct FCollisionResponseContainer& ResponseToChannels);
 
diff --git a/Engine/Source/Runtime/Experimental/GeometryCollectionEngine/Private/GeometryCollection/GeometryCollectionComponent.cpp b/Engine/Source/Runtime/Experimental/GeometryCollectionEngine/Private/GeometryCollection/GeometryCollectionComponent.cpp
--- a/Engine/Source/Runtime/Experimental/GeometryCollectionEngine/Private/GeometryCollection/GeometryCollectionComponent.cpp
+++ b/Engine/Source/Runtime/Experimental/GeometryCollectionEngine/Private/GeometryCollection/GeometryCollectionComponent.cpp
@@ -5003,8 +5003,8 @@ void UGeometryCollectionComponent::BuildInitialFilterData()
 	FMaskFilter FilterMask = BodyInstance.GetMaskFilter();
 	BodyInstance.BuildBodyFilterData(FilterData);
 
-	InitialSimFilter = FilterData.SimFilter;
-	InitialQueryFilter = FilterData.QuerySimpleFilter;
+	Chaos::Filter::FCombinedShapeFilterData CombinedFilter(FilterData.SimpleShapeFilterData, FilterData.FilterInstanceData);
+	Chaos::Filter::FShapeFilterBuilder::GetLegacyShapeFilter(CombinedFilter, InitialQueryFilter, InitialSimFilter);
 
 	// Enable for complex and simple (no dual representation currently like other meshes)
 	InitialQueryFilter.Word3 |= (EPDF_SimpleCollision | EPDF_ComplexCollision);
diff --git a/Engine/Source/Runtime/PhysicsCore/Public/PhysicsInterfaceTypesCore.h b/Engine/Source/Runtime/PhysicsCore/Public/PhysicsInterfaceTypesCore.h
--- a/Engine/Source/Runtime/PhysicsCore/Public/PhysicsInterfaceTypesCore.h
+++ b/Engine/Source/Runtime/PhysicsCore/Public/PhysicsInterfaceTypesCore.h
@@ -60,9 +60,26 @@ enum class ECollisionShapeType : uint8
 /** Helper struct holding physics body filter data during initialisation */
 struct FBodyCollisionFilterData
 {
+	// These must be explicitly declared to get rid of deprecation warnings.
+	PRAGMA_DISABLE_DEPRECATION_WARNINGS
+	FBodyCollisionFilterData() = default;
+	FBodyCollisionFilterData(const FBodyCollisionFilterData&) = default;
+	FBodyCollisionFilterData(FBodyCollisionFilterData&&) = default;
+
+	FBodyCollisionFilterData& operator=(const FBodyCollisionFilterData&) = default;
+	FBodyCollisionFilterData& operator=(FBodyCollisionFilterData&&) = default;
+	PRAGMA_ENABLE_DEPRECATION_WARNINGS
+
+	UE_DEPRECATED(5.8, "Use the FilterInstanceData and Simple/Complex ShapeFilterData objects instead")
 	FCollisionFilterData SimFilter;
+	UE_DEPRECATED(5.8, "Use the FilterInstanceData and Simple/Complex ShapeFilterData objects instead")
 	FCollisionFilterData QuerySimpleFilter;
+	UE_DEPRECATED(5.8, "Use the FilterInstanceData and Simple/Complex ShapeFilterData objects instead")
 	FCollisionFilterData QueryComplexFilter;
+
+	Chaos::Filter::FInstanceData FilterInstanceData;
+	Chaos::Filter::FShapeFilterData SimpleShapeFilterData;
+	Chaos::Filter::FShapeFilterData ComplexShapeFilterData;
 };
 
 struct FBodyCollisionFlags
