Collision Filter Refactor Part 3.10: Updating CollisionConversions to use new API and structures.

diff --git a/Engine/Source/Runtime/Engine/Private/Collision/CollisionConversions.cpp b/Engine/Source/Runtime/Engine/Private/Collision/CollisionConversions.cpp
--- a/Engine/Source/Runtime/Engine/Private/Collision/CollisionConversions.cpp
+++ b/Engine/Source/Runtime/Engine/Private/Collision/CollisionConversions.cpp
@@ -61,7 +61,7 @@ extern int32 CVarShowInitialOverlaps;
 
 // Forward declare, I don't want to move the entire function right now or we lose change history.
 template <typename THitLocation>
-static bool ConvertOverlappedShapeToImpactHit(const UWorld* World, const THitLocation& Hit, const FVector& StartLoc, const FVector& EndLoc, FHitResult& OutResult, const FPhysicsGeometry& Geom, const FTransform& QueryTM, const FCollisionFilterData& QueryFilter, bool bReturnPhysMat);
+static bool ConvertOverlappedShapeToImpactHit(const UWorld* World, const THitLocation& Hit, const FVector& StartLoc, const FVector& EndLoc, FHitResult& OutResult, const FPhysicsGeometry& Geom, const FTransform& QueryTM, const Chaos::Filter::FQueryFilterData& QueryFilter, bool bReturnPhysMat);
 
 DECLARE_CYCLE_STAT(TEXT("ConvertQueryHit"), STAT_ConvertQueryImpactHit, STATGROUP_Collision);
 DECLARE_CYCLE_STAT(TEXT("ConvertOverlapToHit"), STAT_CollisionConvertOverlapToHit, STATGROUP_Collision);
@@ -280,8 +280,9 @@ const FPhysicsActor* GetGTActor(const Chaos::FGeometryParticleHandle* PTActor)
 }
 
 template <typename THitLocation>
-EConvertQueryResult ConvertQueryImpactHitImp(const UWorld* World, const THitLocation& Hit, FHitResult& OutResult, float CheckLength, const FCollisionFilterData& QueryFilter, const FVector& StartLoc, const FVector& EndLoc, const FPhysicsGeometry* Geom, const FTransform& QueryTM, bool bReturnFaceIndex, bool bReturnPhysMat)
+EConvertQueryResult ConvertQueryImpactHitImp(const UWorld* World, const THitLocation& Hit, FHitResult& OutResult, float CheckLength, const Chaos::Filter::FQueryFilterData& QueryFilter, const FVector& StartLoc, const FVector& EndLoc, const FPhysicsGeometry* Geom, const FTransform& QueryTM, bool bReturnFaceIndex, bool bReturnPhysMat)
 {
+	const FCollisionFilterData LegacyQueryFilter = Chaos::Filter::FQueryFilterBuilder::GetLegacyQueryFilter(QueryFilter);
 	using namespace ChaosInterface;
 
 	SCOPE_CYCLE_COUNTER(STAT_ConvertQueryImpactHit);
@@ -348,7 +349,7 @@ EConvertQueryResult ConvertQueryImpactHitImp(const UWorld* World, const THitLoca
 
 	// See if this is a 'blocking' hit
 	const FCollisionFilterData ShapeFilter = Chaos::Filter::FShapeFilterBuilder::GetLegacyShapeQueryFilter(GetCombinedShapeFilterData(HitShape));
-	const ECollisionQueryHitType HitType = FCollisionQueryFilterCallback::CalcQueryHitType(QueryFilter, ShapeFilter);
+	const ECollisionQueryHitType HitType = FCollisionQueryFilterCallback::CalcQueryHitType(LegacyQueryFilter, ShapeFilter);
 	OutResult.bBlockingHit = (HitType == ECollisionQueryHitType::Block);
 	OutResult.bStartPenetrating = bInitialOverlap;
 
@@ -465,17 +466,17 @@ EConvertQueryResult ConvertQueryImpactHitImp(const UWorld* World, const THitLoca
 	return EConvertQueryResult::Valid;
 }
 
-EConvertQueryResult ConvertQueryImpactHit(const UWorld* World, const FHitLocation& Hit, FHitResult& OutResult, float CheckLength, const FCollisionFilterData& QueryFilter, const FVector& StartLoc, const FVector& EndLoc, const FPhysicsGeometry* Geom, const FTransform& QueryTM, bool bReturnFaceIndex, bool bReturnPhysMat)
+EConvertQueryResult ConvertQueryImpactHit(const UWorld* World, const FHitLocation& Hit, FHitResult& OutResult, float CheckLength, const Chaos::Filter::FQueryFilterData& QueryFilter, const FVector& StartLoc, const FVector& EndLoc, const FPhysicsGeometry* Geom, const FTransform& QueryTM, bool bReturnFaceIndex, bool bReturnPhysMat)
 {
 	return ConvertQueryImpactHitImp(World, Hit, OutResult, CheckLength, QueryFilter, StartLoc, EndLoc, Geom, QueryTM, bReturnFaceIndex, bReturnPhysMat);
 }
-EConvertQueryResult ConvertQueryImpactHit(const UWorld* World, const ChaosInterface::FPTLocationHit& Hit, FHitResult& OutResult, float CheckLength, const FCollisionFilterData& QueryFilter, const FVector& StartLoc, const FVector& EndLoc, const FPhysicsGeometry* Geom, const FTransform& QueryTM, bool bReturnFaceIndex, bool bReturnPhysMat)
+EConvertQueryResult ConvertQueryImpactHit(const UWorld* World, const ChaosInterface::FPTLocationHit& Hit, FHitResult& OutResult, float CheckLength, const Chaos::Filter::FQueryFilterData& QueryFilter, const FVector& StartLoc, const FVector& EndLoc, const FPhysicsGeometry* Geom, const FTransform& QueryTM, bool bReturnFaceIndex, bool bReturnPhysMat)
 {
 	return ConvertQueryImpactHitImp(World, Hit, OutResult, CheckLength, QueryFilter, StartLoc, EndLoc, Geom, QueryTM, bReturnFaceIndex, bReturnPhysMat);
 }
 
 template <typename HitType>
-EConvertQueryResult ConvertTraceResults(bool& OutHasValidBlockingHit, const UWorld* World, int32 NumHits, HitType* Hits, float CheckLength, const FCollisionFilterData& QueryFilter, TArray<FHitResult>& OutHits, const FVector& StartLoc, const FVector& EndLoc, const FPhysicsGeometry* Geom, const FTransform& QueryTM, float MaxDistance, bool bReturnFaceIndex, bool bReturnPhysMat)
+EConvertQueryResult ConvertTraceResults(bool& OutHasValidBlockingHit, const UWorld* World, int32 NumHits, HitType* Hits, float CheckLength, const Chaos::Filter::FQueryFilterData& QueryFilter, TArray<FHitResult>& OutHits, const FVector& StartLoc, const FVector& EndLoc, const FPhysicsGeometry* Geom, const FTransform& QueryTM, float MaxDistance, bool bReturnFaceIndex, bool bReturnPhysMat)
 {
 	OutHits.Reserve(OutHits.Num() + NumHits);
 	EConvertQueryResult ConvertResult = EConvertQueryResult::Valid;
@@ -517,7 +518,7 @@ EConvertQueryResult ConvertTraceResults(bool& OutHasValidBlockingHit, const UWor
 }
 
 template <typename Hit>
-EConvertQueryResult ConvertTraceResults(bool& OutHasValidBlockingHit, const UWorld* World, int32 NumHits, Hit* Hits, float CheckLength, const FCollisionFilterData& QueryFilter, FHitResult& OutHit, const FVector& StartLoc, const FVector& EndLoc, const FPhysicsGeometry* Geom, const FTransform& QueryTM, float MaxDistance, bool bReturnFaceIndex, bool bReturnPhysMat)
+EConvertQueryResult ConvertTraceResults(bool& OutHasValidBlockingHit, const UWorld* World, int32 NumHits, Hit* Hits, float CheckLength, const Chaos::Filter::FQueryFilterData& QueryFilter, FHitResult& OutHit, const FVector& StartLoc, const FVector& EndLoc, const FPhysicsGeometry* Geom, const FTransform& QueryTM, float MaxDistance, bool bReturnFaceIndex, bool bReturnPhysMat)
 {
 	const FVector Dir = (EndLoc - StartLoc).GetSafeNormal();
 	if constexpr (std::is_same_v<Hit, FHitSweep> || std::is_same_v<Hit, ChaosInterface::FPTSweepHit>)
@@ -532,18 +533,18 @@ EConvertQueryResult ConvertTraceResults(bool& OutHasValidBlockingHit, const UWor
 	return Result;
 }
 
-template EConvertQueryResult ConvertTraceResults<FHitSweep>(bool& OutHasValidBlockingHit, const UWorld* World, int32 NumHits, FHitSweep* Hits, float CheckLength, const FCollisionFilterData& QueryFilter, TArray<FHitResult>& OutHits, const FVector& StartLoc, const FVector& EndLoc, const FPhysicsGeometry* Geom, const FTransform& QueryTM, float MaxDistance, bool bReturnFaceIndex, bool bReturnPhysMat);
-template EConvertQueryResult ConvertTraceResults<FHitSweep>(bool& OutHasValidBlockingHit, const UWorld* World, int32 NumHits, FHitSweep* Hits, float CheckLength, const FCollisionFilterData& QueryFilter, FHitResult& OutHit, const FVector& StartLoc, const FVector& EndLoc, const FPhysicsGeometry* Geom, const FTransform& QueryTM, float MaxDistance, bool bReturnFaceIndex, bool bReturnPhysMat);
-template EConvertQueryResult ConvertTraceResults<ChaosInterface::FPTSweepHit>(bool& OutHasValidBlockingHit, const UWorld* World, int32 NumHits, ChaosInterface::FPTSweepHit* Hits, float CheckLength, const FCollisionFilterData& QueryFilter, TArray<FHitResult>& OutHits, const FVector& StartLoc, const FVector& EndLoc, const FPhysicsGeometry* Geom, const FTransform& QueryTM, float MaxDistance, bool bReturnFaceIndex, bool bReturnPhysMat);
-template EConvertQueryResult ConvertTraceResults<ChaosInterface::FPTSweepHit>(bool& OutHasValidBlockingHit, const UWorld* World, int32 NumHits, ChaosInterface::FPTSweepHit* Hits, float CheckLength, const FCollisionFilterData& QueryFilter, FHitResult& OutHit, const FVector& StartLoc, const FVector& EndLoc, const FPhysicsGeometry* Geom, const FTransform& QueryTM, float MaxDistance, bool bReturnFaceIndex, bool bReturnPhysMat);
-template EConvertQueryResult ConvertTraceResults<FHitRaycast>(bool& OutHasValidBlockingHit, const UWorld* World, int32 NumHits, FHitRaycast* Hits, float CheckLength, const FCollisionFilterData& QueryFilter, TArray<FHitResult>& OutHits, const FVector& StartLoc, const FVector& EndLoc, const FPhysicsGeometry* Geom, const FTransform& QueryTM, float MaxDistance, bool bReturnFaceIndex, bool bReturnPhysMat);
-template EConvertQueryResult ConvertTraceResults<FHitRaycast>(bool& OutHasValidBlockingHit, const UWorld* World, int32 NumHits, FHitRaycast* Hits, float CheckLength, const FCollisionFilterData& QueryFilter, FHitResult& OutHit, const FVector& StartLoc, const FVector& EndLoc, const FPhysicsGeometry* Geom, const FTransform& QueryTM, float MaxDistance, bool bReturnFaceIndex, bool bReturnPhysMat);
-template EConvertQueryResult ConvertTraceResults<ChaosInterface::FPTRaycastHit>(bool& OutHasValidBlockingHit, const UWorld* World, int32 NumHits, ChaosInterface::FPTRaycastHit* Hits, float CheckLength, const FCollisionFilterData& QueryFilter, TArray<FHitResult>& OutHits, const FVector& StartLoc, const FVector& EndLoc, const FPhysicsGeometry* Geom, const FTransform& QueryTM, float MaxDistance, bool bReturnFaceIndex, bool bReturnPhysMat);
-template EConvertQueryResult ConvertTraceResults<ChaosInterface::FPTRaycastHit>(bool& OutHasValidBlockingHit, const UWorld* World, int32 NumHits, ChaosInterface::FPTRaycastHit* Hits, float CheckLength, const FCollisionFilterData& QueryFilter, FHitResult& OutHit, const FVector& StartLoc, const FVector& EndLoc, const FPhysicsGeometry* Geom, const FTransform& QueryTM, float MaxDistance, bool bReturnFaceIndex, bool bReturnPhysMat);
+template EConvertQueryResult ConvertTraceResults<FHitSweep>(bool& OutHasValidBlockingHit, const UWorld* World, int32 NumHits, FHitSweep* Hits, float CheckLength, const Chaos::Filter::FQueryFilterData& QueryFilter, TArray<FHitResult>& OutHits, const FVector& StartLoc, const FVector& EndLoc, const FPhysicsGeometry* Geom, const FTransform& QueryTM, float MaxDistance, bool bReturnFaceIndex, bool bReturnPhysMat);
+template EConvertQueryResult ConvertTraceResults<FHitSweep>(bool& OutHasValidBlockingHit, const UWorld* World, int32 NumHits, FHitSweep* Hits, float CheckLength, const Chaos::Filter::FQueryFilterData& QueryFilter, FHitResult& OutHit, const FVector& StartLoc, const FVector& EndLoc, const FPhysicsGeometry* Geom, const FTransform& QueryTM, float MaxDistance, bool bReturnFaceIndex, bool bReturnPhysMat);
+template EConvertQueryResult ConvertTraceResults<ChaosInterface::FPTSweepHit>(bool& OutHasValidBlockingHit, const UWorld* World, int32 NumHits, ChaosInterface::FPTSweepHit* Hits, float CheckLength, const Chaos::Filter::FQueryFilterData& QueryFilter, TArray<FHitResult>& OutHits, const FVector& StartLoc, const FVector& EndLoc, const FPhysicsGeometry* Geom, const FTransform& QueryTM, float MaxDistance, bool bReturnFaceIndex, bool bReturnPhysMat);
+template EConvertQueryResult ConvertTraceResults<ChaosInterface::FPTSweepHit>(bool& OutHasValidBlockingHit, const UWorld* World, int32 NumHits, ChaosInterface::FPTSweepHit* Hits, float CheckLength, const Chaos::Filter::FQueryFilterData& QueryFilter, FHitResult& OutHit, const FVector& StartLoc, const FVector& EndLoc, const FPhysicsGeometry* Geom, const FTransform& QueryTM, float MaxDistance, bool bReturnFaceIndex, bool bReturnPhysMat);
+template EConvertQueryResult ConvertTraceResults<FHitRaycast>(bool& OutHasValidBlockingHit, const UWorld* World, int32 NumHits, FHitRaycast* Hits, float CheckLength, const Chaos::Filter::FQueryFilterData& QueryFilter, TArray<FHitResult>& OutHits, const FVector& StartLoc, const FVector& EndLoc, const FPhysicsGeometry* Geom, const FTransform& QueryTM, float MaxDistance, bool bReturnFaceIndex, bool bReturnPhysMat);
+template EConvertQueryResult ConvertTraceResults<FHitRaycast>(bool& OutHasValidBlockingHit, const UWorld* World, int32 NumHits, FHitRaycast* Hits, float CheckLength, const Chaos::Filter::FQueryFilterData& QueryFilter, FHitResult& OutHit, const FVector& StartLoc, const FVector& EndLoc, const FPhysicsGeometry* Geom, const FTransform& QueryTM, float MaxDistance, bool bReturnFaceIndex, bool bReturnPhysMat);
+template EConvertQueryResult ConvertTraceResults<ChaosInterface::FPTRaycastHit>(bool& OutHasValidBlockingHit, const UWorld* World, int32 NumHits, ChaosInterface::FPTRaycastHit* Hits, float CheckLength, const Chaos::Filter::FQueryFilterData& QueryFilter, TArray<FHitResult>& OutHits, const FVector& StartLoc, const FVector& EndLoc, const FPhysicsGeometry* Geom, const FTransform& QueryTM, float MaxDistance, bool bReturnFaceIndex, bool bReturnPhysMat);
+template EConvertQueryResult ConvertTraceResults<ChaosInterface::FPTRaycastHit>(bool& OutHasValidBlockingHit, const UWorld* World, int32 NumHits, ChaosInterface::FPTRaycastHit* Hits, float CheckLength, const Chaos::Filter::FQueryFilterData& QueryFilter, FHitResult& OutHit, const FVector& StartLoc, const FVector& EndLoc, const FPhysicsGeometry* Geom, const FTransform& QueryTM, float MaxDistance, bool bReturnFaceIndex, bool bReturnPhysMat);
 
 /** Util to convert an overlapped shape into a sweep hit result, returns whether it was a blocking hit. */
 template <typename THitLocation>
-static bool ConvertOverlappedShapeToImpactHit(const UWorld* World, const THitLocation& Hit, const FVector& StartLoc, const FVector& EndLoc, FHitResult& OutResult, const FPhysicsGeometry& Geom, const FTransform& QueryTM, const FCollisionFilterData& QueryFilter, bool bReturnPhysMat)
+static bool ConvertOverlappedShapeToImpactHit(const UWorld* World, const THitLocation& Hit, const FVector& StartLoc, const FVector& EndLoc, FHitResult& OutResult, const FPhysicsGeometry& Geom, const FTransform& QueryTM, const Chaos::Filter::FQueryFilterData& QueryFilter, bool bReturnPhysMat)
 {
 	using namespace ChaosInterface;
 
@@ -562,7 +563,8 @@ static bool ConvertOverlappedShapeToImpactHit(const UWorld* World, const THitLoc
 
 	// See if this is a 'blocking' hit
 	FCollisionFilterData ShapeFilter = Chaos::Filter::FShapeFilterBuilder::GetLegacyShapeQueryFilter(GetCombinedShapeFilterData(HitShape));
-	const ECollisionQueryHitType HitType = FCollisionQueryFilterCallback::CalcQueryHitType(QueryFilter, ShapeFilter);
+	const FCollisionFilterData LegacyQueryFilter = Chaos::Filter::FQueryFilterBuilder::GetLegacyQueryFilter(QueryFilter);
+	const ECollisionQueryHitType HitType = FCollisionQueryFilterCallback::CalcQueryHitType(LegacyQueryFilter, ShapeFilter);
 	const bool bBlockingHit = (HitType == ECollisionQueryHitType::Block);
 	OutResult.bBlockingHit = bBlockingHit;
 
@@ -671,7 +673,7 @@ static bool ConvertOverlappedShapeToImpactHit(const UWorld* World, const THitLoc
 }
 
 
-void ConvertQueryOverlap(const FPhysicsShape& Shape, const FPhysicsActor& Actor, FOverlapResult& OutOverlap, const FCollisionFilterData& QueryFilter)
+void ConvertQueryOverlap(const FPhysicsShape& Shape, const FPhysicsActor& Actor, FOverlapResult& OutOverlap, const Chaos::Filter::FQueryFilterData& QueryFilter)
 {
 	using namespace ChaosInterface;
 
@@ -775,13 +777,14 @@ static void AddUniqueOverlap(TArray<FOverlapResult>& OutOverlaps, const FOverlap
 	OutOverlaps.Add(NewOverlap);
 }
 
-bool IsBlocking(const FPhysicsShape& Shape, const FCollisionFilterData& QueryFilter)
+bool IsBlocking(const FPhysicsShape& Shape, const Chaos::Filter::FQueryFilterData& QueryFilter)
 {
 	using namespace ChaosInterface;
+	const FCollisionFilterData LegacyQueryFilter = Chaos::Filter::FQueryFilterBuilder::GetLegacyQueryFilter(QueryFilter);
 
 	// See if this is a 'blocking' hit
 	const FCollisionFilterData ShapeFilter = Chaos::Filter::FShapeFilterBuilder::GetLegacyShapeQueryFilter(GetCombinedShapeFilterData(Shape));
-	const ECollisionQueryHitType HitType = FCollisionQueryFilterCallback::CalcQueryHitType(QueryFilter, ShapeFilter);
+	const ECollisionQueryHitType HitType = FCollisionQueryFilterCallback::CalcQueryHitType(LegacyQueryFilter, ShapeFilter);
 	const bool bBlock = (HitType == ECollisionQueryHitType::Block);
 	return bBlock;
 }
@@ -796,7 +799,7 @@ static FAutoConsoleVariableRef GTestOverlapSpeed(
 	);
 
 template <typename THitOverlap>
-bool ConvertOverlapResultsImp(int32 NumOverlaps, THitOverlap* OverlapResults, const FCollisionFilterData& QueryFilter, TArray<FOverlapResult>& OutOverlaps)
+bool ConvertOverlapResultsImp(int32 NumOverlaps, THitOverlap* OverlapResults, const Chaos::Filter::FQueryFilterData& QueryFilter, TArray<FOverlapResult>& OutOverlaps)
 {
 	SCOPE_CYCLE_COUNTER(STAT_CollisionConvertOverlap);
 
@@ -883,12 +886,12 @@ bool ConvertOverlapResultsImp(int32 NumOverlaps, THitOverlap* OverlapResults, co
 	return bBlockingFound;
 }
 
-bool ConvertOverlapResults(int32 NumOverlaps, FHitOverlap* OverlapResults, const FCollisionFilterData& QueryFilter, TArray<FOverlapResult>& OutOverlaps)
+bool ConvertOverlapResults(int32 NumOverlaps, FHitOverlap* OverlapResults, const Chaos::Filter::FQueryFilterData& QueryFilter, TArray<FOverlapResult>& OutOverlaps)
 {
 	return ConvertOverlapResultsImp(NumOverlaps, OverlapResults, QueryFilter, OutOverlaps);
 }
 
-bool ConvertOverlapResults(int32 NumOverlaps, ChaosInterface::FPTOverlapHit* OverlapResults, const FCollisionFilterData& QueryFilter, TArray<FOverlapResult>& OutOverlaps)
+bool ConvertOverlapResults(int32 NumOverlaps, ChaosInterface::FPTOverlapHit* OverlapResults, const Chaos::Filter::FQueryFilterData& QueryFilter, TArray<FOverlapResult>& OutOverlaps)
 {
 	return ConvertOverlapResultsImp(NumOverlaps, OverlapResults, QueryFilter, OutOverlaps);
 }
diff --git a/Engine/Source/Runtime/Engine/Private/Collision/CollisionConversions.h b/Engine/Source/Runtime/Engine/Private/Collision/CollisionConversions.h
--- a/Engine/Source/Runtime/Engine/Private/Collision/CollisionConversions.h
+++ b/Engine/Source/Runtime/Engine/Private/Collision/CollisionConversions.h
@@ -1,5 +1,5 @@
 // Copyright Epic Games, Inc. All Rights Reserved.
-// Utilities to convert from PhysX result structs to Unreal ones
+// Utilities to convert from Chaos result structs to Unreal ones
 
 #pragma once 
 
@@ -16,9 +16,9 @@ enum class EConvertQueryResult
 };
 
 /** 
- * Util to convert single physX hit (raycast or sweep) to our hit result
+ * Util to convert single chaos hit (raycast or sweep) to unreal hit result
  *
- * @param	PHit		PhysX hit data
+ * @param	PHit		Chaos hit data
  * @param	OutResult	(out) Result converted
  * @param	CheckLength	Distance of trace
  * @param	QueryFilter	Query Filter 
@@ -29,11 +29,11 @@ enum class EConvertQueryResult
  * @param	bReturnPhysMat		True if we want to lookup the physical material
  * @return	Whether result passed NaN/Inf checks.
  */
-EConvertQueryResult ConvertQueryImpactHit(const UWorld* World, const FHitLocation& PHit, FHitResult& OutResult, float CheckLength, const FCollisionFilterData& QueryFilter, const FVector& StartLoc, const FVector& EndLoc, const FPhysicsGeometry* Geom, const FTransform& QueryTM, bool bReturnFaceIndex, bool bReturnPhysMat);
-EConvertQueryResult ConvertQueryImpactHit(const UWorld* World, const ChaosInterface::FPTLocationHit& PHit, FHitResult& OutResult, float CheckLength, const FCollisionFilterData& QueryFilter, const FVector& StartLoc, const FVector& EndLoc, const FPhysicsGeometry* Geom, const FTransform& QueryTM, bool bReturnFaceIndex, bool bReturnPhysMat);
+EConvertQueryResult ConvertQueryImpactHit(const UWorld* World, const FHitLocation& PHit, FHitResult& OutResult, float CheckLength, const Chaos::Filter::FQueryFilterData& QueryFilter, const FVector& StartLoc, const FVector& EndLoc, const FPhysicsGeometry* Geom, const FTransform& QueryTM, bool bReturnFaceIndex, bool bReturnPhysMat);
+EConvertQueryResult ConvertQueryImpactHit(const UWorld* World, const ChaosInterface::FPTLocationHit& PHit, FHitResult& OutResult, float CheckLength, const Chaos::Filter::FQueryFilterData& QueryFilter, const FVector& StartLoc, const FVector& EndLoc, const FPhysicsGeometry* Geom, const FTransform& QueryTM, bool bReturnFaceIndex, bool bReturnPhysMat);
 
 /** 
- * Util to convert physX sweep results to unreal hit results and add to array
+ * Util to convert Chaos sweep results to unreal hit results and add to array
  *
  * @param	OutHasValidBlockingHit set to whether there is a valid blocking hit found in the results.
  * @param	NumHits		Number of Hits
@@ -48,20 +48,20 @@ EConvertQueryResult ConvertQueryImpactHit(const UWorld* World, const ChaosInterf
  * @return	Whether all results passed NaN/Inf checks.
  */
 template <typename Hit>
-EConvertQueryResult ConvertTraceResults(bool& OutHasValidBlockingHit, const UWorld* World, int32 NumHits, Hit* Hits, float CheckLength, const FCollisionFilterData& QueryFilter, TArray<FHitResult>& OutHits, const FVector& StartLoc, const FVector& EndLoc, const FPhysicsGeometry* Geom, const FTransform& QueryTM, float MaxDistance, bool bReturnFaceIndex, bool bReturnPhysMat);
+EConvertQueryResult ConvertTraceResults(bool& OutHasValidBlockingHit, const UWorld* World, int32 NumHits, Hit* Hits, float CheckLength, const Chaos::Filter::FQueryFilterData& QueryFilter, TArray<FHitResult>& OutHits, const FVector& StartLoc, const FVector& EndLoc, const FPhysicsGeometry* Geom, const FTransform& QueryTM, float MaxDistance, bool bReturnFaceIndex, bool bReturnPhysMat);
 
 template <typename Hit>
-EConvertQueryResult ConvertTraceResults(bool& OutHasValidBlockingHit, const UWorld* World, int32 NumHits, Hit* Hits, float CheckLength, const FCollisionFilterData& QueryFilter, FHitResult& OutHits, const FVector& StartLoc, const FVector& EndLoc, const FPhysicsGeometry* Geom, const FTransform& QueryTM, float MaxDistance, bool bReturnFaceIndex, bool bReturnPhysMat);
+EConvertQueryResult ConvertTraceResults(bool& OutHasValidBlockingHit, const UWorld* World, int32 NumHits, Hit* Hits, float CheckLength, const Chaos::Filter::FQueryFilterData& QueryFilter, FHitResult& OutHits, const FVector& StartLoc, const FVector& EndLoc, const FPhysicsGeometry* Geom, const FTransform& QueryTM, float MaxDistance, bool bReturnFaceIndex, bool bReturnPhysMat);
 
 /** 
- * Util to convert physX or Chaos overlap query to our overlap result
+ * Util to convert Chaos overlap query to our overlap result
  *
  * @param	PShape		Shape that overlaps
  * @param	PActor		Specific actor as PShape might be shared among many actors
  * @param	OutOverlap	(out) Result converted
  * @param	QueryFilter	Query Filter 
  */
-void ConvertQueryOverlap(const UWorld* World, const FPhysicsShape& PShape, const FPhysicsActorHandle& PActor, FOverlapResult& OutOverlap, const FCollisionFilterData& QueryFilter);
+void ConvertQueryOverlap(const UWorld* World, const FPhysicsShape& PShape, const FPhysicsActorHandle& PActor, FOverlapResult& OutOverlap, const Chaos::Filter::FQueryFilterData& QueryFilter);
 
 /**
  * Util to determine if a shape is deemed blocking based on the query filter
@@ -70,7 +70,7 @@ void ConvertQueryOverlap(const UWorld* World, const FPhysicsShape& PShape, const
  * @QueryFilter Query Filter
  * @return true if the query filter and shape filter resolve to be blocking
  */
-bool IsBlocking(const FPhysicsShape& PShape, const FCollisionFilterData& QueryFilter);
+bool IsBlocking(const FPhysicsShape& PShape, const Chaos::Filter::FQueryFilterData& QueryFilter);
 
 
 /** 
@@ -81,8 +81,8 @@ bool IsBlocking(const FPhysicsShape& PShape, const FCollisionFilterData& QueryFi
  * @param	QueryFilter	Query filter for converting
  * @return	OutOverlaps	Converted data
  */
-bool ConvertOverlapResults(int32 NumOverlaps, FHitOverlap* POverlapResults, const FCollisionFilterData& QueryFilter, TArray<FOverlapResult>& OutOverlaps);
-bool ConvertOverlapResults(int32 NumOverlaps, ChaosInterface::FPTOverlapHit* POverlapResults, const FCollisionFilterData& QueryFilter, TArray<FOverlapResult>& OutOverlaps);
+bool ConvertOverlapResults(int32 NumOverlaps, FHitOverlap* POverlapResults, const Chaos::Filter::FQueryFilterData& QueryFilter, TArray<FOverlapResult>& OutOverlaps);
+bool ConvertOverlapResults(int32 NumOverlaps, ChaosInterface::FPTOverlapHit* POverlapResults, const Chaos::Filter::FQueryFilterData& QueryFilter, TArray<FOverlapResult>& OutOverlaps);
 
 /**
  * Converts an overlap result to a hit result.
diff --git a/Engine/Source/Runtime/Engine/Private/Collision/SceneQuery.cpp b/Engine/Source/Runtime/Engine/Private/Collision/SceneQuery.cpp
--- a/Engine/Source/Runtime/Engine/Private/Collision/SceneQuery.cpp
+++ b/Engine/Source/Runtime/Engine/Private/Collision/SceneQuery.cpp
@@ -575,7 +575,7 @@ bool TSceneCastCommonImpWithRetryRequest(const UWorld* World, typename Traits::O
 		}
 		if (NumHits > 0 && !Traits::IsTest())
 		{
-			bool bSuccess = ConvertTraceResults(bBlockingHit, World, NumHits, Traits::GetHits(HitBufferSync), DeltaMag, Filter, OutHits, Start, End, GeomInputs.GetGeometry(), StartTM, MinBlockingDistance, Params.bReturnFaceIndex, Params.bReturnPhysicalMaterial) == EConvertQueryResult::Valid;
+			bool bSuccess = ConvertTraceResults(bBlockingHit, World, NumHits, Traits::GetHits(HitBufferSync), DeltaMag, QueryFilterData, OutHits, Start, End, GeomInputs.GetGeometry(), StartTM, MinBlockingDistance, Params.bReturnFaceIndex, Params.bReturnPhysicalMaterial) == EConvertQueryResult::Valid;
 
 			if (!bSuccess)
 			{
@@ -1165,7 +1165,7 @@ bool GeomOverlapMultiImp(const UWorld* World, const FPhysicsGeometry& Geom, cons
 
 				if(NumHits > 0)
 				{
-					bHaveBlockingHit = ConvertOverlapResults(NumHits, OverlapBuffer.GetHits(), Filter, OutOverlaps);
+					bHaveBlockingHit = ConvertOverlapResults(NumHits, OverlapBuffer.GetHits(), QueryFilterData, OutOverlaps);
 
 					// This block is only necessary on the game thread when dealing with cluster unions.
 					// TODO: Is there a way to get this to generalize better to the PT as well? Right now it depends on GT functions on the cluster union component.
diff --git a/Engine/Source/Runtime/Engine/Private/Components/PrimitiveComponent.cpp b/Engine/Source/Runtime/Engine/Private/Components/PrimitiveComponent.cpp
--- a/Engine/Source/Runtime/Engine/Private/Components/PrimitiveComponent.cpp
+++ b/Engine/Source/Runtime/Engine/Private/Components/PrimitiveComponent.cpp
@@ -3555,7 +3555,7 @@ bool UPrimitiveComponent::LineTraceComponent(FHitResult& OutHit, const FVector S
 		{
 			bHaveHit = true;
 
-			const FCollisionFilterData QueryFilter = Chaos::Filter::FQueryFilterBuilder::GetLegacyQueryFilter(Chaos::Filter::FQueryObjectFilterBuilder::BuildBlockAll());
+			const Chaos::Filter::FQueryFilterData QueryFilter = Chaos::Filter::FQueryObjectFilterBuilder::BuildBlockAll();
 			ChaosInterface::SetFlags(BestHit, EHitFlags::Distance | EHitFlags::Normal | EHitFlags::Position);
 
 			ConvertQueryImpactHit(GetWorld(), BestHit, OutHit, (End - Start).Size(), QueryFilter, Start, End, nullptr, FTransform{ Start }, true, Params.bReturnPhysicalMaterial);
@@ -3619,7 +3619,7 @@ bool UPrimitiveComponent::SweepComponent(FHitResult& OutHit, const FVector Start
 	SweepParams.bComputeMTD = true;
 	if (CollisionInterface.ShapeSweep(Objects, Geometry, FTransform{ ShapeWorldRotation, Start }, End, SweepParams, BestHit))
 	{
-		const FCollisionFilterData QueryFilter = Chaos::Filter::FQueryFilterBuilder::GetLegacyQueryFilter(Chaos::Filter::FQueryObjectFilterBuilder::BuildBlockAll());
+		const Chaos::Filter::FQueryFilterData QueryFilter = Chaos::Filter::FQueryObjectFilterBuilder::BuildBlockAll();
 		ChaosInterface::SetFlags(BestHit, EHitFlags::Distance | EHitFlags::Normal | EHitFlags::Position | EHitFlags::FaceIndex);
 
 		bool bHasHit = false;
@@ -3727,7 +3727,7 @@ bool UPrimitiveComponent::ComponentOverlapComponentWithResultImpl(const class UP
 			{
 				TArray<FOverlapResult> Overlaps;
 
-				const FCollisionFilterData QueryFilter = Chaos::Filter::FQueryFilterBuilder::GetLegacyQueryFilter(Chaos::Filter::FQueryObjectFilterBuilder::BuildBlockAll());
+				const Chaos::Filter::FQueryFilterData QueryFilter = Chaos::Filter::FQueryObjectFilterBuilder::BuildBlockAll();
 				ConvertOverlapResults(OverlapHits.Num(), OverlapHits.GetData(), QueryFilter, Overlaps);
 
 				if (!Overlaps.IsEmpty())
@@ -3775,7 +3775,7 @@ bool UPrimitiveComponent::OverlapComponentWithResult(const FVector& Pos, const F
 	{
 		TArray<FOverlapResult> Overlaps;
 
-		const FCollisionFilterData QueryFilter = Chaos::Filter::FQueryFilterBuilder::GetLegacyQueryFilter(Chaos::Filter::FQueryObjectFilterBuilder::BuildBlockAll());
+		const Chaos::Filter::FQueryFilterData QueryFilter = Chaos::Filter::FQueryObjectFilterBuilder::BuildBlockAll();
 		ConvertOverlapResults(OverlapHits.Num(), OverlapHits.GetData(), QueryFilter, Overlaps);
 
 		if (!Overlaps.IsEmpty())
diff --git a/Engine/Source/Runtime/Engine/Private/PhysicsEngine/Experimental/PhysInterface_Chaos.cpp b/Engine/Source/Runtime/Engine/Private/PhysicsEngine/Experimental/PhysInterface_Chaos.cpp
--- a/Engine/Source/Runtime/Engine/Private/PhysicsEngine/Experimental/PhysInterface_Chaos.cpp
+++ b/Engine/Source/Runtime/Engine/Private/PhysicsEngine/Experimental/PhysInterface_Chaos.cpp
@@ -984,7 +984,7 @@ bool FPhysInterface_Chaos::LineTrace_Geom(FHitResult& OutHit, const FBodyInstanc
 						SetFlags(BestHit, EHitFlags::Distance | EHitFlags::Normal | EHitFlags::Position);
 
 						// Verify the hit is made by doing an overlap test.
-						const FCollisionFilterData QueryFilter = Chaos::Filter::FQueryFilterBuilder::GetLegacyQueryFilter(Chaos::Filter::FQueryTraceFilterBuilder::BuildOverlapAll());
+						const Chaos::Filter::FQueryFilterData QueryFilter = Chaos::Filter::FQueryTraceFilterBuilder::BuildOverlapAll();
 
 						FTransform StartTM(WorldStart);
 						const UPrimitiveComponent* OwnerComponentInst = InInstance->OwnerComponent.Get();
@@ -1081,7 +1081,7 @@ bool FPhysInterface_Chaos::Sweep_Geom(FHitResult& OutHit, const FBodyInstance* I
 						if (bShapeHit && FloatDistance < Hit.Distance)
 						{
 							// we just like to make sure if the hit is made
-							const FCollisionFilterData QueryFilter = Chaos::Filter::FQueryFilterBuilder::GetLegacyQueryFilter(Chaos::Filter::FQueryTraceFilterBuilder::BuildOverlapAll());
+							const Chaos::Filter::FQueryFilterData QueryFilter = Chaos::Filter::FQueryTraceFilterBuilder::BuildOverlapAll();
 
 							// we don't get Shape information when we access via PShape, so I filled it up
 							Hit.Shape = Shape;
diff --git a/Engine/Source/Runtime/Engine/Public/Physics/Experimental/PhysInterface_Chaos.h b/Engine/Source/Runtime/Engine/Public/Physics/Experimental/PhysInterface_Chaos.h
--- a/Engine/Source/Runtime/Engine/Public/Physics/Experimental/PhysInterface_Chaos.h
+++ b/Engine/Source/Runtime/Engine/Public/Physics/Experimental/PhysInterface_Chaos.h
@@ -187,5 +187,3 @@ inline void SetShape(FPhysTypeDummy& Hit, FPhysTypeDummy* Shape)
 {
 
 }
-
-bool IsBlocking(const FPhysicsShape& PShape, const FCollisionFilterData& QueryFilter);
