Fixed compile errors when removing deprecated methods in TSoftObjectPtr

diff --git a/Engine/Source/Runtime/Engine/Private/ActorEditor.cpp b/Engine/Source/Runtime/Engine/Private/ActorEditor.cpp
--- a/Engine/Source/Runtime/Engine/Private/ActorEditor.cpp
+++ b/Engine/Source/Runtime/Engine/Private/ActorEditor.cpp
@@ -1960,7 +1960,7 @@ void AActor::FixupDataLayers(bool bRevertChangesOnLockedDataLayer /*= false*/)
 		TArray<const UDataLayerAsset*> CleanedUpDataLayers = ResolveDataLayerAssets(DataLayerAssets);
 		CleanupDataLayers(CleanedUpDataLayers);
 		DataLayerAssets.Empty(CleanedUpDataLayers.Num());
-		Algo::Transform(CleanedUpDataLayers, DataLayerAssets, [](const UDataLayerAsset* DataLayerAsset) { return TSoftObjectPtr<UDataLayerAsset>(DataLayerAsset); });
+		Algo::Transform(CleanedUpDataLayers, DataLayerAssets, [](const UDataLayerAsset* DataLayerAsset) { return TSoftObjectPtr<UDataLayerAsset>(const_cast<UDataLayerAsset*>(DataLayerAsset)); });
 	}
 	CleanupDataLayers(DataLayers);
 	PRAGMA_ENABLE_DEPRECATION_WARNINGS
