[Spline Mesh] Update default collision if necessary after BodySetup initialization

* Will allow for the Spline Mesh to respect the "Use Default Collision" flag.

#rnx

diff --git a/Engine/Source/Runtime/Engine/Classes/Components/SplineMeshComponent.h b/Engine/Source/Runtime/Engine/Classes/Components/SplineMeshComponent.h
--- a/Engine/Source/Runtime/Engine/Classes/Components/SplineMeshComponent.h
+++ b/Engine/Source/Runtime/Engine/Classes/Components/SplineMeshComponent.h
@@ -231,6 +231,7 @@ protected:
 	FBox ComputeDistortedBounds(const FTransform& InLocalToWorld, const FBoxSphereBounds& InMeshBounds, const FBoxSphereBounds* InBoundsToDistort = nullptr) const;
 
 public:
+	ENGINE_API virtual bool SupportsDefaultCollision() override;
 	ENGINE_API virtual class UBodySetup* GetBodySetup() override;
 	ENGINE_API class UBodySetup* GetBodySetup() const;
 #if WITH_EDITOR
diff --git a/Engine/Source/Runtime/Engine/Private/Components/SplineMeshComponent.cpp b/Engine/Source/Runtime/Engine/Private/Components/SplineMeshComponent.cpp
--- a/Engine/Source/Runtime/Engine/Private/Components/SplineMeshComponent.cpp
+++ b/Engine/Source/Runtime/Engine/Private/Components/SplineMeshComponent.cpp
@@ -770,6 +770,14 @@ FBox USplineMeshComponent::ComputeDistortedBounds(const FTransform& InLocalToWor
 	return FSplineMeshSceneProxyDesc(this).ComputeDistortedBounds(InLocalToWorld, InMeshBounds, InBoundsToDistort);
 }
 
+bool USplineMeshComponent::SupportsDefaultCollision()
+{
+	// For Spline Mesh, if we have a Body Setup, we support the default collision, because the body setup is only initialized if the static mesh
+	// has a body setup. But because the spline mesh body setup is different from the static mesh body setup (because of the deformation of the mesh)
+	// the UStaticMeshComponent::SupportsDefaultCollision will always return false.
+	return GetBodySetup() != nullptr;
+}
+
 FTransform USplineMeshComponent::GetSocketTransform(FName InSocketName, ERelativeTransformSpace TransformSpace) const
 {
 	if (InSocketName != NAME_None)
@@ -1146,6 +1154,10 @@ void USplineMeshComponent::RecreateCollision()
 		}
 
 		BodySetup->CreatePhysicsMeshes();
+
+		// After the BodySetup is created/updated, we need to make sure the default collision is respected
+		// This function is called in the OnRegister, and at that point for Spline Meshes, BodySetup is null.
+		UpdateCollisionFromStaticMesh();
 	}
 	else
 	{
