[NetPhysics] Debugging cvar to trigger resim with an interval as long as we receive an input or state on that frame.

diff --git a/Engine/Source/Runtime/Engine/Private/PhysicsEngine/NetworkPhysicsComponent.cpp b/Engine/Source/Runtime/Engine/Private/PhysicsEngine/NetworkPhysicsComponent.cpp
--- a/Engine/Source/Runtime/Engine/Private/PhysicsEngine/NetworkPhysicsComponent.cpp
+++ b/Engine/Source/Runtime/Engine/Private/PhysicsEngine/NetworkPhysicsComponent.cpp
@@ -85,6 +85,9 @@ namespace PhysicsReplicationCVars
 		
 		bool bApplyPredictiveInterpolationWhenBehindServer = true;
 		static FAutoConsoleVariableRef CVarResimApplyPredictiveInterpolationWhenBehindServer(TEXT("np2.Resim.ApplyPredictiveInterpolationWhenBehindServer"), bApplyPredictiveInterpolationWhenBehindServer, TEXT("When true, switch over to replicating with Predictive Interpolation temporarily, when the client receive target states from the server for frames that have not yet simulated on the client. When false apply the received target via a resimulation when the client has caught up and simulated the corresponding frame."));
+
+		int32 bDebugTriggerResimEveryNFrames = 0;
+		static FAutoConsoleVariableRef CVarDebugTriggerResimEveryNFrames(TEXT("np2.Resim.Debug.TriggerResimEveryNFrames"), bDebugTriggerResimEveryNFrames, TEXT("When above 0, trigger a resim at an interval of the set value, as long as we receive a state or input on that frame."));
 	}
 }
 
@@ -957,7 +960,7 @@ void UNetworkPhysicsComponent::AddDeltaSourceInput()
 {
 	// Get the data entry for the correct index in the data sources array
 	const int32 Index = GetDeltaSourceIndexForFrame(ReplicatedDeltaSourceInput.History->GetLatestFrame());
-	check(Index <= DeltaSourceInputs.Num());
+	check(Index < DeltaSourceInputs.Num());
 	FNetworkPhysicsData* PhysicsData = DeltaSourceInputs[Index].Get();
 
 	// Extract the data from the replicated DeltaSources property
@@ -1037,7 +1040,7 @@ void UNetworkPhysicsComponent::AddDeltaSourceState()
 {
 	// Get the data entry for the correct index in the data sources array
 	const int32 Index = GetDeltaSourceIndexForFrame(ReplicatedDeltaSourceState.History->GetLatestFrame());
-	check(Index <= DeltaSourceStates.Num());
+	check(Index < DeltaSourceStates.Num());
 	FNetworkPhysicsData* PhysicsData = DeltaSourceStates[Index].Get();
 
 	// Extract the data from the replicated DeltaSources property
@@ -1889,6 +1892,17 @@ void FAsyncNetworkPhysicsComponent::ConsumeAsyncInput(const int32 PhysicsStep)
 					TriggerResimulation(ResimFrame);
 				}
 
+#if !UE_BUILD_SHIPPING
+				if (PhysicsReplicationCVars::ResimulationCVars::bDebugTriggerResimEveryNFrames > 0 && !bIsServer && IsLocallyControlled())
+				{
+					int32 LatestReceivedFrame = ReceiveData->GetLatestFrame() - NetworkPhysicsTickOffset;
+					if (PhysicsStep % PhysicsReplicationCVars::ResimulationCVars::bDebugTriggerResimEveryNFrames == 0)
+					{
+						TriggerResimulation(LatestReceivedFrame);
+					}
+				}
+#endif // !UE_BUILD_SHIPPING
+
 #if DEBUG_NETWORK_PHYSICS
 				{
 					FString NetRoleString = bIsServer ? FString("SERVER") : (IsLocallyControlled() ? FString("AUTONO") : FString("PROXY "));
