Add 'Ignore Pose Asset Source Animation Curves' property to the pose blend node

diff --git a/Engine/Source/Runtime/AnimGraphRuntime/Private/AnimNodes/AnimNode_PoseBlendNode.cpp b/Engine/Source/Runtime/AnimGraphRuntime/Private/AnimNodes/AnimNode_PoseBlendNode.cpp
--- a/Engine/Source/Runtime/AnimGraphRuntime/Private/AnimNodes/AnimNode_PoseBlendNode.cpp
+++ b/Engine/Source/Runtime/AnimGraphRuntime/Private/AnimNodes/AnimNode_PoseBlendNode.cpp
@@ -61,7 +61,7 @@ void FAnimNode_PoseBlendNode::Evaluate_AnyThread(FPoseContext& Output)
 		});
 
 		FAnimationPoseData CurrentAnimationPoseData(CurrentPose);
-		if (CachedPoseAsset->GetAnimationPose(CurrentAnimationPoseData, PoseExtractContext))
+		if (CachedPoseAsset->GetAnimationPose(CurrentAnimationPoseData, PoseExtractContext, bIgnorePoseAssetSourceAnimationCurves))
 		{
 			// once we get it, we have to blend by weight
 			if (CachedPoseAsset->IsValidAdditive())
diff --git a/Engine/Source/Runtime/AnimGraphRuntime/Public/AnimNodes/AnimNode_PoseBlendNode.h b/Engine/Source/Runtime/AnimGraphRuntime/Public/AnimNodes/AnimNode_PoseBlendNode.h
--- a/Engine/Source/Runtime/AnimGraphRuntime/Public/AnimNodes/AnimNode_PoseBlendNode.h
+++ b/Engine/Source/Runtime/AnimGraphRuntime/Public/AnimNodes/AnimNode_PoseBlendNode.h
@@ -28,6 +28,9 @@ public:
 	UPROPERTY(EditAnywhere, Category = "Blend")
 	TObjectPtr<UCurveFloat> CustomCurve;
 
+	UPROPERTY(EditAnywhere, Category = "Blend")
+	bool bIgnorePoseAssetSourceAnimationCurves = false;
+
 private:
 	// Cached curves to extract
 	UE::Anim::TNamedValueArray<FDefaultAllocator, UE::Anim::FNamedIndexElement> BulkCurves;
diff --git a/Engine/Source/Runtime/Engine/Classes/Animation/PoseAsset.h b/Engine/Source/Runtime/Engine/Classes/Animation/PoseAsset.h
--- a/Engine/Source/Runtime/Engine/Classes/Animation/PoseAsset.h
+++ b/Engine/Source/Runtime/Engine/Classes/Animation/PoseAsset.h
@@ -235,7 +235,7 @@ public:
 	*/
 	UE_DEPRECATED(4.26, "Use GetAnimationPose with other signature")
 	ENGINE_API bool GetAnimationPose(struct FCompactPose& OutPose, FBlendedCurve& OutCurve, const FAnimExtractContext& ExtractionContext) const;
-	ENGINE_API bool GetAnimationPose(struct FAnimationPoseData& OutAnimationPoseData, const FAnimExtractContext& ExtractionContext) const;
+	ENGINE_API bool GetAnimationPose(struct FAnimationPoseData& OutAnimationPoseData, const FAnimExtractContext& ExtractionContext, const bool bIgnoreSourceAnimationCurves = false) const;
 
 	UE_DEPRECATED(4.26, "Use GetBaseAnimationPose with other signature")
 	ENGINE_API void GetBaseAnimationPose(struct FCompactPose& OutPose, FBlendedCurve& OutCurve) const;
diff --git a/Engine/Source/Runtime/Engine/Private/Animation/PoseAsset.cpp b/Engine/Source/Runtime/Engine/Private/Animation/PoseAsset.cpp
--- a/Engine/Source/Runtime/Engine/Private/Animation/PoseAsset.cpp
+++ b/Engine/Source/Runtime/Engine/Private/Animation/PoseAsset.cpp
@@ -604,7 +604,7 @@ bool UPoseAsset::GetAnimationPose(struct FCompactPose& OutPose, FBlendedCurve& O
 	return GetAnimationPose(OutPoseData, ExtractionContext);
 }
 
-bool UPoseAsset::GetAnimationPose(struct FAnimationPoseData& OutAnimationPoseData, const FAnimExtractContext& ExtractionContext) const
+bool UPoseAsset::GetAnimationPose(struct FAnimationPoseData& OutAnimationPoseData, const FAnimExtractContext& ExtractionContext, const bool bIgnoreSourceAnimationCurves) const
 {
 	ANIM_MT_SCOPE_CYCLE_COUNTER(PoseAssetGetAnimationPose, !IsInGameThread());
 
@@ -645,8 +645,11 @@ bool UPoseAsset::GetAnimationPose(struct FAnimationPoseData& OutAnimationPoseDat
 				// Only generate pose if the single weight is actually relevant
 				if(FAnimWeight::IsRelevant(Weight))
 				{
-					// Blend curve
-					PoseContainer.BlendPoseCurve(&Pose, OutCurve, Weight);
+					if (!bIgnoreSourceAnimationCurves)
+					{
+						// Blend curve
+						PoseContainer.BlendPoseCurve(&Pose, OutCurve, Weight);
+					}
 
 					// Per-track (bone) transform
 					for (int32 TrackIndex = 0; TrackIndex < TrackNum; ++TrackIndex)
@@ -773,22 +776,28 @@ bool UPoseAsset::GetAnimationPose(struct FAnimationPoseData& OutAnimationPoseDat
 						float& PoseWeight = PoseWeights[WeightedPoseIndex];
 						PoseWeight /= TotalWeight;
 
-						// Do curve blending inline as we are looping over weights anyway
-						const FPoseData& Pose = PoseContainer.Poses[WeightedPoseIndex];
-						PoseContainer.BlendPoseCurve(&Pose, OutCurve, PoseWeight);
+						if (!bIgnoreSourceAnimationCurves)
+						{
+							// Do curve blending inline as we are looping over weights anyway
+							const FPoseData& Pose = PoseContainer.Poses[WeightedPoseIndex];
+							PoseContainer.BlendPoseCurve(&Pose, OutCurve, PoseWeight);
+						}
 					}
 				}
 				else
 				{
-					// Take the matching curve weights from the selected poses, and blend them using the
-					// weighting that we need from each pose. This is much faster than grabbing each
-					// blend curve and blending them in their entirety, especially when there are very
-					// few active curves for each pose and many curves for the entire pose asset.
-					for (const int32& WeightedPoseIndex : WeightedPoseIndices)
+					if (!bIgnoreSourceAnimationCurves)
 					{
-						const FPoseData& Pose = PoseContainer.Poses[WeightedPoseIndex];
-						const float& Weight = PoseWeights[WeightedPoseIndex];
-						PoseContainer.BlendPoseCurve(&Pose, OutCurve, Weight);
+						// Take the matching curve weights from the selected poses, and blend them using the
+						// weighting that we need from each pose. This is much faster than grabbing each
+						// blend curve and blending them in their entirety, especially when there are very
+						// few active curves for each pose and many curves for the entire pose asset.
+						for (const int32& WeightedPoseIndex : WeightedPoseIndices)
+						{
+							const FPoseData& Pose = PoseContainer.Poses[WeightedPoseIndex];
+							const float& Weight = PoseWeights[WeightedPoseIndex];
+							PoseContainer.BlendPoseCurve(&Pose, OutCurve, Weight);
+						}
 					}
 				}
 
