Fixed UDeletedObjectPlaceholder::GetOriginalObject to also return object even if pending kill.
This fixes the case where we save an actor that is first destroyed & marked as garbage and then moved outside of its package.
This fix follows 46731314 which was preventing a crash but didn't properly cleanup the actor's associated UActorDescContainer.

diff --git a/Engine/Source/Runtime/Engine/Public/DeletedObjectPlaceholder.h b/Engine/Source/Runtime/Engine/Public/DeletedObjectPlaceholder.h
--- a/Engine/Source/Runtime/Engine/Public/DeletedObjectPlaceholder.h
+++ b/Engine/Source/Runtime/Engine/Public/DeletedObjectPlaceholder.h
@@ -25,7 +25,7 @@ public:
 	virtual bool IsAsset() const override;
 	virtual void PostEditUndo() override;
 	
-	const UObject* GetOriginalObject() const { return OriginalObject.Get(); }
+	const UObject* GetOriginalObject() const { return OriginalObject.Get(true); }
 	const FString& GetDisplayName() const { return DisplayName; }
 	const FExternalDataLayerUID GetExternalDataLayerUID() const { return FExternalDataLayerUID(ExternalDataLayerUID); }
 
