Silence static analysis false positive

diff --git a/Engine/Source/Runtime/Engine/Private/Animation/ReferenceSkeleton.cpp b/Engine/Source/Runtime/Engine/Private/Animation/ReferenceSkeleton.cpp
--- a/Engine/Source/Runtime/Engine/Private/Animation/ReferenceSkeleton.cpp
+++ b/Engine/Source/Runtime/Engine/Private/Animation/ReferenceSkeleton.cpp
@@ -512,7 +512,7 @@ void FReferenceSkeleton::RebuildRefSkeleton(const USkeleton* Skeleton, bool bReb
 			const int32 ActualIndex = VirtualBoneIdx + RawRefBoneInfo.Num();
 			const FVirtualBone& VB = Skeleton->GetVirtualBones()[VirtualBoneIdx];
 
-			const int32 SourceIndex = GetRawSourceBoneIndex(Skeleton, VB.SourceBoneName);
+			const int32 SourceIndex = GetRawSourceBoneIndex(Skeleton, VB.SourceBoneName);	//-V595
 			const int32 ParentIndex = FindBoneIndex(VB.SourceBoneName);
 			const int32 TargetIndex = FindBoneIndex(VB.TargetBoneName);
 			if(ParentIndex != INDEX_NONE && TargetIndex != INDEX_NONE)
