Added Double Spring functions to Blueprint Spring Math Library

#virtualized

diff --git a/Engine/Source/Runtime/Engine/Classes/Kismet/BlueprintSpringMathLibrary.h b/Engine/Source/Runtime/Engine/Classes/Kismet/BlueprintSpringMathLibrary.h
--- a/Engine/Source/Runtime/Engine/Classes/Kismet/BlueprintSpringMathLibrary.h
+++ b/Engine/Source/Runtime/Engine/Classes/Kismet/BlueprintSpringMathLibrary.h
@@ -86,6 +86,96 @@ public:
 	UFUNCTION(BlueprintCallable, meta=(Category="Math|Spring|Experimental"))
 	static ENGINE_API void CriticalSpringDampRotator(UPARAM(ref) FRotator& InOutRotation, UPARAM(ref) FVector& InOutAngularVelocity, const FRotator& TargetRotation, float DeltaTime, float SmoothingTime = 0.2f);
 
+
+	// Double Spring //
+
+	UE_EXPERIMENTAL(5.7, "SpringMath is experimental")
+	/** Interpolates the value InOutX towards TargetX using a critically damped spring via an intermediate state - producing a more S-shaped curve. The velocity of X is stored in InOutV.
+	 *
+	 * @param InOutX The value to be damped
+	 * @param InOutV The speed of the value to be damped
+	 * @param InOutXi The intermediate state of the double spring
+	 * @param InOutVi The intermediate velocity of the double spring
+	 * @param TargetX The goal to damp towards
+	 * @param SmoothingTime The smoothing time to use for the spring. Longer times lead to more damped behaviour. A time of 0 will snap InOutX to TargetX
+	 * @param DeltaTime Timestep in seconds
+	 */
+	UFUNCTION(BlueprintCallable, meta = (Category = "Math|Spring|Experimental"))
+	static ENGINE_API void CriticalDoubleSpringDampFloat(UPARAM(ref) float& InOutX, UPARAM(ref) float& InOutV, UPARAM(ref) float& InOutXi, UPARAM(ref) float& InOutVi, const float& TargetX, float DeltaTime, float SmoothingTime = 0.2f);
+
+	UE_EXPERIMENTAL(5.7, "SpringMath is experimental")
+	/** Interpolates the angle InOutAngle towards TargetAngle using a critically damped spring via an intermediate state - producing a more S-shaped curve. The velocity of InOutAngle is stored in InOutAngularVelocity in deg/s.
+	 *
+	 * @param InOutAngle The value to be damped in degrees
+	 * @param InOutAngularVelocity The speed of the value to be damped in deg/s
+	 * @param InOutIntermediateAngle The intermediate state of the double spring
+	 * @param InOutIntermediateAngularVelocity The intermediate velocity of the double spring
+	 * @param TargetAngle The goal to damp towards in degrees
+	 * @param SmoothingTime The smoothing time to use for the spring. Longer times lead to more damped behaviour. A time of 0 will snap InOutAngle to TargetAngle
+	 * @param DeltaTime Timestep in seconds
+	 */
+	UFUNCTION(BlueprintCallable, meta = (Category = "Math|Spring|Experimental"))
+	static ENGINE_API void CriticalDoubleSpringDampAngle(UPARAM(ref) float& InOutAngle,UPARAM(ref) float& InOutAngularVelocity, UPARAM(ref) float& InOutIntermediateAngle, UPARAM(ref) float& InOutIntermediateAngularVelocity, const float& TargetAngle, float DeltaTime, float SmoothingTime = 0.2f);
+
+	UE_EXPERIMENTAL(5.7, "SpringMath is experimental")
+	/** Interpolates the value InOutX towards TargetX using a critically damped spring via an intermediate state - producing a more S-shaped curve. The velocity of X is stored in InOutV.
+	 *
+	 * @param InOutX The value to be damped
+	 * @param InOutV The speed of the value to be damped
+	 * @param InOutXi The intermediate state of the double spring
+	 * @param InOutVi The intermediate velocity of the double spring
+	 * @param TargetX The goal to damp towards
+	 * @param SmoothingTime The smoothing time to use for the spring. Longer times lead to more damped behaviour. A time of 0 will snap InOutX to TargetX
+	 * @param DeltaTime Timestep in seconds
+	 */
+	UFUNCTION(BlueprintCallable, meta = (Category = "Math|Spring|Experimental"))
+	static ENGINE_API void CriticalDoubleSpringDampVector(UPARAM(ref) FVector& InOutX, UPARAM(ref) FVector& InOutV, UPARAM(ref) FVector& InOutXi, UPARAM(ref) FVector& InOutVi, const FVector& TargetX, float DeltaTime, float SmoothingTime = 0.2f);
+
+	UE_EXPERIMENTAL(5.7, "SpringMath is experimental")
+	/** Interpolates the value InOutX towards TargetX using a critically damped spring via an intermediate state - producing a more S-shaped curve. The velocity of X is stored in InOutV.
+	 *
+	 * @param InOutX The value to be damped
+	 * @param InOutV The speed of the value to be damped
+	 * @param InOutXi The intermediate state of the double spring
+	 * @param InOutVi The intermediate velocity of the double spring
+	 * @param TargetX The goal to damp towards
+	 * @param SmoothingTime The smoothing time to use for the spring. Longer times lead to more damped behaviour. A time of 0 will snap InOutX to TargetX
+	 * @param DeltaTime Timestep in seconds
+	 */
+	UFUNCTION(BlueprintCallable, meta = (Category = "Math|Spring|Experimental"))
+	static ENGINE_API void CriticalDoubleSpringDampVector2D(UPARAM(ref) FVector2D& InOutX, UPARAM(ref) FVector2D& InOutV, UPARAM(ref) FVector2D& InOutXi, UPARAM(ref) FVector2D& InOutVi, const FVector2D& TargetX, float DeltaTime, float SmoothingTime = 0.2f);
+
+	UE_EXPERIMENTAL(5.7, "SpringMath is experimental")
+	/** Interpolates the value InOutRotation towards TargetRotation using a critically damped spring via an intermediate state - producing a more S-shaped curve. The velocity of InOutRotation is stored in InOutAngularVelocity in deg/s.
+	 *
+	 * @param InOutRotation The value to be damped
+	 * @param InOutAngularVelocity The speed of the value to be damped in deg/s
+	 * @param InOutIntermediateRotation The intermediate state of the double spring
+	 * @param InOutIntermediateAngularVelocity The intermediate velocity of the double spring
+	 * @param TargetRotation The goal to damp towards
+	 * @param SmoothingTime The smoothing time to use for the spring. Longer times lead to more damped behaviour. A time of 0 will snap InOutRotation to TargetRotation
+	 * @param DeltaTime Timestep in seconds
+	 */
+	UFUNCTION(BlueprintCallable, meta = (Category = "Math|Spring|Experimental"))
+	static ENGINE_API void CriticalDoubleSpringDampQuat(UPARAM(ref) FQuat& InOutRotation, UPARAM(ref) FVector& InOutAngularVelocity, UPARAM(ref) FQuat& InOutIntermediateRotation, UPARAM(ref) FVector& InOutIntermediateAngularVelocity, const FQuat& TargetRotation, float DeltaTime, float SmoothingTime = 0.2f);
+
+	UE_EXPERIMENTAL(5.7, "SpringMath is experimental")
+	/** Interpolates the value InOutRotation towards TargetRotation using a critically damped spring via an intermediate value - producing a more S-shaped curve. The velocity of InOutRotation is stored in InOutAngularVelocity in deg/s.
+	 *
+	 * @param InOutRotation The value to be damped
+	 * @param InOutAngularVelocity The speed of the value to be damped in deg/s
+	 * @param InOutIntermediateRotation The intermediate state of the double spring
+	 * @param InOutIntermediateAngularVelocity The intermediate velocity of the double spring
+	 * @param TargetRotation The goal to damp towards
+	 * @param SmoothingTime The smoothing time to use for the spring. Longer times lead to more damped behaviour. A time of 0 will snap InOutRotation to TargetRotation
+	 * @param DeltaTime Timestep in seconds
+	 */
+	UFUNCTION(BlueprintCallable, meta = (Category = "Math|Spring|Experimental"))
+	static ENGINE_API void CriticalDoubleSpringDampRotator(UPARAM(ref) FRotator& InOutRotation, UPARAM(ref) FVector& InOutAngularVelocity, UPARAM(ref) FRotator& InOutIntermediateRotation, UPARAM(ref) FVector& InOutIntermediateAngularVelocity, const FRotator& TargetRotation, float DeltaTime, float SmoothingTime = 0.2f);
+
+
+	// Velocity Spring //
+
 	UE_EXPERIMENTAL(5.7, "SpringMath is experimental")
 	/** A velocity spring will damp towards a target that follows a fixed linear target velocity, allowing control of the interpolation speed
 	 * while still giving a smoothed behavior. A SmoothingTime of 0 will give a linear interpolation between X and TargetX
@@ -131,6 +221,7 @@ public:
 	UFUNCTION(BlueprintCallable, meta=(Category="Math|Spring|Experimental"))
 	static ENGINE_API void VelocitySpringDampVector2D(UPARAM(ref) FVector2D& InOutX, UPARAM(ref) FVector2D& InOutV, UPARAM(ref) FVector2D& InOutVi, const FVector2D& TargetX, float MaxSpeed, float DeltaTime,  float SmoothingTime = 0.2f);
 
+
 	// Damper //
 
 	UE_EXPERIMENTAL(5.7, "SpringMath is experimental")
@@ -211,6 +302,9 @@ public:
 	UFUNCTION(BlueprintPure, meta=(Category="Math|Spring|Experimental"))
 	static ENGINE_API FRotator DampRotator(const FRotator& Rotation, const FRotator& TargetRotation, float DeltaTime, float SmoothingTime = 0.2f);
 
+
+	// Character //
+
 	UE_EXPERIMENTAL(5.7, "SpringMath is experimental")
 	/** Update the position of a character given a target velocity using a simple damped spring
 	 * 
@@ -241,7 +335,8 @@ public:
 	UFUNCTION(BlueprintCallable, meta=(Category="Math|Spring|Experimental"))
 	static ENGINE_API void VelocitySpringCharacterUpdate(UPARAM(ref) FVector& InOutPosition, UPARAM(ref) FVector& InOutVelocity, UPARAM(ref) FVector& InOutVelocityIntermediate, UPARAM(ref) FVector& InOutAcceleration, const FVector& TargetVelocity, float DeltaTime, float SmoothingTime = 0.2f, float MaxAcceleration = 1.0f);
 
-	// Conversion Methods
+
+	// Conversion Methods //
 
 	/** Convert from smoothing time to spring strength.
 	* 
diff --git a/Engine/Source/Runtime/Engine/Private/BlueprintSpringMathLibrary.cpp b/Engine/Source/Runtime/Engine/Private/BlueprintSpringMathLibrary.cpp
--- a/Engine/Source/Runtime/Engine/Private/BlueprintSpringMathLibrary.cpp
+++ b/Engine/Source/Runtime/Engine/Private/BlueprintSpringMathLibrary.cpp
@@ -48,6 +48,58 @@ void UBlueprintSpringMathLibrary::CriticalSpringDampRotator(FRotator& InOutRotat
 	InOutAngularVelocity = FMath::RadiansToDegrees(InOutAngularVelocityRadians);
 }
 
+void UBlueprintSpringMathLibrary::CriticalDoubleSpringDampVector(FVector& InOutX, FVector& InOutV, FVector& InOutXi, FVector& InOutVi, const FVector& TargetX, float DeltaTime, float SmoothingTime)
+{
+	SpringMath::CriticalDoubleSpringDamper(InOutX, InOutV, InOutXi, InOutVi, TargetX, SmoothingTime, DeltaTime);
+}
+
+void UBlueprintSpringMathLibrary::CriticalDoubleSpringDampVector2D(FVector2D& InOutX, FVector2D& InOutV, FVector2D& InOutXi, FVector2D& InOutVi, const FVector2D& TargetX, float DeltaTime, float SmoothingTime)
+{
+	SpringMath::CriticalDoubleSpringDamper(InOutX, InOutV, InOutXi, InOutVi, TargetX, SmoothingTime, DeltaTime);
+}
+
+void UBlueprintSpringMathLibrary::CriticalDoubleSpringDampFloat(float& InOutX, float& InOutV, float& InOutXi, float& InOutVi, const float& TargetX, float DeltaTime, float SmoothingTime)
+{
+	SpringMath::CriticalDoubleSpringDamper(InOutX, InOutV, InOutXi, InOutVi, TargetX, SmoothingTime, DeltaTime);
+}
+
+void UBlueprintSpringMathLibrary::CriticalDoubleSpringDampAngle(float& InOutAngle, float& InOutAngularVelocity, float& InOutIntermediateAngle, float& InOutIntermediateAngularVelocity, const float& TargetAngle, float DeltaTime, float SmoothingTime)
+{
+	float InOutAngleRadians = FMath::DegreesToRadians(InOutAngle);
+	float InOutAngularVelocityRadians = FMath::DegreesToRadians(InOutAngularVelocity);
+	float InOutIntermediateAngleRadians = FMath::DegreesToRadians(InOutIntermediateAngle);
+	float InOutIntermediateAngularVelocityRadians = FMath::DegreesToRadians(InOutIntermediateAngularVelocity);
+	SpringMath::CriticalDoubleSpringDamperAngle(InOutAngleRadians, InOutAngularVelocityRadians, InOutIntermediateAngleRadians, InOutIntermediateAngularVelocityRadians, FMath::DegreesToRadians(TargetAngle), SmoothingTime, DeltaTime);
+	InOutAngle = FMath::RadiansToDegrees(InOutAngleRadians);
+	InOutAngularVelocity = FMath::RadiansToDegrees(InOutAngularVelocityRadians);
+	InOutIntermediateAngle = FMath::RadiansToDegrees(InOutIntermediateAngleRadians);
+	InOutIntermediateAngularVelocity = FMath::RadiansToDegrees(InOutIntermediateAngularVelocityRadians);
+}
+
+void UBlueprintSpringMathLibrary::CriticalDoubleSpringDampQuat(FQuat& InOutRotation, FVector& InOutAngularVelocity, FQuat& InOutIntermediateRotation, FVector& InOutIntermediateAngularVelocity, const FQuat& TargetRotation,
+	float DeltaTime, float SmoothingTime)
+{
+	FVector InOutAngularVelocityRadians = FMath::DegreesToRadians(InOutAngularVelocity);
+	FVector InOutIntermediateAngularVelocityRadians = FMath::DegreesToRadians(InOutIntermediateAngularVelocity);
+	SpringMath::CriticalDoubleSpringDamperQuat(InOutRotation, InOutAngularVelocityRadians, InOutIntermediateRotation, InOutIntermediateAngularVelocityRadians, TargetRotation, SmoothingTime, DeltaTime);
+	InOutAngularVelocity = FMath::RadiansToDegrees(InOutAngularVelocityRadians);
+	InOutIntermediateAngularVelocity = FMath::RadiansToDegrees(InOutIntermediateAngularVelocityRadians);
+}
+
+void UBlueprintSpringMathLibrary::CriticalDoubleSpringDampRotator(FRotator& InOutRotation, FVector& InOutAngularVelocity, FRotator& InOutIntermediateRotation, FVector& InOutIntermediateAngularVelocity, const FRotator& TargetRotation,
+	float DeltaTime, float SmoothingTime)
+{
+	FQuat InOutRotationQuat = InOutRotation.Quaternion();
+	FVector InOutAngularVelocityRadians = FMath::DegreesToRadians(InOutAngularVelocity);
+	FQuat InOutIntermediateRotationQuat = InOutIntermediateRotation.Quaternion();
+	FVector InOutIntermediateAngularVelocityRadians = FMath::DegreesToRadians(InOutIntermediateAngularVelocity);
+	SpringMath::CriticalDoubleSpringDamperQuat(InOutRotationQuat, InOutAngularVelocityRadians, InOutIntermediateRotationQuat, InOutIntermediateAngularVelocityRadians, TargetRotation.Quaternion(), SmoothingTime, DeltaTime);
+	InOutRotation = InOutRotationQuat.Rotator();
+	InOutAngularVelocity = FMath::RadiansToDegrees(InOutAngularVelocityRadians);
+	InOutIntermediateRotation = InOutIntermediateRotationQuat.Rotator();
+	InOutIntermediateAngularVelocity = FMath::RadiansToDegrees(InOutIntermediateAngularVelocityRadians);
+}
+
 void UBlueprintSpringMathLibrary::VelocitySpringDampFloat(float& InOutX, float& InOutV, float& InOutVi, float TargetX, float MaxSpeed,
                                                           float DeltaTime, float SmoothingTime)
 {
diff --git a/Engine/Source/Runtime/Engine/Public/Animation/SpringMath.h b/Engine/Source/Runtime/Engine/Public/Animation/SpringMath.h
--- a/Engine/Source/Runtime/Engine/Public/Animation/SpringMath.h
+++ b/Engine/Source/Runtime/Engine/Public/Animation/SpringMath.h
@@ -144,12 +144,92 @@ public:
 		FVector J1 = InOutAngularVelocityRadians + J0 * Y;
 
 		float EyDt = FMath::InvExpApprox(Y * DeltaTime);
-
 		
 		InOutRotation = FQuat::MakeFromRotationVector(EyDt * (J0 + J1 * DeltaTime)) * TargetRotation;
 		InOutAngularVelocityRadians = EyDt * (InOutAngularVelocityRadians - J1 * Y * DeltaTime);
 	}
 
+	/** A double spring damper that interpolates the value InOutX towards TargetX via some intermediate state - producing a more S-shaped curve. 
+	 * The velocity of InOutX is stored in InOutV.
+	 *
+	 * @tparam T The type to be damped
+	 * @param InOutX The value to be damped
+	 * @param InOutV The speed of the value to be damped
+	 * @param InOutXi The value of the intermediate state
+	 * @param InOutVi The velocity of the intermediate state
+	 * @param TargetX The goal to damp towards
+	 * @param SmoothingTime The smoothing time to use for the spring. Longer times lead to more damped behaviour. A time of 0 will snap InOutX to TargetX
+	 * @param DeltaTime Timestep in seconds
+	 */
+	template <typename T>
+	UE_EXPERIMENTAL(5.7, "SpringMath is experimental")
+		static void CriticalDoubleSpringDamper(
+			T& InOutX,
+			T& InOutV,
+			T& InOutXi,
+			T& InOutVi,
+			T TargetX,
+			float SmoothingTime,
+			float DeltaTime)
+	{
+		PRAGMA_DISABLE_EXPERIMENTAL_WARNINGS
+		CriticalSpringDamper(InOutXi, InOutVi, TargetX, SmoothingTime / 2.0f, DeltaTime);
+		CriticalSpringDamper(InOutX, InOutV, InOutXi, SmoothingTime / 2.0f, DeltaTime);
+		PRAGMA_ENABLE_EXPERIMENTAL_WARNINGS
+	}
+
+	/** Specialized angle version of CriticalDoubleSpringDamper that handles angle wrapping.
+	 *
+	 * @param InOutAngleRadians The value to be damped
+	 * @param InOutAngularVelocityRadians The speed of the value to be damped
+	 * @param InOutIntermediateAngleRadians The value of the intermediate state
+	 * @param InOutIntermediateAngularVelocityRadians The velocity of the intermediate state
+	 * @param TargetAngleRadians The goal to damp towards
+	 * @param SmoothingTime The smoothing time to use for the spring. Longer times lead to more damped behaviour. A time of 0 will snap InOutAngleRadians to TargetAngleRadians
+	 * @param DeltaTime Timestep in seconds
+	 */
+	UE_EXPERIMENTAL(5.7, "SpringMath is experimental")
+		static void CriticalDoubleSpringDamperAngle(
+			float& InOutAngleRadians,
+			float& InOutAngularVelocityRadians,
+			float& InOutIntermediateAngleRadians,
+			float& InOutIntermediateAngularVelocityRadians,
+			float TargetAngleRadians,
+			float SmoothingTime,
+			float DeltaTime)
+	{
+		PRAGMA_DISABLE_EXPERIMENTAL_WARNINGS
+		CriticalSpringDamperAngle(InOutIntermediateAngleRadians, InOutIntermediateAngularVelocityRadians, TargetAngleRadians, SmoothingTime / 2.0f, DeltaTime);
+		CriticalSpringDamperAngle(InOutAngleRadians, InOutAngularVelocityRadians, InOutIntermediateAngleRadians, SmoothingTime / 2.0f, DeltaTime);
+		PRAGMA_ENABLE_EXPERIMENTAL_WARNINGS
+	}
+
+	/** Specialized quaternion version of CriticalDoubleSpringDamper, uses FVector for angular velocity
+	 *
+	 * @param InOutRotation The value to be damped
+	 * @param InOutAngularVelocityRadians The angular velocity of the rotation in radians
+	 * @param InOutIntermediateRotationRadians The value of the intermediate state
+	 * @param InOutIntermediateAngularVelocityRadians The velocity of the intermediate state
+	 * @param TargetRotation The target rotation to damp towards
+	 * @param SmoothingTime he smoothing time to use for the spring. Longer times lead to more damped behaviour. A time of 0 will snap InOutRotation to TargetRotation
+	 * @param DeltaTime
+	 */
+	UE_EXPERIMENTAL(5.7, "SpringMath is experimental")
+		static void CriticalDoubleSpringDamperQuat(
+			FQuat& InOutRotation,
+			FVector& InOutAngularVelocityRadians,
+			FQuat& InOutIntermediateRotation,
+			FVector& InOutIntermediateAngularVelocityRadians,
+			const FQuat& TargetRotation,
+			float SmoothingTime,
+			float DeltaTime)
+	{
+		PRAGMA_DISABLE_EXPERIMENTAL_WARNINGS
+		CriticalSpringDamperQuat(InOutIntermediateRotation, InOutIntermediateAngularVelocityRadians, TargetRotation, SmoothingTime / 2.0f, DeltaTime);
+		CriticalSpringDamperQuat(InOutRotation, InOutAngularVelocityRadians, InOutIntermediateRotation, SmoothingTime / 2.0f, DeltaTime);
+		PRAGMA_ENABLE_EXPERIMENTAL_WARNINGS
+	}
+
 	/** A velocity spring will damp towards a target that follows a fixed linear target velocity, allowing control of the interpolation speed
 	 * while still giving a smoothed behavior. A SmoothingTime of 0 will give a linear interpolation between X and TargetX
 	 * 
