Fixed typos in Spring Math documentation

diff --git a/Engine/Source/Runtime/Engine/Public/Animation/SpringMath.h b/Engine/Source/Runtime/Engine/Public/Animation/SpringMath.h
--- a/Engine/Source/Runtime/Engine/Public/Animation/SpringMath.h
+++ b/Engine/Source/Runtime/Engine/Public/Animation/SpringMath.h
@@ -123,7 +123,7 @@ public:
 	 * @param InOutRotation The value to be damped
 	 * @param InOutAngularVelocityRadians The angular velocity of the rotation in radians
 	 * @param TargetRotation The target rotation to damp towards
-	 * @param SmoothingTime he smoothing time to use for the spring. Longer times lead to more damped behaviour. A time of 0 will snap InOutRotation to TargetRotation
+	 * @param SmoothingTime The smoothing time to use for the spring. Longer times lead to more damped behaviour. A time of 0 will snap InOutRotation to TargetRotation
 	 * @param DeltaTime
 	 */
 	UE_EXPERIMENTAL(5.7, "SpringMath is experimental")
@@ -208,10 +208,10 @@ public:
 	 *
 	 * @param InOutRotation The value to be damped
 	 * @param InOutAngularVelocityRadians The angular velocity of the rotation in radians
-	 * @param InOutIntermediateRotationRadians The value of the intermediate state
+	 * @param InOutIntermediateRotation The value of the intermediate state
 	 * @param InOutIntermediateAngularVelocityRadians The velocity of the intermediate state
 	 * @param TargetRotation The target rotation to damp towards
-	 * @param SmoothingTime he smoothing time to use for the spring. Longer times lead to more damped behaviour. A time of 0 will snap InOutRotation to TargetRotation
+	 * @param SmoothingTime The smoothing time to use for the spring. Longer times lead to more damped behaviour. A time of 0 will snap InOutRotation to TargetRotation
 	 * @param DeltaTime
 	 */
 	UE_EXPERIMENTAL(5.7, "SpringMath is experimental")
