Update function signature to correct delegate type

#FailedValidations Preflight

diff --git a/Engine/Source/Runtime/Engine/Classes/Animation/Skeleton.h b/Engine/Source/Runtime/Engine/Classes/Animation/Skeleton.h
--- a/Engine/Source/Runtime/Engine/Classes/Animation/Skeleton.h
+++ b/Engine/Source/Runtime/Engine/Classes/Animation/Skeleton.h
@@ -966,7 +966,7 @@ public:
 	ENGINE_API void UnregisterOnSkeletonHierarchyChanged(FDelegateUserObject Unregister);
 
 	/** Registers a delegate to be called after the reference skeleton has changed */
-	ENGINE_API void RegisterOnReferenceSkeletonChanged(const FOnSkeletonHierarchyChanged& Delegate);
+	ENGINE_API void RegisterOnReferenceSkeletonChanged(const FOnReferenceSkeletonChanged& Delegate);
 	ENGINE_API void UnregisterOnReferenceSkeletonChanged(FDelegateUserObject Unregister);
 
 	/** Removes the supplied bones from the skeleton */
diff --git a/Engine/Source/Runtime/Engine/Private/Animation/Skeleton.cpp b/Engine/Source/Runtime/Engine/Private/Animation/Skeleton.cpp
--- a/Engine/Source/Runtime/Engine/Private/Animation/Skeleton.cpp
+++ b/Engine/Source/Runtime/Engine/Private/Animation/Skeleton.cpp
@@ -1492,7 +1492,7 @@ void USkeleton::UnregisterOnSkeletonHierarchyChanged(FDelegateUserObject Unregis
 	OnSkeletonHierarchyChanged.RemoveAll(Unregister);
 }
 
-void USkeleton::RegisterOnReferenceSkeletonChanged(const FOnSkeletonHierarchyChanged& Delegate)
+void USkeleton::RegisterOnReferenceSkeletonChanged(const FOnReferenceSkeletonChanged& Delegate)
 {
 	OnReferenceSkeletonChanged.Add(Delegate);
 }
