[WorldPartition] Added ENGINE_API to WorldPartitionRuntimeCellTransformer methods (PR #13922)

diff --git a/Engine/Source/Runtime/Engine/Private/WorldPartition/WorldPartitionRuntimeCellTransformer.cpp b/Engine/Source/Runtime/Engine/Private/WorldPartition/WorldPartitionRuntimeCellTransformer.cpp
--- a/Engine/Source/Runtime/Engine/Private/WorldPartition/WorldPartitionRuntimeCellTransformer.cpp
+++ b/Engine/Source/Runtime/Engine/Private/WorldPartition/WorldPartitionRuntimeCellTransformer.cpp
@@ -9,6 +9,23 @@ UWorldPartitionRuntimeCellTransformer::UWorldPartitionRuntimeCellTransformer(con
 {}
 
 #if WITH_EDITOR
+bool UWorldPartitionRuntimeCellTransformer::IsLevelTransformable(ULevel* InLevel) const
+{
+	return true;
+}
+
+void UWorldPartitionRuntimeCellTransformer::PreTransform(ULevel* InLevel)
+{
+}
+
+void UWorldPartitionRuntimeCellTransformer::Transform(ULevel* InLevel)
+{
+}
+
+void UWorldPartitionRuntimeCellTransformer::PostTransform(ULevel* InLevel)
+{
+}
+
 void UWorldPartitionRuntimeCellTransformer::ForEachIgnoredComponentClass(TFunctionRef<bool(const TSubclassOf<UActorComponent>&)> Func) const
 {
 	for (const TSubclassOf<UActorComponent>& IgnoredComponentClass : GetDefault<UWorldPartitionRuntimeCellTransformerSettings>()->IgnoredComponentClasses)
diff --git a/Engine/Source/Runtime/Engine/Public/WorldPartition/WorldPartitionRuntimeCellTransformer.h b/Engine/Source/Runtime/Engine/Public/WorldPartition/WorldPartitionRuntimeCellTransformer.h
--- a/Engine/Source/Runtime/Engine/Public/WorldPartition/WorldPartitionRuntimeCellTransformer.h
+++ b/Engine/Source/Runtime/Engine/Public/WorldPartition/WorldPartitionRuntimeCellTransformer.h
@@ -5,6 +5,8 @@
 #include "Components/ActorComponent.h"
 #include "WorldPartitionRuntimeCellTransformer.generated.h"
 
+#define UE_API ENGINE_API
+
 UCLASS(MinimalAPI, Config=Engine)
 class UWorldPartitionRuntimeCellTransformerSettings : public UObject
 {
@@ -28,26 +30,28 @@ class UWorldPartitionRuntimeCellTransformer : public UObject
 #if WITH_EDITOR
 public:
 	/** Whether the level can be processed by the transformer. */
-	virtual bool IsLevelTransformable(ULevel* InLevel) const { return true; }
-	virtual void PreTransform(ULevel* InLevel) {}
-	virtual void Transform(ULevel* InLevel) {}
-	virtual void PostTransform(ULevel* InLevel) {}
-	ENGINE_API virtual void ForEachIgnoredComponentClass(TFunctionRef<bool(const TSubclassOf<UActorComponent>&)> Func) const;
-	ENGINE_API virtual void ForEachIgnoredExactComponentClass(TFunctionRef<bool(const TSubclassOf<UActorComponent>&)> Func) const;
+	UE_API virtual bool IsLevelTransformable(ULevel* InLevel) const;
+	UE_API virtual void PreTransform(ULevel* InLevel);
+	UE_API virtual void Transform(ULevel* InLevel);
+	UE_API virtual void PostTransform(ULevel* InLevel);
+	UE_API virtual void ForEachIgnoredComponentClass(TFunctionRef<bool(const TSubclassOf<UActorComponent>&)> Func) const;
+	UE_API virtual void ForEachIgnoredExactComponentClass(TFunctionRef<bool(const TSubclassOf<UActorComponent>&)> Func) const;
 #endif
 
 	bool IsEnabled() const { return bEnabled; }
 
 protected:
 #if WITH_EDITOR
-	ENGINE_API virtual bool CanIgnoreComponent(const UActorComponent* InComponent) const;
+	UE_API virtual bool CanIgnoreComponent(const UActorComponent* InComponent) const;
 #endif
 	
 protected:
 	// Tag used to force exclude actors from any cell transformation
-	ENGINE_API static const FName NAME_CellTransformerIgnoreActor;
+	UE_API static const FName NAME_CellTransformerIgnoreActor;
 
 private:
 	UPROPERTY(EditAnywhere, Category = Transformer)
 	bool bEnabled = true;
-};
\ No newline at end of file
+};
+
+#undef UE_API
\ No newline at end of file
