Added CVar p.Chaos.AsyncPhysicsStateTask.PriorityExtraTimeBudgetMS to grant additional time for the the async physics state task if a high-priority load is in progress or the world is undergoing seamless travel (disabled by default).

diff --git a/Engine/Source/Runtime/Engine/Private/PhysicsEngine/Experimental/PhysScene_AsyncPhysicsStateJobQueue.cpp b/Engine/Source/Runtime/Engine/Private/PhysicsEngine/Experimental/PhysScene_AsyncPhysicsStateJobQueue.cpp
--- a/Engine/Source/Runtime/Engine/Private/PhysicsEngine/Experimental/PhysScene_AsyncPhysicsStateJobQueue.cpp
+++ b/Engine/Source/Runtime/Engine/Private/PhysicsEngine/Experimental/PhysScene_AsyncPhysicsStateJobQueue.cpp
@@ -5,6 +5,7 @@
 #include "PhysicsEngine/BodySetup.h"
 #include "Engine/World.h"
 #include "Components/ActorComponent.h"
+#include "GameFramework/WorldSettings.h"
 
 namespace Chaos
 {
@@ -12,8 +13,16 @@ namespace Chaos
 	static FAutoConsoleVariableRef CVarAsyncPhysicsStateTaskTimeBudgetMS(
 		TEXT("p.Chaos.AsyncPhysicsStateTask.TimeBudgetMS"),
 		GAsyncPhysicsStateTaskTimeBudgetMS,
-		TEXT("Maximum time budget in milliseconds for the async physics state task (0 = no time limit)"),
+		TEXT("Maximum time budget (in milliseconds) for the async physics state task (0 = unlimited)"),
 		ECVF_Default);
+
+	static float GPriorityAsyncPhysicsStateTaskExtraTimeBudgetMS = 0.0f;
+	static FAutoConsoleVariableRef CVarPriorityAsyncPhysicsStateTaskExtraTimeBudgetMS(
+		TEXT("p.Chaos.AsyncPhysicsStateTask.PriorityExtraTimeBudgetMS"),
+		GPriorityAsyncPhysicsStateTaskExtraTimeBudgetMS,
+		TEXT("Additional time budget (in milliseconds) allocated to the async physics state task during high-priority loading or seamless travel."),
+		ECVF_Default
+	);
 }
 
 FPhysScene_AsyncPhysicsStateJobQueue::FPhysScene_AsyncPhysicsStateJobQueue(FPhysScene* InPhysicScene)
@@ -125,7 +134,21 @@ void FPhysScene_AsyncPhysicsStateJobQueue::LaunchAsyncJobTask()
 		}
 	}
 
-	const double TimeBudgetSeconds = Chaos::GAsyncPhysicsStateTaskTimeBudgetMS > 0 ? Chaos::GAsyncPhysicsStateTaskTimeBudgetMS / 1000 : MAX_dbl;
+	// Grant additional time if a high-priority load is in progress or the world is undergoing seamless travel.
+	float ExtraTimeMS = 0.f;
+	if (Chaos::GAsyncPhysicsStateTaskTimeBudgetMS > 0.f && Chaos::GPriorityAsyncPhysicsStateTaskExtraTimeBudgetMS > 0.f)
+	{
+		const UWorld* World = PhysScene->GetOwningWorld();
+		if (AWorldSettings* WorldSettings = World ? World->GetWorldSettings(false, false) : nullptr)
+		{
+			if (WorldSettings->bHighPriorityLoading || WorldSettings->bHighPriorityLoadingLocal || (World && World->IsInSeamlessTravel()))
+			{
+				ExtraTimeMS = Chaos::GPriorityAsyncPhysicsStateTaskExtraTimeBudgetMS;
+			}
+		}
+	}
+
+	double TimeBudgetSeconds = Chaos::GAsyncPhysicsStateTaskTimeBudgetMS > 0 ? (Chaos::GAsyncPhysicsStateTaskTimeBudgetMS + ExtraTimeMS)/1000 : MAX_dbl;
 	AsyncJobTask = UE::Tasks::Launch(UE_SOURCE_LOCATION, [this, TimeBudgetSeconds]() { ExecuteJobsAsync(TimeBudgetSeconds); }, UE::Tasks::ETaskPriority::BackgroundHigh);
 }
 
