Fix FDataLayerInstanceDesc::operator< not sorting properly since CL24818573.

diff --git a/Engine/Source/Runtime/Engine/Private/WorldPartition/DataLayer/WorldDataLayersActorDesc.cpp b/Engine/Source/Runtime/Engine/Private/WorldPartition/DataLayer/WorldDataLayersActorDesc.cpp
--- a/Engine/Source/Runtime/Engine/Private/WorldPartition/DataLayer/WorldDataLayersActorDesc.cpp
+++ b/Engine/Source/Runtime/Engine/Private/WorldPartition/DataLayer/WorldDataLayersActorDesc.cpp
@@ -99,7 +99,7 @@ bool operator < (const FDataLayerInstanceDesc& Lhs, const FDataLayerInstanceDesc
 							{
 								if (Lhs.bIsPrivate == Rhs.bIsPrivate)
 								{
-									return (int)Lhs.bPrivateDataLayerSupportsActorFilter == (int)Rhs.bPrivateDataLayerSupportsActorFilter;
+									return (int)Lhs.bPrivateDataLayerSupportsActorFilter < (int)Rhs.bPrivateDataLayerSupportsActorFilter;
 								}
 
 								return (int)Lhs.bIsPrivate < (int)Rhs.bIsPrivate;
