[WorldPartition] Fixed UWorldPartitionStreamingPolicy streaming sources velocity.
The optimized code path was calling UWorldPartitionSubsystem::GetStreamingSources to update the velocity since CL 30969827.
In this case, the velocity was transformed to local space twice since CL 32514660, because UWorldPartitionSubsystem::GetStreamingSources now also transforms the velocity to local space.
Added UWorldPartitionSubsystem::UpdateStreamingSourcesVelocity that updates provided streaming sources velocity.

#code-analyzed

diff --git a/Engine/Source/Runtime/Engine/Private/WorldPartition/WorldPartitionStreamingPolicy.cpp b/Engine/Source/Runtime/Engine/Private/WorldPartition/WorldPartitionStreamingPolicy.cpp
--- a/Engine/Source/Runtime/Engine/Private/WorldPartition/WorldPartitionStreamingPolicy.cpp
+++ b/Engine/Source/Runtime/Engine/Private/WorldPartition/WorldPartitionStreamingPolicy.cpp
@@ -118,15 +118,8 @@ void UWorldPartitionStreamingPolicy::UpdateStreamingSources(bool bCanOptimizeUpd
 	const uint32 NewUpdateStreamingSourcesHash = WorldPartitionSubsystem->GetStreamingSourcesHash();
 	if (bCanOptimizeUpdate && (UpdateStreamingSourcesHash == NewUpdateStreamingSourcesHash))
 	{
-		TArray<FWorldPartitionStreamingSource> LocalStreamingSources;
-		WorldPartitionSubsystem->GetStreamingSources(WorldPartition, LocalStreamingSources);
-		check(LocalStreamingSources.Num() == CurrentState.StreamingSources.Num());
-		const FTransform WorldToLocal = WorldPartition->GetInstanceTransform().Inverse();
-		for (int32 i=0; i<LocalStreamingSources.Num(); i++)
-		{
-			check(CurrentState.StreamingSources[i].Name == LocalStreamingSources[i].Name);
-			CurrentState.StreamingSources[i].Velocity = WorldToLocal.TransformVector(LocalStreamingSources[i].Velocity);
-		}
+		// Streaming Source Hash doesn't take velocity into account, we still need to update the streaming sources velocity
+		WorldPartitionSubsystem->UpdateStreamingSourcesVelocity(WorldPartition, CurrentState.StreamingSources);
 		return;
 	}
 
diff --git a/Engine/Source/Runtime/Engine/Private/WorldPartition/WorldPartitionSubsystem.cpp b/Engine/Source/Runtime/Engine/Private/WorldPartition/WorldPartitionSubsystem.cpp
--- a/Engine/Source/Runtime/Engine/Private/WorldPartition/WorldPartitionSubsystem.cpp
+++ b/Engine/Source/Runtime/Engine/Private/WorldPartition/WorldPartitionSubsystem.cpp
@@ -1101,6 +1101,22 @@ void UWorldPartitionSubsystem::GetStreamingSources(const UWorldPartition* InWorl
 	}
 }
 
+void UWorldPartitionSubsystem::UpdateStreamingSourcesVelocity(const UWorldPartition* InWorldPartition, TArray<FWorldPartitionStreamingSource>& InOutStreamingSources) const
+{
+	if (InOutStreamingSources.Num())
+	{
+		const FTransform WorldToLocal = InWorldPartition->GetInstanceTransform().Inverse();
+		for (FWorldPartitionStreamingSource& StreamingSourceToUpdate : InOutStreamingSources)
+		{
+			// Here we can't assume that StreamingSources matches InOutStreamingSources, so we find the streaming sources based on its name
+			if (const FWorldPartitionStreamingSource* StreamingSource = StreamingSources.FindByPredicate([Name = StreamingSourceToUpdate.Name](const FWorldPartitionStreamingSource& Source) { return Source.Name == Name; }))
+			{
+				StreamingSourceToUpdate.Velocity = WorldToLocal.TransformVector(StreamingSource->Velocity);
+			}
+		}
+	}
+}
+
 void UWorldPartitionSubsystem::OnWorldBeginPlay(UWorld& InWorld)
 {
 	Super::OnWorldBeginPlay(InWorld);
diff --git a/Engine/Source/Runtime/Engine/Public/WorldPartition/WorldPartitionSubsystem.h b/Engine/Source/Runtime/Engine/Public/WorldPartition/WorldPartitionSubsystem.h
--- a/Engine/Source/Runtime/Engine/Public/WorldPartition/WorldPartitionSubsystem.h
+++ b/Engine/Source/Runtime/Engine/Public/WorldPartition/WorldPartitionSubsystem.h
@@ -149,6 +149,7 @@ private:
 	// Streaming Sources
 	void UpdateStreamingSources();
 	ENGINE_API void GetStreamingSources(const UWorldPartition* InWorldPartition, TArray<FWorldPartitionStreamingSource>& OutStreamingSources) const;
+	ENGINE_API void UpdateStreamingSourcesVelocity(const UWorldPartition* InWorldPartition, TArray<FWorldPartitionStreamingSource>& InOutStreamingSources) const;
 	uint32 GetStreamingSourcesHash() const { return StreamingSourcesHash; }
 
 	void OnWorldPartitionInitialized(UWorldPartition* InWorldPartition);
