[WorldPartition] Fixed 2D debug display of RuntimeSpatialHash now displaying velocity vector on streaming sources.

#code-analyzed

diff --git a/Engine/Source/Runtime/Engine/Private/WorldPartition/WorldPartitionRuntimeSpatialHash.cpp b/Engine/Source/Runtime/Engine/Private/WorldPartition/WorldPartitionRuntimeSpatialHash.cpp
--- a/Engine/Source/Runtime/Engine/Private/WorldPartition/WorldPartitionRuntimeSpatialHash.cpp
+++ b/Engine/Source/Runtime/Engine/Private/WorldPartition/WorldPartitionRuntimeSpatialHash.cpp
@@ -1093,6 +1093,15 @@ void FSpatialHashStreamingGrid::Draw2D(const UWorldPartitionRuntimeSpatialHash*
 		{
 			DrawStreamingSource2D(GridScreenBounds, Shape, WorldToScreen, Color, DrawContext);
 		});
+
+		// Draw velocity vector			
+		const FVector2D Velocity2D = FVector2D(Source.Velocity);
+
+		if (!Velocity2D.IsNearlyZero())
+		{
+			const FVector2D Center2D = FVector2D(Source.Location);
+			DrawContext.PushDrawSegment(GridScreenBounds, WorldToScreen(Center2D), WorldToScreen(Center2D + Velocity2D.GetSafeNormal() * GridLoadingRange * 0.5f), Color, 1);
+		}
 	}
 }
 
