Add DefaultRuntimeState to WorldDataLayersActorDesc.DataLayerInstanceDesc.

diff --git a/Engine/Source/Runtime/Core/Public/UObject/FortniteMainBranchObjectVersions.inl b/Engine/Source/Runtime/Core/Public/UObject/FortniteMainBranchObjectVersions.inl
--- a/Engine/Source/Runtime/Core/Public/UObject/FortniteMainBranchObjectVersions.inl
+++ b/Engine/Source/Runtime/Core/Public/UObject/FortniteMainBranchObjectVersions.inl
@@ -697,5 +697,8 @@ FORTNITE_MAIN_VERSION(UAFProceduralSystems, 227)
 // Firefight: version 228 was introduced and backed out, adding it back to allow assets created in the meantime to load
 FORTNITE_MAIN_VERSION(Firefight_228, 228)
 
+// Add default runtime state to FDataLayerInstanceDesc
+FORTNITE_MAIN_VERSION(WorldPartitionDataLayerDefaultRuntimeState, 229)
+
 // -----<new versions can be added above this line>-------------------------------------------------
 #undef FORTNITE_MAIN_VERSION
\ No newline at end of file
diff --git a/Engine/Source/Runtime/Engine/Private/WorldPartition/DataLayer/WorldDataLayersActorDesc.cpp b/Engine/Source/Runtime/Engine/Private/WorldPartition/DataLayer/WorldDataLayersActorDesc.cpp
--- a/Engine/Source/Runtime/Engine/Private/WorldPartition/DataLayer/WorldDataLayersActorDesc.cpp
+++ b/Engine/Source/Runtime/Engine/Private/WorldPartition/DataLayer/WorldDataLayersActorDesc.cpp
@@ -65,6 +65,11 @@ FArchive& operator<<(FArchive& Ar, FDataLayerInstanceDesc& Desc)
 		}
 	}
 
+	if (Ar.CustomVer(FFortniteMainBranchObjectVersion::GUID) >= FFortniteMainBranchObjectVersion::WorldPartitionDataLayerDefaultRuntimeState)
+	{
+		Ar << Desc.InitialRuntimeState;
+	}
+
 	return Ar;
 }
 
@@ -78,7 +83,8 @@ bool operator == (const FDataLayerInstanceDesc& Lhs, const FDataLayerInstanceDes
 		(Lhs.PrivateShortName == Rhs.PrivateShortName) &&
 		(Lhs.bIsIncludedInActorFilterDefault == Rhs.bIsIncludedInActorFilterDefault) &&
 		(Lhs.bIsPrivate == Rhs.bIsPrivate) &&
-		(Lhs.bPrivateDataLayerSupportsActorFilter == Rhs.bPrivateDataLayerSupportsActorFilter);
+		(Lhs.bPrivateDataLayerSupportsActorFilter == Rhs.bPrivateDataLayerSupportsActorFilter) &&
+		(Lhs.InitialRuntimeState == Rhs.InitialRuntimeState);
 }
 
 bool operator < (const FDataLayerInstanceDesc& Lhs, const FDataLayerInstanceDesc& Rhs)
@@ -99,12 +105,14 @@ bool operator < (const FDataLayerInstanceDesc& Lhs, const FDataLayerInstanceDesc
 							{
 								if (Lhs.bIsPrivate == Rhs.bIsPrivate)
 								{
+									if (Lhs.bPrivateDataLayerSupportsActorFilter == Rhs.bPrivateDataLayerSupportsActorFilter)
+									{
+										return (int)Lhs.InitialRuntimeState < (int)Rhs.InitialRuntimeState;
+									}
 									return (int)Lhs.bPrivateDataLayerSupportsActorFilter < (int)Rhs.bPrivateDataLayerSupportsActorFilter;
 								}
-
 								return (int)Lhs.bIsPrivate < (int)Rhs.bIsPrivate;
 							}
-							
 							return (int)Lhs.bIsIncludedInActorFilterDefault < (int)Rhs.bIsIncludedInActorFilterDefault;
 						}
 						return Lhs.PrivateShortName < Rhs.PrivateShortName;
@@ -159,6 +167,11 @@ FString FDataLayerInstanceDesc::GetShortName() const
 	return PrivateShortName;
 }
 
+EDataLayerRuntimeState FDataLayerInstanceDesc::GetInitialRuntimeState() const
+{
+	return InitialRuntimeState;
+}
+
 bool FDataLayerInstanceDesc::SupportsActorFilters() const
 {
 	if (bIsPrivate)
@@ -181,6 +194,7 @@ FDataLayerInstanceDesc::FDataLayerInstanceDesc()
 : bIsUsingAsset(false)
 , bIsIncludedInActorFilterDefault(true)
 , bIsPrivate(false)
+, InitialRuntimeState(EDataLayerRuntimeState::Unloaded)
 , bPrivateDataLayerSupportsActorFilter(false)
 , bDeprecatedIsRuntime(false)
 {
@@ -203,6 +217,7 @@ void FDataLayerInstanceDesc::Init(UDataLayerInstance* InDataLayerInstance)
 	bPrivateDataLayerSupportsActorFilter = InDataLayerInstance->SupportsActorFilters();
 	PrivateShortName = InDataLayerInstance->GetDataLayerShortName();
 	bDeprecatedIsRuntime = InDataLayerInstance->IsRuntime();
+	InitialRuntimeState = InDataLayerInstance->GetInitialRuntimeState();
 }
 
 FWorldDataLayersActorDesc::FWorldDataLayersActorDesc()
diff --git a/Engine/Source/Runtime/Engine/Public/WorldPartition/DataLayer/WorldDataLayersActorDesc.h b/Engine/Source/Runtime/Engine/Public/WorldPartition/DataLayer/WorldDataLayersActorDesc.h
--- a/Engine/Source/Runtime/Engine/Public/WorldPartition/DataLayer/WorldDataLayersActorDesc.h
+++ b/Engine/Source/Runtime/Engine/Public/WorldPartition/DataLayer/WorldDataLayersActorDesc.h
@@ -8,6 +8,8 @@
 #include "WorldPartition/DataLayer/DataLayerType.h"
 #include "Misc/Optional.h"
 
+enum class EDataLayerRuntimeState : uint8;
+
 class FDataLayerInstanceDesc
 {
 public:
@@ -24,6 +26,7 @@ public:
 	ENGINE_API EDataLayerType GetDataLayerType() const;
 	ENGINE_API FString GetShortName() const;
 	bool IsIncludedInActorFilterDefault() const { return bIsIncludedInActorFilterDefault; }
+	ENGINE_API EDataLayerRuntimeState GetInitialRuntimeState() const;
 	ENGINE_API bool SupportsActorFilters() const;
 
 private:
@@ -43,6 +46,8 @@ private:
 	bool bIsIncludedInActorFilterDefault;
 	// DataLayer Asset is Private
 	bool bIsPrivate;
+	// Initial runtime state of this data layer instance.
+	EDataLayerRuntimeState InitialRuntimeState;
 	//~ End UDataLayerInstance
 
 	// If DataLayer Asset is Private store the SupportsActorFilter() flag
