[WorldPartition] Fixed ABlockingVolume actor not showing in PIE in the outliner when the actor is the 2nd actor of a world partition runtime cell. Fixed this by creating the default brush for PIE world partition runtime cell levels.

diff --git a/Engine/Source/Runtime/Engine/Classes/Engine/World.h b/Engine/Source/Runtime/Engine/Classes/Engine/World.h
--- a/Engine/Source/Runtime/Engine/Classes/Engine/World.h
+++ b/Engine/Source/Runtime/Engine/Classes/Engine/World.h
@@ -1831,6 +1831,7 @@ private:
 	FRemoveLevelFromWorldExtensionEvent RemoveLevelFromWorldExtensionEvent;
 
 	friend FWorldPartitionEvents;
+	friend class FWorldPartitionLevelHelper;
 
 #if WITH_EDITOR
 
diff --git a/Engine/Source/Runtime/Engine/Private/WorldPartition/WorldPartitionLevelHelper.cpp b/Engine/Source/Runtime/Engine/Private/WorldPartition/WorldPartitionLevelHelper.cpp
--- a/Engine/Source/Runtime/Engine/Private/WorldPartition/WorldPartitionLevelHelper.cpp
+++ b/Engine/Source/Runtime/Engine/Private/WorldPartition/WorldPartitionLevelHelper.cpp
@@ -454,6 +454,12 @@ ULevel* FWorldPartitionLevelHelper::CreateEmptyLevelForRuntimeCell(const UWorldP
 	check(NewWorld->GetWorldSettings());
 	check(UWorld::FindWorldInPackage(CellPackage) == NewWorld);
 	check(InPackage || (NewWorld->GetPathName() == InWorldAssetName));
+	// Make sure to have the default brush in PIE to fix cases where the editor expect to have a default brush at index 1 in the level's actors array
+	if (InWorld->IsPlayInEditor())
+	{
+		NewWorld->RepairDefaultBrush();
+		check(NewWorld->GetDefaultBrush());
+	}
 	// We don't need the cell level's world setting to replicate
 	FSetActorReplicates SetActorReplicates(NewWorld->GetWorldSettings(), false);
 	
