Fix D3DDebug warnings for ES31 on Windows about missing input layout elements during PSO Precaching

diff --git a/Engine/Source/Runtime/Engine/Private/LocalVertexFactory.cpp b/Engine/Source/Runtime/Engine/Private/LocalVertexFactory.cpp
--- a/Engine/Source/Runtime/Engine/Private/LocalVertexFactory.cpp
+++ b/Engine/Source/Runtime/Engine/Private/LocalVertexFactory.cpp
@@ -329,14 +329,14 @@ void FLocalVertexFactory::ValidateCompiledResult(const FVertexFactoryType* Type,
 void FLocalVertexFactory::GetPSOPrecacheVertexFetchElements(EVertexInputStreamType VertexInputStreamType, FVertexDeclarationElementList& Elements)
 {
 	Elements.Add(FVertexElement(0, 0, VET_Float3, 0, sizeof(float)*3u, false));
-	
+
 	if (VertexInputStreamType == EVertexInputStreamType::PositionAndNormalOnly)
 	{
 		// 2-axis TangentBasis components in a single buffer, hence *2u
 		Elements.Add(FVertexElement(1, 4, VET_PackedNormal, 2, sizeof(FPackedNormal)*2u, false));
 	}
 
-	if (UseGPUScene(GMaxRHIShaderPlatform, GMaxRHIFeatureLevel) 
+	if (UseGPUScene(GMaxRHIShaderPlatform, GMaxRHIFeatureLevel)
 		&& !PlatformGPUSceneUsesUniformBufferView(GMaxRHIShaderPlatform))
 	{
 		switch (VertexInputStreamType)
@@ -344,6 +344,20 @@ void FLocalVertexFactory::GetPSOPrecacheVertexFetchElements(EVertexInputStreamTy
 		case EVertexInputStreamType::Default:
 		{
 			Elements.Add(FVertexElement(1, 0, VET_UInt, 13, sizeof(uint32), true));
+			if (!RHISupportsManualVertexFetch(GMaxRHIShaderPlatform))
+			{
+				Elements.Add(FVertexElement(2, 0, VET_Float3, 1, 36, false));
+				Elements.Add(FVertexElement(2, 12, VET_Float3, 2, 36, false));
+				Elements.Add(FVertexElement(2, 24, VET_Float3, 3, 36, false));
+				Elements.Add(FVertexElement(3, 0, VET_Float3, 4, 48, false));
+				Elements.Add(FVertexElement(3, 12, VET_Float3, 5, 48, false));
+				Elements.Add(FVertexElement(3, 24, VET_Float3, 6, 48, false));
+				Elements.Add(FVertexElement(3, 36, VET_Float3, 7, 48, false));
+				if (IsGPUSkinPassThroughSupported(GMaxRHIShaderPlatform))
+				{
+					Elements.Add(FVertexElement(4, 0, VET_Float3, 14, sizeof(float)*3u, false));
+				}
+			}
 			break;
 		}
 		case EVertexInputStreamType::PositionOnly:
