DSTM: Transactionalizing TickableObjects.

diff --git a/Engine/Source/Runtime/Engine/Private/StreamableManager.cpp b/Engine/Source/Runtime/Engine/Private/StreamableManager.cpp
--- a/Engine/Source/Runtime/Engine/Private/StreamableManager.cpp
+++ b/Engine/Source/Runtime/Engine/Private/StreamableManager.cpp
@@ -177,6 +177,23 @@ public:
 	/** Calls all delegates, call from synchronous flushes */
 	void FlushDelegates()
 	{
+		if (!PendingDelegates.Head)
+		{
+			return;
+		}
+
+		if (AutoRTFM::IsTransactional())
+		{
+			ensureMsgf(false, TEXT("%hs was not expected to be called inside a transaction (it won't complete until the tx ends, so no side effects will be visible during the tx)"), __func__);
+
+			AutoRTFM::OnCommit([this]()
+			{
+				FlushDelegates();
+			});
+
+			return;
+		}
+
 		while (PendingDelegates.Head)
 		{
 			Tick(0.0f);
@@ -192,6 +209,17 @@ public:
 			return;
 		}
 
+		if (AutoRTFM::IsTransactional())
+		{
+			// Complete this transaction, then call this function again outside the transaction
+			AutoRTFM::OnCommit([this, DeltaTime]()
+			{
+				Tick(DeltaTime);
+			});
+
+			return;
+		}
+
 		FPendingDelegateList DelegatesToCall;
 		{
 			UE::TScopeLock Lock(DataLock);
diff --git a/Engine/Source/Runtime/Engine/Private/Tickable.cpp b/Engine/Source/Runtime/Engine/Private/Tickable.cpp
--- a/Engine/Source/Runtime/Engine/Private/Tickable.cpp
+++ b/Engine/Source/Runtime/Engine/Private/Tickable.cpp
@@ -196,13 +196,25 @@ void FTickableGameObject::TickObjects(UWorld* World, const ELevelTick LevelTickT
 					if ((GIsEditor && TickableObject->IsTickableInEditor()) ||
 						(bIsGameWorld && ((!bIsPaused && LevelTickType != LEVELTICK_TimeOnly) || (bIsPaused && TickableObject->IsTickableWhenPaused()))))
 					{
+#if UE_WITH_REMOTE_OBJECT_HANDLE && UE_AUTORTFM
+						UE::RemoteExecutor::EnqueueWork(TickableObject->GetStatId().GetName(), true, [TickableObject, DeltaSeconds]() {
+#endif
+						
 						SCOPE_CYCLE_COUNTER_STATID(TickableObject->GetStatId());
 						TickableObject->Tick(DeltaSeconds);
+						
+#if UE_WITH_REMOTE_OBJECT_HANDLE && UE_AUTORTFM
+						});
+#endif
 					}
 				}
 			}
 		}
 
+#if UE_WITH_REMOTE_OBJECT_HANDLE && UE_AUTORTFM
+		UE::RemoteExecutor::ExecutePendingWork();
+#endif
+
 		Statics.FinishTicking();
 	}
 }
