Added additional error information to the streamable manager errors.

diff --git a/Engine/Source/Runtime/Engine/Classes/Engine/Experimental/StreamableManagerError.h b/Engine/Source/Runtime/Engine/Classes/Engine/Experimental/StreamableManagerError.h
--- a/Engine/Source/Runtime/Engine/Classes/Engine/Experimental/StreamableManagerError.h
+++ b/Engine/Source/Runtime/Engine/Classes/Engine/Experimental/StreamableManagerError.h
@@ -39,6 +39,7 @@ namespace UE::UnifiedError::StreamableManager
 	public:
 		FString DebugName;
 		TArray<FSoftObjectPath> RequestedAssets;
+		bool bAttemptedDownload = false;
 	};
 	
 }
@@ -58,8 +59,9 @@ inline void SerializeForLog(FCbWriter& Writer, const UE::UnifiedError::Streamabl
 {
 	Writer.BeginObject(); 
 	Writer.AddString(ANSITEXTVIEW("$type"), TErrorStructFeatures<UE::UnifiedError::StreamableManager::FStreamableManagerRequestContext>::GetErrorContextTypeNameAsString());
-	Writer.AddString(ANSITEXTVIEW("$format"), TEXT("(RequestDebugName: {DebugName}, RequestedAssets: {RequestedAssets})"));
+	Writer.AddString(ANSITEXTVIEW("$format"), TEXT("(RequestDebugName: {DebugName}, AttemptedDownload:{AttemptedDownload}, RequestedAssets: {RequestedAssets})"));
 	Writer.AddString(ANSITEXTVIEW("DebugName"), Context.DebugName);
+	Writer.AddBool(ANSITEXTVIEW("AttemptedDownload"), Context.bAttemptedDownload);
 	Writer.BeginArray(ANSITEXTVIEW("RequestedAssets"));
 	for (const FSoftObjectPath& RequestedAsset : Context.RequestedAssets )
 	{
diff --git a/Engine/Source/Runtime/Engine/Private/StreamableManager.cpp b/Engine/Source/Runtime/Engine/Private/StreamableManager.cpp
--- a/Engine/Source/Runtime/Engine/Private/StreamableManager.cpp
+++ b/Engine/Source/Runtime/Engine/Private/StreamableManager.cpp
@@ -1270,7 +1270,7 @@ void FStreamableHandle::SetError(const UE::UnifiedError::FError& InError)
 	{
 		Error = InError;
 		// add some additional context for this specific request 
-		UE::UnifiedError::StreamableManager::FStreamableManagerRequestContext RequestContext = { DebugName, RequestedAssets };
+		UE::UnifiedError::StreamableManager::FStreamableManagerRequestContext RequestContext = { DebugName, RequestedAssets, DownloadCachePin.IsValid() };
 		Error->PushErrorContext(MoveTemp(RequestContext), UE::UnifiedError::EDetailFilter::Default);
 
 		SetErrorOnParents(Error.GetValue());
