[WorldPartition] Fixed PIE client connecting to a cooked server with a partitioned map. Without the fix, the client would get disconnected because the server wouldn't find the client's levels in its streaming levels.

diff --git a/Engine/Source/Runtime/Engine/Classes/Engine/Level.h b/Engine/Source/Runtime/Engine/Classes/Engine/Level.h
--- a/Engine/Source/Runtime/Engine/Classes/Engine/Level.h
+++ b/Engine/Source/Runtime/Engine/Classes/Engine/Level.h
@@ -1545,6 +1545,10 @@ private:
 	void OnMultipleLSAsPopupClicked();
 	void OnMultipleLSAsPopupDismissed();
 
+	/** Returns the optional cook package name (used by FUpdateLevelVisibilityLevelInfo) */
+	ENGINE_API FName GetCookPackageName() const;
+	friend struct FUpdateLevelVisibilityLevelInfo;
+
 #endif // WITH_EDITOR
 };
 
diff --git a/Engine/Source/Runtime/Engine/Classes/GameFramework/UpdateLevelVisibilityLevelInfo.h b/Engine/Source/Runtime/Engine/Classes/GameFramework/UpdateLevelVisibilityLevelInfo.h
--- a/Engine/Source/Runtime/Engine/Classes/GameFramework/UpdateLevelVisibilityLevelInfo.h
+++ b/Engine/Source/Runtime/Engine/Classes/GameFramework/UpdateLevelVisibilityLevelInfo.h
@@ -46,7 +46,7 @@ struct TStructOpsTypeTraits<FNetLevelVisibilityTransactionId> : public TStructOp
 	};
 };
 
-/** This structure is used to pass arguments to ServerUpdateLevelVisibilty() and ServerUpdateMultipleLevelsVisibility() server RPC functions */
+/** This structure is used to pass arguments to ServerUpdateLevelVisibility() and ServerUpdateMultipleLevelsVisibility() server RPC functions */
 USTRUCT()
 struct FUpdateLevelVisibilityLevelInfo
 {
@@ -54,6 +54,7 @@ struct FUpdateLevelVisibilityLevelInfo
 
 	FUpdateLevelVisibilityLevelInfo()
 		: PackageName(NAME_None)
+		, CookPackageName(NAME_None)
 		, FileName(NAME_None)
 		, bIsVisible(false)
 		, bTryMakeVisible(false)
@@ -72,6 +73,12 @@ struct FUpdateLevelVisibilityLevelInfo
 	UPROPERTY()
 	FName PackageName;
 
+	/** The name of the cook package for the level whose status changed (can be none). 
+	  * Can be used as a fallback for PackageName when a PIE client connects to a cooked server to update level visiblity.
+	  */
+	UPROPERTY()
+	FName CookPackageName;
+
 	/** The name / path of the asset file for the level whose status changed. */
 	UPROPERTY()
 	FName FileName;
diff --git a/Engine/Source/Runtime/Engine/Private/GameFramework/UpdateLevelVisibilityLevelInfo.cpp b/Engine/Source/Runtime/Engine/Private/GameFramework/UpdateLevelVisibilityLevelInfo.cpp
--- a/Engine/Source/Runtime/Engine/Private/GameFramework/UpdateLevelVisibilityLevelInfo.cpp
+++ b/Engine/Source/Runtime/Engine/Private/GameFramework/UpdateLevelVisibilityLevelInfo.cpp
@@ -18,7 +18,8 @@ namespace PlayerControllerCVars
 }
 
 FUpdateLevelVisibilityLevelInfo::FUpdateLevelVisibilityLevelInfo(const ULevel* const Level, const bool bInIsVisible, const bool bInTryMakeVisible)
-	: bIsVisible(bInIsVisible)
+	: CookPackageName(NAME_None)
+	, bIsVisible(bInIsVisible)
 	, bTryMakeVisible(bInTryMakeVisible)
 	, bSkipCloseOnError(false)
 {
@@ -27,6 +28,9 @@ FUpdateLevelVisibilityLevelInfo::FUpdateLevelVisibilityLevelInfo(const ULevel* c
 	check(!bTryMakeVisible || (bIsVisible != bTryMakeVisible));
 	const UPackage* const LevelPackage = Level->GetOutermost();
 	PackageName = LevelPackage->GetFName();
+#if WITH_EDITOR
+	CookPackageName = Level->GetCookPackageName();
+#endif
 
 	// When packages are duplicated for PIE, they may not have a FileName.
 	// For now, just revert to the old behavior.
@@ -39,10 +43,12 @@ bool FUpdateLevelVisibilityLevelInfo::NetSerialize(FArchive& Ar, UPackageMap* Pa
 	bool bArePackageAndFileTheSame = !!((PlayerControllerCVars::LevelVisibilityDontSerializeFileName) || (FileName == PackageName) || (FileName == NAME_None));
 	bool bLocalIsVisible = !!bIsVisible;
 	bool bLocalTryMakeVisible = !!bTryMakeVisible;
+	bool bArePackageAndCookPackageTheSame = (CookPackageName == PackageName) || (CookPackageName.IsNone());
 
 	Ar.SerializeBits(&bArePackageAndFileTheSame, 1);
 	Ar.SerializeBits(&bLocalIsVisible, 1);
 	Ar.SerializeBits(&bLocalTryMakeVisible, 1);
+	Ar.SerializeBits(&bArePackageAndCookPackageTheSame, 1);
 	Ar << PackageName;
 
 	if (!bArePackageAndFileTheSame)
@@ -54,6 +60,15 @@ bool FUpdateLevelVisibilityLevelInfo::NetSerialize(FArchive& Ar, UPackageMap* Pa
 		FileName = PackageName;
 	}
 
+	if (!bArePackageAndCookPackageTheSame)
+	{
+		Ar << CookPackageName;
+	}
+	else if (Ar.IsLoading())
+	{
+		CookPackageName = PackageName;
+	}
+
 	VisibilityRequestId.NetSerialize(Ar, PackageMap, bOutSuccess);
 
 	bIsVisible = bLocalIsVisible;
diff --git a/Engine/Source/Runtime/Engine/Private/Level.cpp b/Engine/Source/Runtime/Engine/Private/Level.cpp
--- a/Engine/Source/Runtime/Engine/Private/Level.cpp
+++ b/Engine/Source/Runtime/Engine/Private/Level.cpp
@@ -2792,6 +2792,12 @@ void ULevel::PostEditChangeProperty(FPropertyChangedEvent& PropertyChangedEvent)
 	}
 } 
 
+FName ULevel::GetCookPackageName() const
+{
+	UWorldPartitionRuntimeCell* Cell = WorldPartitionRuntimeCell.Get();
+	return Cell ? Cell->GetCookPackageName() : NAME_None;
+}
+
 #endif // WITH_EDITOR
 
 void ULevel::MarkLevelBoundsDirty()
diff --git a/Engine/Source/Runtime/Engine/Private/NetConnection.cpp b/Engine/Source/Runtime/Engine/Private/NetConnection.cpp
--- a/Engine/Source/Runtime/Engine/Private/NetConnection.cpp
+++ b/Engine/Source/Runtime/Engine/Private/NetConnection.cpp
@@ -159,30 +159,45 @@ namespace UE::Net::Connection::Private
 		// If its in our own streaming level list, its good
 
 		FValidateLevelVisibilityResult Result;
-		FString PackageNameStr = LevelVisibility.PackageName.ToString();
-		Result.Package = FindPackage(nullptr, *PackageNameStr);
-#if WITH_EDITOR
-		if (Result.Package && Result.Package->IsDynamicPIEPackagePending())
+
+		auto FindStreamingLevel = [World, &LevelVisibility, &Result](FName InPackageName)
 		{
-			// Package is a dynamic PIE package with pending external objects still loading, mimic same behavior as if it was AsyncLoad
-			Result.Package = nullptr;
-		}
+			FString PackageNameStr = InPackageName.ToString();
+			Result.Package = FindPackage(nullptr, *PackageNameStr);
+#if WITH_EDITOR
+			if (Result.Package && Result.Package->IsDynamicPIEPackagePending())
+			{
+				// Package is a dynamic PIE package with pending external objects still loading, mimic same behavior as if it was AsyncLoad
+				Result.Package = nullptr;
+			}
 #endif
-		Result.Linker = FLinkerLoad::FindExistingLinkerForPackage(Result.Package);
+			Result.Linker = FLinkerLoad::FindExistingLinkerForPackage(Result.Package);
 
-		Result.StreamingLevel = FLevelUtils::FindStreamingLevel(World, LevelVisibility.PackageName);
-		Result.bLevelExists = (Result.Linker != nullptr) || (Result.StreamingLevel != nullptr);
+			Result.StreamingLevel = FLevelUtils::FindStreamingLevel(World, InPackageName);
+			Result.bLevelExists = (Result.Linker != nullptr) || (Result.StreamingLevel != nullptr);
 
-		if (!Result.bLevelExists)
-		{
-			Result.bLevelExists = FLevelUtils::IsValidStreamingLevel(World, *PackageNameStr);
-			
 			if (!Result.bLevelExists)
 			{
-				Result.bLevelExists = (!FPackageName::IsMemoryPackage(PackageNameStr) && FPackageName::DoesPackageExist(LevelVisibility.FileName.ToString()));
+				Result.bLevelExists = FLevelUtils::IsValidStreamingLevel(World, *PackageNameStr);
+
+				if (!Result.bLevelExists)
+				{
+					Result.bLevelExists = (!FPackageName::IsMemoryPackage(PackageNameStr) && FPackageName::DoesPackageExist(LevelVisibility.FileName.ToString()));
+				}
 			}
+		};
+
+		FindStreamingLevel(LevelVisibility.PackageName);
+
+		// If the level is not found, try resolving it using the optional LevelVisibility.CookPackageName.
+		// This is useful when the in-memory PackageName differs from the cooked package name.
+		// A common example is when a PIE client sends PIE memory-generated packages to a cooked server.
+		if (!Result.bLevelExists && (!LevelVisibility.CookPackageName.IsNone() && LevelVisibility.CookPackageName != LevelVisibility.PackageName))
+		{
+			FindStreamingLevel(LevelVisibility.CookPackageName);
 		}
 
+
 		return Result;
 	}
 
diff --git a/Engine/Source/Runtime/Engine/Private/WorldPartition/WorldPartitionRuntimeLevelStreamingCell.cpp b/Engine/Source/Runtime/Engine/Private/WorldPartition/WorldPartitionRuntimeLevelStreamingCell.cpp
--- a/Engine/Source/Runtime/Engine/Private/WorldPartition/WorldPartitionRuntimeLevelStreamingCell.cpp
+++ b/Engine/Source/Runtime/Engine/Private/WorldPartition/WorldPartitionRuntimeLevelStreamingCell.cpp
@@ -62,6 +62,16 @@ bool UWorldPartitionRuntimeLevelStreamingCell::HasActors() const
 }
 
 #if WITH_EDITOR
+FName UWorldPartitionRuntimeLevelStreamingCell::GetCookPackageName() const
+{
+	if (UWorld* World = GetOwningWorld())
+	{
+		const FString CookPackageName = FString::Printf(TEXT("%s/%s/%s"), *UWorld::RemovePIEPrefix(World->GetPackage()->GetName()), FPackageName::GetGeneratedPackageSubPath(), *GetName());
+		return FName(CookPackageName);
+	}
+	return NAME_None;
+}
+
 TSet<FName> UWorldPartitionRuntimeLevelStreamingCell::GetActorPackageNames() const
 {
 	TSet<FName> Actors;
diff --git a/Engine/Source/Runtime/Engine/Public/WorldPartition/WorldPartitionRuntimeCell.h b/Engine/Source/Runtime/Engine/Public/WorldPartition/WorldPartitionRuntimeCell.h
--- a/Engine/Source/Runtime/Engine/Public/WorldPartition/WorldPartitionRuntimeCell.h
+++ b/Engine/Source/Runtime/Engine/Public/WorldPartition/WorldPartitionRuntimeCell.h
@@ -309,6 +309,7 @@ class UWorldPartitionRuntimeCell : public UObject, public IWorldPartitionCell, p
 	ENGINE_API void SetDataLayers(const TArray<const UDataLayerInstance*>& InDataLayerInstances);
 	void SetContentBundleUID(const FGuid& InContentBundleID) { ContentBundleID = InContentBundleID; }
 	void SetLevelPackageName(const FName& InLevelPackageName) { LevelPackageName = InLevelPackageName; }
+	virtual FName GetCookPackageName() const PURE_VIRTUAL(UWorldPartitionRuntimeCell::GetCookPackageName, return NAME_None;);
 	
 	//~Begin IWorldPartitionCell Interface
 	virtual TSet<FName> GetActorPackageNames() const override { return TSet<FName>(); }
diff --git a/Engine/Source/Runtime/Engine/Public/WorldPartition/WorldPartitionRuntimeLevelStreamingCell.h b/Engine/Source/Runtime/Engine/Public/WorldPartition/WorldPartitionRuntimeLevelStreamingCell.h
--- a/Engine/Source/Runtime/Engine/Public/WorldPartition/WorldPartitionRuntimeLevelStreamingCell.h
+++ b/Engine/Source/Runtime/Engine/Public/WorldPartition/WorldPartitionRuntimeLevelStreamingCell.h
@@ -52,6 +52,7 @@ class UWorldPartitionRuntimeLevelStreamingCell : public UWorldPartitionRuntimeCe
 	ENGINE_API virtual int32 GetActorCount() const override;
 	ENGINE_API virtual void DumpStateLog(FHierarchicalLogArchive& Ar) const override;
 	ENGINE_API virtual bool PrepareCellForCook(const IWorldPartitionCookPackageContext& InCookContext, UPackage* InGeneratedPackage) override;
+	ENGINE_API virtual FName GetCookPackageName() const override;
 	//~End UWorldPartitionRuntimeCell interface
 
 	virtual void OnCookEvent(UE::Cook::ECookEvent CookEvent, UE::Cook::FCookEventContext& CookContext) override;
