Collision Filter Refactor Part 3.12: Deprecating a bunch of old stuff, mostly old SQ related classes. This simplifies a later change a little bit.

diff --git a/Engine/Source/Runtime/Engine/Private/Collision/SceneQueryLowLevel.cpp b/Engine/Source/Runtime/Engine/Private/Collision/SceneQueryLowLevel.cpp
--- a/Engine/Source/Runtime/Engine/Private/Collision/SceneQueryLowLevel.cpp
+++ b/Engine/Source/Runtime/Engine/Private/Collision/SceneQueryLowLevel.cpp
@@ -10,61 +10,8 @@
 #include "Physics/Experimental/PhysScene_Chaos.h"
 #include "PhysicsEngine/ExternalSpatialAccelerationPayload.h"
 
-int32 ForceStandardSQ = 0;
-FAutoConsoleVariableRef CVarForceStandardSQ(TEXT("p.ForceStandardSQ"), ForceStandardSQ, TEXT("If enabled, we force the standard scene query even if custom SQ structure is enabled"));
-
-
-#if !UE_BUILD_SHIPPING
-int32 SerializeSQs = 0;
-int32 SerializeSQSamples = 100;
-int32 SerializeBadSQs = 0;
-int32 ReplaySQs = 0;
-int32 EnableRaycastSQCapture = 1;
-int32 EnableOverlapSQCapture = 1;
-int32 EnableSweepSQCapture = 1;
-
-FAutoConsoleVariableRef CVarSerializeSQs(TEXT("p.SerializeSQs"), SerializeSQs, TEXT("If enabled, we create a sq capture per sq that takes more than provided value in microseconds. This can be very expensive as the entire scene is saved out"));
-FAutoConsoleVariableRef CVarSerializeSQSamples(TEXT("p.SerializeSQSampleCount"), SerializeSQSamples, TEXT("If Query exceeds duration threshold, we will re-measure SQ this many times before serializing. Larger values cause hitching."));
-FAutoConsoleVariableRef CVarReplaySweeps(TEXT("p.ReplaySQs"), ReplaySQs, TEXT("If enabled, we rerun the sq against chaos"));
-FAutoConsoleVariableRef CVarSerializeBadSweeps(TEXT("p.SerializeBadSQs"), SerializeBadSQs, TEXT("If enabled, we create a sq capture whenever chaos and physx diverge"));
-FAutoConsoleVariableRef CVarSerializeSQsRaycastEnabled(TEXT("p.SerializeSQsRaycastEnabled"), EnableRaycastSQCapture, TEXT("If disabled, p.SerializeSQs will not consider raycasts"));
-FAutoConsoleVariableRef CVarSerializeSQsOverlapEnabled(TEXT("p.SerializeSQsOverlapEnabled"), EnableOverlapSQCapture, TEXT("If disabled, p.SerializeSQs will not consider overlaps"));
-FAutoConsoleVariableRef CVarSerializeSQsSweepEnabled(TEXT("p.SerializeSQsSweepEnabled"), EnableSweepSQCapture, TEXT("If disabled, p.SerializeSQs will not consider sweeps"));
-#else
-constexpr int32 SerializeSQs = 0;
-constexpr int32 ReplaySQs = 0;
-constexpr int32 SerializeSQSamples = 0;
-constexpr int32 EnableRaycastSQCapture = 0;
-constexpr int32 EnableOverlapSQCapture = 0;
-constexpr int32 EnableSweepSQCapture = 0;
-#endif
-
 namespace
 {
-	void FinalizeCapture(FPhysTestSerializer& Serializer)
-	{
-#if !UE_BUILD_SHIPPING
-		if (SerializeSQs)
-		{
-			Serializer.Serialize(TEXT("SQCapture"));
-		}
-#if 0
-		if (ReplaySQs)
-		{
-			const bool bReplaySuccess = SQComparisonHelper(Serializer);
-			if (!bReplaySuccess)
-			{
-				UE_LOG(LogPhysicsCore, Warning, TEXT("Chaos SQ does not match physx"));
-				if (SerializeBadSQs && !SerializeSQs)
-				{
-					Serializer.Serialize(TEXT("BadSQCapture"));
-				}
-			}
-		}
-#endif
-#endif
-	}
-
 	/**
 	 * This is meant as a replacement for FChaosSQAccelerator. There doesn't seem to be a need to expose this publicly since it'll just be used internally to allow
 	 * us to use more than one type of acceleration structure payload.
@@ -111,117 +58,6 @@ namespace
 		const Chaos::ISpatialAcceleration<TPayload, Chaos::FReal, 3>& SpatialAcceleration;
 	};
 
-	template<typename TPayload>
-	void SweepSQCaptureHelper(float QueryDurationSeconds, const FGenericChaosSQAccelerator<TPayload>& SQAccelerator, const FPhysScene& Scene, const FPhysicsGeometry& QueryGeom, const FTransform& StartTM, const FVector& Dir, float DeltaMag, const FPhysicsHitCallback<FHitSweep>& HitBuffer, EHitFlags OutputFlags, FQueryFlags QueryFlags, const FCollisionFilterData& Filter, const ChaosInterface::FQueryFilterData& QueryFilterData, ICollisionQueryFilterCallbackBase* QueryCallback, const ChaosInterface::FQueryDebugParams& DebugParams)
-	{
-#if !UE_BUILD_SHIPPING
-		float QueryDurationMicro = QueryDurationSeconds * 1000.0 * 1000.0;
-		if (((SerializeSQs && QueryDurationMicro > SerializeSQs)) && IsInGameThread())
-		{
-			// Measure average time of query over multiple samples to reduce fluke from context switches or that kind of thing.
-			uint32 Cycles = 0.0;
-			const uint32 SampleCount = SerializeSQSamples;
-			for (uint32 Samples = 0; Samples < SampleCount; ++Samples)
-			{
-				// Reset output to not skew times with large buffer
-				FPhysicsHitCallback<FHitSweep> ScratchHitBuffer = FPhysicsHitCallback<FHitSweep>(HitBuffer.WantsSingleResult());
-
-				uint32 StartTime = FPlatformTime::Cycles();
-				SQAccelerator.Sweep(QueryGeom, StartTM, Dir, DeltaMag, ScratchHitBuffer, OutputFlags, QueryFilterData, *QueryCallback, DebugParams);
-				Cycles += FPlatformTime::Cycles() - StartTime;
-			}
-
-			float Milliseconds = FPlatformTime::ToMilliseconds(Cycles);
-			float AvgMicroseconds = (Milliseconds * 1000) / SampleCount;
-
-			if (AvgMicroseconds > SerializeSQs)
-			{
-				FPhysTestSerializer Serializer;
-				Serializer.SetPhysicsData(*Scene.GetSolver()->GetEvolution());
-				FSQCapture& SweepCapture = Serializer.CaptureSQ();
-				SweepCapture.StartCaptureChaosSweep(*Scene.GetSolver()->GetEvolution(), QueryGeom, StartTM, Dir, DeltaMag, OutputFlags, QueryFilterData, Filter, *QueryCallback);
-				SweepCapture.EndCaptureChaosSweep(HitBuffer);
-
-				FinalizeCapture(Serializer);
-			}
-		}
-#endif
-	}
-
-	template<typename TPayload>
-	void RaycastSQCaptureHelper(float QueryDurationSeconds, const FGenericChaosSQAccelerator<TPayload>& SQAccelerator, const FPhysScene& Scene, const FVector& Start, const FVector& Dir, float DeltaMag, const FPhysicsHitCallback<FHitRaycast>& HitBuffer, EHitFlags OutputFlags, FQueryFlags QueryFlags, const FCollisionFilterData& Filter, const ChaosInterface::FQueryFilterData& QueryFilterData, ICollisionQueryFilterCallbackBase* QueryCallback, const ChaosInterface::FQueryDebugParams& DebugParams)
-	{
-#if !UE_BUILD_SHIPPING
-		float QueryDurationMicro = QueryDurationSeconds * 1000.0 * 1000.0;
-		if (((!!SerializeSQs && QueryDurationMicro > SerializeSQs)) && IsInGameThread())
-		{
-			// Measure average time of query over multiple samples to reduce fluke from context switches or that kind of thing.
-			uint32 Cycles = 0.0;
-			const uint32 SampleCount = SerializeSQSamples;
-			for (uint32 Samples = 0; Samples < SampleCount; ++Samples)
-			{
-				// Reset output to not skew times with large buffer
-				FPhysicsHitCallback<FHitRaycast> ScratchHitBuffer = FPhysicsHitCallback<FHitRaycast>(HitBuffer.WantsSingleResult());
-
-				uint32 StartTime = FPlatformTime::Cycles();
-				SQAccelerator.Raycast(Start, Dir, DeltaMag, ScratchHitBuffer, OutputFlags, QueryFilterData, *QueryCallback, DebugParams);
-				Cycles += FPlatformTime::Cycles() - StartTime;
-			}
-
-			float Milliseconds = FPlatformTime::ToMilliseconds(Cycles);
-			float AvgMicroseconds = (Milliseconds * 1000) / SampleCount;
-
-			if (AvgMicroseconds > SerializeSQs)
-			{
-				FPhysTestSerializer Serializer;
-				Serializer.SetPhysicsData(*Scene.GetSolver()->GetEvolution());
-				FSQCapture& RaycastCapture = Serializer.CaptureSQ();
-				RaycastCapture.StartCaptureChaosRaycast(*Scene.GetSolver()->GetEvolution(), Start, Dir, DeltaMag, OutputFlags, QueryFilterData, Filter, *QueryCallback);
-				RaycastCapture.EndCaptureChaosRaycast(HitBuffer);
-
-				FinalizeCapture(Serializer);
-			}
-		}
-#endif
-	}
-
-	template<typename TPayload>
-	void OverlapSQCaptureHelper(float QueryDurationSeconds, const FGenericChaosSQAccelerator<TPayload>& SQAccelerator, const FPhysScene& Scene, const FPhysicsGeometry& QueryGeom, const FTransform& GeomPose, const FPhysicsHitCallback<FHitOverlap>& HitBuffer, FQueryFlags QueryFlags, const FCollisionFilterData& Filter, const ChaosInterface::FQueryFilterData& QueryFilterData, ICollisionQueryFilterCallbackBase* QueryCallback, const ChaosInterface::FQueryDebugParams& DebugParams)
-	{
-#if !UE_BUILD_SHIPPING
-		float QueryDurationMicro = QueryDurationSeconds * 1000.0 * 1000.0;
-		if (((!!SerializeSQs && QueryDurationMicro > SerializeSQs)) && IsInGameThread())
-		{
-			// Measure average time of query over multiple samples to reduce fluke from context switches or that kind of thing.
-			uint32 Cycles = 0.0;
-			const uint32 SampleCount = SerializeSQSamples;
-			for (uint32 Samples = 0; Samples < SampleCount; ++Samples)
-			{
-				// Reset output to not skew times with large buffer
-				FPhysicsHitCallback<FHitOverlap> ScratchHitBuffer = FPhysicsHitCallback<FHitOverlap>(HitBuffer.WantsSingleResult());
-
-				uint32 StartTime = FPlatformTime::Cycles();
-				SQAccelerator.Overlap(QueryGeom, GeomPose, ScratchHitBuffer, QueryFilterData, *QueryCallback);
-				Cycles += FPlatformTime::Cycles() - StartTime;
-			}
-
-			float Milliseconds = FPlatformTime::ToMilliseconds(Cycles);
-			float AvgMicroseconds = (Milliseconds * 1000) / SampleCount;
-
-			if (AvgMicroseconds > SerializeSQs)
-			{
-				FPhysTestSerializer Serializer;
-				Serializer.SetPhysicsData(*Scene.GetSolver()->GetEvolution());
-				FSQCapture& OverlapCapture = Serializer.CaptureSQ();
-				OverlapCapture.StartCaptureChaosOverlap(*Scene.GetSolver()->GetEvolution(), QueryGeom, GeomPose, QueryFilterData, Filter, *QueryCallback);
-				OverlapCapture.EndCaptureChaosOverlap(HitBuffer);
-
-				FinalizeCapture(Serializer);
-			}
-		}
-#endif
-	}
-
 	template<typename TAccelContainer, bool bGTData>
 	struct FAccelerationContainerTraits
 	{
@@ -266,14 +102,6 @@ void LowLevelRaycastImpl(const TAccelContainer& Container, const FVector& Start,
 			FScopedDurationTimer Timer(Time);
 			SQAccelerator.Raycast(Start, Dir, DeltaMag, HitBuffer, OutputFlags, QueryFilterData, *QueryCallback, DebugParams);
 		}
-
-		if constexpr (bGTData && TTraits::IsPhysScene())
-		{
-			if (!!SerializeSQs && !!EnableRaycastSQCapture)
-			{
-				RaycastSQCaptureHelper(Time, SQAccelerator, Container, Start, Dir, DeltaMag, HitBuffer, OutputFlags, QueryFlags, Filter, QueryFilterData, QueryCallback, DebugParams);
-			}
-		}
 	}
 }
 
@@ -300,14 +128,6 @@ void LowLevelSweepImpl(const TAccelContainer& Container, const FPhysicsGeometry&
 				FScopedDurationTimer Timer(Time);
 				SQAccelerator.Sweep(QueryGeom, StartTM, Dir, DeltaMag, HitBuffer, OutputFlags, QueryFilterData, *QueryCallback, DebugParams);
 			}
-
-			if constexpr (bGTData && FAccelerationContainerTraits<TAccelContainer, bGTData>::IsPhysScene())
-			{
-				if(!!SerializeSQs && !!EnableSweepSQCapture)
-				{
-					SweepSQCaptureHelper(Time, SQAccelerator, Container, QueryGeom, StartTM, Dir, DeltaMag, HitBuffer, OutputFlags, QueryFlags, Filter, QueryFilterData, QueryCallback, DebugParams);
-				}
-			}
 		}
 	}
 }
@@ -335,14 +155,6 @@ void LowLevelOverlapImpl(const TAccelContainer& Container, const FPhysicsGeometr
 			FScopedDurationTimer Timer(Time);
 			SQAccelerator.Overlap(QueryGeom, GeomPose, HitBuffer, QueryFilterData, *QueryCallback, DebugParams);
 		}
-
-		if constexpr(bGTData && FAccelerationContainerTraits<TAccelContainer, bGTData>::IsPhysScene())
-		{
-			if(!!SerializeSQs && !!EnableOverlapSQCapture)
-			{
-				OverlapSQCaptureHelper(Time, SQAccelerator, Container, QueryGeom, GeomPose, HitBuffer, QueryFlags, Filter, QueryFilterData, QueryCallback, DebugParams);
-			}
-		}
 	}
 }
 
diff --git a/Engine/Source/Runtime/Experimental/Chaos/Public/RewindData.h b/Engine/Source/Runtime/Experimental/Chaos/Public/RewindData.h
--- a/Engine/Source/Runtime/Experimental/Chaos/Public/RewindData.h
+++ b/Engine/Source/Runtime/Experimental/Chaos/Public/RewindData.h
@@ -1010,6 +1010,13 @@ struct FPerShapeDataStateBase
 
 	//helper functions for shape API
 	template <typename TParticle>
+	static Chaos::Filter::FCombinedShapeFilterData GetShapeCombinedFilterData(const FPerShapeDataStateBase* State, const TParticle& Particle, int32 ShapeIdx)
+	{
+		return State && State->CollisionData.IsSet() ? State->CollisionData.Read().GetCombinedShapeFilterData() : Particle.ShapesArray()[ShapeIdx]->GetCombinedShapeFilterData();
+	}
+
+	template <typename TParticle>
+	UE_DEPRECATED(5.8, "This is unused and will be removed. Use GetCombinedShapeFilterData instead")
 	static const FCollisionFilterData& GetQueryData(const FPerShapeDataStateBase* State, const TParticle& Particle, int32 ShapeIdx)
 	{
 		PRAGMA_DISABLE_DEPRECATION_WARNINGS
@@ -1028,7 +1035,17 @@ public:
 	{
 	}
 
+	const Chaos::Filter::FCombinedShapeFilterData GetCombinedShapeFilterData() const
+	{
+		PRAGMA_DISABLE_DEPRECATION_WARNINGS
+		return FPerShapeDataStateBase::GetShapeCombinedFilterData(State, Particle, ShapeIdx);
+		PRAGMA_ENABLE_DEPRECATION_WARNINGS
+	}
+
+	PRAGMA_DISABLE_DEPRECATION_WARNINGS
+	UE_DEPRECATED(5.8, "This is unused and will be removed. Use GetCombinedShapeFilterData instead")
 	const FCollisionFilterData& GetQueryData() const { return FPerShapeDataStateBase::GetQueryData(State, Particle, ShapeIdx); }
+	PRAGMA_ENABLE_DEPRECATION_WARNINGS
 private:
 	const FPerShapeDataStateBase* State;
 	const FGeometryParticleHandle& Particle;
diff --git a/Engine/Source/Runtime/PhysicsCore/Private/PhysTestSerializer.cpp b/Engine/Source/Runtime/PhysicsCore/Private/PhysTestSerializer.cpp
--- a/Engine/Source/Runtime/PhysicsCore/Private/PhysTestSerializer.cpp
+++ b/Engine/Source/Runtime/PhysicsCore/Private/PhysTestSerializer.cpp
@@ -91,6 +91,7 @@ void FPhysTestSerializer::Serialize(Chaos::FChaosArchive& Ar)
 		Ar.SetContext(MoveTemp(ChaosContext));	//make sure any context we created during load is used for sqcapture
 	}
 
+	PRAGMA_DISABLE_DEPRECATION_WARNINGS
 	bool bHasSQCapture = !!SQCapture;
 	{
 		Chaos::FChaosArchiveScopedMemory ScopedMemory(Ar, TestSerializerName, false);
@@ -104,6 +105,7 @@ void FPhysTestSerializer::Serialize(Chaos::FChaosArchive& Ar)
 		}
 		SQCapture->Serialize(Ar);
 	}
+	PRAGMA_ENABLE_DEPRECATION_WARNINGS
 	ChaosContext = Ar.StealContext();
 }
 
diff --git a/Engine/Source/Runtime/PhysicsCore/Private/SQAccelerator.cpp b/Engine/Source/Runtime/PhysicsCore/Private/SQAccelerator.cpp
--- a/Engine/Source/Runtime/PhysicsCore/Private/SQAccelerator.cpp
+++ b/Engine/Source/Runtime/PhysicsCore/Private/SQAccelerator.cpp
@@ -15,6 +15,7 @@
 #include "Chaos/GeometryQueries.h"
 #include "Chaos/DebugDrawQueue.h"
 
+PRAGMA_DISABLE_DEPRECATION_WARNINGS
 void FSQAcceleratorUnion::Raycast(const FVector& Start, const FVector& Dir, const float DeltaMagnitude, FPhysicsHitCallback<FHitRaycast>& HitBuffer, EHitFlags OutputFlags, const ChaosInterface::FQueryFilterData& QueryFilterData, ICollisionQueryFilterCallbackBase& QueryCallback) const
 {
 	for(const ISQAccelerator* Accelerator : Accelerators)
@@ -48,6 +49,7 @@ void FSQAcceleratorUnion::RemoveSQAccelerator(ISQAccelerator* AcceleratorToRemov
 {
 	Accelerators.RemoveSingleSwap(AcceleratorToRemove);	//todo(ocohen): probably want to order these in some optimal way
 }
+PRAGMA_ENABLE_DEPRECATION_WARNINGS
 
 FChaosSQAccelerator::FChaosSQAccelerator(const Chaos::ISpatialAcceleration<Chaos::FAccelerationStructureHandle, Chaos::FReal, 3>& InSpatialAcceleration)
 	: SpatialAcceleration(InSpatialAcceleration)
diff --git a/Engine/Source/Runtime/PhysicsCore/Public/PhysTestSerializer.h b/Engine/Source/Runtime/PhysicsCore/Public/PhysTestSerializer.h
--- a/Engine/Source/Runtime/PhysicsCore/Public/PhysTestSerializer.h
+++ b/Engine/Source/Runtime/PhysicsCore/Public/PhysTestSerializer.h
@@ -53,6 +53,8 @@ public:
 		return ChaosContext.Get();
 	}
 
+	PRAGMA_DISABLE_DEPRECATION_WARNINGS
+	UE_DEPRECATED(5.8, "FSQCapture is unused and will be removed")
 	FSQCapture& CaptureSQ()
 	{
 		ensure(!SQCapture);	//we don't have support for multi sweeps yet since the scene could change
@@ -60,6 +62,7 @@ public:
 		return *SQCapture;
 	}
 
+	UE_DEPRECATED(5.8, "FSQCapture is unused and will be removed")
 	const FSQCapture* GetSQCapture()
 	{
 		if (SQCapture)
@@ -72,6 +75,7 @@ public:
 		}
 		return SQCapture.Get();
 	}
+	PRAGMA_ENABLE_DEPRECATION_WARNINGS
 
 	Chaos::FPBDRigidsEvolution* GetChaosData()
 	{
@@ -98,7 +102,10 @@ private:
 	bool bDiskDataIsChaos;
 	bool bChaosDataReady;
 
+	PRAGMA_DISABLE_DEPRECATION_WARNINGS
+	UE_DEPRECATED(5.8, "FSQCapture is unused and will be removed")
 	TUniquePtr<FSQCapture> SQCapture;
+	PRAGMA_ENABLE_DEPRECATION_WARNINGS
 
 	TUniquePtr<Chaos::FPBDRigidsEvolution> ChaosEvolution;
 	Chaos::FParticleUniqueIndicesMultithreaded UniqueIndices;
diff --git a/Engine/Source/Runtime/PhysicsCore/Public/SQAccelerator.h b/Engine/Source/Runtime/PhysicsCore/Public/SQAccelerator.h
--- a/Engine/Source/Runtime/PhysicsCore/Public/SQAccelerator.h
+++ b/Engine/Source/Runtime/PhysicsCore/Public/SQAccelerator.h
@@ -58,11 +58,12 @@ private:
 
 };
 
+PRAGMA_DISABLE_DEPRECATION_WARNINGS
 // An interface to the scene query accelerator that allows us to run queries against either PhysX or Chaos
 // This was used in the 2019 GDC demos and is now broken. To make it work again, we would need to implement
 // the FChaosSQAcceleratorAdapter below to use its internal SQ accelerator and convert the inputs and outputs
 // from/to PhysX types.
-class ISQAccelerator
+class UE_DEPRECATED(5.8, "This is unused and will be removed") ISQAccelerator
 {
 public:
 	virtual ~ISQAccelerator() {};
@@ -71,7 +72,7 @@ public:
 	virtual void Overlap(const FPhysicsGeometry& QueryGeom, const FTransform& GeomPose, FPhysicsHitCallback<FHitOverlap>& HitBuffer, const ChaosInterface::FQueryFilterData& QueryFilterData, ICollisionQueryFilterCallbackBase& QueryCallback) const = 0;
 };
 
-class FSQAcceleratorUnion : public ISQAccelerator
+class UE_DEPRECATED(5.8, "This is unused and will be removed") FSQAcceleratorUnion : public ISQAccelerator
 {
 public:
 
@@ -85,3 +86,4 @@ public:
 private:
 	TArray<ISQAccelerator*> Accelerators;
 };
+PRAGMA_ENABLE_DEPRECATION_WARNINGS
diff --git a/Engine/Source/Runtime/PhysicsCore/Public/SQVerifier.h b/Engine/Source/Runtime/PhysicsCore/Public/SQVerifier.h
--- a/Engine/Source/Runtime/PhysicsCore/Public/SQVerifier.h
+++ b/Engine/Source/Runtime/PhysicsCore/Public/SQVerifier.h
@@ -2,6 +2,8 @@
 
 #pragma once
 
+UE_DEPRECATED_HEADER(5.8, "This file is unused and will be removed")  
+
 #include "PhysicsCore.h"
 #include "ProfilingDebugging/ScopedTimers.h"
 
@@ -14,7 +16,9 @@
 #define SQ_REPLAY_TEST(cond) bEnsureOnMismatch ? ensure(cond) : (cond)
 #endif
 
+PRAGMA_DISABLE_DEPRECATION_WARNINGS
 template <bool bHasPhysX = false>
+UE_DEPRECATED(5.8, "This is unused and will be removed")
 void SQPerfComparisonHelper(const FString& TestName, FPhysTestSerializer& Serializer, bool bEnsureOnMismatch = false)
 {
 	using namespace Chaos;
@@ -91,6 +95,7 @@ void SQPerfComparisonHelper(const FString& TestName, FPhysTestSerializer& Serial
 
 
 #if 0
+UE_DEPRECATED(5.8, "This is unused and will be removed")
 bool SQComparisonHelper(FPhysTestSerializer& Serializer, bool bEnsureOnMismatch = false)
 {
 	using namespace Chaos;
@@ -210,6 +215,7 @@ bool SQComparisonHelper(FPhysTestSerializer& Serializer, bool bEnsureOnMismatch
 
 #endif
 
+UE_DEPRECATED(5.8, "This is unused and will be removed")
 bool SQValidityHelper(FPhysTestSerializer& Serializer)
 {
 	using namespace Chaos;
@@ -275,5 +281,5 @@ bool SQValidityHelper(FPhysTestSerializer& Serializer)
 
 	return bTestPassed;
 }
-
+PRAGMA_ENABLE_DEPRECATION_WARNINGS
 
