Add ExtrudedTaperedCapsule option to Cloth Physics Assets. From the UI side, this just adds a Width parameter to the FKTaperedCapsuleElem. Internally implemented as a separate ImplicitObject type mostly because I didn't want to update all of the methods that are unused by cloth (e.g., anything needed to simulate one of these shapes).

This shape allows making an extruded TaperedCapsule where the X-axis can now have width, so instead of the base shape being a line segment with thickness, it is a rectangle. This is much cheaper than trying to build a similar shape from multiple objects.

Also fixed  FKTaperedCapsuleElem::ScaleElem so it works correctly (it didn't work before this change either).

diff --git a/Engine/Plugins/ChaosCloth/Source/ChaosCloth/Private/ChaosCloth/ChaosClothVisualization.cpp b/Engine/Plugins/ChaosCloth/Source/ChaosCloth/Private/ChaosCloth/ChaosClothVisualization.cpp
--- a/Engine/Plugins/ChaosCloth/Source/ChaosCloth/Private/ChaosCloth/ChaosClothVisualization.cpp
+++ b/Engine/Plugins/ChaosCloth/Source/ChaosCloth/Private/ChaosCloth/ChaosClothVisualization.cpp
@@ -10,6 +10,7 @@
 #include "Chaos/DebugDrawQueue.h"
 #include "Chaos/Capsule.h"
 #include "Chaos/Convex.h"
+#include "Chaos/ExtrudedTaperedCapsule.h"
 #include "Chaos/Levelset.h"
 #include "Chaos/Sphere.h"
 #include "Chaos/TaperedCapsule.h"
@@ -743,6 +744,130 @@ namespace Chaos
 		DrawTaperedCylinder(PDI, Position1, Position2, Radius1, Radius2, 12, Color);
 	}
 
+	static void DrawExtrudedTaperedCapsule(FPrimitiveDrawInterface* PDI, const ::Chaos::FExtrudedTaperedCapsule& TaperedCapsule, const FQuat& Rotation, const FVector& Position, const FReal Scale, const FLinearColor& Color)
+	{
+		const float ScaledHalfLength = TaperedCapsule.GetHalfLength() * Scale;
+		const float ScaledLength = 2.f * ScaledHalfLength;
+		float ScaledRadius0 = TaperedCapsule.GetRadius1() * Scale;
+		float ScaledRadius1 = TaperedCapsule.GetRadius2() * Scale;
+		const float ScaledHalfWidth = TaperedCapsule.GetHalfWidth() * Scale;
+
+		float SphereOffset0 = ScaledHalfLength;
+		float SphereOffset1 = -ScaledHalfLength;
+		if (FMath::IsNearlyZero(ScaledLength) || ScaledLength <= FMath::Abs(ScaledRadius0 - ScaledRadius1))
+		{
+			// Degenerate or one end encompasses the other - we need to draw a sphere, so map one end to the other
+			if (ScaledRadius0 > ScaledRadius1)
+			{
+				SphereOffset1 = SphereOffset0;
+				ScaledRadius1 = ScaledRadius0;
+			}
+			else
+			{
+				SphereOffset0 = SphereOffset1;
+				ScaledRadius0 = ScaledRadius1;
+			}
+		}
+
+		constexpr int32 NumNonExtrudedSides = 12;
+		constexpr int32 NumSides = NumNonExtrudedSides + 2;
+		constexpr int32 NumRings = 8;
+		// The first/last arc are on top of each other.
+		constexpr int32 NumVerts = (NumSides + 1) * (NumRings + 1);
+		TArray<FVector3f> Verts, ArcVerts;
+		Verts.SetNumUninitialized(NumVerts);
+		ArcVerts.SetNumUninitialized(NumRings + 1);
+
+		// Calc arc split point
+		const float OffsetZ = FMath::IsNearlyZero(ScaledLength) ? 0.f : ((ScaledRadius0 - ScaledRadius1) * ScaledRadius0) / ScaledLength;
+		const float SplitAngle = FMath::IsNearlyZero(ScaledRadius0) ? 90.f : 90.0f - FMath::RadiansToDegrees(FMath::Asin(OffsetZ / ScaledRadius0));
+		const int32 SplitPoint = (int32)((float)NumRings * ((180.0f - SplitAngle) / 180.0f));
+		for (int32 RingIdx = 0; RingIdx < NumRings + 1; RingIdx++)
+		{
+			// Note- unit sphere, so position always has mag of one. We can just use it for normal!		
+			FVector3f SpherePos;
+			SpherePos.X = 0.0f;
+			float ZOffset;
+
+			if (RingIdx <= SplitPoint)
+			{
+				float Angle = ((float)RingIdx / (NumRings - 1)) * UE_PI;
+				ZOffset = SphereOffset0;
+				SpherePos.Y = ScaledRadius0 * FMath::Sin(Angle);
+				SpherePos.Z = ScaledRadius0 * FMath::Cos(Angle);
+			}
+			else
+			{
+				float Angle = ((float)(RingIdx - 1) / (NumRings - 1)) * UE_PI;
+				ZOffset = SphereOffset1;
+				SpherePos.Y = ScaledRadius1 * FMath::Sin(Angle);
+				SpherePos.Z = ScaledRadius1 * FMath::Cos(Angle);
+			}
+
+			ArcVerts[RingIdx] = SpherePos + FVector3f(0.f, 0.f, ZOffset);
+		}
+
+		// Then rotate this arc NumSides+1 times.
+		for (int32 SideIdx = 0; SideIdx < NumNonExtrudedSides + 1; SideIdx++)
+		{
+			const FRotator3f ArcRotator(0, 360.f * ((float)SideIdx / NumNonExtrudedSides), 0);
+			const FRotationMatrix44f ArcRot(ArcRotator);
+
+			const FVector3f XOffset(SideIdx <= NumNonExtrudedSides / 2 ? -ScaledHalfWidth : ScaledHalfWidth, 0, 0);
+			const int32 SideIdxToVIxOffset = (SideIdx > NumNonExtrudedSides / 2) ? 1 : 0;
+			for (int32 VertIdx = 0; VertIdx < NumRings + 1; VertIdx++)
+			{
+				const int32 VIx = (NumRings + 1) * (SideIdx + SideIdxToVIxOffset) + VertIdx;
+				Verts[VIx] = ArcRot.TransformPosition(ArcVerts[VertIdx]) + XOffset;
+			}
+
+			if (SideIdx == NumNonExtrudedSides / 2)
+			{
+				// Duplicate vertex on second half
+				for (int32 VertIdx = 0; VertIdx < NumRings + 1; VertIdx++)
+				{
+					const int32 VIx = (NumRings + 1) * (SideIdx + 1) + VertIdx;
+					Verts[VIx] = ArcRot.TransformPosition(ArcVerts[VertIdx]) + FVector3f(ScaledHalfWidth, 0.f, 0.f);
+				}
+			}
+			else if (SideIdx == NumNonExtrudedSides)
+			{
+				// Duplicate vertex on second half
+				for (int32 VertIdx = 0; VertIdx < NumRings + 1; VertIdx++)
+				{
+					const int32 VIx = (NumRings + 1) * (SideIdx + 2) + VertIdx;
+					Verts[VIx] = ArcRot.TransformPosition(ArcVerts[VertIdx]) - FVector3f(ScaledHalfWidth, 0.f, 0.f);
+				}
+			}
+		}
+		FMatrix CapsuleTransform = FRotationMatrix::MakeFromZX(FVector(TaperedCapsule.GetMajorAxis()), FVector(TaperedCapsule.GetMinorAxis()));
+		CapsuleTransform.SetOrigin(FVector(TaperedCapsule.GetOrigin()));
+		FMatrix BoneTransform = Rotation.ToMatrix();
+		BoneTransform.SetOrigin(Position);
+		const FMatrix FullTransform = CapsuleTransform * BoneTransform;
+		// Add all of the lines
+		for (int32 SideIdx = 0; SideIdx < NumSides; SideIdx++)
+		{
+			const int32 Side0Start = (SideIdx + 0) * (NumRings + 1);
+			const int32 Side1Start = (SideIdx + 1) * (NumRings + 1);
+
+			for (int32 RingIdx = 0; RingIdx < NumRings; RingIdx++)
+			{
+				const FVector S0R0 = FullTransform.TransformPosition(FVector(Verts[Side0Start + RingIdx]));
+				const FVector S0R1 = FullTransform.TransformPosition(FVector(Verts[Side0Start + RingIdx + 1]));
+				const FVector S1R0 = FullTransform.TransformPosition(FVector(Verts[Side1Start + RingIdx]));
+				DrawLine(PDI, S0R0, S1R0, Color);
+				DrawLine(PDI, S0R0, S0R1, Color);
+
+				if (RingIdx == NumRings - 1)
+				{
+					const FVector S1R1 = FullTransform.TransformPosition(FVector(Verts[Side1Start + RingIdx + 1]));
+					DrawLine(PDI, S0R1, S1R1, Color);
+				}
+			}
+		}
+	}
+
 	static void DrawConvex(FPrimitiveDrawInterface* PDI, const ::Chaos::FConvex& Convex, const FQuat& Rotation, const FVector& Position, const FReal Scale, const FLinearColor& Color)
 	{
 		const TArray<FConvex::FPlaneType>& Planes = Convex.GetFaces();
@@ -1742,6 +1867,13 @@ namespace Chaos
 							DrawSkinnedTriangleMesh(PDI, MaterialRenderProxy, Solver, Color, Mesh);
 
 						}
+						break;
+						case ImplicitObjectType::ExtrudedTaperedCapsule:
+						{
+							const FExtrudedTaperedCapsule& Capsule = Object->GetObjectChecked<FExtrudedTaperedCapsule>();
+							DrawExtrudedTaperedCapsule(PDI, Capsule, Rotation, Position, Solver->GetLocalSpaceScale(), Color);
+						}
+						break;
 						default:
 							DrawCoordinateSystem(PDI, Rotation, Position);  // Draw everything else as a coordinate for now
 							break;
diff --git a/Engine/Plugins/ChaosCloth/Source/ChaosCloth/Private/ChaosCloth/ChaosClothingSimulationCollider.cpp b/Engine/Plugins/ChaosCloth/Source/ChaosCloth/Private/ChaosCloth/ChaosClothingSimulationCollider.cpp
--- a/Engine/Plugins/ChaosCloth/Source/ChaosCloth/Private/ChaosCloth/ChaosClothingSimulationCollider.cpp
+++ b/Engine/Plugins/ChaosCloth/Source/ChaosCloth/Private/ChaosCloth/ChaosClothingSimulationCollider.cpp
@@ -6,6 +6,7 @@
 #include "ChaosCloth/ChaosClothingSimulationMesh.h"
 #include "ChaosCloth/ChaosClothPrivate.h"
 #include "Chaos/Capsule.h"
+#include "Chaos/ExtrudedTaperedCapsule.h"
 #include "Chaos/Plane.h"
 #include "Chaos/Sphere.h"
 #include "Chaos/TaperedCylinder.h"
@@ -107,13 +108,14 @@ void FClothingSimulationCollider::FLODData::Add(
 
 	const int32 NumConvexes = ClothCollisionData.Convexes.Num();
 	const int32 NumBoxes = ClothCollisionData.Boxes.Num();
+	const int32 NumExtrudedCapsules = ClothCollisionData.ExtrudedTaperedCapsules.Num();
 	const int32 NumLevelSets = InClothCollisionDataExtended.LevelSets.Num();
 	const int32 NumSkinnedLevelSets = InClothCollisionDataExtended.SkinnedLevelSets.Num();
 	const int32 NumMLLevelSets = InClothCollisionDataExtended.MLLevelSets.Num();
 	const int32 NumSkinnedTriangleMeshes = InClothCollisionDataExtended.SkinnedTriangleMeshes.Num();
 
 	const bool bSupportsComplexColliders = !Solver->IsLegacySolver();
-	const int32 NumSimpleColliders = NumSpheres + NumCapsules + NumConvexes + NumBoxes + NumLevelSets;
+	const int32 NumSimpleColliders = NumSpheres + NumCapsules + NumConvexes + NumBoxes + NumExtrudedCapsules + NumLevelSets;
 	const int32 NumComplexColliders = NumSkinnedLevelSets + NumMLLevelSets + NumSkinnedTriangleMeshes;
 
 	NumGeometries = NumSimpleColliders + (bSupportsComplexColliders ? NumComplexColliders : 0);
@@ -328,9 +330,27 @@ void FClothingSimulationCollider::FLODData::Add(
 			Solver->SetCollisionGeometry(CollisionRangeId, Index, MakeImplicitObjectPtr<TBox<FReal, 3>>(-HalfExtents, HalfExtents));
 		}
 	}
+
+	// Extruded Tapered Capsules
+	const int32 ExtrudedCapsuleIndexOffset = BoxIndexOffset + NumBoxes;
+	if (NumExtrudedCapsules != 0)
+	{
+		for (int32 Index = ExtrudedCapsuleIndexOffset, ExtrudedCapsuleIndex = 0; ExtrudedCapsuleIndex < NumExtrudedCapsules; ++Index, ++ExtrudedCapsuleIndex)
+		{
+			const FClothCollisionPrim_ExtrudedTaperedCapsule ExtrudedCapsule = ClothCollisionData.ExtrudedTaperedCapsules[ExtrudedCapsuleIndex];
+			BaseTransforms[Index] = Softs::FSolverRigidTransform3::Identity;
+
+			BoneIndices[Index] = GetMappedBoneIndex(UsedBoneIndices, ExtrudedCapsule.BoneIndex);
+			UE_LOG(LogChaosCloth, VeryVerbose, TEXT("Found collision ExtrudedCapsule on bone index %d."), BoneIndices[Index]);
+
+			const FVec3f MajorAxis = FVec3f(ExtrudedCapsule.LocalRotation.RotateVector(FVector::UpVector));
+			const FVec3f MinorAxis = FVec3f(ExtrudedCapsule.LocalRotation.RotateVector(FVector::ForwardVector));
+			Solver->SetCollisionGeometry(CollisionRangeId, Index, MakeImplicitObjectPtr<FExtrudedTaperedCapsule>(FVec3f(ExtrudedCapsule.LocalPosition * InScale), MajorAxis, MinorAxis, ExtrudedCapsule.Length * InScale * .5f, ExtrudedCapsule.Width * InScale * .5f, ExtrudedCapsule.Radius1 * InScale, ExtrudedCapsule.Radius2 * InScale));
+		}
+	}
 	
 	// Level sets
-	const int32 LevelSetIndexOffset = BoxIndexOffset + NumBoxes;
+	const int32 LevelSetIndexOffset = ExtrudedCapsuleIndexOffset + NumExtrudedCapsules;
 	if (NumLevelSets != 0)
 	{
 		for (int32 Index = LevelSetIndexOffset, LevelSetIndex = 0; LevelSetIndex < NumLevelSets; ++Index, ++LevelSetIndex)
@@ -348,7 +368,7 @@ void FClothingSimulationCollider::FLODData::Add(
 		}
 	}
 
-	UE_LOG(LogChaosCloth, VeryVerbose, TEXT("Added collisions: %d spheres, %d capsules, %d convexes, %d boxes, %d level sets."), NumSpheres, NumCapsules, NumConvexes, NumBoxes, NumLevelSets);
+	UE_LOG(LogChaosCloth, VeryVerbose, TEXT("Added collisions: %d spheres, %d capsules, %d convexes, %d boxes, &d extruded tapered capsules, %d level sets."), NumSpheres, NumCapsules, NumConvexes, NumBoxes, NumExtrudedCapsules, NumLevelSets);
 
 	if (bSupportsComplexColliders)
 	{
@@ -392,7 +412,7 @@ void FClothingSimulationCollider::FLODData::Add(
 			BaseTransforms[Index] = Softs::FSolverRigidTransform3::Identity;
 
 			BoneIndices[Index] = GetMappedBoneIndex(UsedBoneIndices, InClothCollisionDataExtended.SkinnedTriangleMeshes[SkinnedTriMeshIndex].BoneIndex);
-			UE_LOG(LogChaosCloth, VeryVerbose, TEXT("Found collision ml level set on bone index %d."), BoneIndices[Index]);
+			UE_LOG(LogChaosCloth, VeryVerbose, TEXT("Found collision skinned triangle mesh on bone index %d."), BoneIndices[Index]);
 
 			FImplicitObjectPtr SkinnedTriMesh = InClothCollisionDataExtended.SkinnedTriangleMeshes[SkinnedTriMeshIndex].SkinnedTriangleMesh->CopyGeometry();
 			ComplexColliders->AddSkinnedTriangleMesh(Index, InClothCollisionDataExtended.SkinnedTriangleMeshes[SkinnedTriMeshIndex].MappedSkinnedBones, SkinnedTriMesh);
@@ -517,36 +537,6 @@ FClothCollisionData FClothingSimulationCollider::GetCollisionData(const FClothin
 	return ClothCollisionData;
 }
 
-void FClothingSimulationCollider::ExtractPhysicsAssetCollision(
-	const UPhysicsAsset* PhysicsAsset,
-	const FReferenceSkeleton* ReferenceSkeleton,
-	FClothCollisionData& ClothCollisionData,
-	PRAGMA_DISABLE_DEPRECATION_WARNINGS
-	TArray<FLevelSetCollisionData>& LevelSetCollisions,
-	TArray<FSkinnedLevelSetCollisionData>& SkinnedLevelSetCollisions,
-	PRAGMA_ENABLE_DEPRECATION_WARNINGS
-	TArray<int32>& UsedBoneIndices,
-	bool bUseSphylsOnly,
-	bool bSkipMissingBones)
-{
-	FClothCollisionDataExtended ExtendedCollisionData;
-	TArray<int32> UsedSubBoneIndices;
-	ExtractPhysicsAssetCollision(PhysicsAsset, ReferenceSkeleton, ClothCollisionData, ExtendedCollisionData, UsedBoneIndices, UsedSubBoneIndices, bUseSphylsOnly, bSkipMissingBones);
-	LevelSetCollisions = MoveTemp(ExtendedCollisionData.LevelSets);
-	SkinnedLevelSetCollisions = MoveTemp(ExtendedCollisionData.SkinnedLevelSets);
-
-	// Append SubBones and reindex
-	const int32 Offset = UsedBoneIndices.Num();
-	UsedBoneIndices.Append(UsedSubBoneIndices);
-	for (Chaos::FSkinnedLevelSetCollisionData& SkinnedLevelSetData : SkinnedLevelSetCollisions)
-	{
-		for (int32& MappedBone : SkinnedLevelSetData.MappedSkinnedBones)
-		{
-			MappedBone += Offset;
-		}
-	}
-}
-
 void FClothingSimulationCollider::ExtractPhysicsAssetCollision(
 	const UPhysicsAsset* PhysicsAsset,
 	const FReferenceSkeleton* ReferenceSkeleton,
@@ -633,7 +623,19 @@ void FClothingSimulationCollider::ExtractPhysicsAssetCollision(
 			// Add tapered capsules
 			for (const FKTaperedCapsuleElem& TaperedCapsuleElem : AggGeom.TaperedCapsuleElems)
 			{
-				if (TaperedCapsuleElem.Length == 0 && !TaperedCapsuleElem.bOneSidedCollision)
+				if (TaperedCapsuleElem.IsExtruded())
+				{
+					FClothCollisionPrim_ExtrudedTaperedCapsule Capsule;
+					Capsule.LocalPosition = TaperedCapsuleElem.Center;
+					Capsule.LocalRotation = TaperedCapsuleElem.Rotation.Quaternion();
+					Capsule.Length = TaperedCapsuleElem.Length;
+					Capsule.Width = TaperedCapsuleElem.Width;
+					Capsule.Radius1 = TaperedCapsuleElem.Radius0;
+					Capsule.Radius2 = TaperedCapsuleElem.Radius1;
+					Capsule.BoneIndex = MappedBoneIndex;
+					ClothCollisionData.ExtrudedTaperedCapsules.Add(MoveTemp(Capsule));
+				}
+				else if (TaperedCapsuleElem.Length == 0 && !TaperedCapsuleElem.IsOneSidedCollision())
 				{
 					// Add extracted sphere collision data
 					FClothCollisionPrim_Sphere Sphere;
@@ -660,7 +662,7 @@ void FClothingSimulationCollider::ExtractPhysicsAssetCollision(
 					FClothCollisionPrim_SphereConnection SphereConnection;
 					SphereConnection.SphereIndices[0] = ClothCollisionData.Spheres.Add(Sphere0);
 					SphereConnection.SphereIndices[1] = ClothCollisionData.Spheres.Add(Sphere1);
-					if (TaperedCapsuleElem.bOneSidedCollision)
+					if (TaperedCapsuleElem.IsOneSidedCollision())
 					{
 						SphereConnection.OneSidedPlaneNormal = TaperedCapsuleElem.Rotation.RotateVector(FVector::ForwardVector);
 					}
@@ -1008,12 +1010,6 @@ TConstArrayView<FImplicitObjectPtr> FClothingSimulationCollider::GetCollisionGeo
 	return CollisionRangeId != INDEX_NONE ? Solver->GetCollisionGeometryView(CollisionRangeId) : TConstArrayView<FImplicitObjectPtr>();
 }
 
-TConstArrayView<TUniquePtr<FImplicitObject>> FClothingSimulationCollider::GetCollisionGeometries(const FClothingSimulationSolver* Solver, const FClothingSimulationCloth* Cloth, ECollisionDataType CollisionDataType) const
-{
-	check(false);
-	return TConstArrayView<TUniquePtr<FImplicitObject>>();
-}
-
 TConstArrayView<bool> FClothingSimulationCollider::GetCollisionStatus(const FClothingSimulationSolver* Solver, const FClothingSimulationCloth* Cloth, ECollisionDataType CollisionDataType) const
 {
 	check(Solver);
diff --git a/Engine/Plugins/ChaosCloth/Source/ChaosCloth/Public/ChaosCloth/ChaosClothingSimulationCollider.h b/Engine/Plugins/ChaosCloth/Source/ChaosCloth/Public/ChaosCloth/ChaosClothingSimulationCollider.h
--- a/Engine/Plugins/ChaosCloth/Source/ChaosCloth/Public/ChaosCloth/ChaosClothingSimulationCollider.h
+++ b/Engine/Plugins/ChaosCloth/Source/ChaosCloth/Public/ChaosCloth/ChaosClothingSimulationCollider.h
@@ -15,11 +15,6 @@ class UClothingAssetCommon;
 class UPhysicsAsset;
 class FClothingSimulationContextCommon;
 struct FReferenceSkeleton;
-#if UE_ENABLE_INCLUDE_ORDER_DEPRECATED_IN_5_6
-// These classes never actually existed.
-class FLevelSetCollisionData;
-class FSkinnedLevelSetCollisionData;
-#endif
 
 namespace Chaos
 {
@@ -107,11 +102,6 @@ namespace Chaos
 			Count
 		};
 
-		UE_DEPRECATED(5.6, "FLevelSetCollisionData has been moved out of this class.")
-		typedef Chaos::FLevelSetCollisionData FLevelSetCollisionData;
-		UE_DEPRECATED(5.6, "FSkinnedLevelSetCollisionData has been moved out of this class.")
-		typedef Chaos::FSkinnedLevelSetCollisionData FSkinnedLevelSetCollisionData;
-
 		CHAOSCLOTH_API FClothingSimulationCollider(const UPhysicsAsset* InPhysicsAsset, const FReferenceSkeleton* InReferenceSkeleton);
 		CHAOSCLOTH_API ~FClothingSimulationCollider();
 
@@ -152,27 +142,10 @@ namespace Chaos
 		// Return current active LOD collision geometries, not thread safe, to use after solver update.
 		CHAOSCLOTH_API TConstArrayView<FImplicitObjectPtr> GetCollisionGeometry(const FClothingSimulationSolver* Solver, const FClothingSimulationCloth* Cloth, ECollisionDataType CollisionDataType) const;
 
-		UE_DEPRECATED(5.4, "Use GetCollisionGeometry instead.")
-		CHAOSCLOTH_API TConstArrayView<TUniquePtr<FImplicitObject>> GetCollisionGeometries(const FClothingSimulationSolver* Solver, const FClothingSimulationCloth* Cloth, ECollisionDataType CollisionDataType) const;
-
 		// Return whether the collision has been hit by a particle during CCD.
 		CHAOSCLOTH_API TConstArrayView<bool> GetCollisionStatus(const FClothingSimulationSolver* Solver, const FClothingSimulationCloth* Cloth, ECollisionDataType CollisionDataType) const;
 		// ---- End of the debugging and visualization functions ----
 
-		// Extract all collisions from the physics asset into a more cloth friendly format
-		UE_DEPRECATED(5.6, "Use the version that uses FClothCollisionDataExtended.")
-		CHAOSCLOTH_API static void ExtractPhysicsAssetCollision(
-			const UPhysicsAsset* PhysicsAsset,
-			const FReferenceSkeleton* ReferenceSkeleton,
-			FClothCollisionData& ClothCollisionData,
-			PRAGMA_DISABLE_DEPRECATION_WARNINGS
-			TArray<FLevelSetCollisionData>& LevelSetCollisions,
-			TArray<FSkinnedLevelSetCollisionData>& SkinnedLevelSetCollisions,
-			PRAGMA_ENABLE_DEPRECATION_WARNINGS
-			TArray<int32>& UsedBoneIndices,
-			bool bUseSphylOnly = false,
-			bool bSkipMissingBones = false);
-
 		// Extract all collisions from the physics asset into a more cloth friendly format
 		CHAOSCLOTH_API static void ExtractPhysicsAssetCollision(
 			const UPhysicsAsset* PhysicsAsset,
diff --git a/Engine/Source/Runtime/ClothingSystemRuntimeInterface/Private/ClothCollisionData.cpp b/Engine/Source/Runtime/ClothingSystemRuntimeInterface/Private/ClothCollisionData.cpp
--- a/Engine/Source/Runtime/ClothingSystemRuntimeInterface/Private/ClothCollisionData.cpp
+++ b/Engine/Source/Runtime/ClothingSystemRuntimeInterface/Private/ClothCollisionData.cpp
@@ -10,6 +10,7 @@ void FClothCollisionData::Reset()
 	SphereConnections.Reset();
 	Convexes.Reset();
 	Boxes.Reset();
+	ExtrudedTaperedCapsules.Reset();
 }
 
 void FClothCollisionData::Append(const FClothCollisionData& InOther)
@@ -36,6 +37,8 @@ void FClothCollisionData::Append(const FClothCollisionData& InOther)
 	Convexes.Append(InOther.Convexes);
 
 	Boxes.Append(InOther.Boxes);
+
+	ExtrudedTaperedCapsules.Append(InOther.ExtrudedTaperedCapsules);
 }
 
 void FClothCollisionData::AppendTransformed(const FClothCollisionData& InOther, const TArray<FTransform>& BoneTransforms)
@@ -43,6 +46,7 @@ void FClothCollisionData::AppendTransformed(const FClothCollisionData& InOther,
 	const int32 NumSpheresBefore = Spheres.Num();
 	const int32 NumConvexesBefore = Convexes.Num();
 	const int32 NumBoxesBefore = Boxes.Num();
+	const int32 NumExtrudedCapsulesBefore = ExtrudedTaperedCapsules.Num();
 
 	Append(InOther);
 
@@ -65,6 +69,12 @@ void FClothCollisionData::AppendTransformed(const FClothCollisionData& InOther,
 		Boxes[Index].LocalPosition = Transform.TransformPosition(Boxes[Index].LocalPosition);
 		Boxes[Index].LocalRotation = Transform.TransformRotation(Boxes[Index].LocalRotation);
 	}
+	for (int32 Index = NumExtrudedCapsulesBefore; Index < ExtrudedTaperedCapsules.Num(); ++Index)
+	{
+		const FTransform& Transform = BoneTransforms[ExtrudedTaperedCapsules[Index].BoneIndex];
+		ExtrudedTaperedCapsules[Index].LocalPosition = Transform.TransformPosition(ExtrudedTaperedCapsules[Index].LocalPosition);
+		ExtrudedTaperedCapsules[Index].LocalRotation = Transform.TransformRotation(ExtrudedTaperedCapsules[Index].LocalRotation);
+	}
 }
 
 void FClothCollisionData::AppendUnique(const FClothCollisionData& InOther)
@@ -102,4 +112,10 @@ void FClothCollisionData::AppendUnique(const FClothCollisionData& InOther)
 	{
 		Boxes.AddUnique(Box);
 	}
+
+	ExtrudedTaperedCapsules.Reserve(ExtrudedTaperedCapsules.Num() + InOther.ExtrudedTaperedCapsules.Num());
+	for (const FClothCollisionPrim_ExtrudedTaperedCapsule& Capsule : InOther.ExtrudedTaperedCapsules)
+	{
+		ExtrudedTaperedCapsules.AddUnique(Capsule);
+	}
 }
diff --git a/Engine/Source/Runtime/ClothingSystemRuntimeInterface/Public/ClothCollisionData.h b/Engine/Source/Runtime/ClothingSystemRuntimeInterface/Public/ClothCollisionData.h
--- a/Engine/Source/Runtime/ClothingSystemRuntimeInterface/Public/ClothCollisionData.h
+++ b/Engine/Source/Runtime/ClothingSystemRuntimeInterface/Public/ClothCollisionData.h
@@ -28,7 +28,7 @@ struct FClothCollisionData
 
 	bool IsEmpty() const
 	{
-		return Spheres.IsEmpty() && SphereConnections.IsEmpty() && Convexes.IsEmpty() && Boxes.IsEmpty();
+		return Spheres.IsEmpty() && SphereConnections.IsEmpty() && Convexes.IsEmpty() && Boxes.IsEmpty() && ExtrudedTaperedCapsules.IsEmpty();;
 	}
 
 	// Sphere data
@@ -46,6 +46,10 @@ struct FClothCollisionData
 	// Box data
 	UPROPERTY(EditAnywhere, Category = Collision)
 	TArray<FClothCollisionPrim_Box> Boxes;
+
+	// Extruded tapered capsule data
+	UPROPERTY(EditAnywhere, Category = Collision)
+	TArray<FClothCollisionPrim_ExtrudedTaperedCapsule> ExtrudedTaperedCapsules;
 };
 
 #undef UE_API
diff --git a/Engine/Source/Runtime/ClothingSystemRuntimeInterface/Public/ClothCollisionPrim.h b/Engine/Source/Runtime/ClothingSystemRuntimeInterface/Public/ClothCollisionPrim.h
--- a/Engine/Source/Runtime/ClothingSystemRuntimeInterface/Public/ClothCollisionPrim.h
+++ b/Engine/Source/Runtime/ClothingSystemRuntimeInterface/Public/ClothCollisionPrim.h
@@ -183,3 +183,61 @@ struct FClothCollisionPrim_Box
 	UPROPERTY()
 	int32 BoneIndex;
 };
+
+/** Data for an extruded tapered capsule primitive. */
+USTRUCT()
+struct FClothCollisionPrim_ExtrudedTaperedCapsule
+{
+	GENERATED_BODY()
+
+	explicit FClothCollisionPrim_ExtrudedTaperedCapsule(
+		const FVector& InLocalPosition = FVector::ZeroVector,
+		const FQuat& InLocalRotation = FQuat::Identity,
+		float InLength = 0.f,
+		float InWidth = 0.f,
+		float InRadius1 = 0.f,
+		float InRadius2 = 0.f,
+		int32 InBoneIndex = INDEX_NONE)
+		: LocalPosition(InLocalPosition)
+		, LocalRotation(InLocalRotation)
+		, Length(InLength)
+		, Width(InWidth)
+		, Radius1(InRadius1)
+		, Radius2(InRadius2)
+		, BoneIndex(InBoneIndex)
+	{
+	}
+
+	bool operator==(const FClothCollisionPrim_ExtrudedTaperedCapsule& Other) const
+	{
+		return
+			BoneIndex == Other.BoneIndex &&
+			LocalPosition.Equals(Other.LocalPosition) &&
+			LocalRotation.Equals(Other.LocalRotation) &&
+			Length == Other.Length &&
+			Width == Other.Width &&
+			Radius1 == Other.Radius1 &&
+			Radius2 == Other.Radius2;
+	}
+
+	UPROPERTY()
+	FVector LocalPosition;
+
+	UPROPERTY()
+	FQuat LocalRotation;
+
+	UPROPERTY()
+	float Length;
+
+	UPROPERTY()
+	float Width;
+
+	UPROPERTY()
+	float Radius1;
+
+	UPROPERTY()
+	float Radius2;
+
+	UPROPERTY()
+	int32 BoneIndex;
+};
diff --git a/Engine/Source/Runtime/Engine/Classes/PhysicsEngine/TaperedCapsuleElem.h b/Engine/Source/Runtime/Engine/Classes/PhysicsEngine/TaperedCapsuleElem.h
--- a/Engine/Source/Runtime/Engine/Classes/PhysicsEngine/TaperedCapsuleElem.h
+++ b/Engine/Source/Runtime/Engine/Classes/PhysicsEngine/TaperedCapsuleElem.h
@@ -36,8 +36,12 @@ struct FKTaperedCapsuleElem : public FKShapeElem
 	UPROPERTY(Category= Capsule, EditAnywhere)
 	float Length;
 
-	/** (Cloth-only) Treat as one-sided collider, where all collisions are pushed to the +x side.*/
-	UPROPERTY(Category = Capsule, EditAnywhere)
+	/** (Cloth-only) Line segment width for extruded tapered capsules. */
+	UPROPERTY(Category = Capsule, EditAnywhere, meta = (ClampMin = 0.f))
+	float Width = 0.f;
+
+	/** (Cloth-only) Treat as one-sided collider, where all collisions are pushed to the +x side. Not supported when Width != 0*/
+	UPROPERTY(Category = Capsule, EditAnywhere, meta = (EditCondition = "Width <= 0.00000001"))
 	bool bOneSidedCollision = false;
 
 	ENGINE_API FKTaperedCapsuleElem();
@@ -52,6 +56,7 @@ struct FKTaperedCapsuleElem : public FKShapeElem
 			LHS.Radius0 == RHS.Radius0 &&
 			LHS.Radius1 == RHS.Radius1 &&
 			LHS.Length == RHS.Length &&
+			LHS.Width == RHS.Width &&
 			LHS.bOneSidedCollision == RHS.bOneSidedCollision);
 	};
 
@@ -59,7 +64,17 @@ struct FKTaperedCapsuleElem : public FKShapeElem
 	virtual FTransform GetTransform() const override final
 	{
 		return FTransform(Rotation, Center );
-	};
+	}
+
+	bool IsExtruded() const
+	{
+		return Width > UE_SMALL_NUMBER;
+	}
+
+	bool IsOneSidedCollision() const
+	{
+		return bOneSidedCollision && !IsExtruded();
+	}
 
 	void SetTransform( const FTransform& InTransform )
 	{
@@ -87,13 +102,22 @@ struct FKTaperedCapsuleElem : public FKShapeElem
 	ENGINE_API float GetScaledCylinderLength(const FVector& Scale3D) const;
 	/** Returns half of the total scaled length of the Sphyl, which includes the scaled top and bottom caps */
 	ENGINE_API float GetScaledHalfLength(const FVector& Scale3D) const;
+	/** Returns the scaled length of the two cylinders for extruded Sphyls **/
+	ENGINE_API void GetScaledExtrudedCylinderLengths(const FVector& Scale3D, float& OutLength0, float& OutLength1) const;
+	/** Returns half of the total scaled length of the extruded Sphyl, which includes the scaled top and bottom caps */
+	ENGINE_API void GetScaledExtrudedHalfLengths(const FVector& Scale3D, float& OutLength0, float& OutLength1) const;
 
 	/** 
 	 * Draws just the sides of a tapered capsule specified by provided Spheres that can have different radii.  Does not draw the spheres, just the sleeve.
 	 * Extent geometry endpoints not necessarily coplanar with sphere origins (uses hull horizon)
 	 * Otherwise uses the great-circle cap assumption.
 	 */
-	ENGINE_API static void DrawTaperedCapsuleSides(FPrimitiveDrawInterface* PDI, const FTransform& ElemTM, const FVector& InCenter0, const FVector& InCenter1, float InRadius0, float InRadius1, const FColor& Color, bool bInSplitTaperedCylinder = false);
+
+	ENGINE_API static void DrawTaperedCapsuleSides(FPrimitiveDrawInterface* PDI, const FQuat& CapsuleOrientation, const FVector& InCenter0, const FVector& InCenter1, float InRadius0, float InRadius1, const FColor& Color, bool bInSplitTaperedCylinder = false);
+	static void DrawTaperedCapsuleSides(FPrimitiveDrawInterface* PDI, const FTransform& ElemTM, const FVector& InCenter0, const FVector& InCenter1, float InRadius0, float InRadius1, const FColor& Color, bool bInSplitTaperedCylinder = false)
+	{
+		DrawTaperedCapsuleSides(PDI, ElemTM.GetRotation(), InCenter0, InCenter1, InRadius0, InRadius1, Color, bInSplitTaperedCylinder);
+	}
 
 	ENGINE_API static EAggCollisionShape::Type StaticShapeType;
 };
\ No newline at end of file
diff --git a/Engine/Source/Runtime/Engine/Private/Collision/KAggregateGeom.cpp b/Engine/Source/Runtime/Engine/Private/Collision/KAggregateGeom.cpp
--- a/Engine/Source/Runtime/Engine/Private/Collision/KAggregateGeom.cpp
+++ b/Engine/Source/Runtime/Engine/Private/Collision/KAggregateGeom.cpp
@@ -530,8 +530,16 @@ FBox FKTaperedCapsuleElem::CalcAABB(const FTransform& BoneTM, float Scale) const
 	// Scale by length of sphyl
 	const FVector3f AbsDist = (Scale * 0.5f * Length) * AbsAxis;
 
-	const FVector3f MaxPos = TaperedCapsuleCenter + AbsDist;
-	const FVector3f MinPos = TaperedCapsuleCenter - AbsDist;
+	// Get extruded tapered capsule axis direction
+	const FVector3f ExtrudedAxis = (FVector3f)ElemTM.GetScaledAxis(EAxis::X);
+	// Get abs of that vector
+	const FVector3f AbsExtrudedAxis(FMath::Abs(ExtrudedAxis.X), FMath::Abs(ExtrudedAxis.Y), FMath::Abs(ExtrudedAxis.Z));
+	// Scale by width of sphyl
+	const FVector3f AbsExtrudedDist = (Scale * 0.5f * Width) * AbsExtrudedAxis;
+
+
+	const FVector3f MaxPos = TaperedCapsuleCenter + AbsDist + AbsExtrudedDist;
+	const FVector3f MinPos = TaperedCapsuleCenter - AbsDist - AbsExtrudedDist;
 	const FVector3f Extent0(Scale * Radius0);
 	const FVector3f Extent1(Scale * Radius1);
 
diff --git a/Engine/Source/Runtime/Engine/Private/PhysicsEngine/BodySetup.cpp b/Engine/Source/Runtime/Engine/Private/PhysicsEngine/BodySetup.cpp
--- a/Engine/Source/Runtime/Engine/Private/PhysicsEngine/BodySetup.cpp
+++ b/Engine/Source/Runtime/Engine/Private/PhysicsEngine/BodySetup.cpp
@@ -2171,20 +2171,35 @@ float FKSphylElem::GetClosestPointAndNormal(const FVector& WorldPosition, const
 
 void FKTaperedCapsuleElem::ScaleElem(FVector DeltaSize, float MinSize)
 {
-	float DeltaRadius0 = DeltaSize.X;
-	float DeltaRadius1 = DeltaSize.Y;
+	float DeltaRadius = DeltaSize.Y;
+	float DeltaWidth = 0.f;
+	if (IsExtruded())
+	{
+		DeltaWidth = DeltaSize.X;
+	}
+	else
+	{
+		if (FMath::Abs(DeltaSize.X) > FMath::Abs(DeltaRadius))
+		{
+			DeltaRadius = DeltaSize.X;
+		}
+	}
 
-	float DeltaHeight = DeltaSize.Z;
-	float radius0 = FMath::Max(Radius0 + DeltaRadius0, MinSize);
-	float radius1 = FMath::Max(Radius1 + DeltaRadius1, MinSize);
-	float length = Length + DeltaHeight;
+	const float DeltaHeight = DeltaSize.Z;
+	const float NewRadius0 = FMath::Max(Radius0 + DeltaRadius, MinSize);
+	const float NewRadius1 = FMath::Max(Radius1 + DeltaRadius, MinSize);
+	float NewLength = Length + DeltaHeight;
 
-	length += ((Radius1 - radius1) + (Radius0 - radius0)) * 0.5f;
-	length = FMath::Max(0.f, length);
+	NewLength += ((Radius1 - NewRadius1) + (Radius0 - NewRadius0)) * .5f;
+	NewLength = FMath::Max(0.f, NewLength);
 
-	Radius0 = radius0;
-	Radius1 = radius1;
-	Length = length;
+	Radius0 = NewRadius0;
+	Radius1 = NewRadius1;
+	Length = NewLength;
+	if (IsExtruded())
+	{
+		Width = FMath::Max(MinSize, Width + DeltaWidth);
+	}
 }
 
 FKTaperedCapsuleElem FKTaperedCapsuleElem::GetFinalScaled(const FVector& Scale3D, const FTransform& RelativeTM) const
@@ -2198,6 +2213,7 @@ FKTaperedCapsuleElem FKTaperedCapsuleElem::GetFinalScaled(const FVector& Scale3D
 
 	GetScaledRadii(Scale3DAbs, ScaledTaperedCapsuleElem.Radius0, ScaledTaperedCapsuleElem.Radius1);
 	ScaledTaperedCapsuleElem.Length = GetScaledCylinderLength(Scale3DAbs);
+	ScaledTaperedCapsuleElem.Width *= Scale3DAbs.X;
 
 	const FTransform ScaleTransform(FQuat::Identity, FVector::ZeroVector, Scale3D);
 	const FTransform RotationTransform(ScaledTaperedCapsuleElem.Rotation, FVector::ZeroVector, Scale3D);
@@ -2230,6 +2246,22 @@ float FKTaperedCapsuleElem::GetScaledHalfLength(const FVector& Scale3D) const
 	return FMath::Max<float>((Length + Radius0 + Radius1) * FMath::Abs(Scale3D.Z) * 0.5f, 0.1f);
 }
 
+void FKTaperedCapsuleElem::GetScaledExtrudedCylinderLengths(const FVector& Scale3D, float& OutLength0, float& OutLength1) const
+{
+	float ScaledRadius0, ScaledRadius1;
+	GetScaledRadii(Scale3D, ScaledRadius0, ScaledRadius1);
+	float ScaledExtendedHalfLength0, ScaledExtendedHalfLength1;
+	GetScaledExtrudedHalfLengths(Scale3D, ScaledExtendedHalfLength0, ScaledExtendedHalfLength1);
+	OutLength0 = FMath::Max(0.1f, ((ScaledExtendedHalfLength0 - ScaledRadius0) * 2.0f));
+	OutLength1 = FMath::Max(0.1f, ((ScaledExtendedHalfLength1 - ScaledRadius1) * 2.0f));
+}
+
+void FKTaperedCapsuleElem::GetScaledExtrudedHalfLengths(const FVector& Scale3D, float& OutLength0, float& OutLength1) const
+{
+	OutLength0 = FMath::Max<float>((Width + 2.f * Radius0) * FMath::Abs(Scale3D.X) * 0.5f, 0.1f);
+	OutLength1 = FMath::Max<float>((Width + 2.f * Radius1) * FMath::Abs(Scale3D.X) * 0.5f, 0.1f);
+}
+
 float FKConvexElem::GetClosestPointAndNormal(const FVector& WorldPosition, const FTransform& BoneToWorldTM, FVector& ClosestWorldPosition, FVector& Normal) const
 {
 	float MinScale, MinScaleAbs;
diff --git a/Engine/Source/Runtime/Engine/Private/PhysicsEngine/PhysDrawing.cpp b/Engine/Source/Runtime/Engine/Private/PhysicsEngine/PhysDrawing.cpp
--- a/Engine/Source/Runtime/Engine/Private/PhysicsEngine/PhysDrawing.cpp
+++ b/Engine/Source/Runtime/Engine/Private/PhysicsEngine/PhysDrawing.cpp
@@ -426,7 +426,7 @@ void FKSphylElem::DrawElemSolid(FPrimitiveDrawInterface* PDI, const FTransform&
 // FKTaperedCapsuleElem
 /////////////////////////////////////////////////////////////////////////////////////
 
-void FKTaperedCapsuleElem::DrawTaperedCapsuleSides(FPrimitiveDrawInterface* PDI, const FTransform& ElemTM, const FVector& InCenter0, const FVector& InCenter1, float InRadius0, float InRadius1, const FColor& Color, bool bInSplitTaperedCylinder)
+void FKTaperedCapsuleElem::DrawTaperedCapsuleSides(FPrimitiveDrawInterface* PDI, const FQuat& CapsuleOrientation, const FVector& InCenter0, const FVector& InCenter1, float InRadius0, float InRadius1, const FColor& Color, bool bInSplitTaperedCylinder)
 {
 	const FColor HalfTaperedColor = bInSplitTaperedCylinder ? (FLinearColor(Color) * .2f).ToFColor(true) : Color;
 
@@ -441,7 +441,6 @@ void FKTaperedCapsuleElem::DrawTaperedCapsuleSides(FPrimitiveDrawInterface* PDI,
 	{
 		return;
 	}
-	FQuat CapsuleOrientation = ElemTM.GetRotation();
 	float OffsetZ = (InRadius1 - InRadius0) / Distance;
 	float ScaleXY = FMath::Sqrt(1.0f - FMath::Square(OffsetZ));
 
@@ -475,28 +474,9 @@ void FKTaperedCapsuleElem::DrawTaperedCapsuleSides(FPrimitiveDrawInterface* PDI,
 	}
 }
 
-void FKTaperedCapsuleElem::DrawElemWire(class FPrimitiveDrawInterface* PDI, const FTransform& ElemTM, float Scale, const FColor Color) const
-{
-	DrawElemWire(PDI, ElemTM, FVector(Scale), Color);
-}
-
-void FKTaperedCapsuleElem::DrawElemSolid(class FPrimitiveDrawInterface* PDI, const FTransform& ElemTM, float Scale, const FMaterialRenderProxy* MaterialRenderProxy) const
-{
-	DrawElemSolid(PDI, ElemTM, FVector(Scale), MaterialRenderProxy);
-}
-
-void FKTaperedCapsuleElem::DrawElemWire(class FPrimitiveDrawInterface* PDI, const FTransform& ElemTM, const FVector& Scale3D, const FColor Color) const
+static void DrawWireframeTaperedCapsule(class FPrimitiveDrawInterface* PDI, const FColor& Color, const FQuat& CapsuleOrientation, const FVector& Origin, const FVector& XAxis, const FVector& YAxis, const FVector& ZAxis, const float ScaledHalfLength, const float ScaledRadius0, const float ScaledRadius1, const bool bIsOneSided)
 {
-	const FColor HalfTaperedColor = bOneSidedCollision ? ( FLinearColor(Color) * .2f).ToFColor(true) : Color;
-	const FVector Origin = ElemTM.GetLocation();
-	const FVector XAxis = ElemTM.GetScaledAxis(EAxis::X);
-	const FVector YAxis = ElemTM.GetScaledAxis(EAxis::Y);
-	const FVector ZAxis = ElemTM.GetScaledAxis(EAxis::Z);
-	const float ScaledHalfLength = GetScaledCylinderLength(Scale3D) * .5f;
-	float ScaledRadius0;
-	float ScaledRadius1;
-	GetScaledRadii(Scale3D, ScaledRadius0, ScaledRadius1);
-	
+	const FColor HalfTaperedColor = bIsOneSided ? (FLinearColor(Color) * .2f).ToFColor(true) : Color;
 	// Draw top and bottom circles
 	const FVector TopEnd = Origin + (ScaledHalfLength * ZAxis);
 	const FVector BottomEnd = Origin - (ScaledHalfLength * ZAxis);
@@ -504,10 +484,10 @@ void FKTaperedCapsuleElem::DrawElemWire(class FPrimitiveDrawInterface* PDI, cons
 	const FVector Separation = TopEnd - BottomEnd;
 	float Distance = Separation.Size();
 
-	if(Separation.IsNearlyZero() || Distance <= FMath::Abs(ScaledRadius0 - ScaledRadius1))
+	if (Separation.IsNearlyZero() || Distance <= FMath::Abs(ScaledRadius0 - ScaledRadius1))
 	{
 		// Degenerate or one end encompasses the other - draw a sphere
-		if(ScaledRadius0 > ScaledRadius1)
+		if (ScaledRadius0 > ScaledRadius1)
 		{
 			DrawArc(PDI, TopEnd, XAxis, YAxis, -90.f, 90.f, ScaledRadius0, DrawCollisionSides, Color, SDPG_World);
 			DrawArc(PDI, TopEnd, XAxis, ZAxis, -90.f, 90.f, ScaledRadius0, DrawCollisionSides, Color, SDPG_World);
@@ -528,7 +508,7 @@ void FKTaperedCapsuleElem::DrawElemWire(class FPrimitiveDrawInterface* PDI, cons
 	}
 	else
 	{
-		if(ScaledRadius0 > ScaledRadius1)
+		if (ScaledRadius0 > ScaledRadius1)
 		{
 			DrawArc(PDI, TopEnd, XAxis, YAxis, -90.f, 90.f, ScaledRadius0, DrawCollisionSides, Color, SDPG_World);
 			DrawArc(PDI, TopEnd, XAxis, YAxis, 90.f, 270.f, ScaledRadius0, DrawCollisionSides, HalfTaperedColor, SDPG_World);
@@ -537,24 +517,74 @@ void FKTaperedCapsuleElem::DrawElemWire(class FPrimitiveDrawInterface* PDI, cons
 		{
 			DrawArc(PDI, BottomEnd, XAxis, YAxis, -90.f, 90.f, ScaledRadius1, DrawCollisionSides, Color, SDPG_World);
 			DrawArc(PDI, BottomEnd, XAxis, YAxis, 90.f, 270.f, ScaledRadius1, DrawCollisionSides, HalfTaperedColor, SDPG_World);
-		}	
+		}
 
 		const FVector NegZAxis = -ZAxis;
 
-		const float OffsetZ0 = ((ScaledRadius0 - ScaledRadius1) * ScaledRadius0) / Distance;
-		float TopLimit = 90.0f - FMath::RadiansToDegrees(FMath::Asin(OffsetZ0 / ScaledRadius0));
-		DrawArc(PDI, TopEnd, NegZAxis, YAxis, TopLimit, 360.0f-TopLimit, ScaledRadius0, DrawCollisionSides, Color, SDPG_World);
+		const float OffsetZ0 = FMath::IsNearlyZero(Distance) ? 0.f : ((ScaledRadius0 - ScaledRadius1) * ScaledRadius0) / Distance;
+		float TopLimit = FMath::IsNearlyZero(ScaledRadius0) ? 90.f : 90.0f - FMath::RadiansToDegrees(FMath::Asin(OffsetZ0 / ScaledRadius0));
+		DrawArc(PDI, TopEnd, NegZAxis, YAxis, TopLimit, 360.0f - TopLimit, ScaledRadius0, DrawCollisionSides, Color, SDPG_World);
 		DrawArc(PDI, TopEnd, NegZAxis, XAxis, TopLimit, 180.f, ScaledRadius0, DrawCollisionSides, Color, SDPG_World);
-		DrawArc(PDI, TopEnd, NegZAxis, XAxis, 180.f,360.0f - TopLimit , ScaledRadius0, DrawCollisionSides, HalfTaperedColor, SDPG_World);
+		DrawArc(PDI, TopEnd, NegZAxis, XAxis, 180.f, 360.0f - TopLimit, ScaledRadius0, DrawCollisionSides, HalfTaperedColor, SDPG_World);
 
-		const float OffsetZ1 = ((ScaledRadius0 - ScaledRadius1) * ScaledRadius1) / Distance;
-		float BottomLimit =  180.0f - TopLimit;
- 		DrawArc(PDI, BottomEnd, ZAxis, YAxis, BottomLimit, 360.0f-BottomLimit, ScaledRadius1, DrawCollisionSides, Color, SDPG_World);
+		const float OffsetZ1 = FMath::IsNearlyZero(Distance) ? 0.f : ((ScaledRadius0 - ScaledRadius1) * ScaledRadius1) / Distance;
+		float BottomLimit = 180.0f - TopLimit;
+		DrawArc(PDI, BottomEnd, ZAxis, YAxis, BottomLimit, 360.0f - BottomLimit, ScaledRadius1, DrawCollisionSides, Color, SDPG_World);
 		DrawArc(PDI, BottomEnd, ZAxis, XAxis, BottomLimit, 180.f, ScaledRadius1, DrawCollisionSides, Color, SDPG_World);
 		DrawArc(PDI, BottomEnd, ZAxis, XAxis, 180.f, 360.0f - BottomLimit, ScaledRadius1, DrawCollisionSides, HalfTaperedColor, SDPG_World);
 
 		// Draw connecty lines
-		DrawTaperedCapsuleSides(PDI, ElemTM, TopEnd, BottomEnd, ScaledRadius0, ScaledRadius1, Color, bOneSidedCollision);
+		FKTaperedCapsuleElem::DrawTaperedCapsuleSides(PDI, CapsuleOrientation, TopEnd, BottomEnd, ScaledRadius0, ScaledRadius1, Color, bIsOneSided);
+	}
+}
+
+void FKTaperedCapsuleElem::DrawElemWire(class FPrimitiveDrawInterface* PDI, const FTransform& ElemTM, float Scale, const FColor Color) const
+{
+	DrawElemWire(PDI, ElemTM, FVector(Scale), Color);
+}
+
+void FKTaperedCapsuleElem::DrawElemSolid(class FPrimitiveDrawInterface* PDI, const FTransform& ElemTM, float Scale, const FMaterialRenderProxy* MaterialRenderProxy) const
+{
+	DrawElemSolid(PDI, ElemTM, FVector(Scale), MaterialRenderProxy);
+}
+
+void FKTaperedCapsuleElem::DrawElemWire(class FPrimitiveDrawInterface* PDI, const FTransform& ElemTM, const FVector& Scale3D, const FColor Color) const
+{
+	const FVector Origin = ElemTM.GetLocation();
+	const FVector XAxis = ElemTM.GetScaledAxis(EAxis::X);
+	const FVector YAxis = ElemTM.GetScaledAxis(EAxis::Y);
+	const FVector ZAxis = ElemTM.GetScaledAxis(EAxis::Z);
+	const float ScaledHalfLength = GetScaledCylinderLength(Scale3D) * .5f;
+	float ScaledRadius0;
+	float ScaledRadius1;
+	GetScaledRadii(Scale3D, ScaledRadius0, ScaledRadius1);
+
+	if (IsExtruded())
+	{
+		const float ScaledHalfWidth = Width * Scale3D.X * .5f;
+		const FVector FrontOrigin = Origin + (ScaledHalfWidth * XAxis);
+		const FVector BackOrigin = Origin - (ScaledHalfWidth * XAxis);
+		DrawWireframeTaperedCapsule(PDI, Color, ElemTM.GetRotation(), FrontOrigin, XAxis, YAxis, ZAxis, ScaledHalfLength, ScaledRadius0, ScaledRadius1, IsOneSidedCollision());
+		DrawWireframeTaperedCapsule(PDI, Color, ElemTM.GetRotation(), BackOrigin, XAxis, YAxis, ZAxis, ScaledHalfLength, ScaledRadius0, ScaledRadius1, IsOneSidedCollision());
+
+
+		// Draw connecty lines
+		const FVector TopFrontEnd = FrontOrigin + ScaledHalfLength * ZAxis;
+		const FVector TopBackEnd = BackOrigin + ScaledHalfLength * ZAxis;
+		PDI->DrawLine(TopFrontEnd + ScaledRadius0 * ZAxis, TopBackEnd + ScaledRadius0 * ZAxis, Color, SDPG_World);
+		PDI->DrawLine(TopFrontEnd + ScaledRadius0 * YAxis, TopBackEnd + ScaledRadius0 * YAxis, Color, SDPG_World);
+		PDI->DrawLine(TopFrontEnd - ScaledRadius0 * YAxis, TopBackEnd - ScaledRadius0 * YAxis, Color, SDPG_World);
+
+		const FVector BottomFrontEnd = FrontOrigin - ScaledHalfLength * ZAxis;
+		const FVector BottomBackEnd = BackOrigin - ScaledHalfLength * ZAxis;
+		PDI->DrawLine(BottomFrontEnd - ScaledRadius1 * ZAxis, BottomBackEnd - ScaledRadius1 * ZAxis, Color, SDPG_World);
+		PDI->DrawLine(BottomFrontEnd + ScaledRadius1 * YAxis, BottomBackEnd + ScaledRadius1 * YAxis, Color, SDPG_World);
+		PDI->DrawLine(BottomFrontEnd - ScaledRadius1 * YAxis, BottomBackEnd - ScaledRadius1 * YAxis, Color, SDPG_World);
+
+	}
+	else
+	{
+		DrawWireframeTaperedCapsule(PDI, Color, ElemTM.GetRotation(), Origin, XAxis, YAxis, ZAxis, ScaledHalfLength, ScaledRadius0, ScaledRadius1, IsOneSidedCollision());
 	}
 }
 
@@ -563,19 +593,22 @@ struct FScopedTaperedCapsuleBuilder
 	FScopedTaperedCapsuleBuilder(const FKTaperedCapsuleElem& InTaperedCapsule, const FVector& InScale3D, ERHIFeatureLevel::Type InFeatureLevel)
 		: MeshBuilder(InFeatureLevel)
 	{
+		const bool bIsExtrudedCapsule = InTaperedCapsule.IsExtruded();
+
 		const float ScaledLength = InTaperedCapsule.GetScaledCylinderLength(InScale3D);
 		const float ScaledHalfLength = ScaledLength * 0.5f;
 		float ScaledRadius0;
 		float ScaledRadius1;
 		InTaperedCapsule.GetScaledRadii(InScale3D, ScaledRadius0, ScaledRadius1);
+		const float ScaledHalfWidth = bIsExtrudedCapsule ? InTaperedCapsule.Width * InScale3D.X * .5f : 0.f;;
 
 		// Deal with degenerates
 		float SphereOffset0 = ScaledHalfLength;
 		float SphereOffset1 = -ScaledHalfLength;
-		if(FMath::IsNearlyZero(ScaledLength) || ScaledLength <= FMath::Abs(ScaledRadius0 - ScaledRadius1))
+		if (FMath::IsNearlyZero(ScaledLength) || ScaledLength <= FMath::Abs(ScaledRadius0 - ScaledRadius1))
 		{
 			// Degenerate or one end encompasses the other - we need to draw a sphere, so map one end to the other
-			if(ScaledRadius0 > ScaledRadius1)
+			if (ScaledRadius0 > ScaledRadius1)
 			{
 				SphereOffset1 = SphereOffset0;
 				ScaledRadius1 = ScaledRadius0;
@@ -586,21 +619,21 @@ struct FScopedTaperedCapsuleBuilder
 				ScaledRadius0 = ScaledRadius1;
 			}
 		}
-		
-		const int32 NumSides = DrawCollisionSides;
+
+		const int32 NumSides = DrawCollisionSides + (bIsExtrudedCapsule ? 2 : 0);
 		const int32 NumRings = (DrawCollisionSides / 2) + 1;
 
 		// The first/last arc are on top of each other.
 		const int32 NumVerts = (NumSides + 1) * (NumRings + 1);
-		Verts = (FDynamicMeshVertex*)FMemory::Malloc(NumVerts * sizeof(FDynamicMeshVertex));
+		Verts.SetNumUninitialized(NumVerts);
 
 		// Calculate verts for one arc
-		ArcVerts = (FDynamicMeshVertex*)FMemory::Malloc((NumRings + 1) * sizeof(FDynamicMeshVertex));
+		ArcVerts.SetNumUninitialized(NumRings + 1);
 
 		// Calc arc split point
-		const float OffsetZ = ((ScaledRadius0 - ScaledRadius1) * ScaledRadius0) / ScaledLength;
-		float SplitAngle = 90.0f - FMath::RadiansToDegrees(FMath::Asin(OffsetZ / ScaledRadius0));
-		int32 SplitPoint = (int32)((float)NumRings * ((180.0f - SplitAngle) / 180.0f));
+		const float OffsetZ = FMath::IsNearlyZero(ScaledLength) ? 0.f : ((ScaledRadius0 - ScaledRadius1) * ScaledRadius0) / ScaledLength;
+		const float SplitAngle = FMath::IsNearlyZero(ScaledRadius0) ? 90.f : 90.0f - FMath::RadiansToDegrees(FMath::Asin(OffsetZ / ScaledRadius0));
+		const int32 SplitPoint = (int32)((float)NumRings * ((180.0f - SplitAngle) / 180.0f));
 
 		for (int32 RingIdx = 0; RingIdx < NumRings + 1; RingIdx++)
 		{
@@ -633,37 +666,82 @@ struct FScopedTaperedCapsuleBuilder
 				FVector3f(1, 0, 0),
 				FVector3f(0.0f, -SpherePos.Z, SpherePos.Y),
 				(FVector3f)SpherePos
-				);
+			);
 
 			ArcVert->TextureCoordinate[0].X = 0.0f;
 			ArcVert->TextureCoordinate[0].Y = ((float)RingIdx / NumRings);
 		}
 
 		// Then rotate this arc NumSides+1 times.
-		for (int32 SideIdx = 0; SideIdx < NumSides + 1; SideIdx++)
+		for (int32 SideIdx = 0; SideIdx < DrawCollisionSides + 1; SideIdx++)
 		{
-			const FRotator3f ArcRotator(0, 360.f * ((float)SideIdx / NumSides), 0);
+			const FRotator3f ArcRotator(0, 360.f * ((float)SideIdx / DrawCollisionSides), 0);
 			const FRotationMatrix44f ArcRot(ArcRotator);
-			const float XTexCoord = ((float)SideIdx / NumSides);
+			const float XTexCoord = ((float)SideIdx / DrawCollisionSides);
 
+			const FVector3f XOffset(SideIdx <= DrawCollisionSides / 2 ? -ScaledHalfWidth : ScaledHalfWidth, 0, 0);
+			const int32 SideIdxToVIxOffset = (bIsExtrudedCapsule && SideIdx > DrawCollisionSides / 2) ? 1 : 0;
 			for (int32 VertIdx = 0; VertIdx < NumRings + 1; VertIdx++)
 			{
-				int32 VIx = (NumRings + 1)*SideIdx + VertIdx;
+				const int32 VIx = (NumRings + 1) * (SideIdx + SideIdxToVIxOffset) + VertIdx;
 
-				Verts[VIx].Position = ArcRot.TransformPosition(ArcVerts[VertIdx].Position);
+				Verts[VIx].Position = ArcRot.TransformPosition(ArcVerts[VertIdx].Position) + XOffset;
 
 				Verts[VIx].SetTangents(
 					ArcRot.TransformVector(ArcVerts[VertIdx].TangentX.ToFVector3f()),
 					ArcRot.TransformVector(ArcVerts[VertIdx].GetTangentY()),
 					ArcRot.TransformVector(ArcVerts[VertIdx].TangentZ.ToFVector3f())
-					);
+				);
 
 				Verts[VIx].TextureCoordinate[0].X = XTexCoord;
 				Verts[VIx].TextureCoordinate[0].Y = ArcVerts[VertIdx].TextureCoordinate[0].Y;
 			}
+
+			if (bIsExtrudedCapsule)
+			{
+
+				if (SideIdx == DrawCollisionSides / 2)
+				{
+					// Duplicate vertex on second half
+					for (int32 VertIdx = 0; VertIdx < NumRings + 1; VertIdx++)
+					{
+						const int32 VIx = (NumRings + 1) * (SideIdx + 1) + VertIdx;
+
+						Verts[VIx].Position = ArcRot.TransformPosition(ArcVerts[VertIdx].Position) + FVector3f(ScaledHalfWidth, 0.f, 0.f);
+
+						Verts[VIx].SetTangents(
+							ArcRot.TransformVector(ArcVerts[VertIdx].TangentX.ToFVector3f()),
+							ArcRot.TransformVector(ArcVerts[VertIdx].GetTangentY()),
+							ArcRot.TransformVector(ArcVerts[VertIdx].TangentZ.ToFVector3f())
+						);
+
+						Verts[VIx].TextureCoordinate[0].X = XTexCoord;
+						Verts[VIx].TextureCoordinate[0].Y = ArcVerts[VertIdx].TextureCoordinate[0].Y;
+					}
+				}
+				else if (SideIdx == DrawCollisionSides)
+				{
+					// Duplicate vertex on second half
+					for (int32 VertIdx = 0; VertIdx < NumRings + 1; VertIdx++)
+					{
+						const int32 VIx = (NumRings + 1) * (SideIdx + 2) + VertIdx;
+
+						Verts[VIx].Position = ArcRot.TransformPosition(ArcVerts[VertIdx].Position) - FVector3f(ScaledHalfWidth, 0.f, 0.f);
+
+						Verts[VIx].SetTangents(
+							ArcRot.TransformVector(ArcVerts[VertIdx].TangentX.ToFVector3f()),
+							ArcRot.TransformVector(ArcVerts[VertIdx].GetTangentY()),
+							ArcRot.TransformVector(ArcVerts[VertIdx].TangentZ.ToFVector3f())
+						);
+
+						Verts[VIx].TextureCoordinate[0].X = XTexCoord;
+						Verts[VIx].TextureCoordinate[0].Y = ArcVerts[VertIdx].TextureCoordinate[0].Y;
+					}
+				}
+			}
 		}
 
-	
+
 		{
 			// Add all of the vertices to the mesh.
 			for (int32 VertIdx = 0; VertIdx < NumVerts; VertIdx++)
@@ -688,13 +766,11 @@ struct FScopedTaperedCapsuleBuilder
 
 	~FScopedTaperedCapsuleBuilder()
 	{
-		FMemory::Free(Verts);
-		FMemory::Free(ArcVerts);
 	}
 
 	FDynamicMeshBuilder MeshBuilder;
-	FDynamicMeshVertex* Verts;
-	FDynamicMeshVertex* ArcVerts;
+	TArray<FDynamicMeshVertex> Verts;
+	TArray<FDynamicMeshVertex> ArcVerts;
 };
 
 void FKTaperedCapsuleElem::GetElemSolid(const FTransform& ElemTM, const FVector& Scale3D, const FMaterialRenderProxy* MaterialRenderProxy, int32 ViewIndex, FMeshElementCollector& Collector) const
diff --git a/Engine/Source/Runtime/Experimental/Chaos/Private/Chaos/PerParticlePBDCollisionConstraint.ispc b/Engine/Source/Runtime/Experimental/Chaos/Private/Chaos/PerParticlePBDCollisionConstraint.ispc
--- a/Engine/Source/Runtime/Experimental/Chaos/Private/Chaos/PerParticlePBDCollisionConstraint.ispc
+++ b/Engine/Source/Runtime/Experimental/Chaos/Private/Chaos/PerParticlePBDCollisionConstraint.ispc
@@ -28,6 +28,7 @@ enum ImplicitObjectType
 	TaperedCapsule,
 	MLLevelSet,
 	SkinnedTriangleMesh,
+	ExtrudedTaperedCapsule,
 	WeightedLatticeLevelSetType = 38, //== LevelSet | IsWeightedLattice (cannot use expressions like this in ISPC)
 
 
@@ -177,6 +178,18 @@ struct FTaperedCylinder
 	// FAABB3 MLocalBoundingBox; // this code never uses this, and we never assume to know the sizeof(FTaperedCylinder)
 };
 
+struct FExtrudedTaperedCapsule
+{
+	FVector3f Origin;
+	FVector3f MajorAxis;
+	FVector3f MinorAxis;
+	float HalfLength;
+	float HalfWidth;
+	float Radius1;
+	float Radius2;
+	// FAABB3 LocalBoundingBox; // this code never uses this, and we never assume to know the sizeof(FExtrudedTaperedCapsule)
+};
+
 // We want to export all of the structs above, so defining exported method that uses them.
 export uniform int32 ExportHack_PerParticlePBDCollisionConstraintStructs(
 	uniform ImplicitObjectType Type,
@@ -195,7 +208,8 @@ export uniform int32 ExportHack_PerParticlePBDCollisionConstraintStructs(
 	uniform FConvexStructureDataIndexType IndexType,
 	uniform FConvex* uniform _FConvex,
 	uniform FTaperedCapsule* uniform _TaperedCapsule,
-	uniform FTaperedCylinder* uniform _TaperedCylinder
+	uniform FTaperedCylinder* uniform _TaperedCylinder,
+	uniform FExtrudedTaperedCapsule* uniform _ExtrudedTaperedCapsule
 )
 {
 	return 0;
@@ -307,6 +321,20 @@ static inline void PenetrationWithNormalTaperedCapsule(const uniform uint8* unif
 	}
 }
 
+static inline void PenetrationWithNormalExtrudedTaperedCapsule(const uniform uint8* uniform TypedObjectPtr, const FVector3f& PhiWithNormalInput, const uniform float Thickness, FVector3f& Normal, float& Penetration)
+{
+	const uniform FExtrudedTaperedCapsule CapsuleObj = *((const uniform FExtrudedTaperedCapsule *uniform)&TypedObjectPtr[0]);
+
+	const FVector3f OriginToX = PhiWithNormalInput - CapsuleObj.Origin;
+	const float DistanceAlongMajorAxis = clamp(VectorDot(OriginToX, CapsuleObj.MajorAxis), -CapsuleObj.HalfLength, CapsuleObj.HalfLength);
+	const float DistanceAlongMinorAxis = clamp(VectorDot(OriginToX, CapsuleObj.MinorAxis), -CapsuleObj.HalfWidth, CapsuleObj.HalfWidth);
+	const FVector3f ProjectedPoint = CapsuleObj.MajorAxis * DistanceAlongMajorAxis + CapsuleObj.MinorAxis * DistanceAlongMinorAxis + CapsuleObj.Origin;
+	const float Alpha = .5 * (DistanceAlongMajorAxis / CapsuleObj.HalfLength) + .5;
+	const float Radius = select(CapsuleObj.HalfLength > FLOAT_SMALL_NUMBER, CapsuleObj.Radius2 * ((uniform float)(1.0) - Alpha) + CapsuleObj.Radius1 * Alpha, max(CapsuleObj.Radius1, CapsuleObj.Radius2));
+	Normal = PhiWithNormalInput - ProjectedPoint;
+	Penetration = Thickness - (SafeNormalize(Normal) - Radius);
+}
+
 static inline void PenetrationWithNormalTaperedCylinder(const uniform uint8 *uniform TypedObjectPtr, const FVector3f &PhiWithNormalInput, const uniform float Thickness, FVector3f &Normal, float &Penetration, const uniform int SizeofFImplicitObject)
 {
 	const uniform FTaperedCylinder CylinderObj = *((const uniform FTaperedCylinder *uniform)&TypedObjectPtr[0]);
@@ -661,6 +689,12 @@ static inline bool GetSimplePenetrationWithNormal(const FVector3f& PhiWithNormal
 			PenetrationWithNormalConvex(TypedPtr, PhiWithNormalInputFloat, Thickness, Normal, Penetration);
 			return true;
 		}
+	case ExtrudedTaperedCapsule:
+		{
+			PenetrationWithNormalExtrudedTaperedCapsule(TypedPtr, PhiWithNormalInputFloat, Thickness, Normal, Penetration);
+			return true;
+		}
+
 	default:
 		return false;
 	}
diff --git a/Engine/Source/Runtime/Experimental/Chaos/Private/Chaos/PerParticlePBDCollisionConstraintISPCDataVerification.h b/Engine/Source/Runtime/Experimental/Chaos/Private/Chaos/PerParticlePBDCollisionConstraintISPCDataVerification.h
--- a/Engine/Source/Runtime/Experimental/Chaos/Private/Chaos/PerParticlePBDCollisionConstraintISPCDataVerification.h
+++ b/Engine/Source/Runtime/Experimental/Chaos/Private/Chaos/PerParticlePBDCollisionConstraintISPCDataVerification.h
@@ -5,6 +5,7 @@
 #include "Chaos/PBDSoftsEvolutionFwd.h"
 #include "Chaos/Capsule.h"
 #include "Chaos/Convex.h"
+#include "Chaos/ExtrudedTaperedCapsule.h"
 #include "Chaos/ImplicitObjectUnion.h"
 #include "Chaos/Sphere.h"
 #include "Chaos/TaperedCapsule.h"
@@ -57,6 +58,22 @@ static_assert(sizeof(Chaos::FImplicitObject) + offsetof(ispc::FTaperedCapsule, R
 static_assert(sizeof(ispc::FTaperedCapsule::Radius2) == Chaos::FTaperedCapsule::FISPCDataVerifier::SizeOfRadius2());
 static_assert(sizeof(Chaos::FImplicitObject) + offsetof(ispc::FTaperedCapsule, bIsOneSided) == Chaos::FTaperedCapsule::FISPCDataVerifier::OffsetOfIsOneSided());
 static_assert(sizeof(ispc::FTaperedCapsule::bIsOneSided) == Chaos::FTaperedCapsule::FISPCDataVerifier::SizeOfIsOneSided());
+// ExtrudedTaperedCapsule
+static_assert(sizeof(Chaos::FImplicitObject) + offsetof(ispc::FExtrudedTaperedCapsule, Origin) == Chaos::FExtrudedTaperedCapsule::FISPCDataVerifier::OffsetOfOrigin());
+static_assert(sizeof(ispc::FExtrudedTaperedCapsule::Origin) == Chaos::FExtrudedTaperedCapsule::FISPCDataVerifier::SizeOfOrigin());
+static_assert(sizeof(Chaos::FImplicitObject) + offsetof(ispc::FExtrudedTaperedCapsule, MajorAxis) == Chaos::FExtrudedTaperedCapsule::FISPCDataVerifier::OffsetOfMajorAxis());
+static_assert(sizeof(ispc::FExtrudedTaperedCapsule::MajorAxis) == Chaos::FExtrudedTaperedCapsule::FISPCDataVerifier::SizeOfMajorAxis());
+static_assert(sizeof(Chaos::FImplicitObject) + offsetof(ispc::FExtrudedTaperedCapsule, MinorAxis) == Chaos::FExtrudedTaperedCapsule::FISPCDataVerifier::OffsetOfMinorAxis());
+static_assert(sizeof(ispc::FExtrudedTaperedCapsule::MinorAxis) == Chaos::FExtrudedTaperedCapsule::FISPCDataVerifier::SizeOfMinorAxis());
+static_assert(sizeof(Chaos::FImplicitObject) + offsetof(ispc::FExtrudedTaperedCapsule, HalfLength) == Chaos::FExtrudedTaperedCapsule::FISPCDataVerifier::OffsetOfHalfLength());
+static_assert(sizeof(ispc::FExtrudedTaperedCapsule::HalfLength) == Chaos::FExtrudedTaperedCapsule::FISPCDataVerifier::SizeOfHalfLength());
+static_assert(sizeof(Chaos::FImplicitObject) + offsetof(ispc::FExtrudedTaperedCapsule, HalfWidth) == Chaos::FExtrudedTaperedCapsule::FISPCDataVerifier::OffsetOfHalfWidth());
+static_assert(sizeof(ispc::FExtrudedTaperedCapsule::HalfWidth) == Chaos::FExtrudedTaperedCapsule::FISPCDataVerifier::SizeOfHalfWidth());
+static_assert(sizeof(Chaos::FImplicitObject) + offsetof(ispc::FExtrudedTaperedCapsule, Radius1) == Chaos::FExtrudedTaperedCapsule::FISPCDataVerifier::OffsetOfRadius1());
+static_assert(sizeof(ispc::FExtrudedTaperedCapsule::Radius1) == Chaos::FExtrudedTaperedCapsule::FISPCDataVerifier::SizeOfRadius1());
+static_assert(sizeof(Chaos::FImplicitObject) + offsetof(ispc::FExtrudedTaperedCapsule, Radius2) == Chaos::FExtrudedTaperedCapsule::FISPCDataVerifier::OffsetOfRadius2());
+static_assert(sizeof(ispc::FExtrudedTaperedCapsule::Radius2) == Chaos::FExtrudedTaperedCapsule::FISPCDataVerifier::SizeOfRadius2());
+
 // Convex
 static_assert(sizeof(Chaos::FImplicitObject) + offsetof(ispc::FConvex, Planes) == Chaos::FConvex::FISPCDataVerifier::OffsetOfPlanes());
 static_assert(sizeof(ispc::FConvex::Planes) == Chaos::FConvex::FISPCDataVerifier::SizeOfPlanes());
diff --git a/Engine/Source/Runtime/Experimental/Chaos/Public/Chaos/ExtrudedTaperedCapsule.h b/Engine/Source/Runtime/Experimental/Chaos/Public/Chaos/ExtrudedTaperedCapsule.h
new file mode 100644
--- /dev/null
+++ b/Engine/Source/Runtime/Experimental/Chaos/Public/Chaos/ExtrudedTaperedCapsule.h
@@ -0,0 +1,143 @@
+// Copyright Epic Games, Inc. All Rights Reserved.
+#pragma once
+
+#include "Chaos/ImplicitObject.h"
+#include "Chaos/Core.h"
+#include "Math/UnrealMath.h"
+
+namespace Chaos
+{
+	class FExtrudedTaperedCapsule: public FImplicitObject
+	{
+	public:
+		FExtrudedTaperedCapsule()
+		    : FImplicitObject(EImplicitObject::FiniteConvex, ImplicitObjectType::ExtrudedTaperedCapsule)
+		{
+		}
+		FExtrudedTaperedCapsule(const FVec3f& InOrigin, const FVec3f& InMajorAxis, const FVec3f& InMinorAxis, const FRealSingle InHalfLength, const FRealSingle InHalfWidth, const FRealSingle InRadius1, const FRealSingle InRadius2)
+		    : FImplicitObject(EImplicitObject::FiniteConvex, ImplicitObjectType::ExtrudedTaperedCapsule)
+		    , Origin(InOrigin)
+			, MajorAxis(InMajorAxis.GetSafeNormal())
+			, MinorAxis((InMinorAxis -InMinorAxis.Dot(MajorAxis)*MajorAxis).GetSafeNormal())
+			, HalfLength(InHalfLength)
+			, HalfWidth(InHalfWidth)
+		    , Radius1(InRadius1)
+		    , Radius2(InRadius2)
+			, LocalBoundingBox(Origin,Origin)
+		{
+			LocalBoundingBox.GrowToInclude(Origin + MajorAxis * HalfLength + MinorAxis * HalfWidth);
+			LocalBoundingBox.GrowToInclude(Origin + MajorAxis * HalfLength - MinorAxis * HalfWidth);
+			LocalBoundingBox.GrowToInclude(Origin - MajorAxis * HalfLength + MinorAxis * HalfWidth);
+			LocalBoundingBox.GrowToInclude(Origin - MajorAxis * HalfLength - MinorAxis * HalfWidth);
+			const FRealSingle MaxRadius = FMath::Max(Radius1, Radius2);
+			LocalBoundingBox = FAABB3(LocalBoundingBox.Min() - FVec3(MaxRadius), LocalBoundingBox.Max() + FVec3(MaxRadius));
+		}
+		FExtrudedTaperedCapsule(const FExtrudedTaperedCapsule& Other)
+		    : FImplicitObject(EImplicitObject::FiniteConvex, ImplicitObjectType::ExtrudedTaperedCapsule)
+		    , Origin(Other.Origin)
+			, MajorAxis(Other.MajorAxis)
+			, MinorAxis(Other.MinorAxis)
+			, HalfLength(Other.HalfLength)
+			, HalfWidth(Other.HalfWidth)
+			, Radius1(Other.Radius1)
+			, Radius2(Other.Radius2)
+			, LocalBoundingBox(Other.LocalBoundingBox)
+		{
+		}
+		FExtrudedTaperedCapsule(FExtrudedTaperedCapsule&& Other)
+		    : FImplicitObject(EImplicitObject::FiniteConvex, ImplicitObjectType::ExtrudedTaperedCapsule)
+		    , Origin(MoveTemp(Other.Origin))
+			, MajorAxis(MoveTemp(Other.MajorAxis))
+			, MinorAxis(MoveTemp(Other.MinorAxis))
+			, HalfLength(Other.HalfLength)
+			, HalfWidth(Other.HalfWidth)
+			, Radius1(Other.Radius1)
+			, Radius2(Other.Radius2)
+			, LocalBoundingBox(MoveTemp(Other.LocalBoundingBox))
+		{
+		}
+
+		~FExtrudedTaperedCapsule() {}
+
+		static constexpr EImplicitObjectType StaticType() { return ImplicitObjectType::ExtrudedTaperedCapsule; }
+
+		virtual const FAABB3 BoundingBox() const override { return LocalBoundingBox; }
+
+		FReal PhiWithNormal(const FVec3& x, FVec3& OutNormal) const
+		{
+			const FVec3f OriginToX = x - Origin;
+			const FRealSingle DistanceAlongMajorAxis = FMath::Clamp(FVec3f::DotProduct(OriginToX, MajorAxis), -HalfLength, HalfLength);
+			const FRealSingle DistanceAlongMinorAxis = FMath::Clamp(FVec3f::DotProduct(OriginToX, MinorAxis), -HalfWidth, HalfWidth);
+			const FVec3f ClosestPoint = Origin + MajorAxis * DistanceAlongMajorAxis  + MinorAxis * DistanceAlongMinorAxis;
+			const FRealSingle Radius = (HalfLength > UE_SMALL_NUMBER) ? FMath::Lerp(Radius2, Radius1, .5f*(DistanceAlongMajorAxis / HalfLength) + .5f) : FMath::Max(Radius1, Radius2);
+			OutNormal = (x - ClosestPoint);
+			return OutNormal.SafeNormalize() - Radius;
+		}
+
+		/* Radius at +MajorAxis end*/
+		FRealSingle GetRadius1() const { return Radius1; }
+		/* Radius at -MajorAxis end*/
+		FRealSingle GetRadius2() const { return Radius2; }
+		FRealSingle GetHalfLength() const { return HalfLength; }
+		FRealSingle GetHalfWidth() const { return HalfWidth; }
+		/** Returns the center of the rectangle */
+		const FVec3f& GetOrigin() const 
+		{
+			return Origin;
+		}
+		const FVec3f& GetOriginf() const { return Origin; }
+		/** Length axis */
+		const FVec3f& GetMajorAxis() const { return MajorAxis; }
+		/** Width axis */
+		const FVec3f& GetMinorAxis() const { return MinorAxis; }
+
+		virtual uint32 GetTypeHash() const override
+		{
+			uint32 Result = UE::Math::GetTypeHash(Origin);
+			Result = HashCombine(Result, UE::Math::GetTypeHash(MajorAxis));
+			Result = HashCombine(Result, UE::Math::GetTypeHash(MinorAxis));
+			Result = HashCombine(Result, ::GetTypeHash(HalfLength));
+			Result = HashCombine(Result, ::GetTypeHash(HalfWidth));
+			Result = HashCombine(Result, ::GetTypeHash(Radius1));
+			Result = HashCombine(Result, ::GetTypeHash(Radius2));
+			return Result;
+		}
+
+#if INTEL_ISPC
+		// See PerParticlePBDCollisionConstraint.cpp
+		// ISPC code has matching structs for interpreting FImplicitObjects.
+		// This is used to verify that the structs stay the same.
+		struct FISPCDataVerifier
+		{
+			static constexpr int32 OffsetOfOrigin() { return offsetof(FExtrudedTaperedCapsule, Origin); }
+			static constexpr int32 SizeOfOrigin() { return sizeof(FExtrudedTaperedCapsule::Origin); }
+			static constexpr int32 OffsetOfMajorAxis() { return offsetof(FExtrudedTaperedCapsule, MajorAxis); }
+			static constexpr int32 SizeOfMajorAxis() { return sizeof(FExtrudedTaperedCapsule::MajorAxis); }
+			static constexpr int32 OffsetOfMinorAxis() { return offsetof(FExtrudedTaperedCapsule, MinorAxis); }
+			static constexpr int32 SizeOfMinorAxis() { return sizeof(FExtrudedTaperedCapsule::MinorAxis); }
+			static constexpr int32 OffsetOfHalfLength() { return offsetof(FExtrudedTaperedCapsule, HalfLength); }
+			static constexpr int32 SizeOfHalfLength() { return sizeof(FExtrudedTaperedCapsule::HalfLength); }
+			static constexpr int32 OffsetOfHalfWidth() { return offsetof(FExtrudedTaperedCapsule, HalfWidth); }
+			static constexpr int32 SizeOfHalfWidth() { return sizeof(FExtrudedTaperedCapsule::HalfWidth); }
+			static constexpr int32 OffsetOfRadius1() { return offsetof(FExtrudedTaperedCapsule, Radius1); }
+			static constexpr int32 SizeOfRadius1() { return sizeof(FExtrudedTaperedCapsule::Radius1); }
+			static constexpr int32 OffsetOfRadius2() { return offsetof(FExtrudedTaperedCapsule, Radius2); }
+			static constexpr int32 SizeOfRadius2() { return sizeof(FExtrudedTaperedCapsule::Radius2); }
+		};
+		friend FISPCDataVerifier;
+#endif // #if INTEL_ISPC
+
+	private:
+
+		FVec3f Origin;
+		FVec3f MajorAxis;
+		FVec3f MinorAxis;
+		FRealSingle HalfLength;
+		FRealSingle HalfWidth;
+		FRealSingle Radius1;
+		FRealSingle Radius2;
+		FAABB3 LocalBoundingBox;
+	};
+
+
+} // namespace Chaos
diff --git a/Engine/Source/Runtime/Experimental/Chaos/Public/Chaos/ImplicitObjectType.h b/Engine/Source/Runtime/Experimental/Chaos/Public/Chaos/ImplicitObjectType.h
--- a/Engine/Source/Runtime/Experimental/Chaos/Public/Chaos/ImplicitObjectType.h
+++ b/Engine/Source/Runtime/Experimental/Chaos/Public/Chaos/ImplicitObjectType.h
@@ -29,6 +29,7 @@ namespace Chaos
 			TaperedCapsule,
 			MLLevelSet,
 			SkinnedTriangleMesh,
+			ExtrudedTaperedCapsule,
 			//Add entries above this line for serialization
 			ConcreteObjectCount, // Used to ensure bitflags do not overlap concrete type
 			IsWeightedLattice = 1 << 5,
@@ -96,6 +97,7 @@ namespace Chaos
 				TEXT("EImplicitObjectType:TaperedCapsule"),
 				TEXT("EImplicitObjectType:MLLevelSet"),
 				TEXT("EImplicitObjectType:SkinnedTriangleMesh"),
+				TEXT("EImplicitObjectType::ExtrudedTaperedCapsule")
 			};
 
 		static_assert(Chaos::ImplicitObjectType::ConcreteObjectCount == UE_ARRAY_COUNT(Names), "");
