Add missing 'override' to fix warning V1052: Declaring virtual methods in a class marked as 'final' is pointless. Consider inspecting the 'ThreadedPostLoadAssetRegistryTagsOverride' method of the 'UWorld' class.
#rnx

diff --git a/Engine/Source/Runtime/Engine/Classes/Engine/World.h b/Engine/Source/Runtime/Engine/Classes/Engine/World.h
--- a/Engine/Source/Runtime/Engine/Classes/Engine/World.h
+++ b/Engine/Source/Runtime/Engine/Classes/Engine/World.h
@@ -3053,7 +3053,7 @@ public:
 	UE_API virtual void GetAssetRegistryTagMetadata(TMap<FName, FAssetRegistryTagMetadata>& OutMetadata) const override;
 	UE_API virtual void GetExtendedAssetRegistryTagsForSave(const ITargetPlatform* TargetPlatform, TArray<FAssetRegistryTag>& OutTags) const override;
 protected:
-	UE_API virtual void ThreadedPostLoadAssetRegistryTagsOverride(FPostLoadAssetRegistryTagsContext& Context) const;
+	UE_API virtual void ThreadedPostLoadAssetRegistryTagsOverride(FPostLoadAssetRegistryTagsContext& Context) const override;
 public:
 	UE_API virtual bool IsNameStableForNetworking() const override;
 #endif
