Add more information to assert to narrow down root cause

diff --git a/Engine/Source/Runtime/Engine/Private/SkinnedAsset.cpp b/Engine/Source/Runtime/Engine/Private/SkinnedAsset.cpp
--- a/Engine/Source/Runtime/Engine/Private/SkinnedAsset.cpp
+++ b/Engine/Source/Runtime/Engine/Private/SkinnedAsset.cpp
@@ -459,8 +459,10 @@ void USkinnedAsset::FillComponentSpaceTransforms(const TArray<FTransform>& InBon
 	ANIM_MT_SCOPE_CYCLE_COUNTER(FillComponentSpaceTransforms, !IsInGameThread());
 
 	// right now all this does is populate DestSpaceBases
-	check(GetRefSkeleton().GetNum() == InBoneSpaceTransforms.Num());
-	check(GetRefSkeleton().GetNum() == OutComponentSpaceTransforms.Num());
+	checkf(GetRefSkeleton().GetNum() == InBoneSpaceTransforms.Num(),
+		TEXT("RefSkeleton bone count (%d) and input bone transforms array count (%d) mismatches. Skinned Asset: [%s]"), GetRefSkeleton().GetNum(), InBoneSpaceTransforms.Num(), *GetPathName());
+	checkf(GetRefSkeleton().GetNum() == OutComponentSpaceTransforms.Num(),
+		TEXT("RefSkeleton bone count (%d) and output bone transforms array count (%d) mismatches. Skinned Asset: [%s]"), GetRefSkeleton().GetNum(), OutComponentSpaceTransforms.Num(), *GetPathName());
 
 	const int32 NumBones = InBoneSpaceTransforms.Num();
 	
