Fix assert in FRayTracingGeometryManager constructor when launching with the csv profiler recording.

IsRayTracingEnabled() requires the RHI to be initialized and will assert if it's not. Moved setting this to PostEngineInit instead.

#rnx

diff --git a/Engine/Source/Runtime/Engine/Private/Rendering/RayTracingGeometryManager.cpp b/Engine/Source/Runtime/Engine/Private/Rendering/RayTracingGeometryManager.cpp
--- a/Engine/Source/Runtime/Engine/Private/Rendering/RayTracingGeometryManager.cpp
+++ b/Engine/Source/Runtime/Engine/Private/Rendering/RayTracingGeometryManager.cpp
@@ -171,13 +171,11 @@ FRayTracingGeometryManager::FRayTracingGeometryManager()
 	StreamingRequests.SetNum(GRayTracingStreamingMaxPendingRequests);
 
 #if CSV_PROFILER_STATS
-	if (FCsvProfiler* CSVProfiler = FCsvProfiler::Get())
+	// Need to defer until RHI is initialized.
+	FCoreDelegates::OnPostEngineInit.AddLambda([]
 	{
-		CSVProfiler->OnCSVProfileStart().AddLambda([]()
-		{
-			CSV_METADATA(TEXT("RayTracing"), IsRayTracingEnabled() ? TEXT("1") : TEXT("0"));
-		});
-	}
+		CSV_METADATA(TEXT("RayTracing"), IsRayTracingEnabled() ? TEXT("1") : TEXT("0"));
+	});
 #endif
 }
 
