SDCE, fixed cvar definition

diff --git a/Engine/Source/Runtime/Engine/Private/Materials/HLSLMaterialTranslator.cpp b/Engine/Source/Runtime/Engine/Private/Materials/HLSLMaterialTranslator.cpp
--- a/Engine/Source/Runtime/Engine/Private/Materials/HLSLMaterialTranslator.cpp
+++ b/Engine/Source/Runtime/Engine/Private/Materials/HLSLMaterialTranslator.cpp
@@ -78,6 +78,11 @@
 
 #include <functional>
 
+TAutoConsoleVariable<bool> CVarShadersSDCEEnabled(
+	TEXT("r.Shaders.SDCE"),
+	true,
+	TEXT("Strip dead code using symbolic dead code elimination on structured member stores"));
+
 #if ENABLE_COOK_STATS && STATS
 
 static bool GEnableMaterialTranslationLogFile = false;
@@ -86,7 +91,6 @@ static FAutoConsoleVariableRef CVarEnableMaterialTranslationLogFile(
 	GEnableMaterialTranslationLogFile,
 	TEXT("Enables material translation log file generation for tracking materials translation time (written to  \"ShaderDebugInfo/MaterialTranslationLog-X.csv\")."),
 	ECVF_ReadOnly);
-
 /**
  * Utility used to create a MaterialTranslationLog.txt file during cooks that contains the list of all translated materials
  * and other info such as how long the translation took.
