Allow disabling world bounds checking with a cvar.

diff --git a/Engine/Source/Runtime/Engine/Classes/GameFramework/WorldSettings.h b/Engine/Source/Runtime/Engine/Classes/GameFramework/WorldSettings.h
--- a/Engine/Source/Runtime/Engine/Classes/GameFramework/WorldSettings.h
+++ b/Engine/Source/Runtime/Engine/Classes/GameFramework/WorldSettings.h
@@ -866,7 +866,7 @@ public:
 	bool IsAISystemEnabled() const { return bEnableAISystem; }
 
 	/** @return whether given world is restricting actors to +-HALF_WORLD_MAX bounds, destroying actors that move below KillZ */
-	bool AreWorldBoundsChecksEnabled() const { return bEnableWorldBoundsChecks; }
+	ENGINE_API bool AreWorldBoundsChecksEnabled() const;
 	
 	/**
 	 * Called from GameStateBase, calls BeginPlay on all actors
diff --git a/Engine/Source/Runtime/Engine/Private/WorldSettings.cpp b/Engine/Source/Runtime/Engine/Private/WorldSettings.cpp
--- a/Engine/Source/Runtime/Engine/Private/WorldSettings.cpp
+++ b/Engine/Source/Runtime/Engine/Private/WorldSettings.cpp
@@ -50,6 +50,13 @@ DEFINE_LOG_CATEGORY_STATIC(LogWorldSettings, Log, All);
 // input is polled for motion controller devices each frame.
 ENGINE_API float GNewWorldToMetersScale = 0.0f;
 
+/** Allows globally disabling world bound checks. */
+static int32 GDisableWorldBoundsChecks = 0;
+static FAutoConsoleVariableRef CVarActorWorldBoundsOverride(
+	TEXT("WorldBounds.DisableChecks"),
+	GDisableWorldBoundsChecks,
+	TEXT("If set to 1, will ignore all worlds bounds checks."));
+
 #if WITH_EDITOR
 AWorldSettings::FOnBookmarkClassChanged AWorldSettings::OnBookmarkClassChanged;
 AWorldSettings::FOnNumberOfBookmarksChanged AWorldSettings::OnNumberOfBoomarksChanged;
@@ -702,6 +709,10 @@ void AWorldSettings::SetNavigationSystemConfigOverride(UNavigationSystemConfig*
 	}
 }
 
+bool AWorldSettings::AreWorldBoundsChecksEnabled() const
+{
+	return bEnableWorldBoundsChecks && !GDisableWorldBoundsChecks;
+}
 
 #if WITH_EDITOR
 
