Entity Streaming: entities in generated cooked packages now have their outer set in their parent package, removing the need to call Rename in the game thread when loading and unloading streaming cell packages.

#preflighbt

diff --git a/Engine/Source/Runtime/Engine/Private/WorldCookPackageSplitter.cpp b/Engine/Source/Runtime/Engine/Private/WorldCookPackageSplitter.cpp
--- a/Engine/Source/Runtime/Engine/Private/WorldCookPackageSplitter.cpp
+++ b/Engine/Source/Runtime/Engine/Private/WorldCookPackageSplitter.cpp
@@ -73,6 +73,26 @@ bool FWorldCookPackageSplitter::RequiresGeneratorPackageDestructBeforeResplit()
 	return true;
 }
 
+bool FWorldCookPackageSplitter::RequiresEmptyPackageBeforePopulate(const UPackage* Package)
+{
+	TOptional<bool> Result;
+
+	ForEachCookPackageSubSplitters(
+		[Package, &Result](FWorldCookPackageSplitter::ISubSplitter* WorldCookPackageSplitter)
+		{
+			TOptional<bool> SubResult = WorldCookPackageSplitter->RequiresEmptyPackageBeforePopulate(Package);
+			check(!SubResult.IsSet() || !Result.IsSet());
+			Result = SubResult;
+		});
+
+	if (Result.IsSet())
+	{
+		return *Result;
+	}
+
+	return ICookPackageSplitter::RequiresEmptyPackageBeforePopulate(Package);
+}
+
 ICookPackageSplitter::FGenerationManifest FWorldCookPackageSplitter::ReportGenerationManifest(
 	const UPackage* OwnerPackage, const UObject* OwnerObject)
 {
diff --git a/Engine/Source/Runtime/Engine/Private/WorldPartition/Cook/WorldPartitionCookPackageContext.h b/Engine/Source/Runtime/Engine/Private/WorldPartition/Cook/WorldPartitionCookPackageContext.h
--- a/Engine/Source/Runtime/Engine/Private/WorldPartition/Cook/WorldPartitionCookPackageContext.h
+++ b/Engine/Source/Runtime/Engine/Private/WorldPartition/Cook/WorldPartitionCookPackageContext.h
@@ -39,6 +39,16 @@ public:
 	virtual void ReportSaveDependency(UE::Cook::FCookDependency CookDependency) const override;
 	ICookPackageSplitter::FPopulateContext* PopulateContext = nullptr;
 
+	inline void RegisterGeneratedPackage(const FString& GeneratedPackageName) const
+	{
+		CookGeneratedPackages.Add(GeneratedPackageName);
+	}
+
+	inline bool IsGeneratingPackage(const UPackage* Package) const
+	{
+		return CookGeneratedPackages.Contains(Package->GetName());
+	}
+
 private:
 	const FWorldPartitionCookPackage* AddLevelStreamingPackageToGenerate(IWorldPartitionCookPackageGenerator* CookPackageGenerator, const FString& Root, const FString& RelativePath);
 	const FWorldPartitionCookPackage* AddGenericPackageToGenerate(IWorldPartitionCookPackageGenerator* CookPackageGenerator, const FString& Root, const FString& RelativePath);
@@ -49,6 +59,7 @@ private:
 	TMap<FWorldPartitionCookPackage::IDType, IWorldPartitionCookPackageGenerator*> CookGeneratorByPackageId;
 	TMap<IWorldPartitionCookPackageGenerator*, TArray<FWorldPartitionCookPackage*>> PackagesToCookByGenerator;
 	TMap<FObjectKey, FWorldPartitionCookPackage::IDType> PackageObjectToPackageId;
+	mutable TSet<FString> CookGeneratedPackages;
 	FWorldPartitionCookPackageContextParams Params;
 };
 
diff --git a/Engine/Source/Runtime/Engine/Private/WorldPartition/WorldPartitionCookPackageSplitter.cpp b/Engine/Source/Runtime/Engine/Private/WorldPartition/WorldPartitionCookPackageSplitter.cpp
--- a/Engine/Source/Runtime/Engine/Private/WorldPartition/WorldPartitionCookPackageSplitter.cpp
+++ b/Engine/Source/Runtime/Engine/Private/WorldPartition/WorldPartitionCookPackageSplitter.cpp
@@ -42,6 +42,15 @@ TArray<ICookPackageSplitter::FGeneratedPackage> FWorldPartitionCookPackageSplitt
 	return PackagesToGenerate;
 }
 
+TOptional<bool> FWorldPartitionCookPackageSplitter::RequiresEmptyPackageBeforePopulate(const UPackage* Package)
+{
+	if (CookContext.IsGeneratingPackage(Package))
+	{
+		return true;
+	}
+	return TOptional<bool>();
+}
+
 bool FWorldPartitionCookPackageSplitter::PopulateGeneratedPackage(ICookPackageSplitter::FPopulateContext& PopulateContext)
 {
 	const ICookPackageSplitter::FGeneratedPackageForPopulate& GeneratedPackage = *PopulateContext.GetTargetGeneratedPackage();
@@ -125,6 +134,7 @@ void FWorldPartitionCookPackageSplitter::BuildPackagesToGenerateList(TArray<ICoo
 				GeneratedPackage.RelativePath = CookPackage->RelativePath;
 				GeneratedPackage.GenerationHash = CookPackage->GenerationHash;
 
+				CookContext.RegisterGeneratedPackage(CookPackage->GetFullGeneratedPath());
 				UE_LOG(LogWorldPartition, Log, TEXT("Adding Cell %s with GenerationHash: %s to packages to generate"), *CookPackage->GetFullGeneratedPath(), *LexToString(CookPackage->GenerationHash))
 
 				GeneratedPackage.SetCreateAsMap(CookPackage->Type == FWorldPartitionCookPackage::EType::Level);
diff --git a/Engine/Source/Runtime/Engine/Private/WorldPartition/WorldPartitionCookPackageSplitter.h b/Engine/Source/Runtime/Engine/Private/WorldPartition/WorldPartitionCookPackageSplitter.h
--- a/Engine/Source/Runtime/Engine/Private/WorldPartition/WorldPartitionCookPackageSplitter.h
+++ b/Engine/Source/Runtime/Engine/Private/WorldPartition/WorldPartitionCookPackageSplitter.h
@@ -12,6 +12,7 @@ class FWorldPartitionCookPackageSplitter : public FWorldCookPackageSplitter::ISu
 public:
 	//~Begin FWorldCookPackageSplitter::ISubSplitter interface
 	virtual TArray<ICookPackageSplitter::FGeneratedPackage> GetGenerateList(const UPackage* OwnerPackage) override;
+	virtual TOptional<bool> RequiresEmptyPackageBeforePopulate(const UPackage* Package) override;
 	virtual bool PopulateGeneratedPackage(ICookPackageSplitter::FPopulateContext& PopulateContext) override;
 	virtual bool PopulateGeneratorPackage(ICookPackageSplitter::FPopulateContext& PopulateContext) override;
 	virtual void Teardown(ICookPackageSplitter::ETeardown Status) override;
diff --git a/Engine/Source/Runtime/Engine/Public/WorldCookPackageSplitter.h b/Engine/Source/Runtime/Engine/Public/WorldCookPackageSplitter.h
--- a/Engine/Source/Runtime/Engine/Public/WorldCookPackageSplitter.h
+++ b/Engine/Source/Runtime/Engine/Public/WorldCookPackageSplitter.h
@@ -23,6 +23,7 @@ public:
 	{
 		virtual ~ISubSplitter()	{}
 		virtual TArray<ICookPackageSplitter::FGeneratedPackage> GetGenerateList(const UPackage* OwnerPackage) = 0;
+		virtual TOptional<bool> RequiresEmptyPackageBeforePopulate(const UPackage* Package) = 0;
 		virtual bool PopulateGeneratedPackage(ICookPackageSplitter::FPopulateContext& PopulateContext) = 0;
 		virtual bool PopulateGeneratorPackage(ICookPackageSplitter::FPopulateContext& PopulateContext) = 0;
 		virtual void Teardown(ICookPackageSplitter::ETeardown Status) = 0;
@@ -67,6 +68,7 @@ protected:
 	virtual bool UseInternalReferenceToAvoidGarbageCollect() override;
 	virtual EGeneratedRequiresGenerator DoesGeneratedRequireGenerator() override;
 	virtual bool RequiresGeneratorPackageDestructBeforeResplit() override;
+	virtual bool RequiresEmptyPackageBeforePopulate(const UPackage* Package) override;
 	virtual ICookPackageSplitter::FGenerationManifest ReportGenerationManifest(
 		const UPackage* OwnerPackage, const UObject* OwnerObject) override;
 	virtual bool PopulateGeneratedPackage(FPopulateContext& PopulateContext) override;
