Fixed potential mod by zero caught by static analysis

#rnx

diff --git a/Engine/Source/Runtime/Engine/Private/Components/SceneCaptureComponent.cpp b/Engine/Source/Runtime/Engine/Private/Components/SceneCaptureComponent.cpp
--- a/Engine/Source/Runtime/Engine/Private/Components/SceneCaptureComponent.cpp
+++ b/Engine/Source/Runtime/Engine/Private/Components/SceneCaptureComponent.cpp
@@ -744,7 +744,7 @@ void USceneCaptureComponent2D::TickComponent(float DeltaTime, enum ELevelTick Ti
 		return;
 	}
 
-	if (bCaptureEveryFrame)
+	if (bCaptureEveryFrame && NumTiles > 0)
 	{
 		TileID++;
 		TileID %= NumTiles;
