Fix PSO precaching misses
- CustomDepth for Niagara when enabled on the system
- Default overlay material is always precached when set

diff --git a/Engine/Plugins/FX/Niagara/Source/Niagara/Private/NiagaraSystem.cpp b/Engine/Plugins/FX/Niagara/Source/Niagara/Private/NiagaraSystem.cpp
--- a/Engine/Plugins/FX/Niagara/Source/Niagara/Private/NiagaraSystem.cpp
+++ b/Engine/Plugins/FX/Niagara/Source/Niagara/Private/NiagaraSystem.cpp
@@ -1656,6 +1656,14 @@ void UNiagaraSystem::PrecachePSOs()
 		);
 	}
 
+	if (bRenderCustomDepth)
+	{
+		for (FMaterialInterfacePSOPrecacheParams& PrecacheParams : MaterialInterfacePSOPrecacheParamsList)
+		{
+			PrecacheParams.PSOPrecacheParams.bRenderCustomDepth = bRenderCustomDepth;
+		}
+	}
+
 	LaunchPSOPrecaching(MaterialInterfacePSOPrecacheParamsList);
 }
 
diff --git a/Engine/Source/Runtime/Engine/Private/Components/SkinnedMeshComponent.cpp b/Engine/Source/Runtime/Engine/Private/Components/SkinnedMeshComponent.cpp
--- a/Engine/Source/Runtime/Engine/Private/Components/SkinnedMeshComponent.cpp
+++ b/Engine/Source/Runtime/Engine/Private/Components/SkinnedMeshComponent.cpp
@@ -756,11 +756,11 @@ void USkinnedMeshComponent::CollectPSOPrecacheData(const FPSOPrecacheParams& Bas
 	TArray<UMaterialInterface*> AllOverlayMaterials;
 	TArray<TObjectPtr<UMaterialInterface>> AssetAndComponentMaterialSlotsOverlayMaterial;
 	GetMaterialSlotsOverlayMaterial(AssetAndComponentMaterialSlotsOverlayMaterial);
-	bool bAllMaterialSlotOverride = true;
-	FMaterialOverlayHelper::AppendAllOverlayMaterial(AssetAndComponentMaterialSlotsOverlayMaterial, AllOverlayMaterials, bAllMaterialSlotOverride);
+	bool bHaveNullEntry = true;
+	FMaterialOverlayHelper::AppendAllOverlayMaterial(AssetAndComponentMaterialSlotsOverlayMaterial, AllOverlayMaterials, bHaveNullEntry);
 
-	//If all material slot set an overlay material we wont use the global one
-	if (!bAllMaterialSlotOverride)
+	// always add default overlay material if available
+	if (GetOverlayMaterial())
 	{
 		AllOverlayMaterials.AddUnique(GetOverlayMaterial());
 	}
