Fix issues with toggling replicated movement and UCharacterMovementComponent activation in ACharacter at runtime
- Replicate UCharacterMovementComponent's IsActive flag so that simulated proxies stop ticking when deactivated on the server
- Refactor code that controls bOnlyAllowAutonomousTickPose to handle these variables getting toggled

#tests PIE - Lyra including 2 clients in Listen Server mode
#rnx

diff --git a/Engine/Source/Runtime/Engine/Classes/GameFramework/Character.h b/Engine/Source/Runtime/Engine/Classes/GameFramework/Character.h
--- a/Engine/Source/Runtime/Engine/Classes/GameFramework/Character.h
+++ b/Engine/Source/Runtime/Engine/Classes/GameFramework/Character.h
@@ -588,6 +588,15 @@ public:
 	uint8 bWasJumping : 1;
 
 protected:
+	/** CharacterMovement bIsActive replicated for simulated proxies. */
+	UPROPERTY(Transient, ReplicatedUsing=OnRep_IsCharacterMovementActive)
+	uint8 bProxyIsCharacterMovementActive : 1;
+
+	ENGINE_API void SetProxyCharacterMovementActive(bool bInActive);
+
+	UFUNCTION()
+	ENGINE_API virtual void OnRep_IsCharacterMovementActive();
+
 	/** CharacterMovement MovementMode (and custom mode) replicated for simulated proxies. Use CharacterMovementComponent::UnpackNetworkMovementMode() to translate it. */
 	UPROPERTY(Replicated)
 	uint8 ReplicatedMovementMode;
@@ -1092,4 +1101,7 @@ public:
 	ENGINE_API void FillAsyncInput(FCharacterAsyncInput& Input) const;
 	ENGINE_API void InitializeAsyncOutput(FCharacterAsyncOutput& Output) const;
 	ENGINE_API void ApplyAsyncOutput(const FCharacterAsyncOutput& Output);
+
+	/** Sets bOnlyAllowAutonomousTickPose based on the character's current state */
+	ENGINE_API void UpdateMeshAllowAutonomousTickPose();
 };
diff --git a/Engine/Source/Runtime/Engine/Classes/GameFramework/CharacterMovementComponent.h b/Engine/Source/Runtime/Engine/Classes/GameFramework/CharacterMovementComponent.h
--- a/Engine/Source/Runtime/Engine/Classes/GameFramework/CharacterMovementComponent.h
+++ b/Engine/Source/Runtime/Engine/Classes/GameFramework/CharacterMovementComponent.h
@@ -1295,6 +1295,7 @@ public:
 	ENGINE_API virtual void BeginDestroy() override;
 	ENGINE_API virtual void BeginPlay() override;
 	ENGINE_API virtual void PostLoad() override;
+	ENGINE_API virtual void Activate(bool bReset = false) override;
 	ENGINE_API virtual void Deactivate() override;
 	ENGINE_API virtual void RegisterComponentTickFunctions(bool bRegister) override;
 	ENGINE_API virtual void ApplyWorldOffset(const FVector& InOffset, bool bWorldShift) override;
diff --git a/Engine/Source/Runtime/Engine/Private/Character.cpp b/Engine/Source/Runtime/Engine/Private/Character.cpp
--- a/Engine/Source/Runtime/Engine/Private/Character.cpp
+++ b/Engine/Source/Runtime/Engine/Private/Character.cpp
@@ -415,6 +415,20 @@ void ACharacter::OnRep_IsCrouched()
 	}
 }
 
+void ACharacter::SetProxyCharacterMovementActive(bool bInActive)
+{
+	bProxyIsCharacterMovementActive = bInActive;
+}
+
+void ACharacter::OnRep_IsCharacterMovementActive()
+{
+	if (CharacterMovement)
+	{
+		CharacterMovement->SetActiveFlag(bProxyIsCharacterMovementActive);
+		CharacterMovement->OnRep_IsActive();
+	}
+}
+
 void ACharacter::SetReplicateMovement(bool bInReplicateMovement)
 {
 	Super::SetReplicateMovement(bInReplicateMovement);
@@ -430,6 +444,8 @@ void ACharacter::SetReplicateMovement(bool bInReplicateMovement)
 			NetworkPrediction->ServerTimeStamp = GetWorld()->GetTimeSeconds();
 		}
 	}
+
+	UpdateMeshAllowAutonomousTickPose();
 }
 
 bool ACharacter::CanCrouch() const
@@ -993,11 +1009,7 @@ void ACharacter::PossessedBy(AController* NewController)
 {
 	Super::PossessedBy(NewController);
 
-	// If we are controlled remotely, set animation timing to be driven by client's network updates. So timing and events remain in sync.
-	if (Mesh && IsReplicatingMovement() && (GetRemoteRole() == ROLE_AutonomousProxy && GetNetConnection() != nullptr))
-	{
-		Mesh->bOnlyAllowAutonomousTickPose = true;
-	}
+	UpdateMeshAllowAutonomousTickPose();
 }
 
 void ACharacter::UnPossessed()
@@ -1011,10 +1023,7 @@ void ACharacter::UnPossessed()
 	}
 
 	// We're no longer controlled remotely, resume regular ticking of animations.
-	if (Mesh)
-	{
-		Mesh->bOnlyAllowAutonomousTickPose = false;
-	}
+	UpdateMeshAllowAutonomousTickPose();
 }
 
 
@@ -1029,10 +1038,7 @@ void ACharacter::TornOff()
 	}
 
 	// We're no longer controlled remotely, resume regular ticking of animations.
-	if (Mesh)
-	{
-		Mesh->bOnlyAllowAutonomousTickPose = false;
-	}
+	UpdateMeshAllowAutonomousTickPose();
 }
 
 
@@ -1653,6 +1659,7 @@ void ACharacter::PreReplication( IRepChangedPropertyTracker & ChangedPropertyTra
 	}
 
 	SetProxyIsJumpForceApplied(JumpForceTimeRemaining > 0.0f);
+	SetProxyCharacterMovementActive(CharacterMovement->IsActive());
 	SetReplicatedMovementMode(CharacterMovement->PackNetworkMovementMode());
 	SetReplicatedGravityDirection(CharacterMovement->GetGravityDirection());
 
@@ -1725,6 +1732,7 @@ void ACharacter::GetLifetimeReplicatedProps( TArray< FLifetimeProperty > & OutLi
 	DOREPLIFETIME_CONDITION( ACharacter, ReplicatedServerLastTransformUpdateTimeStamp, COND_SimulatedOnlyNoReplay );
 	DOREPLIFETIME_CONDITION( ACharacter, ReplicatedMovementMode,			COND_SimulatedOnly );
 	DOREPLIFETIME_CONDITION( ACharacter, bIsCrouched,						COND_SimulatedOnly );
+	DOREPLIFETIME_CONDITION( ACharacter, bProxyIsCharacterMovementActive,	COND_SimulatedOnly );
 	DOREPLIFETIME_CONDITION( ACharacter, bProxyIsJumpForceApplied,			COND_SimulatedOnly );
 	DOREPLIFETIME_CONDITION( ACharacter, AnimRootMotionTranslationScale,	COND_SimulatedOnly );
 	DOREPLIFETIME_CONDITION( ACharacter, ReplicatedGravityDirection,		COND_SimulatedOnly );
@@ -2013,6 +2021,18 @@ void ACharacter::ApplyAsyncOutput(const FCharacterAsyncOutput& Output)
 	}
 }
 
+void ACharacter::UpdateMeshAllowAutonomousTickPose()
+{
+	if (Mesh)
+	{
+		// If we are controlled remotely and CharacterMovement is active, set animation timing to be driven by client's network updates. So timing and events remain in sync.
+		const bool bCharacterMovementWillTickPose = IsReplicatingMovement() && (CharacterMovement && CharacterMovement->IsActive());
+		const bool bRemotelyControlled = (GetRemoteRole() == ROLE_AutonomousProxy && GetNetConnection() != nullptr);
+
+		Mesh->bOnlyAllowAutonomousTickPose = bCharacterMovementWillTickPose && bRemotelyControlled;
+	}
+}
+
 UAnimMontage* FRepRootMotionMontage::GetAnimMontage() const
 {
 	return Cast<UAnimMontage>(Animation);
diff --git a/Engine/Source/Runtime/Engine/Private/Components/CharacterMovementComponent.cpp b/Engine/Source/Runtime/Engine/Private/Components/CharacterMovementComponent.cpp
--- a/Engine/Source/Runtime/Engine/Private/Components/CharacterMovementComponent.cpp
+++ b/Engine/Source/Runtime/Engine/Private/Components/CharacterMovementComponent.cpp
@@ -944,6 +944,19 @@ void UCharacterMovementComponent::BeginDestroy()
 	Super::BeginDestroy();
 }
 
+void UCharacterMovementComponent::Activate(bool bReset /*= false*/)
+{
+	Super::Activate(bReset);
+
+	if (IsActive())
+	{
+		if (CharacterOwner)
+		{
+			CharacterOwner->UpdateMeshAllowAutonomousTickPose();
+		}
+	}
+}
+
 void UCharacterMovementComponent::Deactivate()
 {
 	NavMovementProperties.bStopMovementAbortPaths = false; // Mirrors StopMovementKeepPathing(), because Super calls StopMovement() and we want that handled differently.
@@ -954,6 +967,7 @@ void UCharacterMovementComponent::Deactivate()
 		if (CharacterOwner)
 		{
 			CharacterOwner->ResetJumpState();
+			CharacterOwner->UpdateMeshAllowAutonomousTickPose();
 		}
 	}
 	NavMovementProperties.bStopMovementAbortPaths = true;
