[Engine]
*  Force the decay to function-pointer for default parameter. This fixes linker error when building with header units

diff --git a/Engine/Source/Runtime/Engine/Public/ContentStreaming.h b/Engine/Source/Runtime/Engine/Public/ContentStreaming.h
--- a/Engine/Source/Runtime/Engine/Public/ContentStreaming.h
+++ b/Engine/Source/Runtime/Engine/Public/ContentStreaming.h
@@ -483,7 +483,7 @@ struct IRenderAssetStreamingManager : public IStreamingManager
 	virtual void GetAssetComponents(
 		const UStreamableRenderAsset* RenderAsset,
 		TArray<const UPrimitiveComponent*>& OutComps,
-		TFunction<bool(const UPrimitiveComponent*)> ShouldChoose = ShouldChooseDefault) = 0;
+		TFunction<bool(const UPrimitiveComponent*)> ShouldChoose = +ShouldChooseDefault) = 0;
 
 	//BEGIN: APIs for backward compatibility
 	ENGINE_API void UpdateIndividualTexture(UTexture2D* Texture);
