[HLOD] Fix issue where, after building HLOD using the button in HLOD Actor's Details tab, actors that are source actors of the HLOD Actor can no longer be saved in the editor.

diff --git a/Engine/Source/Runtime/Engine/Private/WorldPartition/HLOD/HLODSourceActorsFromCell.cpp b/Engine/Source/Runtime/Engine/Private/WorldPartition/HLOD/HLODSourceActorsFromCell.cpp
--- a/Engine/Source/Runtime/Engine/Private/WorldPartition/HLOD/HLODSourceActorsFromCell.cpp
+++ b/Engine/Source/Runtime/Engine/Private/WorldPartition/HLOD/HLODSourceActorsFromCell.cpp
@@ -9,6 +9,7 @@
 #include "WorldPartition/WorldPartitionLevelStreamingDynamic.h"
 
 #include "Editor/UnrealEd/Public/EditorLevelUtils.h" // MoveActorsToLevel
+#include "UObject/Linker.h"
 #include "WorldPartition/ContentBundle/ContentBundlePaths.h"
 #include "WorldPartition/DataLayer/ExternalDataLayerHelper.h"
 
@@ -74,6 +75,15 @@ bool UWorldPartitionHLODSourceActorsFromCell::LoadSourceActors(bool& bOutDirty,
 	{
 		TArray<UPackage*> ModifiedPackages;
 		FWorldPartitionLevelHelper::MoveExternalActorsToLevel(ActorsCopy, TargetWorld->PersistentLevel, ModifiedPackages);
+		
+		TArray<UObject*> PackagesToResetLoaders;
+		PackagesToResetLoaders.Reserve(ModifiedPackages.Num());
+		for (UPackage* Package : ModifiedPackages)
+		{
+			PackagesToResetLoaders.Add(Package);
+		}
+		ResetLoaders(PackagesToResetLoaders);
+
 		return true;
 	}
 
