[World] Merging CL 47101183 from //UE5/Partner-CDPR/...

Provide a ForEachDataLayerInstance method in Data Layer Utils with which it's possible to iterate on the object / descriptor instances (whichever is available).

#tests local test with private, asset, EDL and deprecated data layers
#rnx

diff --git a/Engine/Source/Runtime/Engine/Private/WorldPartition/DataLayer/DataLayerUtils.cpp b/Engine/Source/Runtime/Engine/Private/WorldPartition/DataLayer/DataLayerUtils.cpp
--- a/Engine/Source/Runtime/Engine/Private/WorldPartition/DataLayer/DataLayerUtils.cpp
+++ b/Engine/Source/Runtime/Engine/Private/WorldPartition/DataLayer/DataLayerUtils.cpp
@@ -13,14 +13,16 @@
 
 #define LOCTEXT_NAMESPACE "DataLayer"
 
-FDataLayerInstanceNames FDataLayerUtils::ResolveDataLayerInstanceNames(const UDataLayerManager* InDataLayerManager, const FWorldPartitionActorDesc* InActorDesc, const TArray<const FWorldDataLayersActorDesc*>& InWorldDataLayersActorDescs)
+bool FDataLayerUtils::ForEachDataLayerInstance(const UDataLayerManager* InDataLayerManager,
+	const FWorldPartitionActorDesc* InActorDesc,
+	const TArray<const FWorldDataLayersActorDesc*>& InWorldDataLayersActorDescs,
+	TFunctionRef<void(const UDataLayerInstance*, const FDataLayerInstanceDesc*, bool)> Callback)
 {
-	const bool bIncludeExternalDataLayer = false;
 	const FName ActorDescExternalDataLayer = InActorDesc->GetExternalDataLayer();
-	const TArray<FName> ActorDescNonExternalDataLayers = InActorDesc->GetDataLayers(bIncludeExternalDataLayer);
+	const TArray<FName>& ActorDescNonExternalDataLayers = InActorDesc->GetNonExternalDataLayers();
 	if (ActorDescNonExternalDataLayers.IsEmpty() && ActorDescExternalDataLayer.IsNone())
 	{
-		return FDataLayerInstanceNames();
+		return true;
 	}
 	
 	// DataLayers not using DataLayer Assets represent DataLayerInstanceNames
@@ -28,28 +30,26 @@ FDataLayerInstanceNames FDataLayerUtils::ResolveDataLayerInstanceNames(const UDa
 	{
 		if (InDataLayerManager && InDataLayerManager->CanResolveDataLayers())
 		{
-			TArray<FName> Result;
 			for (const FName& DataLayerInstanceName : ActorDescNonExternalDataLayers)
 			{
 				if (const UDataLayerInstance* DataLayerInstance = InDataLayerManager->GetDataLayerInstance(DataLayerInstanceName))
 				{
-					Result.Add(DataLayerInstanceName);
+					Callback(DataLayerInstance, nullptr, false);
 				}
 			}
-			return FDataLayerInstanceNames(Result, false);
+			return true;
 		}
 		// Fallback on FWorldDataLayersActorDesc
 		else if (!InWorldDataLayersActorDescs.IsEmpty() && AreWorldDataLayersActorDescsSane(InWorldDataLayersActorDescs))
 		{
-			TArray<FName> Result;
 			for (const FName& DataLayerInstanceName : ActorDescNonExternalDataLayers)
 			{
-				if (GetDataLayerInstanceDescFromInstanceName(InWorldDataLayersActorDescs, DataLayerInstanceName) != nullptr)
+				if (const FDataLayerInstanceDesc* DataLayerInstanceDesc = GetDataLayerInstanceDescFromInstanceName(InWorldDataLayersActorDescs, DataLayerInstanceName))
 				{
-					Result.Add(DataLayerInstanceName);
+					Callback(nullptr, DataLayerInstanceDesc, false);
 				}
 			}
-			return FDataLayerInstanceNames(Result, false);
+			return true;
 		}
 	}
 	// ActorDesc DataLayers represents DataLayer Asset Paths
@@ -57,45 +57,79 @@ FDataLayerInstanceNames FDataLayerUtils::ResolveDataLayerInstanceNames(const UDa
 	{
 		if (InDataLayerManager && InDataLayerManager->CanResolveDataLayers())
 		{
-			TArray<FName> ResolvedDLINames;
-			ResolvedDLINames.Reserve(ActorDescNonExternalDataLayers.Num() + 1);
-			const UDataLayerInstance* ExternalDataLayerInstance = InDataLayerManager->GetDataLayerInstanceFromAssetName(ActorDescExternalDataLayer);
-			FName ExternalDataLayer = ExternalDataLayerInstance ? ExternalDataLayerInstance->GetDataLayerFName() : NAME_None;
-			bool bIsFirstDataLayerExternal = !ExternalDataLayer.IsNone();
-			if (bIsFirstDataLayerExternal)
+			if (const UDataLayerInstance* ExternalDataLayerInstance = InDataLayerManager->GetDataLayerInstanceFromAssetName(ActorDescExternalDataLayer))
 			{
-				ResolvedDLINames.Add(ExternalDataLayer);
+				Callback(ExternalDataLayerInstance, nullptr, true);
 			}
 			for (const FName& DataLayerAssetPath : ActorDescNonExternalDataLayers)
 			{
 				if (const UDataLayerInstance* DataLayerInstance = InDataLayerManager->GetDataLayerInstanceFromAssetName(DataLayerAssetPath))
 				{
-					ResolvedDLINames.Add(DataLayerInstance->GetDataLayerFName());
+					Callback(DataLayerInstance, nullptr, false);
 				}
 			}
-			return FDataLayerInstanceNames(ResolvedDLINames, bIsFirstDataLayerExternal);
+			return true;
 		}
 		// Fallback on FWorldDataLayersActorDesc
 		else if (!InWorldDataLayersActorDescs.IsEmpty() && AreWorldDataLayersActorDescsSane(InWorldDataLayersActorDescs))
 		{
-			TArray<FName> ResolvedDLINames;
-			ResolvedDLINames.Reserve(ActorDescNonExternalDataLayers.Num() + 1);
-			const FDataLayerInstanceDesc* ExternalDataLayerInstanceDesc = GetDataLayerInstanceDescFromAssetPath(InWorldDataLayersActorDescs, ActorDescExternalDataLayer);
-			FName ExternalDataLayer = ExternalDataLayerInstanceDesc ? ExternalDataLayerInstanceDesc->GetName() : NAME_None;
-			bool bIsFirstDataLayerExternal = !ExternalDataLayer.IsNone();
-			if (bIsFirstDataLayerExternal)
+			if (const FDataLayerInstanceDesc* ExternalDataLayerInstanceDesc = GetDataLayerInstanceDescFromAssetPath(InWorldDataLayersActorDescs, ActorDescExternalDataLayer))
 			{
-				ResolvedDLINames.Add(ExternalDataLayer);
+				Callback(nullptr, ExternalDataLayerInstanceDesc, true);
 			}
 			for (const FName& DataLayerAssetPath : ActorDescNonExternalDataLayers)
 			{
 				if (const FDataLayerInstanceDesc* DataLayerInstanceDesc = GetDataLayerInstanceDescFromAssetPath(InWorldDataLayersActorDescs, DataLayerAssetPath))
 				{
-					ResolvedDLINames.Add(DataLayerInstanceDesc->GetName());
+					Callback(nullptr, DataLayerInstanceDesc, false);
 				}
 			}
-			return FDataLayerInstanceNames(ResolvedDLINames, bIsFirstDataLayerExternal);
+			return true;
+		}
+	}
+	return false;
+}
+
+FDataLayerInstanceNames FDataLayerUtils::ResolveDataLayerInstanceNames(const UDataLayerManager* InDataLayerManager, const FWorldPartitionActorDesc* InActorDesc, const TArray<const FWorldDataLayersActorDesc*>& InWorldDataLayersActorDescs)
+{
+	const FName ActorDescExternalDataLayer = InActorDesc->GetExternalDataLayer();
+	const TArray<FName>& ActorDescNonExternalDataLayers = InActorDesc->GetNonExternalDataLayers();
+	if (ActorDescNonExternalDataLayers.IsEmpty() && ActorDescExternalDataLayer.IsNone())
+	{
+		return FDataLayerInstanceNames();
+	}
+
+	TArray<FName> ResolvedDLINames;
+
+	bool bIsFirstDataLayerExternal = false;
+	const bool bSuccess = ForEachDataLayerInstance(InDataLayerManager, InActorDesc, InWorldDataLayersActorDescs,
+		[&ResolvedDLINames, &bIsFirstDataLayerExternal, &ActorDescNonExternalDataLayers](
+			const UDataLayerInstance* Instance, const FDataLayerInstanceDesc* InstanceDesc, const bool bExternal)
+	{
+		if (ResolvedDLINames.Num() == 0)
+		{
+			bIsFirstDataLayerExternal = bExternal;
+			ResolvedDLINames.Reserve(ActorDescNonExternalDataLayers.Num() + 1);
+		}
+		else
+		{
+		    check(!bExternal);
 		}
+
+		if (Instance)
+		{
+			ResolvedDLINames.Add(Instance->GetDataLayerFName());
+		}
+		else
+		{
+			check(InstanceDesc);
+			ResolvedDLINames.Add(InstanceDesc->GetName());
+		}
+	});
+
+	if (bSuccess)
+	{
+		return FDataLayerInstanceNames(ResolvedDLINames, bIsFirstDataLayerExternal);
 	}
 	return FDataLayerInstanceNames(ActorDescNonExternalDataLayers, ActorDescExternalDataLayer);
 }
diff --git a/Engine/Source/Runtime/Engine/Public/WorldPartition/DataLayer/DataLayerUtils.h b/Engine/Source/Runtime/Engine/Public/WorldPartition/DataLayer/DataLayerUtils.h
--- a/Engine/Source/Runtime/Engine/Public/WorldPartition/DataLayer/DataLayerUtils.h
+++ b/Engine/Source/Runtime/Engine/Public/WorldPartition/DataLayer/DataLayerUtils.h
@@ -43,6 +43,17 @@ public:
 
 	PRAGMA_ENABLE_DEPRECATION_WARNINGS
 
+	/**
+	 * Calls Callback for each data layer instance that an actor belongs to.
+	 *
+	 * @param InDataLayerManager If provided, we'll try to fetch the data layer instances for the current world
+	 * @param InActorDesc The actor for which we check the assigned instances
+	 * @param InWorldDataLayersActorDescs List of Data Layers to consider when Data Layer Manager is not available
+	 * @param Callback The function to call for each instance; only one of the two pointers will be valid
+	 * @return false if we failed to resolve data layer instances
+	 */
+	static ENGINE_API bool ForEachDataLayerInstance(const UDataLayerManager* InDataLayerManager, const FWorldPartitionActorDesc* InActorDesc, const TArray<const FWorldDataLayersActorDesc*>& InWorldDataLayersActorDescs, TFunctionRef<void(const UDataLayerInstance*, const FDataLayerInstanceDesc*, bool /** bExternal */)> Callback);
+
 	static ENGINE_API FDataLayerInstanceNames ResolveDataLayerInstanceNames(const UDataLayerManager* InDataLayerManager, const FWorldPartitionActorDesc* InActorDesc, const TArray<const FWorldDataLayersActorDesc*>& InWorldDataLayersActorDescs = TArray<const FWorldDataLayersActorDesc*>());
 
 	static ENGINE_API bool ResolveRuntimeDataLayerInstanceNames(const UDataLayerManager* InDataLayerManager, const IWorldPartitionActorDescInstanceView& InActorDescView, const TArray<const FWorldDataLayersActorDesc*>& InWorldDataLayersActorDescs, FDataLayerInstanceNames& OutRuntimeDataLayerInstanceNames);
diff --git a/Engine/Source/Runtime/Engine/Public/WorldPartition/WorldPartitionActorDesc.h b/Engine/Source/Runtime/Engine/Public/WorldPartition/WorldPartitionActorDesc.h
--- a/Engine/Source/Runtime/Engine/Public/WorldPartition/WorldPartitionActorDesc.h
+++ b/Engine/Source/Runtime/Engine/Public/WorldPartition/WorldPartitionActorDesc.h
@@ -330,6 +330,7 @@ public:
 	inline bool GetActorIsHLODRelevant() const { return bActorIsHLODRelevant; }
 	inline FSoftObjectPath GetHLODLayer() const { return HLODLayer; }
 	ENGINE_API TArray<FName> GetDataLayers(bool bIncludeExternalDataLayer = true) const;
+	inline const TArray<FName>& GetNonExternalDataLayers() const { return DataLayers; }
 	ENGINE_API FName GetExternalDataLayer() const;
 	inline const FSoftObjectPath& GetExternalDataLayerAsset() const { return ExternalDataLayerAsset; }
 
