UAssetManager::OnAssetRegistryAvailableAfterInitialization: If FPlatformFilePak is initialized then don't fall back to the loose filesystem when querying if AssetRegistry.bin files exist. This avoids a lot of unnecessary and expensive FileExists calls

diff --git a/Engine/Source/Runtime/Engine/Private/AssetManager.cpp b/Engine/Source/Runtime/Engine/Private/AssetManager.cpp
--- a/Engine/Source/Runtime/Engine/Private/AssetManager.cpp
+++ b/Engine/Source/Runtime/Engine/Private/AssetManager.cpp
@@ -3038,7 +3038,10 @@ bool UAssetManager::OnAssetRegistryAvailableAfterInitialization(FName InName, FA
 		{
 			TArray<uint8> Bytes;
 			const FString Filename = FPaths::ProjectDir() / (TEXT("AssetRegistry") + InName.ToString()) + TEXT(".bin");
-			if (FPaths::FileExists(Filename) && FFileHelper::LoadFileToArray(Bytes, *Filename))
+			// Optimization: if we're using pak files then only search paks (avoid unnecessary fallback to loose)
+			FPakPlatformFile* PakPlatformFile = (FPakPlatformFile*)(FPlatformFileManager::Get().FindPlatformFile(FPakPlatformFile::GetTypeName()));
+			bool bFileExists = PakPlatformFile != nullptr ? PakPlatformFile->FindFileInPakFiles(*Filename) : FPaths::FileExists(Filename);
+			if (bFileExists && FFileHelper::LoadFileToArray(Bytes, *Filename))
 			{
 				bLoaded = true;
 				FMemoryReader Ar(Bytes);
