Remove ensure to make sure we get the error logs in crash reporter.

diff --git a/Engine/Source/Runtime/Engine/Public/WorldPartition/ActorDescList.h b/Engine/Source/Runtime/Engine/Public/WorldPartition/ActorDescList.h
--- a/Engine/Source/Runtime/Engine/Public/WorldPartition/ActorDescList.h
+++ b/Engine/Source/Runtime/Engine/Public/WorldPartition/ActorDescList.h
@@ -227,7 +227,7 @@ void TActorDescList<DescType>::RemoveActorDescriptor(DescType* ActorDesc)
 	check(!bIsProxy);
 	check(ActorDesc);
 	
-	if (!ensure((ActorsByGuid.Remove(ActorDesc->GetGuid()))))
+	if (!ActorsByGuid.Remove(ActorDesc->GetGuid()))
 	{
 		UE_LOG(LogWorldPartition, Error, TEXT("Removing unknown actor descriptor: %s"), *ActorDesc->ToString(FWorldPartitionActorDesc::EToStringMode::Full));
 	}
