Pulled CL from sarah.noonan
Add option to force set the outer world of an external streaming object if its outer world on load is invalid. Disabled by default, can be enabled via ExternalDataLayerManager.ShouldForceSetExternalStreamingObjectOuterWorldIfNull

#rnx

diff --git a/Engine/Source/Runtime/Engine/Private/WorldPartition/DataLayer/ExternalDataLayerManager.cpp b/Engine/Source/Runtime/Engine/Private/WorldPartition/DataLayer/ExternalDataLayerManager.cpp
--- a/Engine/Source/Runtime/Engine/Private/WorldPartition/DataLayer/ExternalDataLayerManager.cpp
+++ b/Engine/Source/Runtime/Engine/Private/WorldPartition/DataLayer/ExternalDataLayerManager.cpp
@@ -32,6 +32,16 @@
 
 #define LOCTEXT_NAMESPACE "ExternalDataLayerManager"
 
+namespace ExternalDataLayerManagerPrivate
+{
+	static bool bShouldForceSetExternalStreamingObjectOuterWorldIfNull = false;
+	static FAutoConsoleVariableRef CVarShouldForceSetExternalStreamingObjectOuterWorldIfNull(
+		TEXT("ExternalDataLayerManager.ShouldForceSetExternalStreamingObjectOuterWorldIfNull"),
+		bShouldForceSetExternalStreamingObjectOuterWorldIfNull,
+		TEXT("Whether to assign the current world as the outer world of an external streaming object if it does not have a valid outer world"),
+		ECVF_Default);
+}
+
 UExternalDataLayerManager::UExternalDataLayerManager()
 	: bIsInitialized(false)
 	, bIsRunningGameOrInstancedWorldPartition(false)
@@ -158,7 +168,8 @@ bool UExternalDataLayerManager::CanInjectExternalDataLayerAsset(const UExternalD
 
 bool UExternalDataLayerManager::RegisterExternalStreamingObjectForGameWorld(const UExternalDataLayerAsset* InExternalDataLayerAsset)
 {
-	check(GetTypedOuter<UWorld>()->IsGameWorld());
+	UWorld* OuterWorld = GetTypedOuter<UWorld>();
+	check(OuterWorld && OuterWorld->IsGameWorld());
 #if !WITH_EDITOR
 	check(InExternalDataLayerAsset);
 	if (ExternalStreamingObjects.Contains(InExternalDataLayerAsset))
@@ -179,7 +190,7 @@ bool UExternalDataLayerManager::RegisterExternalStreamingObjectForGameWorld(cons
 	UPackage* DestPackage = nullptr;
 	FLinkerInstancingContext InstancingContext;
 	FString SourceWorldPath, RemappedWorldPath;
-	if (GetTypedOuter<UWorld>()->GetSoftObjectPathMapping(SourceWorldPath, RemappedWorldPath))
+	if (OuterWorld->GetSoftObjectPathMapping(SourceWorldPath, RemappedWorldPath))
 	{
 		FString Source = FTopLevelAssetPath(SourceWorldPath).GetPackageName().ToString();
 		FString Remapped = FTopLevelAssetPath(RemappedWorldPath).GetPackageName().ToString();
@@ -208,8 +219,14 @@ bool UExternalDataLayerManager::RegisterExternalStreamingObjectForGameWorld(cons
 		return false;
 	}
 
+	if (ExternalDataLayerManagerPrivate::bShouldForceSetExternalStreamingObjectOuterWorldIfNull && !ExternalStreamingObject->GetOuterWorld())
+	{
+		UE_LOG(LogWorldPartition, Warning, TEXT("[EDL: %s] Outer world is invalid, force setting with current world"), *InExternalDataLayerAsset->GetName());
+		ExternalStreamingObject->SetOuterWorld(OuterWorld);
+	}
+
 	// Do some validation on ExternalStreamingObject's OuterWorld and OwningWorld
-	check(ExternalStreamingObject->GetOuterWorld() == GetTypedOuter<UWorld>());
+	check(ExternalStreamingObject->GetOuterWorld() == OuterWorld);
 	check(ExternalStreamingObject->GetOwningWorld() == GetOuterUWorldPartition()->GetWorld());
 
 	ExternalStreamingObjects.Emplace(InExternalDataLayerAsset, ExternalStreamingObject);
diff --git a/Engine/Source/Runtime/Engine/Private/WorldPartition/WorldPartitionRuntimeHash.cpp b/Engine/Source/Runtime/Engine/Private/WorldPartition/WorldPartitionRuntimeHash.cpp
--- a/Engine/Source/Runtime/Engine/Private/WorldPartition/WorldPartitionRuntimeHash.cpp
+++ b/Engine/Source/Runtime/Engine/Private/WorldPartition/WorldPartitionRuntimeHash.cpp
@@ -105,6 +105,11 @@ UWorld* URuntimeHashExternalStreamingObjectBase::GetOwningWorld() const
 	return OwningWorld.IsSet() ? OwningWorld.GetValue().Get() : GetOuterWorld()->GetWorldPartition()->GetWorld();
 }
 
+void URuntimeHashExternalStreamingObjectBase::SetOuterWorld(UWorld* InWorld)
+{
+	OuterWorld = InWorld;
+}
+
 void URuntimeHashExternalStreamingObjectBase::OnStreamingObjectLoaded(UWorld* InjectedWorld)
 {
 #if !WITH_EDITOR
diff --git a/Engine/Source/Runtime/Engine/Public/WorldPartition/WorldPartitionRuntimeHash.h b/Engine/Source/Runtime/Engine/Public/WorldPartition/WorldPartitionRuntimeHash.h
--- a/Engine/Source/Runtime/Engine/Public/WorldPartition/WorldPartitionRuntimeHash.h
+++ b/Engine/Source/Runtime/Engine/Public/WorldPartition/WorldPartitionRuntimeHash.h
@@ -86,6 +86,8 @@ public:
 	UExternalDataLayerInstance* GetRootExternalDataLayerInstance() { return const_cast<UExternalDataLayerInstance*>(RootExternalDataLayerInstance.Get()); }
 	ENGINE_API const UObject* GetLevelMountPointContextObject() const;
 
+	void SetOuterWorld(UWorld* InWorld);
+
 #if WITH_EDITOR
 	UE_DEPRECATED(5.4, "PopulateGeneratorPackageForCook is depreacted and was replaced by OnPopulateGeneratorPackageForCook")
 	ENGINE_API void PopulateGeneratorPackageForCook();
