NavHeightMap: don't draw things outside of view and reduce mesh count

diff --git a/Engine/Source/Runtime/Engine/Public/DebugRenderSceneProxy.h b/Engine/Source/Runtime/Engine/Public/DebugRenderSceneProxy.h
--- a/Engine/Source/Runtime/Engine/Public/DebugRenderSceneProxy.h
+++ b/Engine/Source/Runtime/Engine/Public/DebugRenderSceneProxy.h
@@ -105,6 +105,13 @@ public:
 		return ViewFrustum.IntersectLineSegment(StartPoint, EndPoint);
 	}
 
+	inline static bool BoxInFrustum(const FBox& Box, const FConvexVolume& ViewFrustum)
+	{
+		FVector Center, Extents;
+		Box.GetCenterAndExtents(Center, Extents);
+		return BoxInFrustum(Center, Extents, ViewFrustum);
+	}
+
 	inline static bool BoxInFrustum(const FVector& Origin, const FVector& Extent, const FConvexVolume& ViewFrustum)
 	{
 		return ViewFrustum.IntersectBox(Origin, Extent);
