Wave Player Optimization
- New FSoundWaveProxyPlayer that should replace FSoundWaveProxyReader, handling Multi Channel Rendering, Resampling, Speed, and Loop settings.
- SoundWaveProxyPlayer can be reassigned with new settings and sound waves without re-allocating whole new buffers
- SoundWavePlayerNodeV2 using the new SoundWaveProxyPlayer
- Ability to toggle the new SoundWavePlayerNodeV2 with a CVar, injecting the V2 operator into the original SoundWavePlayerNode

#rnx

diff --git a/Engine/Plugins/Runtime/Metasound/Source/MetasoundEngine/Private/MetasoundWavePlayerNode.cpp b/Engine/Plugins/Runtime/Metasound/Source/MetasoundEngine/Private/MetasoundWavePlayerNode.cpp
--- a/Engine/Plugins/Runtime/Metasound/Source/MetasoundEngine/Private/MetasoundWavePlayerNode.cpp
+++ b/Engine/Plugins/Runtime/Metasound/Source/MetasoundEngine/Private/MetasoundWavePlayerNode.cpp
@@ -17,6 +17,7 @@
 #include "MetasoundTrigger.h"
 #include "MetasoundVertex.h"
 #include "MetasoundWave.h"
+#include "MetasoundWavePlayerNodeV2.h"
 #include "Sound/SoundWaveProxyReader.h"
 
 #define LOCTEXT_NAMESPACE "MetasoundWaveNode"
@@ -331,6 +332,14 @@ namespace Metasound
 		TEXT("Default: 512"),
 		ECVF_Default);
 
+	static bool UseNewWavePlayerOperator = false;
+	FAutoConsoleVariableRef CVarMetaSoundWavePlayerUseNewWavePlayer(
+		TEXT("au.MetaSound.WavePlayer.UseNewOperator"),
+		UseNewWavePlayerOperator,
+		TEXT("Whether to use the new wave player operator\n")
+		TEXT("0: Use the original Wave Player operator\n 1: Use the NEW Wave player Operator"),
+		ECVF_Default);
+
 	namespace WavePlayerOperatorPrivate
 	{
 		template<typename AudioChannelConfigurationInfoType>
@@ -907,6 +916,8 @@ namespace Metasound
 		 */
 		void GenerateSourceAudio(Audio::FMultichannelCircularBuffer& OutBuffer, WavePlayerNodePrivate::FSourceBufferState& OutSourceState)
 		{
+			METASOUND_TRACE_CPUPROFILER_EVENT_SCOPE(Metasound::FWavePlayerOperator::GenerateSourceAudio);
+			
 			using namespace WavePlayerNodePrivate;
 
 			if (bIsPlaying)
@@ -1094,6 +1105,8 @@ namespace Metasound
 
 		void GeneratePitchedAudio(Audio::FMultichannelBufferView& OutBuffer)
 		{
+			METASOUND_TRACE_CPUPROFILER_EVENT_SCOPE(Metasound::FWavePlayerOperator::GeneratePitchedAudio);
+			
 			using namespace Audio;
 
 			// Outside callers should ensure that bIsPlaying is true if calling this function.
@@ -1196,6 +1209,11 @@ namespace Metasound
 
 		virtual TUniquePtr<IOperator> CreateOperator(const FBuildOperatorParams& InParams, FBuildResults& OutResults) override
 		{
+			if (UseNewWavePlayerOperator)
+			{
+				return Nodes::WavePlayerV2::CreateOperator(OutputAudioVertices.Num(), InParams, OutResults);	
+			}
+			
 			using namespace Audio;
 			using namespace WavePlayerVertexNames;
 
diff --git a/Engine/Plugins/Runtime/Metasound/Source/MetasoundEngine/Private/MetasoundWavePlayerNodeV2.cpp b/Engine/Plugins/Runtime/Metasound/Source/MetasoundEngine/Private/MetasoundWavePlayerNodeV2.cpp
new file mode 100644
--- /dev/null
+++ b/Engine/Plugins/Runtime/Metasound/Source/MetasoundEngine/Private/MetasoundWavePlayerNodeV2.cpp
@@ -0,0 +1,615 @@
+// Copyright Epic Games, Inc. All Rights Reserved.
+
+#include "MetasoundWavePlayerNodeV2.h"
+#include "MetasoundAudioBuffer.h"
+#include "MetasoundEngineNodesNames.h"
+#include "MetasoundExecutableOperator.h"
+#include "MetasoundFacade.h"
+#include "MetasoundLog.h"
+#include "MetasoundNodeRegistrationMacro.h"
+#include "MetasoundOperatorInterface.h"
+#include "MetasoundParamHelper.h"
+#include "MetasoundPrimitives.h"
+#include "MetasoundTime.h"
+#include "MetasoundTrigger.h"
+#include "MetasoundTrace.h"
+#include "MetasoundWave.h"
+#include "Sound/SoundWaveProxyPlayer.h"
+
+#define LOCTEXT_NAMESPACE "MetasoundWaveNode"
+
+namespace Metasound
+{
+	namespace Inputs
+	{
+		METASOUND_PARAM(Play, "Play", "Play the wave player.");
+		METASOUND_PARAM(Stop, "Stop", "Stop the wave player.");
+		METASOUND_PARAM(WaveAsset, "Wave Asset", "The wave asset to be real-time decoded.");
+		METASOUND_PARAM(StartTime, "Start Time", "Time (in seconds) into the wave asset to start (seek) the wave asset.");
+		METASOUND_PARAM(PitchShift, "Pitch Shift", "The pitch shift to use for the wave asset in semitones.");
+		METASOUND_PARAM(Loop, "Loop", "Whether or not to loop between the start and specified end times.");
+		METASOUND_PARAM(LoopStartTime, "Loop Start", "When to start the loop");
+		METASOUND_PARAM(LoopDuration, "Loop Duration", "The duration of the loop when wave player is enabled for looping. A negative value will loop the whole wave asset.");
+		METASOUND_PARAM(MaintainAudioSync, "Maintain Audio Sync", "Force audio to seek to maintain sample accurate playback even if the source audio is latent")
+	}
+
+	namespace Outputs
+	{
+		METASOUND_PARAM(OnPlay, "On Play", "Triggers when Play is triggered.");
+		METASOUND_PARAM(OnFinished, "On Finished", "Triggers when the wave played has finished playing or Stop is Triggered");
+		METASOUND_PARAM(OnNearlyFinished, "On Nearly Finished", "Triggers when the wave played has almost finished playing (the block before it finished). Allows time for logic to trigger different variations to play seamlessly.");
+		METASOUND_PARAM(OnLooped, "On Looped", "Triggers when the wave player has looped.");
+		METASOUND_PARAM(OnCuePoint, "On Cue Point", "Triggers when a wave cue point was hit during playback.");
+		METASOUND_PARAM(CuePointID, "Cue Point ID", "The cue point ID that was triggered");
+		METASOUND_PARAM(CuePointLabel, "Cue Point Label", "The cue point label that was triggered (if there was a label parsed in the imported .wav file)");
+		METASOUND_PARAM(LoopProgress, "Loop Percent", "Returns the current playback progress as a value between (0-1) given the loop start and the loop duration if looping is enabled");
+		METASOUND_PARAM(PlaybackProgress, "Playback Location", "Returns the current absolute playback progress of the wave playback as a value between (0-1) based on the total wave duration");
+		METASOUND_PARAM(PlaybackTime, "Playback Time", "Returns the current absolute playback time of the wave");
+
+		namespace Audio
+		{
+			METASOUND_PARAM(Mono, "Out Mono", "The mono channel audio output.")
+			METASOUND_PARAM(Left, "Out Left", "The left channel audio output.")
+			METASOUND_PARAM(Right, "Out Right", "The right channel audio output.")
+			METASOUND_PARAM(FrontRight, "Out Front Right", "The front right channel audio output.")
+			METASOUND_PARAM(FrontLeft, "Out Front Left", "The front left channel audio output.")
+			METASOUND_PARAM(FrontCenter, "Out Front Center", "The front center channel audio output.")
+			METASOUND_PARAM(LowFrequency, "Out Low Frequency", "The low frequency channel audio output.")
+			METASOUND_PARAM(SideRight, "Out Side Right", "The side right channel audio output.")
+			METASOUND_PARAM(SideLeft, "Out Side Left", "The side left channel audio output.")
+			METASOUND_PARAM(BackRight, "Out Back Right", "The back right channel audio output.")
+			METASOUND_PARAM(BackLeft, "Out Back Left", "The back left channel audio output.")
+		}
+	}
+
+	class FWavePlayerOperatorV2 : public TExecutableOperator<FWavePlayerOperatorV2>
+	{
+	public:
+
+		struct FInputs
+		{
+			FTriggerReadRef PlayTrigger;
+			FTriggerReadRef StopTrigger;
+			FWaveAssetReadRef WaveAsset;
+			FTimeReadRef StartTime;
+			FFloatReadRef PitchShift;
+			FBoolReadRef Loop;
+			FTimeReadRef LoopStartTime;
+			FTimeReadRef LoopDuration;
+			bool bMaintainAudioSync;
+		};
+		
+		struct FOutputs
+		{
+			FTriggerWriteRef OnFinished;
+			FTriggerWriteRef OnLooped;
+			FTimeWriteRef PlaybackTime;
+
+			TArray<FAudioBufferWriteRef> Audio;
+
+			FTriggerWriteRef OnNearlyFinished;
+			FTriggerWriteRef OnCuePoint;
+			FInt32WriteRef CuePointID;
+			FStringWriteRef CuePointLabel;
+			FFloatWriteRef LoopProgress;
+			FFloatWriteRef PlaybackProgress;
+
+			FOutputs(const FOperatorSettings& InOperatorSettings, int32 NumChannels)
+				: OnFinished(FTriggerWriteRef::CreateNew(InOperatorSettings))
+				, OnLooped(FTriggerWriteRef::CreateNew(InOperatorSettings))
+				, PlaybackTime(FTimeWriteRef::CreateNew(0.0f))
+				, OnNearlyFinished(FTriggerWriteRef::CreateNew(InOperatorSettings))
+				, OnCuePoint(FTriggerWriteRef::CreateNew(InOperatorSettings))
+				, CuePointID(FInt32WriteRef::CreateNew(0))
+				, CuePointLabel(FStringWriteRef::CreateNew(TEXT("")))
+				, LoopProgress(FFloatWriteRef::CreateNew(0.0f))
+				, PlaybackProgress(FFloatWriteRef::CreateNew(0.0f))
+			{
+				for (int32 Idx = 0; Idx < NumChannels; ++Idx)
+				{
+					Audio.Add(
+						FAudioBufferWriteRef::CreateNew(InOperatorSettings)
+					);
+				}
+			}
+
+			void ZeroBuffers()
+			{
+				for (int32 Idx = 0; Idx < Audio.Num(); ++Idx)
+				{
+					Audio[Idx]->Zero();
+				}
+			}
+
+			void Reset(const FResetParams& ResetParams, int32 NumChannels)
+			{
+				for (int32 Idx = 0; Idx < Audio.Num(); ++Idx)
+				{
+					Audio[Idx]->Zero();
+				}
+
+				OnFinished->Reset();
+				OnLooped->Reset();
+				OnNearlyFinished->Reset();
+				OnCuePoint->Reset();
+				*CuePointID = 0;
+				*CuePointLabel = TEXT("");
+				*LoopProgress = 0.0f;
+				*PlaybackProgress = 0.0f;
+				*PlaybackTime = FTime(0.0f);
+			}
+		};
+
+		static const FNodeClassName& GetNodeClassName(FName VariantName)
+		{
+			static const FName NodeName = TEXT("Wave Player V2");
+			static const TMap<FName, FNodeClassName> ClassNames
+			{
+				{ EngineNodes::MonoVariant, { EngineNodes::Namespace, NodeName, EngineNodes::MonoVariant } },
+				{ EngineNodes::StereoVariant, { EngineNodes::Namespace, NodeName, EngineNodes::StereoVariant } },
+				{ EngineNodes::QuadVariant, { EngineNodes::Namespace, NodeName, EngineNodes::QuadVariant } },
+				{ EngineNodes::FiveDotOneVariant, { EngineNodes::Namespace, NodeName, EngineNodes::FiveDotOneVariant } },
+				{ EngineNodes::SevenDotOneVariant, { EngineNodes::Namespace, NodeName, EngineNodes::SevenDotOneVariant } },
+			};
+			check(ClassNames.Contains(VariantName));
+			return ClassNames[VariantName];
+		}
+
+		static const FNodeClassMetadata& GetNodeInfo()
+		{
+			auto InitNodeInfo = []() -> FNodeClassMetadata
+				{
+					FNodeClassMetadata Info;
+					Info.ClassName        = GetNodeClassName(EngineNodes::MonoVariant);
+					Info.MajorVersion     = 1;
+					Info.MinorVersion     = 0;
+					Info.DisplayName      = METASOUND_LOCTEXT("MetasoundWavePlayerV2Node_StereoDisplayName", "Wave Player V2 (1.0, Mono) (Exp.)");
+					Info.Description      = METASOUND_LOCTEXT("WavePlayerV2_Description", "Plays a wave asset. The wave's channel configurations will be up or down mixed to match the wave players audio channel format.");
+					Info.Author           = PluginAuthor;
+					Info.PromptIfMissing  = PluginNodeMissingPrompt;
+					Info.DefaultInterface = GetVertexInterface();
+					Info.Keywords         = { METASOUND_LOCTEXT("WavePlayerV2Keyword_Sound", "Sound")
+											, METASOUND_LOCTEXT("WavePlayerV2Keyword_Cue", "Cue")};
+					return Info;
+				};
+
+			static const FNodeClassMetadata Info = InitNodeInfo();
+			return Info;
+		}
+
+		static const FVertexInterface& GetVertexInterface()
+		{
+			using namespace Metasound;
+
+			auto CreateVertexInterface = []() -> FVertexInterface
+				{
+					FInputVertexInterface InputVertexInterface {
+						TInputDataVertex<FTrigger>(METASOUND_GET_PARAM_NAME_AND_METADATA(Inputs::Play)),
+						TInputDataVertex<FTrigger>(METASOUND_GET_PARAM_NAME_AND_METADATA(Inputs::Stop)),
+						TInputDataVertex<FWaveAsset>(METASOUND_GET_PARAM_NAME_AND_METADATA(Inputs::WaveAsset)),
+						TInputDataVertex<FTime>(METASOUND_GET_PARAM_NAME_AND_METADATA(Inputs::StartTime), 0.0f),
+						TInputDataVertex<float>(METASOUND_GET_PARAM_NAME_AND_METADATA(Inputs::PitchShift), 0.0f),
+						TInputDataVertex<bool>(METASOUND_GET_PARAM_NAME_AND_METADATA(Inputs::Loop), false),
+						TInputDataVertex<FTime>(METASOUND_GET_PARAM_NAME_AND_METADATA_ADVANCED(Inputs::LoopStartTime), 0.0f),
+						TInputDataVertex<FTime>(METASOUND_GET_PARAM_NAME_AND_METADATA_ADVANCED(Inputs::LoopDuration), -1.0f),
+						TInputConstructorVertex<bool>(METASOUND_GET_PARAM_NAME_AND_METADATA_ADVANCED(Inputs::MaintainAudioSync), false)
+					};
+			
+					FOutputVertexInterface OutputVertexInterface;
+					OutputVertexInterface.Add(TOutputDataVertex<FTrigger>(METASOUND_GET_PARAM_NAME_AND_METADATA(Outputs::OnPlay)));
+					OutputVertexInterface.Add(TOutputDataVertex<FTrigger>(METASOUND_GET_PARAM_NAME_AND_METADATA(Outputs::OnFinished)));
+					OutputVertexInterface.Add(TOutputDataVertex<FTrigger>(METASOUND_GET_PARAM_NAME_AND_METADATA(Outputs::OnLooped)));
+					OutputVertexInterface.Add(TOutputDataVertex<FTime>(METASOUND_GET_PARAM_NAME_AND_METADATA(Outputs::PlaybackTime)));
+					OutputVertexInterface.Add(TOutputDataVertex<FAudioBuffer>(METASOUND_GET_PARAM_NAME_AND_METADATA(Outputs::Audio::Mono)));
+
+					OutputVertexInterface.Add(TOutputDataVertex<FTrigger>(METASOUND_GET_PARAM_NAME_AND_METADATA_ADVANCED(Outputs::OnNearlyFinished)));
+					OutputVertexInterface.Add(TOutputDataVertex<FTrigger>(METASOUND_GET_PARAM_NAME_AND_METADATA_ADVANCED(Outputs::OnCuePoint)));
+					OutputVertexInterface.Add(TOutputDataVertex<int32>(METASOUND_GET_PARAM_NAME_AND_METADATA_ADVANCED(Outputs::CuePointID)));
+					OutputVertexInterface.Add(TOutputDataVertex<FString>(METASOUND_GET_PARAM_NAME_AND_METADATA_ADVANCED(Outputs::CuePointLabel)));
+					OutputVertexInterface.Add(TOutputDataVertex<float>(METASOUND_GET_PARAM_NAME_AND_METADATA_ADVANCED(Outputs::LoopProgress)));
+					OutputVertexInterface.Add(TOutputDataVertex<float>(METASOUND_GET_PARAM_NAME_AND_METADATA_ADVANCED(Outputs::PlaybackProgress)));
+
+					return FVertexInterface(InputVertexInterface, OutputVertexInterface);
+				};
+
+			static const FVertexInterface Interface = CreateVertexInterface();
+			return Interface;
+		}
+		
+		static TUniquePtr<IOperator> CreateOperator(const FBuildOperatorParams& InParams, FBuildResults& OutResults)
+		{
+			FInputs Inputs
+			{
+				InParams.InputData.GetOrCreateDefaultDataReadReference<FTrigger>(METASOUND_GET_PARAM_NAME(Inputs::Play), InParams.OperatorSettings),
+				InParams.InputData.GetOrCreateDefaultDataReadReference<FTrigger>(METASOUND_GET_PARAM_NAME(Inputs::Stop), InParams.OperatorSettings),
+				InParams.InputData.GetOrCreateDefaultDataReadReference<FWaveAsset>(METASOUND_GET_PARAM_NAME(Inputs::WaveAsset), InParams.OperatorSettings),
+				InParams.InputData.GetOrCreateDefaultDataReadReference<FTime>(METASOUND_GET_PARAM_NAME(Inputs::StartTime), InParams.OperatorSettings),
+				InParams.InputData.GetOrCreateDefaultDataReadReference<float>(METASOUND_GET_PARAM_NAME(Inputs::PitchShift), InParams.OperatorSettings),
+				InParams.InputData.GetOrCreateDefaultDataReadReference<bool>(METASOUND_GET_PARAM_NAME(Inputs::Loop), InParams.OperatorSettings),
+				InParams.InputData.GetOrCreateDefaultDataReadReference<FTime>(METASOUND_GET_PARAM_NAME(Inputs::LoopStartTime), InParams.OperatorSettings),
+				InParams.InputData.GetOrCreateDefaultDataReadReference<FTime>(METASOUND_GET_PARAM_NAME(Inputs::LoopDuration), InParams.OperatorSettings),
+				InParams.InputData.GetOrCreateDefaultValue<bool>(METASOUND_GET_PARAM_NAME(Inputs::MaintainAudioSync), InParams.OperatorSettings),
+			};
+
+			FOutputs Outputs(InParams.OperatorSettings, 1);
+
+			return MakeUnique<FWavePlayerOperatorV2>(InParams, MoveTemp(Inputs), MoveTemp(Outputs));
+		}
+		
+		FWavePlayerOperatorV2(const FBuildOperatorParams& InParams, FInputs&& InInputs, FOutputs&& InOutputs)
+			: OperatorSettings(InParams.OperatorSettings)
+			, CostReporter(InParams.GraphRenderCost ? InParams.GraphRenderCost->AddNode(InParams.Node.GetInstanceID(), InParams.Environment) : FNodeRenderCost{})
+			, Inputs(MoveTemp(InInputs))
+			, Outputs(MoveTemp(InOutputs))
+		{
+			const int32 NumFramesPerBlock = InParams.OperatorSettings.GetNumFramesPerBlock();
+			const float SampleRate = InParams.OperatorSettings.GetSampleRate();
+			const int32 NumChannels = Outputs.Audio.Num();
+			FSoundWaveProxyPlayer::FSettings PlayerSettings(NumFramesPerBlock, SampleRate, NumChannels);
+			PlayerSettings.bMaintainAudioSync = Inputs.bMaintainAudioSync;
+			ProxyPlayer = FSoundWaveProxyPlayer::Create(PlayerSettings);
+
+			// setup buffer views
+			for (int32 Idx = 0; Idx < NumChannels; ++Idx)
+			{
+				MultiChannelOutputView.Emplace(Outputs.Audio[Idx]->GetData(), Outputs.Audio[Idx]->Num());
+			}
+		}
+
+		void Reset(const IOperator::FResetParams& ResetParams)
+		{
+			ensure(ProxyPlayer->GetOutputSampleRate() == ResetParams.OperatorSettings.GetSampleRate());
+			ensure(ProxyPlayer->GetOutputNumChannels() == Outputs.Audio.Num());
+			OperatorSettings = ResetParams.OperatorSettings;
+			Outputs.Reset(ResetParams, Outputs.Audio.Num());
+			
+			const int32 NumFramesPerBlock = ResetParams.OperatorSettings.GetNumFramesPerBlock();
+            const float SampleRate = ResetParams.OperatorSettings.GetSampleRate();
+            const int32 NumChannels = Outputs.Audio.Num();
+			FSoundWaveProxyPlayer::FSettings PlayerSettings(NumFramesPerBlock, SampleRate, NumChannels);
+			PlayerSettings.bMaintainAudioSync = Inputs.bMaintainAudioSync;
+			ProxyPlayer->Reset(PlayerSettings);
+
+			CurrentWaveAsset = FWaveAsset();
+			bIsPlaying = false;
+			bDidRenderThisBlock = false;
+		}
+
+		virtual void BindInputs(FInputVertexInterfaceData& InVertexData) override
+		{
+			InVertexData.BindReadVertex(Inputs::PlayName, Inputs.PlayTrigger);
+			InVertexData.BindReadVertex(Inputs::StopName, Inputs.StopTrigger);
+			InVertexData.BindReadVertex(Inputs::WaveAssetName, Inputs.WaveAsset);
+			InVertexData.BindReadVertex(Inputs::StartTimeName, Inputs.StartTime);
+			InVertexData.BindReadVertex(Inputs::PitchShiftName, Inputs.PitchShift);
+			InVertexData.BindReadVertex(Inputs::LoopName, Inputs.Loop);
+			InVertexData.BindReadVertex(Inputs::LoopStartTimeName, Inputs.LoopStartTime);
+			InVertexData.BindReadVertex(Inputs::LoopDurationName, Inputs.LoopDuration);
+			InVertexData.SetValue(Inputs::MaintainAudioSyncName, Inputs.bMaintainAudioSync);
+		}
+
+		virtual void BindOutputs(FOutputVertexInterfaceData& InVertexData) override
+		{
+			int32 NumChannels = Outputs.Audio.Num();
+			const TArray<FName>& BufferNames = GetOutputAudioBufferNames(NumChannels);
+			check(BufferNames.Num() == NumChannels);
+			for (int32 Idx = 0; Idx < NumChannels; ++Idx)
+			{
+				InVertexData.BindReadVertex(BufferNames[Idx], Outputs.Audio[Idx]);
+			}
+
+			// bind the input play trigger to the output play trigger to pass through the triggers
+			InVertexData.BindReadVertex(Outputs::OnPlayName, Inputs.PlayTrigger);
+			InVertexData.BindReadVertex(Outputs::OnFinishedName, Outputs.OnFinished);
+			InVertexData.BindReadVertex(Outputs::OnLoopedName, Outputs.OnLooped);
+			InVertexData.BindReadVertex(Outputs::PlaybackTimeName, Outputs.PlaybackTime);
+
+			InVertexData.BindReadVertex(Outputs::OnNearlyFinishedName, Outputs.OnNearlyFinished);
+			InVertexData.BindReadVertex(Outputs::OnCuePointName, Outputs.OnCuePoint);
+			InVertexData.BindReadVertex(Outputs::CuePointIDName, Outputs.CuePointID);
+			InVertexData.BindReadVertex(Outputs::CuePointLabelName, Outputs.CuePointLabel);
+			InVertexData.BindReadVertex(Outputs::PlaybackProgressName, Outputs.PlaybackProgress);
+			InVertexData.BindReadVertex(Outputs::LoopProgressName, Outputs.LoopProgress);
+		}
+
+		static const TArray<FName>& GetOutputAudioBufferNames(int32 NumChannels)
+		{
+			static const TArray<FName> None;
+			static const TArray<FName> Mono { Outputs::Audio::MonoName };
+			static const TArray<FName> Stereo
+			{
+				Outputs::Audio::LeftName,
+				Outputs::Audio::RightName
+			};
+			static const TArray<FName> Quad
+			{
+				Outputs::Audio::FrontLeftName,
+				Outputs::Audio::FrontRightName,
+				Outputs::Audio::SideLeftName,
+				Outputs::Audio::SideRightName
+			};
+			static const TArray<FName> FiveDotOne
+			{
+				Outputs::Audio::FrontLeftName,
+				Outputs::Audio::FrontRightName,
+				Outputs::Audio::FrontCenterName,
+				Outputs::Audio::LowFrequencyName,
+				Outputs::Audio::SideLeftName,
+				Outputs::Audio::SideRightName
+			};
+			static const TArray<FName> SevenDotOne
+			{
+				Outputs::Audio::FrontLeftName,
+				Outputs::Audio::FrontRightName,
+				Outputs::Audio::FrontCenterName,
+				Outputs::Audio::LowFrequencyName,
+				Outputs::Audio::SideLeftName,
+				Outputs::Audio::SideRightName,
+				Outputs::Audio::BackLeftName,
+				Outputs::Audio::BackRightName
+			};
+			switch (NumChannels)
+			{
+			case 1:
+				return Mono;
+			case 2:
+				return Stereo;
+			case 4:
+				return Quad;
+			case 6:
+				return FiveDotOne;
+			case 8:
+				return SevenDotOne;
+			default:
+				checkNoEntry();
+				return None;
+			}
+		}
+
+		void Execute()
+		{
+			METASOUND_TRACE_CPUPROFILER_EVENT_SCOPE(Metasound::FWavePlayerOperatorV2::Execute);
+
+			Outputs.OnFinished->AdvanceBlock();
+			Outputs.OnNearlyFinished->AdvanceBlock();
+			Outputs.OnCuePoint->AdvanceBlock();
+			Outputs.OnLooped->AdvanceBlock();
+
+			bDidRenderThisBlock = false;
+			
+			if (ProxyPlayer.IsValid())
+			{
+				ProxyPlayer->SetSpeed(CalculateSpeedFromPitchShift());
+				ProxyPlayer->SetLoop(*Inputs.Loop, Inputs.LoopStartTime->GetSeconds(), Inputs.LoopDuration->GetSeconds());
+				ExecuteSubblocks();
+				*Outputs.LoopProgress = ProxyPlayer->GetLoopProgress();
+				*Outputs.PlaybackProgress = ProxyPlayer->GetPlaybackProgress();
+				*Outputs.PlaybackTime = FTime::FromSeconds(ProxyPlayer->GetCurrentPlaybackTimeSeconds());
+			}
+			else
+			{
+				*Outputs.LoopProgress = 0.0f;
+				*Outputs.PlaybackProgress = 0.0f;
+				*Outputs.PlaybackTime = FTime::FromSeconds(0.0f);
+				ZeroFrameRange(0, OperatorSettings.GetNumFramesPerBlock());
+			}
+
+			if (bDidRenderThisBlock)
+			{
+				CostReporter.SetRenderCost(1.0f);
+			}
+		}
+
+	private:
+
+		void ExecuteSubblocks()
+		{
+#if !UE_BUILD_SHIPPING
+			// hot reload sound
+			if (bIsPlaying && Inputs.WaveAsset->IsSoundWaveValid() && Inputs.WaveAsset->GetSoundWaveProxy()->TryGetLatest())
+			{
+				StartPlaying();
+			}
+#endif // !UE_BUILD_SHIPPING
+			
+			// Parse triggers and render audio
+			int32 PlayTrigIndex = 0;
+			int32 NextPlayFrame = 0;
+			const int32 NumPlayTrigs = Inputs.PlayTrigger->NumTriggeredInBlock();
+
+			int32 StopTrigIndex = 0;
+			int32 NextStopFrame = 0;
+			const int32 NumStopTrigs = Inputs.StopTrigger->NumTriggeredInBlock();
+
+			int32 CurrAudioFrame = 0;
+			int32 NextAudioFrame = 0;
+			const int32 LastAudioFrame = OperatorSettings.GetNumFramesPerBlock() - 1;
+			const int32 NoTrigger = OperatorSettings.GetNumFramesPerBlock() << 1;
+
+			while (NextAudioFrame < LastAudioFrame)
+			{
+				// get the next Play and Stop indices
+				// (play)
+				if (PlayTrigIndex < NumPlayTrigs)
+				{
+					NextPlayFrame = (*Inputs.PlayTrigger)[PlayTrigIndex];
+				}
+				else
+				{
+					NextPlayFrame = NoTrigger;
+				}
+
+				// (stop)
+				if (StopTrigIndex < NumStopTrigs)
+				{
+					NextStopFrame = (*Inputs.StopTrigger)[StopTrigIndex];
+				}
+				else
+				{
+					NextStopFrame = NoTrigger;
+				}
+
+				// determine the next audio frame we are going to render up to
+				NextAudioFrame = FMath::Min(NextPlayFrame, NextStopFrame);
+
+				// no more triggers, rendering to the end of the block
+				if (NextAudioFrame == NoTrigger)
+				{
+					NextAudioFrame = OperatorSettings.GetNumFramesPerBlock();
+				}
+
+				// render audio (while loop handles looping audio)
+				while (CurrAudioFrame != NextAudioFrame)
+				{
+					if (bIsPlaying)
+					{
+						RenderFrameRange(CurrAudioFrame, NextAudioFrame);
+					}
+					else
+					{
+						ZeroFrameRange(CurrAudioFrame, NextAudioFrame);
+					}
+					CurrAudioFrame = NextAudioFrame;
+				}
+
+				// execute the next trigger
+				if (CurrAudioFrame == NextPlayFrame)
+				{
+					StartPlaying();
+					if (!bIsPlaying)
+					{
+						Outputs.OnFinished->TriggerFrame(CurrAudioFrame);
+					}
+
+					++PlayTrigIndex;
+				}
+
+				if (CurrAudioFrame == NextStopFrame)
+				{
+					bIsPlaying = false;
+					Outputs.OnFinished->TriggerFrame(CurrAudioFrame);
+					++StopTrigIndex;
+				}
+			}
+		}
+
+		void StartPlaying()
+		{
+			check(ProxyPlayer.IsValid());
+			CurrentWaveAsset = *Inputs.WaveAsset;
+			if (ProxyPlayer->SetSoundWave(CurrentWaveAsset.GetSoundWaveProxy()))
+			{
+				ProxyPlayer->SetSpeed(CalculateSpeedFromPitchShift());
+				ProxyPlayer->SetLoop(*Inputs.Loop, Inputs.LoopStartTime->GetSeconds(), Inputs.LoopDuration->GetSeconds());
+				ProxyPlayer->SeekToTime(Inputs.StartTime->GetSeconds());
+			}
+
+			bIsPlaying = ProxyPlayer.IsValid() && ProxyPlayer->IsPlayerValid() && !ProxyPlayer->IsFinished();
+		}
+
+		void RenderFrameRange(int32 StartFrame, int32 EndFrame)
+		{
+			check(ProxyPlayer.IsValid());
+			check(ProxyPlayer->IsPlayerValid());
+			check(bIsPlaying);
+			const int32 NumFramesToGenerate = EndFrame - StartFrame;
+			bDidRenderThisBlock = true;
+
+			TArray<FSoundWaveProxyPlayer::FSourceEvent> SourceEvents;
+			if (NumFramesToGenerate > 0)
+			{
+				Audio::FMultichannelBufferView OutputSlice = Audio::SliceMultichannelBufferView(MultiChannelOutputView, StartFrame, NumFramesToGenerate);
+				ProxyPlayer->RenderMultiChannelAudio(StartFrame, OutputSlice, SourceEvents);
+				bIsPlaying = !ProxyPlayer->IsFinished();
+			}
+
+			for (const FSoundWaveProxyPlayer::FSourceEvent& SourceEvent : SourceEvents)
+			{
+				switch (SourceEvent.Type)
+				{
+				case FSoundWaveProxyPlayer::FSourceEvent::Loop:
+					Outputs.OnLooped->TriggerFrame(SourceEvent.OutputFrameIndex);
+					break;
+				case FSoundWaveProxyPlayer::FSourceEvent::OnNearlyFinished:
+					Outputs.OnNearlyFinished->TriggerFrame(SourceEvent.OutputFrameIndex);
+					break;
+				case FSoundWaveProxyPlayer::FSourceEvent::OnFinished:
+					Outputs.OnFinished->TriggerFrame(SourceEvent.OutputFrameIndex);
+					break;
+				case FSoundWaveProxyPlayer::FSourceEvent::CuePoint:
+					{
+						check(ProxyPlayer->GetCuePoints().IsValidIndex(SourceEvent.CuePointIndex));
+						const FSoundWaveCuePoint& CuePoint = ProxyPlayer->GetCuePoints()[SourceEvent.CuePointIndex];
+						*Outputs.CuePointID = CuePoint.CuePointID;
+						*Outputs.CuePointLabel = CuePoint.Label;
+						Outputs.OnCuePoint->TriggerFrame(SourceEvent.OutputFrameIndex);
+					}
+					break;
+				}
+			}
+		}
+
+		void ZeroFrameRange(int32 StartFrame, int32 EndFrame)
+		{
+			const int32 NumFramesToZero = EndFrame - StartFrame;
+			if (NumFramesToZero > 0)
+			{
+				Audio::FMultichannelBufferView OutputSlice = Audio::SliceMultichannelBufferView(MultiChannelOutputView, StartFrame, NumFramesToZero);
+				for (TArrayView<float>& Audio : OutputSlice)
+				{
+					FMemory::Memset(Audio.GetData(), 0, Audio.Num() * Audio.GetTypeSize());
+				}
+			}	
+		}
+
+		float GetPitchShiftClamped() const
+		{
+			static constexpr float MaxAbsPitchShiftInOctaves = 6.0f;
+			return FMath::Clamp(*Inputs.PitchShift, -12.0f * MaxAbsPitchShiftInOctaves, 12.0f * MaxAbsPitchShiftInOctaves);
+		}
+
+		float CalculateSpeedFromPitchShift() const
+		{
+			return FMath::Pow(2.0f, GetPitchShiftClamped() / 12.0f);
+		}
+
+		FOperatorSettings OperatorSettings;
+		FNodeRenderCost CostReporter;
+		FInputs Inputs;
+		FOutputs Outputs;
+
+		bool bIsPlaying = false;
+		bool bDidRenderThisBlock = false;
+
+		FWaveAsset CurrentWaveAsset;
+		TUniquePtr<FSoundWaveProxyPlayer> ProxyPlayer;
+		TArray<TArrayView<float>> MultiChannelOutputView;
+	};
+
+	// using FMonoWavePlayerNodeV2 = Metasound::TNodeFacade<FWavePlayerOperatorV2>;
+	// METASOUND_REGISTER_NODE(FMonoWavePlayerNodeV2);
+
+	TUniquePtr<IOperator> Nodes::WavePlayerV2::CreateOperator(int32 NumChannels, const FBuildOperatorParams& InParams, FBuildResults& OutResults)
+	{
+		check(FWavePlayerOperatorV2::GetOutputAudioBufferNames(NumChannels).Num() == NumChannels);
+		
+		FWavePlayerOperatorV2::FInputs Inputs
+		{
+			InParams.InputData.GetOrCreateDefaultDataReadReference<FTrigger>(METASOUND_GET_PARAM_NAME(Inputs::Play), InParams.OperatorSettings),
+			InParams.InputData.GetOrCreateDefaultDataReadReference<FTrigger>(METASOUND_GET_PARAM_NAME(Inputs::Stop), InParams.OperatorSettings),
+			InParams.InputData.GetOrCreateDefaultDataReadReference<FWaveAsset>(METASOUND_GET_PARAM_NAME(Inputs::WaveAsset), InParams.OperatorSettings),
+			InParams.InputData.GetOrCreateDefaultDataReadReference<FTime>(METASOUND_GET_PARAM_NAME(Inputs::StartTime), InParams.OperatorSettings),
+			InParams.InputData.GetOrCreateDefaultDataReadReference<float>(METASOUND_GET_PARAM_NAME(Inputs::PitchShift), InParams.OperatorSettings),
+			InParams.InputData.GetOrCreateDefaultDataReadReference<bool>(METASOUND_GET_PARAM_NAME(Inputs::Loop), InParams.OperatorSettings),
+			InParams.InputData.GetOrCreateDefaultDataReadReference<FTime>(METASOUND_GET_PARAM_NAME(Inputs::LoopStartTime), InParams.OperatorSettings),
+			InParams.InputData.GetOrCreateDefaultDataReadReference<FTime>(METASOUND_GET_PARAM_NAME(Inputs::LoopDuration), InParams.OperatorSettings),
+			InParams.InputData.GetOrCreateDefaultValue<bool>(METASOUND_GET_PARAM_NAME(Inputs::MaintainAudioSync), InParams.OperatorSettings),
+		};
+
+		FWavePlayerOperatorV2::FOutputs Outputs(InParams.OperatorSettings, NumChannels);
+
+		return MakeUnique<FWavePlayerOperatorV2>(InParams, MoveTemp(Inputs), MoveTemp(Outputs));
+	}
+}
+
+#undef LOCTEXT_NAMESPACE
\ No newline at end of file
diff --git a/Engine/Plugins/Runtime/Metasound/Source/MetasoundEngine/Private/MetasoundWavePlayerNodeV2.h b/Engine/Plugins/Runtime/Metasound/Source/MetasoundEngine/Private/MetasoundWavePlayerNodeV2.h
new file mode 100644
--- /dev/null
+++ b/Engine/Plugins/Runtime/Metasound/Source/MetasoundEngine/Private/MetasoundWavePlayerNodeV2.h
@@ -0,0 +1,18 @@
+// Copyright Epic Games, Inc. All Rights Reserved.
+
+#pragma once
+
+#include "MetasoundOperatorInterface.h"
+#include "Templates/UniquePtr.h"
+#include "HAL/Platform.h"
+
+namespace Metasound
+{
+	struct FBuildOperatorParams;
+	struct FBuildResults;
+	
+	namespace Nodes::WavePlayerV2
+	{
+		TUniquePtr<IOperator> CreateOperator(int32 NumChannels, const FBuildOperatorParams& InParams, FBuildResults& OutResults);
+	}
+}
diff --git a/Engine/Source/Runtime/Engine/Classes/Sound/SoundWaveProxyPlayer.h b/Engine/Source/Runtime/Engine/Classes/Sound/SoundWaveProxyPlayer.h
new file mode 100644
--- /dev/null
+++ b/Engine/Source/Runtime/Engine/Classes/Sound/SoundWaveProxyPlayer.h
@@ -0,0 +1,355 @@
+// Copyright Epic Games, Inc. All Rights Reserved.
+
+#pragma once
+
+#include "AudioDecompress.h"
+#include "Containers/Array.h"
+#include "Containers/ArrayView.h"
+#include "DSP/ConvertDeinterleave.h"
+#include "DSP/RuntimeResampler.h"
+#include "HAL/Platform.h"
+
+// Forward declare
+class FSoundWaveProxy;
+class ICompressedAudioInfo;
+
+#define UE_API ENGINE_API
+
+class UE_INTERNAL FSoundWaveProxyPlayer
+{
+public:
+	using FSoundWaveProxyRef = TSharedRef<FSoundWaveProxy, ESPMode::ThreadSafe>;
+
+	struct FSettings
+	{
+		FSettings(int32 InNumFramesPerBlock, float InOutputSampleRate, int32 InOutputNumChannels)
+			: OutputSampleRate(InOutputSampleRate)
+			, OutputNumChannels(InOutputNumChannels)
+			, NumFramesPerBlock(InNumFramesPerBlock)
+		{
+			check(OutputSampleRate > 0.0f);
+			check(OutputNumChannels > 0);
+			check(OutputNumChannels <= 8);
+		}
+		
+		float OutputSampleRate = 0.0f;
+		int32 OutputNumChannels = 0;
+		int32 NumFramesPerBlock = 0;
+		int32 MaxDecodeSizeFrames = 1024;
+		bool bMaintainAudioSync = false;
+	};
+
+	struct FSourceEvent
+	{
+		enum EType
+		{
+			CuePoint,
+			Loop,
+			OnFinished,
+			OnNearlyFinished
+		};
+		
+		FSourceEvent(EType InType, int32 InSourceFrameIndex, int32 InOutputFrameIndex, int32 InCuePointIndex = INDEX_NONE)
+			: Type(InType)
+			, SourceFrameIndex(InSourceFrameIndex)
+			, OutputFrameIndex(InOutputFrameIndex)
+			, CuePointIndex(InCuePointIndex)
+		{
+		}
+
+		EType Type = Loop;
+
+		// The Source Frame into the Source Sound Wave
+		int32 SourceFrameIndex = INDEX_NONE;
+
+		// The Output Frame Index relative to the render block
+		int32 OutputFrameIndex = INDEX_NONE;
+
+		// For Cue Point Events only - the index into the sorted cue point
+		// default to INDEX_NONE 
+		int32 CuePointIndex = INDEX_NONE;
+	};
+
+	~FSoundWaveProxyPlayer() {}
+
+	/**
+	 * Create a SoundWaveProxyPlayer initialized with the given output settings
+	 * (Note: Settings can be changed after creation, so no need to create a new one if you need to change the settings)
+	 * 
+	 * @param InSettings 
+	 * @return a UniquePtr to a new WaveProxyPlayer
+	 */
+	UE_API static TUniquePtr<FSoundWaveProxyPlayer> Create(const FSettings& InSettings);
+
+	/**
+	 * Set the SoundWave Proxy to play
+	 * 
+	 * if Null, will reset the Proxy Player and return 'false'
+	 * if the WaveProxy has invalid data or is unabled to be decoded,
+	 * the player will be reset and this will return 'false'
+	 * 
+	 * @param InWaveProxy 
+	 * @return whether the player is valid for playback
+	 */
+	UE_API bool SetSoundWave(FSoundWaveProxyPtr InWaveProxy);
+
+	/**
+	 * 
+	 * @return the number of channels this player will render out to
+	 */
+	UE_API int32 GetOutputNumChannels() const;
+
+	/**
+	 *
+	 * @return the sample rate this player will render out to in Hz 
+	 */
+	UE_API float GetOutputSampleRate() const;
+
+	/**
+	 * 
+	 * @return whether this player has a valid Sound Wave and Decoder
+	 */
+	UE_API bool IsPlayerValid() const;
+	
+	/**
+	 * Reset the Wave Player with a new Wave Proxy and new settings
+	 * Helpful to avoid creating a whole new Wave Player with new buffers etc.
+	 * 
+	 * @param NewSettings
+	 */
+	UE_API void Reset(const FSettings& NewSettings);
+
+	/**
+	 * Render audio from the Wave Source based on the current Speed, Loop Settings, and Desired output Sample Rate
+	 * Expects the number of channels in the output buffer be the same as the initial settings
+	 *
+	 * Asserts the SoundWave is set
+	 * 
+	 * @param OutMultiChannelAudio 
+	 */
+	UE_API void RenderMultiChannelAudio(Audio::FMultichannelBufferView& OutMultiChannelAudio);
+
+	/**
+	 * Render audio from the Wave Source based on the current Speed, Loop Settings, and Desired output Sample Rate
+	 * Expects the number of channels in the output buffer be the same as the initial settings
+	 *
+	 * Also, Generate a list of source events (Cue Points & Loop Points) based on the current location of the player
+	 *
+	 * Output Event Output Frames will be relative to the given Block Frame Index (accounts for resampling)
+	 * Output Event Source Frames will be in absolute Frames into the Source Sound (no resampling)
+	 *
+	 * Asserts the SoundWave is set
+	 * 
+	 * @param BlockFrameIndex The Current BlockFrameIndex for this render block
+	 * @param OutMultiChannelAudio The output audio buffer view
+	 * @param OutEvents the output source events triggered during this render
+	 */
+	UE_API void RenderMultiChannelAudio(int32 BlockFrameIndex, Audio::FMultichannelBufferView& OutMultiChannelAudio, TArray<FSourceEvent>& OutEvents);
+
+	/**
+	 * Directly generate audio from the Wave Source into an interleaved buffer
+	 *
+	 * This will update the current playback state of the wave player
+	 *
+	 * will honor loop settings
+	 *
+	 * Asserts that the SoundWave is set
+	 * 
+	 * @param OutAudio 
+	 * @return The number of frames generated
+	 */
+	UE_API int32 GenerateSourceAudio(TArrayView<float> OutAudio);
+
+	/**
+	 * Set the speed of the wave player
+	 * clamped to MinSpeed & MaxSpeed
+	 * @param InSpeed 
+	 */
+	UE_API void SetSpeed(float InSpeed);
+
+	/**
+	 * @return the current speed of the wave player
+	 */
+	UE_API float GetSpeed() const;
+
+	/**
+	 * Seek this Wave Player to the specified time
+	 * Can fail if the wave player is not initialized or the wave asset does not properly support seeking
+	 * 
+	 * @param TimeSeconds 
+	 * @return whether the wave player successfully seeked to the desired time
+	 */
+	UE_API bool SeekToTime(float TimeSeconds);
+
+	/**
+	 * Seek this wave player to a specified frame.
+	 * Can fail if the wave player is not initialized or the wave asset does not properly support seeking
+	 * 
+	 * @param InFrame 
+	 * @return whether the wave player successfully seeked to the desired frame
+	 */
+	UE_API bool SeekToFrame(int32 InFrame);
+
+	/**
+	 * Set whether this wave player should loop with optional loop parameters
+	 * 
+	 * @param InIsLooping 
+	 * @param InLoopStartTimeSeconds 
+	 * @param InLoopDurationSeconds 
+	 */
+	UE_API void SetLoop(bool InIsLooping, float InLoopStartTimeSeconds = 0.0f, float InLoopDurationSeconds = -1.0f);
+
+	/**
+	 * @return whether this wave player is set to looping
+	 */
+	UE_API bool IsLooping() const;
+
+	/**
+	 * 
+	 * @return true if the Player has reached the end of the SoundWave
+	 */
+	UE_API bool IsFinished() const;
+	
+	/**
+	 * Returns the current playback position into the Source 
+	 *
+	 * @return Current Playback Position in Seconds
+	 */
+	UE_API float GetCurrentPlaybackTimeSeconds() const;
+
+	/**
+	 * Returns the current playback progress based on the playback time into the source file
+	 *
+	 * @return PlayProgress in the range [0.0, 1.0]
+	 */
+	UE_API float GetPlaybackProgress() const;
+
+	/**
+	 * Returns the current loop progress between the current loop region
+	 * based on the playback time in the source file
+	 *
+	 * @return LoopProgress in the range [0.0, 1.0]
+	 */
+	UE_API float GetLoopProgress() const;
+
+	/**
+	 * Get the Current Cue Points as a sorted array
+	 * If the Sound Wave is not valid, cue points will be empty
+	 * 
+	 * @return Sorted Array of Cue Points
+	 */
+	UE_API const TArray<FSoundWaveCuePoint>& GetCuePoints() const;
+
+private:
+	
+	/**
+	 * Intentionally Private constructor. call the static method "Create" to make a UniquePtr to this object
+	 * 
+	 * @param InSettings 
+	 */
+	FSoundWaveProxyPlayer(const FSettings& InSettings);
+
+	/**
+	 * Check wether the Wave Proxy has valid data (valid Sample Rate, Num Channels, Num Frames)
+	 * 
+	 * @param InWaveProxy 
+	 * @return whether the WaveProxy is valid
+	 */
+	static bool IsValidSoundWaveProxy(FSoundWaveProxyRef InWaveProxy);
+
+	// 
+	/**
+	 * performs the actual streaming/reading of the audio from the decoder into an interleaved output buffer
+	 *
+	 * Called by GenerateSourceAudio to generate subblock range of audio
+	 * 
+	 * @param OutAudio 
+	 * @return 
+	 */
+	int32 GenerateSourceAudioInternal(TArrayView<float> OutAudio);
+
+	/**
+	 * Generate a list of source events (Cue Points & Loop Points) given a current Source Start Frame Index
+	 *
+	 * Called immediately before RenderMultiChannelAudio
+	 *
+	 * @param BlockFrameIndex 
+	 * @param NumFrames 
+	 * @param OutEvents 
+	 */
+	void GenerateSourceEvents(int32 BlockFrameIndex, int32 NumFrames, TArray<FSourceEvent>& OutEvents);
+
+	/**
+	 * Generate a list of source events (Cue Points & Loop Points) given a current Source Start Frame Index
+	 *
+	 * Used by GenerateSourceEvents to generate subblock ranges of events
+	 *
+	 * @param InSourceStartFrameIndex 
+	 * @param InOutputStartFrameIndex 
+	 * @param NumFrames 
+	 * @param OutEvents 
+	 */
+	void GenerateSourceEventsInternal(int32 InSourceStartFrameIndex, int32 InOutputStartFrameIndex, int32 NumFrames, TArray<FSourceEvent>& OutEvents);
+	
+	/**
+	 * "seek forward" by discarding the requested number of frames from the decoder and advancing the "source" frame
+	 * does NOT update "audio sync" frame
+	 * 
+	 * @param NumFrames 
+	 * @return Number of frames that were actually seeked
+	 */
+	int32 SimulateSeekForward(int32 NumFrames);
+
+	/**
+	 * Calculate the current frame ratio based on Speed & Resampling (if needed)
+	 * 
+	 * @return Current Frame Ratio
+	 */
+	float GetCurrentFrameRatio() const;
+
+	/**
+	 * Initialize the Decoder with the given wave proxy
+	 * also initialize the decode buffers
+	 * @param WaveProxy
+	 * @param MaxDecodeSizeInFrames 
+	 * @return whether the decoder was successfully initialized
+	 */
+	bool InitializeDecoder(const FSoundWaveProxyRef& InWaveProxy, int32 MaxDecodeSizeInFrames);
+
+	FSoundWaveProxyPtr SoundWave;
+	FSettings Settings;
+	TUniquePtr<ICompressedAudioInfo> CompressedAudioInfo;
+	TUniquePtr<Audio::IConvertDeinterleave> ConvertDeinterleave;
+	Audio::FRuntimeResampler RuntimeResampler;
+
+	TArray<int16, Audio::FAudioBufferAlignedAllocator> DecodeBuffer;
+	Audio::FAlignedFloatBuffer SourceAudio;
+	Audio::FAlignedFloatBuffer ResampledAudio;
+
+	// source wave data
+	float DurationSeconds = 0.0f;
+	float SampleRate = 0.0f;
+	int32 NumChannels = 0;
+	int32 NumFramesInWave = 0;
+	TArray<FSoundWaveCuePoint> SortedCuePoints;
+
+	// playback
+	int32 SourceFrameIndex = 0;
+	int32 AudioSyncFrameIndex = 0;
+	float Speed = 1.0f;
+	bool bIsLooping = false;
+	float CurrentLoopStartTime = 0.0f;
+	float CurrentLoopDuration = -1.0f;
+	int32 LoopStartFrameIndex = 0;
+	int32 LoopEndFrameIndex = -1;
+	bool bRequiresResampling = false;
+	bool bHasNearlyFinished = false;
+
+	static constexpr float MinLoopDurationSeconds = 0.05f;
+	static constexpr float MinSpeed = 0.01f;
+	static constexpr float MaxSpeed = 100.0f;
+	
+	bool bFallbackSeekMethodWarningLogged = false;
+};
+
+#undef UE_API
\ No newline at end of file
diff --git a/Engine/Source/Runtime/Engine/Private/SoundWaveProxyPlayer.cpp b/Engine/Source/Runtime/Engine/Private/SoundWaveProxyPlayer.cpp
new file mode 100644
--- /dev/null
+++ b/Engine/Source/Runtime/Engine/Private/SoundWaveProxyPlayer.cpp
@@ -0,0 +1,786 @@
+// Copyright Epic Games, Inc. All Rights Reserved.
+
+#include "Sound/SoundWaveProxyPlayer.h"
+
+#include "Audio.h"
+#include "DSP/FloatArrayMath.h"
+#include "Interfaces/IAudioFormat.h"
+#include "ProfilingDebugging/CpuProfilerTrace.h"
+#include "Templates/SharedPointer.h"
+
+int32 ConformDecodeSize(int32 InMaxDesiredDecodeSizeInFrames)
+{
+	/** Minimum number of frames to decode per a call to the decoder.  */
+	static constexpr int32 DefaultMinDecodeSizeInFrames = 128;
+
+	/** Some codecs have strict requirements on decode size. In order to be
+	 * functional with all supported codecs, the decode size must be a multiple
+	 * of 128.
+	 */
+	static constexpr int32 DecodeSizeQuantizationInFrames = 128;
+	
+	static_assert(DefaultMinDecodeSizeInFrames >= DecodeSizeQuantizationInFrames, "Min decode size less than decode size quantization");
+	static_assert((DefaultMinDecodeSizeInFrames % DecodeSizeQuantizationInFrames) == 0, "Min decode size must be equally divisible by decode size quantization");
+
+	InMaxDesiredDecodeSizeInFrames = FMath::Max(InMaxDesiredDecodeSizeInFrames, DefaultMinDecodeSizeInFrames);
+	const int32 QuantizedDecodeSizeInFrames = (InMaxDesiredDecodeSizeInFrames / DecodeSizeQuantizationInFrames) * DecodeSizeQuantizationInFrames;
+	const int32 ConformedDecodeSizeInFrames = FMath::Max(QuantizedDecodeSizeInFrames, DefaultMinDecodeSizeInFrames);
+
+	check((ConformedDecodeSizeInFrames % QuantizedDecodeSizeInFrames) == 0);
+	check(ConformedDecodeSizeInFrames > 0);
+
+	return ConformedDecodeSizeInFrames;
+}
+
+FSoundWaveProxyPlayer::FSoundWaveProxyPlayer(const FSettings& InSettings)
+	: Settings(InSettings)
+	, RuntimeResampler(1)
+{
+
+}
+
+bool FSoundWaveProxyPlayer::SetSoundWave(FSoundWaveProxyPtr InWaveProxy)
+{
+	if (!InWaveProxy.IsValid())
+	{
+		Reset(Settings);
+		return false;
+	}
+	
+	if (!IsValidSoundWaveProxy(InWaveProxy.ToSharedRef()))
+	{
+		UE_LOG(LogAudio, Warning, TEXT("Failed to Set Sound Wave in FSoundWaveProxyPlayer(), Package: %s"), *WriteToString<64>(InWaveProxy->GetPackageName()));
+		Reset(Settings);
+		return false;
+	}
+
+	if (!InitializeDecoder(InWaveProxy.ToSharedRef(), Settings.MaxDecodeSizeFrames))
+	{
+		UE_LOG(LogAudio, Warning, TEXT("Failed to InitializeDecoder in FSoundWaveProxyPlayer(), Package: %s"), *WriteToString<64>(InWaveProxy->GetPackageName()));
+		Reset(Settings);
+		return false;
+	}
+	
+	SoundWave = InWaveProxy;
+	SourceFrameIndex = 0;
+	AudioSyncFrameIndex = 0;
+	// update the loop with the cached loop settings to update the new
+	SetLoop(bIsLooping, CurrentLoopStartTime, CurrentLoopDuration);
+	bFallbackSeekMethodWarningLogged = false;
+	bRequiresResampling = false;
+	bHasNearlyFinished = false;
+	SortedCuePoints = SoundWave->GetCuePointsAndLoopRegions();
+	Algo::SortBy(SortedCuePoints, &FSoundWaveCuePoint::FramePosition);
+
+	RuntimeResampler.Reset(NumChannels);
+	
+	// Create algorithm for channel conversion and deinterleave 
+	Audio::FConvertDeinterleaveParams ConvertDeinterleaveParams;
+	ConvertDeinterleaveParams.NumInputChannels = NumChannels;
+	ConvertDeinterleaveParams.NumOutputChannels = Settings.OutputNumChannels;
+	// Original implementation of MetaSound WavePlayer upmixed 
+	// mono using FullVolume. In the future, the mono upmix 
+	// method may be exposed as a node input to facilitate 
+	// better control.
+	ConvertDeinterleaveParams.MonoUpmixMethod = Audio::EChannelMapMonoUpmixMethod::FullVolume;
+	ConvertDeinterleave = Audio::IConvertDeinterleave::Create(ConvertDeinterleaveParams);
+
+	return true;
+}
+
+int32 FSoundWaveProxyPlayer::GetOutputNumChannels() const
+{
+	return Settings.OutputNumChannels;
+}
+
+float FSoundWaveProxyPlayer::GetOutputSampleRate() const
+{
+	return Settings.OutputSampleRate;
+}
+
+bool FSoundWaveProxyPlayer::IsPlayerValid() const
+{
+	return SoundWave.IsValid() && CompressedAudioInfo.IsValid();
+}
+
+void FSoundWaveProxyPlayer::Reset(const FSettings& NewSettings)
+{
+	SoundWave.Reset();
+	Settings = NewSettings;
+	DecodeBuffer.Reset();
+	CompressedAudioInfo.Reset();
+	ConvertDeinterleave.Reset();
+	SourceAudio.Reset();
+	ResampledAudio.Reset();
+	SourceFrameIndex = 0;
+	AudioSyncFrameIndex = 0;
+	Speed = 1.0f;
+	bIsLooping = false;
+	LoopStartFrameIndex = 0;
+	LoopEndFrameIndex = -1;
+	bRequiresResampling = false;
+	bFallbackSeekMethodWarningLogged = false;
+	bHasNearlyFinished = false;
+
+	SampleRate = -1.0f;
+	NumChannels = -1;
+	NumFramesInWave = -1;
+	DurationSeconds = -1.0f;
+	SortedCuePoints.Reset();
+}
+
+TUniquePtr<FSoundWaveProxyPlayer> FSoundWaveProxyPlayer::Create(const FSettings& InSettings)
+{
+	return TUniquePtr<FSoundWaveProxyPlayer>(new FSoundWaveProxyPlayer(InSettings));
+}
+
+bool FSoundWaveProxyPlayer::IsValidSoundWaveProxy(FSoundWaveProxyRef InWaveProxy)
+{
+	if (InWaveProxy->GetSampleRate() <= 0.f)
+	{
+		UE_LOG(LogAudio, Warning, TEXT("Cannot create FSoundWaveProxyPlayer due to invalid sample rate (%f). Package: %s"), InWaveProxy->GetSampleRate(), *InWaveProxy->GetPackageName().ToString());
+		return false;
+	}
+
+	if (InWaveProxy->GetNumChannels() <= 0)
+	{
+		UE_LOG(LogAudio, Warning, TEXT("Cannot create FSoundWaveProxyPlayer due to invalid num channels (%d). Package: %s"), InWaveProxy->GetNumChannels(), *InWaveProxy->GetPackageName().ToString());
+		return false;
+	}
+
+	if (InWaveProxy->GetNumFrames() <= 0)
+	{
+		UE_LOG(LogAudio, Warning, TEXT("Cannot create FSoundWaveProxyPlayer due to invalid num frames (%d). Package: %s"), InWaveProxy->GetNumFrames(), *InWaveProxy->GetPackageName().ToString());
+		return  false;
+	}
+
+	return true;
+}
+
+void FSoundWaveProxyPlayer::GenerateSourceEvents(int32 BlockFrameIndex, int32 NumFrames, TArray<FSourceEvent>& OutEvents)
+{
+	check(SoundWave.IsValid());
+
+	if (IsFinished())
+	{
+		return;
+	}
+
+	float FrameRatio = GetCurrentFrameRatio();
+	if (bRequiresResampling || !FMath::IsNearlyEqual(FrameRatio, 1.0f))
+	{
+		bRequiresResampling = true;
+		RuntimeResampler.SetFrameRatio(FrameRatio);
+		NumFrames = RuntimeResampler.GetNumInputFramesNeededToProduceOutputFrames(NumFrames);
+	}
+	
+	int32 StartFrameIndex = Settings.bMaintainAudioSync ? AudioSyncFrameIndex : SourceFrameIndex;
+
+	if (!bIsLooping)
+	{
+		if (StartFrameIndex >= NumFramesInWave)
+		{
+			return;
+		}
+		
+		GenerateSourceEventsInternal(StartFrameIndex, BlockFrameIndex, NumFrames, OutEvents);
+
+		// check if we're finished
+		if (StartFrameIndex + NumFrames >= NumFramesInWave)
+		{
+			int32 OutputFrameIndex = NumFramesInWave - StartFrameIndex - 1;
+			if (bRequiresResampling)
+			{
+				OutputFrameIndex = RuntimeResampler.GetNumOutputFramesProducedByInputFrames(OutputFrameIndex);
+			}
+			OutputFrameIndex += BlockFrameIndex;
+			OutEvents.Add(FSourceEvent(FSourceEvent::OnFinished, NumFramesInWave - 1, OutputFrameIndex));
+		}
+
+		// check if we'll finish in the next render, assuming the NumFrames we'll render next time wil be the same
+		// note that we're using the actual number of requested input frames based on the Resampler
+		// so in most cases, this should trigger exactly one block before "OnFinished" is triggered
+		// In some cases, it can render a few blocks before "OnFinished" if the Frame Ratio changes between renders
+		// It's also possible for this event to trigger in the same block as "OnFinished" in cases where the
+		// Sound can be rendered in a single block, due to either a high speed setting or a tiny wave file
+		if (!bHasNearlyFinished && StartFrameIndex + NumFrames * 2 >= NumFramesInWave)
+		{
+			// trigger "On nearly finished" at the beginning of this block
+			OutEvents.Add(FSourceEvent(FSourceEvent::OnNearlyFinished, StartFrameIndex, BlockFrameIndex));
+
+			// set this flag to avoid re-triggering in case the frame ratio changes between now and the next render
+			bHasNearlyFinished = true;
+		}
+
+		return;
+	}
+
+	if (StartFrameIndex >= LoopEndFrameIndex)
+	{
+		StartFrameIndex = LoopStartFrameIndex;
+		OutEvents.Add(FSourceEvent(FSourceEvent::Loop, LoopEndFrameIndex, BlockFrameIndex));
+	}
+
+	bool bLoopingThisBlock = (StartFrameIndex + NumFrames) > LoopEndFrameIndex;
+
+	if (!bLoopingThisBlock)
+	{
+		GenerateSourceEventsInternal(StartFrameIndex, BlockFrameIndex, NumFrames, OutEvents);
+		return;
+	}
+
+	int32 FirstNumFrames = LoopEndFrameIndex - StartFrameIndex;
+	GenerateSourceEventsInternal(StartFrameIndex, BlockFrameIndex, FirstNumFrames, OutEvents);
+
+	int32 OutputFrameIndex = LoopEndFrameIndex - StartFrameIndex - 1;
+	if (bRequiresResampling)
+	{
+		OutputFrameIndex = RuntimeResampler.GetNumOutputFramesProducedByInputFrames(OutputFrameIndex);
+	}
+	OutputFrameIndex += BlockFrameIndex;
+	OutEvents.Add(FSourceEvent(FSourceEvent::Loop, LoopEndFrameIndex, OutputFrameIndex));
+	int32 SecondNumFrames = NumFrames - FirstNumFrames;
+	GenerateSourceEventsInternal(LoopStartFrameIndex, OutputFrameIndex, SecondNumFrames, OutEvents);
+}
+
+void FSoundWaveProxyPlayer::GenerateSourceEventsInternal(int32 InSourceStartFrameIndex, int32 InOutputStartFrameIndex, int32 NumFrames, TArray<FSourceEvent>& OutEvents)
+{
+	if (SortedCuePoints.IsEmpty())
+	{
+		return;
+	}
+	
+	int32 EndFrameIndex = InSourceStartFrameIndex + NumFrames;
+	int32 CuePointIndex = Algo::LowerBoundBy(SortedCuePoints, InSourceStartFrameIndex, &FSoundWaveCuePoint::FramePosition);
+	for ( ;SortedCuePoints.IsValidIndex(CuePointIndex); ++CuePointIndex)
+	{
+		int32 CuePointFrameIndex = SortedCuePoints[CuePointIndex].FramePosition;
+		if (CuePointFrameIndex >= EndFrameIndex)
+		{
+			break;
+		}
+
+		int32 OutputFrameIndex = CuePointFrameIndex - InSourceStartFrameIndex;
+		if (bRequiresResampling)
+		{
+			OutputFrameIndex = RuntimeResampler.GetNumOutputFramesProducedByInputFrames(OutputFrameIndex);
+		}
+		OutputFrameIndex += InOutputStartFrameIndex;
+		OutEvents.Add(FSourceEvent(FSourceEvent::CuePoint, CuePointFrameIndex, OutputFrameIndex, CuePointIndex));
+	}
+}
+
+void FSoundWaveProxyPlayer::RenderMultiChannelAudio(int32 BlockFrameIndex, Audio::FMultichannelBufferView& OutMultiChannelAudio, TArray<FSourceEvent>& OutEvents)
+{
+	const int32 NumFrames = Audio::GetMultichannelBufferNumFrames(OutMultiChannelAudio);
+	GenerateSourceEvents(BlockFrameIndex, NumFrames, OutEvents);
+	RenderMultiChannelAudio(OutMultiChannelAudio);
+}
+
+void FSoundWaveProxyPlayer::RenderMultiChannelAudio(Audio::FMultichannelBufferView& OutMultiChannelAudio)
+{
+	check(SoundWave.IsValid());
+	
+	if (!ensure(Settings.OutputNumChannels == OutMultiChannelAudio.Num()))
+	{
+		return;
+	}
+	
+	if (IsFinished())
+	{
+		for (TArrayView<float>& Audio : OutMultiChannelAudio)
+		{
+			FMemory::Memset(Audio.GetData(), 0, Audio.Num() * Audio.GetTypeSize());
+		}
+		return;
+	}
+	
+	const float FrameRatio = GetCurrentFrameRatio();
+
+	const int32 OutputNumFrames = Audio::GetMultichannelBufferNumFrames(OutMultiChannelAudio);
+	if (!bRequiresResampling && FMath::IsNearlyEqual(FrameRatio, 1.0f))
+	{
+		if (1 == NumChannels && NumChannels == Settings.OutputNumChannels)
+		{
+			GenerateSourceAudio(OutMultiChannelAudio[0]);
+		}
+		else
+		{
+			
+			SourceAudio.SetNum(OutputNumFrames * NumChannels);
+			GenerateSourceAudio(SourceAudio);
+			ConvertDeinterleave->ProcessAudio(SourceAudio, OutMultiChannelAudio);
+		}
+	}
+	else
+	{
+		// once we've started resampling, we'll need to continue resampling to maintain continuity
+		// even when the frame ratio ever goes back to 1
+		bRequiresResampling = true;
+		
+		// resampling requires us to generate the source audio into a separate buffer
+		RuntimeResampler.SetFrameRatio(FrameRatio);
+		int32 SourceNumFrames = RuntimeResampler.GetNumInputFramesNeededToProduceOutputFrames(OutputNumFrames);
+
+		// include the samples we used from the last render
+		int32 LastNumSamples = SourceAudio.Num();
+		SourceAudio.SetNum(SourceNumFrames * NumChannels);
+		GenerateSourceAudio(MakeArrayView(SourceAudio.GetData() + LastNumSamples, SourceAudio.Num() - LastNumSamples));
+
+		int32 NumFramesConsumed = 0;
+		int32 NumFramesProduced = 0;
+		if (1 == NumChannels && NumChannels == Settings.OutputNumChannels)
+		{
+			// Mono -> Mono
+			// write the output of the resampler directly to the output
+			RuntimeResampler.ProcessInterleaved(SourceAudio, OutMultiChannelAudio[0], NumFramesConsumed, NumFramesProduced);
+		}
+		else
+		{
+			// other sources require down/up mixing using the deinterleave converter
+			ResampledAudio.SetNum(OutputNumFrames * NumChannels);
+			RuntimeResampler.ProcessInterleaved(SourceAudio, ResampledAudio, NumFramesConsumed, NumFramesProduced);
+			ConvertDeinterleave->ProcessAudio(ResampledAudio, OutMultiChannelAudio);
+		}
+
+		int32 NumSamplesConsumed = NumFramesConsumed * NumChannels;
+		if (NumSamplesConsumed < SourceAudio.Num())
+		{
+			// move the last frame of samples to the start of the SourceBuffer for future processing
+			int32 NumSamplesCopy = SourceAudio.Num() - NumSamplesConsumed;
+			FMemory::Memcpy(SourceAudio.GetData(), SourceAudio.GetData() + NumSamplesConsumed, NumSamplesCopy * sizeof(float));
+			SourceAudio.SetNum(NumSamplesCopy);
+		}
+		else
+		{
+			SourceAudio.SetNum(0);
+		}
+	}
+}
+
+int32 FSoundWaveProxyPlayer::GenerateSourceAudio(TArrayView<float> OutAudio)
+{
+	TRACE_CPUPROFILER_EVENT_SCOPE(FSoundWaveProxyPlayer::GenerateSourceAudio);
+
+	check(SoundWave.IsValid());
+
+	int32 NumFramesRequested = OutAudio.Num() / NumChannels;
+
+	// check to see if the source audio is in sync 
+	if (Settings.bMaintainAudioSync && AudioSyncFrameIndex != SourceFrameIndex)
+	{
+		// assert that our source frame is never ahead of the sync frame
+		check(AudioSyncFrameIndex > SourceFrameIndex);
+		// try to force the source audio to be in sync
+		int32 NumSyncFrames = AudioSyncFrameIndex - SourceFrameIndex;
+		int32 NumFramesSeeked = SimulateSeekForward(NumSyncFrames);
+		if (NumSyncFrames == NumFramesSeeked)
+		{
+			// it should be in sync now
+			check(AudioSyncFrameIndex == SourceFrameIndex);
+		}
+	}
+
+	if (!bIsLooping)
+	{
+		int32 NumFramesGenerated = GenerateSourceAudioInternal(OutAudio);
+		return NumFramesGenerated;
+	}
+
+	if (SourceFrameIndex >= LoopEndFrameIndex)
+	{
+		SeekToFrame(LoopStartFrameIndex);
+	}
+
+	bool bLoopingThisChunk = (SourceFrameIndex + NumFramesRequested) > LoopEndFrameIndex;
+
+	if (!bLoopingThisChunk)
+	{
+		int32 NumFramesGenerated = GenerateSourceAudioInternal(OutAudio);
+		return NumFramesGenerated;
+	}
+	
+	int32 NumFramesRemaining = NumFramesRequested;
+	int32 SampleOffset = 0;
+	while (NumFramesRemaining > 0)
+	{
+		int32 NumFrames = FMath::Min(LoopEndFrameIndex - SourceFrameIndex, NumFramesRemaining);
+		int32 NumSamples = NumFrames * NumChannels;
+		GenerateSourceAudioInternal(OutAudio.Slice(SampleOffset, NumSamples));
+		SampleOffset += NumSamples;
+		NumFramesRemaining -= NumFrames;
+		if (NumFramesRemaining > 0)
+		{
+			SeekToFrame(LoopStartFrameIndex);
+		}
+	}
+	return NumFramesRequested;
+}
+
+int32 FSoundWaveProxyPlayer::GenerateSourceAudioInternal(TArrayView<float> OutAudio)
+{
+	check(CompressedAudioInfo.IsValid());
+	check(DecodeBuffer.Num() > 0);
+	check(OutAudio.Num() % NumChannels == 0);
+	const int32 NumFramesTotal = FMath::Min(OutAudio.Num() / NumChannels, NumFramesInWave - SourceFrameIndex);
+	int32 NumFramesRemaining = NumFramesTotal;
+	int32 SampleOffset = 0;
+	uint32 MaxBufferSize = DecodeBuffer.Num() * DecodeBuffer.GetTypeSize();
+	uint32 DesiredBufferSize = OutAudio.Num() * DecodeBuffer.GetTypeSize();
+	uint32 BufferSize = FMath::Min(DesiredBufferSize, MaxBufferSize);
+	uint8* ByteBuffer = reinterpret_cast<uint8*>(DecodeBuffer.GetData());
+
+	bool bFinished = false;
+	while (!bFinished && NumFramesRemaining > 0)
+	{
+		int32 NumBytesStreamed = 0;
+		if (SoundWave->IsStreaming())
+		{
+			bFinished = CompressedAudioInfo->StreamCompressedData(ByteBuffer, false, BufferSize, NumBytesStreamed);
+		}
+		else
+		{
+			bFinished = CompressedAudioInfo->ReadCompressedData(ByteBuffer, false, BufferSize);
+			NumBytesStreamed = BufferSize;
+		}
+
+		if (NumBytesStreamed == 0)
+		{
+			break;
+		}
+
+		const int32 NumSamplesStreamed = NumBytesStreamed / sizeof(int16);
+		const int32 NumFramesStreamed = NumSamplesStreamed / NumChannels;
+		
+		Audio::ArrayPcm16ToFloat(MakeArrayView(DecodeBuffer.GetData(), NumSamplesStreamed), MakeArrayView(OutAudio.GetData() + SampleOffset, NumSamplesStreamed));
+		SampleOffset += NumSamplesStreamed;
+		NumFramesRemaining -= NumFramesStreamed;
+		DesiredBufferSize -= NumBytesStreamed;
+		BufferSize = FMath::Min(DesiredBufferSize, MaxBufferSize);
+	}
+
+	// zero unset samples
+	if (SampleOffset < OutAudio.Num())
+	{
+		FMemory::Memset(OutAudio.GetData() + SampleOffset, 0, (OutAudio.Num() - SampleOffset) * sizeof(float));
+	}
+	int32 NumFramesGenerated = NumFramesTotal - NumFramesRemaining;
+	SourceFrameIndex = FMath::Min(SourceFrameIndex + NumFramesGenerated, NumFramesInWave);
+	AudioSyncFrameIndex = FMath::Min(AudioSyncFrameIndex + NumFramesTotal, NumFramesInWave);
+	return NumFramesGenerated;
+}
+
+int32 FSoundWaveProxyPlayer::SimulateSeekForward(int32 NumFrames)
+{
+	int32 NumFramesRemaining = NumFrames;
+	int32 SampleOffset = 0;
+	uint32 MaxBufferSize = NumFrames * NumChannels * DecodeBuffer.GetTypeSize();
+	uint32 BufferSize = FMath::Min(DecodeBuffer.Num() * DecodeBuffer.GetTypeSize(), MaxBufferSize);
+	uint8* ByteBuffer = (uint8*)DecodeBuffer.GetData();
+
+	bool bFinished = false;
+	while (!bFinished && NumFramesRemaining > 0)
+	{
+		int32 NumBytesStreamed = 0;
+		if (SoundWave->IsStreaming())
+		{
+			bFinished = CompressedAudioInfo->StreamCompressedData(ByteBuffer, false, BufferSize, NumBytesStreamed);
+		}
+		else
+		{
+			bFinished = CompressedAudioInfo->ReadCompressedData(ByteBuffer, false, BufferSize);
+			NumBytesStreamed = BufferSize;
+		}
+
+		if (NumBytesStreamed == 0)
+		{
+			break;
+		}
+
+		const int32 NumSamplesStreamed = NumBytesStreamed / sizeof(int16);
+		const int32 NumFramesStreamed = NumSamplesStreamed / NumChannels;
+
+		SampleOffset += NumSamplesStreamed;
+		NumFramesRemaining -= NumFramesStreamed;
+		BufferSize -= NumBytesStreamed;
+	}
+
+	int32 NumFramesGenerated = NumFrames - NumFramesRemaining;
+	SourceFrameIndex = FMath::Min(SourceFrameIndex + NumFramesGenerated, NumFramesInWave);
+	return NumFramesGenerated;
+}
+
+float FSoundWaveProxyPlayer::GetCurrentFrameRatio() const
+{
+	const float SampleRateFrameRatio = SampleRate / Settings.OutputSampleRate;
+	return FMath::Clamp(SampleRateFrameRatio * Speed, Audio::FRuntimeResampler::MinFrameRatio, Audio::FRuntimeResampler::MaxFrameRatio);
+}
+
+void FSoundWaveProxyPlayer::SetSpeed(float InSpeed)
+{
+	ensureMsgf(Speed >= MinSpeed, TEXT("Speed is too slow! Clamping to MinSpeed = %.2f"), MinSpeed);
+	ensureMsgf(Speed <= MaxSpeed, TEXT("Speed is too high! Clamping to MaxSpeed = %.2f"), MaxSpeed);
+	Speed = FMath::Clamp(InSpeed, MinSpeed, MaxSpeed);
+}
+
+float FSoundWaveProxyPlayer::GetSpeed() const
+{
+	return Speed;	
+}
+
+bool FSoundWaveProxyPlayer::SeekToTime(float TimeSeconds)
+{
+	int32 InFrame = FMath::Clamp(static_cast<int32>(TimeSeconds * SampleRate), 0, NumFramesInWave);
+	return SeekToFrame(InFrame);
+}
+
+void FSoundWaveProxyPlayer::SetLoop(bool InIsLooping, float InLoopStartTimeSeconds, float InLoopDurationSeconds)
+{
+	bIsLooping = InIsLooping;
+	CurrentLoopStartTime = InLoopStartTimeSeconds;
+	CurrentLoopDuration = InLoopDurationSeconds;
+
+	if (!SoundWave)
+	{
+		return;
+	}
+	
+	float LoopStartTimeSeconds = FMath::Clamp(CurrentLoopStartTime, 0.0f, DurationSeconds - MinLoopDurationSeconds);
+	float LoopDurationSeconds = CurrentLoopDuration > 0.0f ? CurrentLoopDuration : DurationSeconds - CurrentLoopStartTime;
+	LoopDurationSeconds = FMath::Clamp(LoopDurationSeconds, MinLoopDurationSeconds, DurationSeconds - CurrentLoopStartTime);
+
+	LoopStartFrameIndex = FMath::Clamp(FMath::FloorToInt32(LoopStartTimeSeconds * SampleRate), 0, NumFramesInWave);
+	LoopEndFrameIndex = FMath::Clamp(FMath::FloorToInt32(LoopDurationSeconds * SampleRate) + LoopStartFrameIndex, 0, NumFramesInWave);
+	
+	check(LoopStartFrameIndex >= 0);
+	check(LoopEndFrameIndex > LoopStartFrameIndex);
+	check(LoopEndFrameIndex <= NumFramesInWave);
+}
+
+bool FSoundWaveProxyPlayer::IsLooping() const
+{
+	if (!SoundWave)
+	{
+		return false;
+	}
+	return bIsLooping;
+}
+
+bool FSoundWaveProxyPlayer::IsFinished() const
+{
+	if (!SoundWave)
+	{
+		return false;
+	}
+
+	if (bIsLooping)
+	{
+		return false;
+	}
+	
+	return SourceFrameIndex >= NumFramesInWave;
+}
+
+float FSoundWaveProxyPlayer::GetCurrentPlaybackTimeSeconds() const
+{
+	if (!SoundWave)
+	{
+		return 0.0f;
+	}
+	
+	if (!ensure(SampleRate > 0.0f))
+	{
+		return 0.0f;
+	}
+	return SourceFrameIndex / SampleRate;	
+}
+
+float FSoundWaveProxyPlayer::GetPlaybackProgress() const
+{
+	if (!SoundWave)
+	{
+		return 0.0f;
+	}
+	
+	if (!ensure(NumFramesInWave > 0))
+	{
+		return 0.0f;
+	}
+	return SourceFrameIndex / static_cast<float>(NumFramesInWave);
+}
+
+float FSoundWaveProxyPlayer::GetLoopProgress() const
+{
+	if (!SoundWave)
+	{
+		return 0.0f;
+	}
+	
+	if (!bIsLooping)
+	{
+		return 0.0f;
+	}
+	
+	if (!ensure(NumFramesInWave > 0))
+	{
+		return 0.0f;
+	}
+	check(LoopStartFrameIndex >= 0);
+	check(LoopEndFrameIndex > LoopStartFrameIndex);
+	check(LoopEndFrameIndex <= NumFramesInWave);
+
+	return (SourceFrameIndex - LoopStartFrameIndex) / static_cast<float>(LoopEndFrameIndex - LoopStartFrameIndex);
+}
+
+const TArray<FSoundWaveCuePoint>& FSoundWaveProxyPlayer::GetCuePoints() const
+{
+	return SortedCuePoints;
+}
+
+bool FSoundWaveProxyPlayer::SeekToFrame(int32 InFrame)
+{
+	if (!SoundWave)
+	{
+		UE_LOG(LogAudio, Error, TEXT("FSoundWaveProxyPlayer::SeekToFrame, failed to seek - no Sound Wave set!"));
+		return false;
+	}
+	
+	InFrame = FMath::Clamp(InFrame, 0, NumFramesInWave);
+	
+	if (!CompressedAudioInfo.IsValid())
+	{
+		// set the current frame index, but indicate that we've still failed
+		SourceFrameIndex = InFrame;
+		AudioSyncFrameIndex = InFrame;
+		UE_LOG(LogAudio, Verbose, TEXT("FSoundWaveProxyPlayer::SeekToFrame, failed to seek due to the decoder being invalid!: %s"), *SoundWave->GetFName().ToString());
+		return false;
+	}
+
+	// ignore seek request if we're already at the specified time
+	if (InFrame == SourceFrameIndex)
+	{
+		return true;
+	}
+
+	// if we seek, then we should reset the resampler
+	bRequiresResampling = false;
+	RuntimeResampler.Reset(NumChannels);
+	ResampledAudio.Reset();
+
+	// reset the "NearlyFinished" flag
+	bHasNearlyFinished = false;
+
+	if (SoundWave->IsSeekable() && CompressedAudioInfo)
+	{
+		CompressedAudioInfo->SeekToFrame(InFrame);
+		SourceFrameIndex = InFrame;
+		AudioSyncFrameIndex = InFrame;
+		return true;
+	}
+
+	FName Format = SoundWave->GetRuntimeFormat();
+	UE_LOG(LogAudio, Warning, TEXT("Attempt to seek on non-seekable wave: (format:%s) for wave (package:%s) to frame '%d'"),
+		*Format.ToString(),
+		*SoundWave->GetPackageName().ToString(),
+		InFrame);
+	
+	// For non-seekable streaming waves, use a fallback method to seek
+	IConsoleVariable* SimulateSeekCVar = IConsoleManager::Get().FindConsoleVariable(TEXT("au.SoundWaveProxyReader.SimulateSeek"));
+	const bool bSimulateSeek = SimulateSeekCVar ? SimulateSeekCVar->GetInt() == 1 : false;
+	if (bSimulateSeek)
+	{
+		UE_CLOG(bFallbackSeekMethodWarningLogged == false, LogAudio, Warning, TEXT("Simulating seeking in wave which is not seekable (package:%s). For better performance, set wave to a seekable format"), *SoundWave->GetPackageName().ToString());
+		bFallbackSeekMethodWarningLogged = true;
+
+		if (InFrame < SourceFrameIndex)
+		{
+			// if we're seeking backwards, we need to reset the decoder
+			if (!InitializeDecoder(SoundWave.ToSharedRef(), Settings.MaxDecodeSizeFrames))
+			{
+				SourceFrameIndex = InFrame;
+				AudioSyncFrameIndex = InFrame;
+				UE_LOG(LogAudio, Verbose, TEXT("FSoundWaveProxyPlayer::SeekToFrame, failed to seek due to the decoder being invalid!: %s"), *SoundWave->GetFName().ToString());
+				return false;
+			}
+			SourceFrameIndex = 0;
+		}
+
+		int32 SeekNumFrames = InFrame - SourceFrameIndex;
+		int32 NumFramesGenerated = SimulateSeekForward(SeekNumFrames);
+		SourceFrameIndex = FMath::Min(SourceFrameIndex + NumFramesGenerated, NumFramesInWave);
+		AudioSyncFrameIndex = InFrame;
+		return NumFramesGenerated == SeekNumFrames;
+	}
+
+	return false;
+}
+
+
+bool FSoundWaveProxyPlayer::InitializeDecoder(const FSoundWaveProxyRef& InWaveProxy, int32 MaxDecodeSizeInFrames)
+{
+	using namespace Audio;
+
+	TRACE_CPUPROFILER_EVENT_SCOPE(FSoundWaveProxyPlayer::InitializeDecoder);
+
+	FName Format = InWaveProxy->GetRuntimeFormat();
+	IAudioInfoFactory* Factory = IAudioInfoFactoryRegistry::Get().Find(Format);
+	if (!ensure(Factory))
+	{
+		UE_LOG(LogAudio, Error, TEXT("FSoundWaveProxyPlayer::InitializeDecoder: Failed to create CompressedAudioInfo for wave (package: %s). Unable to find AudioInfoFactory for format: %s"),
+			*InWaveProxy->GetPackageName().ToString(), *Format.ToString());
+		return false;
+	}
+
+	TUniquePtr<ICompressedAudioInfo> InfoInstance;
+	InfoInstance.Reset(Factory->Create());
+
+	if (!ensure(InfoInstance.IsValid()))
+	{
+		UE_LOG(LogAudio, Error, TEXT("FSoundWaveProxyPlayer::InitializeDecoder: Failed to create CompressedAudioInfo for wave (package: %s). Unable to create info from factory for for format: %s"), *InWaveProxy->GetPackageName().ToString(), *Format.ToString());
+		return false;
+	}
+
+	FSoundQualityInfo Info;
+	if (InWaveProxy->IsStreaming())
+	{
+		if (!InfoInstance->StreamCompressedInfo(InWaveProxy, &Info))
+		{
+			UE_LOG(LogAudio, Error, TEXT("FSoundWaveProxyPlayer::InitializeDecoder: Failed to create CompressedAudioInfo for wave (package: %s). Unable to stream compressed info for streaming wave"), *InWaveProxy->GetPackageName().ToString());
+			return false;
+		}
+	}
+	else
+	{
+		if (!InfoInstance->ReadCompressedInfo(InWaveProxy->GetResourceData(), InWaveProxy->GetResourceSize(), &Info))
+		{
+			UE_LOG(LogAudio, Error, TEXT("FSoundWaveProxyPlayer::InitializeDecoder: Failed to create CompressedAudioInfo for wave (package: %s). Unable to read compressed info for non-streaming wave"), *InWaveProxy->GetPackageName().ToString());
+			return false;
+		}
+	}
+
+	SampleRate = InWaveProxy->GetSampleRate();
+	NumChannels = InWaveProxy->GetNumChannels();
+	NumFramesInWave = InWaveProxy->GetNumFrames();
+	DurationSeconds = InWaveProxy->GetDuration();
+
+	CompressedAudioInfo.Reset(InfoInstance.Release());
+
+	// BEGIN HACK
+	// Read the sample rate and number of frames from the header 
+	// Similar to refreshing the wave data in FMixerBuffer::CreateStreamingBuffer
+	// This is a runtime hack to address incorrect sample rate on soundwaves 
+	// on platforms with Resample for Device enabled (UE-183237)
+	SampleRate = Info.SampleRate;
+	uint32 NumFrames = (uint32)((float)Info.Duration * Info.SampleRate);
+	if (NumFrames > 0)
+	{
+		NumFramesInWave = NumFrames;
+	}
+	// END HACK
+
+	// initialized decode buffers
+	const int32 DecodeSize = ConformDecodeSize(MaxDecodeSizeInFrames);
+	check(DecodeSize > 0);
+	DecodeBuffer.SetNum(DecodeSize * NumChannels, EAllowShrinking::Yes);
+	return true;
+}
\ No newline at end of file
