Making FEngineAnalytics::CreateAnalyticsUserId public.

diff --git a/Engine/Source/Runtime/Engine/Private/EngineAnalytics.cpp b/Engine/Source/Runtime/Engine/Private/EngineAnalytics.cpp
--- a/Engine/Source/Runtime/Engine/Private/EngineAnalytics.cpp
+++ b/Engine/Source/Runtime/Engine/Private/EngineAnalytics.cpp
@@ -79,11 +79,6 @@ static TSharedPtr<IAnalyticsProviderET> CreateEpicAnalyticsProvider()
 	return FAnalyticsET::Get().CreateAnalyticsProvider(Config);
 }
 
-FString CreateAnalyticsUserId(const FString& EpicAccountId)
-{
-	return FString::Printf(TEXT("%s|%s|%s"), *FPlatformMisc::GetLoginId(), *EpicAccountId, *FPlatformMisc::GetOperatingSystemId());
-}
-
 IAnalyticsProviderET& FEngineAnalytics::GetProvider()
 {
 	checkf(bIsInitialized && IsAvailable(), TEXT("FEngineAnalytics::GetProvider called outside of Initialize/Shutdown."));
@@ -232,6 +227,11 @@ void FEngineAnalytics::LowDriveSpaceDetected()
 #endif
 }
 
+FString FEngineAnalytics::CreateAnalyticsUserId(const FString& EpicAccountId)
+{
+	return FString::Printf(TEXT("%s|%s|%s"), *FPlatformMisc::GetLoginId(), *EpicAccountId, *FPlatformMisc::GetOperatingSystemId());
+}
+
 void FEngineAnalytics::AppendMachineStats(TArray<FAnalyticsEventAttribute>& EventAttributes)
 {
 	const UGeneralProjectSettings& ProjectSettings = *GetDefault<UGeneralProjectSettings>();
diff --git a/Engine/Source/Runtime/Engine/Public/EngineAnalytics.h b/Engine/Source/Runtime/Engine/Public/EngineAnalytics.h
--- a/Engine/Source/Runtime/Engine/Public/EngineAnalytics.h
+++ b/Engine/Source/Runtime/Engine/Public/EngineAnalytics.h
@@ -58,6 +58,8 @@ public:
 
 	static ENGINE_API void LowDriveSpaceDetected();
 
+	static ENGINE_API FString CreateAnalyticsUserId(const FString& EpicAccountId);
+
 private:
 	static void AppendMachineStats(TArray<FAnalyticsEventAttribute>& EventAttributes);
 	static void SendMachineInfoForAccount(const FString& EpicAccountId);
