Fix wrong motion vectors on GPUSKIN_MORPH and GPUSKIN_CLOTH object with the mobile renderer

diff --git a/Engine/Source/Runtime/Engine/Private/GPUSkinVertexFactory.cpp b/Engine/Source/Runtime/Engine/Private/GPUSkinVertexFactory.cpp
--- a/Engine/Source/Runtime/Engine/Private/GPUSkinVertexFactory.cpp
+++ b/Engine/Source/Runtime/Engine/Private/GPUSkinVertexFactory.cpp
@@ -841,7 +841,7 @@ void TGPUSkinVertexFactory<BoneInfluenceType>::ModifyCompilationEnvironment(cons
 	OutEnvironment.SetDefine(TEXT("VF_SUPPORTS_PRIMITIVE_SCENE_DATA"), bSupportsPrimitiveIdStream && bUseGPUScene);
 
 	// Mobile doesn't support motion blur, don't use previous frame morph delta for mobile.
-	const bool bIsMobile = IsMobilePlatform(Parameters.Platform);
+	const bool bIsMobile = IsMobilePlatform(Parameters.Platform) && !MobileSupportsSM5MaterialNodes(Parameters.Platform);
 	OutEnvironment.SetDefine(TEXT("GPUSKIN_MORPH_USE_PREVIOUS"), !bIsMobile);
 
 	// Whether the material supports morph targets
@@ -1298,7 +1298,7 @@ void TGPUSkinAPEXClothVertexFactory<BoneInfluenceType>::ModifyCompilationEnviron
 	OutEnvironment.SetDefine(TEXT("GPUSKIN_APEX_CLOTH"),TEXT("1"));
 	
 	// Mobile doesn't support motion blur, don't use previous frame data.
-	const bool bIsMobile = IsMobilePlatform(Parameters.Platform);
+	const bool bIsMobile = IsMobilePlatform(Parameters.Platform) && !MobileSupportsSM5MaterialNodes(Parameters.Platform);
 	OutEnvironment.SetDefine(TEXT("GPUSKIN_APEX_CLOTH_PREVIOUS"), !bIsMobile);
 }
 
@@ -1592,4 +1592,4 @@ FGPUBaseSkinVertexFactory::FUpdateScope::FUpdateScope()
 {}
 
 #undef IMPLEMENT_GPUSKINNING_VERTEX_FACTORY_PARAMETER_TYPE
-#undef IMPLEMENT_GPUSKINNING_VERTEX_FACTORY_TYPE
\ No newline at end of file
+#undef IMPLEMENT_GPUSKINNING_VERTEX_FACTORY_TYPE
