Adjust RHI usage in ParticleGpuSimulation.cpp to fix warnings from D3D12 debug layer about uninitialized resources

diff --git a/Engine/Source/Runtime/Engine/Private/Particles/ParticleGpuSimulation.cpp b/Engine/Source/Runtime/Engine/Private/Particles/ParticleGpuSimulation.cpp
--- a/Engine/Source/Runtime/Engine/Private/Particles/ParticleGpuSimulation.cpp
+++ b/Engine/Source/Runtime/Engine/Private/Particles/ParticleGpuSimulation.cpp
@@ -323,8 +323,8 @@ class FParticleStateTextures : public FRenderResource
 {
 public:
 
-	const TCHAR* PaticlePositonTextureName = TEXT("FParticleStatePosition");
-	const TCHAR* PaticleVelocityTextureName = TEXT("FParticleStateVelocity");
+	const TCHAR* ParticlePositionTextureName = TEXT("FParticleStatePosition");
+	const TCHAR* ParticleVelocityTextureName = TEXT("FParticleStateVelocity");
 
 	/** Contains the positions of all simulating particles. */
 	FTextureRHIRef PositionTextureRHI;
@@ -353,7 +353,7 @@ public:
 			.SetFormat(PixelFormat)
 			.SetClearValue(FClearValueBinding::Transparent)
 			.SetFlags(ETextureCreateFlags::RenderTargetable | ETextureCreateFlags::ShaderResource)
-			.SetInitialState(ERHIAccess::SRVMask);
+			.SetInitialState(ERHIAccess::RTV);
 
 		return RHICmdList.CreateTexture(Desc);
 	}
@@ -365,19 +365,19 @@ public:
 	{
 		// 32-bit per channel RGBA texture for position.
 		check( !IsValidRef( PositionTextureRHI ) );
-		PositionTextureRHI = CreateStateTextureRHI(RHICmdList, PaticlePositonTextureName, GetParticlePositionFormat());
+		PositionTextureRHI = CreateStateTextureRHI(RHICmdList, ParticlePositionTextureName, GetParticlePositionFormat());
 
 		// 16-bit per channel RGBA texture for velocity.
 		check(!IsValidRef(VelocityTextureRHI));
-		VelocityTextureRHI = CreateStateTextureRHI(RHICmdList, PaticleVelocityTextureName, GetParticleVelocityFormat());
+		VelocityTextureRHI = CreateStateTextureRHI(RHICmdList, ParticleVelocityTextureName, GetParticleVelocityFormat());
 
 		// using FName's ability to append a number to a string (..._0) without an extra string allocation, except suffixing is done when number > 0 hence the +1 here : 
-		FName PositionTextureName(PaticlePositonTextureName, ParticleStateIndex + 1);
-		FName VelocityTextureName(PaticleVelocityTextureName, ParticleStateIndex + 1);
+		FName PositionTextureName(ParticlePositionTextureName, ParticleStateIndex + 1);
+		FName VelocityTextureName(ParticleVelocityTextureName, ParticleStateIndex + 1);
 		PositionTextureRHI->SetName(PositionTextureName);
 		VelocityTextureRHI->SetName(VelocityTextureName);
-		RHICmdList.BindDebugLabelName(PositionTextureRHI, PaticlePositonTextureName);
-		RHICmdList.BindDebugLabelName(VelocityTextureRHI, PaticleVelocityTextureName);
+		RHICmdList.BindDebugLabelName(PositionTextureRHI, ParticlePositionTextureName);
+		RHICmdList.BindDebugLabelName(VelocityTextureRHI, ParticleVelocityTextureName);
 
 		bTexturesCleared = false;
 	}
@@ -400,26 +400,36 @@ public:
 		RenderTargetSize.Y = SizeY;
 
 		// allocate the new textures
-		FTextureRHIRef NewPositionTextureRHI = CreateStateTextureRHI(RHICmdList, PaticlePositonTextureName, PF_A32B32G32R32F);
-		FTextureRHIRef NewVelocityTextureRHI = CreateStateTextureRHI(RHICmdList, PaticleVelocityTextureName, PF_FloatRGBA);
-		
-		RHICmdList.Transition(FRHITransitionInfo(NewPositionTextureRHI, ERHIAccess::SRVMask, ERHIAccess::CopyDest));
-		RHICmdList.Transition(FRHITransitionInfo(PositionTextureRHI, ERHIAccess::SRVMask, ERHIAccess::CopySrc));
+		FTextureRHIRef NewPositionTextureRHI = CreateStateTextureRHI(RHICmdList, ParticlePositionTextureName, GetParticlePositionFormat());
+		FTextureRHIRef NewVelocityTextureRHI = CreateStateTextureRHI(RHICmdList, ParticleVelocityTextureName, GetParticleVelocityFormat());
+
+		if (bTexturesCleared)
+		{
+			{
+				SCOPED_GPU_MASK(RHICmdList, FRHIGPUMask::All());
 
-		RHICmdList.Transition(FRHITransitionInfo(NewVelocityTextureRHI, ERHIAccess::SRVMask, ERHIAccess::CopyDest));
-		RHICmdList.Transition(FRHITransitionInfo(VelocityTextureRHI, ERHIAccess::SRVMask, ERHIAccess::CopySrc));
+				FRHITexture* NewRenderTargets[] = { NewPositionTextureRHI, NewVelocityTextureRHI };
+				FRHIRenderPassInfo RenderPassInfo(UE_ARRAY_COUNT(NewRenderTargets), NewRenderTargets, ERenderTargetActions::Clear_Store);
+				RHICmdList.BeginRenderPass(RenderPassInfo, TEXT("GPUParticlesClearStateTextures"));
+				RHICmdList.EndRenderPass();
+			}
 
-		FRHICopyTextureInfo CopyTextureInfo;
-		CopyTextureInfo.Size = FIntVector(PositionTextureRHI->GetSizeX(), PositionTextureRHI->GetSizeY(), 0);
+			RHICmdList.Transition({
+				FRHITransitionInfo(NewPositionTextureRHI, ERHIAccess::RTV, ERHIAccess::CopyDest),
+				FRHITransitionInfo(PositionTextureRHI, ERHIAccess::SRVMask, ERHIAccess::CopySrc),
+				FRHITransitionInfo(NewVelocityTextureRHI, ERHIAccess::RTV, ERHIAccess::CopyDest),
+				FRHITransitionInfo(VelocityTextureRHI, ERHIAccess::SRVMask, ERHIAccess::CopySrc) });
 
-		RHICmdList.CopyTexture(PositionTextureRHI, NewPositionTextureRHI, CopyTextureInfo);
-		RHICmdList.CopyTexture(VelocityTextureRHI, NewVelocityTextureRHI, CopyTextureInfo);
+			FRHICopyTextureInfo CopyTextureInfo;
+			CopyTextureInfo.Size = FIntVector(PositionTextureRHI->GetSizeX(), PositionTextureRHI->GetSizeY(), 0);
 
-		RHICmdList.Transition(FRHITransitionInfo(NewPositionTextureRHI, ERHIAccess::CopyDest, ERHIAccess::SRVMask));
-		RHICmdList.Transition(FRHITransitionInfo(PositionTextureRHI, ERHIAccess::CopySrc, ERHIAccess::SRVMask));
+			RHICmdList.CopyTexture(PositionTextureRHI, NewPositionTextureRHI, CopyTextureInfo);
+			RHICmdList.CopyTexture(VelocityTextureRHI, NewVelocityTextureRHI, CopyTextureInfo);
 
-		RHICmdList.Transition(FRHITransitionInfo(NewVelocityTextureRHI, ERHIAccess::CopyDest, ERHIAccess::SRVMask));
-		RHICmdList.Transition(FRHITransitionInfo(VelocityTextureRHI, ERHIAccess::CopySrc, ERHIAccess::SRVMask));
+			RHICmdList.Transition({
+				FRHITransitionInfo(NewPositionTextureRHI, ERHIAccess::CopyDest, ERHIAccess::SRVMask),
+				FRHITransitionInfo(NewVelocityTextureRHI, ERHIAccess::CopyDest, ERHIAccess::SRVMask) });
+		}
 
 		// release the old textures
 		ReleaseRHI();
@@ -437,7 +447,7 @@ class FParticleAttributesTexture : public FRenderResource
 {
 public:
 
-	const TCHAR* PaticleAttributeTextureName = TEXT("FParticleAttributesTexture");
+	const TCHAR* ParticleAttributeTextureName = TEXT("FParticleAttributesTexture");
 
 	/** Contains the attributes of all simulating particles. */
 	FTextureRHIRef TextureRHI;
@@ -447,7 +457,7 @@ public:
 	FTextureRHIRef CreateAttributesTextureRHI(FRHICommandListBase& RHICmdList)
 	{
 		const FRHITextureCreateDesc Desc =
-			FRHITextureCreateDesc::Create2D(PaticleAttributeTextureName)
+			FRHITextureCreateDesc::Create2D(ParticleAttributeTextureName)
 			.SetExtent(RenderTargetSize.X, RenderTargetSize.Y)
 			.SetFormat(PF_B8G8R8A8)
 			.SetClearValue(FClearValueBinding::Transparent)
@@ -487,8 +497,11 @@ public:
 		// allocate the new textures
 		FTextureRHIRef NewAttributesTextureRHI = CreateAttributesTextureRHI(RHICmdList);
 
-		RHICmdList.Transition(FRHITransitionInfo(NewAttributesTextureRHI, ERHIAccess::RTV, ERHIAccess::CopyDest));
-		RHICmdList.Transition(FRHITransitionInfo(TextureRHI, ERHIAccess::SRVMask, ERHIAccess::CopySrc));
+		ClearRenderTarget(RHICmdList, NewAttributesTextureRHI);
+
+		RHICmdList.Transition({
+			FRHITransitionInfo(NewAttributesTextureRHI, ERHIAccess::RTV, ERHIAccess::CopyDest),
+			FRHITransitionInfo(TextureRHI, ERHIAccess::SRVMask, ERHIAccess::CopySrc) });
 
 		FRHICopyTextureInfo CopyTextureInfo;
 		CopyTextureInfo.Size = FIntVector(TextureRHI->GetSizeX(), TextureRHI->GetSizeY(), 0);
@@ -496,7 +509,6 @@ public:
 		RHICmdList.CopyTexture(TextureRHI, NewAttributesTextureRHI, CopyTextureInfo);
 
 		RHICmdList.Transition(FRHITransitionInfo(NewAttributesTextureRHI, ERHIAccess::CopyDest, ERHIAccess::SRVMask));
-		RHICmdList.Transition(FRHITransitionInfo(TextureRHI, ERHIAccess::CopySrc, ERHIAccess::SRVMask));
 
 		// release the old textures
 		ReleaseRHI();
@@ -572,15 +584,6 @@ public:
 			ParticleResources->StateTextures[1].InitResource(RHICmdList);
 			ParticleResources->RenderAttributesTexture.InitResource(RHICmdList);
 			ParticleResources->SimulationAttributesTexture.InitResource(RHICmdList);
-
-			RHICmdList.Transition({ 
-				FRHITransitionInfo(ParticleResources->RenderAttributesTexture.TextureRHI, ERHIAccess::Unknown, ERHIAccess::SRVMask),
-				FRHITransitionInfo(ParticleResources->SimulationAttributesTexture.TextureRHI, ERHIAccess::Unknown, ERHIAccess::SRVMask),
-				FRHITransitionInfo(ParticleResources->StateTextures[0].PositionTextureRHI, ERHIAccess::Unknown, ERHIAccess::SRVMask),
-				FRHITransitionInfo(ParticleResources->StateTextures[0].VelocityTextureRHI, ERHIAccess::Unknown, ERHIAccess::SRVMask),
-				FRHITransitionInfo(ParticleResources->StateTextures[1].PositionTextureRHI, ERHIAccess::Unknown, ERHIAccess::SRVMask),
-				FRHITransitionInfo(ParticleResources->StateTextures[1].VelocityTextureRHI, ERHIAccess::Unknown, ERHIAccess::SRVMask)
-			});
 		});
 	}
 
@@ -4762,11 +4765,24 @@ void FFXSystem::PrepareGPUSimulation(FRHICommandListImmediate& RHICmdList)
 	// Grab resources.
 	FParticleStateTextures& CurrentStateTextures = ParticleSimulationResources->GetCurrentStateTextures();
 
-	// Setup render states.
-	RHICmdList.Transition({
-		FRHITransitionInfo(CurrentStateTextures.PositionTextureRHI, ERHIAccess::SRVMask, ERHIAccess::RTV),
-		FRHITransitionInfo(CurrentStateTextures.VelocityTextureRHI, ERHIAccess::SRVMask, ERHIAccess::RTV)
-	});
+	if (CurrentStateTextures.bTexturesCleared)
+	{
+		// Setup render states.
+		RHICmdList.Transition({
+			FRHITransitionInfo(CurrentStateTextures.PositionTextureRHI, ERHIAccess::SRVMask, ERHIAccess::RTV),
+			FRHITransitionInfo(CurrentStateTextures.VelocityTextureRHI, ERHIAccess::SRVMask, ERHIAccess::RTV) });
+	}
+	else
+	{
+		SCOPED_GPU_MASK(RHICmdList, FRHIGPUMask::All());
+
+		FRHITexture* CurrentStateRenderTargets[] = { CurrentStateTextures.PositionTextureRHI, CurrentStateTextures.VelocityTextureRHI };
+		FRHIRenderPassInfo RenderPassInfo(UE_ARRAY_COUNT(CurrentStateRenderTargets), CurrentStateRenderTargets, ERenderTargetActions::Clear_Store);
+		RHICmdList.BeginRenderPass(RenderPassInfo, TEXT("PrepareGPUSimulation_GPUParticlesClearStateTextures"));
+		RHICmdList.EndRenderPass();
+
+		CurrentStateTextures.bTexturesCleared = true;
+	}
 }
 
 void FFXSystem::FinalizeGPUSimulation(FRHICommandListImmediate& RHICmdList)
@@ -4821,7 +4837,6 @@ void FFXSystem::SimulateGPUParticles(
 #endif
 
 	{
-		
 		// On some platforms, the textures are filled with garbage after creation, so we need to clear them to black the first time we use them
 		if ( !CurrentStateTextures.bTexturesCleared )
 		{
@@ -4830,7 +4845,7 @@ void FFXSystem::SimulateGPUParticles(
 			FRHIRenderPassInfo RenderPassInfo(2, CurrentStateRenderTargets, ERenderTargetActions::Clear_Store);
 			RHICmdList.BeginRenderPass(RenderPassInfo, TEXT("GPUParticlesClearStateTextures"));
 			RHICmdList.EndRenderPass();
-			
+
 			CurrentStateTextures.bTexturesCleared = true;
 		}
 		
@@ -4838,10 +4853,6 @@ void FFXSystem::SimulateGPUParticles(
 		{
 			SCOPED_GPU_MASK(RHICmdList, FRHIGPUMask::All());
 
-			RHICmdList.Transition({
-				FRHITransitionInfo(PreviousStateRenderTargets[0], ERHIAccess::SRVMask, ERHIAccess::RTV),
-				FRHITransitionInfo(PreviousStateRenderTargets[1], ERHIAccess::SRVMask, ERHIAccess::RTV)});
-
 			FRHIRenderPassInfo RPInfo(2, PreviousStateRenderTargets, ERenderTargetActions::Clear_Store);
 			RHICmdList.BeginRenderPass(RPInfo, TEXT("GPUParticlesClearPreviousStateTextures"));
 			RHICmdList.EndRenderPass();
@@ -4849,7 +4860,7 @@ void FFXSystem::SimulateGPUParticles(
 			RHICmdList.Transition({
 				FRHITransitionInfo(PreviousStateRenderTargets[0], ERHIAccess::RTV, ERHIAccess::SRVMask),
 				FRHITransitionInfo(PreviousStateRenderTargets[1], ERHIAccess::RTV, ERHIAccess::SRVMask) });
-			
+
 			PrevStateTextures.bTexturesCleared = true;
 		}
 	}
@@ -5069,7 +5080,7 @@ void FFXSystem::SimulateGPUParticles(
 		FRHIRenderPassInfo RPInfo(4, InjectRenderTargets, ERenderTargetActions::Load_Store);
 		{
 			// Transition attribute textures to writeble, particle state texture are in writeble state already
-			RHICmdList.Transition({ 
+			RHICmdList.Transition({
 				FRHITransitionInfo(InjectRenderTargets[2], ERHIAccess::SRVMask, ERHIAccess::RTV),
 				FRHITransitionInfo(InjectRenderTargets[3], ERHIAccess::SRVMask, ERHIAccess::RTV)
 			});
@@ -5105,7 +5116,7 @@ void FFXSystem::SimulateGPUParticles(
 	}
 	
 	// finish current state render
-	RHICmdList.Transition({ 
+	RHICmdList.Transition({
 		FRHITransitionInfo(CurrentStateRenderTargets[0], ERHIAccess::RTV, ERHIAccess::SRVMask),
 		FRHITransitionInfo(CurrentStateRenderTargets[1], ERHIAccess::RTV, ERHIAccess::SRVMask)});
 
@@ -5115,7 +5126,7 @@ void FFXSystem::SimulateGPUParticles(
 		//will make this second step simulate an interpolated extra 7ms.  This second interpolated step is what we render on THIS frame, but it is NOT fed into the next frame's simulation.
 		FParticleStateTextures& VisualizeStateTextures = ParticleSimulationResources->GetPreviousStateTextures();
 		
-		RHICmdList.Transition({ 
+		RHICmdList.Transition({
 			FRHITransitionInfo(VisualizeStateTextures.PositionTextureRHI, ERHIAccess::SRVMask, ERHIAccess::RTV),
 			FRHITransitionInfo(VisualizeStateTextures.VelocityTextureRHI, ERHIAccess::SRVMask, ERHIAccess::RTV)
 		});
@@ -5147,7 +5158,7 @@ void FFXSystem::SimulateGPUParticles(
 			RHICmdList.EndRenderPass();
 		}
 
-		RHICmdList.Transition({ 
+		RHICmdList.Transition({
 			FRHITransitionInfo(VisualizeStateTextures.PositionTextureRHI, ERHIAccess::RTV, ERHIAccess::SRVMask),
 			FRHITransitionInfo(VisualizeStateTextures.VelocityTextureRHI, ERHIAccess::RTV, ERHIAccess::SRVMask)
 		});
