[NetPhysics] Dynamic input buffer scaling for networked physics via NetworkPhysicsComponent

diff --git a/Engine/Source/Runtime/Engine/Classes/GameFramework/PlayerController.h b/Engine/Source/Runtime/Engine/Classes/GameFramework/PlayerController.h
--- a/Engine/Source/Runtime/Engine/Classes/GameFramework/PlayerController.h
+++ b/Engine/Source/Runtime/Engine/Classes/GameFramework/PlayerController.h
@@ -2281,6 +2281,9 @@ private:
 	/** Tick Offset sync interval timer */
 	float TickOffsetSyncCountdown = 0.f;
 
+	/** Time to add to the tick offset, works as a buffer on the server to handle fluctuating network conditions */
+	float TickOffsetBufferTime;
+
 	/** Time that tick offset has been desynced for since it was last in sync, used to trigger a correction */
 	float NetworkPhysicsTickOffsetDesyncAccumulatedTime = 0.0f;
 
@@ -2321,6 +2324,12 @@ public:
 	/** True if the NetworkPhysicsTickOffset is setup, before this is setup the GetNetworkPhysicsTickOffset() and GetPhysicsTimestamp() functions will not return valid results */
 	bool GetNetworkPhysicsTickOffsetAssigned() const { return bNetworkPhysicsTickOffsetAssigned; }
 
+	/** Get the current buffer added to the physics tick-offset */
+	float GetTickOffsetBufferTime() const { return TickOffsetBufferTime; }
+
+	/** Set a new buffer time to add to the physics tick-offset */
+	void SetTickOffsetBufferTime(float NewBufferTime);
+
 	/** Return the handshake cookie of the connection this player controller is associated. */
 	ENGINE_API uint32 GetClientHandshakeId() const;
 
diff --git a/Engine/Source/Runtime/Engine/Private/PhysicsEngine/NetworkPhysicsComponent.cpp b/Engine/Source/Runtime/Engine/Private/PhysicsEngine/NetworkPhysicsComponent.cpp
--- a/Engine/Source/Runtime/Engine/Private/PhysicsEngine/NetworkPhysicsComponent.cpp
+++ b/Engine/Source/Runtime/Engine/Private/PhysicsEngine/NetworkPhysicsComponent.cpp
@@ -21,26 +21,45 @@ namespace PhysicsReplicationCVars
 	namespace ResimulationCVars
 	{
 		int32 RedundantInputs = 2;
-		static FAutoConsoleVariableRef CVarResimRedundantInputs(TEXT("np2.Resim.RedundantInputs"), RedundantInputs, TEXT("How many extra inputs to send with each unreliable network message, to account for packetloss. From owning client to server and back to owning client. NOTE: This is disabled while np2.Resim.DynamicInputScaling.Enabled is enabled."));
+		static FAutoConsoleVariableRef CVarResimRedundantInputs(TEXT("np2.Resim.RedundantInputs"), RedundantInputs, TEXT("How many extra inputs to send with each unreliable network message, to account for packetloss. From owning client to server and back to owning client. NOTE: This is disabled while np2.Resim.DynamicInputReplicationScaling.Enabled is enabled."));
 		int32 RedundantRemoteInputs = 1;
 		static FAutoConsoleVariableRef CVarResimRedundantRemoteInputs(TEXT("np2.Resim.RedundantRemoteInputs"), RedundantRemoteInputs, TEXT("How many extra inputs to send with each unreliable network message, to account for packetloss. From server to remote clients."));
 		int32 RedundantStates = 0;
 		static FAutoConsoleVariableRef CVarResimRedundantStates(TEXT("np2.Resim.RedundantStates"), RedundantStates, TEXT("How many extra states to send with each unreliable network message, to account for packetloss."));
 
-		bool bDynamicInputScalingEnabled = true;
-		static FAutoConsoleVariableRef CVarDynamicInputScalingEnabled(TEXT("np2.Resim.DynamicInputScaling.Enabled"), bDynamicInputScalingEnabled, TEXT("Enable dynmic scaling of number of inputs sent from owning client to the server to account for packet loss. The server will control the value based on how often the server has a hole in its input buffer. NOTE: This overrides np2.Resim.RedundantInputs"));
-		float DynamicInputScalingMaxInputsPercent = 0.1f;
-		static FAutoConsoleVariableRef CVarDynamicInputScalingMaxInputsPercent(TEXT("np2.Resim.DynamicInputScaling.MaxInputsPercent"), DynamicInputScalingMaxInputsPercent, TEXT("Default 0.1 (= 10%, value in percent as multiplier). Sets the max scalable number of inputs to network from owning client to server as a percentage of the physics fixed tick-rate. 10% of 30Hz = 3 inputs at max."));
-		int32 DynamicInputScalingMinInputs = 2;
-		static FAutoConsoleVariableRef CVarDynamicInputScalingMinInputs(TEXT("np2.Resim.DynamicInputScaling.MinInputs"), DynamicInputScalingMinInputs, TEXT("Default 2. Sets the minimum scalable number of inputs to network from owning client to server."));
-		float DynamicInputScalingIncreaseAverageMultiplier = 0.2f;
-		static FAutoConsoleVariableRef CVarDynamicInputScalingIncreaseAverageMultiplier(TEXT("np2.Resim.DynamicInputScaling.IncreaseAverageMultiplier"), DynamicInputScalingIncreaseAverageMultiplier, TEXT("Default 0.2 (= 20%). Multiplier for how fast the average input scaling value increases. NOTE it's recommended to have a higher value than np2.Resim.DynamicInputScaling.DecreaseAverageMultiplier so the average can grow quick when network conditions gets worse."));
-		float DynamicInputScalingDecreaseAverageMultiplier = 0.1f;
-		static FAutoConsoleVariableRef CVarDynamicInputScalingDecreaseAverageMultiplier(TEXT("np2.Resim.DynamicInputScaling.DecreaseAverageMultiplier"), DynamicInputScalingDecreaseAverageMultiplier, TEXT("Default 0.1 (= 10%). Multiplier for how fast the average input scaling value decreases. NOTE it's recommended to have a lower value than np2.Resim.DynamicInputScaling.IncreaseAverageMultiplier so the average doesn't try to decrease too quickly which can cause repeated desyncs."));
-		float DynamicInputScalingIncreaseTimeInterval = 2.0f;
-		static FAutoConsoleVariableRef CVarDynamicInputScalingIncreaseTimeInterval(TEXT("np2.Resim.DynamicInputScaling.IncreaseTimeInterval"), DynamicInputScalingIncreaseTimeInterval, TEXT("Default 2.0 (value in seconds). How often dynamic scaling can increase the number of inputs to send." ));
-		float DynamicInputScalingDecreaseTimeInterval = 10.0f;
-		static FAutoConsoleVariableRef CVarDynamicInputScalingDecreaseTimeInterval(TEXT("np2.Resim.DynamicInputScaling.DecreaseTimeInterval"), DynamicInputScalingDecreaseTimeInterval, TEXT("Default 10.0 (value in seconds). How often dynamic scaling can decrease the number of inputs to send."));
+		bool bDynamicInputReplicationScalingEnabled = true;
+		static FAutoConsoleVariableRef CVarDynamicInputReplicationScalingEnabled(TEXT("np2.Resim.DynamicInputReplicationScaling.Enabled"), bDynamicInputReplicationScalingEnabled, TEXT("Enable dynmic scaling of number of inputs sent from owning client to the server to account for packet loss. The server will control the value based on how often the server has a hole in its input buffer. NOTE: This overrides np2.Resim.RedundantInputs"));
+		float DynamicInputReplicationScalingMaxInputsPercent = 0.1f;
+		static FAutoConsoleVariableRef CVarDynamicInputReplicationScalingMaxInputsPercent(TEXT("np2.Resim.DynamicInputReplicationScaling.MaxInputsPercent"), DynamicInputReplicationScalingMaxInputsPercent, TEXT("Default 0.1 (= 10%, value in percent as multiplier). Sets the max scalable number of inputs to network from owning client to server as a percentage of the physics fixed tick-rate. 10% of 30Hz = 3 inputs at max."));
+		int32 DynamicInputReplicationScalingMinInputs = 2;
+		static FAutoConsoleVariableRef CVarDynamicInputReplicationScalingMinInputs(TEXT("np2.Resim.DynamicInputReplicationScaling.MinInputs"), DynamicInputReplicationScalingMinInputs, TEXT("Default 2. Sets the minimum scalable number of inputs to network from owning client to server."));
+		float DynamicInputReplicationScalingIncreaseAverageMultiplier = 0.2f;
+		static FAutoConsoleVariableRef CVarDynamicInputReplicationScalingIncreaseAverageMultiplier(TEXT("np2.Resim.DynamicInputReplicationScaling.IncreaseAverageMultiplier"), DynamicInputReplicationScalingIncreaseAverageMultiplier, TEXT("Default 0.2 (= 20%). Multiplier for how fast the average input scaling value increases. NOTE it's recommended to have a higher value than np2.Resim.DynamicInputReplicationScaling.DecreaseAverageMultiplier so the average can grow quick when network conditions gets worse."));
+		float DynamicInputReplicationScalingDecreaseAverageMultiplier = 0.1f;
+		static FAutoConsoleVariableRef CVarDynamicInputReplicationScalingDecreaseAverageMultiplier(TEXT("np2.Resim.DynamicInputReplicationScaling.DecreaseAverageMultiplier"), DynamicInputReplicationScalingDecreaseAverageMultiplier, TEXT("Default 0.1 (= 10%). Multiplier for how fast the average input scaling value decreases. NOTE it's recommended to have a lower value than np2.Resim.DynamicInputReplicationScaling.IncreaseAverageMultiplier so the average doesn't try to decrease too quickly which can cause repeated desyncs."));
+		float DynamicInputReplicationScalingIncreaseTimeInterval = 2.0f;
+		static FAutoConsoleVariableRef CVarDynamicInputReplicationScalingIncreaseTimeInterval(TEXT("np2.Resim.DynamicInputReplicationScaling.IncreaseTimeInterval"), DynamicInputReplicationScalingIncreaseTimeInterval, TEXT("Default 2.0 (value in seconds). How often dynamic scaling can increase the number of inputs to send." ));
+		float DynamicInputReplicationScalingDecreaseTimeInterval = 10.0f;
+		static FAutoConsoleVariableRef CVarDynamicInputReplicationScalingDecreaseTimeInterval(TEXT("np2.Resim.DynamicInputReplicationScaling.DecreaseTimeInterval"), DynamicInputReplicationScalingDecreaseTimeInterval, TEXT("Default 10.0 (value in seconds). How often dynamic scaling can decrease the number of inputs to send."));
+
+		bool bDynamicInputBufferScalingEnabled = true;
+		static FAutoConsoleVariableRef CVarDynamicInputBufferScalingEnabled(TEXT("np2.Resim.DynamicInputBufferScaling.Enabled"), bDynamicInputBufferScalingEnabled, TEXT("Enable dynmic scaling of input buffer on the server."));
+		float DynamicInputBufferScalingMinBufferMs = 30.0f;
+		static FAutoConsoleVariableRef CVarDynamicInputBufferScalingMinBufferMs(TEXT("np2.Resim.DynamicInputBufferScaling.MinBufferTime"), DynamicInputBufferScalingMinBufferMs, TEXT("Time in milliseconds for the lowest allowed input buffer size, when below this the buffer will scale up over time."));
+		float DynamicInputBufferScalingMaxBufferMs = 150.0f;
+		static FAutoConsoleVariableRef CVarDynamicInputBufferScalingMaxBufferMs(TEXT("np2.Resim.DynamicInputBufferScaling.MaxBufferTime"), DynamicInputBufferScalingMaxBufferMs, TEXT("Time in milliseconds to cap out the input buffer size."));
+		float DynamicInputBufferScalingAverageTime = 0.5f;
+		static FAutoConsoleVariableRef CVarDynamicInputBufferScalingAverageTime(TEXT("np2.Resim.DynamicInputBufferScaling.AverageTime"), DynamicInputBufferScalingAverageTime, TEXT("Time in seconds to keep a running average of the input buffer size over."));
+		float DynamicInputBufferScalingBumpUpMultiplier = 1.0f;
+		static FAutoConsoleVariableRef CVarDynamicInputBufferScalingBumpUpMultiplier(TEXT("np2.Resim.DynamicInputBufferScaling.BumpUpMultiplier"), DynamicInputBufferScalingBumpUpMultiplier, TEXT("Multiplier for how much of a fixed delta time should be added to the target input buffer size instantly."));
+		float DynamicInputBufferScalingScaleUpMultiplier = 1.0f;
+		static FAutoConsoleVariableRef CVarDynamicInputBufferScalingScaleUpMultiplier(TEXT("np2.Resim.DynamicInputBufferScaling.ScaleUpMultiplier"), DynamicInputBufferScalingScaleUpMultiplier, TEXT("How fast the input buffer scales up when too low. Default 1.0 = one fixed delta time per second."));
+		float DynamicInputBufferScalingScaleDownMinMultiplier = 0.01f;
+		static FAutoConsoleVariableRef CVarDynamicInputBufferScalingScaleDownMinMultiplier(TEXT("np2.Resim.DynamicInputBufferScaling.ScaleDownMinMultiplier"), DynamicInputBufferScalingScaleDownMinMultiplier, TEXT("How fast the input buffer scales down when buffer is slightly too large."));
+		float DynamicInputBufferScalingScaleDownMaxMultiplier = 0.1f;
+		static FAutoConsoleVariableRef CVarDynamicInputBufferScalingScaleDownMaxMultiplier(TEXT("np2.Resim.DynamicInputBufferScaling.ScaleDownMaxMultiplier"), DynamicInputBufferScalingScaleDownMaxMultiplier, TEXT("How fast the input buffer scales down when buffer is slightly too large."));
+		bool bDynamicInputBufferScalingDebugLogs = false;
+		static FAutoConsoleVariableRef CVarDynamicInputBufferScalingDebugLogs(TEXT("np2.Resim.DynamicInputBufferScaling.DebugLogs"), bDynamicInputBufferScalingDebugLogs, TEXT("Print logs for debugging."));
 
 		bool bAllowRewindToClosestState = true;
 		static FAutoConsoleVariableRef CVarResimAllowRewindToClosestState(TEXT("np2.Resim.AllowRewindToClosestState"), bAllowRewindToClosestState, TEXT("When rewinding to a specific frame, if the client doens't have state data for that frame, use closest data available. Only affects the first rewind frame, when FPBDRigidsEvolution is set to Reset."));
@@ -744,6 +763,15 @@ void UNetworkPhysicsComponent::NetworkMarshaledData()
 			InputsToNetwork_Owner = *AsyncOutput->InputsToNetwork_Owner;
 		}
 
+		if (AsyncOutput->TargetBufferSizeMs.IsSet())
+		{
+			// Only marshaled from PT to GT on the server
+			if (APlayerController* PC = GetPlayerController())
+			{
+				PC->SetTickOffsetBufferTime(*AsyncOutput->TargetBufferSizeMs);
+			}
+		}
+
 		// Unimportant / Unreliable
 		if (bEnableUnreliableFlow
 			&& AsyncOutput->InputData
@@ -1797,7 +1825,18 @@ void FAsyncNetworkPhysicsComponent::ConsumeAsyncInput(const int32 PhysicsStep)
 			}
 			if (AsyncInput->NetworkPhysicsTickOffset.IsSet())
 			{
-				NetworkPhysicsTickOffset = *AsyncInput->NetworkPhysicsTickOffset;
+				if (NetworkPhysicsTickOffset != *AsyncInput->NetworkPhysicsTickOffset)
+				{
+					NetworkPhysicsTickOffset = *AsyncInput->NetworkPhysicsTickOffset;
+
+					if (ComponentSettings.GetDynamicInputBufferScalingEnabled())
+					{
+						// Also reset dynamic input buffer if tick offset was altered
+						const float MinInputBufferMs = PhysicsReplicationCVars::ResimulationCVars::DynamicInputBufferScalingMinBufferMs;
+						const float MaxInputBufferMs = PhysicsReplicationCVars::ResimulationCVars::DynamicInputBufferScalingMaxBufferMs;
+						InputBufferTargetSizeMs = FMath::Min(InputBufferTargetSizeMs, MaxInputBufferMs - ((MaxInputBufferMs - MinInputBufferMs) * 0.5f));
+					}
+				}
 			}
 			if (AsyncInput->InputsToNetwork_Owner.IsSet())
 			{
@@ -2350,17 +2389,22 @@ void FAsyncNetworkPhysicsComponent::SendInputData_Internal(FAsyncNetworkPhysicsC
 {
 	const bool bIsServer = IsServer();
 
-	if (bIsServer)
-	{
-		UpdateDynamicInputScaling();
-		AsyncOutput.InputsToNetwork_Owner = InputsToNetwork_Owner;
-	}
-
 	// Inputs are sent from the server or locally controlled actors/pawns
 	if (AsyncOutput.InputData && InputHistory && (IsLocallyControlled() || bIsServer))
 	{
 		const FNetworkPhysicsSettingsNetworkPhysicsComponent& ComponentSettings = GetComponentSettings();
 
+		if (bIsServer)
+		{
+			UpdateDynamicInputReplicationScaling();
+			AsyncOutput.InputsToNetwork_Owner = InputsToNetwork_Owner;
+
+			if (ComponentSettings.GetDynamicInputBufferScalingEnabled())
+			{
+				AsyncOutput.TargetBufferSizeMs = UpdateDynamicInputBufferScaling();
+			}
+		}
+
 		// Send latest N frames from history
 		const int32 ToFrame = FMath::Max(0, PhysicsStep);
 
@@ -2566,9 +2610,9 @@ const float FAsyncNetworkPhysicsComponent::GetCurrentInputDecay() const
 	return InputDecay;
 }
 
-void FAsyncNetworkPhysicsComponent::UpdateDynamicInputScaling()
+void FAsyncNetworkPhysicsComponent::UpdateDynamicInputReplicationScaling()
 {
-	if (!PhysicsReplicationCVars::ResimulationCVars::bDynamicInputScalingEnabled)
+	if (!PhysicsReplicationCVars::ResimulationCVars::bDynamicInputReplicationScalingEnabled)
 	{
 		InputsToNetwork_Owner = InputsToNetwork_OwnerDefault;
 		return;
@@ -2585,41 +2629,161 @@ void FAsyncNetworkPhysicsComponent::UpdateDynamicInputScaling()
 		return;
 	}
 
-	const float TimeSinceLastDynamicScaling = RigidSolver->GetSolverTime() - TimeOfLastDynamicInputScaling;
+	const float TimeSinceLastDynamicScaling = RigidSolver->GetSolverTime() - TimeOfLastDynamicInputReplicationScaling;
 
 	if (MissingInputCount > 0)
 	{
-		if (TimeSinceLastDynamicScaling > PhysicsReplicationCVars::ResimulationCVars::DynamicInputScalingIncreaseTimeInterval)
+		if (TimeSinceLastDynamicScaling > PhysicsReplicationCVars::ResimulationCVars::DynamicInputReplicationScalingIncreaseTimeInterval)
 		{
-			const uint16 MaxInputsValue = static_cast<uint16>(FMath::CeilToInt32(PhysicsReplicationCVars::ResimulationCVars::DynamicInputScalingMaxInputsPercent / (RigidSolver->GetAsyncDeltaTime())));
+			const uint16 MaxInputsValue = static_cast<uint16>(FMath::CeilToInt32(PhysicsReplicationCVars::ResimulationCVars::DynamicInputReplicationScalingMaxInputsPercent / (RigidSolver->GetAsyncDeltaTime())));
 
 			// Increase the amount of inputs the owner sends
 			InputsToNetwork_Owner++;
 
 			// Update the average value for minimum clamping
-			DynamicInputScalingAverageInputs += (static_cast<float>(InputsToNetwork_Owner) - DynamicInputScalingAverageInputs) * PhysicsReplicationCVars::ResimulationCVars::DynamicInputScalingIncreaseAverageMultiplier;
+			DynamicInputReplicationScalingAverageInputs += (static_cast<float>(InputsToNetwork_Owner) - DynamicInputReplicationScalingAverageInputs) * PhysicsReplicationCVars::ResimulationCVars::DynamicInputReplicationScalingIncreaseAverageMultiplier;
 			
 			// Clamp to maximum valid value
 			InputsToNetwork_Owner = FMath::Min(InputsToNetwork_Owner, MaxInputsValue);
 
-			TimeOfLastDynamicInputScaling = RigidSolver->GetSolverTime();
+			TimeOfLastDynamicInputReplicationScaling = RigidSolver->GetSolverTime();
 			MissingInputCount = 0;
 		}
 	}
-	else if (TimeSinceLastDynamicScaling > PhysicsReplicationCVars::ResimulationCVars::DynamicInputScalingDecreaseTimeInterval)
+	else if (TimeSinceLastDynamicScaling > PhysicsReplicationCVars::ResimulationCVars::DynamicInputReplicationScalingDecreaseTimeInterval)
 	{
 		// Decrease the amount of inputs the owner sends
 		InputsToNetwork_Owner--;
 
 		// Update the average value for minimum clamping, perform before clamping to allow for decreasing average even if the clamp might still round up.
-		DynamicInputScalingAverageInputs += (static_cast<float>(InputsToNetwork_Owner) - DynamicInputScalingAverageInputs) * PhysicsReplicationCVars::ResimulationCVars::DynamicInputScalingDecreaseAverageMultiplier;
+		DynamicInputReplicationScalingAverageInputs += (static_cast<float>(InputsToNetwork_Owner) - DynamicInputReplicationScalingAverageInputs) * PhysicsReplicationCVars::ResimulationCVars::DynamicInputReplicationScalingDecreaseAverageMultiplier;
 
 		// Clamp to minimum valid value
-		const uint16 MinInputsValue = static_cast<uint16>(FMath::Max(static_cast<uint16>(FMath::RoundToInt(DynamicInputScalingAverageInputs)), static_cast<uint16>(PhysicsReplicationCVars::ResimulationCVars::DynamicInputScalingMinInputs)));
+		const uint16 MinInputsValue = static_cast<uint16>(FMath::Max(static_cast<uint16>(FMath::RoundToInt(DynamicInputReplicationScalingAverageInputs)), static_cast<uint16>(PhysicsReplicationCVars::ResimulationCVars::DynamicInputReplicationScalingMinInputs)));
 		InputsToNetwork_Owner = FMath::Max(InputsToNetwork_Owner, MinInputsValue);
 
-		TimeOfLastDynamicInputScaling = RigidSolver->GetSolverTime();
+		TimeOfLastDynamicInputReplicationScaling = RigidSolver->GetSolverTime();
+	}
+}
+
+float FAsyncNetworkPhysicsComponent::UpdateDynamicInputBufferScaling()
+{
+	if (!IsServer() || !InputHistory)
+	{
+		return InputBufferTargetSizeMs;
+	}
+
+	Chaos::FPBDRigidsSolver* RigidSolver = GetRigidSolver();
+	if (!RigidSolver || FMath::IsNearlyZero(RigidSolver->GetAsyncDeltaTime()))
+	{
+		return InputBufferTargetSizeMs;
+	}
+
+	const float Dt = RigidSolver->GetAsyncDeltaTime();
+	if (FMath::IsNearlyZero(Dt))
+	{
+		return InputBufferTargetSizeMs;
 	}
+
+	const int32 InputBufferSize = InputHistory->GetLatestFrame() - RigidSolver->GetCurrentFrame();
+	if (InputBufferSize < 0)
+	{
+		// If buffer is negative, we are not yet setup
+		return InputBufferTargetSizeMs;
+	}
+
+	// Cache fixed delta time as milliseconds (used for input buffer size calculations while Dt in seconds is used for time scaling / s
+	const float FixedDeltaTimeMs = Dt * 1000.0f;
+
+	const float MinInputBufferMs = FMath::Max(PhysicsReplicationCVars::ResimulationCVars::DynamicInputBufferScalingMinBufferMs, FixedDeltaTimeMs);
+	const float MaxInputBufferMs = FMath::Max(PhysicsReplicationCVars::ResimulationCVars::DynamicInputBufferScalingMaxBufferMs, MinInputBufferMs + (FixedDeltaTimeMs * 2.0f));
+
+	// Current input buffer size
+	const float InputBufferTimeMs = static_cast<float>(InputBufferSize) * FixedDeltaTimeMs;
+
+	// Running average input buffer size
+	const float AverageOverSeconds = PhysicsReplicationCVars::ResimulationCVars::DynamicInputBufferScalingAverageTime > 0.0f ? PhysicsReplicationCVars::ResimulationCVars::DynamicInputBufferScalingAverageTime : 0.5f;
+	const float AverageAlpha = FMath::Clamp(Dt / (Dt + AverageOverSeconds), 0.0f, 1.0f);
+	InputBufferAverageMs += (InputBufferTimeMs - InputBufferAverageMs) * AverageAlpha;
+
+	// Calculate what's currently considered the max good input buffer
+	const float GoodInputBufferSize = MinInputBufferMs + FixedDeltaTimeMs + InputBufferDynamicMinSizeMs;
+
+	// The input buffer is empty
+	if (InputBufferSize <= 0)
+	{
+		// Bump up the target input buffer size
+		const float BumpUpMultiplier = PhysicsReplicationCVars::ResimulationCVars::DynamicInputBufferScalingBumpUpMultiplier > 0.0f ? PhysicsReplicationCVars::ResimulationCVars::DynamicInputBufferScalingBumpUpMultiplier : 0.25f;
+		InputBufferTargetSizeMs += FixedDeltaTimeMs * BumpUpMultiplier;
+
+		// Bump up the dynamic minimum size
+		InputBufferDynamicMinSizeMs += (InputBufferTargetSizeMs - InputBufferDynamicMinSizeMs) * BumpUpMultiplier;
+
+#if !(UE_BUILD_SHIPPING)
+		if (PhysicsReplicationCVars::ResimulationCVars::bDynamicInputBufferScalingDebugLogs)
+		{
+			UE_LOG(LogChaos, Log, TEXT("Dynamic Input Buffer - EMPTY     InputBufferTargetSizeMs: %f - Added: %f - BumpUpMultiplier: %f")
+				, InputBufferTargetSizeMs, (FixedDeltaTimeMs * BumpUpMultiplier), BumpUpMultiplier);
+		}
+#endif
+	}
+	// The input buffer is low
+	else if(InputBufferTimeMs <= MinInputBufferMs)
+	{
+		// Scale up the target input buffer size
+		const float ScaleUpMultiplier = PhysicsReplicationCVars::ResimulationCVars::DynamicInputBufferScalingScaleUpMultiplier > 0.0f ? PhysicsReplicationCVars::ResimulationCVars::DynamicInputBufferScalingScaleUpMultiplier : 1.0f;
+		InputBufferTargetSizeMs += (ScaleUpMultiplier * FixedDeltaTimeMs * Dt);
+
+		// Scale up the dynamic minimum size
+		InputBufferDynamicMinSizeMs += (InputBufferTargetSizeMs - InputBufferDynamicMinSizeMs) * (ScaleUpMultiplier * Dt);
+
+#if !(UE_BUILD_SHIPPING)
+		if (PhysicsReplicationCVars::ResimulationCVars::bDynamicInputBufferScalingDebugLogs)
+		{
+			UE_LOG(LogChaos, Log, TEXT("Dynamic Input Buffer - TOO LOW   InputBufferTimeMs: %f - InputBufferAverageMs %f - InputBufferTargetSizeMs: %f - ScaleUpMultiplier: %f")
+				, InputBufferTimeMs, InputBufferAverageMs, InputBufferTargetSizeMs, ScaleUpMultiplier);
+		}
+#endif
+	}
+	// The input buffer is too large
+	else if (InputBufferTimeMs > GoodInputBufferSize && InputBufferAverageMs > GoodInputBufferSize)
+	{
+		// Scale down the target input buffer size
+		const float ScaleDownMinMultiplier = PhysicsReplicationCVars::ResimulationCVars::DynamicInputBufferScalingScaleDownMinMultiplier > 0.0f ? PhysicsReplicationCVars::ResimulationCVars::DynamicInputBufferScalingScaleDownMinMultiplier : 0.01f;
+		const float ScaleDownMaxMultiplier = PhysicsReplicationCVars::ResimulationCVars::DynamicInputBufferScalingScaleDownMaxMultiplier > 0.0f ? PhysicsReplicationCVars::ResimulationCVars::DynamicInputBufferScalingScaleDownMaxMultiplier : 0.1f;
+
+		const float OvershootMs = InputBufferAverageMs - MinInputBufferMs;
+		const float OvershootAlpha = FMath::Clamp(OvershootMs / (MaxInputBufferMs - MinInputBufferMs), 0.0f, 1.0f);
+		const float DecayRate = FMath::Lerp(ScaleDownMinMultiplier, ScaleDownMaxMultiplier, OvershootAlpha);
+		InputBufferTargetSizeMs -= InputBufferTargetSizeMs * DecayRate * Dt;
+
+#if !(UE_BUILD_SHIPPING)
+		if (PhysicsReplicationCVars::ResimulationCVars::bDynamicInputBufferScalingDebugLogs)
+		{
+			UE_LOG(LogChaos, Log, TEXT("Dynamic Input Buffer - TOO LARGE InputBufferTimeMs: %f - InputBufferAverageMs %f - GoodInputBufferSize: %f - InputBufferTargetSizeMs: %f - InputBufferDynamicMinSizeMs: %f - DecayRate: %f")
+				, InputBufferTimeMs, InputBufferAverageMs, GoodInputBufferSize, InputBufferTargetSizeMs, InputBufferDynamicMinSizeMs, DecayRate);
+		}
+#endif
+	}
+	// The input buffer is good
+	else
+	{
+		// Scale down the dynamic minimum size
+		InputBufferDynamicMinSizeMs -= InputBufferDynamicMinSizeMs * Dt;
+
+#if !(UE_BUILD_SHIPPING)
+		if (PhysicsReplicationCVars::ResimulationCVars::bDynamicInputBufferScalingDebugLogs)
+		{
+			UE_LOG(LogChaos, Log, TEXT("Dynamic Input Buffer - GOOD      InputBufferTimeMs: %f - InputBufferAverageMs %f - GoodInputBufferSize: %f - InputBufferTargetSizeMs: %f - InputBufferDynamicMinSizeMs: %f")
+				, InputBufferTimeMs, InputBufferAverageMs, GoodInputBufferSize, InputBufferTargetSizeMs, InputBufferDynamicMinSizeMs);
+		}
+#endif
+	}
+
+	// Clamp input buffer target size between min and max
+	InputBufferTargetSizeMs = FMath::Clamp(InputBufferTargetSizeMs, MinInputBufferMs, MaxInputBufferMs);
+
+	return InputBufferTargetSizeMs;
 }
 
 void FAsyncNetworkPhysicsComponent::RegisterDataHistoryInRewindData()
diff --git a/Engine/Source/Runtime/Engine/Private/PlayerController.cpp b/Engine/Source/Runtime/Engine/Private/PlayerController.cpp
--- a/Engine/Source/Runtime/Engine/Private/PlayerController.cpp
+++ b/Engine/Source/Runtime/Engine/Private/PlayerController.cpp
@@ -239,6 +239,8 @@ APlayerController::APlayerController(const FObjectInitializer& ObjectInitializer
 
 	ClickEventKeys.Add(EKeys::LeftMouseButton);
 
+	TickOffsetBufferTime = NetworkPhysicsCvars::TickOffsetBufferTime;
+
 	if (RootComponent)
 	{
 		// We want to drive rotation with ControlRotation regardless of attachment state.
@@ -6538,6 +6540,11 @@ void APlayerController::SetClientHandshakeId(uint32 InClientHandshakeId)
 	ClientHandshakeId = InClientHandshakeId;
 }
 
+void APlayerController::SetTickOffsetBufferTime(float NewBufferTime)
+{
+	TickOffsetBufferTime = FMath::Max(0.0f, NewBufferTime);
+};
+
 void APlayerController::UpdateServerAsyncPhysicsTickOffset()
 {
 	if (NetworkPhysicsCvars::TickOffsetUpdateIntervalTime <= 0 || TickOffsetSyncCountdown < (NetworkPhysicsCvars::TickOffsetUpdateIntervalTime * 0.001f))
@@ -6570,9 +6577,14 @@ void APlayerController::ServerSendLatestAsyncPhysicsTimestamp_Implementation(FAs
 		}
 	}
 
+	if (FMath::IsNearlyZero(AsyncFixedTimeStepSize))
+	{
+		return;
+	}
+
 	// Get current server timestamp and add the frame buffer to the ServerFrame
 	FAsyncPhysicsTimestamp ActualTimestamp = GetPhysicsTimestamp();
-	const int32 BufferTickSize = FMath::CeilToInt((NetworkPhysicsCvars::TickOffsetBufferTime * 0.001f) / AsyncFixedTimeStepSize);
+	const int32 BufferTickSize = FMath::CeilToInt((TickOffsetBufferTime * 0.001f) / AsyncFixedTimeStepSize);
 	ActualTimestamp.ServerFrame += BufferTickSize;
 
 	// Mark offset as assigned when we get a valid predicted server frame.
diff --git a/Engine/Source/Runtime/Engine/Public/Physics/NetworkPhysicsComponent.h b/Engine/Source/Runtime/Engine/Public/Physics/NetworkPhysicsComponent.h
--- a/Engine/Source/Runtime/Engine/Public/Physics/NetworkPhysicsComponent.h
+++ b/Engine/Source/Runtime/Engine/Public/Physics/NetworkPhysicsComponent.h
@@ -1532,7 +1532,7 @@ private:
 
 	/** Set default number of inputs to send over the network with each message, clamped to 1 as minimum 
 	* From owning client to server.
-	* NOTE: this will be overridden if np2.Resim.DynamicInputScaling.Enabled is enabled. */
+	* NOTE: this will be overridden if np2.Resim.DynamicInputReplicationScaling.Enabled is enabled. */
 	void SetNumberOfInputsToNetwork(const uint16 NumInputs)
 	{
 		InputsToNetwork_OwnerDefault = FMath::Max(NumInputs, static_cast<uint16>(1));
@@ -1638,7 +1638,7 @@ private:
 	// Local default state data in legacy type
 	TUniquePtr<FNetworkPhysicsData> StateDataDefault_Legacy;
 
-	// The number of inputs the owning client should send to the server with each RPC, replicated from the server. This is dynamically scaled based on when there are holes in the inputs buffer if np2.Resim.DynamicInputScaling.Enabled is enabled
+	// The number of inputs the owning client should send to the server with each RPC, replicated from the server. This is dynamically scaled based on when there are holes in the inputs buffer if np2.Resim.DynamicInputReplicationScaling.Enabled is enabled
 	UPROPERTY( Replicated )
 	uint16 InputsToNetwork_Owner = 3;
 
@@ -1858,6 +1858,7 @@ struct FAsyncNetworkPhysicsComponentInput : public Chaos::FSimCallbackInput
 struct FAsyncNetworkPhysicsComponentOutput : public Chaos::FSimCallbackOutput
 {
 	TOptional<uint16> InputsToNetwork_Owner;
+	TOptional<float> TargetBufferSizeMs;
 
 	TUniquePtr<Chaos::FBaseRewindHistory> InputData;
 	TUniquePtr<Chaos::FBaseRewindHistory> StateData;
@@ -1868,6 +1869,7 @@ struct FAsyncNetworkPhysicsComponentOutput : public Chaos::FSimCallbackOutput
 	void Reset()
 	{
 		InputsToNetwork_Owner.Reset();
+		TargetBufferSizeMs.Reset();
 
 		if (InputData)
 		{
@@ -1969,7 +1971,10 @@ private:
 	void SendStateData_Internal(FAsyncNetworkPhysicsComponentOutput& AsyncOutput, const int32 PhysicsStep);
 
 	/** Update function on the server to manage how many inputs the owning client should send with each RPC */
-	void UpdateDynamicInputScaling();
+	void UpdateDynamicInputReplicationScaling();
+
+	/** Update the dynamic input buffer and return the new size */
+	float UpdateDynamicInputBufferScaling();
 
 private:
 	bool bIsLocallyControlled;
@@ -2023,14 +2028,19 @@ private:
 	uint16 InputsToNetwork_Owner = 3; // From owning client, i.e. autonomous proxy or client owning an actor with bIsRelayingLocalInputs enabled
 	uint16 InputsToNetwork_Simulated = 2; // To simulated proxies
 	
-	// Properties for dynamic scaling of inputs
-	float TimeOfLastDynamicInputScaling = 0.0f;
-	float DynamicInputScalingAverageInputs = 0.0f;
+	// Properties for dynamic scaling of number of inputs to replicate
+	float TimeOfLastDynamicInputReplicationScaling = 0.0f;
+	float DynamicInputReplicationScalingAverageInputs = 0.0f;
 	int32 MissingInputCount = 0;
 
 	// Send last N number of states each replication call
 	uint16 StatesToNetwork = 1;
 
+	// Properties for dynamic scaling of buffered inputs
+	float InputBufferAverageMs = 90.0;
+	float InputBufferTargetSizeMs = 90.0;
+	float InputBufferDynamicMinSizeMs = 0.0f;
+
 	// Cache predicted states and then compare incoming states via FNetworkPhysicsData::CompareData to trigger a resim if they desync
 	bool bCompareStateToTriggerRewind;
 
diff --git a/Engine/Source/Runtime/Engine/Public/Physics/NetworkPhysicsSettingsComponent.h b/Engine/Source/Runtime/Engine/Public/Physics/NetworkPhysicsSettingsComponent.h
--- a/Engine/Source/Runtime/Engine/Public/Physics/NetworkPhysicsSettingsComponent.h
+++ b/Engine/Source/Runtime/Engine/Public/Physics/NetworkPhysicsSettingsComponent.h
@@ -71,6 +71,7 @@ namespace PhysicsReplicationCVars
 		extern int32 RedundantInputs;
 		extern int32 RedundantRemoteInputs;
 		extern int32 RedundantStates;
+		extern bool bDynamicInputBufferScalingEnabled;
 		extern bool bCompareStateToTriggerRewind;
 		extern bool bCompareStateToTriggerRewindIncludeSimProxies;
 		extern bool bCompareInputToTriggerRewind;
@@ -495,6 +496,7 @@ struct FNetworkPhysicsSettingsNetworkPhysicsComponent
 		: bOverrideRedundantInputs(0)
 		, bOverrideRedundantRemoteInputs(0)
 		, bOverrideRedundantStates(0)
+		, bOverrideDynamicInputBufferScalingEnabled(0)
 		, bOverrideCompareStateToTriggerRewind(0)
 		, bOverridebCompareStateToTriggerRewindIncludeSimProxies(0)
 		, bOverrideCompareInputToTriggerRewind(0)
@@ -532,6 +534,8 @@ struct FNetworkPhysicsSettingsNetworkPhysicsComponent
 	UPROPERTY(config, EditDefaultsOnly, Category = "Overrides", Meta = (InlineEditConditionToggle))
 	uint32 bOverrideRedundantStates : 1;
 	UPROPERTY(config, EditDefaultsOnly, Category = "Overrides", Meta = (InlineEditConditionToggle))
+	uint32 bOverrideDynamicInputBufferScalingEnabled : 1;
+	UPROPERTY(config, EditDefaultsOnly, Category = "Overrides", Meta = (InlineEditConditionToggle))
 	uint32 bOverrideCompareStateToTriggerRewind : 1;
 	UPROPERTY(config, EditDefaultsOnly, Category = "Overrides", Meta = (InlineEditConditionToggle))
 	uint32 bOverridebCompareStateToTriggerRewindIncludeSimProxies : 1;
@@ -563,7 +567,7 @@ struct FNetworkPhysicsSettingsNetworkPhysicsComponent
 	uint32 bOverrideInputDecaySetTime : 1;
 
 	/* Overrides CVar : np2.Resim.RedundantInputs -- How many extra inputs to send with each unreliable network message, to account for packetloss.From owning client to server and server to owning client.
-	 * NOTE: This is disabled while np2.Resim.DynamicInputScaling.Enabled is enabled. */
+	 * NOTE: This is disabled while np2.Resim.DynamicInputReplicationScaling.Enabled is enabled. */
 	UPROPERTY(config, EditDefaultsOnly, Category = "Overrides", Meta = (EditCondition = "bOverrideRedundantInputs"))
 	uint16 RedundantInputs = (uint16)PhysicsReplicationCVars::ResimulationCVars::RedundantInputs;
 	const uint16 GetRedundantInputs() const { return bOverrideRedundantInputs ? RedundantInputs : (uint16)PhysicsReplicationCVars::ResimulationCVars::RedundantInputs; }
@@ -578,6 +582,11 @@ struct FNetworkPhysicsSettingsNetworkPhysicsComponent
 	uint16 RedundantStates = (uint16)PhysicsReplicationCVars::ResimulationCVars::RedundantStates;
 	const uint16 GetRedundantStates() const { return bOverrideRedundantStates ? RedundantStates : (uint16)PhysicsReplicationCVars::ResimulationCVars::RedundantStates; }
 
+	// Overrides CVar: np2.Resim.DynamicInputBufferScaling.Enabled -- Enable dynmic scaling of input buffer on the server.
+	UPROPERTY(config, EditDefaultsOnly, Category = "Overrides", Meta = (EditCondition = "bOverrideDynamicInputBufferScalingEnabled"))
+	bool bDynamicInputBufferScalingEnabled = PhysicsReplicationCVars::ResimulationCVars::bDynamicInputBufferScalingEnabled;
+	const bool GetDynamicInputBufferScalingEnabled() const { return bOverrideDynamicInputBufferScalingEnabled ? bDynamicInputBufferScalingEnabled : PhysicsReplicationCVars::ResimulationCVars::bDynamicInputBufferScalingEnabled; }
+
 	// Overrides CVar: np2.Resim.CompareStateToTriggerRewind -- When true, cache local FNetworkPhysicsData state in rewind history and compare the predicted state with incoming server state to trigger resimulations if they differ, comparison done through FNetworkPhysicsData::CompareData.
 	UPROPERTY(config, EditDefaultsOnly, Category = "Overrides", Meta = (EditCondition = "bOverrideCompareStateToTriggerRewind"))
 	bool bCompareStateToTriggerRewind = PhysicsReplicationCVars::ResimulationCVars::bCompareStateToTriggerRewind;
