Remove no longer needed additive total bone weight array from the BlendProfileScratchData and SlotEvaluatePoseWithBlendProfiles

diff --git a/Engine/Source/Runtime/Engine/Private/Animation/AnimInstanceProxy.cpp b/Engine/Source/Runtime/Engine/Private/Animation/AnimInstanceProxy.cpp
--- a/Engine/Source/Runtime/Engine/Private/Animation/AnimInstanceProxy.cpp
+++ b/Engine/Source/Runtime/Engine/Private/Animation/AnimInstanceProxy.cpp
@@ -1551,7 +1551,6 @@ void FAnimInstanceProxy::SlotEvaluatePoseWithBlendProfiles(const FName& SlotNode
 	FBlendProfileScratchData& BlendProfileScratchData = FBlendProfileScratchData::Get();
 	TArray<TArray<float>>& PerBoneWeights = BlendProfileScratchData.PerBoneWeights;
 	TArray<float>& PerBoneWeightTotals = BlendProfileScratchData.PerBoneWeightTotals;
-	TArray<float>& PerBoneWeightTotalsAdditive = BlendProfileScratchData.PerBoneWeightTotalsAdditive;
 	TArray<float>& BoneBlendProfileScales = BlendProfileScratchData.BoneBlendProfileScales;
 	check(BlendProfileScratchData.IsEmpty());
 
@@ -1561,7 +1560,6 @@ void FAnimInstanceProxy::SlotEvaluatePoseWithBlendProfiles(const FName& SlotNode
 	{
 		PerBoneWeightTotals[BoneIndex] = 0.0f;
 	}
-	PerBoneWeightTotalsAdditive = PerBoneWeightTotals;
 	PerBoneWeights.Reset(GetMontageEvaluationData().Num());
 
 	//------------------------------------------
@@ -1647,7 +1645,7 @@ void FAnimInstanceProxy::SlotEvaluatePoseWithBlendProfiles(const FName& SlotNode
 				}
 			}
 
-			// Sum up the total weights per bone.
+			// Sum up the total weights per bone for non-additive poses.
 			if (AdditiveAnimType == AAT_None)
 			{
 				for (int32 BoneIndex = 0; BoneIndex < NumBones; ++BoneIndex)
@@ -1655,13 +1653,6 @@ void FAnimInstanceProxy::SlotEvaluatePoseWithBlendProfiles(const FName& SlotNode
 					PerBoneWeightTotals[BoneIndex] += PerBoneWeights[CurrentPoseIndex][BoneIndex];
 				}
 			}
-			else
-			{
-				for (int32 BoneIndex = 0; BoneIndex < NumBones; ++BoneIndex)
-				{
-					PerBoneWeightTotalsAdditive[BoneIndex] += PerBoneWeights[CurrentPoseIndex][BoneIndex];
-				}
-			}
 			
 			CurrentPoseIndex++;
 		} // If montage slot is valid.
diff --git a/Engine/Source/Runtime/Engine/Private/Animation/BlendProfileScratchData.h b/Engine/Source/Runtime/Engine/Private/Animation/BlendProfileScratchData.h
--- a/Engine/Source/Runtime/Engine/Private/Animation/BlendProfileScratchData.h
+++ b/Engine/Source/Runtime/Engine/Private/Animation/BlendProfileScratchData.h
@@ -21,7 +21,6 @@ struct FBlendProfileScratchData : public TThreadSingleton<FBlendProfileScratchDa
 
 	TArray<TArray<float>> PerBoneWeights;		// A set of bone weights, per montage instance. Index this like [PoseIndex][CompactPoseBoneIndex].
 	TArray<float> PerBoneWeightTotals;			// The bone weight totals for non-additive poses, used for normalizing weights.
-	TArray<float> PerBoneWeightTotalsAdditive;	// The bone weight totals for additive poses, used for normalizing weights.
 	TArray<float> BoneBlendProfileScales;		// The bone profile scale values.
 
 	TArray<FSlotEvaluationPose> Poses;			// Non additive poses.
@@ -39,7 +38,6 @@ struct FBlendProfileScratchData : public TThreadSingleton<FBlendProfileScratchDa
 	{
 		PerBoneWeights.Reset();
 		PerBoneWeightTotals.Reset();
-		PerBoneWeightTotalsAdditive.Reset();
 		BoneBlendProfileScales.Reset();
 		Poses.Reset();
 		AdditivePoses.Reset();
@@ -55,7 +53,6 @@ struct FBlendProfileScratchData : public TThreadSingleton<FBlendProfileScratchDa
 	{
 		return PerBoneWeights.IsEmpty() &&
 			PerBoneWeightTotals.IsEmpty() &&
-			PerBoneWeightTotalsAdditive.IsEmpty() &&
 			BoneBlendProfileScales.IsEmpty() &&
 			Poses.IsEmpty() &&
 			AdditivePoses.IsEmpty() &&
