Add GetDisabled, SetDisabled, SetGravityEnabled, GetGravityEnabled, SetSleepType, GetSleepType, on UClusterUnionComponent

diff --git a/Engine/Source/Runtime/Engine/Classes/PhysicsEngine/ClusterUnionComponent.h b/Engine/Source/Runtime/Engine/Classes/PhysicsEngine/ClusterUnionComponent.h
--- a/Engine/Source/Runtime/Engine/Classes/PhysicsEngine/ClusterUnionComponent.h
+++ b/Engine/Source/Runtime/Engine/Classes/PhysicsEngine/ClusterUnionComponent.h
@@ -366,6 +366,18 @@ public:
 
 	virtual void OnChildToParentUpdated(UPrimitiveComponent* ChangedComponent, const FLocalBonesToTransformMap& PerBoneChildToParent, const FTransform& NewTransform, TArray<TTuple<UPrimitiveComponent*, FTransform>>* OutNewComponents) {}
 
+	ENGINE_API virtual void SetDisabled(const bool Disabled);
+	ENGINE_API virtual bool GetDisabled() const;
+	
+	ENGINE_API virtual void SetSleepType(const Chaos::ESleepType SleepType);
+	ENGINE_API virtual Chaos::ESleepType GetSleepType() const;
+	
+	ENGINE_API virtual void SetRigidState(const Chaos::EObjectStateType ObjectState);
+	ENGINE_API virtual Chaos::EObjectStateType GetRigidState() const;
+
+	ENGINE_API virtual void SetGravityEnabled(const bool Enabled);
+	ENGINE_API virtual bool GetGravityEnabled() const;
+	
 protected:
 
 	ENGINE_API void BroadcastComponentAddedEvents(UPrimitiveComponent* ChangedComponent, const TArray<FClusterUnionBoneData>& BoneIds, bool bIsNew, const TArray<FClusterUnionBoneData>& RemovedBoneIDs);
@@ -374,8 +386,6 @@ protected:
 	Chaos::FClusterUnionPhysicsProxy* GetPhysicsProxy() const { return PhysicsProxy; }
 	Chaos::FClusterUnionPhysicsProxy* GetPhysicsProxy() { return PhysicsProxy; }
 
-	ENGINE_API virtual void SetRigidState(Chaos::EObjectStateType ObjectState);
-
 	// We need to keep track of the mapping of primitive components to physics objects.
 	// This way we know the right physics objects to pass when removing the component (because
 	// it's possible to get a different list of physics objects when we get to removal). A
diff --git a/Engine/Source/Runtime/Engine/Private/PhysicsEngine/ClusterUnionComponent.cpp b/Engine/Source/Runtime/Engine/Private/PhysicsEngine/ClusterUnionComponent.cpp
--- a/Engine/Source/Runtime/Engine/Private/PhysicsEngine/ClusterUnionComponent.cpp
+++ b/Engine/Source/Runtime/Engine/Private/PhysicsEngine/ClusterUnionComponent.cpp
@@ -1354,7 +1354,7 @@ void UClusterUnionComponent::BroadcastComponentRemovedEvents(UPrimitiveComponent
 void UClusterUnionComponent::OnRep_RigidState()
 {
 	// Prior to having a proxy we can't set its state.
-	// Also, if our state is 0 we have not yet recieved a valid object state from
+	// Also, if our state is 0 we have not yet receive a valid object state from
 	// the server (we call this onrep in OnCreatePhysicsState too incase we got a
 	// new state but didn't have a proxy at that time)
 	if (!PhysicsProxy || (ReplicatedRigidState.ObjectState == 0 && GSkipZeroStateInOnRep))
@@ -1416,7 +1416,61 @@ void UClusterUnionComponent::ForceSetChildToParent(UPrimitiveComponent* InCompon
 	}
 }
 
-void UClusterUnionComponent::SetRigidState(Chaos::EObjectStateType ObjectState)
+void UClusterUnionComponent::SetDisabled(const bool Disabled)
+{
+	if (PhysicsProxy)
+	{
+		PhysicsProxy->SetDisabled_External(Disabled);
+	}
+}
+
+bool UClusterUnionComponent::GetDisabled() const
+{
+	if (PhysicsProxy)
+	{
+		return PhysicsProxy->GetDisabled_External();
+	}
+
+	return false;
+}
+
+void UClusterUnionComponent::SetGravityEnabled(const bool Enabled)
+{
+	if (PhysicsProxy)
+	{
+		PhysicsProxy->SetGravityEnabled_External(Enabled);
+	}
+}
+
+bool UClusterUnionComponent::GetGravityEnabled() const
+{
+	if (PhysicsProxy)
+	{
+		return PhysicsProxy->GetGravityEnabled_External();
+	}
+
+	return false;
+}
+
+void UClusterUnionComponent::SetSleepType(const Chaos::ESleepType SleepType)
+{
+	if (PhysicsProxy)
+	{
+		PhysicsProxy->SetSleepType_External(SleepType);
+	}
+}
+
+Chaos::ESleepType UClusterUnionComponent::GetSleepType() const
+{
+	if (PhysicsProxy)
+	{
+		return PhysicsProxy->GetSleepType_External();
+	}
+
+	return Chaos::ESleepType::MaterialSleep;
+}
+	
+void UClusterUnionComponent::SetRigidState(const Chaos::EObjectStateType ObjectState)
 {
 	if (PhysicsProxy)
 	{
@@ -1426,6 +1480,16 @@ void UClusterUnionComponent::SetRigidState(Chaos::EObjectStateType ObjectState)
 	}
 }
 
+Chaos::EObjectStateType UClusterUnionComponent::GetRigidState() const
+{
+	if (PhysicsProxy)
+	{
+		return PhysicsProxy->GetObjectState_External();
+	}
+
+	return Chaos::EObjectStateType::Uninitialized;
+}
+
 void UClusterUnionComponent::SetSimulatePhysics(bool bSimulate)
 {
 	SetRigidState(bSimulate ? Chaos::EObjectStateType::Dynamic : Chaos::EObjectStateType::Kinematic);
diff --git a/Engine/Source/Runtime/Experimental/Chaos/Private/PhysicsProxy/ClusterUnionPhysicsProxy.cpp b/Engine/Source/Runtime/Experimental/Chaos/Private/PhysicsProxy/ClusterUnionPhysicsProxy.cpp
--- a/Engine/Source/Runtime/Experimental/Chaos/Private/PhysicsProxy/ClusterUnionPhysicsProxy.cpp
+++ b/Engine/Source/Runtime/Experimental/Chaos/Private/PhysicsProxy/ClusterUnionPhysicsProxy.cpp
@@ -286,6 +286,109 @@ namespace Chaos
 		);
 	}
 
+	ESleepType FClusterUnionPhysicsProxy::GetSleepType_External() const
+	{
+		if (!ensure(Particle_External))
+		{
+			return ESleepType::MaterialSleep;
+		}
+
+		return Particle_External->SleepType();
+	}
+
+	void FClusterUnionPhysicsProxy::SetSleepType_External(const ESleepType SleepType)
+	{
+		if (!Solver || !ensure(Particle_External))
+		{
+			return;
+		}
+
+		Solver->EnqueueCommandImmediate(
+			[this, SleepType]() mutable
+			{
+				if (Particle_Internal)
+				{
+					if (Particle_Internal->SleepType() != SleepType)
+					{
+						if (FPBDRigidsEvolutionGBF* Evolution = GetEvolution(this))
+						{
+							Evolution->SetParticleSleepType(Particle_Internal, SleepType);
+						}
+					}
+				}
+			}
+		);
+	}
+
+	bool FClusterUnionPhysicsProxy::GetDisabled_External() const
+	{
+		if (!ensure(Particle_External))
+		{
+			return false;
+		}
+				
+		return Particle_External->Disabled();
+	}
+
+	void FClusterUnionPhysicsProxy::SetDisabled_External(const bool Disabled)
+	{
+		if (!Solver || !ensure(Particle_External))
+		{
+			return;
+		}
+
+		Solver->EnqueueCommandImmediate(
+			[this, Disabled]() mutable
+			{
+				if (Particle_Internal)
+				{
+					if (Particle_Internal->Disabled() != Disabled)
+					{
+						if (FPBDRigidsEvolutionGBF* Evolution = GetEvolution(this))
+						{
+							if (Disabled)
+							{
+								Evolution->DisableParticle(Particle_Internal);
+							}
+							else
+							{
+								Evolution->EnableParticle(Particle_Internal);
+							}
+						}
+					}
+				}
+			}
+		);
+	}
+
+	bool FClusterUnionPhysicsProxy::GetGravityEnabled_External() const
+	{
+		if (!ensure(Particle_External))
+		{
+			return false;
+		}
+
+		return Particle_External->GravityEnabled();
+	}
+
+	void FClusterUnionPhysicsProxy::SetGravityEnabled_External(const bool Enabled)
+	{
+		if (!Solver || !ensure(Particle_External))
+		{
+			return;
+		}
+
+		Solver->EnqueueCommandImmediate(
+			[this, Enabled]() mutable
+			{
+				if (Particle_Internal)
+				{
+					Particle_Internal->SetGravityEnabled(Enabled);
+				}
+			}
+		);
+	}
+	
 	void FClusterUnionPhysicsProxy::Wake_External()
 	{
 		if (!Solver || !ensure(Particle_External))
diff --git a/Engine/Source/Runtime/Experimental/Chaos/Public/PhysicsProxy/ClusterUnionPhysicsProxy.h b/Engine/Source/Runtime/Experimental/Chaos/Public/PhysicsProxy/ClusterUnionPhysicsProxy.h
--- a/Engine/Source/Runtime/Experimental/Chaos/Public/PhysicsProxy/ClusterUnionPhysicsProxy.h
+++ b/Engine/Source/Runtime/Experimental/Chaos/Public/PhysicsProxy/ClusterUnionPhysicsProxy.h
@@ -83,6 +83,15 @@ namespace Chaos
 		// Set Object State.
 		CHAOS_API EObjectStateType GetObjectState_External() const;
 		CHAOS_API void SetObjectState_External(EObjectStateType State);
+		
+		CHAOS_API ESleepType GetSleepType_External() const;
+		CHAOS_API void SetSleepType_External(const ESleepType SleepType);
+
+		CHAOS_API bool GetDisabled_External() const;
+		CHAOS_API void SetDisabled_External(const bool Disabled);
+
+		CHAOS_API bool GetGravityEnabled_External() const;
+		CHAOS_API void SetGravityEnabled_External(const bool Enabled);
 
 		// Explicitly wake the physics object. The wakes the physics object if it is sleeping and resets any sleep state (whether awake or sleeping)
 		CHAOS_API void Wake_External();
