Iris
* Added NetRootObjectFactory class
* This factory can be used by non-actor classes that want to replicate autonomously without being attached to an actor as a subobject.
* The factory does not currently support objects that want to be spatially filtered. Instead the factory should be overriden and users can gather the proper world location of the object manually.

#rnx

diff --git a/Engine/Source/Runtime/Engine/Private/Net/Iris/ReplicationSystem/NetEngineFactories.cpp b/Engine/Source/Runtime/Engine/Private/Net/Iris/ReplicationSystem/NetEngineFactories.cpp
--- a/Engine/Source/Runtime/Engine/Private/Net/Iris/ReplicationSystem/NetEngineFactories.cpp
+++ b/Engine/Source/Runtime/Engine/Private/Net/Iris/ReplicationSystem/NetEngineFactories.cpp
@@ -4,10 +4,10 @@
 
 #include "Net/Iris/ReplicationSystem/NetActorFactory.h"
 #include "Net/Iris/ReplicationSystem/NetSubObjectFactory.h"
+#include "Net/Iris/ReplicationSystem/NetRootObjectFactory.h"
 
 #include "Iris/ReplicationSystem/NetObjectFactoryRegistry.h"
 
-
 namespace UE::Net
 {
 
@@ -24,6 +24,7 @@ void InitEngineNetObjectFactories()
 	{
 		FNetObjectFactoryRegistry::RegisterFactory(UNetActorFactory::StaticClass(), UNetActorFactory::GetFactoryName());
 		FNetObjectFactoryRegistry::RegisterFactory(UNetSubObjectFactory::StaticClass(), UNetSubObjectFactory::GetFactoryName());
+		FNetObjectFactoryRegistry::RegisterFactory(UNetRootObjectFactory::StaticClass(), UNetRootObjectFactory::GetFactoryName());
 		bAreFactoriesRegistered = true;
 	}
 }
@@ -36,6 +37,7 @@ void ShutdownEngineNetObjectFactories()
 	{
 		FNetObjectFactoryRegistry::UnregisterFactory(UNetActorFactory::GetFactoryName());
 		FNetObjectFactoryRegistry::UnregisterFactory(UNetSubObjectFactory::GetFactoryName());
+		FNetObjectFactoryRegistry::UnregisterFactory(UNetRootObjectFactory::GetFactoryName());
 		bAreFactoriesRegistered = false;
 	}
 }
diff --git a/Engine/Source/Runtime/Engine/Private/Net/Iris/ReplicationSystem/NetRootObjectFactory.cpp b/Engine/Source/Runtime/Engine/Private/Net/Iris/ReplicationSystem/NetRootObjectFactory.cpp
new file mode 100644
--- /dev/null
+++ b/Engine/Source/Runtime/Engine/Private/Net/Iris/ReplicationSystem/NetRootObjectFactory.cpp
@@ -0,0 +1,303 @@
+// Copyright Epic Games, Inc. All Rights Reserved.
+
+#include "Net/Iris/ReplicationSystem/NetRootObjectFactory.h"
+
+#include "HAL/IConsoleManager.h"
+
+#include "Iris/Core/IrisLog.h"
+
+#include "Iris/ReplicationSystem/ObjectReplicationBridge.h"
+
+#include "Iris/Serialization/NetBitStreamUtil.h"
+#include "Iris/Serialization/NetBitStreamReader.h"
+#include "Iris/Serialization/NetBitStreamWriter.h"
+#include "Iris/Serialization/NetSerializationContext.h"
+#include "Iris/Serialization/ObjectNetSerializer.h"
+
+#include "UObject/Package.h"
+
+#include "Net/DataBunch.h"
+
+#include UE_INLINE_GENERATED_CPP_BY_NAME(NetRootObjectFactory)
+
+namespace UE::Net::Private
+{
+	static const FName RootObjectFactoryName(TEXT("RootObjectFactory"));
+
+	float DefaultRootObjectNetUpdateFrequency = 30.f;
+	static FAutoConsoleVariableRef CVarDefaultRootObjectNetUpdateFrequency(
+		TEXT("net.Iris.DefaultRootObjectNetUpdateFrequency"),
+		DefaultRootObjectNetUpdateFrequency,
+		TEXT("The default net update frequency of objects using the NetRootObjectFactory."));
+}
+
+FName UNetRootObjectFactory::GetFactoryName()
+{
+	return UE::Net::Private::RootObjectFactoryName;
+}
+
+void UNetRootObjectFactory::OnInit()
+{
+	Super::OnInit();
+}
+
+UNetObjectFactory::FInstantiateResult UNetRootObjectFactory::InstantiateReplicatedObjectFromHeader(const FInstantiateContext& Context, const UE::Net::FNetObjectCreationHeader* Header)
+{
+	using namespace UE::Net;
+
+	const FNetRootObjectCreationHeader* RootObjectHeader = static_cast<const FNetRootObjectCreationHeader*>(Header);
+
+	// For stable objects, find the object via its pathname and bind it
+	if (RootObjectHeader->bIsStable)
+	{
+		UObject* RootObjectPtr = Bridge->ResolveObjectReference(RootObjectHeader->Reference, Context.ResolveContext);
+
+		if (!RootObjectPtr)
+		{
+			UE_LOG(LogIris, Error, TEXT("%hs: Failed to resolve stable reference: %s . Could not find stable object to bind to."), __FUNCTION__, *Bridge->DescribeObjectReference(RootObjectHeader->Reference, Context.ResolveContext));
+			return FInstantiateResult();
+		}
+
+		UE_LOG(LogIris, Verbose, TEXT("%hs: Found stable object: %s with Handle: %s via: %s"), __FUNCTION__, *GetNameSafe(RootObjectPtr), *Context.Handle.ToString(), *Bridge->DescribeObjectReference(RootObjectHeader->Reference, Context.ResolveContext));
+
+		return FInstantiateResult{ .Instance = RootObjectPtr };
+	}
+
+	// For preregistered objects, find the object in the list
+	if (RootObjectHeader->bIsPreregistered)
+	{
+		// Nothing useable in the header
+
+		UObject* RegisteredObject = Bridge->GetPreRegisteredObject(Context.Handle);
+		if (!RegisteredObject)
+		{
+			UE_LOG(LogIris, Error, TEXT("%hs Unable to find pre-registered root object: %s"), __FUNCTION__, *Context.Handle.ToString());
+			return FInstantiateResult();
+		}
+
+		UE_LOG(LogIris, Verbose, TEXT("%hs: Found pre-registered object: %s with Handle: %s"), __FUNCTION__, *GetNameSafe(RegisteredObject), *Context.Handle.ToString());
+
+		return FInstantiateResult{ .Instance = RegisteredObject };
+	}
+
+	// For dynamic objects, spawn them from the template
+	
+	// Find the template
+	UObject* Template = Bridge->ResolveObjectReference(RootObjectHeader->Reference, Context.ResolveContext);
+	if (!Template)
+	{
+		UE_LOG(LogIris, Error, TEXT("%hs Unable to instantiate root object, failed to resolve template: %s"), __FUNCTION__, *Bridge->DescribeObjectReference(RootObjectHeader->Reference, Context.ResolveContext));
+		return FInstantiateResult();
+	}
+
+	UObject* DynamicRootObject = NewObject<UObject>(GetTransientPackage(), Template->GetClass(), NAME_None, RF_NoFlags, Template);
+
+	FInstantiateResult Result
+	{ 
+		.Instance = DynamicRootObject,
+		.Template = Template, // Force the template to be used as the default init state source. It's possible GetArchetype() will be different then the template
+	};
+
+	// Dynamic objects may be destroyed.
+	Result.Flags |= EReplicationBridgeCreateNetRefHandleResultFlags::AllowDestroyInstanceFromRemote;
+
+	UE_LOG(LogIris, Verbose, TEXT("%hs: Spawned %s from template: %s with Handle: %s"), __FUNCTION__, *GetNameSafe(DynamicRootObject), *GetNameSafe(Template), *Context.Handle.ToString());
+			
+	return Result;
+}
+
+void UNetRootObjectFactory::DetachedFromReplication(const FDestroyedContext& Context)
+{
+	UE_LOG(LogIris, Verbose, TEXT("%hs: %s | Reason: %s | Flags: %s"), __FUNCTION__, *GetNameSafe(Context.DestroyedInstance), LexToString(Context.DestroyReason), LexToString(Context.DestroyFlags));
+
+	if (Context.DestroyReason == EReplicationBridgeDestroyInstanceReason::DoNotDestroy ||
+		Context.DestroyReason == EReplicationBridgeDestroyInstanceReason::TearOff)
+	{
+		return;
+	}
+
+	if (EnumHasAnyFlags(Context.DestroyFlags, EReplicationBridgeDestroyInstanceFlags::AllowDestroyInstanceFromRemote))
+	{
+		Context.DestroyedInstance->PreDestroyFromReplication();
+
+		//TODO: Some objects will want to opt-out of being MarkAsGarbage to support verse semantics
+		Context.DestroyedInstance->MarkAsGarbage();
+	}
+}
+
+TOptional<UNetObjectFactory::FWorldInfoData> UNetRootObjectFactory::GetWorldInfo(const FWorldInfoContext& Context) const
+{
+	//TODO: Currently the only way to add support for world infos to a class type is to derive from this factory.
+	//		We should offer a faster customization path without forcing users to declare new concrete factories for every class type.
+
+	ensureMsgf(false, TEXT("UNetRootObjectFactory::GetWorldInfo does not support object: %s with a world location. Classes that want to support spatial filtering need to implement the function in their own factory."), *GetNameSafe(Context.Instance));
+	return NullOpt;
+}
+
+float UNetRootObjectFactory::GetPollFrequency(UE::Net::FNetRefHandle RootObjectHandle, UObject* RootObjectInstance)
+{
+	//TODO: Figure what we want to do with the default frequency of these objects.
+	//      Could be time to add a enum (every frame, 30htz, etc.)
+
+	return UE::Net::Private::DefaultRootObjectNetUpdateFrequency;
+}
+
+TUniquePtr<UE::Net::FNetObjectCreationHeader> UNetRootObjectFactory::CreateAndFillHeader(UE::Net::FNetRefHandle Handle)
+{
+	using namespace UE::Net;
+
+	UObject* RootObjectPtr = Bridge->GetReplicatedObject(Handle);
+
+	if (!RootObjectPtr)
+	{
+		UE_LOG(LogIris, Error, TEXT("%hs: no instance associated with handle: %s"), __FUNCTION__, *Handle.ToString());
+		return nullptr;
+	}
+
+	TUniquePtr<FNetRootObjectCreationHeader> RootObjectHeader(new FNetRootObjectCreationHeader);
+
+	const bool bSuccess = FillHeader(RootObjectHeader.Get(), Handle, RootObjectPtr);
+
+	if (!bSuccess)
+	{
+		return nullptr;
+	}
+	
+	return RootObjectHeader;
+}
+
+bool UNetRootObjectFactory::FillHeader(UE::Net::FNetRootObjectCreationHeader* RootObjectHeader, UE::Net::FNetRefHandle Handle, const UObject* RootObjectPtr)
+{
+	using namespace UE::Net;
+
+	const FNetObjectReference RootObjectReference = Bridge->GetOrCreateObjectReference(RootObjectPtr);
+
+	if (!RootObjectReference.IsValid())
+	{
+		UE_LOG(LogIris, Error, TEXT("%hs: could not create NetReference for %s (handle: %s)"), __FUNCTION__, *GetNameSafe(RootObjectPtr), *Handle.ToString());
+		return false;
+	}
+
+	// Stable object
+	if (Handle.IsStatic() || RootObjectPtr->IsNameStableForNetworking())
+	{
+		RootObjectHeader->bIsStable = true;
+		RootObjectHeader->Reference = RootObjectReference;
+
+		return true;
+	}
+
+	// Preregistered object
+	if (Bridge->IsNetRefHandlePreRegistered(Handle))
+	{
+		RootObjectHeader->bIsPreregistered = true;
+
+		return true;
+	}
+
+	// Dynamic object
+	const UObject* Archetype = RootObjectPtr->GetArchetype();
+	checkf(Archetype, TEXT("No archetype found for: %s"), *GetNameSafe(RootObjectPtr));
+	checkf(Archetype->NeedsLoadForClient(), TEXT("Archetype: %s of object: %s will not exist on the client"), *GetNameSafe(Archetype), *GetNameSafe(RootObjectPtr));
+		
+	const FNetObjectReference ArchetypeReference = Bridge->GetOrCreateObjectReference(Archetype);
+	if (!ArchetypeReference.IsValid())
+	{
+		UE_LOG(LogIris, Error, TEXT("%hs: could not create NetReference for archetype: %s"), __FUNCTION__, *GetNameSafe(Archetype));
+		return false;
+	}
+
+	RootObjectHeader->Reference = ArchetypeReference;
+	
+	return true;
+}
+
+TUniquePtr<UE::Net::FNetObjectCreationHeader> UNetRootObjectFactory::CreateAndDeserializeHeader(const UE::Net::FCreationHeaderContext& Context)
+{
+	using namespace UE::Net;
+
+	TUniquePtr<FNetRootObjectCreationHeader> RootObjectHeader(new FNetRootObjectCreationHeader);
+	const bool bSuccess = RootObjectHeader->Deserialize(Context);
+
+	if (!bSuccess)
+	{
+		return nullptr;
+	}
+	
+	return RootObjectHeader;
+}
+
+bool UNetRootObjectFactory::SerializeHeader(const UE::Net::FCreationHeaderContext& Context, const UE::Net::FNetObjectCreationHeader* Header)
+{
+	using namespace UE::Net;
+
+	const FNetRootObjectCreationHeader* RootObjectHeader = static_cast<const FNetRootObjectCreationHeader*>(Header);
+	return RootObjectHeader->Serialize(Context);
+}
+
+//------------------------------------------------------------------------
+// FNetRootObjectCreationHeader
+//------------------------------------------------------------------------
+
+namespace UE::Net
+{
+
+FString FNetRootObjectCreationHeader:: ToString() const
+{
+	return FString::Printf(TEXT("\n\tFNetRootObjectCreationHeader (ProtocolId:0x%x):\n\t"
+								"%s\n\t"
+								"Reference=%s")
+						   , GetProtocolId()
+						   , bIsStable ? TEXT("Stable") : (bIsPreregistered ? TEXT("Preregistered") : TEXT("Dynamic"))
+						   , *Reference.ToString()
+						   );
+}
+
+bool FNetRootObjectCreationHeader::Serialize(const UE::Net::FCreationHeaderContext& Context) const
+{
+	FNetBitStreamWriter* Writer = Context.Serialization.GetBitStreamWriter();
+	
+	Writer->WriteBool(bIsStable);
+
+	if (bIsStable)
+	{
+		WriteFullNetObjectReference(Context.Serialization, Reference);
+	}
+	else
+	{
+		Writer->WriteBool(bIsPreregistered);
+
+		if (!bIsPreregistered)
+		{
+			WriteFullNetObjectReference(Context.Serialization, Reference);
+		}
+	}
+
+	return Writer->IsOverflown() == false;
+}
+
+bool FNetRootObjectCreationHeader::Deserialize(const UE::Net::FCreationHeaderContext& Context)
+{
+	FNetBitStreamReader* Reader = Context.Serialization.GetBitStreamReader();
+
+	bIsStable = Reader->ReadBool();
+
+	if (bIsStable)
+	{
+		ReadFullNetObjectReference(Context.Serialization, Reference);
+	}
+	else
+	{
+		bIsPreregistered = Reader->ReadBool();
+
+		if (!bIsPreregistered)
+		{
+			ReadFullNetObjectReference(Context.Serialization, Reference);
+		}
+	}
+
+	return Reader->IsOverflown() == false;
+}
+
+} // end namespace UE::Net
+
diff --git a/Engine/Source/Runtime/Engine/Public/Net/Iris/ReplicationSystem/NetRootObjectFactory.h b/Engine/Source/Runtime/Engine/Public/Net/Iris/ReplicationSystem/NetRootObjectFactory.h
new file mode 100644
--- /dev/null
+++ b/Engine/Source/Runtime/Engine/Public/Net/Iris/ReplicationSystem/NetRootObjectFactory.h
@@ -0,0 +1,93 @@
+// Copyright Epic Games, Inc. All Rights Reserved.
+
+#pragma once
+
+#include "Iris/ReplicationSystem/NetObjectFactory.h"
+
+#include "Iris/Core/NetObjectReference.h"
+
+#include "NetRootObjectFactory.generated.h"
+
+
+namespace UE::Net
+{
+
+class FNetRootObjectCreationHeader : public FNetObjectCreationHeader
+{
+public:
+	
+	// Option 1: Root Object is stable
+	// Option 2: Root Object is pre-registered
+	// Option 3: Root Object is dynamic
+
+	/** References a pathname we want to load on the client. Can be either the object itself when bIsStable or the template to clone from when dynamic. */
+	FNetObjectReference Reference;
+
+	//TODO: Add optional way to send the object's Outer
+
+	/** When stable we expect the object to exist on the client */
+	bool bIsStable = false;
+
+	/** When preregistered we will find the object via it's NetHandle in the bridge */
+	bool bIsPreregistered = false;
+
+
+public:
+
+	ENGINE_API virtual FString ToString() const override;
+
+	/** Serialize/Deserialize the information held in the header */
+	ENGINE_API virtual bool Serialize(const UE::Net::FCreationHeaderContext& Context) const;
+	ENGINE_API virtual bool Deserialize(const UE::Net::FCreationHeaderContext& Context);
+};
+
+
+} // end namespace UE::Net
+
+
+/**
+ * Responsible for creating headers allowing remote factories to spawn replicated actors
+ */
+UCLASS(MinimalAPI)
+class UNetRootObjectFactory : public UNetObjectFactory
+{
+	GENERATED_BODY()
+
+public:
+
+	ENGINE_API static FName GetFactoryName();
+
+public:
+
+	// UNetObjectFactory interface
+
+	ENGINE_API virtual void OnInit() override;
+
+	ENGINE_API virtual FInstantiateResult InstantiateReplicatedObjectFromHeader(const FInstantiateContext& Context, const UE::Net::FNetObjectCreationHeader* Header) override;
+
+	ENGINE_API virtual void DetachedFromReplication(const FDestroyedContext& Context) override;
+
+	ENGINE_API virtual TOptional<FWorldInfoData> GetWorldInfo(const FWorldInfoContext& Context) const override;
+
+	ENGINE_API virtual float GetPollFrequency(UE::Net::FNetRefHandle RootObjectHandle, UObject* RootObjectInstance) override;
+
+protected:
+
+	// UNetObjectFactory interface
+
+	ENGINE_API virtual TUniquePtr<UE::Net::FNetObjectCreationHeader> CreateAndFillHeader(UE::Net::FNetRefHandle Handle) override;
+
+	ENGINE_API virtual TUniquePtr<UE::Net::FNetObjectCreationHeader> CreateAndDeserializeHeader(const UE::Net::FCreationHeaderContext& Context) override;
+
+	ENGINE_API virtual bool SerializeHeader(const UE::Net::FCreationHeaderContext& Context, const UE::Net::FNetObjectCreationHeader* Header)  override;
+
+protected:
+
+	// Callbacks for derived factories to use
+
+	// Methods available to derived classes that want to reuse basic functionalities
+
+	/** Fill the basic header info for a given object */
+	ENGINE_API bool FillHeader(UE::Net::FNetRootObjectCreationHeader* RootObjectHeader, UE::Net::FNetRefHandle Handle, const UObject* RootObjectPtr);
+
+};
