Fixed TSet remove during iteration in UWorld::TransferBlueprintDebugReferences.

Swapped out the map before iterating, also made a minor reduction in FString allocations.

diff --git a/Engine/Source/Runtime/Engine/Private/World.cpp b/Engine/Source/Runtime/Engine/Private/World.cpp
--- a/Engine/Source/Runtime/Engine/Private/World.cpp
+++ b/Engine/Source/Runtime/Engine/Private/World.cpp
@@ -3152,31 +3152,33 @@ void UWorld::TransferBlueprintDebugReferences(UWorld* NewWorld)
 {
 #if WITH_EDITOR
 	// First create a list of blueprints that already exist in the new world
-	TArray< FString > NewWorldExistingBlueprintNames;
+	TArray<FName> NewWorldExistingBlueprintNames;
 	for (FBlueprintToDebuggedObjectMap::TIterator It(NewWorld->BlueprintObjectsBeingDebugged); It; ++It)
 	{
 		if (UBlueprint* TargetBP = It.Key().Get())
 		{	
-			NewWorldExistingBlueprintNames.AddUnique( TargetBP->GetName() );
+			NewWorldExistingBlueprintNames.AddUnique( TargetBP->GetFName() );
 		}
 	}
 
+	// SetObjectBeingDebugged may modify BlueprintObjectsBeingDebugged during iteration, so swap out before iterating.
+	const FBlueprintToDebuggedObjectMap LocalBlueprintObjectsBeingDebugged = MoveTemp(BlueprintObjectsBeingDebugged);
 	// Move debugging object associations across from the old world to the new world that are not already there
-	for (FBlueprintToDebuggedObjectMap::TIterator It(BlueprintObjectsBeingDebugged); It; ++It)
+	for (FBlueprintToDebuggedObjectMap::TConstIterator It(LocalBlueprintObjectsBeingDebugged); It; ++It)
 	{
 		if (UBlueprint* TargetBP = It.Key().Get())
-		{	
-			FString SourceName = TargetBP->GetName();
+		{
+			const FName SourceName = TargetBP->GetFName();
 			// If this blueprint is not already listed in the ones bieng debugged in the new world, add it.
 			if( NewWorldExistingBlueprintNames.Find( SourceName ) == INDEX_NONE )			
 			{
-				TWeakObjectPtr<UObject>& WeakTargetObject = It.Value();
+				const TWeakObjectPtr<UObject>& WeakTargetObject = It.Value();
 				UObject* NewTargetObject = nullptr;
 				bool bForceClear = false;
 
 				if (WeakTargetObject.IsValid())
 				{
-					UObject* OldTargetObject = WeakTargetObject.Get();
+					const UObject* OldTargetObject = WeakTargetObject.Get();
 					check(OldTargetObject);
 
 					// We don't map from PIE objects in a client world back to editor objects, as that will transfer to the server on the next execution
@@ -3226,7 +3228,7 @@ void UWorld::TransferBlueprintDebugReferences(UWorld* NewWorld)
 							UE_LOG(LogWorld, Warning, TEXT("  NewTargetObject path is %s"), *NewTargetObject->GetPathName());
 							UE_LOG(LogWorld, Warning, TEXT("  NewTargetObject class path is %s"), *NewTargetObject->GetClass()->GetPathName());
 
-							UObject* OldTargetObject = WeakTargetObject.Get();
+							const UObject* OldTargetObject = WeakTargetObject.Get();
 							UE_LOG(LogWorld, Warning, TEXT("  OldObject path is %s"), *OldTargetObject->GetPathName());
 							UE_LOG(LogWorld, Warning, TEXT("  OldObject class path is %s"), *OldTargetObject->GetClass()->GetPathName());
 
