Move a RW lock to be transactionally safe.

#rnx

diff --git a/Engine/Source/Runtime/Engine/Private/Curves/Spline.cpp b/Engine/Source/Runtime/Engine/Private/Curves/Spline.cpp
--- a/Engine/Source/Runtime/Engine/Private/Curves/Spline.cpp
+++ b/Engine/Source/Runtime/Engine/Private/Curves/Spline.cpp
@@ -6,6 +6,7 @@
 #include "Async/TransactionallySafeMutex.h"
 #include "Components/SplineComponent.h"	// for FSplineCurves, ideally removed
 #include "Misc/Base64.h"
+#include "Misc/TransactionallySafeRWLock.h"
 #include "Serialization/MemoryReader.h"
 #include "Serialization/MemoryWriter.h"
 #include "UObject/FortniteMainBranchObjectVersion.h"
@@ -311,7 +312,7 @@ private:
 	FInterpCurveFloat ReparamTable;
 	FChangeStamp ReparamTableNextVersion;
 	uint32 ReparamTableVersion = 0;
-	FRWLock ReparamTableRWLock;
+	FTransactionallySafeRWLock ReparamTableRWLock;
 	
 	static inline const FName RotationAttrName = FName("Rotation");
 	static inline const FName ScaleAttrName = FName("Scale");
@@ -2394,7 +2395,7 @@ float FNewSpline::GetParameterAtDistance(float Distance)
 {
 	UpdateReparamTable();
 	
-	FReadScopeLock Lock(ReparamTableRWLock);
+	UE::TReadScopeLock Lock(ReparamTableRWLock);
 	
 	return ReparamTable.Eval(Distance);
 }
@@ -2408,7 +2409,7 @@ float FNewSpline::GetDistanceAtParameter(float Parameter)
 	
 	UpdateReparamTable();
 	
-	FReadScopeLock Lock(ReparamTableRWLock);
+	UE::TReadScopeLock Lock(ReparamTableRWLock);
 
 	const float ParameterMax = static_cast<float>(GetNumberOfSegments());
 	const float Key = (Parameter / ParameterMax) * (ReparamTable.Points.Num() - 1);
@@ -2560,7 +2561,7 @@ void FNewSpline::UpdateReparamTable()
 		return;
 	}
 
-	FWriteScopeLock Lock(ReparamTableRWLock);
+	UE::TWriteScopeLock Lock(ReparamTableRWLock);
 
 	if (!IsReparamTableDirty())
 	{
