UE-203619 - Bug with setting the ChildActorClass in ChildActorComponents in Blueprint

When a child actor component is registered, we typically spawn a new actor, which also destroys the old instance data cache and updates it with a new instance. However, we weren't destroying the instance data if ChildActorClass happened to be null. This could lead to a scenario where the cache might hold onto its previous value. As a result, re-running the construction scripts via PostEditChange could accidentally reset the ChildActorClass to the cached value instead of null.

The fix here is to explicitly reset the cached instanced data in OnRegister when ChildActorClass is null.

diff --git a/Engine/Source/Runtime/Engine/Classes/Components/ChildActorComponent.h b/Engine/Source/Runtime/Engine/Classes/Components/ChildActorComponent.h
--- a/Engine/Source/Runtime/Engine/Classes/Components/ChildActorComponent.h
+++ b/Engine/Source/Runtime/Engine/Classes/Components/ChildActorComponent.h
@@ -160,6 +160,8 @@ private:
 	/** When true, does not modify ChildActorName during spawn such as removing _UAID_ */
 	uint8 bChildActorNameIsExact:1;
 
+	void DeleteCachedInstanceData();
+
 #if WITH_EDITOR
 	ENGINE_API virtual void SetPackageExternal(bool bExternal, bool bShouldDirty) override;
 #endif
diff --git a/Engine/Source/Runtime/Engine/Private/Components/ChildActorComponent.cpp b/Engine/Source/Runtime/Engine/Private/Components/ChildActorComponent.cpp
--- a/Engine/Source/Runtime/Engine/Private/Components/ChildActorComponent.cpp
+++ b/Engine/Source/Runtime/Engine/Private/Components/ChildActorComponent.cpp
@@ -108,6 +108,11 @@ void UChildActorComponent::OnRegister()
 	{
 		CreateChildActor();
 	}
+	else
+	{
+		DeleteCachedInstanceData();
+		CachedInstanceData = new FChildActorComponentInstanceData(this);
+	}
 }
 
 void UChildActorComponent::Serialize(FArchive& Ar)
@@ -194,6 +199,15 @@ void UChildActorComponent::Serialize(FArchive& Ar)
 #endif
 }
 
+void UChildActorComponent::DeleteCachedInstanceData()
+{
+	if (CachedInstanceData)
+	{
+		delete CachedInstanceData;
+		CachedInstanceData = nullptr;
+	}
+}
+
 #if WITH_EDITOR
 void UChildActorComponent::SetPackageExternal(bool bExternal, bool bShouldDirty)
 {
@@ -226,11 +240,7 @@ void UChildActorComponent::PostEditImport()
 	}
 
 	// Any cached instance data is invalid if we've had data imported in to us
-	if (CachedInstanceData)
-	{
-		delete CachedInstanceData;
-		CachedInstanceData = nullptr;
-	}
+	DeleteCachedInstanceData();
 }
 
 void UChildActorComponent::PostEditChangeProperty(FPropertyChangedEvent& PropertyChangedEvent)
@@ -448,11 +458,7 @@ void UChildActorComponent::BeginDestroy()
 {
 	Super::BeginDestroy();
 
-	if (CachedInstanceData)
-	{
-		delete CachedInstanceData;
-		CachedInstanceData = nullptr;
-	}
+	DeleteCachedInstanceData();
 }
 
 TStructOnScope<FActorComponentInstanceData> UChildActorComponent::GetComponentInstanceData() const
@@ -881,12 +887,7 @@ void UChildActorComponent::CreateChildActor(TFunction<void(AActor*)> CustomizerF
 		}
 	}
 
-	// This is no longer needed
-	if (CachedInstanceData)
-	{
-		delete CachedInstanceData;
-		CachedInstanceData = nullptr;
-	}
+	DeleteCachedInstanceData();
 
 	if (ChildActor)
 	{
@@ -933,11 +934,7 @@ void UChildActorComponent::DestroyChildActor()
 			if (!bIsChildActorPendingKillOrUnreachable)
 			{
 #if WITH_EDITOR
-				if (CachedInstanceData)
-				{
-					delete CachedInstanceData;
-					CachedInstanceData = nullptr;
-				}
+				DeleteCachedInstanceData();
 #else
 				check(!CachedInstanceData);
 #endif
