[HLOD] Reset loaders of source actors before moving them to external level
Resetting loaders after the move could trigger unnecessary object creation via ConditionalFlushAsyncLoadingForLinkers

diff --git a/Engine/Source/Runtime/Engine/Private/WorldPartition/HLOD/HLODSourceActorsFromCell.cpp b/Engine/Source/Runtime/Engine/Private/WorldPartition/HLOD/HLODSourceActorsFromCell.cpp
--- a/Engine/Source/Runtime/Engine/Private/WorldPartition/HLOD/HLODSourceActorsFromCell.cpp
+++ b/Engine/Source/Runtime/Engine/Private/WorldPartition/HLOD/HLODSourceActorsFromCell.cpp
@@ -73,16 +73,20 @@ bool UWorldPartitionHLODSourceActorsFromCell::LoadSourceActors(bool& bOutDirty,
 
 	if (FWorldPartitionLevelHelper::LoadActors(MoveTemp(Params)))
 	{
-		TArray<UPackage*> ModifiedPackages;
-		FWorldPartitionLevelHelper::MoveExternalActorsToLevel(ActorsCopy, TargetWorld->PersistentLevel, ModifiedPackages);
-		
-		TArray<UObject*> PackagesToResetLoaders;
-		PackagesToResetLoaders.Reserve(ModifiedPackages.Num());
-		for (UPackage* Package : ModifiedPackages)
+		TArray<UObject*> ActorsToResetLoaders;
+		ActorsToResetLoaders.Reserve(ActorsCopy.Num());
+		for (const FWorldPartitionRuntimeCellObjectMapping& ActorMapping : ActorsCopy)
 		{
-			PackagesToResetLoaders.Add(Package);
+			AActor* Actor = FindObject<AActor>(nullptr, *ActorMapping.LoadedPath.ToString());
+			if (Actor)
+			{
+				ActorsToResetLoaders.Add(Actor);
+			}
 		}
-		ResetLoaders(PackagesToResetLoaders);
+		ResetLoaders(ActorsToResetLoaders);
+	
+		TArray<UPackage*> ModifiedPackages;
+		FWorldPartitionLevelHelper::MoveExternalActorsToLevel(ActorsCopy, TargetWorld->PersistentLevel, ModifiedPackages);
 
 		return true;
 	}
