World Partition - HLOD: When building Level HLODs, call ResetLoaders() on the source level package to release the file lock

diff --git a/Engine/Source/Runtime/Engine/Private/WorldPartition/HLOD/HLODSourceActorsFromLevel.cpp b/Engine/Source/Runtime/Engine/Private/WorldPartition/HLOD/HLODSourceActorsFromLevel.cpp
--- a/Engine/Source/Runtime/Engine/Private/WorldPartition/HLOD/HLODSourceActorsFromLevel.cpp
+++ b/Engine/Source/Runtime/Engine/Private/WorldPartition/HLOD/HLODSourceActorsFromLevel.cpp
@@ -5,8 +5,10 @@
 #if WITH_EDITOR
 #include "WorldPartition/HLOD/HLODActor.h"
 #include "WorldPartition/HLOD/HLODHashBuilder.h"
+#include "Engine/Level.h"
 #include "Engine/LevelStreamingDynamic.h"
 #include "Engine/World.h"
+#include "UObject/Linker.h"
 #endif
 
 #include UE_INLINE_GENERATED_CPP_BY_NAME(HLODSourceActorsFromLevel)
@@ -23,7 +25,14 @@ bool UWorldPartitionHLODSourceActorsFromLevel::LoadSourceActors(bool& bOutDirty,
 	bool bSuccess = false;
 	ULevelStreamingDynamic* LevelStreaming = ULevelStreamingDynamic::LoadLevelInstanceBySoftObjectPtr(TargetWorld, SourceLevel, FTransform::Identity, bSuccess);
 	TargetWorld->FlushLevelStreaming();
-	return bSuccess && LevelStreaming && LevelStreaming->GetLoadedLevel();
+
+	ULevel* LoadedLevel = LevelStreaming ? LevelStreaming->GetLoadedLevel() : nullptr;
+	if (LoadedLevel != nullptr)
+	{
+		ResetLoaders(LoadedLevel->GetPackage());
+	}
+
+	return bSuccess && LoadedLevel != nullptr;
 }
 
 void UWorldPartitionHLODSourceActorsFromLevel::ComputeHLODHash(FHLODHashBuilder& InHashBuilder) const
