[World] Move manifest bCookSaveGeneratorPackage control to sub-splitters

#rnx

diff --git a/Engine/Source/Runtime/Engine/Private/WorldCookPackageSplitter.cpp b/Engine/Source/Runtime/Engine/Private/WorldCookPackageSplitter.cpp
--- a/Engine/Source/Runtime/Engine/Private/WorldCookPackageSplitter.cpp
+++ b/Engine/Source/Runtime/Engine/Private/WorldCookPackageSplitter.cpp
@@ -5,7 +5,7 @@
 #if WITH_EDITOR
 #include "Engine/World.h"
 #include "Editor.h"
-#include "WorldPartition/WorldPartitionSettings.h"
+#include "WorldPartition/WorldPartitionLog.h"
 
 REGISTER_COOKPACKAGE_SPLITTER(FWorldCookPackageSplitter, UWorld);
 
@@ -115,8 +115,10 @@ ICookPackageSplitter::FGenerationManifest FWorldCookPackageSplitter::ReportGener
 	check(CookPackageSubSplitters.Num());
 
 	ICookPackageSplitter::FGenerationManifest Manifest;
+	TOptional<bool> SkipSaveGeneratorPackage;
+
 	ForEachCookPackageSubSplitters(
-		[this, &Manifest, OwnerPackage, OwnerWorld](FWorldCookPackageSplitter::ISubSplitter* WorldCookPackageSplitter)
+		[this, &Manifest, &SkipSaveGeneratorPackage, OwnerPackage, OwnerWorld](FWorldCookPackageSplitter::ISubSplitter* WorldCookPackageSplitter)
 		{
 			const TArray<ICookPackageSplitter::FGeneratedPackage> GenerateList = WorldCookPackageSplitter->GetGenerateList(OwnerPackage);
 			
@@ -136,12 +138,27 @@ ICookPackageSplitter::FGenerationManifest FWorldCookPackageSplitter::ReportGener
 				check(!SplittersGenerateListMap.Contains(FullPackageName));
 				SplittersGenerateListMap.Add(FullPackageName, WorldCookPackageSplitter);
 			}
+
+			// Any sub-splitter can disable save of the map
+			TOptional<bool> SkipSave = WorldCookPackageSplitter->SkipSaveGeneratorPackage();
+			if (SkipSave.IsSet())
+			{
+				if (SkipSaveGeneratorPackage.IsSet() && SkipSave.GetValue() != SkipSaveGeneratorPackage.GetValue())
+				{
+					UE_LOG(LogWorldPartition, Error, TEXT("[Cook][ReportGenerationManifest] At least 2 sub-splitters disagree on whether to skip the generator package save, which also means 2 sub-splitters are attempting to handle the same package. It will stay enabled as a consequence."));
+					SkipSaveGeneratorPackage = false; // Do not skip, as this will avoid missing packages and be easier to debug for unconcerned users.
+				}
+				else
+				{
+					SkipSaveGeneratorPackage = SkipSave.GetValue();
+				}
+			}
 		});
 
-	FWorldDelegates::OnWorldCleanup.AddRaw(this, &FWorldCookPackageSplitter::OnWorldCleanup);
+	// If any sub-splitter requested that we skip the generator package save, and none requested that we keep it, disable the save. Otherwise the default is to save.
+	Manifest.bCookSaveGeneratorPackage = !SkipSaveGeneratorPackage.Get(false);
 
-	// If we have EDL cook filters, we are trying to cook the EDLs in isolation. Do not save the map even though it was loaded.
-	Manifest.bCookSaveGeneratorPackage = UWorldPartitionSettings::GetEDLCookFilters().IsEmpty();
+	FWorldDelegates::OnWorldCleanup.AddRaw(this, &FWorldCookPackageSplitter::OnWorldCleanup);
 
 	return Manifest;
 }
diff --git a/Engine/Source/Runtime/Engine/Private/WorldPartition/WorldPartitionCookPackageSplitter.cpp b/Engine/Source/Runtime/Engine/Private/WorldPartition/WorldPartitionCookPackageSplitter.cpp
--- a/Engine/Source/Runtime/Engine/Private/WorldPartition/WorldPartitionCookPackageSplitter.cpp
+++ b/Engine/Source/Runtime/Engine/Private/WorldPartition/WorldPartitionCookPackageSplitter.cpp
@@ -104,6 +104,16 @@ bool FWorldPartitionCookPackageSplitter::PopulateGeneratorPackage(ICookPackageSp
 	return bIsSuccess;
 }
 
+TOptional<bool> FWorldPartitionCookPackageSplitter::SkipSaveGeneratorPackage()
+{
+	// If we have EDL cook filters, we are trying to cook the EDLs in isolation. Do not save the map even though it was loaded.
+	if (UWorldPartitionSettings::GetEDLCookFilters().Num() > 0)
+	{
+		return true;
+	}
+	return TOptional<bool>();
+}
+
 void FWorldPartitionCookPackageSplitter::Teardown(ICookPackageSplitter::ETeardown Status)
 {
 	UE_LOG(LogWorldPartition, Display, TEXT("[Cook] Debug(TearDown): OwnerWorld=%s"), *GetFullNameSafe(World));
diff --git a/Engine/Source/Runtime/Engine/Private/WorldPartition/WorldPartitionCookPackageSplitter.h b/Engine/Source/Runtime/Engine/Private/WorldPartition/WorldPartitionCookPackageSplitter.h
--- a/Engine/Source/Runtime/Engine/Private/WorldPartition/WorldPartitionCookPackageSplitter.h
+++ b/Engine/Source/Runtime/Engine/Private/WorldPartition/WorldPartitionCookPackageSplitter.h
@@ -15,6 +15,7 @@ public:
 	virtual TOptional<bool> RequiresEmptyPackageBeforePopulate(const UPackage* Package) override;
 	virtual bool PopulateGeneratedPackage(ICookPackageSplitter::FPopulateContext& PopulateContext) override;
 	virtual bool PopulateGeneratorPackage(ICookPackageSplitter::FPopulateContext& PopulateContext) override;
+	virtual TOptional<bool> SkipSaveGeneratorPackage() override;
 	virtual void Teardown(ICookPackageSplitter::ETeardown Status) override;
 	//~End FWorldCookPackageSplitter::ISubSplitter interface
 
diff --git a/Engine/Source/Runtime/Engine/Public/WorldCookPackageSplitter.h b/Engine/Source/Runtime/Engine/Public/WorldCookPackageSplitter.h
--- a/Engine/Source/Runtime/Engine/Public/WorldCookPackageSplitter.h
+++ b/Engine/Source/Runtime/Engine/Public/WorldCookPackageSplitter.h
@@ -26,6 +26,11 @@ public:
 		virtual TOptional<bool> RequiresEmptyPackageBeforePopulate(const UPackage* Package) = 0;
 		virtual bool PopulateGeneratedPackage(ICookPackageSplitter::FPopulateContext& PopulateContext) = 0;
 		virtual bool PopulateGeneratorPackage(ICookPackageSplitter::FPopulateContext& PopulateContext) = 0;
+		/**
+		 * This can return true to disable the save of the generator package, for fine-tuned output control. Return unset by default.
+		 * All splitters should return the same value, or leave unset (as only one sub-splitter should be responsible for each package), otherwise the resulting mismatch will log an error.
+		 */
+		virtual TOptional<bool> SkipSaveGeneratorPackage() = 0;
 		virtual void Teardown(ICookPackageSplitter::ETeardown Status) = 0;
 	};
 
