Fix for race condition when setting skin weight buffer.

diff --git a/Engine/Source/Runtime/Engine/Private/SkeletalRenderGPUSkin.cpp b/Engine/Source/Runtime/Engine/Private/SkeletalRenderGPUSkin.cpp
--- a/Engine/Source/Runtime/Engine/Private/SkeletalRenderGPUSkin.cpp
+++ b/Engine/Source/Runtime/Engine/Private/SkeletalRenderGPUSkin.cpp
@@ -1793,6 +1793,7 @@ void FSkeletalMeshObjectGPUSkin::FSkeletalMeshObjectLOD::UpdateSkinWeights(const
 			ENQUEUE_RENDER_COMMAND(UpdateSkinWeightsGPUSkin)(UE::RenderCommandPipe::SkeletalMesh,
 				[NewMeshObjectWeightBuffer, VertexBuffers, Self](FRHICommandList& RHICmdList)
 			{
+				Self->MeshObjectWeightBuffer_RenderThread = NewMeshObjectWeightBuffer;
 				Self->GPUSkinVertexFactories.UpdateVertexFactoryData(VertexBuffers);
 			});
 		}
@@ -1997,7 +1998,7 @@ FGPUBaseSkinVertexFactory const* FSkeletalMeshObjectGPUSkin::GetBaseSkinVertexFa
 const FSkinWeightVertexBuffer* FSkeletalMeshObjectGPUSkin::GetSkinWeightVertexBuffer(int32 LODIndex) const
 {
 	checkSlow(LODs.IsValidIndex(LODIndex));
-	return LODs[LODIndex].MeshObjectWeightBuffer;
+	return LODs[LODIndex].MeshObjectWeightBuffer_RenderThread;
 }
 
 FSkinningSceneExtensionProxy* FSkeletalMeshObjectGPUSkin::CreateSceneExtensionProxy(const USkinnedAsset* InSkinnedAsset, bool bAllowScaling) const
@@ -2559,6 +2560,7 @@ void FSkeletalMeshObjectGPUSkin::FSkeletalMeshObjectLOD::InitResources(
 	ENQUEUE_RENDER_COMMAND(FSkeletalMeshObjectLOD_InitResources)(UE::RenderCommandPipe::SkeletalMesh,
 		[this, InGPUSkinCache, VertexBuffers, &LODData, InFeatureLevel, VertexAttributeMask](FRHICommandList& RHICmdList)
 	{
+		MeshObjectWeightBuffer_RenderThread = VertexBuffers.SkinWeightVertexBuffer;
 		GPUSkinVertexFactories.InitVertexFactories(RHICmdList, VertexBuffers, LODData.RenderSections, InFeatureLevel, VertexAttributeMask, GPUSkinTechnique);
 
 		if (DynamicBoundsStartOffset == INDEX_NONE && GPUSkinTechnique == ESkeletalMeshGPUSkinTechnique::GPUSkinCache)
diff --git a/Engine/Source/Runtime/Engine/Private/SkeletalRenderGPUSkin.h b/Engine/Source/Runtime/Engine/Private/SkeletalRenderGPUSkin.h
--- a/Engine/Source/Runtime/Engine/Private/SkeletalRenderGPUSkin.h
+++ b/Engine/Source/Runtime/Engine/Private/SkeletalRenderGPUSkin.h
@@ -629,6 +629,7 @@ private:
 
 		/** Skin weight buffer to use, could be from asset or component override */
 		const FSkinWeightVertexBuffer* MeshObjectWeightBuffer = nullptr;
+		const FSkinWeightVertexBuffer* MeshObjectWeightBuffer_RenderThread = nullptr;
 
 		/** Color buffer to user, could be from asset or component override */
 		FColorVertexBuffer* MeshObjectColorBuffer = nullptr;
