Extending TextureLODGroup quality configuration

diff --git a/Engine/Plugins/Mutable/Source/CustomizableObjectEditor/Private/MuCOE/GenerateMutableSource/GenerateMutableSource.cpp b/Engine/Plugins/Mutable/Source/CustomizableObjectEditor/Private/MuCOE/GenerateMutableSource/GenerateMutableSource.cpp
--- a/Engine/Plugins/Mutable/Source/CustomizableObjectEditor/Private/MuCOE/GenerateMutableSource/GenerateMutableSource.cpp
+++ b/Engine/Plugins/Mutable/Source/CustomizableObjectEditor/Private/MuCOE/GenerateMutableSource/GenerateMutableSource.cpp
@@ -1604,7 +1604,7 @@ uint32 ComputeLODBiasForTexture(const FMutableGraphGenerationContext& Generation
 
 		// Additional LODBias of the Texture
 		const int32 ReferenceTextureLODBias = LODSettings.CalculateLODBias(ReferenceTextureSize, ReferenceTextureSize, 0,	ReferenceTexture->LODGroup,
-			ReferenceTexture->LODBias, 0, ReferenceTexture->MipGenSettings, ReferenceTexture->IsCurrentlyVirtualTextured());
+			ReferenceTexture->LODBias, 0, ReferenceTexture->MipGenSettings, ReferenceTexture->IsCurrentlyVirtualTextured(), ReferenceTexture->IsStreamable());
 
 		return FMath::Max(NumMipsSource - NumMipsRefSource + PlatformLODBias + ReferenceTextureLODBias, 0);
 	}
@@ -1616,7 +1616,7 @@ uint32 ComputeLODBiasForTexture(const FMutableGraphGenerationContext& Generation
 	const int64 TextureSize = SourceSize >> PlatformLODBias;
 
 	// Additional LODBias of the Texture
-	const int32 TextureLODBias = LODSettings.CalculateLODBias(TextureSize, TextureSize, 0, Texture.LODGroup, Texture.LODBias, 0, Texture.MipGenSettings, Texture.IsCurrentlyVirtualTextured());
+	const int32 TextureLODBias = LODSettings.CalculateLODBias(TextureSize, TextureSize, 0, Texture.LODGroup, Texture.LODBias, 0, Texture.MipGenSettings, Texture.IsCurrentlyVirtualTextured(), Texture.IsStreamable());
 
 	return FMath::Max(PlatformLODBias + TextureLODBias, 0);
 }
@@ -1651,7 +1651,7 @@ int32 GetTextureSizeInGame(const UTexture2D& Texture, const UTextureLODSettings&
 	const int32 MaxTextureSizeAllowed = SourceSize >> PlatformLODBias;
 
 	// Calculate the LODBias specific for this texture 
-	const int32 TextureLODBias = LODSettings.CalculateLODBias(MaxTextureSizeAllowed, MaxTextureSizeAllowed, 0, Texture.LODGroup, Texture.LODBias, 0, Texture.MipGenSettings, Texture.IsCurrentlyVirtualTextured());
+	const int32 TextureLODBias = LODSettings.CalculateLODBias(MaxTextureSizeAllowed, MaxTextureSizeAllowed, 0, Texture.LODGroup, Texture.LODBias, 0, Texture.MipGenSettings, Texture.IsCurrentlyVirtualTextured(), Texture.IsStreamable());
 
 	return MaxTextureSizeAllowed >> TextureLODBias;
 }
diff --git a/Engine/Source/Editor/TextureEditor/Private/TextureEditorToolkit.cpp b/Engine/Source/Editor/TextureEditor/Private/TextureEditorToolkit.cpp
--- a/Engine/Source/Editor/TextureEditor/Private/TextureEditorToolkit.cpp
+++ b/Engine/Source/Editor/TextureEditor/Private/TextureEditorToolkit.cpp
@@ -750,7 +750,7 @@ void FTextureEditorToolkit::PopulateQuickInfo( )
 											const UTextureLODSettings* LODSettings = (UTextureLODSettings*)UDeviceProfileManager::Get().FindProfile(TargetPlatform->PlatformName());
 											int32 FirstMipToSerialize = LODSettings->CalculateLODBias(
 												TextureDescription.TopMipSizeX, TextureDescription.TopMipSizeY, Texture->MaxTextureSize, 
-												Texture->LODGroup, Texture->LODBias, 0, Texture->MipGenSettings, false);
+												Texture->LODGroup, Texture->LODBias, 0, Texture->MipGenSettings, false, Texture->IsStreamable());
 
 											int32 StrippedNumMips = TextureDescription.NumMips;
 											const int32 LastMip = FMath::Max(StrippedNumMips - 1, 0);
diff --git a/Engine/Source/Runtime/Engine/Classes/Engine/TextureLODSettings.h b/Engine/Source/Runtime/Engine/Classes/Engine/TextureLODSettings.h
--- a/Engine/Source/Runtime/Engine/Classes/Engine/TextureLODSettings.h
+++ b/Engine/Source/Runtime/Engine/Classes/Engine/TextureLODSettings.h
@@ -72,6 +72,22 @@ struct FTextureLODGroup
 	UPROPERTY()
 	int32 MaxLODSize_Smallest = -1;
 
+	/** Cap the number of mips such that the largest mip is this big. Has no effect for textures with no mip chain. Used in combination with r.TextureGroup.QualityLevel, r.TextureGroup.MaxQualityLevel and r.TextureGroup.OptionalQualityLevel. Not used for virtual textures. */
+	UPROPERTY()
+	int32 MaxLODSize_Low = -1;
+
+	/** Cap the number of mips such that the largest mip is this big. Has no effect for textures with no mip chain. Used in combination with r.TextureGroup.QualityLevel, r.TextureGroup.MaxQualityLevel and r.TextureGroup.OptionalQualityLevel. Not used for virtual textures. */
+	UPROPERTY()
+	int32 MaxLODSize_Medium = -1;
+
+	/** Cap the number of mips such that the largest mip is this big. Has no effect for textures with no mip chain. Used in combination with r.TextureGroup.QualityLevel, r.TextureGroup.MaxQualityLevel and r.TextureGroup.OptionalQualityLevel. Not used for virtual textures. */
+	UPROPERTY()
+	int32 MaxLODSize_High = -1;
+
+	/** Cap the number of mips such that the largest mip is this big. Has no effect for textures with no mip chain. Used in combination with r.TextureGroup.QualityLevel, r.TextureGroup.MaxQualityLevel and r.TextureGroup.OptionalQualityLevel. Not used for virtual textures. */
+	UPROPERTY()
+	int32 MaxLODSize_Epic = -1;
+
 	/** Cap the number of mips such that the largest mip is this big. Has no effect for textures with no mip chain. Used for virtual textures. */
 	UPROPERTY()
 	int32 MaxLODSize_VT = 0;
@@ -126,6 +142,8 @@ struct FTextureLODGroup
 	ENGINE_API void SetupGroup();
 
 	ENGINE_API bool operator==(const FTextureLODGroup& Other) const;
+
+	ENGINE_API void ComputeMaxSizesForQualityLevels(int32 OutMaxSizes[4]) const;
 };
 
 /**
@@ -159,9 +177,10 @@ public:
 	 * @param	NumCinematicMipLevels		The texture cinematic mip levels to include in the calculation
 	 * @param	MipGenSetting				Mip generation setting
 	 * @param	bVirtualTexture				If VT is enabled (in this case group's max LOD is ignored)
+	* @param	bCurrentlyStreamableTexture	If texture is currently streamable (if it's not, then quality levels don't apply to it)
 	 * @return	LOD bias
 	 */
-	ENGINE_API int32 CalculateLODBias( int32 Width, int32 Height, int32 MaxSize, int32 LODGroup, int32 LODBias, int32 NumCinematicMipLevels, TextureMipGenSettings MipGenSetting, bool bVirtualTexture ) const;
+	ENGINE_API int32 CalculateLODBias( int32 Width, int32 Height, int32 MaxSize, int32 LODGroup, int32 LODBias, int32 NumCinematicMipLevels, TextureMipGenSettings MipGenSetting, bool bVirtualTexture, bool bCurrentlyStreamableTexture = false) const;
 
 	/**
 	 * Calculate num optional mips
diff --git a/Engine/Source/Runtime/Engine/Private/Texture.cpp b/Engine/Source/Runtime/Engine/Private/Texture.cpp
--- a/Engine/Source/Runtime/Engine/Private/Texture.cpp
+++ b/Engine/Source/Runtime/Engine/Private/Texture.cpp
@@ -4100,7 +4100,7 @@ void UTexture::GetBuiltTextureSize(const ITargetPlatformSettings* TargetPlatform
 		const bool bVirtualTextureStreaming = VirtualTextureStreaming && IsVirtualTexturingEnabled(TargetPlatformSettings);
 
 		const UTextureLODSettings& LODSettings = TargetPlatformSettings->GetTextureLODSettings();
-		const uint32 LODBiasNoCinematics = FMath::Max<int32>(LODSettings.CalculateLODBias(SizeX, SizeY, MaxTextureSize, LODGroup, LODBias, 0, MipGenSettings, bVirtualTextureStreaming), 0);
+		const uint32 LODBiasNoCinematics = FMath::Max<int32>(LODSettings.CalculateLODBias(SizeX, SizeY, MaxTextureSize, LODGroup, LODBias, 0, MipGenSettings, bVirtualTextureStreaming, IsStreamable()), 0);
 		SizeX = FMath::Max<int32>(SizeX >> LODBiasNoCinematics, 1);
 		SizeY = FMath::Max<int32>(SizeY >> LODBiasNoCinematics, 1);
 		if ( bIsVolume )
diff --git a/Engine/Source/Runtime/Engine/Private/Texture2D.cpp b/Engine/Source/Runtime/Engine/Private/Texture2D.cpp
--- a/Engine/Source/Runtime/Engine/Private/Texture2D.cpp
+++ b/Engine/Source/Runtime/Engine/Private/Texture2D.cpp
@@ -1343,7 +1343,7 @@ int32 UTexture2D::Blueprint_GetSizeX() const
 		//	MaxTextureSize, PadForPow2, Downscale, etc. are not applied
 
 		const UTextureLODSettings* LODSettings = UDeviceProfileManager::Get().GetActiveProfile()->GetTextureLODSettings();
-		const int32 CookedLODBias = LODSettings->CalculateLODBias(Source.SizeX, Source.SizeY, MaxTextureSize, LODGroup, LODBias, 0, MipGenSettings, IsCurrentlyVirtualTextured());
+		const int32 CookedLODBias = LODSettings->CalculateLODBias(Source.SizeX, Source.SizeY, MaxTextureSize, LODGroup, LODBias, 0, MipGenSettings, IsCurrentlyVirtualTextured(), IsStreamable());
 		return FMath::Max<int32>(Source.SizeX >> CookedLODBias, 1);
 	}
 #endif
@@ -1406,7 +1406,7 @@ int32 UTexture2D::Blueprint_GetSizeY() const
 	if (!GetSizeY())
 	{
 		const UTextureLODSettings* LODSettings = UDeviceProfileManager::Get().GetActiveProfile()->GetTextureLODSettings();
-		const int32 CookedLODBias = LODSettings->CalculateLODBias(Source.SizeX, Source.SizeY, MaxTextureSize, LODGroup, LODBias, 0, MipGenSettings, IsCurrentlyVirtualTextured());
+		const int32 CookedLODBias = LODSettings->CalculateLODBias(Source.SizeX, Source.SizeY, MaxTextureSize, LODGroup, LODBias, 0, MipGenSettings, IsCurrentlyVirtualTextured(), IsStreamable());
 		return FMath::Max<int32>(Source.SizeY >> CookedLODBias, 1);
 	}
 #endif
diff --git a/Engine/Source/Runtime/Engine/Private/TextureDerivedData.cpp b/Engine/Source/Runtime/Engine/Private/TextureDerivedData.cpp
--- a/Engine/Source/Runtime/Engine/Private/TextureDerivedData.cpp
+++ b/Engine/Source/Runtime/Engine/Private/TextureDerivedData.cpp
@@ -841,7 +841,7 @@ static void ModifyMaxTextureResolutionBuildSettingsForPlatformLimit(
 		// after scaling to MaxTextureResolution, LODBiasNoCinematics will be applied
 			
 		const UTextureLODSettings& LODSettings = TargetPlatform->GetTextureLODSettings();
- 		const uint32 LODBiasNoCinematics = FMath::Max<int32>(LODSettings.CalculateLODBias(BuiltSizeX, BuiltSizeY, Texture.MaxTextureSize, Texture.LODGroup, Texture.LODBias, 0, Texture.MipGenSettings, OutSettings.bVirtualStreamable), 0);
+ 		const uint32 LODBiasNoCinematics = FMath::Max<int32>(LODSettings.CalculateLODBias(BuiltSizeX, BuiltSizeY, Texture.MaxTextureSize, Texture.LODGroup, Texture.LODBias, 0, Texture.MipGenSettings, OutSettings.bVirtualStreamable, Texture.IsStreamable()), 0);
 
 		int64 MaxTextureResolutionUp = ((int64)OutSettings.MaxTextureResolution)<<LODBiasNoCinematics;
 
@@ -1323,8 +1323,8 @@ static void GetTextureBuildSettings(
 		OutBuildSettings.CompositeTextureMode = CTM_Disabled;
 	}
 
-	OutBuildSettings.LODBias = TextureLODSettings.CalculateLODBias(SourceSize.X, SourceSize.Y, Texture.MaxTextureSize, Texture.LODGroup, Texture.LODBias, Texture.NumCinematicMipLevels, Texture.MipGenSettings, bVirtualTextureStreaming);
-	OutBuildSettings.LODBiasWithCinematicMips = TextureLODSettings.CalculateLODBias(SourceSize.X, SourceSize.Y, Texture.MaxTextureSize, Texture.LODGroup, Texture.LODBias, 0, Texture.MipGenSettings, bVirtualTextureStreaming);
+	OutBuildSettings.LODBias = TextureLODSettings.CalculateLODBias(SourceSize.X, SourceSize.Y, Texture.MaxTextureSize, Texture.LODGroup, Texture.LODBias, Texture.NumCinematicMipLevels, Texture.MipGenSettings, bVirtualTextureStreaming, Texture.IsStreamable());
+	OutBuildSettings.LODBiasWithCinematicMips = TextureLODSettings.CalculateLODBias(SourceSize.X, SourceSize.Y, Texture.MaxTextureSize, Texture.LODGroup, Texture.LODBias, 0, Texture.MipGenSettings, bVirtualTextureStreaming, Texture.IsStreamable());
 	OutBuildSettings.PowerOfTwoMode = Texture.PowerOfTwoMode;
 	OutBuildSettings.PaddingColor = Texture.PaddingColor;
 	OutBuildSettings.bPadWithBorderColor = Texture.bPadWithBorderColor;
@@ -2945,7 +2945,7 @@ static void SerializePlatformData(
 		const int32 FirstMipTailMip = NumMips - (int32)PlatformData->GetNumMipsInTail();
 		check(FirstMipTailMip >= 0);
 
-		FirstMipToSerialize = Ar.CookingTarget()->GetTextureLODSettings().CalculateLODBias(Width, Height, Texture->MaxTextureSize, LODGroup, LODBias, 0, MipGenSetting, bIsVirtual);
+		FirstMipToSerialize = Ar.CookingTarget()->GetTextureLODSettings().CalculateLODBias(Width, Height, Texture->MaxTextureSize, LODGroup, LODBias, 0, MipGenSetting, bIsVirtual, Texture->IsStreamable());
 		if (!bIsVirtual)
 		{
 			// Reassign NumMips as the number of mips starting from FirstMipToSerialize.
diff --git a/Engine/Source/Runtime/Engine/Private/TextureLODSettings.cpp b/Engine/Source/Runtime/Engine/Private/TextureLODSettings.cpp
--- a/Engine/Source/Runtime/Engine/Private/TextureLODSettings.cpp
+++ b/Engine/Source/Runtime/Engine/Private/TextureLODSettings.cpp
@@ -36,6 +36,29 @@ FAutoConsoleVariableRef CVarGlobalTextureLODBias(
 	ECVF_Scalability
 );
 
+int32 TextureGroupMaxQualityLevel = -1;
+FAutoConsoleVariableRef CVarTextureGroupMaxQualityLevel(
+	TEXT("r.TextureGroup.MaxQualityLevel"),
+	TextureGroupMaxQualityLevel,
+	TEXT("Used to apply texture quality bias to all textures and to limit TextureGroupQualityLevel. (default=-1)"),
+	ECVF_Default
+);
+
+int32 TextureGroupQualityLevel = -1;
+FAutoConsoleVariableRef CVarTextureGroupQualityLevel(
+	TEXT("r.TextureGroup.QualityLevel"),
+	TextureGroupQualityLevel,
+	TEXT("Used to apply texture quality bias to streamable textures at runtime. Clamped by r.TextureGroup.MaxQualityLevel. (default=-1)"),
+	ECVF_Default
+);
+
+int32 TextureGroupOptionalQualityLevel = -1;
+FAutoConsoleVariableRef CVarTextureGroupOptionalQualityLevel(
+	TEXT("r.TextureGroup.OptionalQualityLevel"),
+	TextureGroupOptionalQualityLevel,
+	TEXT("If set, it will use MaxLODSize for the given quality instead of OptionalMaxLODSize during cook time. (default=-1)"),
+	ECVF_Default
+);
 
 void FTextureLODGroup::SetupGroup()
 {
@@ -78,6 +101,13 @@ void FTextureLODGroup::SetupGroup()
 			default:
 				break;
 		}
+
+		if (TextureGroupMaxQualityLevel >= 0 && TextureGroupMaxQualityLevel < 4)
+		{
+			int32 MaxSizes[4];
+			ComputeMaxSizesForQualityLevels(MaxSizes);
+			MaxLODSize = MaxSizes[TextureGroupMaxQualityLevel];
+		}
 	}
 
 	MaxLODMipCount = FMath::CeilLogTwo(MaxLODSize);
@@ -123,6 +153,25 @@ bool FTextureLODGroup::operator==(const FTextureLODGroup& Other) const
 	return ScriptStruct->CompareScriptStruct(this, &Other, 0);
 }
 
+void FTextureLODGroup::ComputeMaxSizesForQualityLevels(int32 OutMaxSizes[4]) const
+{
+	// Copy or compute fallbacks for sizes
+	OutMaxSizes[0] = MaxLODSize_Low == -1 ? OptionalMaxLODSize : MaxLODSize_Low;
+	// In case none of the values are set
+	if (OutMaxSizes[0] == -1)
+	{
+		OutMaxSizes[0] = 128;
+	}
+	OutMaxSizes[1] = MaxLODSize_Medium == -1 ? OutMaxSizes[0] << 1 : MaxLODSize_Medium;
+	OutMaxSizes[2] = MaxLODSize_High == -1 ? OutMaxSizes[1] << 1 : MaxLODSize_High;
+	OutMaxSizes[3] = MaxLODSize_Epic == -1 ? MaxLODSize : FMath::Min(MaxLODSize_Epic, MaxLODSize);
+
+	// Ensure size[n] is not higher than size[n + 1]
+	OutMaxSizes[2] = FMath::Min(OutMaxSizes[2], OutMaxSizes[3]);
+	OutMaxSizes[1] = FMath::Min(OutMaxSizes[1], OutMaxSizes[2]);
+	OutMaxSizes[0] = FMath::Min(OutMaxSizes[0], OutMaxSizes[1]);
+}
+
 UTextureLODSettings::UTextureLODSettings(const FObjectInitializer& ObjectInitializer)
 	: Super(ObjectInitializer)
 {
@@ -170,10 +219,10 @@ int32 UTextureLODSettings::CalculateLODBias(const UTexture* Texture, bool bIncCi
 	float Width = Texture->GetSurfaceWidth();
 	float Height = Texture->GetSurfaceHeight();
 
-	return CalculateLODBias( FMath::RoundToInt(Width), FMath::RoundToInt(Height), TextureMaxSize, Texture->LODGroup, Texture->LODBias, bIncCinematicMips ? Texture->NumCinematicMipLevels : 0, MipGenSetting, Texture->IsCurrentlyVirtualTextured());
+	return CalculateLODBias( FMath::RoundToInt(Width), FMath::RoundToInt(Height), TextureMaxSize, Texture->LODGroup, Texture->LODBias, bIncCinematicMips ? Texture->NumCinematicMipLevels : 0, MipGenSetting, Texture->IsCurrentlyVirtualTextured(), Texture->IsStreamable());
 }
 
-int32 UTextureLODSettings::CalculateLODBias(int32 Width, int32 Height, int32 MaxSize, int32 LODGroup, int32 LODBias, int32 NumCinematicMipLevels, TextureMipGenSettings InMipGenSetting, bool bVirtualTexture ) const
+int32 UTextureLODSettings::CalculateLODBias(int32 Width, int32 Height, int32 MaxSize, int32 LODGroup, int32 LODBias, int32 NumCinematicMipLevels, TextureMipGenSettings InMipGenSetting, bool bVirtualTexture, bool bCurrentlyStreamableTexture) const
 {	
 	checkf(LODGroup < TextureLODGroups.Num(), TEXT("A texture had passed a bad LODGroup to UTextureLODSettings::CalculateLODBias (%d, out of %d groups). This code does not have the texture name. The LODSettings object is '%s'"), LODGroup, TextureLODGroups.Num(), *GetPathName());
 
@@ -234,8 +283,26 @@ int32 UTextureLODSettings::CalculateLODBias(int32 Width, int32 Height, int32 Max
 	//	MaxLODSize caps the largest mip size.
 	//	MinLODSize prevents LODBias from making the largest mip smaller than the given value. Does *not*
 	//	affect the size of the smallest mip. Almost never used.
+	int32 MaxLODSize = LODGroupInfo.MaxLODSize;
+	if (!GIsEditor)
+	{
+		// Non-streamable textures effectively bake in the bias once created and don't respond to changes, which means we can't apply bias that can change at runtime
+		if (bCurrentlyStreamableTexture && !bVirtualTexture && TextureGroupQualityLevel >= 0 && TextureGroupQualityLevel < 4)
+		{
+			int32 QualityLevel = TextureGroupQualityLevel;
+			if (TextureGroupMaxQualityLevel >= 0)
+			{
+				QualityLevel = FMath::Min(QualityLevel, TextureGroupMaxQualityLevel);
+			}
+
+			int32 MaxSizes[4];
+			LODGroupInfo.ComputeMaxSizesForQualityLevels(MaxSizes);
+			MaxLODSize = MaxSizes[QualityLevel];
+		}
+	}
+
 	int32 MinLOD = FMath::CeilLogTwo(LODGroupInfo.MinLODSize);
-	int32 MaxLOD = FMath::CeilLogTwo(LODGroupInfo.MaxLODSize);
+	int32 MaxLOD = FMath::CeilLogTwo(MaxLODSize);
 
 	if (bVirtualTexture)
 	{
@@ -264,7 +331,19 @@ int32 UTextureLODSettings::CalculateNumOptionalMips(int32 LODGroup, const int32
 		return 0;
 	}
 
-	int32 OptionalLOD = FMath::Min<int32>(FMath::CeilLogTwo(LODGroupInfo.OptionalMaxLODSize) + 1, NumMips);
+	int32 OptionalMaxLODSize;
+	if (TextureGroupOptionalQualityLevel >= 0 && TextureGroupOptionalQualityLevel < 4)
+	{
+		int32 MaxSizes[4];
+		LODGroupInfo.ComputeMaxSizesForQualityLevels(MaxSizes);
+		OptionalMaxLODSize = MaxSizes[TextureGroupOptionalQualityLevel];
+	}
+	else
+	{
+		OptionalMaxLODSize = LODGroupInfo.OptionalMaxLODSize;
+	}
+
+	int32 OptionalLOD = FMath::Min<int32>(FMath::CeilLogTwo(OptionalMaxLODSize) + 1, NumMips);
 
 	// "MinMipToInline" actually comes from NumNonStreaming
 	//  this ensures that optional mips are always streaming mips
