[SimpleStreamableAssetManager] Fixed crash in TLocklessGrowingStorage where a shard could be used after being freed (now internally uses TConsumeAllMpmcQueue).

diff --git a/Engine/Source/Runtime/Engine/Private/Streaming/SimpleStreamableAssetManager.cpp b/Engine/Source/Runtime/Engine/Private/Streaming/SimpleStreamableAssetManager.cpp
--- a/Engine/Source/Runtime/Engine/Private/Streaming/SimpleStreamableAssetManager.cpp
+++ b/Engine/Source/Runtime/Engine/Private/Streaming/SimpleStreamableAssetManager.cpp
@@ -230,64 +230,56 @@ void FSimpleStreamableAssetManager::UpdateTask_Async()
 		
 	const int32 LastRegisteredObjectCount = RegisteredObjectCount;
 	const int32 LastRegisteredAssetsCount = UsedAssetIndices;
-	
-
-	TArray<TLocklessGrowingStorage<FUpdate>::FStorageShard*> Pending_UpdateRecords{};
-	TArray<TLocklessGrowingStorage<FUpdateLastRenderTime>::FStorageShard*> Pending_UpdateLastRenderTimeRecords{};
-	TArray<TLocklessGrowingStorage<FRegister>::FStorageShard*> Pending_RegisterRecords{};
-	TArray<TLocklessGrowingStorage<FUnregister>::FStorageShard*> Pending_UnregisterRecords{};
-	TArray<TLocklessGrowingStorage<FRemovedAssetRecord>::FStorageShard*> Pending_RemovedAssetRecords{};
 
+	check(Pending_UpdateRecords.IsEmpty());
+	check(Pending_UpdateLastRenderTimeRecords.IsEmpty());
+	check(Pending_RegisterRecords.IsEmpty());
+	check(Pending_UnregisterRecords.IsEmpty());
+	check(Pending_RemovedAssetRecords.IsEmpty());
+	
 	{
-		/** It is important to extract Shards in this order
-		 * Since the process in not blocking we want to ensure we have
-		 * request first  add for updated primitives
-		 * becouse those two can come from different/competeing threads 
-		 */
+		// Here, the order of extraction of the records is important.
+		// This consumer is non-blocking; producers may push from other threads.
 		SCOPED_NAMED_EVENT(FSimpleStreamableAssetManager_MoveData, FColor::Silver);
-		Pending_UpdateRecords = UpdateRecords.ExtractShards();
-		Pending_UpdateLastRenderTimeRecords = UpdateLastRenderTimeRecords.ExtractShards();
-		Pending_RegisterRecords = RegisterRecords.ExtractShards();
-		Pending_UnregisterRecords = UnregisterRecords.ExtractShards();
-		Pending_RemovedAssetRecords = RemovedAssetsRecords.ExtractShards();
+		UpdateRecords.ExtractAll(Pending_UpdateRecords);
+		UpdateLastRenderTimeRecords.ExtractAll(Pending_UpdateLastRenderTimeRecords);
+		RegisterRecords.ExtractAll(Pending_RegisterRecords);
+		UnregisterRecords.ExtractAll(Pending_UnregisterRecords);
+		RemovedAssetsRecords.ExtractAll(Pending_RemovedAssetRecords);
 	}
 	
 	/** Assign a indice to each proxy **/
 	{
 		SCOPED_NAMED_EVENT(FSimpleStreamableAssetManager_AssignObjectIndex, FColor::Silver);
 		
-		for (TLocklessGrowingStorage<FRegister>::FStorageShard* Shard : Pending_RegisterRecords)
+		const int32 NeedToReserve = RegisteredObjectCount + Pending_RegisterRecords.Num() - MaxObjects;
+		if (NeedToReserve > 0)
 		{
-			TArrayView<FRegister> View = Shard->GetData();
-			const int32 NeedToReserve = RegisteredObjectCount + View.Num() - MaxObjects;
-			if (NeedToReserve > 0)
+			SCOPED_NAMED_EVENT(FSimpleStreamableAssetManager_Resize, FColor::Silver);
+			ObjectUsedIndices.Add(false, NeedToReserve);
+			ObjectRegistrationIndexToAssetProperty.AddZeroed(NeedToReserve);
+			MaxObjects += NeedToReserve;
+
+			// We store 4 Object bounds in one entry
+			const int32 MaxObject4BoundsNeeded = MaxObjects / 4 + int32(MaxObjects % 4 != 0);
+			const int32 Object4BoundsNeedToReserve = MaxObject4BoundsNeeded - ObjectBounds4.Num();
+			if (Object4BoundsNeedToReserve > 0)
 			{
-				SCOPED_NAMED_EVENT(FSimpleStreamableAssetManager_Resize, FColor::Silver);
-				ObjectUsedIndices.Add(false, NeedToReserve);
-				ObjectRegistrationIndexToAssetProperty.AddZeroed(NeedToReserve);
-				MaxObjects += NeedToReserve;
-
-				// We store 4 Object bounds in one entry
-				const int32 MaxObject4BoundsNeeded = MaxObjects/4 + int32(MaxObjects % 4 != 0);
-				const int32 Object4BoundsNeedToReserve = MaxObject4BoundsNeeded - ObjectBounds4.Num();
-				if (Object4BoundsNeedToReserve > 0)
-				{
-					ObjectBounds4.AddDefaulted(Object4BoundsNeedToReserve);
-				}
+				ObjectBounds4.AddDefaulted(Object4BoundsNeedToReserve);
 			}
+		}
 
-			// Cheap index assignment for proxy
-			for (FRegister& Entry : View)
+		for (const FRegister& Record : Pending_RegisterRecords)
+		{
+			// Lets Skip reserving indices if no streamable assets
+			if (Record.Assets.Num() > 0)
 			{
-				// Lets Skip reserving indices if no streamable assets
-				if (Entry.Assets.Num() > 0)
-				{
-					const int32 ObjectIndex = ObjectUsedIndices.FindAndSetFirstZeroBit(FreeObjectIndexHint);
-					check(ObjectIndex != INDEX_NONE);
-					FreeObjectIndexHint = ObjectIndex + 1;
-					++RegisteredObjectCount;
-					*Entry.ObjectRegistrationIndex = ObjectIndex;
-				}
+				// Cheap index assignment for proxy
+				const int32 ObjectIndex = ObjectUsedIndices.FindAndSetFirstZeroBit(FreeObjectIndexHint);
+				check(ObjectIndex != INDEX_NONE);
+				FreeObjectIndexHint = ObjectIndex + 1;
+				++RegisteredObjectCount;
+				*Record.ObjectRegistrationIndex = ObjectIndex;
 			}
 		}
 	}
@@ -297,68 +289,48 @@ void FSimpleStreamableAssetManager::UpdateTask_Async()
 
 		{
 			SCOPED_NAMED_EVENT(FSimpleStreamableAssetManager_Unregister, FColor::Silver);
-			for (TLocklessGrowingStorage<FUnregister>::FStorageShard* Shard : Pending_UnregisterRecords)
+			for (FUnregister& Record : Pending_UnregisterRecords)
 			{
-				TArrayView<FUnregister> View = Shard->GetData();
-				for (FUnregister& Record : View)
-				{
-					UnregisterRecord(Record);
-				}
+				UnregisterRecord(Record);
 			}
 		}
 		
 		{
 			SCOPED_NAMED_EVENT(FSimpleStreamableAssetManager_Register, FColor::Silver);
-			for (TLocklessGrowingStorage<FRegister>::FStorageShard* Shard : Pending_RegisterRecords)
+			for (FRegister& Record : Pending_RegisterRecords)
 			{
-				TArrayView<FRegister> View = Shard->GetData();
-				for (FRegister& Entry : View)
-				{
-					RegisterRecord(Entry);
-				}
+				RegisterRecord(Record);
 			}
 		}
 
 		{
 			SCOPED_NAMED_EVENT(FSimpleStreamableAssetManager_Update, FColor::Silver);
-			for (TLocklessGrowingStorage<FUpdate>::FStorageShard* Shard : Pending_UpdateRecords)
+			for (FUpdate& Record : Pending_UpdateRecords)
 			{
-				TArrayView<FUpdate> View = Shard->GetData();
-				for (FUpdate& Entry : View)
-				{
-					UpdateRecord(Entry);
-				}
+				UpdateRecord(Record);
 			}
 		}
 
 		{
 			SCOPED_NAMED_EVENT(FSimpleStreamableAssetManager_UpdateLastRenderTime, FColor::Silver);
-			for (TLocklessGrowingStorage<FUpdateLastRenderTime>::FStorageShard* Shard : Pending_UpdateLastRenderTimeRecords)
+			for (FUpdateLastRenderTime& Record : Pending_UpdateLastRenderTimeRecords)
 			{
-				TArrayView<FUpdateLastRenderTime> View = Shard->GetData();
-				for (FUpdateLastRenderTime& Entry : View)
-				{
-					UpdateRecord(Entry);
-				}
+				UpdateRecord(Record);
 			}
 		}
 
 		{
 			SCOPED_NAMED_EVENT(FSimpleStreamableAssetManager_RemoveAssets, FColor::Silver);
-			for (TLocklessGrowingStorage<FRemovedAssetRecord>::FStorageShard* Shard : Pending_RemovedAssetRecords)
+			for (FRemovedAssetRecord& Record : Pending_RemovedAssetRecords)
 			{
-				TArrayView<FRemovedAssetRecord> View = Shard->GetData();
-				for (const FRemovedAssetRecord& Record : View)
+				const int32 AssetIndex = *Record.SimpleStreamableAssetManagerIndex;
+				if (AssetIndex != INDEX_NONE)
 				{
-					const int32 AssetIndex = *Record.SimpleStreamableAssetManagerIndex;
-					if (AssetIndex != INDEX_NONE)
-					{
-						AssetUsedIndices[AssetIndex] = false;
-						--UsedAssetIndices;
-						FreeAssetIndexHint = FMath::Min(FreeAssetIndexHint, AssetIndex);
-						AssetIndexToBounds4Index[AssetIndex].Empty();
-						*Record.SimpleStreamableAssetManagerIndex = INDEX_NONE;
-					}
+					AssetUsedIndices[AssetIndex] = false;
+					--UsedAssetIndices;
+					FreeAssetIndexHint = FMath::Min(FreeAssetIndexHint, AssetIndex);
+					AssetIndexToBounds4Index[AssetIndex].Empty();
+					*Record.SimpleStreamableAssetManagerIndex = INDEX_NONE;
 				}
 			}
 		}
@@ -373,54 +345,23 @@ void FSimpleStreamableAssetManager::UpdateTask_Async()
 	{
 		TRACE_COUNTER_SET(RegisteredAssets, UsedAssetIndices);
 	}
-	
-	{
-#if COUNTERSTRACE_ENABLED
-		int32 Updates = 0;	
-#endif
-		for (TLocklessGrowingStorage<FUpdate>::FStorageShard* Shard : Pending_UpdateRecords)
-		{
-#if COUNTERSTRACE_ENABLED
-			Updates += Shard->GetData().Num();
-#endif
-			delete Shard;
-		}
-		TRACE_COUNTER_SET(UpdatedObjects, Updates);
 
 #if COUNTERSTRACE_ENABLED
-		int32 Added = 0;
-#endif
-		for (TLocklessGrowingStorage<FRegister>::FStorageShard* Shard : Pending_RegisterRecords)
-		{
-#if COUNTERSTRACE_ENABLED
-			Added += Shard->GetData().Num();
-#endif
-			delete Shard;
-		}
-		TRACE_COUNTER_SET(AddedObjects, Added);
+	const int32 Removed = Pending_UnregisterRecords.Num();
+	TRACE_COUNTER_SET(RemovedObjects, Removed);
 
-#if COUNTERSTRACE_ENABLED
-		int32 Removed = 0;
-#endif
-		for (TLocklessGrowingStorage<FUnregister>::FStorageShard* Shard : Pending_UnregisterRecords)
-		{
-#if COUNTERSTRACE_ENABLED
-			Removed += Shard->GetData().Num();
-#endif
-			delete Shard;
-		}
-		TRACE_COUNTER_SET(RemovedObjects, Removed);
+	const int32 Added = Pending_RegisterRecords.Num();
+	TRACE_COUNTER_SET(AddedObjects, Added);
 
-		for (TLocklessGrowingStorage<FRemovedAssetRecord>::FStorageShard* Shard : Pending_RemovedAssetRecords)
-		{
-			delete Shard;
-		}
+	const int32 Updates = Pending_UpdateRecords.Num();
+	TRACE_COUNTER_SET(UpdatedObjects, Updates);
+#endif
 
-		for (TLocklessGrowingStorage<FUpdateLastRenderTime>::FStorageShard* Shard : Pending_UpdateLastRenderTimeRecords)
-		{
-			delete Shard;
-		}
-	}
+	Pending_UpdateRecords.Reset();
+	Pending_UpdateLastRenderTimeRecords.Reset();
+	Pending_RegisterRecords.Reset();
+	Pending_UnregisterRecords.Reset();
+	Pending_RemovedAssetRecords.Reset();
 }
 
 void FSimpleStreamableAssetManager::RegisterRecord(FRegister& Record)
diff --git a/Engine/Source/Runtime/Engine/Public/Streaming/SimpleStreamableAssetManager.h b/Engine/Source/Runtime/Engine/Public/Streaming/SimpleStreamableAssetManager.h
--- a/Engine/Source/Runtime/Engine/Public/Streaming/SimpleStreamableAssetManager.h
+++ b/Engine/Source/Runtime/Engine/Public/Streaming/SimpleStreamableAssetManager.h
@@ -5,6 +5,7 @@
 #include "Components/PrimitiveComponent.h"
 #include "Misc/SpinLock.h"
 #include "Engine/StreamableRenderAsset.h"
+#include "Containers/ConsumeAllMpmcQueue.h"
 
 class FPrimitiveSceneProxy;
 class IPrimitiveComponent;
@@ -18,111 +19,26 @@ struct FBounds4;
 
 #define SIMPLE_STREAMABLE_ASSET_MANAGER_ALLOW_DEBUG_POINTER !(UE_BUILD_SHIPPING)
 
-template<typename T, uint32 N>
-struct TLocklessStaticStorage
-{
-public:
-	constexpr static uint32 MaxElements = N;
-	uint32 Reserve() { return ReservedIndex.fetch_add(1, std::memory_order_acq_rel); }
-		
-	void Store(T&& In, uint32 ItemReservedIndex)
-	{
-		// can assert if Relased is equal ReservedIndex and lower than N
-		Storage[ItemReservedIndex] = MoveTemp(In);
-	}
-
-	void Release(uint32 ToRelease)
-	{
-		uint32 Expected = ToRelease;
-		while (!ReleasedIndex.compare_exchange_weak(Expected, ToRelease + 1, std::memory_order_acq_rel)) 
-		{ 
-			Expected = ToRelease; 
-		}
-	}
-
-	void WaitForWrites() const
-	{
-		while(ReleasedIndex.load(std::memory_order_acquire) < std::min(MaxElements, ReservedIndex.load(std::memory_order_acquire))){}
-	}
-
-	TArrayView<T> GetData() { return TArrayView<T>(Storage, ReleasedIndex.load()); };
-	
-private:
-	std::atomic_uint32_t ReservedIndex = 0;
-	std::atomic_uint32_t ReleasedIndex = 0;
-	T Storage[MaxElements]{};
-};
-
 template<typename T>
 class TLocklessGrowingStorage
 {
 public:
-	using FStorageShard = TLocklessStaticStorage<T, 512>;
-
+	
 	void Push(T&& In)
 	{
-		while (true)
-		{
-			FStorageShard* Storage = Shard.load(std::memory_order_acquire);
-			if (Storage)
-			{
-				const uint32 StorageIndex = Storage->Reserve();
-				if (StorageIndex < FStorageShard::MaxElements)
-				{
-					Storage->Store(MoveTemp(In), StorageIndex);
-					Storage->Release(StorageIndex);
-					return;
-				}
-			}
-			uint32 Expected = 0;
-			if (StorageAllocatorGuard.compare_exchange_strong(Expected, 1, std::memory_order_acq_rel))
-			{
-				FStorageShard* NewStorage = new FStorageShard();
-				if (Shard.compare_exchange_strong(Storage, NewStorage, std::memory_order_acq_rel))
-				{
-					if (Storage)
-					{
-						LockedStorage.Add(Storage);
-					}
-				}
-				else
-				{
-					delete NewStorage;
-				}
-				StorageAllocatorGuard.store(0, std::memory_order_release);
-			}
-		}
-	}
-	TArray<FStorageShard*> ExtractShards()
-	{
-		TArray<FStorageShard*> Shards;
-		uint32 Expected = 0;
-		while (!StorageAllocatorGuard.compare_exchange_strong(Expected, 1, std::memory_order_acq_rel)){ Expected =  0; }
-		Shards = MoveTemp(LockedStorage);
-		FStorageShard* GrabbedStorage = Shard.exchange(nullptr);
-		StorageAllocatorGuard.store(0, std::memory_order_release);
-		if (GrabbedStorage)
-		{
-			GrabbedStorage->WaitForWrites();
-			Shards.Add(GrabbedStorage);
-		}
-			
-		return Shards;
+		Queue.ProduceItem(MoveTemp(In));
 	}
 
-	~TLocklessGrowingStorage()
+	void ExtractAll(TArray<T>& Out)
 	{
-		TArray<FStorageShard*> Shards = ExtractShards();
-		for (FStorageShard* ShardElement : Shards)
+		Queue.ConsumeAllFifo([&Out](T Item)
 		{
-			delete ShardElement;
-		}
+			Out.Add(MoveTemp(Item));
+		});
 	}
+
 private:
-		
-	std::atomic_uint32_t StorageAllocatorGuard = 0;
-	std::atomic<FStorageShard*> Shard = nullptr;
-	TArray<FStorageShard*> LockedStorage;
+	UE::TConsumeAllMpmcQueue<T> Queue;
 };
 
 class FSimpleStreamableAssetManager
@@ -374,6 +290,12 @@ private:
 	TLocklessGrowingStorage<FUpdate> UpdateRecords;
 	TLocklessGrowingStorage<FUpdateLastRenderTime> UpdateLastRenderTimeRecords;
 
+	TArray<FUpdate> Pending_UpdateRecords;
+	TArray<FUpdateLastRenderTime> Pending_UpdateLastRenderTimeRecords;
+	TArray<FRegister> Pending_RegisterRecords;
+	TArray<FUnregister> Pending_UnregisterRecords;
+	TArray<FRemovedAssetRecord> Pending_RemovedAssetRecords;
+
 	// ** Variables to manage Object registration ** //
 	int32 RegisteredObjectCount = 0;
 	int32 MaxObjects = 0;
