[NetPhysics] Fix replicated simulated child component getting the wrong world position on the client during detachment from parent component.

diff --git a/Engine/Source/Runtime/Engine/Private/Components/SceneComponent.cpp b/Engine/Source/Runtime/Engine/Private/Components/SceneComponent.cpp
--- a/Engine/Source/Runtime/Engine/Private/Components/SceneComponent.cpp
+++ b/Engine/Source/Runtime/Engine/Private/Components/SceneComponent.cpp
@@ -81,6 +81,20 @@ namespace SceneComponentCVars
 	static FAutoConsoleVariableRef AllowRemoteAttachedChildren(TEXT("DSTM.AllowRemoteAttachedChildren"), bAllowRemoteAttachedChildren, TEXT("Do we allow AttachChildren of a SceneComponent to exist only Remotely?"));
 #endif
 
+	bool bStopReplicatedClientFromAlteringAttachmentLocally = true;
+	static FAutoConsoleVariableRef CVarStopReplicatedClientFromAlteringAttachmentLocally(
+		TEXT("s.StopReplicatedClientFromAlteringAttachmentLocally"),
+		bStopReplicatedClientFromAlteringAttachmentLocally,
+		TEXT("True: Don't let the client locally change the attachment status of this component which could make it deviate from the servers attachment. Attachment is replicated and the client will be updated with the servers attachment state."),
+		ECVF_Default);
+
+	bool bStopReplicatedClientFromApplyingTransformDuringDetachmentLocally = true;
+	static FAutoConsoleVariableRef CVarStopReplicatedClientFromApplyingTransformDuringDetachmentLocally(
+		TEXT("s.StopReplicatedClientFromApplyingTransformDuringDetachmentLocally"),
+		bStopReplicatedClientFromApplyingTransformDuringDetachmentLocally,
+		TEXT("True: Don't let the client locally apply its current relative transform as a world transform, if both client and server do this the clients world transform will be doubled until the next replicated transform arrive from the server."),
+		ECVF_Default);
+
 #if !(UE_BUILD_SHIPPING || UE_BUILD_TEST)
 	bool bWarnSceneComponentStaticMobilityMove = true;
 	static FAutoConsoleVariableRef CVarWarnSceneComponentStaticMobilityMove(
@@ -816,6 +830,12 @@ void USceneComponent::UpdateComponentToWorldWithParent(USceneComponent* Parent,F
 
 void USceneComponent::AttachToParent()
 {
+	if (SceneComponentCVars::bStopReplicatedClientFromAlteringAttachmentLocally && IsNetSimulating())
+	{
+		// Attachment is replicated from the server, early out if we are networked and not server
+		return;
+	}
+
 	// If we need to perform a call to AttachTo, do that now
 	// At this point scene component still has no any state (rendering, physics),
 	// so this call will just add this component to an AttachChildren array of a the Parent component
@@ -2226,6 +2246,12 @@ void USceneComponent::SetRelativeRotationCache(const FRotationConversionCache& I
 
 void USceneComponent::SetupAttachment(class USceneComponent* InParent, FName InSocketName)
 {
+	if (SceneComponentCVars::bStopReplicatedClientFromAlteringAttachmentLocally && IsNetSimulating())
+	{
+		// Attachment is replicated from the server, early out if we are networked and not server
+		return;
+	}
+
 	if (InParent != AttachParent || InSocketName != AttachSocketName)
 	{
 		if (ensureMsgf(!bRegistered, TEXT("SetupAttachment should only be used to initialize AttachParent and AttachSocketName for a future AttachToComponent. Once a component is registered you must use AttachToComponent. Component [%s] Owner [%s], InParent [%s], InSocketName [%s]"), *GetFullName(), *GetPathNameSafe(GetOwner()), *GetNameSafe(InParent), *InSocketName.ToString()))
@@ -2699,32 +2725,36 @@ void USceneComponent::DetachFromComponent(const FDetachmentTransformRules& Detac
 
 		OnAttachmentChanged();
 
-		// If desired, update RelativeLocation and RelativeRotation to maintain current world position after detachment
-		switch (DetachmentRules.LocationRule)
+		// When component is replicated only apply location, rotation and scale on the server
+		if (!SceneComponentCVars::bStopReplicatedClientFromApplyingTransformDuringDetachmentLocally || !IsNetSimulating())
 		{
-		case EDetachmentRule::KeepRelative:
-			break;
-		case EDetachmentRule::KeepWorld:
-			SetRelativeLocation_Direct(GetComponentTransform().GetTranslation()); // or GetComponentLocation, but worried about custom location...
-			break;
-		}
+			// If desired, update RelativeLocation and RelativeRotation to maintain current world position after detachment
+			switch (DetachmentRules.LocationRule)
+			{
+			case EDetachmentRule::KeepRelative:
+				break;
+			case EDetachmentRule::KeepWorld:
+				SetRelativeLocation_Direct(GetComponentTransform().GetTranslation()); // or GetComponentLocation, but worried about custom location...
+				break;
+			}
 
-		switch (DetachmentRules.RotationRule)
-		{
-		case EDetachmentRule::KeepRelative:
-			break;
-		case EDetachmentRule::KeepWorld:
-			SetRelativeRotation_Direct(GetComponentRotation());
-			break;
-		}
+			switch (DetachmentRules.RotationRule)
+			{
+			case EDetachmentRule::KeepRelative:
+				break;
+			case EDetachmentRule::KeepWorld:
+				SetRelativeRotation_Direct(GetComponentRotation());
+				break;
+			}
 
-		switch (DetachmentRules.ScaleRule)
-		{
-		case EDetachmentRule::KeepRelative:
-			break;
-		case EDetachmentRule::KeepWorld:
-			SetRelativeScale3D_Direct(GetComponentScale());
-			break;
+			switch (DetachmentRules.ScaleRule)
+			{
+			case EDetachmentRule::KeepRelative:
+				break;
+			case EDetachmentRule::KeepWorld:
+				SetRelativeScale3D_Direct(GetComponentScale());
+				break;
+			}
 		}
 
 		// calculate transform with new attachment condition
