[World] Old DataLayer deprecation continuation part 1 - Move to Data Layer assets in Rule Processor and City Sample

#rnx

diff --git a/Engine/Source/Runtime/Engine/Public/WorldPartition/DataLayer/WorldDataLayers.h b/Engine/Source/Runtime/Engine/Public/WorldPartition/DataLayer/WorldDataLayers.h
--- a/Engine/Source/Runtime/Engine/Public/WorldPartition/DataLayer/WorldDataLayers.h
+++ b/Engine/Source/Runtime/Engine/Public/WorldPartition/DataLayer/WorldDataLayers.h
@@ -136,7 +136,8 @@ public:
 
 	// Allows overriding of DataLayers with PlayFromHere
 	template<class T>
-	void OverwriteDataLayerRuntimeStates(const TArray<T>* InActiveDataLayers, const TArray<T>* InLoadedDataLayers );
+	UE_DEPRECATED(5.8, "Use SetDataLayerRuntimeState instead.")
+	void OverwriteDataLayerRuntimeStates(const TArray<T>* InActiveDataLayers, const TArray<T>* InLoadedDataLayers);
 
 	ENGINE_API bool IsReadOnly(FText* OutReason = nullptr) const;
 	ENGINE_API bool IsSubWorldDataLayers() const;
@@ -442,12 +443,26 @@ DataLayerInstanceType* AWorldDataLayers::CreateDataLayer(CreationsArgs... InCrea
 	return NewDataLayer;
 }
 
-//todo_ow: Remove this
 template<class IdentifierType>
 void AWorldDataLayers::OverwriteDataLayerRuntimeStates(const TArray<IdentifierType>* InActiveDataLayers, const TArray<IdentifierType>* InLoadedDataLayers)
 {
 	if (GetLocalRole() == ROLE_Authority)
 	{
+		// Handle Wrapper types
+		auto IdentifierToInstance = [this](const IdentifierType& Identifier) -> const UDataLayerInstance*
+		{
+			if constexpr (std::is_same_v<IdentifierType, TSoftObjectPtr<UDataLayerAsset>>)
+			{
+				return GetDataLayerInstance(Identifier.Get());
+			}
+			else
+			{
+			PRAGMA_DISABLE_DEPRECATION_WARNINGS
+				return GetDataLayerInstance(Identifier);
+			PRAGMA_ENABLE_DEPRECATION_WARNINGS
+			}
+		};
+
 		FlushNetDormancy();
 
 		// This should get called before game starts. It doesn't send out events
@@ -458,9 +473,7 @@ void AWorldDataLayers::OverwriteDataLayerRuntimeStates(const TArray<IdentifierTy
 			ActiveDataLayerNames.Empty(InActiveDataLayers->Num());
 			for (const IdentifierType& DataLayerIdentitier : *InActiveDataLayers)
 			{
-				PRAGMA_DISABLE_DEPRECATION_WARNINGS
-				const UDataLayerInstance* DataLayerInstance = GetDataLayerInstance(DataLayerIdentitier);
-				PRAGMA_ENABLE_DEPRECATION_WARNINGS
+				const UDataLayerInstance* DataLayerInstance = IdentifierToInstance(DataLayerIdentitier);
 				if (DataLayerInstance && DataLayerInstance->IsRuntime())
 				{
 					ActiveDataLayerNames.Add(DataLayerInstance->GetDataLayerFName());
@@ -476,9 +489,7 @@ void AWorldDataLayers::OverwriteDataLayerRuntimeStates(const TArray<IdentifierTy
 			LoadedDataLayerNames.Empty(InLoadedDataLayers->Num());
 			for (const IdentifierType& DataLayerIdentitier : *InLoadedDataLayers)
 			{
-				PRAGMA_DISABLE_DEPRECATION_WARNINGS
-				const UDataLayerInstance* DataLayerInstance = GetDataLayerInstance(DataLayerIdentitier);
-				PRAGMA_ENABLE_DEPRECATION_WARNINGS
+				const UDataLayerInstance* DataLayerInstance = IdentifierToInstance(DataLayerIdentitier);
 				if (DataLayerInstance && DataLayerInstance->IsRuntime())
 				{
 					LoadedDataLayerNames.Add(DataLayerInstance->GetDataLayerFName());
