[Texture Streaming] Fixed FRenderAssetInstanceState::PreAddComponent potentially not computing the correct list of RenderAssets for a primitive bounds (when r.Streaming.AllowParallelRenderAssetStreamingManagerIncrementalUpdate is enabled).

diff --git a/Engine/Source/Runtime/Engine/Private/Streaming/TextureInstanceState.cpp b/Engine/Source/Runtime/Engine/Private/Streaming/TextureInstanceState.cpp
--- a/Engine/Source/Runtime/Engine/Private/Streaming/TextureInstanceState.cpp
+++ b/Engine/Source/Runtime/Engine/Private/Streaming/TextureInstanceState.cpp
@@ -540,7 +540,7 @@ EAddComponentResult FRenderAssetInstanceState::AddComponentInternal(const UPrimi
 			{
 				check(OutPreAddPayload);
 				FPreAddBounds& OutBounds = OutPreAddPayload->Bounds.Emplace_GetRef(Info.Bounds, PackedRelativeBox_Identity, false, Component->GetBounds().Origin, MinDistanceSq, MinRangeSq, MaxRangeSq);
-				UE::Private::PreAddRenderAssetElements(RenderAssetInstanceInfos, OutBounds.RenderAssets);
+				UE::Private::PreAddRenderAssetElements(TArrayView<FStreamingRenderAssetPrimitiveInfo>(RenderAssetInstanceInfos.GetData() + InfoIndex, NumOfBoundReferences), OutBounds.RenderAssets);
 			}
 
 			InfoIndex += NumOfBoundReferences;
