[Audio] - introduce API to let subsystems intercept ParseParams each update

#rnx

diff --git a/Engine/Source/Runtime/Engine/Private/ActiveSound.cpp b/Engine/Source/Runtime/Engine/Private/ActiveSound.cpp
--- a/Engine/Source/Runtime/Engine/Private/ActiveSound.cpp
+++ b/Engine/Source/Runtime/Engine/Private/ActiveSound.cpp
@@ -1103,6 +1103,9 @@ void FActiveSound::UpdateWaveInstances(TArray<FWaveInstance*> &InWaveInstances,
 			}
 		}
 
+		// Allow subsystems one last opportunity to adjust ParseParams
+		AudioDevice->SubsystemsAdjustParseParams(*this, ParseParams);
+
 		Sound->Parse(AudioDevice, 0, *this, ParseParams, ThisSoundsWaveInstances);
 
 		// Track this active sound's min pitch value. This is used to scale it's possible duration value.
diff --git a/Engine/Source/Runtime/Engine/Private/AudioDevice.cpp b/Engine/Source/Runtime/Engine/Private/AudioDevice.cpp
--- a/Engine/Source/Runtime/Engine/Private/AudioDevice.cpp
+++ b/Engine/Source/Runtime/Engine/Private/AudioDevice.cpp
@@ -3661,6 +3661,15 @@ void FAudioDevice::ApplyInteriorSettings(FActiveSound& ActiveSound, FSoundParseP
 	});
 }
 
+void FAudioDevice::SubsystemsAdjustParseParams(FActiveSound& ActiveSound, FSoundParseParameters& ParseParams) const
+{
+	SubsystemCollection.ForEachSubsystem<IActiveSoundUpdateInterface>([&ActiveSound, &ParseParams](IActiveSoundUpdateInterface* ActiveSoundUpdate)
+	{
+		ActiveSoundUpdate->SubsystemsAdjustParseParams(ActiveSound, ParseParams);
+		return true;
+	});
+}
+
 void FAudioDevice::NotifySubsystemsActiveSoundCreated(FActiveSound& ActiveSound) const
 {
 	SubsystemCollection.ForEachSubsystem<IActiveSoundUpdateInterface>([&ActiveSound](IActiveSoundUpdateInterface* ActiveSoundUpdate)
diff --git a/Engine/Source/Runtime/Engine/Public/ActiveSoundUpdateInterface.h b/Engine/Source/Runtime/Engine/Public/ActiveSoundUpdateInterface.h
--- a/Engine/Source/Runtime/Engine/Public/ActiveSoundUpdateInterface.h
+++ b/Engine/Source/Runtime/Engine/Public/ActiveSoundUpdateInterface.h
@@ -40,6 +40,22 @@ public:
 	 */
 	virtual void ApplyInteriorSettings(const FActiveSound& ActiveSound, FSoundParseParameters& ParseParams) {}
 
+	/**
+	 * Allows subsystems to modify SoundParseParameters just before Sound.Parse checks them to populate WaveInstances.
+	 * 
+	 * Modifying Transform is not recommended. It is used for calculations before this function is called,
+	 * which will not be able to observe any modifications done here.
+	 * 
+	 * Modifying SoundClass is allowed but not compatible with use of FAudioDevice::SetSoundClassDistanceScale.
+	 * The scale value will always be taken from the old class, never the new one.
+	 * 
+	 * NOTE! Called on the AudioThread during each audio frame Update
+	 *
+	 * @param ActiveSound	The active sound affected
+	 * @param ParseParams	The parameters to apply to the wave instances
+	 */
+	virtual void SubsystemsAdjustParseParams(const FActiveSound& ActiveSound, FSoundParseParameters& ParseParams) {}
+
 	/**
 	 * Called while an active sound is being constructed. Will be followed by either
 	 * NotifyActiveSoundCreated or NotifyVirtualizedSoundCreated.
diff --git a/Engine/Source/Runtime/Engine/Public/AudioDevice.h b/Engine/Source/Runtime/Engine/Public/AudioDevice.h
--- a/Engine/Source/Runtime/Engine/Public/AudioDevice.h
+++ b/Engine/Source/Runtime/Engine/Public/AudioDevice.h
@@ -848,6 +848,11 @@ public:
 	 */
 	ENGINE_API void ApplyInteriorSettings(FActiveSound& ActiveSound, FSoundParseParameters& ParseParams) const;
 
+	/**
+	 * Allows subsystems to modify ParseParams before Sound.Parse begins
+	 */
+	ENGINE_API void SubsystemsAdjustParseParams(FActiveSound& ActiveSound, FSoundParseParameters& ParseParams) const;
+
 	/**
 	 * Notifies subsystems an active sound is about to be deleted (called on audio thread) - Deprecated, see NotifySubsystemsActiveSoundDeleting
 	 */
