[World] Fix Undo on World Partition Settings

#tests Tested undo on old maps and new maps
[RN] Fixed Undo in world Partition Settings not working correctly

diff --git a/Engine/Plugins/Experimental/FastGeoStreaming/Source/FastGeoStreaming/Private/FastGeoWorldPartitionRuntimeCellTransformer.cpp b/Engine/Plugins/Experimental/FastGeoStreaming/Source/FastGeoStreaming/Private/FastGeoWorldPartitionRuntimeCellTransformer.cpp
--- a/Engine/Plugins/Experimental/FastGeoStreaming/Source/FastGeoStreaming/Private/FastGeoWorldPartitionRuntimeCellTransformer.cpp
+++ b/Engine/Plugins/Experimental/FastGeoStreaming/Source/FastGeoStreaming/Private/FastGeoWorldPartitionRuntimeCellTransformer.cpp
@@ -129,6 +129,19 @@ namespace FastGeo
 }
 #endif
 
+void UFastGeoTransformerSettings::PostLoad()
+{
+	Super::PostLoad();
+
+#if WITH_EDITOR
+	// Fix old objects that were created when the transformer was missing RF_Transactional
+	if (!HasAnyFlags(RF_ClassDefaultObject | RF_Transactional))
+	{
+		SetFlags(RF_Transactional);
+	}
+#endif
+}
+
 UFastGeoWorldPartitionRuntimeCellTransformer::UFastGeoWorldPartitionRuntimeCellTransformer(const FObjectInitializer& ObjectInitializer)
 	: Super(ObjectInitializer)
 {
diff --git a/Engine/Plugins/Experimental/FastGeoStreaming/Source/FastGeoStreaming/Public/FastGeoWorldPartitionRuntimeCellTransformer.h b/Engine/Plugins/Experimental/FastGeoStreaming/Source/FastGeoStreaming/Public/FastGeoWorldPartitionRuntimeCellTransformer.h
--- a/Engine/Plugins/Experimental/FastGeoStreaming/Source/FastGeoStreaming/Public/FastGeoWorldPartitionRuntimeCellTransformer.h
+++ b/Engine/Plugins/Experimental/FastGeoStreaming/Source/FastGeoStreaming/Public/FastGeoWorldPartitionRuntimeCellTransformer.h
@@ -54,6 +54,10 @@ class UFastGeoTransformerSettings : public UDataAsset
 	GENERATED_BODY()
 	
 public:
+	//~ Begin UObject Interface.
+	virtual void PostLoad() override;
+	//~ End UObject Interface.
+
 	/** Allowed actor classes (recursive) to convert to FastGeo */
 	UPROPERTY(EditAnywhere, Category="Settings", meta= (AllowAbstract = "true"))
 	TArray<TSubclassOf<AActor>> AllowedActorClasses;
diff --git a/Engine/Source/Runtime/Engine/Private/WorldPartition/WorldPartition.cpp b/Engine/Source/Runtime/Engine/Private/WorldPartition/WorldPartition.cpp
--- a/Engine/Source/Runtime/Engine/Private/WorldPartition/WorldPartition.cpp
+++ b/Engine/Source/Runtime/Engine/Private/WorldPartition/WorldPartition.cpp
@@ -535,7 +535,7 @@ void UWorldPartition::PostEditChangeChainProperty(FPropertyChangedChainEvent& Pr
 		
 		UWorldPartitionRuntimeCellTransformer* OldTransformerInstance = TransformerInstance.Instance;
 
-		TransformerInstance.Instance = TransformerInstance.Class ? NewObject<UWorldPartitionRuntimeCellTransformer>(this, TransformerInstance.Class, NAME_None) : nullptr;
+		TransformerInstance.Instance = TransformerInstance.Class ? NewObject<UWorldPartitionRuntimeCellTransformer>(this, TransformerInstance.Class, NAME_None, RF_Transactional) : nullptr;
 
 		if (OldTransformerInstance && TransformerInstance.Instance)
 		{
@@ -1224,7 +1224,7 @@ UWorldPartition* UWorldPartition::CreateOrRepairWorldPartition(AWorldSettings* W
 
 	if (!WorldPartition)
 	{
-		WorldPartition = NewObject<UWorldPartition>(WorldSettings);
+		WorldPartition = NewObject<UWorldPartition>(WorldSettings, NAME_None, RF_Transactional);
 		WorldSettings->SetWorldPartition(WorldPartition);
 
 		// New maps should include GridSize in name
@@ -1702,6 +1702,19 @@ void UWorldPartition::Serialize(FArchive& Ar)
 	}
 }
 
+void UWorldPartition::PostLoad()
+{
+	Super::PostLoad();
+
+#if WITH_EDITOR
+	// Fix old objects that were created without RF_Transactional
+	if (!HasAnyFlags(RF_ClassDefaultObject | RF_Transactional))
+	{
+		SetFlags(RF_Transactional);
+	}
+#endif
+}
+
 UWorld* UWorldPartition::GetWorld() const
 {
 	if (World)
@@ -2632,7 +2645,7 @@ void UWorldPartition::AddRuntimeCellTransformerOfType(TSubclassOf<UWorldPartitio
 			RuntimeCellsTransformerStack.InsertDefaulted(Index);
 			FRuntimeCellTransformerInstance& TransformerInstance = RuntimeCellsTransformerStack[Index];
 			TransformerInstance.Class = CellTransformerClass;
-			TransformerInstance.Instance = TransformerInstance.Class ? NewObject<UWorldPartitionRuntimeCellTransformer>(this, TransformerInstance.Class, NAME_None) : nullptr;
+			TransformerInstance.Instance = TransformerInstance.Class ? NewObject<UWorldPartitionRuntimeCellTransformer>(this, TransformerInstance.Class, NAME_None, RF_Transactional) : nullptr;
 		}
 	}
 }
diff --git a/Engine/Source/Runtime/Engine/Private/WorldPartition/WorldPartitionRuntimeCellTransformer.cpp b/Engine/Source/Runtime/Engine/Private/WorldPartition/WorldPartitionRuntimeCellTransformer.cpp
--- a/Engine/Source/Runtime/Engine/Private/WorldPartition/WorldPartitionRuntimeCellTransformer.cpp
+++ b/Engine/Source/Runtime/Engine/Private/WorldPartition/WorldPartitionRuntimeCellTransformer.cpp
@@ -8,6 +8,19 @@ UWorldPartitionRuntimeCellTransformer::UWorldPartitionRuntimeCellTransformer(con
 	: Super(ObjectInitializer)
 {}
 
+void UWorldPartitionRuntimeCellTransformer::PostLoad()
+{
+	Super::PostLoad();
+
+#if WITH_EDITOR
+	// Fix old objects that were created without RF_Transactional
+	if (!HasAnyFlags(RF_ClassDefaultObject | RF_Transactional))
+	{
+		SetFlags(RF_Transactional);
+	}
+#endif
+}
+
 #if WITH_EDITOR
 bool UWorldPartitionRuntimeCellTransformer::IsLevelTransformable(ULevel* InLevel) const
 {
diff --git a/Engine/Source/Runtime/Engine/Public/WorldPartition/WorldPartition.h b/Engine/Source/Runtime/Engine/Public/WorldPartition/WorldPartition.h
--- a/Engine/Source/Runtime/Engine/Public/WorldPartition/WorldPartition.h
+++ b/Engine/Source/Runtime/Engine/Public/WorldPartition/WorldPartition.h
@@ -223,6 +223,7 @@ public:
 	ENGINE_API virtual void PreSave(FObjectPreSaveContext SaveContext) override;
 #endif //WITH_EDITOR
 	ENGINE_API virtual void Serialize(FArchive& Ar) override;
+	ENGINE_API virtual void PostLoad() override;
 	ENGINE_API virtual UWorld* GetWorld() const override;
 	ENGINE_API virtual bool ResolveSubobject(const TCHAR* SubObjectPath, UObject*& OutObject, bool bLoadIfExists) override;
 	ENGINE_API virtual void BeginDestroy() override;
diff --git a/Engine/Source/Runtime/Engine/Public/WorldPartition/WorldPartitionRuntimeCellTransformer.h b/Engine/Source/Runtime/Engine/Public/WorldPartition/WorldPartitionRuntimeCellTransformer.h
--- a/Engine/Source/Runtime/Engine/Public/WorldPartition/WorldPartitionRuntimeCellTransformer.h
+++ b/Engine/Source/Runtime/Engine/Public/WorldPartition/WorldPartitionRuntimeCellTransformer.h
@@ -27,8 +27,12 @@ class UWorldPartitionRuntimeCellTransformer : public UObject
 {
 	GENERATED_UCLASS_BODY()
 
-#if WITH_EDITOR
 public:
+	//~ Begin UObject Interface.
+	UE_API virtual void PostLoad() override;
+	//~ End UObject Interface.
+
+#if WITH_EDITOR
 	/** Whether the level can be processed by the transformer. */
 	UE_API virtual bool IsLevelTransformable(ULevel* InLevel) const;
 	UE_API virtual void PreTransform(ULevel* InLevel);
