Anim Notify fixes, among other things.

- Deprecated bNeedsUpdate variable public access in UAnimInstance, instead use getter.
- Update CharacterMovementComponent to properly dispatch notifies every time the character pose is ticked (autonomous proxy).
- This fixes the continous begin/end event dispatch for anim notify states.
- Added a.Notifies.DispatchPolicy to control if an UAnimInstance dispatches notifies from all sources (montage/anim graph) or only those that have updated before the dispatch request.
- This is mainlly used to maintain lifetimes of active anim state notifies when triggering notifies (since we swap active arrays).
- Updated TriggerAnimNotifies() to add overrides to force dispatch notifies from montage only, anim graph only, or all sources.

diff --git a/Engine/Plugins/Runtime/GameplayAbilities/Source/GameplayAbilities/Private/AbilitySystemComponent_Abilities.cpp b/Engine/Plugins/Runtime/GameplayAbilities/Source/GameplayAbilities/Private/AbilitySystemComponent_Abilities.cpp
--- a/Engine/Plugins/Runtime/GameplayAbilities/Source/GameplayAbilities/Private/AbilitySystemComponent_Abilities.cpp
+++ b/Engine/Plugins/Runtime/GameplayAbilities/Source/GameplayAbilities/Private/AbilitySystemComponent_Abilities.cpp
@@ -3395,7 +3395,11 @@ void UAbilitySystemComponent::OnRep_ReplicatedAnimMontage()
 						{
 							MontageInstance->UpdateWeight(DeltaTime);
 							MontageInstance->HandleEvents(CurrentPosition, ConstRepAnimMontageInfo.Position, nullptr);
-							AnimInstance->TriggerAnimNotifies(DeltaTime);
+							
+							// @todo: Not sure how we haven't noticed missed notifies, since we are queueing AND firing notifies from the UAnimInstance we need to manually handle the unfiltered notifies.
+							
+							// To keep the anim graph notify states alive, we force to only trigger montage notifies.
+							AnimInstance->TriggerAnimNotifies(DeltaTime, UAnimInstance::ENotifyTriggerMode::ForceMontageOnly);
 						}
 					}
 					AnimInstance->Montage_SetPosition(LocalAnimMontageInfo.AnimMontage, ConstRepAnimMontageInfo.Position);
diff --git a/Engine/Source/Runtime/Engine/Classes/Animation/AnimInstance.h b/Engine/Source/Runtime/Engine/Classes/Animation/AnimInstance.h
--- a/Engine/Source/Runtime/Engine/Classes/Animation/AnimInstance.h
+++ b/Engine/Source/Runtime/Engine/Classes/Animation/AnimInstance.h
@@ -380,6 +380,7 @@ class UAnimInstance : public UObject
 	uint8 bUsingCopyPoseFromMesh : 1;
 
 	/** Flag to check back on the game thread that indicates we need to run PostUpdateAnimation() in the post-eval call */
+	UE_DEPRECATED(5.7, "Public access to bNeedsUpdate is deprecated, and it will become private in a future release. Please update code to use the NeedsUpdate() getter instead.")
 	uint8 bNeedsUpdate : 1;
 
 	/** Flag to check if created by LinkedAnimGraph in ReinitializeLinkedAnimInstance */
@@ -414,12 +415,18 @@ private:
 	uint8 bInitializing : 1;
 
 	/** Used to guard against recursive calls to UpdateAnimation */
-	bool bPostUpdatingAnimation;
+	uint8 bPostUpdatingAnimation : 1;
 
 	/** Used to guard against recursive calls to UpdateAnimation */
-	bool bUpdatingAnimation;
+	uint8 bUpdatingAnimation : 1;
 #endif
 
+	/** Used to prevent firing notifies queued from the anim graph if we have already done the animation update this tick (non-autonomous). */
+	uint8 bHasDispatchedAnimGraphNotifiesThisTick : 1;
+
+	/** Used to prevent firing notifies queued from montages if we have already done the montage animation update this tick (non-autonomous). */
+	uint8 bHasDispatchedMontageNotifiesThisTick : 1;
+	
 public:
 
 	// @todo document
@@ -1012,10 +1019,10 @@ public:
 
 #if DO_CHECK
 	/** Is this animation currently running post update */
-	bool IsPostUpdatingAnimation() const { return bPostUpdatingAnimation; }
+	bool IsPostUpdatingAnimation() const { return static_cast<bool>(bPostUpdatingAnimation); }
 
 	/** Is this animation currently running update */
-	bool IsUpdatingAnimation() const { return bUpdatingAnimation; }
+	bool IsUpdatingAnimation() const { return static_cast<bool>(bUpdatingAnimation); }
 #endif
 
 	/** Set RootMotionMode */
@@ -1331,6 +1338,26 @@ public:
 		Default
 	};
 
+	/** Controls which notifies are triggered (used by TriggerAnimNotifies()). Useful mainly when handling active anim notify state lifetimes. */
+	enum class ENotifyTriggerMode : uint8
+	{
+		/**
+		 * Uses the global dispatch policy (a.Notifies.DispatchPolicy) for eligibility:
+		 *   1: only trigger notifies from sources (anim graph or montage) that updated before the trigger request.
+		 *   0: trigger notifies from all sources.
+		 */
+		Default,
+
+		/** Override: Only trigger notifies queued by the anim graph. Ignores global dispatch policy. */
+		ForceAnimGraphOnly,
+
+		/** Override: Only trigger notifies queued by montages. Ignores global dispatch policy. */
+		ForceMontageOnly,
+
+		/** Override: Trigger notifies from all sources (anim graph/montages). Ignores global dispatch policy. */
+		ForceAllSources,
+	};
+	
 	// Animation phase trigger
 	// start with initialize
 	// update happens in every tick. Can happen in parallel with others if conditions are right.
@@ -1552,7 +1579,7 @@ public:
 	ENGINE_API virtual UWorld* GetWorld() const override;
 
 	/** Trigger AnimNotifies **/
-	ENGINE_API void TriggerAnimNotifies(float DeltaSeconds);
+	ENGINE_API void TriggerAnimNotifies(float DeltaSeconds, ENotifyTriggerMode Mode = ENotifyTriggerMode::Default);
 
 	/** Trigger an AnimNotify.  Note that this version does not provide any context for the event **/ 
 	ENGINE_API void TriggerSingleAnimNotify(const FAnimNotifyEvent* AnimNotifyEvent);
diff --git a/Engine/Source/Runtime/Engine/Private/Animation/AnimInstance.cpp b/Engine/Source/Runtime/Engine/Private/Animation/AnimInstance.cpp
--- a/Engine/Source/Runtime/Engine/Private/Animation/AnimInstance.cpp
+++ b/Engine/Source/Runtime/Engine/Private/Animation/AnimInstance.cpp
@@ -117,6 +117,31 @@ namespace MontageCVars
 		TEXT("Whether we should flush all completed montages IMMEDIATELY when a new montage stomps the group. Use this to prevent accumulating montages when animation tick is paused, and firing all at once when we unpause."));
 }
 
+/** Used to determine notify dispatch behaviour when sources (montage, anim graph) may tick at different times within a frame. */
+enum class EAnimNotifyDispatchPolicy : int32
+{
+	/** Previous behavior - update active notify buffers and dispatch from any source, even if that source hasn't updated before the dispatch request. */
+	AllSources = 0,
+	
+	/**
+	 *  Only dispatch from sources that HAVE updated before the dispatch request.
+	 *  
+	 *  Anim graph notifies/states are not triggered/ticked if the anim graph hasn't updated yet this tick.
+	 *  Active notify states from skipped sources remain valid (useful in MoveAutonomous / multi-tick scenarios).
+	 */
+	UpdatedSourcesOnly = 1,
+};
+
+int32 GAnimNotifyDispatchPolicy = static_cast<int32>(EAnimNotifyDispatchPolicy::UpdatedSourcesOnly);
+static FAutoConsoleVariableRef CVarNotifiesDispatchPolicy(
+	TEXT("a.Notifies.DispatchPolicy"),
+	GAnimNotifyDispatchPolicy,
+	TEXT("Controls notify dispatch behavior when it's source (montage, anim graph) may tick at different times within a frame.\n")
+		TEXT("  0 = AllSources: Update active notify buffers and dispatch from any source, even if that source has not updated yet this tick (previous behavior).\n")
+		TEXT("  1 = UpdatedSourcesOnly: Only dispatch from sources that have updated this tick. Anim graph notifies/states are not triggered/ticked if the graph has not updated yet, ")
+		TEXT("but active notify states from skipped sources remain valid. Useful for multi-tick scenarios such as MoveAutonomous."),
+	ECVF_Default);
+
 /////////////////////////////////////////////////////
 // UAnimInstance
 /////////////////////////////////////////////////////
@@ -129,8 +154,11 @@ UAnimInstance::UAnimInstance(const FObjectInitializer& ObjectInitializer)
 #endif
 {
 	RootMotionMode = ERootMotionMode::RootMotionFromMontagesOnly;
-	bNeedsUpdate = false;
 
+	PRAGMA_DISABLE_DEPRECATION_WARNINGS
+	bNeedsUpdate = false;
+	PRAGMA_ENABLE_DEPRECATION_WARNINGS
+	
 	// Default to using threaded animation update.
 	bUseMultiThreadedAnimationUpdate = true;
 	bCreatedByLinkedAnimGraph = false;
@@ -145,8 +173,13 @@ UAnimInstance::UAnimInstance(const FObjectInitializer& ObjectInitializer)
 	bUpdateAnimationEnabled = true;
 #if DO_CHECK
 	bInitializing = false;
+	bPostUpdatingAnimation = false;
+	bUpdatingAnimation = false;
 #endif
 
+	bHasDispatchedAnimGraphNotifiesThisTick = false;
+	bHasDispatchedMontageNotifiesThisTick = false;
+	
 #if WITH_EDITOR
 	if(!HasAnyFlags(RF_ClassDefaultObject) && !GetClass()->HasAnyClassFlags(CLASS_Native))
 	{
@@ -581,7 +614,7 @@ void UAnimInstance::UpdateAnimation(float DeltaSeconds, bool bNeedsValidRootMoti
 
 #if DO_CHECK
 	checkf(!bUpdatingAnimation, TEXT("UpdateAnimation already in progress, circular detected for SkeletalMeshComponent [%s], AnimInstance [%s]"), *GetNameSafe(GetOwningComponent()),  *GetName());
-	TGuardValue<bool> CircularGuard(bUpdatingAnimation, true);
+	FGuardValue_Bitfield(bUpdatingAnimation, true);
 #endif
 
 	if (!IsUpdateAnimationEnabled())
@@ -736,8 +769,16 @@ void UAnimInstance::PreUpdateAnimation(float DeltaSeconds)
 	LLM_SCOPE(ELLMTag::Animation);
 	SCOPE_CYCLE_COUNTER(STAT_PreUpdateAnimation);
 
+	PRAGMA_DISABLE_DEPRECATION_WARNINGS
 	bNeedsUpdate = true;
-
+	PRAGMA_ENABLE_DEPRECATION_WARNINGS
+	
+	if (USkeletalMeshComponent* SkelMeshComp = GetSkelMeshComponent(); SkelMeshComp && !SkelMeshComp->PoseTickedThisFrame())
+	{
+		bHasDispatchedMontageNotifiesThisTick = false;
+		bHasDispatchedAnimGraphNotifiesThisTick = false;
+	}
+	
 	ClearQueuedAnimEvents(true);
 
 	RootMotionBlendQueue.Reset();
@@ -750,7 +791,7 @@ void UAnimInstance::PostUpdateAnimation()
 	LLM_SCOPE(ELLMTag::Animation);
 #if DO_CHECK
 	checkf(!bPostUpdatingAnimation, TEXT("PostUpdateAnimation already in progress, recursion detected for SkeletalMeshComponent [%s], AnimInstance [%s]"), *GetNameSafe(GetOwningComponent()), *GetName());
-	TGuardValue<bool> CircularGuard(bPostUpdatingAnimation, true);
+	FGuardValue_Bitfield(bPostUpdatingAnimation, true);
 #endif
 
 	SCOPE_CYCLE_COUNTER(STAT_PostUpdateAnimation);
@@ -777,13 +818,15 @@ void UAnimInstance::PostUpdateAnimation()
 	}
 
 	// Early out here if we are not set to needing an update
-	if(!bNeedsUpdate)
+	if(!NeedsUpdate())
 	{
 		return;
 	}
 
+	PRAGMA_DISABLE_DEPRECATION_WARNINGS
 	bNeedsUpdate = false;
-
+	PRAGMA_ENABLE_DEPRECATION_WARNINGS
+	
 	// acquire the proxy as we need to update
 	FAnimInstanceProxy& Proxy = GetProxyOnGameThread<FAnimInstanceProxy>();
 
@@ -866,6 +909,14 @@ void UAnimInstance::DispatchQueuedAnimEvents()
 			}
 		}
 	}
+
+	// We need to determine which anim update path's notifies have we triggered in the case we get multiple dispatch request per frame
+	// (from the pose being ticked in CharacterMovementComponent) to avoid clearing the ActiveNotify arrays after parallel anim eval is done.
+	if (USkeletalMeshComponent* SkelMeshComp = GetSkelMeshComponent())
+	{
+		bHasDispatchedMontageNotifiesThisTick = SkelMeshComp->bIsAutonomousTickPose && SkelMeshComp->PoseTickedThisFrame();
+		bHasDispatchedAnimGraphNotifiesThisTick = SkelMeshComp->bIsAutonomousTickPose && !NeedsUpdate() && SkelMeshComp->PoseTickedThisFrame();
+	}
 }
 
 void UAnimInstance::ClearQueuedAnimEvents(bool bShouldUpdateActiveAnimNotifiesSinceLastTick)
@@ -919,7 +970,9 @@ bool UAnimInstance::NeedsImmediateUpdate(float DeltaSeconds, bool bNeedsValidRoo
 
 bool UAnimInstance::NeedsUpdate() const
 {
+	PRAGMA_DISABLE_DEPRECATION_WARNINGS
 	return bNeedsUpdate;
+	PRAGMA_ENABLE_DEPRECATION_WARNINGS
 }
 
 bool UAnimInstance::IsInitialized() const
@@ -1641,7 +1694,7 @@ float UAnimInstance::GetDeltaSeconds() const
 	return GetProxyOnAnyThread<FAnimInstanceProxy>().GetDeltaSeconds();
 }
 
-void UAnimInstance::TriggerAnimNotifies(float DeltaSeconds)
+void UAnimInstance::TriggerAnimNotifies(float DeltaSeconds, ENotifyTriggerMode Mode)
 {
 	SCOPE_CYCLE_COUNTER(STAT_AnimTriggerAnimNotifies);
 	USkeletalMeshComponent* SkelMeshComp = GetSkelMeshComponent();
@@ -1653,6 +1706,38 @@ void UAnimInstance::TriggerAnimNotifies(float DeltaSeconds)
 	TArray<FAnimNotifyEventReference> NewActiveAnimNotifyEventReference;
 	NewActiveAnimNotifyEventReference.Reserve(NotifyQueue.AnimNotifies.Num());
 
+	// Override modes (filters).
+	const bool bForceMontageOnly = Mode == ENotifyTriggerMode::ForceMontageOnly;
+	const bool bForceAnimGraphOnly = Mode == ENotifyTriggerMode::ForceAnimGraphOnly;
+
+	// Filter based on override selection
+	if (bForceMontageOnly || bForceAnimGraphOnly)
+	{
+		TArray<FAnimNotifyEventReference, FAnimStackAllocator> FilteredAnimNotifies;
+		FilteredAnimNotifies.Reserve(NotifyQueue.AnimNotifies.Num());
+		
+		for (FAnimNotifyEventReference& Notify : NotifyQueue.AnimNotifies)
+		{
+			const bool bIsFromMontage = Notify.GetContextData<UE::Anim::FAnimNotifyMontageInstanceContext>() != nullptr;
+			const bool bFilterGraphNotify = !bIsFromMontage && bForceMontageOnly;
+			const bool bFilterMontageNotify = bIsFromMontage && bForceAnimGraphOnly;
+
+			if (bFilterGraphNotify || bFilterMontageNotify)
+			{
+				continue;
+			}
+
+			FilteredAnimNotifies.Push(Notify);
+		}
+
+		NotifyQueue.AnimNotifies = MoveTemp(FilteredAnimNotifies);
+	}
+	
+	const bool bDispatchFromSourcesOnly = static_cast<EAnimNotifyDispatchPolicy>(GAnimNotifyDispatchPolicy) == EAnimNotifyDispatchPolicy::UpdatedSourcesOnly && Mode == ENotifyTriggerMode::Default;
+
+	// We need to skip triggering notifies if we haven't updated the source (due to parallel update) or if we have already processed the notifies this tick (multiple anim updates/dispatches only allowed from autonomous tick/move requests). 
+	const bool bSkipNotifiesSourcedFromAnimGraph = bDispatchFromSourcesOnly && (NeedsUpdate() || (bHasDispatchedAnimGraphNotifiesThisTick && !SkelMeshComp->bIsAutonomousTickPose));
+	const bool bSkipNotifiesSourcedFromMontages = bDispatchFromSourcesOnly && (bHasDispatchedMontageNotifiesThisTick && !SkelMeshComp->bIsAutonomousTickPose);
 	
 	// AnimNotifyState freshly added that need their 'NotifyBegin' event called.
 	TArray<const FAnimNotifyEvent *> NotifyStateBeginEvent;
@@ -1660,14 +1745,14 @@ void UAnimInstance::TriggerAnimNotifies(float DeltaSeconds)
 
 	for (int32 Index=0; Index<NotifyQueue.AnimNotifies.Num(); Index++)
 	{
-		if(const FAnimNotifyEvent* AnimNotifyEvent = NotifyQueue.AnimNotifies[Index].GetNotify())
+		if (const FAnimNotifyEvent* AnimNotifyEvent = NotifyQueue.AnimNotifies[Index].GetNotify())
 		{
-			// AnimNotifyState
+			// Filter any notify states that are still active from the old tracking arrays to avoid triggering their respective being/end events.
 			if (AnimNotifyEvent->NotifyStateClass)
 			{
 				int32 ExistingItemIndex = INDEX_NONE;
-
-				if (ActiveAnimNotifyState.Find(*AnimNotifyEvent, ExistingItemIndex))
+				
+				if (const bool bIsNotifyStateAlreadyActive = ActiveAnimNotifyState.Find(*AnimNotifyEvent, ExistingItemIndex))
 				{
 					check(ActiveAnimNotifyState.Num() == ActiveAnimNotifyEventReference.Num());
 					ActiveAnimNotifyState.RemoveAtSwap(ExistingItemIndex, EAllowShrinking::No);
@@ -1695,25 +1780,40 @@ void UAnimInstance::TriggerAnimNotifies(float DeltaSeconds)
 	{
 		const FAnimNotifyEvent& AnimNotifyEvent = ActiveAnimNotifyState[Index];
 		const FAnimNotifyEventReference& EventReference = ActiveAnimNotifyEventReference[Index];
-		if (AnimNotifyEvent.NotifyStateClass && ShouldTriggerAnimNotifyState(AnimNotifyEvent.NotifyStateClass))
+
+		const bool bIsFromMontage = EventReference.GetContextData<UE::Anim::FAnimNotifyMontageInstanceContext>() != nullptr;
+		const bool bShouldSkipAnimGraphNotify = bSkipNotifiesSourcedFromAnimGraph && !bIsFromMontage;
+		const bool bShouldSkipMontageNotify = bSkipNotifiesSourcedFromMontages && bIsFromMontage;
+		
+		// Any notify state we skip from being considered as ended needs to be re-added to our new notify array to be able to properly keep track of its lifetime.
+		if (bShouldSkipAnimGraphNotify || bShouldSkipMontageNotify)
+		{
+			NewActiveAnimNotifyState.Add(AnimNotifyEvent);
+			FAnimNotifyEventReference& NewEventRef = NewActiveAnimNotifyEventReference.Add_GetRef(EventReference);
+			NewEventRef.SetNotify(&NewActiveAnimNotifyState.Top());
+		}
+		else
 		{
+			if (AnimNotifyEvent.NotifyStateClass && ShouldTriggerAnimNotifyState(AnimNotifyEvent.NotifyStateClass))
+			{
 #if WITH_EDITOR
-			// Prevent firing notifies in animation editors if requested 
-			if((SkelMeshComp && !SkelMeshComp->IsA<UDebugSkelMeshComponent>()) || AnimNotifyEvent.NotifyStateClass->ShouldFireInEditor())
+				// Prevent firing notifies in animation editors if requested 
+				if((SkelMeshComp && !SkelMeshComp->IsA<UDebugSkelMeshComponent>()) || AnimNotifyEvent.NotifyStateClass->ShouldFireInEditor())
 #endif
+				{
+					TRACE_ANIM_NOTIFY(this, AnimNotifyEvent, End);
+					AnimNotifyEvent.NotifyStateClass->NotifyEnd(SkelMeshComp, Cast<UAnimSequenceBase>(AnimNotifyEvent.NotifyStateClass->GetOuter()), EventReference);
+				}
+			}
+			// The NotifyEnd callback above may have triggered actor destruction and the tear down
+			// of this instance via UninitializeAnimation which empties ActiveAnimNotifyState.
+			// If that happened, we should stop iterating the ActiveAnimNotifyState array
+			if (ActiveAnimNotifyState.IsValidIndex(Index) == false)
 			{
-				TRACE_ANIM_NOTIFY(this, AnimNotifyEvent, End);
-				AnimNotifyEvent.NotifyStateClass->NotifyEnd(SkelMeshComp, Cast<UAnimSequenceBase>(AnimNotifyEvent.NotifyStateClass->GetOuter()), EventReference);
+				ensureMsgf(false, TEXT("UAnimInstance::ActiveAnimNotifyState has been invalidated by NotifyEnd. AnimInstance: %s, Owning Component: %s, Owning Actor: %s "), *GetNameSafe(this), *GetNameSafe(GetOwningComponent()), *GetNameSafe(GetOwningActor()));
+				return;
 			}
 		}
-		// The NotifyEnd callback above may have triggered actor destruction and the tear down
-		// of this instance via UninitializeAnimation which empties ActiveAnimNotifyState.
-		// If that happened, we should stop iterating the ActiveAnimNotifyState array
-		if (ActiveAnimNotifyState.IsValidIndex(Index) == false)
-		{
-			ensureMsgf(false, TEXT("UAnimInstance::ActiveAnimNotifyState has been invalidated by NotifyEnd. AnimInstance: %s, Owning Component: %s, Owning Actor: %s "), *GetNameSafe(this), *GetNameSafe(GetOwningComponent()), *GetNameSafe(GetOwningActor()));
-			return;
-		}
 	}
 
 	check(NotifyStateBeginEventReference.Num() == NotifyStateBeginEvent.Num());
@@ -1737,12 +1837,20 @@ void UAnimInstance::TriggerAnimNotifies(float DeltaSeconds)
 	// Switch our arrays.
 	ActiveAnimNotifyState = MoveTemp(NewActiveAnimNotifyState);
 	ActiveAnimNotifyEventReference = MoveTemp(NewActiveAnimNotifyEventReference);
+	
 	// Tick currently active AnimNotifyState
 	for (int32 Index = 0; Index < ActiveAnimNotifyState.Num(); Index++)
 	{
 		const FAnimNotifyEvent& AnimNotifyEvent = ActiveAnimNotifyState[Index];
 		const FAnimNotifyEventReference& EventReference = ActiveAnimNotifyEventReference[Index];
-		if (ShouldTriggerAnimNotifyState(AnimNotifyEvent.NotifyStateClass))
+
+		// Only tick notifies states if their source is active.
+		const bool bIsFromMontage = EventReference.GetContextData<UE::Anim::FAnimNotifyMontageInstanceContext>() != nullptr;
+		const bool bShouldSkipAnimGraphNotify = bSkipNotifiesSourcedFromAnimGraph && !bIsFromMontage;
+		const bool bShouldSkipMontageNotify = bSkipNotifiesSourcedFromMontages && bIsFromMontage;
+		const bool bShouldTickNotifyState = !bShouldSkipMontageNotify || !bShouldSkipAnimGraphNotify; 
+
+		if (ShouldTriggerAnimNotifyState(AnimNotifyEvent.NotifyStateClass) && bShouldTickNotifyState)
 		{
 #if WITH_EDITOR
 			// Prevent firing notifies in animation editors if requested 
diff --git a/Engine/Source/Runtime/Engine/Private/Components/CharacterMovementComponent.cpp b/Engine/Source/Runtime/Engine/Private/Components/CharacterMovementComponent.cpp
--- a/Engine/Source/Runtime/Engine/Private/Components/CharacterMovementComponent.cpp
+++ b/Engine/Source/Runtime/Engine/Private/Components/CharacterMovementComponent.cpp
@@ -10505,8 +10505,6 @@ bool UCharacterMovementComponent::ServerMoveOld_Validate(float OldTimeStamp, FVe
 	return true;
 }
 
-
-
 void UCharacterMovementComponent::MoveAutonomous
 	(
 	float ClientTimeStamp,
@@ -10544,7 +10542,7 @@ void UCharacterMovementComponent::MoveAutonomous
 	}
 
 	// If not playing root motion, tick animations after physics. We do this here to keep events, notifies, states and transitions in sync with client updates.
-	if( CharacterOwner && !CharacterOwner->bClientUpdating && !CharacterOwner->IsPlayingRootMotion() && CharacterOwner->GetMesh() )
+	if (CharacterOwner && !CharacterOwner->bClientUpdating && !CharacterOwner->IsPlayingRootMotion() && CharacterOwner->GetMesh())
 	{
 		if (!bWasPlayingRootMotion) // If we were playing root motion before PerformMovement but aren't anymore, we're on the last frame of anim root motion and have already ticked character
 		{
@@ -10555,25 +10553,8 @@ void UCharacterMovementComponent::MoveAutonomous
 		USkeletalMeshComponent* OwnerMesh = CharacterOwner->GetMesh();
 		check(OwnerMesh != nullptr)
 		
-		if (CharacterMovementCVars::EnableQueuedAnimEventsOnServer)
-		{
-			if (UAnimInstance* AnimInstance = OwnerMesh->GetAnimInstance())
-			{
-				if (OwnerMesh->VisibilityBasedAnimTickOption <= EVisibilityBasedAnimTickOption::AlwaysTickPose && AnimInstance->NeedsUpdate())
-				{
-					// If we are doing a full graph update on the server but its doing a parallel update,
-					// trigger events right away since these are notifies queued from the montage update, and we could be receiving multiple ServerMoves per frame.
-					OwnerMesh->ConditionallyDispatchQueuedAnimEvents();
-
-					// We need to manually clear the anim notify queue (since normally its only is cleared in PreUpdateAnimation()) otherwise if animation ticks, the notifies queued from the ServerMove would fire twice.
-					AnimInstance->ClearQueuedAnimEvents(false);
-
-					// When animation ticks, we want its queued events to be triggered.
-					OwnerMesh->AllowQueuedAnimEventsNextDispatch();
-				}
-			}
-		}
-		else
+		// Fallback to old behaviour if needed. Otherwise, event are dispatched in TickCharacterPose().
+		if (!CharacterMovementCVars::EnableQueuedAnimEventsOnServer)
 		{
 			// Revert back to old behavior if wanted/needed.
 			if (OwnerMesh->ShouldOnlyTickMontages(DeltaTime) || OwnerMesh->ShouldOnlyTickMontagesAndRefreshBones(DeltaTime))
@@ -11880,6 +11861,38 @@ void UCharacterMovementComponent::TickCharacterPose(float DeltaTime)
 			}
 #endif // !(UE_BUILD_SHIPPING)
 		}
+
+		// We need to dispatch any queued events, since these can affect the next move request, after we ticked the character pose.
+		if (CharacterMovementCVars::EnableQueuedAnimEventsOnServer)
+		{
+			UAnimInstance* AnimInstance = CharacterMesh->GetAnimInstance();
+			UAnimInstance* PostProcessAnimInstance = CharacterMesh->GetPostProcessInstance();
+			
+			if (AnimInstance || PostProcessAnimInstance)
+			{
+				const bool bIsAnimGraphUpdateEnabled = CharacterMesh->VisibilityBasedAnimTickOption <= EVisibilityBasedAnimTickOption::AlwaysTickPose;
+				const bool bIsMainAnimGraphUpdatePending = AnimInstance ? AnimInstance->NeedsUpdate() : false;
+				const bool bIsPostProcessAnimGraphUpdatePending =  PostProcessAnimInstance ? PostProcessAnimInstance->NeedsUpdate() : false;
+				const bool bIsAnimGraphUpdatePending = bIsMainAnimGraphUpdatePending || bIsPostProcessAnimGraphUpdatePending;
+				
+				if (bIsAnimGraphUpdateEnabled && bIsAnimGraphUpdatePending)
+				{
+					// If we are doing a full graph update on the server but its doing a parallel update,
+					// trigger events right away since these are notifies queued from the montage update, and we could be receiving multiple ServerMoves per frame.
+					CharacterMesh->ConditionallyDispatchQueuedAnimEvents();
+
+					// We need to manually clear the anim notify queue (since normally its only is cleared in PreUpdateAnimation()) otherwise if animation ticks, the notifies queued from the ServerMove would fire twice.
+					CharacterMesh->ForEachAnimInstance([](UAnimInstance* InAnimInstance)
+					{
+						check(InAnimInstance);
+						InAnimInstance->ClearQueuedAnimEvents(false);
+					});
+
+					// When animation ticks, we want its queued events to be triggered.
+					CharacterMesh->AllowQueuedAnimEventsNextDispatch();
+				}
+			}
+		}
 	}
 
 	CharacterMesh->bIsAutonomousTickPose = false;
