Updated GameUserSettings and Scalability::FQualityLevels to treat LandscapeQuality consistently with other quality levels:
- Added GameUserSetting APIs to set/get LandscapeQuality
- LandscapeQuality is considered when deciding whether currently using one overall quality level

diff --git a/Engine/Source/Runtime/Engine/Classes/GameFramework/GameUserSettings.h b/Engine/Source/Runtime/Engine/Classes/GameFramework/GameUserSettings.h
--- a/Engine/Source/Runtime/Engine/Classes/GameFramework/GameUserSettings.h
+++ b/Engine/Source/Runtime/Engine/Classes/GameFramework/GameUserSettings.h
@@ -261,6 +261,15 @@ public:
 	UFUNCTION(BlueprintCallable, Category = Settings)
 	ENGINE_API int32 GetShadingQuality() const;
 
+	// Sets the landscape quality (0..4, higher is better)
+	// @param Value 0:low, 1:medium, 2:high, 3:epic, 4:cinematic (gets clamped if needed)
+	UFUNCTION(BlueprintCallable, Category = Settings)
+	ENGINE_API void SetLandscapeQuality(int32 Value);
+
+	// Returns the landscape quality (0..4, higher is better)
+	UFUNCTION(BlueprintCallable, Category = Settings)
+	ENGINE_API int32 GetLandscapeQuality() const;
+
 	/** Checks if any user settings is different from current */
 	UFUNCTION(BlueprintPure, Category=Settings)
 	ENGINE_API virtual bool IsDirty() const;
diff --git a/Engine/Source/Runtime/Engine/Private/GameUserSettings.cpp b/Engine/Source/Runtime/Engine/Private/GameUserSettings.cpp
--- a/Engine/Source/Runtime/Engine/Private/GameUserSettings.cpp
+++ b/Engine/Source/Runtime/Engine/Private/GameUserSettings.cpp
@@ -1021,6 +1021,16 @@ int32 UGameUserSettings::GetShadingQuality() const
 	return ScalabilityQuality.ShadingQuality;
 }
 
+void UGameUserSettings::SetLandscapeQuality(int32 Value)
+{
+	ScalabilityQuality.SetLandscapeQuality(Value);
+}
+
+int32 UGameUserSettings::GetLandscapeQuality() const
+{
+	return ScalabilityQuality.LandscapeQuality;
+}
+
 UGameUserSettings* UGameUserSettings::GetGameUserSettings()
 {
 	return GEngine->GetGameUserSettings();
diff --git a/Engine/Source/Runtime/Engine/Private/Scalability.cpp b/Engine/Source/Runtime/Engine/Private/Scalability.cpp
--- a/Engine/Source/Runtime/Engine/Private/Scalability.cpp
+++ b/Engine/Source/Runtime/Engine/Private/Scalability.cpp
@@ -1085,7 +1085,7 @@ int32 FQualityLevels::GetSingleQualityLevel() const
 	int32 Result = ViewDistanceQuality;
 
 	const int32 Target = ViewDistanceQuality;
-	if ((Target == AntiAliasingQuality) && (Target == ShadowQuality) && (Target == GlobalIlluminationQuality) && (Target == ReflectionQuality) && (Target == PostProcessQuality) && (Target == TextureQuality) && (Target == EffectsQuality) && (Target == FoliageQuality) && (Target == ShadingQuality))
+	if ((Target == AntiAliasingQuality) && (Target == ShadowQuality) && (Target == GlobalIlluminationQuality) && (Target == ReflectionQuality) && (Target == PostProcessQuality) && (Target == TextureQuality) && (Target == EffectsQuality) && (Target == FoliageQuality) && (Target == ShadingQuality) && (Target == LandscapeQuality))
 	{
 		if (GetRenderScaleLevelFromQualityLevel(Target) == ResolutionQuality)
 		{
