[Backout] - CL47161616
#rnx
Original CL Desc
-----------------------------------------------------------------
Pulled CL from sarah.noonan
Add option to force set the outer world of an external streaming object if its outer world on load is invalid. Disabled by default, can be enabled via ExternalDataLayerManager.ShouldForceSetExternalStreamingObjectOuterWorldIfNull

#rnx

diff --git a/Engine/Source/Runtime/Engine/Private/WorldPartition/DataLayer/ExternalDataLayerManager.cpp b/Engine/Source/Runtime/Engine/Private/WorldPartition/DataLayer/ExternalDataLayerManager.cpp
--- a/Engine/Source/Runtime/Engine/Private/WorldPartition/DataLayer/ExternalDataLayerManager.cpp
+++ b/Engine/Source/Runtime/Engine/Private/WorldPartition/DataLayer/ExternalDataLayerManager.cpp
@@ -32,16 +32,6 @@
 
 #define LOCTEXT_NAMESPACE "ExternalDataLayerManager"
 
-namespace ExternalDataLayerManagerPrivate
-{
-	static bool bShouldForceSetExternalStreamingObjectOuterWorldIfNull = false;
-	static FAutoConsoleVariableRef CVarShouldForceSetExternalStreamingObjectOuterWorldIfNull(
-		TEXT("ExternalDataLayerManager.ShouldForceSetExternalStreamingObjectOuterWorldIfNull"),
-		bShouldForceSetExternalStreamingObjectOuterWorldIfNull,
-		TEXT("Whether to assign the current world as the outer world of an external streaming object if it does not have a valid outer world"),
-		ECVF_Default);
-}
-
 UExternalDataLayerManager::UExternalDataLayerManager()
 	: bIsInitialized(false)
 	, bIsRunningGameOrInstancedWorldPartition(false)
@@ -168,8 +158,7 @@ bool UExternalDataLayerManager::CanInjectExternalDataLayerAsset(const UExternalD
 
 bool UExternalDataLayerManager::RegisterExternalStreamingObjectForGameWorld(const UExternalDataLayerAsset* InExternalDataLayerAsset)
 {
-	UWorld* OuterWorld = GetTypedOuter<UWorld>();
-	check(OuterWorld && OuterWorld->IsGameWorld());
+	check(GetTypedOuter<UWorld>()->IsGameWorld());
 #if !WITH_EDITOR
 	check(InExternalDataLayerAsset);
 	if (ExternalStreamingObjects.Contains(InExternalDataLayerAsset))
@@ -190,7 +179,7 @@ bool UExternalDataLayerManager::RegisterExternalStreamingObjectForGameWorld(cons
 	UPackage* DestPackage = nullptr;
 	FLinkerInstancingContext InstancingContext;
 	FString SourceWorldPath, RemappedWorldPath;
-	if (OuterWorld->GetSoftObjectPathMapping(SourceWorldPath, RemappedWorldPath))
+	if (GetTypedOuter<UWorld>()->GetSoftObjectPathMapping(SourceWorldPath, RemappedWorldPath))
 	{
 		FString Source = FTopLevelAssetPath(SourceWorldPath).GetPackageName().ToString();
 		FString Remapped = FTopLevelAssetPath(RemappedWorldPath).GetPackageName().ToString();
@@ -219,14 +208,8 @@ bool UExternalDataLayerManager::RegisterExternalStreamingObjectForGameWorld(cons
 		return false;
 	}
 
-	if (ExternalDataLayerManagerPrivate::bShouldForceSetExternalStreamingObjectOuterWorldIfNull && !ExternalStreamingObject->GetOuterWorld())
-	{
-		UE_LOG(LogWorldPartition, Warning, TEXT("[EDL: %s] Outer world is invalid, force setting with current world"), *InExternalDataLayerAsset->GetName());
-		ExternalStreamingObject->SetOuterWorld(OuterWorld);
-	}
-
 	// Do some validation on ExternalStreamingObject's OuterWorld and OwningWorld
-	check(ExternalStreamingObject->GetOuterWorld() == OuterWorld);
+	check(ExternalStreamingObject->GetOuterWorld() == GetTypedOuter<UWorld>());
 	check(ExternalStreamingObject->GetOwningWorld() == GetOuterUWorldPartition()->GetWorld());
 
 	ExternalStreamingObjects.Emplace(InExternalDataLayerAsset, ExternalStreamingObject);
diff --git a/Engine/Source/Runtime/Engine/Private/WorldPartition/WorldPartitionRuntimeHash.cpp b/Engine/Source/Runtime/Engine/Private/WorldPartition/WorldPartitionRuntimeHash.cpp
--- a/Engine/Source/Runtime/Engine/Private/WorldPartition/WorldPartitionRuntimeHash.cpp
+++ b/Engine/Source/Runtime/Engine/Private/WorldPartition/WorldPartitionRuntimeHash.cpp
@@ -105,11 +105,6 @@ UWorld* URuntimeHashExternalStreamingObjectBase::GetOwningWorld() const
 	return OwningWorld.IsSet() ? OwningWorld.GetValue().Get() : GetOuterWorld()->GetWorldPartition()->GetWorld();
 }
 
-void URuntimeHashExternalStreamingObjectBase::SetOuterWorld(UWorld* InWorld)
-{
-	OuterWorld = InWorld;
-}
-
 void URuntimeHashExternalStreamingObjectBase::OnStreamingObjectLoaded(UWorld* InjectedWorld)
 {
 #if !WITH_EDITOR
diff --git a/Engine/Source/Runtime/Engine/Public/WorldPartition/WorldPartitionRuntimeHash.h b/Engine/Source/Runtime/Engine/Public/WorldPartition/WorldPartitionRuntimeHash.h
--- a/Engine/Source/Runtime/Engine/Public/WorldPartition/WorldPartitionRuntimeHash.h
+++ b/Engine/Source/Runtime/Engine/Public/WorldPartition/WorldPartitionRuntimeHash.h
@@ -86,8 +86,6 @@ public:
 	UExternalDataLayerInstance* GetRootExternalDataLayerInstance() { return const_cast<UExternalDataLayerInstance*>(RootExternalDataLayerInstance.Get()); }
 	ENGINE_API const UObject* GetLevelMountPointContextObject() const;
 
-	void SetOuterWorld(UWorld* InWorld);
-
 #if WITH_EDITOR
 	UE_DEPRECATED(5.4, "PopulateGeneratorPackageForCook is depreacted and was replaced by OnPopulateGeneratorPackageForCook")
 	ENGINE_API void PopulateGeneratorPackageForCook();
