[Montage] Add Montage frame advancing delegate, for custom Montage frame advancing.  Currently used for music synchronization.

#tests PIE BSE

diff --git a/Engine/Source/Runtime/Engine/Classes/Animation/AnimInstance.h b/Engine/Source/Runtime/Engine/Classes/Animation/AnimInstance.h
--- a/Engine/Source/Runtime/Engine/Classes/Animation/AnimInstance.h
+++ b/Engine/Source/Runtime/Engine/Classes/Animation/AnimInstance.h
@@ -78,6 +78,9 @@ DECLARE_DELEGATE_TwoParams(FOnMontageBlendingOutStarted, UAnimMontage*, bool /*b
 DECLARE_DELEGATE_OneParam(FOnMontageBlendedInEnded, UAnimMontage*)
 DECLARE_DELEGATE_ThreeParams(FOnMontageSectionChanged, UAnimMontage*, FName /*SectionName*/ , bool /*bLooped*/)
 
+// Experimental 5.8, FComputeDeltaFrameTime is experimental, and may change in the future.
+DECLARE_DELEGATE_RetVal_ThreeParams(float, FComputeDeltaFrameTime, int32 /*InstanceId*/, float /*CurrentTime*/, float /*DeltaTime*/)
+
 /**
 * Delegate for when Montage is started
 */
@@ -789,6 +792,10 @@ public:
 	If Montage reference is NULL, it will pick the first active montage found. */
 	ENGINE_API FOnMontageBlendingOutStarted* Montage_GetBlendingOutDelegate(UAnimMontage* Montage = NULL);
 
+	/**  Experimental: Returns the delegate if Montage is found, otherwise returns first one found */
+	UE_EXPERIMENTAL(5.8, "Experimental custom way to drive the montage frame time forward, API may change in the future")
+	ENGINE_API FComputeDeltaFrameTime* Montage_GetComputeFrameDeltaTimeDelegate(const UAnimMontage* Montage = nullptr);
+	
 	ENGINE_API void Montage_SetSectionChangedDelegate(FOnMontageSectionChanged& InOnMontageSectionChanged, UAnimMontage* Montage = NULL);
 	
 	/** Get pointer to SectionChanged delegate for Montage.
diff --git a/Engine/Source/Runtime/Engine/Classes/Animation/AnimMontage.h b/Engine/Source/Runtime/Engine/Classes/Animation/AnimMontage.h
--- a/Engine/Source/Runtime/Engine/Classes/Animation/AnimMontage.h
+++ b/Engine/Source/Runtime/Engine/Classes/Animation/AnimMontage.h
@@ -203,6 +203,14 @@ DECLARE_DELEGATE_OneParam(FOnMontageBlendedInEnded, class UAnimMontage*)
  */
 DECLARE_DELEGATE_ThreeParams( FOnMontageSectionChanged, class UAnimMontage*, FName /*SectionName*/ , bool /*bInterrupted*/)
 
+/* Delegate allows a c++ client to get a callback when it's time to advance the montage frame time.
+ * It then returns the DeltaTime to use. Currently used to drive montages according to a music clock timeline
+ * rather than advancing through world time. Note that playrate both at the montage and segment level
+ * are basically ignored by this system, because the current use does not require it.
+ */
+// Experimental 5.8, FComputeDeltaFrameTime is experimental, and may change in the future.
+DECLARE_DELEGATE_RetVal_ThreeParams(float, FComputeDeltaFrameTime, int32 /*InstanceId*/, float /*CurrentTime*/, float /*DeltaTime*/)
+
 /**
 	Helper struct to sub step through Montages when advancing time.
 	These require stopping at sections and branching points to potential jumps and loops.
@@ -346,7 +354,10 @@ struct FAnimMontageInstance
 	FOnMontageBlendingOutStarted OnMontageBlendingOutStarted;
 	FOnMontageBlendedInEnded OnMontageBlendedInEnded;
 	FOnMontageSectionChanged OnMontageSectionChanged;
-	
+
+	UE_EXPERIMENTAL(5.8, "FComputeDeltaFrameTime is experimental, and may change in the future.")
+	FComputeDeltaFrameTime OnComputeDeltaFrameTime;
+
 	UPROPERTY()
 	bool bPlaying;
 
diff --git a/Engine/Source/Runtime/Engine/Private/Animation/AnimInstance.cpp b/Engine/Source/Runtime/Engine/Private/Animation/AnimInstance.cpp
--- a/Engine/Source/Runtime/Engine/Private/Animation/AnimInstance.cpp
+++ b/Engine/Source/Runtime/Engine/Private/Animation/AnimInstance.cpp
@@ -3115,6 +3115,32 @@ FOnMontageBlendingOutStarted* UAnimInstance::Montage_GetBlendingOutDelegate(UAni
 	return NULL;
 }
 
+FComputeDeltaFrameTime* UAnimInstance::Montage_GetComputeFrameDeltaTimeDelegate(const UAnimMontage* Montage)
+{
+	if (Montage)
+	{
+		FAnimMontageInstance* MontageInstance = GetActiveInstanceForMontage(Montage);
+		if (MontageInstance)
+		{
+			return &MontageInstance->OnComputeDeltaFrameTime;
+		}
+	}
+	else
+	{
+		// If no Montage reference, use first active one found.
+		for (int32 InstanceIndex = 0; InstanceIndex < MontageInstances.Num(); InstanceIndex++)
+		{
+			FAnimMontageInstance* MontageInstance = MontageInstances[InstanceIndex];
+			if (MontageInstance && MontageInstance->IsActive())
+			{
+				return &MontageInstance->OnComputeDeltaFrameTime;
+			}
+		}
+	}
+
+	return nullptr;
+}
+
 void UAnimInstance::Montage_SetSectionChangedDelegate(FOnMontageSectionChanged& InOnMontageSectionChanged, UAnimMontage* Montage)
 {
 	if (Montage)
diff --git a/Engine/Source/Runtime/Engine/Private/Animation/AnimMontage.cpp b/Engine/Source/Runtime/Engine/Private/Animation/AnimMontage.cpp
--- a/Engine/Source/Runtime/Engine/Private/Animation/AnimMontage.cpp
+++ b/Engine/Source/Runtime/Engine/Private/Animation/AnimMontage.cpp
@@ -2176,7 +2176,15 @@ EMontageSubStepResult FMontageSubStepper::Advance(float& InOut_P_Original, const
 		bPlayingForward = (DeltaMove >= 0.f);
 		TimeStretchMarkerIndex = INDEX_NONE;
 	}
-	else 
+	else if (MontageInstance->OnComputeDeltaFrameTime.IsBound())
+	{
+		// Note that the compute delta frame time ignores play rate, at least for now.
+		DeltaMove = MontageInstance->OnComputeDeltaFrameTime.Execute(MontageInstance->GetInstanceID(), InOut_P_Original, TimeRemaining);
+		// we only set playrate for a sign check later, otherwise unused.
+		PlayRate = DeltaMove / TimeRemaining;
+		bPlayingForward = (DeltaMove >= 0.f); // *could* go backwards a bit. since music clock is not purely coupled with gameclock. 
+	}
+	else
 	{
 		PlayRate = MontageInstance->PlayRate * Montage->RateScale;
 
