Iris
* Fixed some comments in NetRootObjectFactory

diff --git a/Engine/Source/Runtime/Engine/Public/Net/Iris/ReplicationSystem/NetRootObjectFactory.h b/Engine/Source/Runtime/Engine/Public/Net/Iris/ReplicationSystem/NetRootObjectFactory.h
--- a/Engine/Source/Runtime/Engine/Public/Net/Iris/ReplicationSystem/NetRootObjectFactory.h
+++ b/Engine/Source/Runtime/Engine/Public/Net/Iris/ReplicationSystem/NetRootObjectFactory.h
@@ -31,7 +31,6 @@ public:
 	/** When preregistered we will find the object via it's NetHandle in the bridge */
 	bool bIsPreregistered = false;
 
-
 public:
 
 	ENGINE_API virtual FString ToString() const override;
@@ -41,12 +40,12 @@ public:
 	ENGINE_API virtual bool Deserialize(const UE::Net::FCreationHeaderContext& Context);
 };
 
-
 } // end namespace UE::Net
 
 
 /**
- * Responsible for creating headers allowing remote factories to spawn replicated actors
+ * Factory that can be used for any UObject based class that wants to be replicated autonomously.
+ * Can be overriden if the class implementation requires more advanced control over discovery, instantiation, location, etc.
  */
 UCLASS(MinimalAPI)
 class UNetRootObjectFactory : public UNetObjectFactory
@@ -60,7 +59,6 @@ public:
 public:
 
 	// UNetObjectFactory interface
-
 	ENGINE_API virtual void OnInit() override;
 
 	ENGINE_API virtual FInstantiateResult InstantiateReplicatedObjectFromHeader(const FInstantiateContext& Context, const UE::Net::FNetObjectCreationHeader* Header) override;
@@ -74,7 +72,6 @@ public:
 protected:
 
 	// UNetObjectFactory interface
-
 	ENGINE_API virtual TUniquePtr<UE::Net::FNetObjectCreationHeader> CreateAndFillHeader(UE::Net::FNetRefHandle Handle) override;
 
 	ENGINE_API virtual TUniquePtr<UE::Net::FNetObjectCreationHeader> CreateAndDeserializeHeader(const UE::Net::FCreationHeaderContext& Context) override;
@@ -83,8 +80,6 @@ protected:
 
 protected:
 
-	// Callbacks for derived factories to use
-
 	// Methods available to derived classes that want to reuse basic functionalities
 
 	/** Fill the basic header info for a given object */
