[Montage] Minor fixes to the use of OnComputeDeltaFrameTime.  Returns NotMoved if DeltaMove is nearly zero, and resets TimeStretchMarkerIndex

#tests PIE greenroom BSE

diff --git a/Engine/Source/Runtime/Engine/Private/Animation/AnimMontage.cpp b/Engine/Source/Runtime/Engine/Private/Animation/AnimMontage.cpp
--- a/Engine/Source/Runtime/Engine/Private/Animation/AnimMontage.cpp
+++ b/Engine/Source/Runtime/Engine/Private/Animation/AnimMontage.cpp
@@ -2180,9 +2180,15 @@ EMontageSubStepResult FMontageSubStepper::Advance(float& InOut_P_Original, const
 	{
 		// Note that the compute delta frame time ignores play rate, at least for now.
 		DeltaMove = MontageInstance->OnComputeDeltaFrameTime.Execute(MontageInstance->GetInstanceID(), InOut_P_Original, TimeRemaining);
+		if (FMath::IsNearlyZero(DeltaMove))
+		{
+			return EMontageSubStepResult::NotMoved;
+		}
+
 		// we only set playrate for a sign check later, otherwise unused.
 		PlayRate = DeltaMove / TimeRemaining;
 		bPlayingForward = (DeltaMove >= 0.f); // *could* go backwards a bit. since music clock is not purely coupled with gameclock. 
+		TimeStretchMarkerIndex = INDEX_NONE;
 	}
 	else
 	{
