Populate load error log for missing skeleton bones, compared to the mesh, with the bone names that are missing

diff --git a/Engine/Source/Runtime/Engine/Private/Animation/Skeleton.cpp b/Engine/Source/Runtime/Engine/Private/Animation/Skeleton.cpp
--- a/Engine/Source/Runtime/Engine/Private/Animation/Skeleton.cpp
+++ b/Engine/Source/Runtime/Engine/Private/Animation/Skeleton.cpp
@@ -776,6 +776,11 @@ void USkeleton::BuildLinkupData(const USkinnedAsset* InSkinnedAsset, FSkeletonTo
 	NewMeshLinkup.MeshToSkeletonTable.Empty(NumMeshBones);
 	NewMeshLinkup.MeshToSkeletonTable.AddUninitialized(NumMeshBones);
 
+#if WITH_EDITOR
+	TArray<FString> MissingBoneNames;
+	bool bPopulatedMissingBoneNames = false;
+#endif
+
 	for (int32 MeshBoneIndex = 0; MeshBoneIndex < NumMeshBones; MeshBoneIndex++)
 	{
 		const FName MeshBoneName = MeshRefSkel.GetBoneName(MeshBoneIndex);
@@ -788,11 +793,37 @@ void USkeleton::BuildLinkupData(const USkinnedAsset* InSkinnedAsset, FSkeletonTo
 #if WITH_EDITOR
 			// If we're in editor, and skeleton is missing a bone, fix it.
 			// not currently supported in-game.
+			 
+			// Gather remaining missing bone data
+			// In case the merge succeeds, we need to gather this data before the merge to populate the message log correctly 
+			if (!bPopulatedMissingBoneNames)
+			{
+				MissingBoneNames.Add(MeshBoneName.ToString());
+				for (int32 CheckMeshBoneIndex = MeshBoneIndex + 1; CheckMeshBoneIndex < NumMeshBones; ++CheckMeshBoneIndex)
+				{
+					const FName MissingMeshBoneName = MeshRefSkel.GetBoneName(CheckMeshBoneIndex);
+					if (SkeletonRefSkel.FindBoneIndex(MissingMeshBoneName) == INDEX_NONE)
+					{
+						MissingBoneNames.Add(MissingMeshBoneName.ToString());
+					}
+				}
+				bPopulatedMissingBoneNames = true;
+			}
+			
+			// Re-add all SkelMesh bones to the Skeleton.
+			const bool bSuccessfulMerge = MergeAllBonesToBoneTree(InSkinnedAsset);
+
+			// Fix missing bone.
+			SkeletonBoneIndex = SkeletonRefSkel.FindBoneIndex(MeshBoneName);
 
 			static FName NAME_LoadErrors("LoadErrors");
 
 			if (IsInGameThread() || IsInSlateThread())
 			{
+				// If the merge failed we only display the current missing bone as this error will be hit again for future missing bones. 
+				// If the merge succeeded, we don't expect this error to be hit again so we display all missing bone names. 
+				const FString AllMissingBonesNames = bSuccessfulMerge ? FString::Join(MissingBoneNames, TEXT(",")) : MeshBoneName.ToString();
+
 				FMessageLog LoadErrors(NAME_LoadErrors);
 
 				TSharedRef<FTokenizedMessage> Message = LoadErrors.Info();
@@ -801,16 +832,12 @@ void USkeleton::BuildLinkupData(const USkinnedAsset* InSkinnedAsset, FSkeletonTo
 				Message->AddToken(FTextToken::Create(LOCTEXT("SkeletonBuildLinkupMissingBones2", " is missing bones that SkeletalMesh ")));
 				Message->AddToken(FAssetNameToken::Create(InSkinnedAsset->GetPathName(), FText::FromString(GetNameSafe(InSkinnedAsset))));
 				Message->AddToken(FTextToken::Create(LOCTEXT("SkeletonBuildLinkupMissingBones3", "  needs. They will be added now. Please save the Skeleton!")));
+				Message->AddToken(FTextToken::Create(FText::Format(LOCTEXT("SkeletonBuildLinkupMissingBones4", " Missing bone name(s): {0}"), FText::FromString(AllMissingBonesNames))));
 				LoadErrors.Open();
 
 				bSpawnedErrorMessage = true;
 			}
 
-			// Re-add all SkelMesh bones to the Skeleton.
-			MergeAllBonesToBoneTree(InSkinnedAsset);
-
-			// Fix missing bone.
-			SkeletonBoneIndex = SkeletonRefSkel.FindBoneIndex(MeshBoneName);
 #endif // WITH_EDITOR
 
 			if (!bSpawnedErrorMessage)
