ISMPoolComponent PSO Precaching fix

diff --git a/Engine/Source/Runtime/Engine/Classes/Components/PrimitiveComponent.h b/Engine/Source/Runtime/Engine/Classes/Components/PrimitiveComponent.h
--- a/Engine/Source/Runtime/Engine/Classes/Components/PrimitiveComponent.h
+++ b/Engine/Source/Runtime/Engine/Classes/Components/PrimitiveComponent.h
@@ -698,11 +698,19 @@ protected:
 
 	uint8 bIgnoreBoundsForEditorFocus : 1;
 
+#if UE_WITH_PSO_PRECACHING
+	/** Whether PrecachedPSOs was called. */
+	uint8 bPrecachedPSOs : 1;
+#endif
+
 public:
 #if UE_WITH_PSO_PRECACHING
 	/** PSOs requested priority */
 	EPSOPrecachePriority PSOPrecacheRequestPriority : 2;
 	static_assert((int)EPSOPrecachePriority::Highest < 1 << 2);
+
+	/** Whether the PrecachePSOs call in SetStaticMesh should be skipped. */
+	uint8 bSkipPSOPrecaching : 1;
 #endif
 
 #if WITH_EDITOR
diff --git a/Engine/Source/Runtime/Engine/Private/Components/PrimitiveComponent.cpp b/Engine/Source/Runtime/Engine/Private/Components/PrimitiveComponent.cpp
--- a/Engine/Source/Runtime/Engine/Private/Components/PrimitiveComponent.cpp
+++ b/Engine/Source/Runtime/Engine/Private/Components/PrimitiveComponent.cpp
@@ -409,6 +409,8 @@ UPrimitiveComponent::UPrimitiveComponent(const FObjectInitializer& ObjectInitial
 #if UE_WITH_PSO_PRECACHING
 	bPSOPrecacheCalled = false;
 	PSOPrecacheRequestPriority = EPSOPrecachePriority::Medium;
+	bSkipPSOPrecaching = false;
+	bPrecachedPSOs = false;
 #endif // UE_WITH_PSO_PRECACHING
 	
 	bApplyImpulseOnDamage = true;
@@ -722,6 +724,10 @@ void UPrimitiveComponent::OnUnregister()
 #if WITH_EDITOR
 	FStaticLightingSystemInterface::OnPrimitiveComponentUnregistered.Broadcast(this);
 #endif
+
+#if UE_WITH_PSO_PRECACHING
+	bPrecachedPSOs = false;
+#endif
 }
 
 FPrimitiveComponentInstanceData::FPrimitiveComponentInstanceData(const UPrimitiveComponent* SourceComponent)
@@ -5097,7 +5103,7 @@ void UPrimitiveComponent::PrecachePSOs()
 	// Also only request PSOs from game thread because TStrongObjectPtr is used on the material to make
 	// it's not deleted via garbage collection when PSO precaching is still busy. TStrongObjectPtr can only
 	// be constructed on the GameThread
-	if (!FApp::CanEverRender() || !IsComponentPSOPrecachingEnabled())
+	if (!FApp::CanEverRender() || !IsComponentPSOPrecachingEnabled() || bSkipPSOPrecaching)
 	{
 		return;
 	}
@@ -5116,6 +5122,8 @@ void UPrimitiveComponent::PrecachePSOs()
 	PrecacheMaterialPSOs(PSOPrecacheDataArray, MaterialPSOPrecacheRequestIDs, GraphEvents);
 
 	RequestRecreateRenderStateWhenPSOPrecacheFinished(GraphEvents);
+
+	bPrecachedPSOs = true;
 #endif
 }
 
diff --git a/Engine/Source/Runtime/Engine/Private/Components/StaticMeshComponent.cpp b/Engine/Source/Runtime/Engine/Private/Components/StaticMeshComponent.cpp
--- a/Engine/Source/Runtime/Engine/Private/Components/StaticMeshComponent.cpp
+++ b/Engine/Source/Runtime/Engine/Private/Components/StaticMeshComponent.cpp
@@ -208,7 +208,6 @@ UStaticMeshComponent::UStaticMeshComponent(const FObjectInitializer& ObjectIniti
 	WireframeColorOverride = FColor(255, 255, 255, 255);
 
 	MinLOD = 0;
-	bPrecachedPSOs = false;
 	bOverrideLightMapRes = false;
 	OverriddenLightMapRes = 64;
 	SubDivisionStepSize = 32;
@@ -895,13 +894,11 @@ void UStaticMeshComponent::OnUnregister()
 
 	Super::OnUnregister();
 	ReleaseMaterialCacheTextures();
-	bPrecachedPSOs = false;
 }
 
 void UStaticMeshComponent::PrecachePSOs()
 {
 	Super::PrecachePSOs();
-	bPrecachedPSOs = true;
 }
 
 void UStaticMeshComponent::BeginPlay()
diff --git a/Engine/Source/Runtime/Experimental/ISMPool/Private/ISMPool/ISMPoolComponent.cpp b/Engine/Source/Runtime/Experimental/ISMPool/Private/ISMPool/ISMPoolComponent.cpp
--- a/Engine/Source/Runtime/Experimental/ISMPool/Private/ISMPool/ISMPoolComponent.cpp
+++ b/Engine/Source/Runtime/Experimental/ISMPool/Private/ISMPool/ISMPoolComponent.cpp
@@ -168,6 +168,11 @@ void FISMPoolISM::InitISM(const FISMPoolStaticMeshInstance& InMeshInstance, bool
 		bDisallowNanite |= Material->GetBlendMode() == BLEND_Translucent;
 	}
 
+#if UE_WITH_PSO_PRECACHING
+	// disable the implicit precachePSO in SetStaticMesh, so we can explicitly call it later to make sure we precache the correct state
+	ISMComponent->bSkipPSOPrecaching = true;
+#endif
+
 	ISMComponent->SetStaticMesh(StaticMesh);
 	ISMComponent->SetMobility((MeshInstance.Desc.Flags & FISMPoolComponentDescription::StaticMobility) != 0 ? EComponentMobility::Static : EComponentMobility::Movable);
 
@@ -231,6 +236,11 @@ void FISMPoolISM::InitISM(const FISMPoolStaticMeshInstance& InMeshInstance, bool
 	// Use a fixed seed to avoid getting a different seed at every run (see UInstancedStaticMeshComponent::OnRegister())
 	// A possible improvement would be to compute a hash from the owner component and use that as the seed.
 	ISMComponent->InstancingRandomSeed = 1;	
+
+#if UE_WITH_PSO_PRECACHING
+	ISMComponent->bSkipPSOPrecaching = false;
+	ISMComponent->PrecachePSOs();
+#endif
 }
 
 FISMPoolInstanceGroups::FInstanceGroupId FISMPoolISM::AddInstanceGroup(int32 InstanceCount, TArrayView<const float> CustomDataFloats)
