ISMPoolComponent PSO Precache fix: missing file in CL 47305905

diff --git a/Engine/Source/Runtime/Engine/Classes/Components/StaticMeshComponent.h b/Engine/Source/Runtime/Engine/Classes/Components/StaticMeshComponent.h
--- a/Engine/Source/Runtime/Engine/Classes/Components/StaticMeshComponent.h
+++ b/Engine/Source/Runtime/Engine/Classes/Components/StaticMeshComponent.h
@@ -199,9 +199,6 @@ protected:
 	/** Whether mesh painting is supported. This is can be set by derived component types that don't support mesh painting. */
 	uint8 bSupportMeshPainting : 1;
 
-	/** Whether PrecachedPSOs was called. */
-	uint8 bPrecachedPSOs : 1;
-
 public:
 
 #if WITH_EDITORONLY_DATA
