Optimize CVar sink function for UWorldPartitionHLODRuntimeSubsystem. Cache cvar values outside of any subsystem instance to avoid GetObjectsOfClass unless a cvar actually changed.

#tests PIE
#rnx

diff --git a/Engine/Source/Runtime/Engine/Private/WorldPartition/HLOD/HLODRuntimeSubsystem.cpp b/Engine/Source/Runtime/Engine/Private/WorldPartition/HLOD/HLODRuntimeSubsystem.cpp
--- a/Engine/Source/Runtime/Engine/Private/WorldPartition/HLOD/HLODRuntimeSubsystem.cpp
+++ b/Engine/Source/Runtime/Engine/Private/WorldPartition/HLOD/HLODRuntimeSubsystem.cpp
@@ -64,17 +64,32 @@ static TAutoConsoleVariable<int32> CVarHLODWarmupVTSizeClamp(
 	2048,
 	TEXT("Clamp VT warmup requests for safety."));
 
+namespace FHLODRuntimeSubsystem
+{
+	//Cached value for UWorldPartitionHLODRuntimeSubsystem::IsWarmupEnabled
+	static bool bWarmupEnabled = true;
+}
+
 static void HLODRuntimeSubsystemCVarSinkFunction()
 {
-	for (UWorld* World : TObjectRange<UWorld>(RF_ClassDefaultObject | RF_ArchetypeObject, true, EInternalObjectFlags::Garbage))
+	// Wait to evaluate if warmup is enabled until after RHI is initialized. Sink functions get called very early in engine init
+	if (!GIsRHIInitialized)
 	{
-		if (World->WorldType == EWorldType::Game || World->WorldType == EWorldType::PIE)
-		{
-			if (UWorldPartitionHLODRuntimeSubsystem* HLODRuntimeSubsystem = World->GetSubsystem<UWorldPartitionHLODRuntimeSubsystem>())
-			{
-				HLODRuntimeSubsystem->OnCVarsChanged();
-			}
-		}
+		return;
+	}
+
+	// Check if any cvars affecting warmup have changed
+	bool bNewWarmupEnabled = UWorldPartitionHLODRuntimeSubsystem::IsWarmupEnabled();
+	if (FHLODRuntimeSubsystem::bWarmupEnabled == bNewWarmupEnabled)
+	{
+		return;
+	}
+
+	// Update the static cached value and every subsystem instance
+	FHLODRuntimeSubsystem::bWarmupEnabled = bNewWarmupEnabled;
+	for (UWorldPartitionHLODRuntimeSubsystem* HLODRuntimeSubsystem : TObjectRange<UWorldPartitionHLODRuntimeSubsystem>(RF_ClassDefaultObject | RF_ArchetypeObject, false, EInternalObjectFlags::Garbage))
+	{
+		HLODRuntimeSubsystem->OnCVarsChanged();
 	}
 }
 
@@ -125,7 +140,7 @@ public:
 
 UWorldPartitionHLODRuntimeSubsystem::UWorldPartitionHLODRuntimeSubsystem()
 	: UWorldSubsystem()
-	, bCachedShouldPerformWarmup(true)
+	, bCachedShouldPerformWarmup(FHLODRuntimeSubsystem::bWarmupEnabled)
 {
 }
 
@@ -176,6 +191,36 @@ bool UWorldPartitionHLODRuntimeSubsystem::IsHLODEnabled()
 	return UWorldPartitionHLODRuntimeSubsystem::WorldPartitionHLODEnabled;
 }
 
+bool UWorldPartitionHLODRuntimeSubsystem::IsWarmupEnabled()
+{
+	// Test if warmup is disabled globally.
+	const bool bWarmupEnabled = CVarHLODWarmupEnabled.GetValueOnGameThread() != 0;
+	if (!bWarmupEnabled)
+	{
+		return false;
+	}
+
+	// If warmup num of frames is invalid, no warmup needed
+	if (CVarHLODWarmupNumFrames.GetValueOnGameThread() <= 0)
+	{
+		return false;
+	}
+
+	// If warmup num of frames is invalid, no warmup needed
+	const EShaderPlatform ShaderPlatform = GShaderPlatformForFeatureLevel[GMaxRHIFeatureLevel];
+	const bool bNaniteEnabled = UseNanite(ShaderPlatform);
+	const bool bVirtualTextureEnabled = UseVirtualTexturing(ShaderPlatform);
+	const bool bWarmupNanite = CVarHLODWarmupNanite.GetValueOnGameThread() != 0;
+	const bool bWarmupVT = CVarHLODWarmupVT.GetValueOnGameThread() != 0;
+	const bool bWarmupNeeded = (bNaniteEnabled && bWarmupNanite) || (bVirtualTextureEnabled && bWarmupVT);
+	if (!bWarmupNeeded)
+	{
+		return false;
+	}
+
+	return true;
+}
+
 bool UWorldPartitionHLODRuntimeSubsystem::DoesSupportWorldType(const EWorldType::Type WorldType) const
 {
 	return WorldType == EWorldType::Game || WorldType == EWorldType::PIE;
@@ -750,27 +795,8 @@ void UWorldPartitionHLODRuntimeSubsystem::OnCVarsChanged()
 
 bool UWorldPartitionHLODRuntimeSubsystem::ShouldPerformWarmup() const
 {
-	// Test if warmup is disabled globally.
-	const bool bWarmupEnabled = CVarHLODWarmupEnabled.GetValueOnGameThread() != 0;
-	if (!bWarmupEnabled)
-	{
-		return false;
-	}
-
-	// If warmup num of frames is invalid, no warmup needed
-	if (CVarHLODWarmupNumFrames.GetValueOnGameThread() <= 0)
-	{
-		return false;
-	}
-
-	// If warmup num of frames is invalid, no warmup needed
-	const EShaderPlatform ShaderPlatform = GShaderPlatformForFeatureLevel[GMaxRHIFeatureLevel];
-	const bool bNaniteEnabled = UseNanite(ShaderPlatform);
-	const bool bVirtualTextureEnabled = UseVirtualTexturing(ShaderPlatform);
-	const bool bWarmupNanite = CVarHLODWarmupNanite.GetValueOnGameThread() != 0;
-	const bool bWarmupVT = CVarHLODWarmupVT.GetValueOnGameThread() != 0;
-	const bool bWarmupNeeded = (bNaniteEnabled && bWarmupNanite) || (bVirtualTextureEnabled && bWarmupVT);
-	if (!bWarmupNeeded)
+	// If warmup is disabled by cvar/platform setting, no warmup is needed
+	if (!FHLODRuntimeSubsystem::bWarmupEnabled)
 	{
 		return false;
 	}
diff --git a/Engine/Source/Runtime/Engine/Public/WorldPartition/HLOD/HLODRuntimeSubsystem.h b/Engine/Source/Runtime/Engine/Public/WorldPartition/HLOD/HLODRuntimeSubsystem.h
--- a/Engine/Source/Runtime/Engine/Public/WorldPartition/HLOD/HLODRuntimeSubsystem.h
+++ b/Engine/Source/Runtime/Engine/Public/WorldPartition/HLOD/HLODRuntimeSubsystem.h
@@ -59,6 +59,7 @@ public:
 	ENGINE_API FWorldPartitionHLODForEachHLODObjectInCellEvent& GetForEachHLODObjectInCellEvent();
 
 	static ENGINE_API bool IsHLODEnabled();
+	static ENGINE_API bool IsWarmupEnabled();
 
 	ENGINE_API const TArray<IWorldPartitionHLODObject*>& GetHLODObjectsForCell(const UWorldPartitionRuntimeCell* InCell) const;
 	FWorldPartitionHLODObjectRegisteredEvent& OnHLODObjectRegisteredEvent() { return HLODObjectRegisteredEvent; }
