Making CanJumpWhileFalling a concept at Character level so character subclasses can override whether to allow the first jump to happen during some portion of their fall (i.e. coyote time seen in a lot of platformers).
Also updated existing implementation of coyote time to take advantage of this new facility instead of doing things the old way.

#tests made sure both jumping and coyote time still work as intended for impacted projects

diff --git a/Engine/Source/Runtime/Engine/Classes/GameFramework/Character.h b/Engine/Source/Runtime/Engine/Classes/GameFramework/Character.h
--- a/Engine/Source/Runtime/Engine/Classes/GameFramework/Character.h
+++ b/Engine/Source/Runtime/Engine/Classes/GameFramework/Character.h
@@ -986,6 +986,12 @@ public:
 	 */
 	ENGINE_API virtual float GetJumpMaxHoldTime() const;
 
+	/** 
+	 * By default this function returns false, which prevents jumping while in falling state (in many cases, we use this to consume the first jump), 
+     * but subclasses can override this behavior to allow jump to happen while falling under specific conditions they may care about 
+	 */
+	ENGINE_API virtual bool CanJumpWhileFalling() const;
+
 	UFUNCTION(Reliable, Client)
 	ENGINE_API void ClientCheatWalk();
 	ENGINE_API virtual void ClientCheatWalk_Implementation();
diff --git a/Engine/Source/Runtime/Engine/Private/Character.cpp b/Engine/Source/Runtime/Engine/Private/Character.cpp
--- a/Engine/Source/Runtime/Engine/Private/Character.cpp
+++ b/Engine/Source/Runtime/Engine/Private/Character.cpp
@@ -314,7 +314,9 @@ bool ACharacter::JumpIsAllowedInternal() const
 		// Ensure JumpHoldTime and JumpCount are valid.
 		if (!bWasJumping || GetJumpMaxHoldTime() <= 0.0f)
 		{
-			if (JumpCurrentCount == 0 && CharacterMovement->IsFalling())
+			// If this would be our first jump, but we're falling, normally we don't allow first jump to happen (we go right to second jump if available) 
+			// One exception is if our subclass overrides and allows first jump to still happen while falling under certain conditions 
+			if (JumpCurrentCount == 0 && CharacterMovement->IsFalling() && !CanJumpWhileFalling())
 			{
 				bJumpIsAllowed = JumpCurrentCount + 1 < JumpMaxCount;
 			}
@@ -1172,6 +1174,11 @@ void ACharacter::StopJumping()
 	ResetJumpState();
 }
 
+bool ACharacter::CanJumpWhileFalling() const
+{
+	return false;
+}
+
 void ACharacter::CheckJumpInput(float DeltaTime)
 {
 	JumpCurrentCountPreJump = JumpCurrentCount;
@@ -1180,10 +1187,10 @@ void ACharacter::CheckJumpInput(float DeltaTime)
 	{
 		if (bPressedJump)
 		{
-			// If this is the first jump and we're already falling,
-			// then increment the JumpCount to compensate.
+			// If this is the first jump and we're already falling, normally we deny that first jump from happening (by skipping right to the second jump, if available)
+			// One exception is if our subclass overrides this to allow first jump to still happen while falling under circumstances they care about
 			const bool bFirstJump = JumpCurrentCount == 0;
-			if (bFirstJump && CharacterMovement->IsFalling())
+			if (bFirstJump && CharacterMovement->IsFalling() && !CanJumpWhileFalling())
 			{
 				JumpCurrentCount++;
 			}
