Fixed UWorld::IsPlayInMobilePreview returning true incorrectly when previewing non-mobile platforms.

diff --git a/Engine/Source/Runtime/Engine/Private/World.cpp b/Engine/Source/Runtime/Engine/Private/World.cpp
--- a/Engine/Source/Runtime/Engine/Private/World.cpp
+++ b/Engine/Source/Runtime/Engine/Private/World.cpp
@@ -9205,8 +9205,13 @@ bool UWorld::IsPlayInMobilePreview() const
 {
 #if WITH_EDITOR
 	if (FPIEPreviewDeviceModule::IsRequestingPreviewDevice()
-		|| FParse::Param(FCommandLine::Get(), TEXT("featureleveles31"))
-		|| FDataDrivenShaderPlatformInfo::GetIsPreviewPlatform(Scene ? Scene->GetShaderPlatform(): GMaxRHIShaderPlatform))
+		|| FParse::Param(FCommandLine::Get(), TEXT("featureleveles31")))
+	{
+		return true;
+	}
+
+	const FStaticShaderPlatform ShaderPlatform = Scene ? Scene->GetShaderPlatform(): GMaxRHIShaderPlatform;
+	if (FDataDrivenShaderPlatformInfo::GetIsPreviewPlatform(ShaderPlatform) && IsMobilePlatform(ShaderPlatform))
 	{
 		return true;
 	}
