FSpline: Remove global TS delegate from spline implementation CVar, only convert internal data on FSpline::UpdateSpline calls

diff --git a/Engine/Source/Runtime/Engine/Classes/Curves/Spline.h b/Engine/Source/Runtime/Engine/Classes/Curves/Spline.h
--- a/Engine/Source/Runtime/Engine/Classes/Curves/Spline.h
+++ b/Engine/Source/Runtime/Engine/Classes/Curves/Spline.h
@@ -28,7 +28,7 @@ struct FSpline
 	ENGINE_API FSpline(const FSpline& Other);
 	ENGINE_API FSpline& operator=(const FSpline& Other);
 	ENGINE_API FSpline& operator=(const FSplineCurves& Other);
-	ENGINE_API ~FSpline();
+	ENGINE_API ~FSpline() = default;
 	
 	/* Control Point Index Interface */
 	
@@ -199,9 +199,7 @@ private:
 	bool WasLegacy() const { return PreviousImplementation == 1; }
 	bool WasNew() const { return PreviousImplementation == 2; }
 
-	/** Called when the implementation is changed at editor time due to a console command. */
 	void OnSplineImplementationChanged();
-	FDelegateHandle OnSplineImplementationChangedHandle;
 #endif
 
 public:
diff --git a/Engine/Source/Runtime/Engine/Private/Curves/Spline.cpp b/Engine/Source/Runtime/Engine/Private/Curves/Spline.cpp
--- a/Engine/Source/Runtime/Engine/Private/Curves/Spline.cpp
+++ b/Engine/Source/Runtime/Engine/Private/Curves/Spline.cpp
@@ -28,32 +28,12 @@ const FInterpCurveVector FSpline::ScaleCurve;
 
 namespace UE::Spline
 {
-
 int32 GImplementation = 0;
-#if WITH_EDITOR
-static FTSSimpleMulticastDelegate& GetOnSplineImplementationChanged()
-{
-	static FTSSimpleMulticastDelegate OnSplineImplementationChanged;
-	return OnSplineImplementationChanged;
-};
-static void SplineImplementationSink(IConsoleVariable*)
-{
-	static int32 GPreviousImplementation = -1;
-	if (GPreviousImplementation != GImplementation)
-	{
-		GImplementation = FMath::Clamp(GImplementation, 0, 2);
-		GetOnSplineImplementationChanged().Broadcast();
-		GPreviousImplementation = GImplementation;
-	}
-}
-#endif
 FAutoConsoleVariableRef CVarSplineImplementation(
 	TEXT("Spline.Implementation"),
 	GImplementation,
 	TEXT("0) Not Implemented - 1) Legacy Implementation - 2) New Implementation")
-#if WITH_EDITOR
-	,FConsoleVariableDelegate::CreateStatic(&SplineImplementationSink)
-#else
+#if !WITH_EDITOR
 	,ECVF_ReadOnly		// In non-editor builds, this CVar should exist so that GImplementation is properly set but it should not be changed by the user.
 #endif
 );
@@ -852,7 +832,6 @@ FSpline::FSpline()
 : CurrentImplementation(UE::Spline::GImplementation)
 , Version(0xffffffff)
 {
-	
 #if WITH_EDITOR
 	PreviousImplementation = 0;
 #endif
@@ -863,10 +842,6 @@ FSpline::FSpline()
 	case 1: LegacyData = MakeShared<FLegacySpline>(); break;
 	case 2: NewData = MakeShared<FNewSpline>(); break;
 	}
-
-#if WITH_EDITOR
-	OnSplineImplementationChangedHandle = UE::Spline::GetOnSplineImplementationChanged().AddRaw(this, &FSpline::OnSplineImplementationChanged);
-#endif
 }
 
 FSpline::FSpline(const FSpline& Other)
@@ -915,17 +890,6 @@ FSpline& FSpline::operator=(const FSplineCurves& Other)
 	return *this;
 }
 
-FSpline::~FSpline()
-{
-#if WITH_EDITOR
-	if (OnSplineImplementationChangedHandle.IsValid())
-	{
-		UE::Spline::GetOnSplineImplementationChanged().Remove(OnSplineImplementationChangedHandle);
-		OnSplineImplementationChangedHandle.Reset();
-	}
-#endif
-}
-
 void FSpline::AddPoint(const FSplinePoint& InPoint)
 {
 	switch (CurrentImplementation)
@@ -1723,6 +1687,13 @@ bool FSpline::IsClosedLoop() const
 
 void FSpline::UpdateSpline(const FUpdateSplineParams& InParams)
 {
+#if WITH_EDITOR
+	if (UE::Spline::GImplementation != CurrentImplementation)
+	{
+		OnSplineImplementationChanged();
+	}
+#endif
+		
 	switch (CurrentImplementation)
     {
     default: break;
