[SplineMesh] Support for bNeverNeedsCookedCollisionData on the spline mesh descriptor

#rnx

diff --git a/Engine/Source/Runtime/Engine/Classes/Engine/SplineMeshComponentDescriptor.h b/Engine/Source/Runtime/Engine/Classes/Engine/SplineMeshComponentDescriptor.h
--- a/Engine/Source/Runtime/Engine/Classes/Engine/SplineMeshComponentDescriptor.h
+++ b/Engine/Source/Runtime/Engine/Classes/Engine/SplineMeshComponentDescriptor.h
@@ -187,6 +187,9 @@ public:
 	UPROPERTY(EditAnywhere, Category = "Component Settings")
 	uint8 bFillCollisionUnderneathForNavmesh : 1;
 
+	UPROPERTY(EditAnywhere, Category = "Component Settings")
+	uint8 bNeverNeedsCookedCollisionData : 1;
+	
 	UPROPERTY(EditAnywhere, Category = "Component Settings")
 	int32 WorldPositionOffsetDisableDistance;
 
diff --git a/Engine/Source/Runtime/Engine/Private/SplineMeshComponentDescriptor.cpp b/Engine/Source/Runtime/Engine/Private/SplineMeshComponentDescriptor.cpp
--- a/Engine/Source/Runtime/Engine/Private/SplineMeshComponentDescriptor.cpp
+++ b/Engine/Source/Runtime/Engine/Private/SplineMeshComponentDescriptor.cpp
@@ -76,10 +76,12 @@ void FSplineMeshComponentDescriptorBase::InitFrom(const UStaticMeshComponent* Te
 	check(Template);
 	bEnableDiscardOnLoad = false;
 
+	const USplineMeshComponent* SplineTemplate = Cast<USplineMeshComponent>(Template);
+
 	// Disregard the template class if it does not stem from a spline mesh component
-	if (Template->IsA(USplineMeshComponent::StaticClass()))
+	if (SplineTemplate)
 	{
-		ComponentClass = Template->GetClass();
+		ComponentClass = SplineTemplate->GetClass();
 	}
 
 	Mobility = Template->Mobility;
@@ -124,6 +126,7 @@ void FSplineMeshComponentDescriptorBase::InitFrom(const UStaticMeshComponent* Te
 	bOverrideNavigationExport = Template->bOverrideNavigationExport;
 	bForceNavigationObstacle = Template->bForceNavigationObstacle;
 	bFillCollisionUnderneathForNavmesh = Template->bFillCollisionUnderneathForNavmesh;
+	bNeverNeedsCookedCollisionData = SplineTemplate ? SplineTemplate->GetbNeverNeedsCookedCollisionData() : false;
 
 #if WITH_EDITORONLY_DATA
 	HLODBatchingPolicy = Template->HLODBatchingPolicy;
@@ -224,6 +227,7 @@ bool FSplineMeshComponentDescriptorBase::operator==(const FSplineMeshComponentDe
 	WorldPositionOffsetDisableDistance == Other.WorldPositionOffsetDisableDistance &&
 	ShadowCacheInvalidationBehavior == Other.ShadowCacheInvalidationBehavior &&
 	DetailMode == Other.DetailMode &&
+	bNeverNeedsCookedCollisionData == Other.bNeverNeedsCookedCollisionData &&
 #if WITH_EDITORONLY_DATA
 	HLODBatchingPolicy == Other.HLODBatchingPolicy &&
 	bIncludeInHLOD == Other.bIncludeInHLOD &&
@@ -348,6 +352,7 @@ void FSplineMeshComponentDescriptorBase::InitComponent(USplineMeshComponent* Spl
 	SplineMeshComponent->WorldPositionOffsetDisableDistance = WorldPositionOffsetDisableDistance;
 	SplineMeshComponent->ShadowCacheInvalidationBehavior = ShadowCacheInvalidationBehavior;
 	SplineMeshComponent->DetailMode = DetailMode;
+	SplineMeshComponent->SetbNeverNeedsCookedCollisionData(bNeverNeedsCookedCollisionData);
 	
 #if WITH_EDITORONLY_DATA
 	SplineMeshComponent->HLODBatchingPolicy = HLODBatchingPolicy;
