[HLOD] Mark packages of source actors' levels as transient to prevent them from being dirtied during MoveExternalActorsToLevel.
Main temporary level is already marked as transient, but other levels are not.

diff --git a/Engine/Source/Runtime/Engine/Private/WorldPartition/HLOD/HLODSourceActorsFromCell.cpp b/Engine/Source/Runtime/Engine/Private/WorldPartition/HLOD/HLODSourceActorsFromCell.cpp
--- a/Engine/Source/Runtime/Engine/Private/WorldPartition/HLOD/HLODSourceActorsFromCell.cpp
+++ b/Engine/Source/Runtime/Engine/Private/WorldPartition/HLOD/HLODSourceActorsFromCell.cpp
@@ -81,6 +81,15 @@ bool UWorldPartitionHLODSourceActorsFromCell::LoadSourceActors(bool& bOutDirty,
 			if (Actor)
 			{
 				ActorsToResetLoaders.Add(Actor);
+
+				// Mark levels' packages as transient, so that they don't get dirtied during MoveExternalActorsToLevel
+				if (ULevel* ActorLevel = Actor->GetLevel())
+				{
+					if (UPackage* ActorLevelPackage = ActorLevel->GetPackage())
+					{
+						ActorLevelPackage->SetFlags(RF_Transient);
+					}
+				}
 			}
 		}
 		ResetLoaders(ActorsToResetLoaders);
