FObjectReplicator::UpdateGuidToReplicatorMap should specify the FNetDeltaSerializeInfo::Connection when gathering GUID references.  The NetDeltaSerializer could be using that Connection as part of its decision making process that determines which data to send.

#rnx

diff --git a/Engine/Source/Runtime/Engine/Private/DataReplication.cpp b/Engine/Source/Runtime/Engine/Private/DataReplication.cpp
--- a/Engine/Source/Runtime/Engine/Private/DataReplication.cpp
+++ b/Engine/Source/Runtime/Engine/Private/DataReplication.cpp
@@ -1507,6 +1507,7 @@ void FObjectReplicator::UpdateGuidToReplicatorMap()
 		FNetSerializeCB NetSerializeCB(Connection->Driver);
 
 		FNetDeltaSerializeInfo Parms;
+		Parms.Connection = Connection;
 		Parms.NetSerializeCB = &NetSerializeCB;
 		Parms.GatherGuidReferences = &LocalReferencedGuids;
 		Parms.TrackedGuidMemoryBytes = &LocalTrackedGuidMemoryBytes;
