PSO Precache Validation: added new cvar to enable conditional break on LogPSOMissInfo based on the shader hash

diff --git a/Engine/Source/Runtime/Engine/Private/PSOPrecacheValidation.cpp b/Engine/Source/Runtime/Engine/Private/PSOPrecacheValidation.cpp
--- a/Engine/Source/Runtime/Engine/Private/PSOPrecacheValidation.cpp
+++ b/Engine/Source/Runtime/Engine/Private/PSOPrecacheValidation.cpp
@@ -78,6 +78,14 @@ static FAutoConsoleVariableRef CVarPSOPrecachingBreakOnShaderHash(
 	ECVF_ReadOnly
 );
 
+FString GPSOPrecachingBreakValidationLogOnShaderHash;
+static FAutoConsoleVariableRef CVarPSOPrecachingBreakValidationLogOnShaderHash(
+	TEXT("r.PSOPrecache.BreakValidationLogOnShaderHash"),
+	GPSOPrecachingBreakValidationLogOnShaderHash,
+	TEXT("Debug break when PSO miss info logs the requested shader hash.\n"),
+	ECVF_ReadOnly
+);
+
 void ConditionalBreakOnPSOPrecacheMaterial(const FMaterial& Material, int32 PSOCollectorIndex)
 {
 	if (!GPSOPrecachingBreakOnMaterialName.IsEmpty())
@@ -126,6 +134,24 @@ void ConditionalBreakOnPSOPrecacheShader(const FRHIComputeShader* ComputeShader)
 	}
 }
 
+void ConditionalBreakOnLogPSOMissInfoShader(FRHIShader* RHIShader)
+{
+	if (!GPSOPrecachingBreakValidationLogOnShaderHash.IsEmpty())
+	{
+		FSHAHash ShaderHash;
+		ShaderHash.FromString(GPSOPrecachingBreakValidationLogOnShaderHash);
+		const auto MatchShaderHash = [&](FRHIShader* RHIShader)
+		{
+			return RHIShader ? RHIShader->GetHash() == ShaderHash : false;
+		};
+
+		if (MatchShaderHash(RHIShader))
+		{
+			UE_DEBUG_BREAK();
+		}
+	}
+}
+
 PSOCollectorStats::EPSOPrecacheValidationMode PSOCollectorStats::GetPrecachingValidationMode()
 {
 	switch (GPSOPrecachingValidationMode)
@@ -544,6 +570,8 @@ static void LogGeneralPSOMissInfo(
 			if (RHIShader)
 			{
 				StringBuilder.Appendf(TEXT("\n\t\t%s:\t\t%s"), ShaderTypeName, *(RHIShader->GetHash().ToString()));
+
+				ConditionalBreakOnLogPSOMissInfoShader(RHIShader);
 			}
 		};
 
