PR #13888: Fix assert in FVirtualTextureTranscodeCache::AddToList when there are more than 0x8000 tasks
#rnx

diff --git a/Engine/Source/Runtime/Engine/Private/VT/VirtualTextureTranscodeCache.h b/Engine/Source/Runtime/Engine/Private/VT/VirtualTextureTranscodeCache.h
--- a/Engine/Source/Runtime/Engine/Private/VT/VirtualTextureTranscodeCache.h
+++ b/Engine/Source/Runtime/Engine/Private/VT/VirtualTextureTranscodeCache.h
@@ -108,8 +108,8 @@ private:
 		uint32 FrameSubmitted;
 		uint16 Magic;
 		uint16 Hash;
-		int16 NextIndex;
-		int16 PrevIndex;
+		uint16 NextIndex;
+		uint16 PrevIndex;
 	};
 
 	void RemoveFromList(int32 Index)
