World Partition - Prevent running the "wp.Editor.DumpStreamingGenerationLog" console command while in PIE

diff --git a/Engine/Source/Runtime/Engine/Private/WorldPartition/WorldPartitionStreamingGeneration.cpp b/Engine/Source/Runtime/Engine/Private/WorldPartition/WorldPartitionStreamingGeneration.cpp
--- a/Engine/Source/Runtime/Engine/Private/WorldPartition/WorldPartitionStreamingGeneration.cpp
+++ b/Engine/Source/Runtime/Engine/Private/WorldPartition/WorldPartitionStreamingGeneration.cpp
@@ -62,7 +62,15 @@ static FAutoConsoleCommand DumpStreamingGenerationLog(
 	TEXT("Dump the streaming generation log."),
 	FConsoleCommandWithArgsDelegate::CreateLambda([](const TArray<FString>& Args)
 	{
-		if (UWorld* World = GEditor->GetEditorWorldContext().World(); World && !World->IsGameWorld())
+		if (!GEditor)
+		{
+			UE_LOG(LogWorldPartition, Error, TEXT("Streaming generation log can only be created in the editor"));
+		}
+		else if (GEditor->GetPIEWorldContext())
+		{
+			UE_LOG(LogWorldPartition, Error, TEXT("Streaming generation log can't be created during a PIE session"));
+		}
+		else if (UWorld* World = GEditor->GetEditorWorldContext().World())
 		{
 			if (UWorldPartition* WorldPartition = World->GetWorldPartition())
 			{
@@ -77,6 +85,10 @@ static FAutoConsoleCommand DumpStreamingGenerationLog(
 					UE_LOG(LogWorldPartition, Display, TEXT("Streaming generation details logged to '%s'"), *FPaths::ConvertRelativePathToFull(Context.OutputLogFilename.GetValue()));
 				}
 			}
+			else
+			{
+				UE_LOG(LogWorldPartition, Warning, TEXT("Current world is not a partitioned world, no streaming generation log created"));
+			}
 		}
 	})
 );
