Fix mesh paint texture borders leaking across edges.
Textures use wrap mode by default and this sets the border filling logic when cooking.
Mesh paint textures typically don't wrap, so clamp by default.

diff --git a/Engine/Source/Runtime/Engine/Private/VT/MeshPaintVirtualTexture.cpp b/Engine/Source/Runtime/Engine/Private/VT/MeshPaintVirtualTexture.cpp
--- a/Engine/Source/Runtime/Engine/Private/VT/MeshPaintVirtualTexture.cpp
+++ b/Engine/Source/Runtime/Engine/Private/VT/MeshPaintVirtualTexture.cpp
@@ -285,7 +285,11 @@ UMeshPaintVirtualTexture::UMeshPaintVirtualTexture(const FObjectInitializer& Obj
 	: UTexture2D(ObjectInitializer)
 {
 	VirtualTextureStreaming = true;
-	
+
+	// Mesh paint textures usually don't tile. We should default to clamping borders during texture build.
+	AddressX = TA_Clamp;
+	AddressY = TA_Clamp;
+
 #if WITH_EDITORONLY_DATA
 	CompressionNone = !MeshPaintVirtualTexture::UseCompressedTextureFormat();
 
