FNewSpline: Dirty legacy curves when changing loop flag

Changing the loop flag was working correctly but was not properly dirtying the source of data used by the visualizer to render the spline, so the closing segment was not drawn.

diff --git a/Engine/Source/Runtime/Engine/Private/Curves/Spline.cpp b/Engine/Source/Runtime/Engine/Private/Curves/Spline.cpp
--- a/Engine/Source/Runtime/Engine/Private/Curves/Spline.cpp
+++ b/Engine/Source/Runtime/Engine/Private/Curves/Spline.cpp
@@ -1981,6 +1981,13 @@ float FNewSpline::GetSplineLength()
 
 void FNewSpline::SetClosedLoop(bool bInClosedLoop)
 {
+	if (bInClosedLoop == GetSpline().IsClosedLoop())
+	{
+		return;
+	}
+
+	MarkLegacyCurvesDirty();
+	
     GetSpline().SetClosedLoop(bInClosedLoop);
 	
 	UE::Geometry::Spline::TTangentBezierSpline<FVector>* ScaleChild = GetTypedAttributeChannel<UE::Geometry::Spline::TTangentBezierSpline<FVector>>(ScaleAttrName);
