Fix performance regression in AssetCompilingManager

#rnx

diff --git a/Engine/Source/Runtime/Engine/Private/AssetCompilingManager.cpp b/Engine/Source/Runtime/Engine/Private/AssetCompilingManager.cpp
--- a/Engine/Source/Runtime/Engine/Private/AssetCompilingManager.cpp
+++ b/Engine/Source/Runtime/Engine/Private/AssetCompilingManager.cpp
@@ -218,12 +218,6 @@ public:
 	class MemorySharingInfo {
 	private:
 		const TFunction<int64 ()> GetAvailableMemory;
-		mutable FCriticalSection UpdateLock;
-
-		FShaderCompilerMemoryLimits SCWMemoryLimits;
-		FShaderCompileMemoryUsage SCWMemoryUsage;
-
-		int64 AvailableAssetCompilationMemory;
 	public:
 		int64 MemoryToLoan;
 		int64 MemoryLoanThreshold;
@@ -241,47 +235,40 @@ public:
 			IsMemorySharingEnabled = GConfig->GetBoolOrDefault(TEXT("AssetCompilation"), TEXT("SCWMemorySharingEnabled"), false, GEngineIni);
 		}	
 
-		int32 GetHighSCWMemory() const
+		int64 GetHighSCWMemory() const
 		{
-			FScopeLock MemLock(&UpdateLock);
+			FShaderCompilerMemoryLimits SCWMemoryLimits = GetShaderCompilerCurrentMemoryLimits();
 			return SCWMemoryLimits.BaseLimit + MemoryToLoan;
 		}
 
-		int32 GetNormalSCWMemory() const
+		int64 GetNormalSCWMemory() const
 		{
-			FScopeLock MemLock(&UpdateLock);
+			FShaderCompilerMemoryLimits SCWMemoryLimits = GetShaderCompilerCurrentMemoryLimits();		
 			return SCWMemoryLimits.BaseLimit;
 		}
 
 		bool IsSCWMemoryWithinLimit() const 
-		{ 
-			FScopeLock MemLock(&UpdateLock);
-			return SCWMemoryUsage.VirtualMemory <= (SCWMemoryLimits.CurrentLimit * 1024 * 1024);
+		{ 			
+			FShaderCompileMemoryUsage SCWMemoryUsage = GShaderCompilingManager->GetExternalMemoryUsage();
+			FShaderCompilerMemoryLimits SCWMemoryLimits = GetShaderCompilerCurrentMemoryLimits();		
+
+			return SCWMemoryUsage.VirtualMemory <= static_cast<uint64>(SCWMemoryLimits.CurrentLimit * 1024 * 1024);
 		}
 		bool IsMemoryAvailableToLoan() const
 		{
-			FScopeLock MemLock(&UpdateLock);
+			int64 AvailableAssetCompilationMemory = GetAvailableMemory() / (1024 * 1024);	
 			return AvailableAssetCompilationMemory >= MemoryLoanThreshold;
 		}
 		bool ShouldRetrieveFromSCW() const 
 		{
-			FScopeLock MemLock(&UpdateLock);
+			FShaderCompilerMemoryLimits SCWMemoryLimits = GetShaderCompilerCurrentMemoryLimits();		
 			return IsMemorySharingEnabled && IsMemoryAvailableToLoan() == false && SCWMemoryLimits.CurrentLimit == GetHighSCWMemory();
 		}
 		bool CanLoanToSCW() const 
 		{
-			FScopeLock MemLock(&UpdateLock);
+			FShaderCompilerMemoryLimits SCWMemoryLimits = GetShaderCompilerCurrentMemoryLimits();		
 			return IsMemorySharingEnabled && IsMemoryAvailableToLoan() && SCWMemoryLimits.CurrentLimit == GetNormalSCWMemory();
 		}
-
-		void Update()
-		{	
-			FScopeLock MemLock(&UpdateLock);			
-
-			GetShaderCompilerCurrentMemoryLimits(SCWMemoryLimits);								
-			SCWMemoryUsage = GShaderCompilingManager->GetExternalMemoryUsage();
-			AvailableAssetCompilationMemory = GetAvailableMemory() / (1024 * 1024);				
-		}
 	};
 
 	class FMemoryBoundScheduledWork : public FQueuedThreadPoolWrapper::FScheduledWork
@@ -297,9 +284,7 @@ public:
 		MemorySharingInfo& MemoryInfo;
 	protected:
 		void DoThreadedWork() override
-		{			
-			MemoryInfo.Update();					
-
+		{						
 			if (MemoryInfo.IsMemorySharingEnabled && FMemoryBoundQueuedThreadPoolWrapper::IsMemoryConstrained())
 			{
 				auto WaitForSCWMemory = [this]() 
@@ -316,8 +301,6 @@ public:
 					int Counter = 0;
 					for (Counter = 0; Counter < WaitIterations; Counter++) 
 					{
-						MemoryInfo.Update();
-
 						if (MemoryInfo.IsSCWMemoryWithinLimit())
 							break;
 
@@ -376,8 +359,7 @@ public:
 		TRACE_COUNTER_SET(AsyncCompilationTotalEstimatedMemory, TotalEstimatedMemory);
 		TRACE_COUNTER_SET(AsyncCompilationConcurrency, GetCurrentConcurrency());
 
-		FShaderCompilerMemoryLimits SCWMemoryLimits; 
-		GetShaderCompilerCurrentMemoryLimits(SCWMemoryLimits);
+		FShaderCompilerMemoryLimits SCWMemoryLimits = GetShaderCompilerCurrentMemoryLimits();
 		TRACE_COUNTER_SET(AsyncCompilationSCWCurrentLimit, SCWMemoryLimits.CurrentLimit * 1024 * 1024);
 		TRACE_COUNTER_SET(AsyncCompilationSCWBaseLimit, SCWMemoryLimits.BaseLimit * 1024 * 1024);
 	}
@@ -415,8 +397,7 @@ public:
 
 		UpdateCounters();
 		
-		MemoryInfo.Update();
-		if (MemoryInfo.CanLoanToSCW() && IsMemoryConstrained())
+		if (IsMemoryConstrained() && MemoryInfo.CanLoanToSCW())
 		{
 			SetShaderCompilerTargetMemoryLimitInMiB(MemoryInfo.GetHighSCWMemory());
 		} 		
diff --git a/Engine/Source/Runtime/Engine/Private/ShaderCompiler/ShaderCompilerThreadRunnable.cpp b/Engine/Source/Runtime/Engine/Private/ShaderCompiler/ShaderCompilerThreadRunnable.cpp
--- a/Engine/Source/Runtime/Engine/Private/ShaderCompiler/ShaderCompilerThreadRunnable.cpp
+++ b/Engine/Source/Runtime/Engine/Private/ShaderCompiler/ShaderCompilerThreadRunnable.cpp
@@ -60,17 +60,19 @@ static const TCHAR* GWorkerOutputFilename = TEXT("WorkerOutputOnly.out");
 
 TRACE_DECLARE_MEMORY_COUNTER(AsyncCompilationSCWMemory, TEXT("AsyncCompilation/SCWMemory"));
 
-void GetShaderCompilerCurrentMemoryLimits(FShaderCompilerMemoryLimits& Limits)
+FShaderCompilerMemoryLimits GetShaderCompilerCurrentMemoryLimits()
 {
-	Limits = {0};
+	FShaderCompilerMemoryLimits Limits = {0};
 
 	Limits.CurrentLimit = GShaderCompilerCurrentMemoryLimitInMiB;
 	Limits.BaseLimit = GShaderCompilerConfiguredMemoryLimitInMiB;
+
+	return Limits;
 }
 
-void SetShaderCompilerTargetMemoryLimitInMiB(int32 NewLimit)
+void SetShaderCompilerTargetMemoryLimitInMiB(int64 NewLimit)
 {
-	GShaderCompilerCurrentMemoryLimitInMiB = NewLimit;
+	GShaderCompilerCurrentMemoryLimitInMiB = static_cast<int32>(NewLimit);
 }
 
 /** Output structure for polling shader process memory violation status. See FShaderCompileThreadRunnable::QueryMemoryLimitViolationStatus(). */
diff --git a/Engine/Source/Runtime/Engine/Public/ShaderCompiler.h b/Engine/Source/Runtime/Engine/Public/ShaderCompiler.h
--- a/Engine/Source/Runtime/Engine/Public/ShaderCompiler.h
+++ b/Engine/Source/Runtime/Engine/Public/ShaderCompiler.h
@@ -46,7 +46,6 @@ class FShaderPipelineCompileJob;
 class UMaterialInterface;
 class FJsonObject;
 struct FAnalyticsEventAttribute;
-struct FShaderCompilerMemoryLimits;
 
 DECLARE_LOG_CATEGORY_EXTERN(LogShaderCompilers, Log, All);
 
@@ -56,9 +55,6 @@ bool CreateShadersOnLoad();
 extern ENGINE_API bool IsShaderJobCacheDDCEnabled();
 extern ENGINE_API bool IsMaterialMapDDCEnabled();
 
-extern ENGINE_API void GetShaderCompilerCurrentMemoryLimits(FShaderCompilerMemoryLimits& limits);
-extern ENGINE_API void SetShaderCompilerTargetMemoryLimitInMiB(int32 NewLimit);
-
 /**
  * Returns true if we should compile shaders that are only compiled in on demand shader compilation
  * modes and not stored in cooked shader maps. This can be useful for debug shaders.
@@ -875,10 +871,13 @@ private:
 
 struct FShaderCompilerMemoryLimits
 {
-	uint64 CurrentLimit;
-	uint64 BaseLimit;
+	int64 CurrentLimit;
+	int64 BaseLimit;
 };
 
+extern ENGINE_API FShaderCompilerMemoryLimits GetShaderCompilerCurrentMemoryLimits();
+extern ENGINE_API void SetShaderCompilerTargetMemoryLimitInMiB(int64 NewLimit);
+
 /**  
  * Manager of asynchronous and parallel shader compilation.
  * This class contains an interface to enqueue and retreive asynchronous shader jobs, and manages a FShaderCompileThreadRunnable.
