EngineNetTest
* Added utility function to teleport a client to a location

#rnx

diff --git a/Engine/Source/Runtime/Engine/Private/Tests/NetTestHelpers.cpp b/Engine/Source/Runtime/Engine/Private/Tests/NetTestHelpers.cpp
--- a/Engine/Source/Runtime/Engine/Private/Tests/NetTestHelpers.cpp
+++ b/Engine/Source/Runtime/Engine/Private/Tests/NetTestHelpers.cpp
@@ -525,6 +525,25 @@ void FTestWorlds::TickClientsAndDelay(uint32 NumFramesToDelay)
 	return nullptr;
 }
 
+void FTestWorlds::MovePlayerToLocation(uint32 ClientIndex, const FVector& Location)
+{
+	if (Clients.IsValidIndex(ClientIndex) == false)
+	{
+		ensureMsgf(false, TEXT("No client exists for that index"));
+		return;
+	}
+
+	APlayerController* ClientPC = GetServerPlayerControllerOfClient(ClientIndex);
+	if (!ClientPC)
+	{
+		ensureMsgf(false, TEXT("No PlayerController found for this client"));
+		return;
+	}
+
+	AActor* ViewTarget = ClientPC->GetViewTarget();
+	ViewTarget->SetActorLocation(Location);
+}
+
 UObject* FTestWorlds::FindReplicatedObjectOnClient(UObject* ServerObject, uint32 ClientIndex) const
 {
 	using namespace UE::Net;
diff --git a/Engine/Source/Runtime/Engine/Public/Tests/NetTestHelpers.h b/Engine/Source/Runtime/Engine/Public/Tests/NetTestHelpers.h
--- a/Engine/Source/Runtime/Engine/Public/Tests/NetTestHelpers.h
+++ b/Engine/Source/Runtime/Engine/Public/Tests/NetTestHelpers.h
@@ -234,6 +234,9 @@ struct FTestWorlds
 	/** Return the Server's player state corresponding to a specific client */
 	ENGINE_API APlayerController* GetServerPlayerControllerOfClient(uint32 ClientIndex);
 
+	/** Move a client's PlayerController to a location on the server */
+	ENGINE_API void MovePlayerToLocation(uint32 ClientIndex, const FVector& Location);
+
 	/**
 	 * Find the remote instance of a replicated object and ensures if it can't find it
 	 * @param ServerObject The server object that you want the remote version of.
