Do not try to unregister tickable game objects during engine shutdown as it may happen after the singleton is destroyed

diff --git a/Engine/Source/Runtime/Engine/Private/Tickable.cpp b/Engine/Source/Runtime/Engine/Private/Tickable.cpp
--- a/Engine/Source/Runtime/Engine/Private/Tickable.cpp
+++ b/Engine/Source/Runtime/Engine/Private/Tickable.cpp
@@ -147,8 +147,11 @@ FTickableGameObject::FTickableGameObject(ETickableTickType StartingTickType)
 FTickableGameObject::~FTickableGameObject()
 {
 	// This won't do anything if it was already set to never tick
-	// Depending on destruction order this could create a new statics object during shutdown, but the removal request will be ignored
-	SetTickableTickType(ETickableTickType::Never);
+	// Ignore on engine exit as the singleton may be in the middle of destruction
+	if (!IsEngineExitRequested())
+	{
+		SetTickableTickType(ETickableTickType::Never);
+	}
 }
 
 void FTickableGameObject::SetTickableTickType(ETickableTickType NewTickType)
