Substrate - use the same behavior as legacy: if not plugged in, the bottom cloar coat normal node will throw an error. This is compatible with Substrate which only copy the connection from this node to the new when running convertion.
Avoids shader compilation crash in editor and instead throw an error.

#tests PC editor
#rnx

diff --git a/Engine/Source/Runtime/Engine/Private/Materials/MaterialExpressions.cpp b/Engine/Source/Runtime/Engine/Private/Materials/MaterialExpressions.cpp
--- a/Engine/Source/Runtime/Engine/Private/Materials/MaterialExpressions.cpp
+++ b/Engine/Source/Runtime/Engine/Private/Materials/MaterialExpressions.cpp
@@ -20015,14 +20015,7 @@ int32  UMaterialExpressionClearCoatNormalCustomOutput::Compile(class FMaterialCo
 	{
 		return Compiler->CustomOutput(this, OutputIndex, Input.Compile(Compiler));
 	}
-	else
-	{
-		if (!Substrate::IsSubstrateEnabled())
-		{
-			return CompilerError(Compiler, TEXT("Input missing"));
-		}
-	}
-	return INDEX_NONE;
+	return CompilerError(Compiler, TEXT("Input missing"));
 }
 
 
