Change to fix an ensure and make navigation bounds get cached correctly before registration.

diff --git a/Engine/Source/Runtime/Engine/Classes/Components/HierarchicalInstancedStaticMeshComponent.h b/Engine/Source/Runtime/Engine/Classes/Components/HierarchicalInstancedStaticMeshComponent.h
--- a/Engine/Source/Runtime/Engine/Classes/Components/HierarchicalInstancedStaticMeshComponent.h
+++ b/Engine/Source/Runtime/Engine/Classes/Components/HierarchicalInstancedStaticMeshComponent.h
@@ -383,20 +383,6 @@ protected:
 
 	ENGINE_API static FBox GetClusterTreeBounds(TArray<FClusterNode> const& InClusterTree, const FVector& InOffset);
 
-	UE_DEPRECATED(5.4, "UHierarchicalInstancedStaticMeshComponent is no longer overriding UInstancedStaticMeshComponent::GetNavigationPerInstanceTransforms.")
-	ENGINE_API virtual void GetNavigationPerInstanceTransforms(const FBox& AreaBox, TArray<FTransform>& InstanceData) const override;
-	UE_DEPRECATED(5.4, "UHierarchicalInstancedStaticMeshComponent is no longer overriding UInstancedStaticMeshComponent::PartialNavigationUpdate.")
-	ENGINE_API virtual void PartialNavigationUpdate(int32 InstanceIdx) override;
-	UE_DEPRECATED(5.4, "UHierarchicalInstancedStaticMeshComponent is no longer overriding UInstancedStaticMeshComponent::SupportsPartialNavigationUpdate.")
-	virtual bool SupportsPartialNavigationUpdate() const override { return Super::SupportsPartialNavigationUpdate(); }
-	UE_DEPRECATED(5.4, "FlushAccumulatedNavigationUpdates has been deprecated. Use UInstancedStaticMeshComponent::PartialNavigationUpdates if you want to manually update the navigation data.")
-	ENGINE_API void FlushAccumulatedNavigationUpdates();
-
-	UE_DEPRECATED(5.4, "AccumulatedNavigationDirtyArea has been deprecated.")
-    mutable FBox AccumulatedNavigationDirtyArea;
-	UE_DEPRECATED(5.4, "AccumulatedNavigationDirtyAreas has been deprecated.")
-	mutable TArray<FBox> AccumulatedNavigationDirtyAreas;	//Remember to also remove the disable deprecation macro wrapper around ~UHierarchicalInstancedStaticMeshComponent() when removing the variable
-
 	FGraphEventArray BuildTreeAsyncTasks;
 	EHISMViewRelevanceType ViewRelevanceType = EHISMViewRelevanceType::HISM;
 
diff --git a/Engine/Source/Runtime/Engine/Classes/Components/InstancedStaticMeshComponent.h b/Engine/Source/Runtime/Engine/Classes/Components/InstancedStaticMeshComponent.h
--- a/Engine/Source/Runtime/Engine/Classes/Components/InstancedStaticMeshComponent.h
+++ b/Engine/Source/Runtime/Engine/Classes/Components/InstancedStaticMeshComponent.h
@@ -625,11 +625,16 @@ public:
 
 	UE_DEPRECATED(5.4, "This function has been added only for the purposes of moving LWI code outside of the engine. Don't use it, it will be removed soon.")
 	void OnPostPopulatePerInstanceData() { OnPostLoadPerInstanceData(); }
+
 	/** Request to navigation system to update only part of navmesh occupied by specified instance. */
-	ENGINE_API virtual void PartialNavigationUpdate(int32 InstanceIdx);
+	ENGINE_API virtual void PartialNavigationUpdate(int32 InstanceIdx, bool bUpdateCachedNavigationBounds);
+	UE_DEPRECATED(5.8, "Use the overload with the flag indicating whether the bounds need to be updated.")
+	ENGINE_API virtual void PartialNavigationUpdate(int32 InstanceIdx) final;
 
 	/** Request to navigation system to update only part of navmesh occupied specified instances transforms. */
-	ENGINE_API virtual void PartialNavigationUpdates(TConstArrayView<FTransform> InstanceTransforms);
+	ENGINE_API virtual void PartialNavigationUpdates(TConstArrayView<FTransform> InstanceTransforms, bool bUpdateCachedNavigationBounds);
+	UE_DEPRECATED(5.8, "Use the overload with the flag indicating whether the bounds need to be updated.")
+	ENGINE_API virtual void PartialNavigationUpdates(TConstArrayView<FTransform> InstanceTransforms) final;
 
 	ENGINE_API TSharedPtr<FISMCInstanceDataSceneProxy, ESPMode::ThreadSafe> GetOrCreateInstanceDataSceneProxy();
 	ENGINE_API const TSharedPtr<FISMCInstanceDataSceneProxy, ESPMode::ThreadSafe>& GetInstanceDataSceneProxy() const;
@@ -799,10 +804,6 @@ private:
 
 	UPROPERTY()
 	mutable FBoundsCacheElement CachedBounds[(int)EBoundsCacheType::Count];
-
-	/** Temporary flag used to skip Navigation CachedBounds invalidation.
-	 * This member was added to avoid deprecating the virtual PartialNavigationUpdates() (which can't be deprecated nicely). */ 
-	bool bNavigationCachedBoundsUpdated = false;
 };
 
 /** InstancedStaticMeshInstance hit proxy */
diff --git a/Engine/Source/Runtime/Engine/Private/HierarchicalInstancedStaticMesh.cpp b/Engine/Source/Runtime/Engine/Private/HierarchicalInstancedStaticMesh.cpp
--- a/Engine/Source/Runtime/Engine/Private/HierarchicalInstancedStaticMesh.cpp
+++ b/Engine/Source/Runtime/Engine/Private/HierarchicalInstancedStaticMesh.cpp
@@ -2073,8 +2073,6 @@ UHierarchicalInstancedStaticMeshComponent::UHierarchicalInstancedStaticMeshCompo
 	bUseAsOccluder = false;
 }
 
-//We deprecated a TArray and that's being referenced by the dtor as TArray has a non-trivial dtor which will trigger the deprecation warning.
-PRAGMA_DISABLE_DEPRECATION_WARNINGS
 UHierarchicalInstancedStaticMeshComponent::~UHierarchicalInstancedStaticMeshComponent()
 {
 	if (ProxySize)
@@ -2083,7 +2081,6 @@ UHierarchicalInstancedStaticMeshComponent::~UHierarchicalInstancedStaticMeshComp
 	}
 	ProxySize = 0;
 }
-PRAGMA_ENABLE_DEPRECATION_WARNINGS
 
 #if WITH_EDITOR
 
@@ -3269,19 +3266,6 @@ FVector UHierarchicalInstancedStaticMeshComponent::GetAverageScale() const
 }
 
 
-void UHierarchicalInstancedStaticMeshComponent::GetNavigationPerInstanceTransforms(const FBox& AreaBox, TArray<FTransform>& InstanceData) const
-{
-	Super::GetNavigationPerInstanceTransforms(AreaBox, InstanceData);
-}
-
-void UHierarchicalInstancedStaticMeshComponent::PartialNavigationUpdate(const int32 InstanceIdx)
-{
-	Super::PartialNavigationUpdate(InstanceIdx);
-}
-
-void UHierarchicalInstancedStaticMeshComponent::FlushAccumulatedNavigationUpdates()
-{}
-
 // recursive helper to gather all instances with locations inside the specified area. Supply a Filter to exclude leaf nodes based on the instance transform.
 static void GatherInstancesOverlappingArea(const UHierarchicalInstancedStaticMeshComponent& Component, const FBox& AreaBox, int32 Child, TFunctionRef<bool(const FMatrix&)> Filter, TArray<int32>& OutInstanceIndices)
 {
diff --git a/Engine/Source/Runtime/Engine/Private/InstancedStaticMesh.cpp b/Engine/Source/Runtime/Engine/Private/InstancedStaticMesh.cpp
--- a/Engine/Source/Runtime/Engine/Private/InstancedStaticMesh.cpp
+++ b/Engine/Source/Runtime/Engine/Private/InstancedStaticMesh.cpp
@@ -3677,8 +3677,8 @@ int32 UInstancedStaticMeshComponent::AddInstanceInternal(int32 InstanceIndex, FI
 	// Update navigation relevancy
 	bNavigationRelevant = IsNavigationRelevant();
 
-	// Indicates if the cached bounds were updated, it's expected to occur in PartialNavigationUpdate().
-	bNavigationCachedBoundsUpdated = false;
+	// Indicates if the cached bounds were updated.
+	bool bNavigationCachedBoundsUpdated = false;
 	
 	// Perform navigation update if the component is relevant to navigation and registered.
 	// Otherwise this will be handled in OnRegister (e.g. Editor manipulations).
@@ -3688,17 +3688,28 @@ int32 UInstancedStaticMeshComponent::AddInstanceInternal(int32 InstanceIndex, FI
 		// since it was skipped because IsNavigationRelevant() requires at least one instance.
 		if (GetInstanceCount() == 1)
 		{
+			CalcAndCacheNavigationBounds();
+			bNavigationCachedBoundsUpdated = true;
+
 			FNavigationSystem::RegisterComponent(*this);
 		}
 
 		if (SupportsPartialNavigationUpdate())
 		{
-			PartialNavigationUpdate(InstanceIndex);
+			PartialNavigationUpdate(InstanceIndex,  /*bUpdateCachedNavigationBounds*/ !bNavigationCachedBoundsUpdated);
 		}
 		else
 		{
+			if (!bNavigationCachedBoundsUpdated)
+			{
+				CalcAndCacheNavigationBounds();
+			}
+
 			FullNavigationUpdate();
 		}
+
+		// At this point navigation bound were calculated and cached
+		bNavigationCachedBoundsUpdated = true;
 	}
 
 	if (FInstancedStaticMeshDelegates::OnInstanceIndexUpdated.IsBound())
@@ -3754,6 +3765,9 @@ TArray<int32> UInstancedStaticMeshComponent::AddInstancesInternal(TConstArrayVie
 	// and IsNavigationRelevant() takes the number of instance into account which can be 0 here.
 	bNavigationRelevant = NumInstances > 0 && Super::IsNavigationRelevant();
 
+	// Indicates if the cached bounds were updated.
+	bool bNavigationCachedBoundsUpdated = false;
+
 	// Navigation update required if explicitly requested and if the component is relevant to navigation and registered.
 	// Otherwise this will be handled in OnRegister (e.g. Editor manipulations).
 	const bool bNavigationUpdateRequired = bUpdateNavigation && bNavigationRelevant && IsRegistered();
@@ -3765,6 +3779,8 @@ TArray<int32> UInstancedStaticMeshComponent::AddInstancesInternal(TConstArrayVie
 		NavigationUpdateTransforms.Reserve(InstanceTransforms.Num());
 	}
 	
+	bool bNeedsNavigationRegistration = false;
+
 	for (const FTransform& InstanceTransform : InstanceTransforms)
 	{
 		FInstancedStaticMeshInstanceData& NewInstanceData = PerInstanceSMData.AddDefaulted_GetRef();
@@ -3784,8 +3800,7 @@ TArray<int32> UInstancedStaticMeshComponent::AddInstancesInternal(TConstArrayVie
 			// since it was skipped because IsNavigationRelevant() requires at least one instance.
 			if (GetInstanceCount() == 1)
 			{
-				CalcAndCacheNavigationBounds();
-				FNavigationSystem::RegisterComponent(*this);
+				bNeedsNavigationRegistration = true;
 			}
 
 			if (bPartialNavigationUpdateRequired)
@@ -3805,17 +3820,41 @@ TArray<int32> UInstancedStaticMeshComponent::AddInstancesInternal(TConstArrayVie
 
 	if (bNavigationUpdateRequired)
 	{
+		if (bNeedsNavigationRegistration)
+		{
+			CalcAndCacheNavigationBounds();
+			bNavigationCachedBoundsUpdated = true;
+
+			FNavigationSystem::RegisterComponent(*this);
+		}
+
 		if (bPartialNavigationUpdateRequired)
 		{
-			PartialNavigationUpdates(NavigationUpdateTransforms);
+			PartialNavigationUpdates(NavigationUpdateTransforms,  /*bUpdateCachedNavigationBounds*/ !bNavigationCachedBoundsUpdated);
 		}
 		else
 		{
+			if (!bNavigationCachedBoundsUpdated)
+			{
+				CalcAndCacheNavigationBounds();
+			}
+
 			FullNavigationUpdate();
 		}
+
+		// At this point navigation bound were calculated and cached
+		bNavigationCachedBoundsUpdated = true;
 	}
 
+	const FBoundsCacheElement NavigationCacheElement = CachedBounds[(int32)EBoundsCacheType::Navigation];
 	InvalidateCachedBounds();
+
+	// Restore Navigation CachedBounds (the invalidation is not needed if it was updated).
+	if (bNavigationCachedBoundsUpdated)
+	{
+		CachedBounds[(int32)EBoundsCacheType::Navigation] = NavigationCacheElement;
+	}
+
 	return NewInstanceIndices;
 }
 
@@ -4043,7 +4082,7 @@ bool UInstancedStaticMeshComponent::RemoveInstanceInternal(int32 InstanceIndex,
 		// the instance's transform to dirty the area and full update will dirty the whole covered by the navigation bounds.
 		if (bNavigationUpdateRequired && SupportsPartialNavigationUpdate())
 		{
-			PartialNavigationUpdate(InstanceIndex);
+			PartialNavigationUpdate(InstanceIndex, /*bUpdateCachedNavigationBounds*/ true );
 		}
 
 		if (bUseRemoveAtSwap)
@@ -4308,7 +4347,7 @@ void UInstancedStaticMeshComponent::UpdateComponentTransform(EUpdateTransformFla
 		// Only handle partial updates since modifying the transform is handled by base class (i.e. SceneComponent).
 		if (bDoPartialNavigationUpdate)
 		{
-			PartialNavigationUpdates(NavigationUpdateTransforms);
+			PartialNavigationUpdates(NavigationUpdateTransforms, /*bUpdateCachedNavigationBounds*/ true );
 		}
 	}
 
@@ -4416,7 +4455,7 @@ bool UInstancedStaticMeshComponent::UpdateInstanceTransform(const int32 Instance
 				{
 					PreviousTransform * GetComponentTransform(),
 					(bWorldSpace ? NewInstanceTransform : NewInstanceTransform * GetComponentTransform())
-				});
+				}, /*bUpdateCachedNavigationBounds*/ true );
 		}
 		else
 		{
@@ -4502,7 +4541,7 @@ bool UInstancedStaticMeshComponent::BatchUpdateInstancesTransforms(int32 StartIn
 	{
 		if (bPartialNavigationUpdateRequired)
 		{
-			PartialNavigationUpdates(NavigationUpdateTransforms);
+			PartialNavigationUpdates(NavigationUpdateTransforms, /*bUpdateCachedNavigationBounds*/ true);
 		}
 		else
 		{
@@ -4589,7 +4628,7 @@ bool UInstancedStaticMeshComponent::BatchUpdateInstancesTransformsInternal(int32
 	{
 		if (bPartialNavigationUpdateRequired)
 		{
-			PartialNavigationUpdates(NavigationUpdateTransforms);
+			PartialNavigationUpdates(NavigationUpdateTransforms, /*bUpdateCachedNavigationBounds*/ true);
 		}
 		else
 		{
@@ -4660,7 +4699,7 @@ bool UInstancedStaticMeshComponent::BatchUpdateInstancesTransform(int32 StartIns
 	{
 		if (bPartialNavigationUpdateRequired)
 		{
-			PartialNavigationUpdates(NavigationUpdateTransforms);
+			PartialNavigationUpdates(NavigationUpdateTransforms, /*bUpdateCachedNavigationBounds*/ true);
 		}
 		else
 		{
@@ -4726,7 +4765,7 @@ bool UInstancedStaticMeshComponent::BatchUpdateInstancesData(int32 StartInstance
 	{
 		if (bPartialNavigationUpdateRequired)
 		{
-			PartialNavigationUpdates(NavigationUpdateTransforms);
+			PartialNavigationUpdates(NavigationUpdateTransforms, /*bUpdateCachedNavigationBounds*/ true);
 		}
 		else
 		{
@@ -5444,10 +5483,10 @@ void UInstancedStaticMeshComponent::PartialNavigateUpdateForCurrentInstances()
 		NavigationUpdateTransforms.Add(FTransform(PerInstanceSMData[i].Transform) * GetComponentTransform());
 	}
 
-	PartialNavigationUpdates(NavigationUpdateTransforms);
+	PartialNavigationUpdates(NavigationUpdateTransforms, /*bUpdateCachedNavigationBounds*/ true);
 }
 
-void UInstancedStaticMeshComponent::PartialNavigationUpdate(int32 InstanceIdx)
+void UInstancedStaticMeshComponent::PartialNavigationUpdate(int32 InstanceIdx, const bool bUpdateCachedNavigationBounds)
 {
 	TRACE_CPUPROFILER_EVENT_SCOPE(UInstancedStaticMeshComponent::PartialNavigationUpdate);
 	if (!IsNavigationRelevant())
@@ -5458,11 +5497,16 @@ void UInstancedStaticMeshComponent::PartialNavigationUpdate(int32 InstanceIdx)
 	FTransform InstanceTransform;
 	if (GetInstanceTransform(InstanceIdx, InstanceTransform, /*bWorldSpace*/true))
 	{
-		PartialNavigationUpdates({InstanceTransform});
+		PartialNavigationUpdates({InstanceTransform}, bUpdateCachedNavigationBounds);
 	}
 }
 
-void UInstancedStaticMeshComponent::PartialNavigationUpdates(const TConstArrayView<FTransform> InstanceTransforms)
+void UInstancedStaticMeshComponent::PartialNavigationUpdate(int32 InstanceIdx)
+{
+	PartialNavigationUpdate(InstanceIdx, /*bUpdateCachedNavigationBounds*/ true);
+}
+
+void UInstancedStaticMeshComponent::PartialNavigationUpdates(const TConstArrayView<FTransform> InstanceTransforms, const bool bUpdateCachedNavigationBounds)
 {
 	if (!IsNavigationRelevant() || InstanceTransforms.IsEmpty())
 	{
@@ -5472,15 +5516,17 @@ void UInstancedStaticMeshComponent::PartialNavigationUpdates(const TConstArrayVi
 	const FBox InstanceBounds = GetInstanceNavigationBounds();
 	if (InstanceBounds.IsValid)
 	{
-		if (CVarISMInvalidateCachedNavigationBoundsOnPartialNavigationUpdates.GetValueOnAnyThread() != 0)
+		if( bUpdateCachedNavigationBounds )
 		{
-			// Make sure to invalidate cached navigation bounds that might have been computed from a different number of instances.
-			CachedBounds[static_cast<int32>(EBoundsCacheType::Navigation)].Invalidate();
-		}
+			if (CVarISMInvalidateCachedNavigationBoundsOnPartialNavigationUpdates.GetValueOnAnyThread() != 0)
+			{
+				// Make sure to invalidate cached navigation bounds that might have been computed from a different number of instances.
+				CachedBounds[static_cast<int32>(EBoundsCacheType::Navigation)].Invalidate();
+			}
 
-		// Update cached navigation bounds
-		CalcAndCacheNavigationBounds();
-		bNavigationCachedBoundsUpdated = true;
+			// Update cached navigation bounds
+			CalcAndCacheNavigationBounds();
+		}
 
 		// Dirty areas around the modified instances
 		TArray<FBox> DirtyAreas;
@@ -5493,6 +5539,12 @@ void UInstancedStaticMeshComponent::PartialNavigationUpdates(const TConstArrayVi
 	}
 }
 
+void UInstancedStaticMeshComponent::PartialNavigationUpdates(const TConstArrayView<FTransform> InstanceTransforms)
+{
+	PartialNavigationUpdates(InstanceTransforms, /*bUpdateCachedNavigationBounds*/ true);
+}
+
+
 bool UInstancedStaticMeshComponent::DoCustomNavigableGeometryExport(FNavigableGeometryExport& GeomExport) const
 {
 	return FInstancedStaticMeshComponentHelper::DoCustomNavigableGeometryExport(*this, GeomExport, FNavDataPerInstanceTransformDelegate::CreateUObject(this, &UInstancedStaticMeshComponent::GetNavigationPerInstanceTransforms));
