Dstm: Fix bliping physics object issue. Set NoPawnPlayerController to correct location when it is created during pawn migration. Then relevancy check will use the correct position

diff --git a/Engine/Source/Runtime/Engine/Private/PlayerController.cpp b/Engine/Source/Runtime/Engine/Private/PlayerController.cpp
--- a/Engine/Source/Runtime/Engine/Private/PlayerController.cpp
+++ b/Engine/Source/Runtime/Engine/Private/PlayerController.cpp
@@ -15,6 +15,7 @@
 #include "EngineStats.h"
 #include "Engine/Canvas.h"
 #include "GameFramework/PlayerStart.h"
+#include "GameFramework/Character.h"
 #include "Components/SkeletalMeshComponent.h"
 #include "Components/AudioComponent.h"
 #include "Components/ForceFeedbackComponent.h"
@@ -5025,6 +5026,41 @@ void APlayerController::GetLifetimeReplicatedProps(TArray< FLifetimeProperty > &
 }
 
 #if UE_WITH_REMOTE_OBJECT_HANDLE
+static void UpdateNetConnectionControllerPosition(UNetConnection* NetConnection)
+{
+	AActor* ControllerActor = Cast<AActor>(NetConnection->PlayerController);
+	bool bNoPawnPlayerControllerPositionSet = false;
+	if (ControllerActor)
+	{
+		if (APlayerController* PlayerController = Cast<APlayerController>(NetConnection->ViewTarget))
+		{
+			FVector TrueViewLocation;
+			FRotator ViewRotation;
+			PlayerController->GetPlayerViewPoint(TrueViewLocation, ViewRotation);
+			ControllerActor->SetActorLocation(TrueViewLocation);
+			UE_LOG(LogPlayerController, Verbose, TEXT("[%s] UpdateNetConnectionControllerPosition: Setting NoPawnPlayerController %s position to %s using position from PlayerController %s"),
+				*FRemoteServerId::GetLocalServerId().ToString(), 
+				*ControllerActor->GetName(), *(TrueViewLocation.ToString()), *PlayerController->GetName());
+			bNoPawnPlayerControllerPositionSet = true;
+		}
+		else if (ACharacter* PlayerCharacter = Cast<ACharacter>(NetConnection->ViewTarget))
+		{
+			FVector TrueViewLocation = PlayerCharacter->GetActorLocation();
+			ControllerActor->SetActorLocation(TrueViewLocation);
+			UE_LOG(LogPlayerController, Verbose, TEXT("[%s] UpdateNetConnectionControllerPosition: Setting NoPawnPlayerController %s position to %s using position from PlayerCharacter %s"),
+				*FRemoteServerId::GetLocalServerId().ToString(),
+				*ControllerActor->GetName(), *(TrueViewLocation.ToString()), *PlayerCharacter->GetName());
+			bNoPawnPlayerControllerPositionSet = true;
+		}
+	}
+	if (!bNoPawnPlayerControllerPositionSet)
+	{
+		UE_LOG(LogPlayerController, Warning, TEXT("[%s] UpdateNetConnectionControllerPosition: Failed to set position of NoPawnPlayerController %s"),
+			*FRemoteServerId::GetLocalServerId().ToString(),
+			*GetNameSafe(NetConnection->PlayerController));
+	}
+}
+
 void APlayerController::PostMigrate(const FUObjectMigrationContext& MigrationContext)
 {
 	// The ANoPawnPlayerController is server specific and should never be migrated between servers.
@@ -5165,7 +5201,10 @@ void APlayerController::PostMigrate(const FUObjectMigrationContext& MigrationCon
 				PrevNetConnection->PlayerController->SetClientHandshakeId(PrevNetConnection->GetClientHandshakeId());
 
 				PrevNetConnection->OwningActor = PrevNetConnection->PlayerController;
-					
+				
+				// Update the position of NoPawnPlayerController during migration, so it will be correct when used for relevancy check. 
+				UpdateNetConnectionControllerPosition(PrevNetConnection);
+
 				UE_LOG(LogPlayerController, Verbose, TEXT("[%s] APlayerController::PostMigrate [%s]: Replacing %s with new player controller %s on connection %s"), 
 					*FRemoteServerId::GetLocalServerId().ToString(), *FRemoteObjectId(this).ToString(), 
 					*GetName(), *PrevNetConnection->PlayerController->GetName(), *PrevNetConnection->GetName());
