[VisualLogger] introducing new defines to split functionalities: Recording, Rendering and Visualizer Tool.

UE_DEBUG_RECORDING_ENABLED:  Indicates whether the UE_DEBUG_* macros will record the element (VisualLogger and/or insight traces)
UE_DEBUG_RECORDING_USING_VLOG:  Indicates whether the UE_DEBUG_* macros will use the VisualLogger recording process to store entries and write to all registered output devices. Otherwise only insight traces are used.

UE_DEBUG_RENDERING_ENABLED:  Indicates whether the UE_DEBUG_* macros will use immediate debug draw visualization

UE_DEBUG_VISUALIZER_TOOL_ENABLED:  Indicates whether the analysis of the recorded data is active and dedicated UI and tools should be compiled */

To maintain compatibility ENABLE_VISUAL_LOG will enable all three functionalities.

Without UE_DEBUG_RECORDING_USING_VLOG we bypass most VisualLogger functionalities to log to the trace directly.
Examples of disabled features:
- CategoryAllowList
- NavigationDataDump
- DebugSnapshot
- ClassAllowList
- ObjectAllowList
- Output devices

diff --git a/Engine/Plugins/Animation/GameplayInsights/Source/RewindDebuggerVLogRuntime/Private/RewindDebuggerVLogRuntime.h b/Engine/Plugins/Animation/GameplayInsights/Source/RewindDebuggerVLogRuntime/Private/RewindDebuggerVLogRuntime.h
--- a/Engine/Plugins/Animation/GameplayInsights/Source/RewindDebuggerVLogRuntime/Private/RewindDebuggerVLogRuntime.h
+++ b/Engine/Plugins/Animation/GameplayInsights/Source/RewindDebuggerVLogRuntime/Private/RewindDebuggerVLogRuntime.h
@@ -2,11 +2,9 @@
 
 #pragma once
 
-#if UE_ENABLE_INCLUDE_ORDER_DEPRECATED_IN_5_8
-#include "CoreMinimal.h"
-#endif // UE_ENABLE_INCLUDE_ORDER_DEPRECATED_IN_5_8
+#include "VisualLogger/VisualLoggerDefines.h"
 
-#include "EngineDefines.h"
+#if UE_DEBUG_RECORDING_ENABLED
 #include "RewindDebuggerRuntimeInterface/IRewindDebuggerRuntimeExtension.h"
 #include "VisualLogger/VisualLogger.h"
 
@@ -14,7 +12,6 @@ namespace RewindDebugger
 {
 	class FRewindDebuggerVLogRuntime : public IRewindDebuggerRuntimeExtension
 	{
-#if ENABLE_VISUAL_LOG
 	public:
 		virtual void RecordingStarted() override
 		{
@@ -35,6 +32,6 @@ namespace RewindDebugger
 
 	private:
 		bool bWasRecording = false;
-#endif // ENABLE_VISUAL
 	};
 }
+#endif // UE_DEBUG_RECORDING_ENABLED
diff --git a/Engine/Plugins/Animation/GameplayInsights/Source/RewindDebuggerVLogRuntime/Private/RewindDebuggerVLogRuntimeModule.cpp b/Engine/Plugins/Animation/GameplayInsights/Source/RewindDebuggerVLogRuntime/Private/RewindDebuggerVLogRuntimeModule.cpp
--- a/Engine/Plugins/Animation/GameplayInsights/Source/RewindDebuggerVLogRuntime/Private/RewindDebuggerVLogRuntimeModule.cpp
+++ b/Engine/Plugins/Animation/GameplayInsights/Source/RewindDebuggerVLogRuntime/Private/RewindDebuggerVLogRuntimeModule.cpp
@@ -2,6 +2,8 @@
 
 #include "RewindDebuggerVLogRuntimeModule.h"
 #include "Modules/ModuleManager.h"
+
+#if UE_DEBUG_RECORDING_ENABLED
 #include "Features/IModularFeatures.h"
 
 void FRewindDebuggerVLogRuntimeModule::StartupModule()
@@ -13,5 +15,6 @@ void FRewindDebuggerVLogRuntimeModule::ShutdownModule()
 {
 	IModularFeatures::Get().UnregisterModularFeature(IRewindDebuggerRuntimeExtension::ModularFeatureName, &VLogExtension);
 }
+#endif // UE_DEBUG_RECORDING_ENABLED
 
 IMPLEMENT_MODULE(FRewindDebuggerVLogRuntimeModule, RewindDebuggerVLogRuntime);
diff --git a/Engine/Plugins/Animation/GameplayInsights/Source/RewindDebuggerVLogRuntime/Private/RewindDebuggerVLogRuntimeModule.h b/Engine/Plugins/Animation/GameplayInsights/Source/RewindDebuggerVLogRuntime/Private/RewindDebuggerVLogRuntimeModule.h
--- a/Engine/Plugins/Animation/GameplayInsights/Source/RewindDebuggerVLogRuntime/Private/RewindDebuggerVLogRuntimeModule.h
+++ b/Engine/Plugins/Animation/GameplayInsights/Source/RewindDebuggerVLogRuntime/Private/RewindDebuggerVLogRuntimeModule.h
@@ -2,17 +2,22 @@
 
 #pragma once
 
-#include "CoreMinimal.h"
+#include "VisualLogger/VisualLoggerDefines.h"
 #include "Modules/ModuleInterface.h"
+
+#if UE_DEBUG_RECORDING_ENABLED
 #include "RewindDebuggerVLogRuntime.h"
+#endif // UE_DEBUG_RECORDING_ENABLED
 
 class FRewindDebuggerVLogRuntimeModule : public IModuleInterface
 {
+#if UE_DEBUG_RECORDING_ENABLED
 public:
-	// IModuleInterface interface
+	//~ IModuleInterface interface
 	virtual void StartupModule() override;
 	virtual void ShutdownModule() override;
-	
+
 private:
 	RewindDebugger::FRewindDebuggerVLogRuntime VLogExtension;
+#endif // UE_DEBUG_RECORDING_ENABLED
 };
diff --git a/Engine/Plugins/Experimental/Mover/Source/Mover/Private/DefaultMovementSet/Modes/AsyncNavWalkingMode.cpp b/Engine/Plugins/Experimental/Mover/Source/Mover/Private/DefaultMovementSet/Modes/AsyncNavWalkingMode.cpp
--- a/Engine/Plugins/Experimental/Mover/Source/Mover/Private/DefaultMovementSet/Modes/AsyncNavWalkingMode.cpp
+++ b/Engine/Plugins/Experimental/Mover/Source/Mover/Private/DefaultMovementSet/Modes/AsyncNavWalkingMode.cpp
@@ -23,7 +23,7 @@
 
 #include UE_INLINE_GENERATED_CPP_BY_NAME(AsyncNavWalkingMode)
 
-#if ENABLE_VISUAL_LOG
+#if UE_DEBUG_RECORDING_ENABLED
 static const FName AsyncWalkModeLogCategory = FName("AsyncNavWalkingMode");
 #endif
 
@@ -501,7 +501,7 @@ void UAsyncNavWalkingMode::FindBestNavMeshLocation(const FVector& TraceStart, co
 		OutHitResult.bBlockingHit = true;
 	}
 
-#if ENABLE_VISUAL_LOG
+#if UE_DEBUG_RECORDING_ENABLED
 	const float DebugNormalLength = 50.0f;
 	UE_VLOG_SEGMENT(GetMoverComponent(), AsyncWalkModeLogCategory, Display, TraceStart, TraceEnd, FColor::Green, TEXT(""));
 	for (int i = 0; i < MultiTraceHits.Num(); i++)
@@ -509,7 +509,7 @@ void UAsyncNavWalkingMode::FindBestNavMeshLocation(const FVector& TraceStart, co
 		UE_VLOG_SPHERE(GetMoverComponent(), "AsyncNavWalkingMode", Display, MultiTraceHits[i].ImpactPoint, i == 0 ? 5.0f : 2.5f, i == 0 ? FColor::Red : FColor::Yellow, TEXT("%i"), i);
 		UE_VLOG_ARROW(GetMoverComponent(), "AsyncNavWalkingMode", Display, MultiTraceHits[i].ImpactPoint, (MultiTraceHits[i].ImpactPoint + DebugNormalLength * MultiTraceHits[i].ImpactNormal), i == 0 ? FColor::Red : FColor::Yellow, TEXT(""));
 	}
-#endif // ENABLE_VISUAL_LOG
+#endif // UE_DEBUG_RECORDING_ENABLED
 
 	if (AsyncWalkingModeCVars::bUseNavMeshNormal && CachedNavLocation.HasNodeRef() && MultiTraceHits.Num() > 0)
 	{
diff --git a/Engine/Source/Developer/LogVisualizer/Private/LogVisualizer.cpp b/Engine/Source/Developer/LogVisualizer/Private/LogVisualizer.cpp
--- a/Engine/Source/Developer/LogVisualizer/Private/LogVisualizer.cpp
+++ b/Engine/Source/Developer/LogVisualizer/Private/LogVisualizer.cpp
@@ -1,15 +1,18 @@
 // Copyright Epic Games, Inc. All Rights Reserved.
 
 #include "LogVisualizer.h"
-#include "LogVisualizerPublic.h"
-#include "GameFramework/SpectatorPawn.h"
+
+#if UE_DEBUG_VISUALIZER_TOOL_ENABLED
 #include "Engine/GameViewportClient.h"
 #include "Engine/LocalPlayer.h"
 #include "EngineUtils.h"
+#include "GameFramework/SpectatorPawn.h"
+#include "LogVisualizerPublic.h"
 #include "LogVisualizerSettings.h"
-#include "VisualLoggerDatabase.h"
-#include "VisualLoggerCameraController.h"
 #include "VisualLogger/VisualLogger.h"
+#include "VisualLoggerCameraController.h"
+#include "VisualLoggerDatabase.h"
+#include "VisualLoggerTimeSliderController.h"
 
 #if WITH_EDITOR
 #include "Editor/EditorEngine.h"
@@ -80,7 +83,10 @@ void FLogVisualizer::Shutdown()
 void FLogVisualizer::Reset()
 {
 	TimeSliderController->SetTimesliderArgs(FVisualLoggerTimeSliderArgs());
+
+#if UE_DEBUG_RECORDING_USING_VLOG
 	FVisualLogger::Get().OnDataReset();
+#endif
 }
 
 FLogVisualizer& FLogVisualizer::Get()
@@ -361,4 +367,5 @@ void FLogVisualizer::GotoLastItemAnyRow()
 void FLogVisualizer::SeekToTime(float Time)
 {
 	GetTimeSliderController()->CommitScrubPosition(Time,  /*bIsScrubbing=*/true);
-}
\ No newline at end of file
+}
+#endif // UE_DEBUG_VISUALIZER_TOOL_ENABLED
\ No newline at end of file
diff --git a/Engine/Source/Developer/LogVisualizer/Private/LogVisualizer.h b/Engine/Source/Developer/LogVisualizer/Private/LogVisualizer.h
--- a/Engine/Source/Developer/LogVisualizer/Private/LogVisualizer.h
+++ b/Engine/Source/Developer/LogVisualizer/Private/LogVisualizer.h
@@ -1,6 +1,14 @@
 // Copyright Epic Games, Inc. All Rights Reserved.
 #pragma once
 
+
+#if UE_ENABLE_INCLUDE_ORDER_DEPRECATED_IN_5_8
+#include "EngineDefines.h"
+#endif // UE_ENABLE_INCLUDE_ORDER_DEPRECATED_IN_5_8
+#include "VisualLogger/VisualLoggerDefines.h"
+
+#if UE_DEBUG_VISUALIZER_TOOL_ENABLED
+
 /* Dependencies
 *****************************************************************************/
 
@@ -32,3 +40,4 @@
 #include "SVisualLoggerTimelineBar.h"
 #include "MovieScene.h"
 #include "VisualLoggerTimeSliderController.h"
+#endif // UE_DEBUG_VISUALIZER_TOOL_ENABLED
diff --git a/Engine/Source/Developer/LogVisualizer/Private/LogVisualizerSettings.cpp b/Engine/Source/Developer/LogVisualizer/Private/LogVisualizerSettings.cpp
--- a/Engine/Source/Developer/LogVisualizer/Private/LogVisualizerSettings.cpp
+++ b/Engine/Source/Developer/LogVisualizer/Private/LogVisualizerSettings.cpp
@@ -109,6 +109,8 @@ void ULogVisualizerSettings::PostEditChangeProperty(struct FPropertyChangedEvent
 }
 #endif
 
+#if UE_DEBUG_VISUALIZER_TOOL_ENABLED
+
 //////////////////////////////////////////////////////////////////////////
 // FVisualLoggerFilters
 //////////////////////////////////////////////////////////////////////////
@@ -360,10 +362,10 @@ void FVisualLoggerFilters::DisableGraphData(FName GraphName, FName DataName, boo
 	}
 }
 
-
 bool FVisualLoggerFilters::IsGraphDataDisabled(FName GraphName, FName DataName)
 {
 	const FName FullName = *(GraphName.ToString() + TEXT("$") + DataName.ToString());
 	return DisabledGraphDatas.Find(FullName) != INDEX_NONE;
 }
 
+#endif // UE_DEBUG_VISUALIZER_TOOL_ENABLED
diff --git a/Engine/Source/Developer/LogVisualizer/Private/SVisualLogger.cpp b/Engine/Source/Developer/LogVisualizer/Private/SVisualLogger.cpp
--- a/Engine/Source/Developer/LogVisualizer/Private/SVisualLogger.cpp
+++ b/Engine/Source/Developer/LogVisualizer/Private/SVisualLogger.cpp
@@ -1,6 +1,8 @@
 // Copyright Epic Games, Inc. All Rights Reserved.
 
 #include "SVisualLogger.h"
+
+#if UE_DEBUG_VISUALIZER_TOOL_ENABLED
 #include "Framework/MultiBox/MultiBoxBuilder.h"
 #include "HAL/FileManager.h"
 #include "Misc/CommandLine.h"
@@ -481,10 +483,12 @@ void SVisualLogger::HandleStopRecordingCommandExecute()
 {
 	UWorld* World = FLogVisualizer::Get().GetWorld();
 
+#if UE_DEBUG_RECORDING_USING_VLOG
 	if (FParse::Param(FCommandLine::Get(), TEXT("LogNavOctree")) == true && GetDefault<ULogVisualizerSettings>()->bLogNavOctreeOnStop)
 	{
 		FVisualLogger::NavigationDataDump(World, LogNavigation, ELogVerbosity::Log, FBox());
 	}
+#endif // UE_DEBUG_RECORDING_USING_VLOG
 
 	FVisualLogger::Get().SetIsRecording(false);
 
@@ -1264,3 +1268,4 @@ FReply SVisualLogger::OnKeyboardRedirection(const FGeometry& MyGeometry, const F
 }
 
 #undef LOCTEXT_NAMESPACE
+#endif // UE_DEBUG_VISUALIZER_TOOL_ENABLED
\ No newline at end of file
diff --git a/Engine/Source/Developer/LogVisualizer/Private/SVisualLogger.h b/Engine/Source/Developer/LogVisualizer/Private/SVisualLogger.h
--- a/Engine/Source/Developer/LogVisualizer/Private/SVisualLogger.h
+++ b/Engine/Source/Developer/LogVisualizer/Private/SVisualLogger.h
@@ -1,13 +1,15 @@
 // Copyright Epic Games, Inc. All Rights Reserved.
 #pragma once
 
-#include "CoreMinimal.h"
-#include "VisualLogger/VisualLoggerTypes.h"
+#include "VisualLogger/VisualLoggerDefines.h"
+
+#if UE_DEBUG_VISUALIZER_TOOL_ENABLED
+#include "Framework/Commands/UICommandList.h"
+#include "Framework/Docking/TabManager.h"
 #include "Input/Reply.h"
+#include "VisualLogger/VisualLoggerTypes.h"
 #include "Widgets/DeclarativeSyntaxSupport.h"
 #include "Widgets/SCompoundWidget.h"
-#include "Framework/Commands/UICommandList.h"
-#include "Framework/Docking/TabManager.h"
 
 class AVisualLoggerCameraController;
 class FMenuBuilder;
@@ -156,3 +158,5 @@ protected:
 	TArray<TWeakObjectPtr<const class AVisualLoggerFilterVolume>> FilterVolumesInLastUsedWorld;
 	TArray<FBox> FilterBoxes;
 };
+
+#endif // UE_DEBUG_VISUALIZER_TOOL_ENABLED
\ No newline at end of file
diff --git a/Engine/Source/Developer/LogVisualizer/Private/SVisualLoggerFilters.cpp b/Engine/Source/Developer/LogVisualizer/Private/SVisualLoggerFilters.cpp
--- a/Engine/Source/Developer/LogVisualizer/Private/SVisualLoggerFilters.cpp
+++ b/Engine/Source/Developer/LogVisualizer/Private/SVisualLoggerFilters.cpp
@@ -1,6 +1,8 @@
 // Copyright Epic Games, Inc. All Rights Reserved.
 
 #include "SVisualLoggerFilters.h"
+
+#if UE_DEBUG_VISUALIZER_TOOL_ENABLED
 #include "EngineGlobals.h"
 #include "Framework/Application/SlateApplication.h"
 #include "Widgets/Layout/SWrapBox.h"
@@ -439,3 +441,4 @@ void SVisualLoggerFilters::OnItemsSelectionChanged(const FVisualLoggerDBRow& Cha
 }
 
 #undef LOCTEXT_NAMESPACE
+#endif // UE_DEBUG_VISUALIZER_TOOL_ENABLED
\ No newline at end of file
diff --git a/Engine/Source/Developer/LogVisualizer/Private/SVisualLoggerFilters.h b/Engine/Source/Developer/LogVisualizer/Private/SVisualLoggerFilters.h
--- a/Engine/Source/Developer/LogVisualizer/Private/SVisualLoggerFilters.h
+++ b/Engine/Source/Developer/LogVisualizer/Private/SVisualLoggerFilters.h
@@ -2,7 +2,9 @@
 
 #pragma once
 
-#include "CoreMinimal.h"
+#include "VisualLogger/VisualLoggerDefines.h"
+
+#if UE_DEBUG_VISUALIZER_TOOL_ENABLED
 #include "SVisualLoggerBaseWidget.h"
 #include "Widgets/SWidget.h"
 #include "Widgets/DeclarativeSyntaxSupport.h"
@@ -65,3 +67,5 @@ protected:
 
 	TMap<FName, TArray<FName> > CachedDatasPerGraph;
 };
+
+#endif // UE_DEBUG_VISUALIZER_TOOL_ENABLED
\ No newline at end of file
diff --git a/Engine/Source/Developer/LogVisualizer/Private/SVisualLoggerLogsList.cpp b/Engine/Source/Developer/LogVisualizer/Private/SVisualLoggerLogsList.cpp
--- a/Engine/Source/Developer/LogVisualizer/Private/SVisualLoggerLogsList.cpp
+++ b/Engine/Source/Developer/LogVisualizer/Private/SVisualLoggerLogsList.cpp
@@ -1,6 +1,8 @@
 // Copyright Epic Games, Inc. All Rights Reserved.
 
 #include "SVisualLoggerLogsList.h"
+
+#if UE_DEBUG_VISUALIZER_TOOL_ENABLED
 #include "SVisualLoggerTableRow.h"
 #include "Misc/OutputDeviceHelper.h"
 #include "Widgets/Text/STextBlock.h"
@@ -248,3 +250,4 @@ void SVisualLoggerLogsList::GenerateLogs(const FVisualLogDevice::FVisualLogEntry
 	LogsLinesWidget->RequestListRefresh();
 }
 #undef LOCTEXT_NAMESPACE
+#endif // UE_DEBUG_VISUALIZER_TOOL_ENABLED
\ No newline at end of file
diff --git a/Engine/Source/Developer/LogVisualizer/Private/SVisualLoggerLogsList.h b/Engine/Source/Developer/LogVisualizer/Private/SVisualLoggerLogsList.h
--- a/Engine/Source/Developer/LogVisualizer/Private/SVisualLoggerLogsList.h
+++ b/Engine/Source/Developer/LogVisualizer/Private/SVisualLoggerLogsList.h
@@ -2,6 +2,9 @@
 
 #pragma once
 
+#include "VisualLogger/VisualLoggerDefines.h"
+
+#if UE_DEBUG_VISUALIZER_TOOL_ENABLED
 #include "Input/Reply.h"
 #include "Widgets/DeclarativeSyntaxSupport.h"
 #include "Framework/Commands/UICommandList.h"
@@ -40,3 +43,4 @@ protected:
 	TSharedPtr<SListView<TSharedPtr<struct FLogEntryItem> > > LogsLinesWidget;
 	TArray<TSharedPtr<struct FLogEntryItem> > CachedLogEntryLines;
 };
+#endif // UE_DEBUG_VISUALIZER_TOOL_ENABLED
\ No newline at end of file
diff --git a/Engine/Source/Developer/LogVisualizer/Private/SVisualLoggerReport.cpp b/Engine/Source/Developer/LogVisualizer/Private/SVisualLoggerReport.cpp
--- a/Engine/Source/Developer/LogVisualizer/Private/SVisualLoggerReport.cpp
+++ b/Engine/Source/Developer/LogVisualizer/Private/SVisualLoggerReport.cpp
@@ -1,6 +1,8 @@
 // Copyright Epic Games, Inc. All Rights Reserved.
 
 #include "SVisualLoggerReport.h"
+
+#if UE_DEBUG_VISUALIZER_TOOL_ENABLED
 #include "SlateOptMacros.h"
 #include "LogVisualizerStyle.h"
 #include "Widgets/Input/SSearchBox.h"
@@ -183,3 +185,4 @@ void SVisualLoggerReport::GenerateReportText()
 
 
 #undef LOCTEXT_NAMESPACE
+#endif // UE_DEBUG_VISUALIZER_TOOL_ENABLED
\ No newline at end of file
diff --git a/Engine/Source/Developer/LogVisualizer/Private/SVisualLoggerReport.h b/Engine/Source/Developer/LogVisualizer/Private/SVisualLoggerReport.h
--- a/Engine/Source/Developer/LogVisualizer/Private/SVisualLoggerReport.h
+++ b/Engine/Source/Developer/LogVisualizer/Private/SVisualLoggerReport.h
@@ -2,7 +2,9 @@
 
 #pragma once
 
-#include "CoreMinimal.h"
+#include "VisualLogger/VisualLoggerDefines.h"
+
+#if UE_DEBUG_VISUALIZER_TOOL_ENABLED
 #include "Widgets/DeclarativeSyntaxSupport.h"
 #include "Widgets/SCompoundWidget.h"
 #include "SVisualLoggerView.h"
@@ -31,3 +33,4 @@ protected:
 	FText ReportText;
 	TArray<FString> CollectedEvents;
 };
+#endif // UE_DEBUG_VISUALIZER_TOOL_ENABLED
diff --git a/Engine/Source/Developer/LogVisualizer/Private/SVisualLoggerTimeline.cpp b/Engine/Source/Developer/LogVisualizer/Private/SVisualLoggerTimeline.cpp
--- a/Engine/Source/Developer/LogVisualizer/Private/SVisualLoggerTimeline.cpp
+++ b/Engine/Source/Developer/LogVisualizer/Private/SVisualLoggerTimeline.cpp
@@ -1,18 +1,20 @@
 // Copyright Epic Games, Inc. All Rights Reserved.
 
 #include "SVisualLoggerTimeline.h"
-#include "Widgets/SBoxPanel.h"
+
+#if UE_DEBUG_VISUALIZER_TOOL_ENABLED
+#include "LogVisualizerSettings.h"
+#include "LogVisualizerStyle.h"
 #include "Styling/CoreStyle.h"
+#include "SVisualLoggerTimelineBar.h"
+#include "SVisualLoggerTimelinesContainer.h"
+#include "VisualLoggerDatabase.h"
+#include "Widgets/Input/SMenuAnchor.h"
 #include "Widgets/Layout/SBorder.h"
-#include "Widgets/Text/STextBlock.h"
 #include "Widgets/Layout/SBox.h"
-#include "Widgets/Input/SMenuAnchor.h"
+#include "Widgets/SBoxPanel.h"
 #include "Widgets/SToolTip.h"
-#include "LogVisualizerSettings.h"
-#include "VisualLoggerDatabase.h"
-#include "LogVisualizerStyle.h"
-#include "SVisualLoggerTimelinesContainer.h"
-#include "SVisualLoggerTimelineBar.h"
+#include "Widgets/Text/STextBlock.h"
 
 class STimelineLabelAnchor : public SMenuAnchor
 {
@@ -270,3 +272,4 @@ void SLogVisualizerTimeline::MoveCursorByDistance(int32 Distance)
 		FLogVisualizer::Get().GotoPreviousItem(GetName(), FMath::Abs(Distance));
 	}
 }
+#endif // UE_DEBUG_VISUALIZER_TOOL_ENABLED
\ No newline at end of file
diff --git a/Engine/Source/Developer/LogVisualizer/Private/SVisualLoggerTimeline.h b/Engine/Source/Developer/LogVisualizer/Private/SVisualLoggerTimeline.h
--- a/Engine/Source/Developer/LogVisualizer/Private/SVisualLoggerTimeline.h
+++ b/Engine/Source/Developer/LogVisualizer/Private/SVisualLoggerTimeline.h
@@ -2,7 +2,9 @@
 
 #pragma once
 
-#include "CoreMinimal.h"
+#include "VisualLogger/VisualLoggerDefines.h"
+
+#if UE_DEBUG_VISUALIZER_TOOL_ENABLED
 #include "Input/Reply.h"
 #include "Widgets/DeclarativeSyntaxSupport.h"
 #include "Widgets/SCompoundWidget.h"
@@ -80,3 +82,4 @@ protected:
 	/** Delegate to execute to get the menu content of this timeline */
 	FOnGetContent OnGetMenuContent;
 };
+#endif // UE_DEBUG_VISUALIZER_TOOL_ENABLED
\ No newline at end of file
diff --git a/Engine/Source/Developer/LogVisualizer/Private/SVisualLoggerTimelineBar.cpp b/Engine/Source/Developer/LogVisualizer/Private/SVisualLoggerTimelineBar.cpp
--- a/Engine/Source/Developer/LogVisualizer/Private/SVisualLoggerTimelineBar.cpp
+++ b/Engine/Source/Developer/LogVisualizer/Private/SVisualLoggerTimelineBar.cpp
@@ -1,6 +1,7 @@
 // Copyright Epic Games, Inc. All Rights Reserved.
 
 #include "SVisualLoggerTimelineBar.h"
+#if UE_DEBUG_VISUALIZER_TOOL_ENABLED
 #include "Layout/ArrangedChildren.h"
 #include "Rendering/DrawElements.h"
 #include "VisualLoggerDatabase.h"
@@ -354,3 +355,4 @@ int32 SVisualLoggerTimelineBar::OnPaint(const FPaintArgs& Args, const FGeometry&
 
 	return RetLayerId;
 }
+#endif // UE_DEBUG_VISUALIZER_TOOL_ENABLED
\ No newline at end of file
diff --git a/Engine/Source/Developer/LogVisualizer/Private/SVisualLoggerTimelineBar.h b/Engine/Source/Developer/LogVisualizer/Private/SVisualLoggerTimelineBar.h
--- a/Engine/Source/Developer/LogVisualizer/Private/SVisualLoggerTimelineBar.h
+++ b/Engine/Source/Developer/LogVisualizer/Private/SVisualLoggerTimelineBar.h
@@ -2,7 +2,9 @@
 
 #pragma once
 
-#include "CoreMinimal.h"
+#include "VisualLogger/VisualLoggerDefines.h"
+
+#if UE_DEBUG_VISUALIZER_TOOL_ENABLED
 #include "Input/Reply.h"
 #include "Widgets/DeclarativeSyntaxSupport.h"
 #include "Widgets/SLeafWidget.h"
@@ -43,3 +45,4 @@ protected:
 	int32 MouseMoveClosestItemIndex = INDEX_NONE;
 	bool bToolTipUpdateRequested = false;
 };
+#endif // UE_DEBUG_VISUALIZER_TOOL_ENABLED
\ No newline at end of file
diff --git a/Engine/Source/Developer/LogVisualizer/Private/SVisualLoggerToolbar.cpp b/Engine/Source/Developer/LogVisualizer/Private/SVisualLoggerToolbar.cpp
--- a/Engine/Source/Developer/LogVisualizer/Private/SVisualLoggerToolbar.cpp
+++ b/Engine/Source/Developer/LogVisualizer/Private/SVisualLoggerToolbar.cpp
@@ -1,6 +1,8 @@
 // Copyright Epic Games, Inc. All Rights Reserved.
 
 #include "SVisualLoggerToolbar.h"
+
+#if UE_DEBUG_VISUALIZER_TOOL_ENABLED
 #include "Styling/AppStyle.h"
 #include "Framework/MultiBox/MultiBoxDefs.h"
 #include "Framework/MultiBox/MultiBoxBuilder.h"
@@ -82,3 +84,4 @@ TSharedRef<SWidget> SVisualLoggerToolbar::MakeClippingDistanceWidget()
 }
 
 #undef LOCTEXT_NAMESPACE
+#endif // UE_DEBUG_VISUALIZER_TOOL_ENABLED
diff --git a/Engine/Source/Developer/LogVisualizer/Private/SVisualLoggerToolbar.h b/Engine/Source/Developer/LogVisualizer/Private/SVisualLoggerToolbar.h
--- a/Engine/Source/Developer/LogVisualizer/Private/SVisualLoggerToolbar.h
+++ b/Engine/Source/Developer/LogVisualizer/Private/SVisualLoggerToolbar.h
@@ -2,7 +2,9 @@
 
 #pragma once
 
-#include "CoreMinimal.h"
+#include "VisualLogger/VisualLoggerDefines.h"
+
+#if UE_DEBUG_VISUALIZER_TOOL_ENABLED
 #include "Widgets/SWidget.h"
 #include "Widgets/DeclarativeSyntaxSupport.h"
 #include "Framework/Commands/UICommandList.h"
@@ -41,3 +43,4 @@ protected:
 
 	TSharedRef<SWidget> MakeClippingDistanceWidget();
 };
+#endif // UE_DEBUG_VISUALIZER_TOOL_ENABLED
\ No newline at end of file
diff --git a/Engine/Source/Developer/LogVisualizer/Private/VisualLoggerCanvasRenderer.cpp b/Engine/Source/Developer/LogVisualizer/Private/VisualLoggerCanvasRenderer.cpp
--- a/Engine/Source/Developer/LogVisualizer/Private/VisualLoggerCanvasRenderer.cpp
+++ b/Engine/Source/Developer/LogVisualizer/Private/VisualLoggerCanvasRenderer.cpp
@@ -1,6 +1,8 @@
 // Copyright Epic Games, Inc. All Rights Reserved.
 
 #include "VisualLoggerCanvasRenderer.h"
+
+#if UE_DEBUG_VISUALIZER_TOOL_ENABLED
 #include "EngineGlobals.h"
 #include "CanvasItem.h"
 #include "DebugRenderSceneProxy.h"
@@ -336,3 +338,4 @@ void FVisualLoggerCanvasRenderer::DrawHistogramGraphs(class UCanvas* Canvas, cla
 		}
 	}
 }
+#endif // UE_DEBUG_VISUALIZER_TOOL_ENABLED
\ No newline at end of file
diff --git a/Engine/Source/Developer/LogVisualizer/Private/VisualLoggerCanvasRenderer.h b/Engine/Source/Developer/LogVisualizer/Private/VisualLoggerCanvasRenderer.h
--- a/Engine/Source/Developer/LogVisualizer/Private/VisualLoggerCanvasRenderer.h
+++ b/Engine/Source/Developer/LogVisualizer/Private/VisualLoggerCanvasRenderer.h
@@ -1,7 +1,9 @@
 // Copyright Epic Games, Inc. All Rights Reserved.
 #pragma once
 
-#include "CoreMinimal.h"
+#include "VisualLogger/VisualLoggerDefines.h"
+
+#if UE_DEBUG_VISUALIZER_TOOL_ENABLED
 #include "VisualLogger/VisualLoggerTypes.h"
 
 struct FVisualLoggerDBRow;
@@ -43,3 +45,4 @@ private:
 	FVisualLogEntry SelectedEntry;
 	TMap<FString, TArray<FString> > UsedGraphCategories;
 };
+#endif // UE_DEBUG_VISUALIZER_TOOL_ENABLED
\ No newline at end of file
diff --git a/Engine/Source/Developer/LogVisualizer/Private/VisualLoggerDatabase.cpp b/Engine/Source/Developer/LogVisualizer/Private/VisualLoggerDatabase.cpp
--- a/Engine/Source/Developer/LogVisualizer/Private/VisualLoggerDatabase.cpp
+++ b/Engine/Source/Developer/LogVisualizer/Private/VisualLoggerDatabase.cpp
@@ -1,6 +1,8 @@
 // Copyright Epic Games, Inc. All Rights Reserved.
 
 #include "VisualLoggerDatabase.h"
+
+#if UE_DEBUG_VISUALIZER_TOOL_ENABLED
 #include "Engine/EngineTypes.h"
 #include "Engine/World.h"
 #include "EngineUtils.h"
@@ -483,3 +485,4 @@ bool FVisualLoggerEditorInterface::MatchCategoryFilters(const FString& String, E
 {
 	return FVisualLoggerFilters::Get().ShouldDisplayCategoryByString(String, Verbosity);
 }
+#endif // UE_DEBUG_VISUALIZER_TOOL_ENABLED
\ No newline at end of file
diff --git a/Engine/Source/Developer/LogVisualizer/Private/VisualLoggerDatabase.h b/Engine/Source/Developer/LogVisualizer/Private/VisualLoggerDatabase.h
--- a/Engine/Source/Developer/LogVisualizer/Private/VisualLoggerDatabase.h
+++ b/Engine/Source/Developer/LogVisualizer/Private/VisualLoggerDatabase.h
@@ -1,7 +1,9 @@
 // Copyright Epic Games, Inc. All Rights Reserved.
 #pragma once
 
-#include "CoreMinimal.h"
+#include "VisualLogger/VisualLoggerDefines.h"
+
+#if UE_DEBUG_VISUALIZER_TOOL_ENABLED
 #include "VisualLogger/VisualLoggerTypes.h"
 
 class AActor;
@@ -226,3 +228,4 @@ struct FVisualLoggerEditorInterface : public IVisualLoggerEditorInterface
 
 	bool MatchCategoryFilters(const FString& String, ELogVerbosity::Type Verbosity = ELogVerbosity::All) override;
 };
+#endif // UE_DEBUG_VISUALIZER_TOOL_ENABLED
diff --git a/Engine/Source/Developer/LogVisualizer/Public/LogVisualizerSettings.h b/Engine/Source/Developer/LogVisualizer/Public/LogVisualizerSettings.h
--- a/Engine/Source/Developer/LogVisualizer/Public/LogVisualizerSettings.h
+++ b/Engine/Source/Developer/LogVisualizer/Public/LogVisualizerSettings.h
@@ -64,8 +64,9 @@ struct FVisualLoggerFiltersData
 USTRUCT()
 struct FVisualLoggerFilters : public FVisualLoggerFiltersData
 {
-	GENERATED_USTRUCT_BODY()
+	GENERATED_BODY()
 
+#if UE_DEBUG_VISUALIZER_TOOL_ENABLED
 	FOnFilterCategoryAdded OnFilterCategoryAdded;
 	FOnFilterCategoryRemoved OnFilterCategoryRemoved;
 
@@ -93,24 +94,12 @@ struct FVisualLoggerFilters : public FVisualLoggerFiltersData
 	 */
 	bool ShouldDisplayCategoryByString(FStringView String, ELogVerbosity::Type Verbosity = ELogVerbosity::All) const;
 
-	UE_DEPRECATED(5.6, "Use ShouldDisplayCategoryVerbosity instead")
-	bool MatchCategoryFilters(FString String, const ELogVerbosity::Type Verbosity = ELogVerbosity::All)
-	{
-		return ShouldDisplayCategoryByString(String, Verbosity);
-	}
-
 	/** @return whether given String is a case-insensitive match to the active search filter */
 	bool IsStringMatchingSearchFilter(const FStringView String) const
 	{
 		return SearchBoxFilter.Equals(String.GetData(), ESearchCase::IgnoreCase);
 	}
 
-	UE_DEPRECATED(5.6, "Use IsStringMatchingSearchFilter instead")
-	bool MatchSearchString(FString String)
-	{
-		return IsStringMatchingSearchFilter(String);
-	}
-
 	void SetSearchString(FString InString)
 	{
 		SearchBoxFilter = MoveTemp(InString);
@@ -141,9 +130,11 @@ protected:
 	void OnNewItemHandler(const FVisualLoggerDBRow& BDRow, int32 ItemIndex);
 
 private:
-	static TSharedPtr< struct FVisualLoggerFilters > StaticInstance;
-	TMap<FName, FCategoryFilter*>	FastCategoryFilterMap;
+	static TSharedPtr<FVisualLoggerFilters> StaticInstance;
+	TMap<FName, FCategoryFilter*> FastCategoryFilterMap;
 	TArray<FName> DisabledGraphDatas;
+
+#endif //UE_DEBUG_VISUALIZER_TOOL_ENABLED
 };
 
 struct FCategoryFiltersManager;
@@ -224,7 +215,7 @@ public:
 
 	// UObject overrides
 #if WITH_EDITOR
-	LOGVISUALIZER_API virtual void PostEditChangeProperty(struct FPropertyChangedEvent& PropertyChangedEvent) override;
+	LOGVISUALIZER_API virtual void PostEditChangeProperty(FPropertyChangedEvent& PropertyChangedEvent) override;
 #endif
 
 	LOGVISUALIZER_API class UMaterial* GetDebugMeshMaterial();
@@ -243,7 +234,7 @@ protected:
 
 	/** A material used to render debug meshes with kind of flat shading, mostly used by Visual Logger tool. */
 	UPROPERTY()
-	TObjectPtr<class UMaterial> DebugMeshMaterialFakeLight;
+	TObjectPtr<UMaterial> DebugMeshMaterialFakeLight;
 
 	/** @todo document */
 	UPROPERTY(config)
diff --git a/Engine/Source/Developer/LogVisualizer/Public/VisualLoggerRenderingActorBase.h b/Engine/Source/Developer/LogVisualizer/Public/VisualLoggerRenderingActorBase.h
--- a/Engine/Source/Developer/LogVisualizer/Public/VisualLoggerRenderingActorBase.h
+++ b/Engine/Source/Developer/LogVisualizer/Public/VisualLoggerRenderingActorBase.h
@@ -70,7 +70,7 @@ public:
 	UE_API double GetFarClippingDistance() const;
 
 protected:
-#if ENABLE_VISUAL_LOG
+#if UE_DEBUG_VISUALIZER_TOOL_ENABLED
 	UE_API void GetDebugShapes(const FVisualLogEntry& EntryItem, bool bAddEntryLocationPointer, FTimelineDebugShapes& OutDebugShapes);
 #endif
 
diff --git a/Engine/Source/Runtime/Engine/Classes/VisualLogger/VisualLoggerBinaryFileDevice.h b/Engine/Source/Runtime/Engine/Classes/VisualLogger/VisualLoggerBinaryFileDevice.h
--- a/Engine/Source/Runtime/Engine/Classes/VisualLogger/VisualLoggerBinaryFileDevice.h
+++ b/Engine/Source/Runtime/Engine/Classes/VisualLogger/VisualLoggerBinaryFileDevice.h
@@ -1,14 +1,19 @@
 // Copyright Epic Games, Inc. All Rights Reserved.
 #pragma once
 
+#if UE_ENABLE_INCLUDE_ORDER_DEPRECATED_IN_5_8
 #include "CoreMinimal.h"
 #include "EngineDefines.h"
-#include "VisualLogger/VisualLoggerTypes.h"
+#endif // UE_ENABLE_INCLUDE_ORDER_DEPRECATED_IN_5_8
 
-#if ENABLE_VISUAL_LOG
+#include "VisualLogger/VisualLoggerDefines.h"
 
 #define VISLOG_FILENAME_EXT TEXT("bvlog")
 
+#if UE_DEBUG_RECORDING_USING_VLOG
+
+#include "VisualLogger/VisualLoggerTypes.h"
+
 class FVisualLoggerBinaryFileDevice : public FVisualLogDevice
 {
 public:
diff --git a/Engine/Source/Runtime/Engine/Classes/VisualLogger/VisualLoggerTraceDevice.h b/Engine/Source/Runtime/Engine/Classes/VisualLogger/VisualLoggerTraceDevice.h
--- a/Engine/Source/Runtime/Engine/Classes/VisualLogger/VisualLoggerTraceDevice.h
+++ b/Engine/Source/Runtime/Engine/Classes/VisualLogger/VisualLoggerTraceDevice.h
@@ -1,29 +1,38 @@
 // Copyright Epic Games, Inc. All Rights Reserved.
 #pragma once
 
+#if UE_ENABLE_INCLUDE_ORDER_DEPRECATED_IN_5_8
 #include "CoreMinimal.h"
 #include "EngineDefines.h"
+#endif // UE_ENABLE_INCLUDE_ORDER_DEPRECATED_IN_5_8
+
+#include "VisualLogger/VisualLoggerDefines.h"
+
+#if UE_DEBUG_RECORDING_ENABLED
 #include "VisualLogger/VisualLoggerTypes.h"
+#include "Delegates/DelegateCombinations.h"
 
-#if ENABLE_VISUAL_LOG
+#define UE_API ENGINE_API
 
 DECLARE_DELEGATE_TwoParams(FImmediateRenderDelegate, const UObject*, const FVisualLogEntry&);
 
 class FVisualLoggerTraceDevice : public FVisualLogDevice
 {
 public:
-	static ENGINE_API FVisualLoggerTraceDevice& Get();
+	static UE_API FVisualLoggerTraceDevice& Get();
 
-	ENGINE_API FVisualLoggerTraceDevice();
-	ENGINE_API virtual void Cleanup(bool bReleaseMemory = false) override;
-	ENGINE_API virtual void StartRecordingToFile(double TimeStamp) override;
-	ENGINE_API virtual void StopRecordingToFile(double TimeStamp) override;
-	ENGINE_API virtual void DiscardRecordingToFile() override;
-	ENGINE_API virtual void SetFileName(const FString& InFileName) override;
-	ENGINE_API virtual void Serialize(const UObject* InLogOwner, const FName& InOwnerName, const FName& InOwnerDisplayName, const FName& InOwnerClassName, const FVisualLogEntry& InLogEntry) override;
-	virtual bool HasFlags(int32 InFlags) const override { return !!(InFlags & (EVisualLoggerDeviceFlags::CanSaveToFile | EVisualLoggerDeviceFlags::StoreLogsLocally)); }
+	UE_API FVisualLoggerTraceDevice();
+	UE_API virtual void Cleanup(bool bReleaseMemory = false) override;
+	UE_API virtual void StartRecordingToFile(double TimeStamp) override;
+	UE_API virtual void StopRecordingToFile(double TimeStamp) override;
+	UE_API virtual void DiscardRecordingToFile() override;
+	UE_API virtual void SetFileName(const FString& InFileName) override;
+	UE_API virtual void Serialize(const UObject* InLogOwner, const FName& InOwnerName, const FName& InOwnerDisplayName, const FName& InOwnerClassName, const FVisualLogEntry& InLogEntry) override;
+	UE_API virtual bool HasFlags(int32 InFlags) const override;
 
 	FImmediateRenderDelegate ImmediateRenderDelegate;
 };
 
-#endif
+#undef UE_API
+
+#endif // UE_DEBUG_RECORDING_ENABLED
diff --git a/Engine/Source/Runtime/Engine/Private/VisualLogger/VisualLogger.cpp b/Engine/Source/Runtime/Engine/Private/VisualLogger/VisualLogger.cpp
--- a/Engine/Source/Runtime/Engine/Private/VisualLogger/VisualLogger.cpp
+++ b/Engine/Source/Runtime/Engine/Private/VisualLogger/VisualLogger.cpp
@@ -6,6 +6,7 @@
 #include "Containers/StringFwd.h"
 #include "EngineStats.h"
 #include "Framework/Docking/TabManager.h"
+#include "Logging/LogCategory.h"
 #include "Misc/CommandLine.h"
 #include "Misc/ScopeRWLock.h"
 #include "Modules/ModuleManager.h"
@@ -26,7 +27,8 @@
 #endif
 
 DEFINE_LOG_CATEGORY(LogVisual);
-#if ENABLE_VISUAL_LOG
 
+#if UE_DEBUG_RECORDING_ENABLED
+
 bool FVisualLogger::bUseVerbosityFilterWhenRecording = false;
 
@@ -75,11 +77,12 @@
 		/** When in Editor, use this TimeStamp as a Base time so our values don't grow too large (and are canonically around 0.0). */
 		static std::atomic<double> EditorBaseTimeStamp = 0.0;
 	}
-#endif
+#endif // WITH_EDITOR
 }
 
 DEFINE_STAT(STAT_VisualLog);
 
+#if UE_DEBUG_RECORDING_USING_VLOG
 template <typename SerializeFuncType>
 static void CollapsedLogfV(const SerializeFuncType& SerializeFunc, const UObject* Object, const FName& CategoryName, ELogVerbosity::Type Verbosity, const TCHAR* Fmt, va_list Args)
 {
@@ -91,13 +94,29 @@
 		SerializeFunc(CurrentEntry, TStringBuilder<512>().AppendV(Fmt, Args).ToString());
 	}
 }
+#else
+template <typename SerializeFuncType>
+static void CollapsedLogfV(const SerializeFuncType& SerializeFunc, const UObject* Object, const FName& CategoryName, ELogVerbosity::Type Verbosity, const TCHAR* Fmt, va_list Args)
+{
+	SCOPE_CYCLE_COUNTER(STAT_VisualLog);
 
+	FVisualLogEntry NewEntry;
+	FVisualLogger::FObjectNames Names;
+	if (FVisualLogger::CreateEntryInternal(Object, NewEntry, Names))
+	{
+		SerializeFunc(&NewEntry, TStringBuilder<512>().AppendV(Fmt, Args).ToString());
+
+		FVisualLoggerTraceDevice::Get().Serialize(Object, Names.Name, Names.DisplayName, Names.ClassName, NewEntry);
+	}
+}
+#endif // UE_DEBUG_RECORDING_USING_VLOG
+
 #define COLLAPSED_LOGF(SerializeSnippet) \
 	va_list Args; \
 	va_start(Args, Fmt); \
 	auto SerializeFunc = [&](FVisualLogEntry* CurrentEntry, const TCHAR* Buffer) { SerializeSnippet; }; \
 	CollapsedLogfV(SerializeFunc, Object, CategoryName, Verbosity, Fmt, Args); \
-	va_end(Args); 
+	va_end(Args);
 
 void FVisualLogger::CategorizedLogfImpl(const UObject* Object, const FName& CategoryName, ELogVerbosity::Type Verbosity, const TCHAR* Fmt, ...)
 {
@@ -225,7 +244,7 @@
 #endif
 		if (!GIsEditor && World == nullptr && GEngine)
 		{
-			World =  GEngine->GetWorld();
+			World = GEngine->GetWorld();
 		}
 
 		// In last resort we fall back on the global world (i.e., logs produced while world gets GCed)
@@ -238,10 +257,12 @@
 	}
 }
 
-TMap<const UWorld*, FVisualLogger::FOwnerToChildrenRedirectionMap> FVisualLogger::WorldToRedirectionMap;
 int32 FVisualLogger::bIsRecording = false;
 FVisualLogger::FNavigationDataDump FVisualLogger::NavigationDataDumpDelegate;
 
+#if UE_DEBUG_RECORDING_USING_VLOG
+TMap<const UWorld*, FVisualLogger::FOwnerToChildrenRedirectionMap> FVisualLogger::WorldToRedirectionMap;
+
 bool FVisualLogger::CheckVisualLogInputInternal(const UObject* Object, const FName& CategoryName, ELogVerbosity::Type, UWorld** OutWorld, FVisualLogEntry** OutCurrentEntry)
 {
 	if (IsRecording() == false || !Object || !GEngine || GEngine->bDisableAILogging || Object->HasAnyFlags(RF_ClassDefaultObject))
@@ -249,7 +270,7 @@
 		return false;
 	}
 
-	FVisualLogger& VisualLogger = FVisualLogger::Get();
+	FVisualLogger& VisualLogger = Get();
 	if (VisualLogger.IsBlockedForAllCategories() && VisualLogger.IsCategoryAllowed(CategoryName) == false)
 	{
 		return false;
@@ -276,17 +297,45 @@
 	return true;
 }
 
-FVisualLogEntry* FVisualLogger::GetEntryToWrite(const UObject* LogOwner, const FLogCategoryBase& LogCategory)
+#else // ^^^ UE_DEBUG_RECORDING_USING_VLOG ^^^ / vvv TRACE ONLY vvv
+
+bool FVisualLogger::CreateEntryInternal(const UObject* Object, FVisualLogEntry& OutEntry, FVisualLogger::FObjectNames& OutNames)
 {
-	FVisualLogEntry* LogEntry;
-	if (CheckVisualLogInputInternal(LogOwner, LogCategory.GetCategoryName(), ELogVerbosity::Log, /*OutWorld*/ nullptr, &LogEntry))
+	if (IsRecording() == false || !Object || !GEngine || GEngine->bDisableAILogging || Object->HasAnyFlags(RF_ClassDefaultObject))
 	{
-		return LogEntry;
+		return false;
 	}
 
-	return nullptr;
+	FVisualLogger& VisualLogger = Get();
+
+	using namespace UE::VisLog::Private;
+	FVisualLogEntry NewEntry;
+	NewEntry.InitializeEntry(VisualLogger.GetTimeStampForObject(Object));
+
+	FObjectNames* Names;
+	{
+		UE::TReadScopeLock Lock(VisualLogger.RedirectRWLock);
+		Names = VisualLogger.ObjectToNamesMap.Find(Object);
+	}
+
+	if (Names == nullptr)
+	{
+		UE::TWriteScopeLock Lock(VisualLogger.RedirectRWLock);
+		Names = VisualLogger.ObjectToNamesMap.Find(Object);
+		if (Names == nullptr)
+		{
+			Names = &VisualLogger.ObjectToNamesMap.Add(Object, VisualLogger.CreateObjectNames(Object));
+		}
+	}
+
+	OutEntry = MoveTemp(NewEntry);
+	OutNames = *Names;
+
+	return true;
 }
 
+#endif // UE_DEBUG_RECORDING_USING_VLOG
+
 double FVisualLogger::GetTimeStampForObject(const UObject* Object) const
 {
 	// Licensees can write their own synchronized clock that can work across the network
@@ -310,7 +359,7 @@
 			return FApp::GetCurrentTime() - EditorOnly::EditorBaseTimeStamp;
 		}
 	}
-#endif
+#endif // WITH_EDITOR
 
 	// This will be the fallback mode in standalone.  We do not have a synchronized clock.
 	const UWorld* WorldForTimeStamp = GetWorldForVisualLogger(Object);
@@ -323,6 +372,18 @@
 	GetTimeStampFunc = Function;
 }
 
+#if UE_DEBUG_RECORDING_USING_VLOG
+FVisualLogEntry* FVisualLogger::GetEntryToWrite(const UObject* LogOwner, const FLogCategoryBase& LogCategory)
+{
+	FVisualLogEntry* LogEntry;
+	if (CheckVisualLogInputInternal(LogOwner, LogCategory.GetCategoryName(), ELogVerbosity::Log, /*OutWorld*/ nullptr, &LogEntry))
+	{
+		return LogEntry;
+	}
+
+	return nullptr;
+}
+
 void FVisualLogger::AddClassToAllowList(UClass& InClass)
 {
 	ClassAllowList.AddUnique(&InClass);
@@ -548,30 +609,8 @@
 		// It's first and only one usage of LogOwner as regular object to get names. We assume once that LogOwner is correct here and only here.
 		CurrentEntry = &CurrentEntryPerObject.Add(LogOwner);
 
-		if (bForceUniqueLogNames && UniqueLogPrefix.IsEmpty())
-		{
-			UniqueLogPrefix = FString::Printf(TEXT("[%d] "), UniqueRunNumber);
-		}
+		ObjectToNamesMap.Add(LogOwner, CreateObjectNames(LogOwner));
+		ObjectToWorldMap.Add(LogOwner, GetWorldForVisualLogger(LogOwner));
 
-		const UWorld* World = GetWorldForVisualLogger(LogOwner);
-		const bool bIsStandalone = (World == nullptr || World->GetNetMode() == NM_Standalone);
-
-		// We use the fully qualified name in order to minimize the chance of name conflict. Otherwise, we end up losing or merging several log entries
-		//  because they share the same object name (e.g. if the LogOwner is a given primitive component, then its name will likely be "PrimitiveComponent", which 
-		//  means that if 2 log entries are added at the same frame for 2 different primitive components, only one will remain, because the OwnerName is the stored key, not the object...)
-		const FName LogName(*FString::Printf(TEXT("%s%s%s"),
-			*UniqueLogPrefix,
-			bIsStandalone ? TEXT("") : *FString::Printf(TEXT("(%s) "), *GetDebugStringForWorld(World)),
-			*LogOwner->GetFullName()));
-
-		const FName LogDisplayName(*FString::Printf(TEXT("%s%s%s%s"),
-			*UniqueLogPrefix,
-			bIsStandalone ? TEXT("") : *FString::Printf(TEXT("(%s) "), *GetDebugStringForWorld(World)),
-			*LogOwner->GetName(),
-			bForceUniqueLogNames ? *FString::Printf(TEXT(" [%d]"), LogOwner->GetUniqueID()) : TEXT("")));
-
-		ObjectToNamesMap.Add(LogOwner, FObjectNames { .Name = LogName, .DisplayName = LogDisplayName, .ClassName = LogOwner->GetClass()->GetFName() });
-		ObjectToWorldMap.Add(LogOwner, World);
-
 		CalculateEntryAllowLogging(CurrentEntry, LogOwner, Object);
 
@@ -652,21 +691,6 @@
 	CurrentEntry->UpdateAllowedToLog();
 }
 
-void FVisualLogger::Tick(float DeltaTime)
-{
-	if (bIsFlushRequired)
-	{
-		Flush();
-		bIsFlushRequired = false;
-	}
-
-	if (bContainsInvalidRedirects)
-	{
-		CleanupRedirects();
-		bContainsInvalidRedirects = false;
-	}
-}
-
 void FVisualLogger::FlushThreadsEntries()
 {
 	UE_MT_SCOPED_READ_ACCESS(ThreadLocalMapListAccessDetector);
@@ -822,6 +846,7 @@
 
 	NavigationDataDumpDelegate.Broadcast(Object, CategoryName, Verbosity, Box, *World, *CurrentEntry);
 }
+#endif // UE_DEBUG_RECORDING_USING_VLOG
 
 FVisualLogger& FVisualLogger::Get()
 {
@@ -832,8 +857,18 @@
 FVisualLogger::FVisualLogger()
 {
 	bForceUniqueLogNames = true;
+
+	if (FParse::Param(FCommandLine::Get(), TEXT("EnableVisualLoggerTrace")))
+	{
+		SetIsRecordingToTrace(true);
+	}
+	else
+	{
+		bIsRecordingToTrace = false;
+	}
+
+#if UE_DEBUG_RECORDING_USING_VLOG
 	bIsRecordingToFile = false;
-	bIsRecordingToTrace = false;
 	bIsFlushRequired = false;
 
 	BlockAllCategories(false);
@@ -877,9 +912,229 @@
 				UE::VisLog::Private::UniqueRunNumber = OldUniqueNumber + 1;
 			}
 		});
-#endif
+#endif // WITH_EDITOR
+#endif // UE_DEBUG_RECORDING_USING_VLOG
 }
 
+void FVisualLogger::SetIsRecording(const bool InIsRecording)
+{
+#if UE_DEBUG_RECORDING_USING_VLOG
+	if (InIsRecording == false && InIsRecording != !!bIsRecording && FParse::Param(FCommandLine::Get(), TEXT("LogNavOctree")))
+	{
+		NavigationDataDump(GetWorldForVisualLogger(nullptr), LogNavigation, ELogVerbosity::Log, FBox());
+	}
+
+	if (IsRecordingToFile())
+	{
+		SetIsRecordingToFile(false);
+	}
+#endif // UE_DEBUG_RECORDING_USING_VLOG
+
+	bIsRecording = InIsRecording;
+};
+
+void FVisualLogger::SetIsRecordingToTrace(const bool InIsRecording)
+{
+	if (!bIsRecording && InIsRecording)
+	{
+		SetIsRecording(true);
+	}
+
+	FVisualLoggerTraceDevice& Device = FVisualLoggerTraceDevice::Get();
+	if (bIsRecordingToTrace && !InIsRecording)
+	{
+		Device.StopRecordingToFile(0.0);
+#if UE_DEBUG_RECORDING_USING_VLOG
+		RemoveDevice(&Device);
+#endif
+	}
+	else if (!bIsRecordingToTrace && InIsRecording)
+	{
+		Device.StartRecordingToFile(0.0);
+#if UE_DEBUG_RECORDING_USING_VLOG
+		AddDevice(&Device);
+#endif
+	}
+
+	bIsRecordingToTrace = InIsRecording;
+}
+
+void FVisualLogger::SetUseUniqueNames(const bool bEnable)
+{
+	bForceUniqueLogNames = bEnable;
+	UE::VisLog::Private::UniqueLogPrefix.Empty();
+}
+
+void FVisualLogger::Redirect(const UObject* FromObject, const UObject* ToObject)
+{
+	FVisualLogger& Logger = Get();
+	UObject* NewRedirection = nullptr;
+	{
+		UE::TWriteScopeLock Lock(Logger.RedirectRWLock);
+		NewRedirection = Logger.RedirectInternal(FromObject, ToObject);
+	}
+	UE_CVLOG(FromObject != nullptr && NewRedirection != nullptr, FromObject, LogVisual, Log, TEXT("Redirected '%s' to '%s'"), *FromObject->GetName(), *NewRedirection->GetName());
+}
+
+UObject* FVisualLogger::RedirectInternal(const UObject* FromObject, const UObject* ToObject)
+{
+	if (FromObject == ToObject || FromObject == nullptr || ToObject == nullptr)
+	{
+		return nullptr;
+	}
+
+	const TWeakObjectPtr<const UObject> FromWeakPtr(FromObject);
+
+	TWeakObjectPtr<const UObject>& WeakOwnerEntry = ChildToOwnerMap.FindOrAdd(FromObject);
+
+#if UE_DEBUG_RECORDING_USING_VLOG
+	const UObject* DirectOwner = WeakOwnerEntry.Get();
+	const UObject* NewDirectOwner = ToObject;
+	if (DirectOwner != ToObject)
+	{
+		FOwnerToChildrenRedirectionMap& OwnerToChildrenMap = GetRedirectionMap(FromObject);
+		const TArray<TWeakObjectPtr<const UObject>>* ObjectChildren = OwnerToChildrenMap.Find(FromObject);
+
+		// A valid direct owner indicates that FromObject was in one or many list(s)
+		// of children so remove it and its children from the owner hierarchy
+		if (DirectOwner != nullptr)
+		{
+			const UObject* Owner = DirectOwner;
+			while (Owner != nullptr)
+			{
+				if (TArray<TWeakObjectPtr<const UObject>>* DirectOwnerChildren = OwnerToChildrenMap.Find(Owner))
+				{
+					DirectOwnerChildren->RemoveSingleSwap(FromWeakPtr);
+					if (ObjectChildren != nullptr)
+					{
+						for (const TWeakObjectPtr<const UObject>& Child : *ObjectChildren)
+						{
+							DirectOwnerChildren->RemoveSingleSwap(Child);
+						}
+					}
+				}
+
+				const TWeakObjectPtr<const UObject>* WeakOwner = ChildToOwnerMap.Find(Owner);
+				Owner = WeakOwner ? WeakOwner->Get() : nullptr;
+			}
+		}
+
+		// Add this child with its children (if any) to the owner hierarchy
+		const UObject* Owner = NewDirectOwner;
+		while (Owner != nullptr)
+		{
+			TArray<TWeakObjectPtr<const UObject>>* OwnerChildren = OwnerToChildrenMap.Find(Owner);
+			if (OwnerChildren == nullptr)
+			{
+				// Entry not found, add it and refresh current object children since it might have been reallocated
+				OwnerChildren = &OwnerToChildrenMap.Emplace(Owner);
+				ObjectChildren = ObjectChildren ? OwnerToChildrenMap.Find(FromObject) : nullptr;
+			}
+
+			check(OwnerChildren);
+			OwnerChildren->Add(FromWeakPtr);
+			if (ObjectChildren != nullptr)
+			{
+				OwnerChildren->Append(*ObjectChildren);
+			}
+
+			const TWeakObjectPtr<const UObject>* WeakOwner = ChildToOwnerMap.Find(Owner);
+			Owner = WeakOwner ? WeakOwner->Get() : nullptr;
+		}
+	}
+#endif // UE_DEBUG_RECORDING_USING_VLOG
+
+	WeakOwnerEntry = ToObject;
+
+	UObject* NewRedirection = FindRedirectionInternal(ToObject);
+	return NewRedirection;
+}
+
+UObject* FVisualLogger::FindRedirectionInternal(const UObject* Object) const
+{
+	TWeakObjectPtr TargetWeakPtr(Object);
+	const TWeakObjectPtr<const UObject>* Parent = &TargetWeakPtr;
+
+	while (Parent)
+	{
+		Parent = ChildToOwnerMap.Find(TargetWeakPtr.Get(/*bEvenIfPendingKill*/true));
+		if (Parent)
+		{
+			if (Parent->IsValid())
+			{
+				TargetWeakPtr = *Parent;
+			}
+			else
+			{
+				Parent = nullptr;
+				bContainsInvalidRedirects = true;
+			}
+		}
+	}
+
+	const UObject* const Target = TargetWeakPtr.Get(/*bEvenIfPendingKill*/true);
+	return const_cast<UObject*>(Target);
+}
+
+FVisualLogger::FObjectNames FVisualLogger::CreateObjectNames(TNotNull<const UObject*> LogOwner)
+{
+	if (bForceUniqueLogNames && UE::VisLog::Private::UniqueLogPrefix.IsEmpty())
+	{
+		UE::VisLog::Private::UniqueLogPrefix = FString::Printf(TEXT("[%d] "), UE::VisLog::Private::UniqueRunNumber);
+	}
+
+	const UWorld* World = GetWorldForVisualLogger(LogOwner);
+	const bool bIsStandalone = (World == nullptr || World->GetNetMode() == NM_Standalone);
+
+	// We use the fully qualified name in order to minimize the chance of name conflict. Otherwise, we end up losing or merging several log entries
+	//  because they share the same object name (e.g. if the LogOwner is a given primitive component, then its name will likely be "PrimitiveComponent", which 
+	//  means that if 2 log entries are added at the same frame for 2 different primitive components, only one will remain, because the OwnerName is the stored key, not the object...)
+	const FName LogName(*FString::Printf(TEXT("%s%s%s"),
+		*UE::VisLog::Private::UniqueLogPrefix,
+		bIsStandalone ? TEXT("") : *FString::Printf(TEXT("(%s) "), *GetDebugStringForWorld(World)),
+		*LogOwner->GetFullName()));
+
+	const FName LogDisplayName(*FString::Printf(TEXT("%s%s%s%s"),
+		*UE::VisLog::Private::UniqueLogPrefix,
+		bIsStandalone ? TEXT("") : *FString::Printf(TEXT("(%s) "), *GetDebugStringForWorld(World)),
+		*LogOwner->GetName(),
+		bForceUniqueLogNames ? *FString::Printf(TEXT(" [%d]"), LogOwner->GetUniqueID()) : TEXT("")));
+
+	return {.Name = LogName, .DisplayName = LogDisplayName, .ClassName = LogOwner->GetClass()->GetFName()};
+}
+
+void FVisualLogger::Tick(float DeltaTime)
+{
+#if UE_DEBUG_RECORDING_USING_VLOG
+	if (bIsFlushRequired)
+	{
+		Flush();
+		bIsFlushRequired = false;
+	}
+#endif //UE_DEBUG_RECORDING_USING_VLOG
+
+	if (bContainsInvalidRedirects)
+	{
+		CleanupRedirects();
+		bContainsInvalidRedirects = false;
+	}
+}
+
+void FVisualLogger::CleanupRedirects()
+{
+	UE::TWriteScopeLock Lock(RedirectRWLock);
+	for (auto It = ChildToOwnerMap.CreateIterator(); It; ++It)
+	{
+		if (!It.Value().IsValid())
+		{
+			FObjectKey Obj = It.Key();
+			ChildToOwnerMap.Remove(Obj);
+		}
+	}
+}
+
+#if UE_DEBUG_RECORDING_USING_VLOG
+
 void FVisualLogger::TearDown()
 {
 #if WITH_EDITOR
@@ -913,6 +1168,13 @@
 #endif
 }
 
+UObject* FVisualLogger::FindRedirection(const UObject* Object)
+{
+	FVisualLogger& Logger = Get();
+	UE::TReadScopeLock Lock(Logger.RedirectRWLock);
+	return Logger.FindRedirectionInternal(Object);
+}
+
 void FVisualLogger::Cleanup(UWorld* OldWorld, const bool bReleaseMemory)
 {
 	const bool WasRecordingToFile = IsRecordingToFile();
@@ -922,7 +1184,7 @@
 	}
 
 	Flush();
-	for (FVisualLogDevice* Device : FVisualLogger::Get().OutputDevices)
+	for (FVisualLogDevice* Device : Get().OutputDevices)
 	{
 		Device->Cleanup(bReleaseMemory);
 	}
@@ -1015,7 +1277,7 @@
 FVisualLogger::FOwnerToChildrenRedirectionMap& FVisualLogger::GetRedirectionMap(const UObject* InObject)
 {
 	const UWorld* World = nullptr;
-	if (const TWeakObjectPtr<const UWorld>* Entry = FVisualLogger::Get().ObjectToWorldMap.Find(InObject))
+	if (const TWeakObjectPtr<const UWorld>* Entry = Get().ObjectToWorldMap.Find(InObject))
 	{
 		World = Entry->Get();
 	}
@@ -1028,130 +1290,6 @@
 	return WorldToRedirectionMap.FindOrAdd(World);
 }
 
-UObject* FVisualLogger::RedirectInternal(const UObject* FromObject, const UObject* ToObject)
-{
-	if (FromObject == ToObject || FromObject == nullptr || ToObject == nullptr)
-	{
-		return nullptr;
-	}
-
-	const TWeakObjectPtr<const UObject> FromWeakPtr(FromObject);
-
-	TWeakObjectPtr<const UObject>& WeakOwnerEntry = ChildToOwnerMap.FindOrAdd(FromObject);
-	const UObject* DirectOwner = WeakOwnerEntry.Get();
-	const UObject* NewDirectOwner = ToObject;
-
-	if (DirectOwner != ToObject)
-	{
-		FOwnerToChildrenRedirectionMap& OwnerToChildrenMap = GetRedirectionMap(FromObject);
-		const TArray<TWeakObjectPtr<const UObject>>* ObjectChildren = OwnerToChildrenMap.Find(FromObject);
-
-		// A valid direct owner indicates that FromObject was in one or many list(s)
-		// of children so remove it and its children from the owner hierarchy
-		if (DirectOwner != nullptr)
-		{
-			const UObject* Owner = DirectOwner;
-			while (Owner != nullptr)
-			{
-				if (TArray<TWeakObjectPtr<const UObject>>* DirectOwnerChildren = OwnerToChildrenMap.Find(Owner))
-				{
-					DirectOwnerChildren->RemoveSingleSwap(FromWeakPtr);
-					if (ObjectChildren != nullptr)
-					{
-						for (const TWeakObjectPtr<const UObject>& Child : *ObjectChildren)
-						{
-							DirectOwnerChildren->RemoveSingleSwap(Child);
-						}
-					}
-				}
-
-				const TWeakObjectPtr<const UObject>* WeakOwner = ChildToOwnerMap.Find(Owner);
-				Owner = WeakOwner ? WeakOwner->Get() : nullptr;
-			}
-		}
-
-		// Add this child with its children (if any) to the owner hierarchy
-		const UObject* Owner = NewDirectOwner;
-		while (Owner != nullptr)
-		{
-			TArray<TWeakObjectPtr<const UObject>>* OwnerChildren = OwnerToChildrenMap.Find(Owner);
-			if (OwnerChildren == nullptr)
-			{
-				// Entry not found, add it and refresh current object children since it might have been reallocated
-				OwnerChildren = &OwnerToChildrenMap.Emplace(Owner);
-				ObjectChildren = ObjectChildren ? OwnerToChildrenMap.Find(FromObject) : nullptr;
-			}
-
-			check(OwnerChildren);
-			OwnerChildren->Add(FromWeakPtr);
-			if (ObjectChildren != nullptr)
-			{
-				OwnerChildren->Append(*ObjectChildren);
-			}
-
-			const TWeakObjectPtr<const UObject>* WeakOwner = ChildToOwnerMap.Find(Owner);
-			Owner = WeakOwner ? WeakOwner->Get() : nullptr;
-		}
-	}
-
-	WeakOwnerEntry = ToObject;
-
-	UObject* NewRedirection = FindRedirectionInternal(ToObject);
-	return NewRedirection;
-}
-
-UObject* FVisualLogger::FindRedirectionInternal(const UObject* Object) const
-{
-	TWeakObjectPtr<const UObject> TargetWeakPtr(Object);
-	const TWeakObjectPtr<const UObject>* Parent = &TargetWeakPtr;
-
-	while (Parent)
-	{
-		Parent = ChildToOwnerMap.Find(TargetWeakPtr.Get(/*bEvenIfPendingKill*/true));
-		if (Parent)
-		{
-			if (Parent->IsValid())
-			{
-				TargetWeakPtr = *Parent;
-			}
-			else
-			{
-				Parent = nullptr;
-				bContainsInvalidRedirects = true;
-			}
-		}
-	}
-
-	const UObject* const Target = TargetWeakPtr.Get(/*bEvenIfPendingKill*/true);
-	return const_cast<UObject*>(Target);
-}
-
-void FVisualLogger::CleanupRedirects()
-{
-	UE::TWriteScopeLock Lock(RedirectRWLock);
-	for (auto It = ChildToOwnerMap.CreateIterator(); It; ++It)
-	{
-		if(!It.Value().IsValid())
-		{
-			FObjectKey Obj = It.Key();
-			ChildToOwnerMap.Remove(Obj);
-		}
-	}
-}
-
-void FVisualLogger::SetIsRecording(const bool InIsRecording)
-{
-	if (InIsRecording == false && InIsRecording != !!bIsRecording && FParse::Param(FCommandLine::Get(), TEXT("LogNavOctree")))
-	{
-		FVisualLogger::NavigationDataDump(GetWorldForVisualLogger(nullptr), LogNavigation, ELogVerbosity::Log, FBox());
-	}
-	if (IsRecordingToFile())
-	{
-		SetIsRecordingToFile(false);
-	}
-	bIsRecording = InIsRecording;
-};
-
 void FVisualLogger::SetIsRecordingToFile(const bool InIsRecording)
 {
 	if (!bIsRecording && InIsRecording)
@@ -1192,34 +1330,6 @@
 	bIsRecordingToFile = InIsRecording;
 }
 
-void FVisualLogger::SetIsRecordingToTrace(const bool InIsRecording)
-{
-	if (!bIsRecording && InIsRecording)
-	{
-		SetIsRecording(true);
-	}
-
-	FVisualLoggerTraceDevice& Device = FVisualLoggerTraceDevice::Get();
-	if (bIsRecordingToTrace && !InIsRecording)
-	{
-		Device.StopRecordingToFile(0.0);
-		RemoveDevice(&Device);
-	}
-	else if (!bIsRecordingToTrace && InIsRecording)
-	{
-		Device.StartRecordingToFile(0.0);
-		AddDevice(&Device);
-	}
-
-	bIsRecordingToTrace = InIsRecording;
-}
-
-void FVisualLogger::SetUseUniqueNames(const bool bEnable)
-{
-	bForceUniqueLogNames = bEnable;
-	UE::VisLog::Private::UniqueLogPrefix.Empty();
-}
-
 void FVisualLogger::DiscardRecordingToFile()
 {
 	if (bIsRecordingToFile)
@@ -1252,12 +1362,12 @@
 	return true;
 }
 
-#endif //ENABLE_VISUAL_LOG
+#endif // UE_DEBUG_RECORDING_USING_VLOG
 
 const FGuid EVisualLoggerVersion::GUID = FGuid(0xA4237A36, 0xCAEA41C9, 0x8FA218F8, 0x58681BF3);
 FCustomVersionRegistration GVisualLoggerVersion(EVisualLoggerVersion::GUID, EVisualLoggerVersion::LatestVersion, TEXT("VisualLogger"));
 
-#if ENABLE_VISUAL_LOG
+#if UE_DEBUG_RECORDING_USING_VLOG
 
 class FLogVisualizerExec : private FSelfRegisteringExec
 {
@@ -1314,4 +1424,6 @@
 	}
 } LogVisualizerExec;
 
-#endif // ENABLE_VISUAL_LOG
+#endif // UE_DEBUG_RECORDING_USING_VLOG
+
+#endif // UE_DEBUG_RECORDING_ENABLED
diff --git a/Engine/Source/Runtime/Engine/Private/VisualLogger/VisualLoggerAutomationTests.cpp b/Engine/Source/Runtime/Engine/Private/VisualLogger/VisualLoggerAutomationTests.cpp
--- a/Engine/Source/Runtime/Engine/Private/VisualLogger/VisualLoggerAutomationTests.cpp
+++ b/Engine/Source/Runtime/Engine/Private/VisualLogger/VisualLoggerAutomationTests.cpp
@@ -37,7 +37,7 @@ UVisualLoggerAutomationTests::UVisualLoggerAutomationTests(const FObjectInitiali
 
 }
 
-#if ENABLE_VISUAL_LOG
+#if UE_DEBUG_RECORDING_USING_VLOG
 
 class FVisualLoggerTestDevice : public FVisualLogDevice
 {
@@ -617,4 +617,4 @@ bool FVisualLogRedirectionsDeepHierarchyToNewParentTest::RunTest(const FString&
 #undef CHECK_SUCCESS
 #undef CHECK_FAIL
 
-#endif //ENABLE_VISUAL_LOG
+#endif // UE_DEBUG_RECORDING_USING_VLOG
diff --git a/Engine/Source/Runtime/Engine/Private/VisualLogger/VisualLoggerBinaryFileDevice.cpp b/Engine/Source/Runtime/Engine/Private/VisualLogger/VisualLoggerBinaryFileDevice.cpp
--- a/Engine/Source/Runtime/Engine/Private/VisualLogger/VisualLoggerBinaryFileDevice.cpp
+++ b/Engine/Source/Runtime/Engine/Private/VisualLogger/VisualLoggerBinaryFileDevice.cpp
@@ -1,11 +1,12 @@
 // Copyright Epic Games, Inc. All Rights Reserved.
+
 #include "VisualLogger/VisualLoggerBinaryFileDevice.h"
+
+#if UE_DEBUG_RECORDING_USING_VLOG
 #include "HAL/FileManager.h"
 #include "Misc/ConfigCacheIni.h"
 #include "VisualLogger/VisualLogger.h"
 
-#if ENABLE_VISUAL_LOG
-
 FVisualLoggerBinaryFileDevice::FVisualLoggerBinaryFileDevice()
 	: FileArchive(nullptr)
 {
@@ -113,4 +114,4 @@ void FVisualLoggerBinaryFileDevice::Serialize(const UObject* InLogOwner, const F
 		FrameCache.Add(FVisualLogEntryItem(InOwnerName, InOwnerDisplayName, InOwnerClassName, InLogEntry));
 	}
 }
-#endif
+#endif // UE_DEBUG_RECORDING_USING_VLOG
diff --git a/Engine/Source/Runtime/Engine/Private/VisualLogger/VisualLoggerKismetLibrary.cpp b/Engine/Source/Runtime/Engine/Private/VisualLogger/VisualLoggerKismetLibrary.cpp
--- a/Engine/Source/Runtime/Engine/Private/VisualLogger/VisualLoggerKismetLibrary.cpp
+++ b/Engine/Source/Runtime/Engine/Private/VisualLogger/VisualLoggerKismetLibrary.cpp
@@ -26,14 +26,14 @@ void UVisualLoggerKismetLibrary::RedirectVislog(UObject* SourceOwner, UObject* D
 
 void UVisualLoggerKismetLibrary::EnableRecording(bool bEnabled)
 {
-#if ENABLE_VISUAL_LOG
+#if UE_DEBUG_RECORDING_ENABLED
 	FVisualLogger::Get().SetIsRecording(bEnabled);
-#endif // ENABLE_VISUAL_LOG
+#endif // UE_DEBUG_RECORDING_ENABLED
 }
 
 void UVisualLoggerKismetLibrary::LogText(UObject* WorldContextObject, FString Text, FName CategoryName, bool bAddToMessageLog)
 {
-#if ENABLE_VISUAL_LOG
+#if UE_DEBUG_RECORDING_ENABLED
 	FVisualLogger::CategorizedLogf(WorldContextObject, CategoryName, DefaultVerbosity
 		, TEXT("%s"), *Text);
 #endif
@@ -45,7 +45,7 @@ void UVisualLoggerKismetLibrary::LogText(UObject* WorldContextObject, FString Te
 
 void UVisualLoggerKismetLibrary::LogLocation(UObject* WorldContextObject, FVector Location, FString Text, FLinearColor ObjectColor, float Radius, FName CategoryName, bool bAddToMessageLog)
 {
-#if ENABLE_VISUAL_LOG
+#if UE_DEBUG_RECORDING_ENABLED
 	FVisualLogger::LocationLogf(WorldContextObject, CategoryName, DefaultVerbosity
 		, Location, static_cast<uint16>(Radius), ObjectColor.ToFColor(true), TEXT("%s"), *Text);
 #endif
@@ -57,7 +57,7 @@ void UVisualLoggerKismetLibrary::LogLocation(UObject* WorldContextObject, FVecto
 
 void UVisualLoggerKismetLibrary::LogSphere(UObject* WorldContextObject, FVector Center, float Radius, FString Text, FLinearColor ObjectColor, FName CategoryName, bool bAddToMessageLog, bool bWireframe)
 {
-#if ENABLE_VISUAL_LOG
+#if UE_DEBUG_RECORDING_ENABLED
 	FVisualLogger::SphereLogf(WorldContextObject, CategoryName, DefaultVerbosity
 		, Center, Radius, ObjectColor.ToFColor(true), bWireframe, TEXT("%s"), *Text);
 #endif
@@ -69,7 +69,7 @@ void UVisualLoggerKismetLibrary::LogSphere(UObject* WorldContextObject, FVector
 
 void UVisualLoggerKismetLibrary::LogCone(UObject* WorldContextObject, FVector Origin, FVector Direction, float Length, float Angle, FString Text, FLinearColor ObjectColor, FName CategoryName, bool bAddToMessageLog, bool bWireframe)
 {
-#if ENABLE_VISUAL_LOG
+#if UE_DEBUG_RECORDING_ENABLED
 	FVisualLogger::ConeLogf(WorldContextObject, CategoryName, DefaultVerbosity
 		, Origin, Direction, Length, Angle, ObjectColor.ToFColor(true), bWireframe, TEXT("%s"), *Text);
 #endif
@@ -81,7 +81,7 @@ void UVisualLoggerKismetLibrary::LogCone(UObject* WorldContextObject, FVector Or
 
 void UVisualLoggerKismetLibrary::LogCylinder(UObject* WorldContextObject, FVector Start, FVector End, float Radius, FString Text, FLinearColor ObjectColor, FName CategoryName, bool bAddToMessageLog, bool bWireframe)
 {
-#if ENABLE_VISUAL_LOG
+#if UE_DEBUG_RECORDING_ENABLED
 	FVisualLogger::CylinderLogf(WorldContextObject, CategoryName, DefaultVerbosity, Start, End, Radius, ObjectColor.ToFColor(true), bWireframe, TEXT("%s"), *Text);
 #endif
 	if (bAddToMessageLog)
@@ -92,7 +92,7 @@ void UVisualLoggerKismetLibrary::LogCylinder(UObject* WorldContextObject, FVecto
 
 void UVisualLoggerKismetLibrary::LogCapsule(UObject* WorldContextObject, FVector Base, float HalfHeight, float Radius, FQuat Rotation, FString Text, FLinearColor ObjectColor, FName CategoryName, bool bAddToMessageLog, bool bWireframe)
 {
-#if ENABLE_VISUAL_LOG
+#if UE_DEBUG_RECORDING_ENABLED
 	FVisualLogger::CapsuleLogf(WorldContextObject, CategoryName, DefaultVerbosity
 		, Base, HalfHeight, Radius, Rotation, ObjectColor.ToFColor(true), bWireframe, TEXT("%s"), *Text);
 #endif
@@ -104,7 +104,7 @@ void UVisualLoggerKismetLibrary::LogCapsule(UObject* WorldContextObject, FVector
 
 void UVisualLoggerKismetLibrary::LogBox(UObject* WorldContextObject, FBox Box, FString Text, FLinearColor ObjectColor, FName CategoryName, bool bAddToMessageLog, bool bWireframe)
 {
-#if ENABLE_VISUAL_LOG
+#if UE_DEBUG_RECORDING_ENABLED
 	FVisualLogger::BoxLogf(WorldContextObject, CategoryName, DefaultVerbosity
 		, Box, FMatrix::Identity, ObjectColor.ToFColor(true), bWireframe, TEXT("%s"), *Text);
 #endif
@@ -116,7 +116,7 @@ void UVisualLoggerKismetLibrary::LogBox(UObject* WorldContextObject, FBox Box, F
 
 void UVisualLoggerKismetLibrary::LogOrientedBox(UObject* WorldContextObject, FBox Box, FTransform Transform, FString Text, FLinearColor ObjectColor, FName CategoryName, bool bAddToMessageLog, bool bWireframe)
 {
-#if ENABLE_VISUAL_LOG
+#if UE_DEBUG_RECORDING_ENABLED
 	FVisualLogger::BoxLogf(WorldContextObject, CategoryName, DefaultVerbosity
 		, Box, Transform.ToMatrixWithScale(), ObjectColor.ToFColor(true), bWireframe, TEXT("%s"), *Text);
 #endif
@@ -128,7 +128,7 @@ void UVisualLoggerKismetLibrary::LogOrientedBox(UObject* WorldContextObject, FBo
 
 void UVisualLoggerKismetLibrary::LogArrow(UObject* WorldContextObject, const FVector SegmentStart, const FVector SegmentEnd, FString Text, FLinearColor ObjectColor, FName CategoryName, bool bAddToMessageLog, float ArrowHeadSize)
 {
-#if ENABLE_VISUAL_LOG
+#if UE_DEBUG_RECORDING_ENABLED
 	FVisualLogger::ArrowLineLogf(WorldContextObject, CategoryName, DefaultVerbosity, SegmentStart, SegmentEnd, ObjectColor.ToFColor(true), ArrowHeadSize, TEXT("%s"), *Text);
 #endif
 	if (bAddToMessageLog)
@@ -139,7 +139,7 @@ void UVisualLoggerKismetLibrary::LogArrow(UObject* WorldContextObject, const FVe
 
 void UVisualLoggerKismetLibrary::LogCircle(UObject* WorldContextObject, FVector Center, FVector UpAxis, float Radius, FString Text, FLinearColor ObjectColor, const float Thickness, FName CategoryName, bool bAddToMessageLog, bool bWireframe)
 {
-#if ENABLE_VISUAL_LOG
+#if UE_DEBUG_RECORDING_ENABLED
 	FVisualLogger::DiscLogf(WorldContextObject, CategoryName, DefaultVerbosity, Center, UpAxis, Radius, ObjectColor.ToFColor(true), static_cast<uint16>(Thickness), bWireframe, TEXT("%s"), *Text);
 #endif
 	if (bAddToMessageLog)
@@ -150,7 +150,7 @@ void UVisualLoggerKismetLibrary::LogCircle(UObject* WorldContextObject, FVector
 
 void UVisualLoggerKismetLibrary::LogSegment(UObject* WorldContextObject, const FVector SegmentStart, const FVector SegmentEnd, FString Text, FLinearColor ObjectColor, const float Thickness, FName CategoryName, bool bAddToMessageLog)
 {
-#if ENABLE_VISUAL_LOG
+#if UE_DEBUG_RECORDING_ENABLED
 	FVisualLogger::SegmentLogf(WorldContextObject, CategoryName, DefaultVerbosity, SegmentStart, SegmentEnd, ObjectColor.ToFColor(true), static_cast<uint16>(Thickness), TEXT("%s"), *Text);
 #endif
 	if (bAddToMessageLog)
diff --git a/Engine/Source/Runtime/Engine/Private/VisualLogger/VisualLoggerTraceDevice.cpp b/Engine/Source/Runtime/Engine/Private/VisualLogger/VisualLoggerTraceDevice.cpp
--- a/Engine/Source/Runtime/Engine/Private/VisualLogger/VisualLoggerTraceDevice.cpp
+++ b/Engine/Source/Runtime/Engine/Private/VisualLogger/VisualLoggerTraceDevice.cpp
@@ -1,7 +1,7 @@
 // Copyright Epic Games, Inc. All Rights Reserved.
 #include "VisualLogger/VisualLoggerTraceDevice.h"
 
-#if ENABLE_VISUAL_LOG
+#if UE_DEBUG_RECORDING_ENABLED
 
 #include "Serialization/BufferArchive.h"
 #include "ObjectTrace.h"
@@ -9,14 +9,14 @@
 #include "VisualLogger/VisualLoggerCustomVersion.h"
 #include "UObject/UE5MainStreamObjectVersion.h"
 
-UE_TRACE_CHANNEL_DEFINE(VisualLoggerChannel);
+UE_TRACE_MINIMAL_CHANNEL_DEFINE(VisualLoggerChannel);
 
-UE_TRACE_EVENT_BEGIN(VisualLogger, VisualLogEntry)
-	UE_TRACE_EVENT_FIELD(uint64, Cycle)
-	UE_TRACE_EVENT_FIELD(double, RecordingTime)
-	UE_TRACE_EVENT_FIELD(uint64, OwnerId)
-	UE_TRACE_EVENT_FIELD(uint8[], LogEntry)
-UE_TRACE_EVENT_END()
+UE_TRACE_MINIMAL_EVENT_BEGIN(VisualLogger, VisualLogEntry)
+	UE_TRACE_MINIMAL_EVENT_FIELD(uint64, Cycle)
+	UE_TRACE_MINIMAL_EVENT_FIELD(double, RecordingTime)
+	UE_TRACE_MINIMAL_EVENT_FIELD(uint64, OwnerId)
+	UE_TRACE_MINIMAL_EVENT_FIELD(uint8[], LogEntry)
+UE_TRACE_MINIMAL_EVENT_END()
 
 FVisualLoggerTraceDevice& FVisualLoggerTraceDevice::Get()
 {
@@ -36,15 +36,15 @@ void FVisualLoggerTraceDevice::Cleanup(bool bReleaseMemory)
 
 void FVisualLoggerTraceDevice::StartRecordingToFile(double TimeStamp)
 {
-#if UE_TRACE_ENABLED
-	UE::Trace::ToggleChannel(TEXT("VisualLogger"), true); 
+#if UE_TRACE_MINIMAL_ENABLED
+	UE::Trace::ToggleChannel(TEXT("VisualLogger"), true);
 #endif
 }
 
 void FVisualLoggerTraceDevice::StopRecordingToFile(double TimeStamp)
 {
-#if UE_TRACE_ENABLED
-	UE::Trace::ToggleChannel(TEXT("VisualLogger"), false); 
+#if UE_TRACE_MINIMAL_ENABLED
+	UE::Trace::ToggleChannel(TEXT("VisualLogger"), false);
 #endif
 }
 
@@ -59,8 +59,10 @@ void FVisualLoggerTraceDevice::SetFileName(const FString& InFileName)
 void FVisualLoggerTraceDevice::Serialize(const UObject* InLogOwner, const FName& InOwnerName, const FName& InOwnerDisplayName, const FName& InOwnerClassName, const FVisualLogEntry& InLogEntry)
 {
 #if OBJECT_TRACE_ENABLED
-	if (UE_TRACE_CHANNELEXPR_IS_ENABLED(VisualLoggerChannel))
+	if (UE_TRACE_MINIMAL_CHANNELEXPR_IS_ENABLED(VisualLoggerChannel))
 	{
+		TRACE_OBJECT(InLogOwner);
+
 		FBufferArchive Archive;
 		Archive.Reserve(1024);
 		Archive.UsingCustomVersion(EVisualLoggerVersion::GUID);
@@ -69,7 +71,7 @@ void FVisualLoggerTraceDevice::Serialize(const UObject* InLogOwner, const FName&
 
 		Archive << const_cast<FVisualLogEntry&>(InLogEntry);
 
-		UE_TRACE_LOG(VisualLogger, VisualLogEntry, VisualLoggerChannel)
+		UE_TRACE_MINIMAL_LOG(VisualLogger, VisualLogEntry, VisualLoggerChannel)
 			<< VisualLogEntry.Cycle(FPlatformTime::Cycles64())
 			<< VisualLogEntry.RecordingTime(FObjectTrace::GetWorldElapsedTime(InLogOwner->GetWorld()))
 			<< VisualLogEntry.OwnerId(FObjectTrace::GetObjectId(InLogOwner))
@@ -80,4 +82,9 @@ void FVisualLoggerTraceDevice::Serialize(const UObject* InLogOwner, const FName&
 	ImmediateRenderDelegate.ExecuteIfBound(InLogOwner, InLogEntry);
 }
 
-#endif
+bool FVisualLoggerTraceDevice::HasFlags(int32 InFlags) const
+{
+	return !!(InFlags & (EVisualLoggerDeviceFlags::CanSaveToFile | EVisualLoggerDeviceFlags::StoreLogsLocally));
+}
+
+#endif // UE_DEBUG_RECORDING_ENABLED
diff --git a/Engine/Source/Runtime/Engine/Private/VisualLogger/VisualLoggerTypes.cpp b/Engine/Source/Runtime/Engine/Private/VisualLogger/VisualLoggerTypes.cpp
--- a/Engine/Source/Runtime/Engine/Private/VisualLogger/VisualLoggerTypes.cpp
+++ b/Engine/Source/Runtime/Engine/Private/VisualLogger/VisualLoggerTypes.cpp
@@ -1,39 +1,35 @@
 // Copyright Epic Games, Inc. All Rights Reserved.
 
 #include "VisualLogger/VisualLoggerTypes.h"
-#include "VisualLogger/VisualLoggerDebugSnapshotInterface.h"
 
-#if ENABLE_VISUAL_LOG
-#include "Engine/World.h"
-#include "UObject/Interface.h"
-#include "UObject/UE5MainStreamObjectVersion.h"
-#include "UObject/FortniteMainBranchObjectVersion.h"
-#include "VisualLogger/VisualLoggerCustomVersion.h"
-#include "VisualLogger/VisualLogger.h"
-#endif
-
-namespace UE::VisualLogger
-{
-	const FName NAME_UnnamedCategory = "UnnamedCategory";
-}
+#if UE_DEBUG_RECORDING_ENABLED || UE_DEBUG_VISUALIZER_TOOL_ENABLED
 
-//----------------------------------------------------------------------//
-// UVisualLoggerDebugSnapshotInterface
-//----------------------------------------------------------------------//
-UVisualLoggerDebugSnapshotInterface::UVisualLoggerDebugSnapshotInterface(const FObjectInitializer& ObjectInitializer)
-	: Super(ObjectInitializer)
-{
-}
+#include "Math/Box.h"
+#include "Math/Quat.h"
 
-#if ENABLE_VISUAL_LOG
+#if UE_DEBUG_RECORDING_ENABLED
+#include "UObject/FortniteMainBranchObjectVersion.h"
+#include "UObject/UE5MainStreamObjectVersion.h"
+#include "VisualLogger/VisualLoggerCustomVersion.h"
+#if UE_DEBUG_RECORDING_USING_VLOG
+#include "Misc/DateTime.h"
+#endif // UE_DEBUG_RECORDING_USING_VLOG
+#endif // UE_DEBUG_RECORDING_ENABLED
 
+#if UE_DEBUG_VISUALIZER_TOOL_ENABLED || UE_DEBUG_RECORDING_USING_VLOG
 #define DEPRECATED_VISUAL_LOGGER_MAGIC_NUMBER 0xFAFAAFAF
 #define VISUAL_LOGGER_MAGIC_NUMBER_OLD_CUSTOM_VERSION 0xAFAFFAFA
 #define VISUAL_LOGGER_MAGIC_NUMBER_CUSTOM_VERSION_CONTAINER 0xBFBFBBFB
 #define VISUAL_LOGGER_MAGIC_NUMBER_LATEST VISUAL_LOGGER_MAGIC_NUMBER_CUSTOM_VERSION_CONTAINER
+#endif // UE_DEBUG_VISUALIZER_TOOL_ENABLED || UE_DEBUG_RECORDING_USING_VLOG
+
+namespace UE::VisualLogger
+{
+	const FName NAME_UnnamedCategory = "UnnamedCategory";
+}
 
 //----------------------------------------------------------------------//
-// FVisualLogShapeElement 
+// FVisualLogShapeElement
 //----------------------------------------------------------------------//
 FVisualLogShapeElement::FVisualLogShapeElement(const EVisualLoggerShapeElement InType)
 	: Type(InType)
@@ -312,7 +308,6 @@ void FVisualLogEntry::AddConvexElement(const TArray<FVector>& Points, const FNam
 	ElementsToDraw.Add(Element);
 }
 
-
 void FVisualLogEntry::AddHistogramData(const FVector2D& DataSample, const FName& CategoryName, ELogVerbosity::Type Verbosity, const FName& GraphName, const FName& DataName)
 {
 	FVisualLogHistogramSample Sample;
@@ -350,6 +345,7 @@ FVisualLogDataBlock& FVisualLogEntry::AddDataBlock(const FString& TagName, const
 	return DataBlocks[Index];
 }
 
+#if UE_DEBUG_RECORDING_ENABLED
 FArchive& operator<<(FArchive& Ar, FVisualLogDataBlock& Data)
 {
 	//!!!!!!!!!!!!!!!! Warning !!!!!!!!!!!!!!!
@@ -644,18 +640,6 @@ FArchive& operator<<(FArchive& Ar, FVisualLogDevice::FVisualLogEntryItem& FrameC
 	return Ar;
 }
 
-FString FVisualLoggerHelpers::GenerateTemporaryFilename(const FString& FileExt)
-{
-	return FString::Printf(TEXT("VTEMP_%s.%s"), *FDateTime::Now().ToString(), *FileExt);
-}
-
-FString FVisualLoggerHelpers::GenerateFilename(const FString& TempFileName, const FString& Prefix, double StartRecordingTime, double EndTimeStamp)
-{
-	const FString FullFilename = FString::Printf(TEXT("%s_%s"), *Prefix, *TempFileName);
-	const FString TimeFrameString = FString::Printf(TEXT("%" INT64_FMT "-%" INT64_FMT "_"), FMath::TruncToInt(StartRecordingTime), FMath::TruncToInt(EndTimeStamp));
-	return FullFilename.Replace(TEXT("VTEMP_"), *TimeFrameString, ESearchCase::CaseSensitive);
-}
-
 FArchive& FVisualLoggerHelpers::Serialize(FArchive& Ar, FName& Name)
 {
 	// Serialize the FName as a string
@@ -673,6 +657,23 @@ FArchive& FVisualLoggerHelpers::Serialize(FArchive& Ar, FName& Name)
 	return Ar;
 }
 
+#if UE_DEBUG_RECORDING_USING_VLOG
+FString FVisualLoggerHelpers::GenerateTemporaryFilename(const FString& FileExt)
+{
+	return FString::Printf(TEXT("VTEMP_%s.%s"), *FDateTime::Now().ToString(), *FileExt);
+}
+
+FString FVisualLoggerHelpers::GenerateFilename(const FString& TempFileName, const FString& Prefix, double StartRecordingTime, double EndTimeStamp)
+{
+	const FString FullFilename = FString::Printf(TEXT("%s_%s"), *Prefix, *TempFileName);
+	const FString TimeFrameString = FString::Printf(TEXT("%" INT64_FMT "-%" INT64_FMT "_"), FMath::TruncToInt(StartRecordingTime), FMath::TruncToInt(EndTimeStamp));
+	return FullFilename.Replace(TEXT("VTEMP_"), *TimeFrameString, ESearchCase::CaseSensitive);
+}
+#endif //UE_DEBUG_RECORDING_USING_VLOG
+
+#endif // UE_DEBUG_RECORDING_ENABLED
+
+#if UE_DEBUG_VISUALIZER_TOOL_ENABLED || UE_DEBUG_RECORDING_USING_VLOG
 FArchive& FVisualLoggerHelpers::Serialize(FArchive& Ar, TArray<FVisualLogDevice::FVisualLogEntryItem>& RecordedLogs)
 {
 	//!!!!!!!!!!!!!!!! Warning !!!!!!!!!!!!!!!
@@ -740,7 +741,9 @@ FArchive& FVisualLoggerHelpers::Serialize(FArchive& Ar, TArray<FVisualLogDevice:
 
 	return Ar;
 }
+#endif // UE_DEBUG_VISUALIZER_TOOL_ENABLED || UE_DEBUG_RECORDING_USING_VLOG
 
+#if UE_DEBUG_VISUALIZER_TOOL_ENABLED
 void FVisualLoggerHelpers::GetCategories(const FVisualLogEntry& EntryItem, TArray<FVisualLoggerCategoryVerbosityPair>& OutCategories)
 {
 	for (const auto& Element : EntryItem.Events)
@@ -780,6 +783,6 @@ void FVisualLoggerHelpers::GetHistogramCategories(const FVisualLogEntry& EntryIt
 		}
 	}
 }
+#endif // UE_DEBUG_VISUALIZER_TOOL_ENABLED
 
-#endif //ENABLE_VISUAL_LOG
-
+#endif // UE_DEBUG_RECORDING_ENABLED || UE_DEBUG_VISUALIZER_TOOL_ENABLED
\ No newline at end of file
diff --git a/Engine/Source/Runtime/Engine/Public/VisualLogger/VisualLogger.h b/Engine/Source/Runtime/Engine/Public/VisualLogger/VisualLogger.h
--- a/Engine/Source/Runtime/Engine/Public/VisualLogger/VisualLogger.h
+++ b/Engine/Source/Runtime/Engine/Public/VisualLogger/VisualLogger.h
@@ -1,33 +1,44 @@
 // Copyright Epic Games, Inc. All Rights Reserved.
 #pragma once
 
-#if UE_ENABLE_INCLUDE_ORDER_DEPRECATED_IN_5_7
+#if UE_ENABLE_INCLUDE_ORDER_DEPRECATED_IN_5_8
 #include "CoreMinimal.h"
-#endif // UE_ENABLE_INCLUDE_ORDER_DEPRECATED_IN_5_7 (switch to 5.8 in main)
-#include "Containers/Ticker.h"
 #include "Engine/EngineTypes.h"
+#include "Misc/ScopeRWLock.h"
+#include "Stats/Stats.h"
+#include "UObject/Object.h"
+#endif // UE_ENABLE_INCLUDE_ORDER_DEPRECATED_IN_5_8
+
+#include "Containers/Ticker.h"
 #include "EngineDefines.h"
+#include "Logging/LogMacros.h"
 #include "Misc/MTTransactionallySafeAccessDetector.h"
-#include "Misc/ScopeRWLock.h"
 #include "Misc/TransactionallySafeRWLock.h"
-#include "Stats/Stats.h"
 #include "Templates/IsArrayOrRefOfTypeByPredicate.h"
 #include "Templates/IsValidVariadicFunctionArg.h"
 #include "Traits/IsCharEncodingCompatibleWith.h"
-#include "UObject/Object.h"
 #include "UObject/ObjectKey.h"
 #include "VisualLogger/VisualLoggerTypes.h"
 
 #define UE_API ENGINE_API
 
-#if ENABLE_VISUAL_LOG
+// helper macros
+#define TEXT_EMPTY TEXT("")
+#define TEXT_NULL TEXT("NULL")
+#define TEXT_TRUE TEXT("TRUE")
+#define TEXT_FALSE TEXT("FALSE")
+#define TEXT_CONDITION(Condition) ((Condition) ? TEXT_TRUE : TEXT_FALSE)
 
+UE_API DECLARE_LOG_CATEGORY_EXTERN(LogVisual, Display, All);
+
+#define CONNECT_WITH_VLOG(Dest) UE_DEPRECATED_MACRO(5.8, "CONNECT_WITH_VLOG is no longer used")
+#define CONNECT_OBJECT_WITH_VLOG(Src, Dest) UE_DEPRECATED_MACRO(5.8, "CONNECT_OBJECT_WITH_VLOG is no longer used")
+
+#if UE_DEBUG_RECORDING_ENABLED
+
 #define REDIRECT_TO_VLOG(Dest) FVisualLogger::Redirect(this, Dest)
 #define REDIRECT_OBJECT_TO_VLOG(Src, Dest) FVisualLogger::Redirect(Src, Dest)
 
-#define CONNECT_WITH_VLOG(Dest)
-#define CONNECT_OBJECT_WITH_VLOG(Src, Dest)
-
 // Text, regular log
 #define UE_VLOG(LogOwner, CategoryName, Verbosity, Format, ...) if( FVisualLogger::IsRecording() ) FVisualLogger::CategorizedLogf(LogOwner, CategoryName, ELogVerbosity::Verbosity, Format, ##__VA_ARGS__)
 #define UE_CVLOG(Condition, LogOwner, CategoryName, Verbosity, Format, ...)  if(FVisualLogger::IsRecording() && Condition) {UE_VLOG(LogOwner, CategoryName, Verbosity, Format, ##__VA_ARGS__);} 
@@ -113,21 +124,10 @@
 #define UE_VLOG_COORDINATESYSTEM(LogOwner, CategoryName, Verbosity, AxisLoc, AxisRot, Scale, Color, Thickness, Format, ...) if(FVisualLogger::IsRecording()) FVisualLogger::CoordinateSystemLogf(LogOwner, CategoryName, ELogVerbosity::Verbosity, AxisLoc, AxisRot, Scale, Color, Thickness, Format, ##__VA_ARGS__)
 #define UE_CVLOG_COORDINATESYSTEM(Condition, LogOwner, CategoryName, Verbosity, AxisLoc, AxisRot, Scale, Color, Thickness, Format, ...) if(FVisualLogger::IsRecording() && Condition) {UE_VLOG_COORDINATESYSTEM(LogOwner, CategoryName, Verbosity, AxisLoc, AxisRot, Scale, Color, Thickness, Format, ##__VA_ARGS__);}
 
-#define DECLARE_VLOG_EVENT(EventName) extern FVisualLogEventBase EventName;
-#define DEFINE_VLOG_EVENT(EventName, Verbosity, UserFriendlyDesc) FVisualLogEventBase EventName(TEXT(#EventName), TEXT(UserFriendlyDesc), ELogVerbosity::Verbosity); 
+#else // UE_DEBUG_RECORDING_ENABLED
 
-#define UE_VLOG_EVENTS(LogOwner, TagNameToLog, ...) if(FVisualLogger::IsRecording()) FVisualLogger::EventLog(LogOwner, TagNameToLog, ##__VA_ARGS__)
-#define UE_CVLOG_EVENTS(Condition, LogOwner, TagNameToLog, ...) if(FVisualLogger::IsRecording() && Condition) {UE_VLOG_EVENTS(LogOwner, TagNameToLog, ##__VA_ARGS__);}
-#define UE_VLOG_EVENT_WITH_DATA(LogOwner, LogEvent, ...) if(FVisualLogger::IsRecording()) FVisualLogger::EventLog(LogOwner, LogEvent, ##__VA_ARGS__)
-#define UE_CVLOG_EVENT_WITH_DATA(Condition, LogOwner, LogEvent, ...) if(FVisualLogger::IsRecording() && Condition) {UE_VLOG_EVENT_WITH_DATA(LogOwner, LogEvent, ##__VA_ARGS__);}
-
-#define UE_IFVLOG(__code_block__) if( FVisualLogger::IsRecording() ) { __code_block__; }
-
-#else // if !ENABLE_VISUAL_LOG
 #define REDIRECT_TO_VLOG(Dest)
 #define REDIRECT_OBJECT_TO_VLOG(Src, Dest)
-#define CONNECT_WITH_VLOG(Dest)
-#define CONNECT_OBJECT_WITH_VLOG(Src, Dest)
 
 #define UE_VLOG(LogOwner, CategoryName, Verbosity, Format, ...)
 #define UE_CVLOG(Condition, LogOwner, CategoryName, Verbosity, Format, ...)
@@ -185,38 +185,36 @@
 #define UE_VLOG_COORDINATESYSTEM(LogOwner, CategoryName, Verbosity, AxisLoc, AxisRot, Scale, Color, Thickness, Format, ...)
 #define UE_CVLOG_COORDINATESYSTEM(Condition, LogOwner, CategoryName, Verbosity, AxisLoc, AxisRot, Scale, Color, Thickness, Format, ...)
 
+#endif // UE_DEBUG_RECORDING_ENABLED
+
+// Definition is the same with/without ENABLE_VISUAL_LOG, because it's dependent on UE_VLOG_UELOG which handles the difference.
+// Text, log with output to regular logs if the condition is met... regular log will still happen if condition is true
+// even when the compiler switch (ENABLE_VISUAL_LOG) is off.
+#define UE_CVLOG_UELOG(Condition, LogOwner, CategoryName, Verbosity, Format, ...) if (Condition) { UE_VLOG_UELOG(LogOwner, CategoryName, Verbosity, Format, ##__VA_ARGS__); }
+
+#if UE_DEBUG_RECORDING_USING_VLOG
+
+#define DECLARE_VLOG_EVENT(EventName) extern FVisualLogEventBase EventName;
+#define DEFINE_VLOG_EVENT(EventName, Verbosity, UserFriendlyDesc) FVisualLogEventBase EventName(TEXT(#EventName), TEXT(UserFriendlyDesc), ELogVerbosity::Verbosity); 
+#define UE_VLOG_EVENTS(LogOwner, TagNameToLog, ...) if (FVisualLogger::IsRecording()) FVisualLogger::EventLog(LogOwner, TagNameToLog, ##__VA_ARGS__)
+#define UE_CVLOG_EVENTS(Condition, LogOwner, TagNameToLog, ...) if (FVisualLogger::IsRecording() && Condition) {UE_VLOG_EVENTS(LogOwner, TagNameToLog, ##__VA_ARGS__);}
+#define UE_VLOG_EVENT_WITH_DATA(LogOwner, LogEvent, ...) if (FVisualLogger::IsRecording()) FVisualLogger::EventLog(LogOwner, LogEvent, ##__VA_ARGS__)
+#define UE_CVLOG_EVENT_WITH_DATA(Condition, LogOwner, LogEvent, ...) if (FVisualLogger::IsRecording() && Condition) {UE_VLOG_EVENT_WITH_DATA(LogOwner, LogEvent, ##__VA_ARGS__);}
+#define UE_IFVLOG(__code_block__) if (FVisualLogger::IsRecording()) { __code_block__; }
+
+#else // UE_DEBUG_RECORDING_USING_VLOG
+
 #define DECLARE_VLOG_EVENT(EventName)
 #define DEFINE_VLOG_EVENT(EventName, Verbosity, UserFriendlyDesc)
 #define UE_VLOG_EVENTS(LogOwner, TagNameToLog, ...) 
-#define UE_CVLOG_EVENTS(Condition, LogOwner, TagNameToLog, ...) 
+#define UE_CVLOG_EVENTS(Condition, LogOwner, TagNameToLog, ...)
 #define UE_VLOG_EVENT_WITH_DATA(LogOwner, LogEvent, ...)
 #define UE_CVLOG_EVENT_WITH_DATA(Condition, LogOwner, LogEvent, ...)
-
 #define UE_IFVLOG(__code_block__)
 
-#endif //ENABLE_VISUAL_LOG
+#endif // UE_DEBUG_RECORDING_USING_VLOG
 
-// Definition is the same with/without ENABLE_VISUAL_LOG, because it's dependent on UE_VLOG_UELOG which handles the difference.
-// Text, log with output to regular logs if the condition is met... regular log will still happen if condition is true
-// even when the compiler switch (ENABLE_VISUAL_LOG) is off.
-#define UE_CVLOG_UELOG(Condition, LogOwner, CategoryName, Verbosity, Format, ...) if (Condition) { UE_VLOG_UELOG(LogOwner, CategoryName, Verbosity, Format, ##__VA_ARGS__); }
+#if UE_DEBUG_RECORDING_ENABLED
 
-// helper macros
-#define TEXT_EMPTY TEXT("")
-#define TEXT_NULL TEXT("NULL")
-#define TEXT_TRUE TEXT("TRUE")
-#define TEXT_FALSE TEXT("FALSE")
-#define TEXT_CONDITION(Condition) ((Condition) ? TEXT_TRUE : TEXT_FALSE)
-
-UE_API DECLARE_LOG_CATEGORY_EXTERN(LogVisual, Display, All);
-
-#if ENABLE_VISUAL_LOG
-
-class FVisualLogDevice;
-class FVisualLogExtensionInterface;
-class UObject;
-class UWorld;
-struct FLogCategoryBase;
-
 DECLARE_DELEGATE_RetVal(FString, FVisualLogFilenameGetterDelegate);
 
@@ -241,7 +239,7 @@
 	static bool IsFilteredOut(const FLogCategoryBase& Category, ELogVerbosity::Type Verbosity)
 	{
 		return bUseVerbosityFilterWhenRecording
-			&& ((Verbosity & ::ELogVerbosity::VerbosityMask) > Category.GetCompileTimeVerbosity()
+			&& ((Verbosity & ELogVerbosity::VerbosityMask) > Category.GetCompileTimeVerbosity()
 				|| Category.IsSuppressed(Verbosity));
 	}
 
@@ -563,13 +561,75 @@
 		CoordinateSystemLogfImpl(LogOwner, Category, Verbosity, AxisLoc, AxisRot, Scale, Color, Thickness, (const TCHAR*)Fmt, Args...);
 	}
 
-	// Navigation data debug snapshot
-	static UE_API void NavigationDataDump(const UObject* LogOwner, const FLogCategoryBase& Category, const ELogVerbosity::Type Verbosity, const FBox& Box);
-	static UE_API void NavigationDataDump(const UObject* LogOwner, const FName& CategoryName, const ELogVerbosity::Type Verbosity, const FBox& Box);
+	// static getter
+	static UE_API FVisualLogger& Get();
 
+	UE_API FVisualLogger();
+	virtual ~FVisualLogger()
+	{
+	}
+
+	/** return information is vlog recording is enabled or not */
+	static bool IsRecording()
+	{
+		return !!bIsRecording;
+	}
+
+	/** Starts visual log collecting and recording */
+	UE_API void SetIsRecording(const bool bInIsRecording);
+
+	/** Starts visual log collecting and recording to insights traces (for Rewind Debugger)*/
+	UE_API void SetIsRecordingToTrace(const bool bInIsRecording);
+
+	/** Set log owner redirection from one object to another, to combine logs */
+	static UE_API void Redirect(const UObject* FromObject, const UObject* ToObject);
+
+	struct FObjectNames
+	{
+		FName Name;
+		FName DisplayName;
+		FName ClassName;
+	};
+
+	/** Sets function to call to get a timestamp instead of the default implementation (e.g. using a network synchronized clock instead of the local world time) */
+	UE_API void SetGetTimeStampFunc(TFunction<double(const UObject*)> Function);
+
+	/** Returns a current time stamp to associate with a recorded event that occurred on Object (used for ordering events on a timeline) */
+	UE_API double GetTimeStampForObject(const UObject* Object) const;
+
+	/** Configure whether VisLog should be using decorated, unique names */
+	UE_API void SetUseUniqueNames(const bool bEnable);
+
+private:
+	virtual void Serialize(const TCHAR* V, ELogVerbosity::Type Verbosity, const FName& Category) override { ensureMsgf(0, TEXT("Regular serialize is forbidden for visual logs")); }
+
+	/** Redirect internal implementation, not thread safe */
+	UE_API UObject* RedirectInternal(const UObject* FromObject, const UObject* ToObject);
+
+	/** Find redirects internal implementation, not thread safe */
+	UE_API UObject* FindRedirectionInternal(const UObject* Object) const;
+
+	/** Populates the struct of names associated to a given object */
+	FObjectNames CreateObjectNames(TNotNull<const UObject*> LogOwner);
+
+	/** Flushes entries recorded in the frame and clean invalid redirects if any. */
+	UE_API void Tick(float DeltaTime);
+
+	/** Cleanup invalid redirects */
+	UE_API void CleanupRedirects();
+
+public:
+	typedef TMap<FObjectKey, TWeakObjectPtr<const UObject>> FChildToOwnerRedirectionMap;
+
 	DECLARE_MULTICAST_DELEGATE_SixParams(FNavigationDataDump, const UObject* /*Object*/, const FName& /*CategoryName*/, const ELogVerbosity::Type /*Verbosity*/, const FBox& /*Box*/, const UWorld& /*World*/, FVisualLogEntry& /*CurrentEntry*/);
 	static UE_API FNavigationDataDump NavigationDataDumpDelegate;
 
+#if UE_DEBUG_RECORDING_USING_VLOG
+
+	/** Navigation data debug snapshot */
+	static UE_API void NavigationDataDump(const UObject* LogOwner, const FLogCategoryBase& Category, const ELogVerbosity::Type Verbosity, const FBox& Box);
+	static UE_API void NavigationDataDump(const UObject* LogOwner, const FName& CategoryName, const ELogVerbosity::Type Verbosity, const FBox& Box);
+
 	/** Log events */
 	static UE_API void EventLog(const UObject* LogOwner, const FName EventTag1, const FVisualLogEventBase& Event1, const FName EventTag2 = NAME_None, const FName EventTag3 = NAME_None, const FName EventTag4 = NAME_None, const FName EventTag5 = NAME_None, const FName EventTag6 = NAME_None);
 	static UE_API void EventLog(const UObject* LogOwner, const FName EventTag1, const FVisualLogEventBase& Event1, const FVisualLogEventBase& Event2);
@@ -579,85 +639,94 @@
 	static UE_API void EventLog(const UObject* LogOwner, const FName EventTag1, const FVisualLogEventBase& Event1, const FVisualLogEventBase& Event2, const FVisualLogEventBase& Event3, const FVisualLogEventBase& Event4, const FVisualLogEventBase& Event5, const FVisualLogEventBase& Event6);
 	static UE_API void EventLog(const UObject* LogOwner, const FVisualLogEventBase& Event1, const FName EventTag1 = NAME_None, const FName EventTag2 = NAME_None, const FName EventTag3 = NAME_None, const FName EventTag4 = NAME_None, const FName EventTag5 = NAME_None, const FName EventTag6 = NAME_None);
 
-	// static getter
-	static UE_API FVisualLogger& Get();
-
-	virtual ~FVisualLogger() {}
-
 	// called on engine shutdown to flush all, etc.
 	UE_API virtual void TearDown() override;
 
-	// Removes all logged data 
+	// Removes all logged data
 	UE_API void Cleanup(UWorld* OldWorld, bool bReleaseMemory = false);
 
 	// Use when a visual logger device has discarded all of its data, waiting for new data
 	UE_API void OnDataReset();
 
-	/** Set log owner redirection from one object to another, to combine logs */
-	static void Redirect(const UObject* FromObject, const UObject* ToObject)
-	{ 
-		FVisualLogger& Logger = FVisualLogger::Get();
-		UObject* NewRedirection = nullptr;
-		{
-			UE::TWriteScopeLock Lock(Logger.RedirectRWLock);
-			NewRedirection = Logger.RedirectInternal(FromObject, ToObject);
-		}
-		UE_CVLOG(FromObject != nullptr && NewRedirection != nullptr, FromObject, LogVisual, Log, TEXT("Redirected '%s' to '%s'"), *FromObject->GetName(), *NewRedirection->GetName());
-	}
-
 	/** find and return redirection object for given object*/
-	static UObject* FindRedirection(const UObject* Object)
-	{ 
-		FVisualLogger& Logger = FVisualLogger::Get();
-		UE::TReadScopeLock Lock(Logger.RedirectRWLock);
-		return Logger.FindRedirectionInternal(Object);
-	}
+	static UE_API UObject* FindRedirection(const UObject* Object);
 
 	/** blocks all categories from logging. It can be bypassed with the category allow list */
-	void BlockAllCategories(const bool bInBlock) { bBlockedAllCategories = bInBlock; }
+	void BlockAllCategories(const bool bInBlock)
+	{
+		bBlockedAllCategories = bInBlock;
+	}
 
 	/** checks if all categories are blocked */
-	bool IsBlockedForAllCategories() const { return !!bBlockedAllCategories; }
+	bool IsBlockedForAllCategories() const
+	{
+		return !!bBlockedAllCategories;
+	}
 
 	/** Returns category allow list for logging */
-	const TArray<FName>& GetCategoryAllowList() const { return CategoryAllowList; }
+	const TArray<FName>& GetCategoryAllowList() const
+	{
+		return CategoryAllowList;
+	}
 
-	bool IsCategoryAllowed(const FName& Name) const { return CategoryAllowList.Find(Name) != INDEX_NONE; }
+	bool IsCategoryAllowed(const FName& Name) const
+	{
+		return CategoryAllowList.Find(Name) != INDEX_NONE;
+	}
 
-	void AddCategoryToAllowList(FName Category) { CategoryAllowList.AddUnique(Category); }
+	void AddCategoryToAllowList(FName Category)
+	{
+		CategoryAllowList.AddUnique(Category);
+	}
 
-	void ClearCategoryAllowList() { CategoryAllowList.Reset(); }
+	void ClearCategoryAllowList()
+	{
+		CategoryAllowList.Reset();
+	}
 
 	/** Generates and returns Id unique for given timestamp - used to connect different logs between (ex. text log with geometry shape) */
 	UE_API int32 GetUniqueId(double Timestamp);
 
 	/** Starts visual log collecting and recording */
-	UE_API void SetIsRecording(const bool bInIsRecording);
-	/** return information is vlog recording is enabled or not */
-	inline static bool IsRecording() { return !!bIsRecording; }
-
-	/** Starts visual log collecting and recording */
 	UE_API void SetIsRecordingToFile(bool InIsRecording);
+
 	/** return information is vlog recording is enabled or not */
-	bool IsRecordingToFile() const { return !!bIsRecordingToFile; }
+	bool IsRecordingToFile() const
+	{
+		return !!bIsRecordingToFile;
+	}
+
 	/** disables recording to file and discards all data without saving it to file */
 	UE_API void DiscardRecordingToFile();
 
-	/** Starts visual log collecting and recording to insights traces (for Rewind Debugger)*/
-	UE_API void SetIsRecordingToTrace(const bool bInIsRecording);
+	void SetIsRecordingOnServer(const bool bInIsRecording)
+	{
+		bIsRecordingOnServer = bInIsRecording;
+	}
 
-	void SetIsRecordingOnServer(const bool bInIsRecording) { bIsRecordingOnServer = bInIsRecording; }
-	bool IsRecordingOnServer() const { return !!bIsRecordingOnServer; }
+	bool IsRecordingOnServer() const
+	{
+		return !!bIsRecordingOnServer;
+	}
 
-	/** Configure whether VisLog should be using decorated, unique names */
-	UE_API void SetUseUniqueNames(const bool bEnable);
-
 	/** Add visual logger output device */
-	void AddDevice(FVisualLogDevice* InDevice) { OutputDevices.AddUnique(InDevice); }
+	void AddDevice(FVisualLogDevice* InDevice)
+	{
+		OutputDevices.AddUnique(InDevice);
+	}
+
 	/** Remove visual logger output device */
-	void RemoveDevice(FVisualLogDevice* InDevice) { OutputDevices.RemoveSwap(InDevice); }
+	void RemoveDevice(FVisualLogDevice* InDevice)
+	{
+		OutputDevices.RemoveSwap(InDevice);
+	}
+
 	/** Remove visual logger output device */
-	const TArray<FVisualLogDevice*>& GetDevices() const { return OutputDevices; }
+	const TArray<FVisualLogDevice*>& GetDevices() const
+	{
+		return OutputDevices;
+	}
+
 	/** Check if log category can be recorded, verify before using GetEntryToWrite! */
 	UE_API bool IsCategoryLogged(const FLogCategoryBase& Category) const;
 
@@ -684,38 +753,56 @@
 
 	/** flush and serialize data if timestamp allows it */
 	UE_API virtual void Flush() override;
+
 	/** Moves all threads entries into the global entry map */
 	UE_API void FlushThreadsEntries();
 
 	/** FileName getter to set project specific file name for vlogs - highly encouraged to use FVisualLogFilenameGetterDelegate::CreateUObject with this */
-	void SetLogFileNameGetter(const FVisualLogFilenameGetterDelegate& InLogFileNameGetter) { LogFileNameGetter = InLogFileNameGetter; }
+	void SetLogFileNameGetter(const FVisualLogFilenameGetterDelegate& InLogFileNameGetter)
+	{
+		LogFileNameGetter = InLogFileNameGetter;
+	}
 
 	/** Register extension to use by LogVisualizer  */
-	void RegisterExtension(FName TagName, FVisualLogExtensionInterface* ExtensionInterface) { check(AllExtensions.Contains(TagName) == false); AllExtensions.Add(TagName, ExtensionInterface); }
+	void RegisterExtension(FName TagName, FVisualLogExtensionInterface* ExtensionInterface)
+	{
+		check(AllExtensions.Contains(TagName) == false); AllExtensions.Add(TagName, ExtensionInterface);
+	}
+
 	/**  Removes previously registered extension */
-	void UnregisterExtension(FName TagName, FVisualLogExtensionInterface* ExtensionInterface) { AllExtensions.Remove(TagName); }
+	void UnregisterExtension(FName TagName, FVisualLogExtensionInterface* ExtensionInterface)
+	{
+		AllExtensions.Remove(TagName);
+	}
+
 	/** returns extension identified by given tag */
-	FVisualLogExtensionInterface* GetExtensionForTag(const FName TagName) const { return AllExtensions.Contains(TagName) ? AllExtensions[TagName] : nullptr; }
+	FVisualLogExtensionInterface* GetExtensionForTag(const FName TagName) const
+	{
+		return AllExtensions.Contains(TagName) ? AllExtensions[TagName] : nullptr;
+	}
+
 	/** Returns reference to map with all registered extension */
-	const TMap<FName, FVisualLogExtensionInterface*>& GetAllExtensions() const { return AllExtensions; }
+	const TMap<FName, FVisualLogExtensionInterface*>& GetAllExtensions() const
+	{
+		return AllExtensions;
+	}
 
 	/** internal check for each usage of visual logger */
 	static UE_API bool CheckVisualLogInputInternal(const UObject* Object, const FName& CategoryName, ELogVerbosity::Type Verbosity, UWorld** OutWorld, FVisualLogEntry** OutCurrentEntry);
 
-	/** Returns a current time stamp to associate with a recorded event that occurred on Object (used for ordering events on a timeline) */
-	UE_API double GetTimeStampForObject(const UObject* Object) const;
-
-	/** Sets function to call to get a timestamp instead of the default implementation (e.g. using a network synchronized clock instead of the local world time) */
-	UE_API void SetGetTimeStampFunc(TFunction<double(const UObject*)> Function);
-
 	typedef TMap<FObjectKey, TArray<TWeakObjectPtr<const UObject>>> FOwnerToChildrenRedirectionMap;
 	static UE_API FOwnerToChildrenRedirectionMap& GetRedirectionMap(const UObject* InObject);
 
-	typedef TMap<FObjectKey, TWeakObjectPtr<const UObject>> FChildToOwnerRedirectionMap;
-	FChildToOwnerRedirectionMap& GetChildToOwnerRedirectionMap() { return ChildToOwnerMap; }
+	FChildToOwnerRedirectionMap& GetChildToOwnerRedirectionMap()
+	{
+		return ChildToOwnerMap;
+	}
 
 	typedef TMap<FObjectKey, TWeakObjectPtr<const UWorld>> FObjectToWorldMapType;
-	FObjectToWorldMapType& GetObjectToWorldMap() { return ObjectToWorldMap; }
+	FObjectToWorldMapType& GetObjectToWorldMap()
+	{
+		return ObjectToWorldMap;
+	}
 
 	UE_API void AddClassToAllowList(UClass& InClass);
 	UE_API bool IsClassAllowed(const UClass& InClass) const;
@@ -731,8 +818,6 @@
 	};
 
 private:
-	UE_API FVisualLogger();
-	virtual void Serialize(const TCHAR* V, ELogVerbosity::Type Verbosity, const FName& Category) override { ensureMsgf(0, TEXT("Regular serialize is forbidden for visual logs")); }
 
 	UE_API FVisualLoggerObjectEntryMap& GetThreadCurrentEntryMap();
 
@@ -744,18 +829,10 @@
 	/** Get global entry to write where all logs are combined, not thread safe */
 	UE_API FVisualLogEntry* GetEntryToWriteInternal(const UObject* Object, double TimeStamp, ECreateIfNeeded ShouldCreate);
 
-	/** Redirect internal implementation, not thread safe */
-	UE_API UObject* RedirectInternal(const UObject* FromObject, const UObject* ToObject);
-	/** Find redirects internal implementation, not thread safe */
-	UE_API UObject* FindRedirectionInternal(const UObject* Object) const;
-	/** Cleanup invalid redirects */
-	UE_API void CleanupRedirects();
 	/** Figure out all conditions if this entry is allowed to log */
 	UE_API void CalculateEntryAllowLogging(FVisualLogEntry* CurrentEntry, const UObject* LogOwner, const UObject* Object) const;
 
 protected:
-	/** Flushes entries recorded in the frame */
-	UE_API void Tick(float DeltaTime);
 
 	/**
 	 * Serializes a single entry and resets it.
@@ -765,89 +842,218 @@
 
 	/** Handle to the registered ticker to flush entries */
 	FTSTicker::FDelegateHandle TickerHandle;
-	
+
 	/** Array of output devices to redirect to */
 	TArray<FVisualLogDevice*> OutputDevices;
-	// Map for inter-objects redirections
+
+	/** Map for inter-objects redirections */
 	static UE_API TMap<const UWorld*, FOwnerToChildrenRedirectionMap> WorldToRedirectionMap;
 
-	// allowed classes - only instances of these classes will be logged. 
-	// if ClassAllowList is empty (default) everything will log
+	/** allowed classes - only instances of these classes will be logged.
+	 * if ClassAllowList is empty (default) everything will log
+	 */
 	TArray<UClass*> ClassAllowList;
 
-	// allowed objects - takes priority over class allow list and should be used to create exceptions in it
-	// if ObjectAllowList is empty (default) everything will log
-	// do NOT read from those pointers, they can be invalid!
+	/** allowed objects - takes priority over class allow list and should be used to create exceptions in it
+	 * if ObjectAllowList is empty (default) everything will log
+	 * do NOT read from those pointers, they can be invalid!
+	 */
 	TSet<FObjectKey> ObjectAllowList;
 
-	// list of categories that are still allowed to be logged when logging is blocking
+	/** list of categories that are still allowed to be logged when logging is blocking */
 	TArray<FName> CategoryAllowList;
 
-	// Visual Logger extensions map
+	/** Visual Logger extensions map */
 	TMap<FName, FVisualLogExtensionInterface*> AllExtensions;
 
-	// last generated unique id for given times tamp
+	/** last generated unique id for given times tamp */
 	TMap<double, int32> LastUniqueIds;
 
-	// Current entry with all data
+	/** Current entry with all data */
 	FVisualLoggerObjectEntryMap CurrentEntryPerObject;
 
-	// Threads current entry maps
+	/** Threads current entry maps */
 	TArray<FVisualLoggerObjectEntryMap*> ThreadCurrentEntryMaps;
 
-	// Multithread access detector for the list of thread local FVisualLoggerObjectEntryMap. Each map has its own MT detector for its entries.
-	UE_MT_DECLARE_RW_ACCESS_DETECTOR(ThreadLocalMapListAccessDetector);
-
-	// Read Write lock protecting object entries
-	mutable FTransactionallySafeRWLock EntryRWLock;
-	struct FObjectNames
-	{
-		FName Name;
-		FName DisplayName;
-		FName ClassName;
-	};
-	// Map to contain various names for Objects (they can be destroyed after while)
-	TMap<FObjectKey, FObjectNames> ObjectToNamesMap;
-
-	// Cached map to world information because it's just raw pointer and not real object
+	/** Cached map to world information because it's just raw pointer and not real object */
 	FObjectToWorldMapType ObjectToWorldMap;
-	// for any object that has requested redirection this map holds where we should
-	// redirect the traffic to
-	FChildToOwnerRedirectionMap ChildToOwnerMap;
-	// Read Write lock protecting redirection maps (ChildToOwnerMap and ObjectToWorldMap)
-	mutable FTransactionallySafeRWLock RedirectRWLock;
-	// if set all categories are blocked from logging
+
+	/** Delegate to set project specific file name for vlogs */
+	FVisualLogFilenameGetterDelegate LogFileNameGetter;
+
+	/** Read Write lock protecting object entries */
+	mutable FTransactionallySafeRWLock EntryRWLock;
+
+	/** start recording time */
+	double StartRecordingToFileTime;
+
+#if WITH_EDITOR
+	/** Handle for registering with PIEStarted to reset the EditorBaseTimeStamp */
+	FDelegateHandle PIEStartedHandle;
+#endif // WITH_EDITOR
+
+	/** Multithread access detector for the list of thread local FVisualLoggerObjectEntryMap. Each map has its own MT detector for its entries. */
+	UE_MT_DECLARE_RW_ACCESS_DETECTOR(ThreadLocalMapListAccessDetector);
+
+	/** if set all categories are blocked from logging */
 	bool bBlockedAllCategories : 1;
-	// if set we are recording to file
+
+	/** if set we are recording to file */
 	bool bIsRecordingToFile : 1;
-	// if set we are recording to insights trace 
-	bool bIsRecordingToTrace : 1;
-	// variable set (from cheat manager) when logging is active on server
+
+	/** variable set (from cheat manager) when logging is active on server */
 	bool bIsRecordingOnServer : 1;
-	// controls how we generate log names. When set to TRUE there's a lower 
-	// chance of name conflict, but it's more expensive
-	bool bForceUniqueLogNames : 1;
+
 	/** Indicates that entries were added/updated and that a flush is required */
 	bool bIsFlushRequired : 1;
+
+#else // ^^^ UE_DEBUG_RECORDING_USING_VLOG ^^^ / vvv INSIGHT TRACES ONLY vvv
+
+	/**
+	 * Not used when UE_DEBUG_RECORDING_USING_VLOG is not enabled; debug macros use CreateEntryInternal instead.
+	 * @note Method is still implemented to allow existing code to compile but will no longer produce logging.
+	 */
+	static UE_API bool CheckVisualLogInputInternal(const UObject* Object, const FName& CategoryName, ELogVerbosity::Type Verbosity, UWorld** OutWorld, FVisualLogEntry** OutCurrentEntry)
+	{
+		return false;
+	}
+
+	/** Creates and sets up a LogEntry intended to be traced synchronously */
+	static UE_API bool CreateEntryInternal(const UObject* Object, FVisualLogEntry& OutEntry, FVisualLogger::FObjectNames& Names);
+
+	/** Starts visual log collecting and recording */
+	void SetIsRecordingToFile(bool)
+	{
+	}
+
+	/** return information is vlog recording is enabled or not */
+	bool IsRecordingToFile() const
+	{
+		return false;
+	}
+
+	/** disables recording to file and discards all data without saving it to file */
+	void DiscardRecordingToFile()
+	{
+	}
+
+	void SetIsRecordingOnServer(const bool)
+	{
+	}
+
+	bool IsRecordingOnServer() const
+	{
+		return false;
+	}
+
+	void AddClassToAllowList(UClass&)
+	{
+	}
+
+	/** Register extension to use by LogVisualizer  */
+	void RegisterExtension(FName TagName, FVisualLogExtensionInterface*)
+	{
+	}
+
+	/**  Removes previously registered extension */
+	void UnregisterExtension(FName TagName, FVisualLogExtensionInterface*)
+	{
+	}
+
+	const TMap<FName, FVisualLogExtensionInterface*>& GetAllExtensions() const
+	{
+		// Not supported when only recording traces
+		static TMap<FName, FVisualLogExtensionInterface*> EmptyExtensionList;
+		return EmptyExtensionList;
+	}
+
+	bool IsCategoryLogged(const FLogCategoryBase&) const
+	{
+		return IsRecording();
+	}
+
+	bool IsCategoryAllowed(const FName&) const
+	{
+		return true;
+	}
+
+	void AddCategoryToAllowList(FName)
+	{
+	}
+
+	void ClearCategoryAllowList()
+	{
+	}
+
+	void BlockAllCategories(const bool)
+	{
+	}
+
+	static FVisualLogEntry* GetEntryToWrite(const UObject*, const FLogCategoryBase&)
+	{
+		return nullptr;
+	}
+
+	void Cleanup(UWorld*, bool bReleaseMemory = false)
+	{
+	}
+
+	/** Add visual logger output device */
+	void AddDevice(FVisualLogDevice* InDevice)
+	{
+	}
+
+	/** Remove visual logger output device */
+	void RemoveDevice(FVisualLogDevice*)
+	{
+	}
+
+	const TArray<FVisualLogDevice*>& GetDevices() const
+	{
+		static TArray<FVisualLogDevice*> EmptyDeviceList;
+		return EmptyDeviceList;
+	}
+
+	bool ExecuteOnLastEntryForObject(const UObject*, TFunctionRef<void(FVisualLogEntry&)>)
+	{
+		return false;
+	}
+
+	void SetLogFileNameGetter(const FVisualLogFilenameGetterDelegate& InLogFileNameGetter)
+	{
+	}
+
+#endif // UE_DEBUG_RECORDING_USING_VLOG
+
+private:
+	/** controls how we generate log names. When set to TRUE there's a lower
+	 * chance of name conflict, but it's more expensive
+	 */
+	bool bForceUniqueLogNames : 1;
+
 	/** Indicates there are entries in the redirection map that are invalid */
 	mutable bool bContainsInvalidRedirects : 1;
-	// start recording time
-	double StartRecordingToFileTime;
-	/** Delegate to set project specific file name for vlogs */
-	FVisualLogFilenameGetterDelegate LogFileNameGetter;
 
+	/** if set we are recording to insights trace */
+	bool bIsRecordingToTrace : 1;
+
+	/** Read Write lock protecting redirection maps (ChildToOwnerMap and ObjectToWorldMap) */
+	mutable FTransactionallySafeRWLock RedirectRWLock;
+
+	/** for any object that has requested redirection this map holds where we should
+	 * redirect the traffic to
+	 */
+	FChildToOwnerRedirectionMap ChildToOwnerMap;
+
+	/** Map to contain various names for Objects (they can be destroyed after while) */
+	TMap<FObjectKey, FObjectNames> ObjectToNamesMap;
+
 	/** Function to call when recording the absolute time stamp of an event. Useful for manually aligning events across multiple instances (e.g. using FPlatformTime::Seconds() rather than WorldTime) */
 	TFunction<double(const UObject*)> GetTimeStampFunc;
 
-#if WITH_EDITOR
-	/** Handle for registering with PIEStarted to reset the EditorBaseTimeStamp */
-	FDelegateHandle PIEStartedHandle;
-#endif
-
-	// if set we are recording and collecting all vlog data
+	/** if set we are recording and collecting all vlog data */
 	static UE_API int32 bIsRecording;
 };
 
-#endif //ENABLE_VISUAL_LOG
-
+#endif //UE_DEBUG_RECORDING_ENABLED
 #undef UE_API
diff --git a/Engine/Source/Runtime/Engine/Public/VisualLogger/VisualLoggerDebugSnapshotInterface.h b/Engine/Source/Runtime/Engine/Public/VisualLogger/VisualLoggerDebugSnapshotInterface.h
--- a/Engine/Source/Runtime/Engine/Public/VisualLogger/VisualLoggerDebugSnapshotInterface.h
+++ b/Engine/Source/Runtime/Engine/Public/VisualLogger/VisualLoggerDebugSnapshotInterface.h
@@ -2,7 +2,10 @@
 
 #pragma once
 
+#if UE_ENABLE_INCLUDE_ORDER_DEPRECATED_IN_5_8
 #include "CoreMinimal.h"
+#endif // UE_ENABLE_INCLUDE_ORDER_DEPRECATED_IN_5_8
+
 #include "UObject/ObjectMacros.h"
 #include "UObject/Interface.h"
 #include "VisualLoggerDebugSnapshotInterface.generated.h"
@@ -10,12 +13,12 @@
 UINTERFACE(MinimalAPI, meta = (CannotImplementInterfaceInBlueprint))
 class UVisualLoggerDebugSnapshotInterface : public UInterface
 {
-	GENERATED_UINTERFACE_BODY()
+	GENERATED_BODY()
 };
 
 class IVisualLoggerDebugSnapshotInterface
 {
-	GENERATED_IINTERFACE_BODY()
+	GENERATED_BODY()
 public:
 	virtual void GrabDebugSnapshot(struct FVisualLogEntry* Snapshot) const {}
 };
diff --git a/Engine/Source/Runtime/Engine/Public/VisualLogger/VisualLoggerDefines.h b/Engine/Source/Runtime/Engine/Public/VisualLogger/VisualLoggerDefines.h
new file mode 100644
--- /dev/null
+++ b/Engine/Source/Runtime/Engine/Public/VisualLogger/VisualLoggerDefines.h
@@ -0,0 +1,29 @@
+// Copyright Epic Games, Inc. All Rights Reserved.
+#pragma once
+
+#include "EngineDefines.h"
+#include "Misc/Build.h"
+#include "Trace/Config.h"
+
+/** Indicates whether the UE_DEBUG_* macros will record the element (VisualLogger and/or insight traces) */
+#ifndef UE_DEBUG_RECORDING_ENABLED
+#define UE_DEBUG_RECORDING_ENABLED ((ENABLE_VISUAL_LOG) || (UE_TRACE_MINIMAL_ENABLED))
+#endif
+
+/**
+ * Indicates whether the UE_DEBUG_* macros will use the VisualLogger recording process
+ * to store entries and write to all registered output devices or use insight traces directly
+ */
+#ifndef UE_DEBUG_RECORDING_USING_VLOG
+#define UE_DEBUG_RECORDING_USING_VLOG ((UE_DEBUG_RECORDING_ENABLED) && !(UE_BUILD_SHIPPING || UE_BUILD_TEST))
+#endif
+
+/** Indicates whether the UE_DEBUG_* macros will use immediate debug draw visualization */
+#ifndef UE_DEBUG_RENDERING_ENABLED
+#define UE_DEBUG_RENDERING_ENABLED ((ENABLE_VISUAL_LOG) || (UE_ENABLE_DEBUG_DRAWING))
+#endif
+
+/** Indicates whether the analysis of the recorded data is active and dedicated UI and tools should be compiled */
+#ifndef UE_DEBUG_VISUALIZER_TOOL_ENABLED
+#define UE_DEBUG_VISUALIZER_TOOL_ENABLED ((ENABLE_VISUAL_LOG) || ((PLATFORM_DESKTOP) && (WITH_EDITOR)))
+#endif
\ No newline at end of file
diff --git a/Engine/Source/Runtime/Engine/Public/VisualLogger/VisualLoggerTypes.h b/Engine/Source/Runtime/Engine/Public/VisualLogger/VisualLoggerTypes.h
--- a/Engine/Source/Runtime/Engine/Public/VisualLogger/VisualLoggerTypes.h
+++ b/Engine/Source/Runtime/Engine/Public/VisualLogger/VisualLoggerTypes.h
@@ -1,20 +1,36 @@
 // Copyright Epic Games, Inc. All Rights Reserved.
 #pragma once
 
+#if UE_ENABLE_INCLUDE_ORDER_DEPRECATED_IN_5_8
 #include "CoreMinimal.h"
 #include "EngineDefines.h"
-#include "VisualLoggerCustomVersion.h"
+#endif // UE_ENABLE_INCLUDE_ORDER_DEPRECATED_IN_5_8
 
-class AActor;
-class UCanvas;
-struct FLogEntryItem;
+#include "VisualLoggerDefines.h"
 
+#define DECLARE_ENUM_TO_STRING(EnumType) FString EnumToString(const EnumType Value)
 #define DEFINE_ENUM_TO_STRING(EnumType, EnumPackage) FString EnumToString(const EnumType Value) \
 { \
 	static const UEnum* TypeEnum = FindObject<UEnum>(nullptr, TEXT(EnumPackage) TEXT(".") TEXT(#EnumType)); \
 	return TypeEnum->GetNameStringByIndex(static_cast<int32>(Value)); \
 }
-#define DECLARE_ENUM_TO_STRING(EnumType) FString EnumToString(const EnumType Value)
+
+#if (UE_DEBUG_RECORDING_ENABLED) || (UE_DEBUG_VISUALIZER_TOOL_ENABLED)
+
+#include "Containers/Map.h"
+#include "Logging/LogVerbosity.h"
+#include "Math/Color.h"
+#include "Math/Matrix.h"
+#include "UObject/NameTypes.h"
+
+#if UE_DEBUG_VISUALIZER_TOOL_ENABLED
+#include "VisualLoggerTypes.h"
+class AActor;
+class UCanvas;
+struct FLogEntryItem;
+#else
+class FVisualLogExtensionInterface;
+#endif // UE_DEBUG_VISUALIZER_TOOL_ENABLED
 
 namespace UE::VisualLogger
 {
@@ -28,7 +44,8 @@ enum class ECreateIfNeeded : int8
 	Create = 1,
 };
 
-// flags describing VisualLogger device's features
+#if UE_DEBUG_RECORDING_ENABLED
+/** flags describing VisualLogger device's features */
 namespace EVisualLoggerDeviceFlags
 {
 	enum Type
@@ -38,8 +55,9 @@ namespace EVisualLoggerDeviceFlags
 		StoreLogsLocally = 2,
 	};
 }
+#endif // UE_DEBUG_RECORDING_ENABLED
 
-//types of shape elements
+/** types of shape elements */
 enum class EVisualLoggerShapeElement : uint8
 {
 	Invalid = 0,
@@ -67,7 +85,6 @@ enum class EVisualLoggerShapeElement : uint8
 	// serialization new enum values need to be added at the end
 };
 
-#if ENABLE_VISUAL_LOG
 struct FVisualLogEventBase
 {
 	const FString Name;
@@ -150,10 +167,8 @@ struct FVisualLogShapeElement
 	{
 	}
 
-PRAGMA_DISABLE_DEPRECATION_WARNINGS
 	FVisualLogShapeElement(const FVisualLogShapeElement& Other) = default;
 	FVisualLogShapeElement& operator=(const FVisualLogShapeElement& Other) = default;
-PRAGMA_ENABLE_DEPRECATION_WARNINGS
 
 	ENGINE_API explicit FVisualLogShapeElement(EVisualLoggerShapeElement InType = EVisualLoggerShapeElement::Invalid);
 	FVisualLogShapeElement(const FString& InDescription, const FColor& InColor, uint16 InThickness, const FName& InCategory);
@@ -169,8 +184,6 @@ PRAGMA_ENABLE_DEPRECATION_WARNINGS
 	union
 	{
 		uint16 Thickness = 0;
-		UE_DEPRECATED(5.6, "Use Thickness instead")
-		uint16 Thicknes;
 		uint16 Radius;
 		uint16 Mag;
 	};
@@ -199,11 +212,9 @@ struct FVisualLogDataBlock
 	TArray<uint8> Data;
 	int32 UniqueId = INDEX_NONE;
 };
-#endif  //ENABLE_VISUAL_LOG
 
 struct FVisualLogEntry final
 {
-#if ENABLE_VISUAL_LOG
 	/** For absolute position of events along a timeline (can involve multiple worlds/game instances such as clients and server) */
 	double TimeStamp = -1.0;
 	/** The time of the event according to its UWorld (can vary widely between game instances such as clients and server) */
@@ -221,7 +232,7 @@ struct FVisualLogEntry final
 	TArray<FVisualLogStatusCategory> Status;
 	TArray<FVisualLogShapeElement> ElementsToDraw;
 	TArray<FVisualLogHistogramSample> HistogramSamples;
-	TArray<FVisualLogDataBlock>	DataBlocks;
+	TArray<FVisualLogDataBlock> DataBlocks;
 
 	FVisualLogEntry() = default;
 
@@ -288,12 +299,8 @@ struct FVisualLogEntry final
 
 	// Moves all content to provided entry and reseting our content.
 	ENGINE_API void MoveTo(FVisualLogEntry& Other);
-
-#endif // ENABLE_VISUAL_LOG
 };
 
-#if  ENABLE_VISUAL_LOG
-
 /**
  * Interface for Visual Logger Device
  */
@@ -323,7 +330,7 @@ public:
 	virtual void StopRecordingToFile(double TimeStamp) { /* Empty */ }
 	virtual void DiscardRecordingToFile() { /* Empty */ }
 	virtual void SetFileName(const FString& InFileName) { /* Empty */ }
-	virtual void GetRecordedLogs(TArray<FVisualLogDevice::FVisualLogEntryItem>& OutLogs) const { /* Empty */ }
+	virtual void GetRecordedLogs(TArray<FVisualLogEntryItem>& OutLogs) const { /* Empty */ }
 	virtual bool HasFlags(int32 InFlags) const { return false; }
 	FGuid GetSessionGUID() const { return SessionGUID; }
 	uint32 GetShortSessionID() const { return SessionGUID[0]; }
@@ -331,6 +338,14 @@ protected:
 	FGuid SessionGUID;
 };
 
+#if UE_DEBUG_VISUALIZER_TOOL_ENABLED || UE_DEBUG_RECORDING_USING_VLOG
+namespace FVisualLoggerHelpers
+{
+	extern ENGINE_API FArchive& Serialize(FArchive& Ar, TArray<FVisualLogDevice::FVisualLogEntryItem>& RecordedLogs);
+};
+#endif
+
+#if UE_DEBUG_VISUALIZER_TOOL_ENABLED
 struct FVisualLoggerCategoryVerbosityPair
 {
 	FVisualLoggerCategoryVerbosityPair(FName Category, ELogVerbosity::Type InVerbosity) : CategoryName(Category), Verbosity(InVerbosity) {}
@@ -345,14 +360,10 @@ struct FVisualLoggerCategoryVerbosityPair
 	}
 };
 
-struct FVisualLoggerHelpers
+namespace FVisualLoggerHelpers
 {
-	static ENGINE_API FString GenerateTemporaryFilename(const FString& FileExt);
-	static ENGINE_API FString GenerateFilename(const FString& TempFileName, const FString& Prefix, double StartRecordingTime, double EndTimeStamp);
-	static ENGINE_API FArchive& Serialize(FArchive& Ar, FName& Name);
-	static ENGINE_API FArchive& Serialize(FArchive& Ar, TArray<FVisualLogDevice::FVisualLogEntryItem>& RecordedLogs);
-	static ENGINE_API void GetCategories(const FVisualLogEntry& RecordedLogs, TArray<FVisualLoggerCategoryVerbosityPair>& OutCategories);
-	static ENGINE_API void GetHistogramCategories(const FVisualLogEntry& RecordedLogs, TMap<FString, TArray<FString> >& OutCategories);
+	extern ENGINE_API void GetCategories(const FVisualLogEntry& RecordedLogs, TArray<FVisualLoggerCategoryVerbosityPair>& OutCategories);
+	extern ENGINE_API void GetHistogramCategories(const FVisualLogEntry& RecordedLogs, TMap<FString, TArray<FString>>& OutCategories);
 };
 
 struct IVisualLoggerEditorInterface
@@ -380,18 +391,34 @@ public:
 	virtual void DrawData(IVisualLoggerEditorInterface* EdInterface, UCanvas* Canvas) = 0;
 
 	virtual void OnItemsSelectionChanged(IVisualLoggerEditorInterface* EdInterface) {};
-	virtual void OnLogLineSelectionChanged(IVisualLoggerEditorInterface* EdInterface, TSharedPtr<struct FLogEntryItem> SelectedItem, int64 UserData) {};
+	virtual void OnLogLineSelectionChanged(IVisualLoggerEditorInterface* EdInterface, TSharedPtr<FLogEntryItem> SelectedItem, int64 UserData) {};
 	virtual void OnScrubPositionChanged(IVisualLoggerEditorInterface* EdInterface, double NewScrubPosition, bool bScrubbing) {}
 };
 
-ENGINE_API  FArchive& operator<<(FArchive& Ar, FVisualLogDevice::FVisualLogEntryItem& FrameCacheItem);
-ENGINE_API  FArchive& operator<<(FArchive& Ar, FVisualLogDataBlock& Data);
-ENGINE_API  FArchive& operator<<(FArchive& Ar, FVisualLogHistogramSample& Sample);
-ENGINE_API  FArchive& operator<<(FArchive& Ar, FVisualLogShapeElement& Element);
-ENGINE_API  FArchive& operator<<(FArchive& Ar, FVisualLogEvent& Event);
-ENGINE_API  FArchive& operator<<(FArchive& Ar, FVisualLogLine& LogLine);
-ENGINE_API  FArchive& operator<<(FArchive& Ar, FVisualLogStatusCategory& Status);
-ENGINE_API  FArchive& operator<<(FArchive& Ar, FVisualLogEntry& LogEntry);
+ENGINE_API FArchive& operator<<(FArchive& Ar, FVisualLogDevice::FVisualLogEntryItem& FrameCacheItem);
+#endif // UE_DEBUG_VISUALIZER_TOOL_ENABLED
+
+#if UE_DEBUG_RECORDING_ENABLED
+namespace FVisualLoggerHelpers
+{
+	extern ENGINE_API FArchive& Serialize(FArchive& Ar, FName& Name);
+
+#if UE_DEBUG_RECORDING_USING_VLOG
+	extern ENGINE_API FString GenerateTemporaryFilename(const FString& FileExt);
+	extern ENGINE_API FString GenerateFilename(const FString& TempFileName, const FString& Prefix, double StartRecordingTime, double EndTimeStamp);
+#endif // UE_DEBUG_RECORDING_USING_VLOG
+
+}
+
+ENGINE_API FArchive& operator<<(FArchive& Ar, FVisualLogDataBlock& Data);
+ENGINE_API FArchive& operator<<(FArchive& Ar, FVisualLogHistogramSample& Sample);
+ENGINE_API FArchive& operator<<(FArchive& Ar, FVisualLogShapeElement& Element);
+ENGINE_API FArchive& operator<<(FArchive& Ar, FVisualLogEvent& Event);
+ENGINE_API FArchive& operator<<(FArchive& Ar, FVisualLogLine& LogLine);
+ENGINE_API FArchive& operator<<(FArchive& Ar, FVisualLogStatusCategory& Status);
+ENGINE_API FArchive& operator<<(FArchive& Ar, FVisualLogEntry& LogEntry);
+
+#endif // UE_DEBUG_RECORDING_ENABLED
 
 inline
 FVisualLogEvent::FVisualLogEvent(const FVisualLogEventBase& Event)
@@ -487,7 +514,6 @@ FColor FVisualLogShapeElement::GetFColor() const
 	return RetColor;
 }
 
-
 inline
 int32 FVisualLogEntry::FindStatusIndex(const FString& CategoryName)
 {
@@ -502,4 +528,4 @@ int32 FVisualLogEntry::FindStatusIndex(const FString& CategoryName)
 	return INDEX_NONE;
 }
 
-#endif // ENABLE_VISUAL_LOG
+#endif // UE_DEBUG_RECORDING_ENABLED || UE_DEBUG_VISUALIZER_TOOL_ENABLED
\ No newline at end of file
