Allow LocalFogVolume to be spatially loaded (default: false).
Show WorldPartition category in the Details panel.

diff --git a/Engine/Source/Runtime/Engine/Classes/Engine/LocalFogVolume.h b/Engine/Source/Runtime/Engine/Classes/Engine/LocalFogVolume.h
--- a/Engine/Source/Runtime/Engine/Classes/Engine/LocalFogVolume.h
+++ b/Engine/Source/Runtime/Engine/Classes/Engine/LocalFogVolume.h
@@ -12,7 +12,7 @@ class ULocalFogVolumeComponent;
  *	Actor used to position a local fog volume in the scene.
  *	@see https://dev.epicgames.com/documentation/en-us/unreal-engine/local-fog-volumes-in-unreal-engine
  */
-UCLASS(showcategories = (Movement, Rendering, Transformation, DataLayers, "Input|MouseInput", "Input|TouchInput"), ClassGroup = Fog, hidecategories = (Info, Object, Input), MinimalAPI)
+UCLASS(showcategories = (Movement, Rendering, Transformation, DataLayers, WorldPartition, "Input|MouseInput", "Input|TouchInput"), ClassGroup = Fog, hidecategories = (Info, Object, Input), MinimalAPI)
 class ALocalFogVolume : public AInfo
 {
 	GENERATED_UCLASS_BODY()
@@ -30,5 +30,11 @@ public:
 
 	/** Returns LocalFogVolumeVolume subobject **/
 	ULocalFogVolumeComponent* GetComponent() const { return LocalFogVolumeVolume; }
+	
+#if WITH_EDITOR
+	//Allow LocalFogVolume instances to be spatially loaded
+	virtual bool CanChangeIsSpatiallyLoadedFlag() const override { return true; }
+#endif
+
 };
 
