Fix invalid output type returned by UMaterialExpressionstaticSwitchParameter causing it to be unusable with Texture2DSamplers or similar expresssions that do not accept float as input.

Override the base GetOutputType to keep the necessary behavior when MaterialAttributes or Subtrate is fed as input but return MCT_Uknown otherwise to allow it to connect to Expressions that require anything other
than MCT_Float<x>.

[RN] minor | Rendering Architecture

diff --git a/Engine/Source/Runtime/Engine/Private/Materials/MaterialExpressions.cpp b/Engine/Source/Runtime/Engine/Private/Materials/MaterialExpressions.cpp
--- a/Engine/Source/Runtime/Engine/Private/Materials/MaterialExpressions.cpp
+++ b/Engine/Source/Runtime/Engine/Private/Materials/MaterialExpressions.cpp
@@ -951,8 +951,23 @@ void UMaterialExpression::CopyMaterialExpressions(const TArray<UMaterialExpressi
 		OutNewComments.Add(NewComment);
 	}
 }
-#endif // WITH_EDITOR
 
+EMaterialValueType UMaterialExpression::GetCompositeOutput(int32 OutputIndex)
+{
+	// Material attributes need to be tested first to work when plugged in main root node (to not return MCT_Substrate when Substrate mateiral is fed)
+	if (IsResultMaterialAttributes(OutputIndex))
+	{
+		return MCT_MaterialAttributes;
+	}
+
+	if (IsResultSubstrateMaterial(OutputIndex))
+	{
+		return MCT_Substrate;
+	}
+
+	return MCT_Unknown;
+}
+#endif // WITH_EDITOR
 
 void UMaterialExpression::Serialize(FStructuredArchive::FRecord Record)
 {
@@ -1898,43 +1913,37 @@ uint32 UMaterialExpression::GetOutputType(int32 OutputIndex)
 {
 	// different outputs should be defined by sub classed expressions 
 
-	// Material attributes need to be tested first to work when plugged in main root node (to not return MCT_Substrate when Substrate mateiral is fed)
-	if (IsResultMaterialAttributes(OutputIndex))
+	EMaterialValueType CompositeResultOutput = GetCompositeOutput(OutputIndex);
+	if (CompositeResultOutput != MCT_Unknown)
 	{
-		return MCT_MaterialAttributes;
+		return CompositeResultOutput;
 	}
-	else if (IsResultSubstrateMaterial(OutputIndex))
-	{
-		return MCT_Substrate;
-	}
-	else
+
+	FExpressionOutput& Output = GetOutputs()[OutputIndex];
+	if (Output.Mask)
 	{
-		FExpressionOutput& Output = GetOutputs()[OutputIndex];
-		if (Output.Mask)
-		{
-			int32 MaskChannelCount = (Output.MaskR ? 1 : 0)
-									+ (Output.MaskG ? 1 : 0)
-									+ (Output.MaskB ? 1 : 0)
-									+ (Output.MaskA ? 1 : 0);
-			switch (MaskChannelCount)
-			{
-			case 1:
-				return MCT_Float;
-			case 2:
-				return MCT_Float2;
-			case 3:
-				return MCT_Float3;
-			case 4:
-				return MCT_Float4;
-			default:
-				return MCT_Unknown;
-			}
-		}
-		else
+		int32 MaskChannelCount = (Output.MaskR ? 1 : 0)
+								+ (Output.MaskG ? 1 : 0)
+								+ (Output.MaskB ? 1 : 0)
+								+ (Output.MaskA ? 1 : 0);
+		switch (MaskChannelCount)
 		{
+		case 1:
 			return MCT_Float;
+		case 2:
+			return MCT_Float2;
+		case 3:
+			return MCT_Float3;
+		case 4:
+			return MCT_Float4;
+		default:
+			return MCT_Unknown;
 		}
 	}
+	else
+	{
+		return MCT_Float;
+	}
 }
 
 int32 UMaterialExpression::GetWidth() const
@@ -8573,6 +8582,11 @@ bool UMaterialExpressionStaticSwitchParameter::IsResultMaterialAttributes(int32
 	}
 }
 
+EMaterialValueType UMaterialExpressionStaticSwitchParameter::GetOutputValueType(int32 OutputIndex)
+{
+	return GetCompositeOutput(OutputIndex);
+}
+
 FExpressionInput* UMaterialExpressionStaticSwitchParameter::GetEffectiveInput(class FMaterialCompiler* Compiler)
 {
 	bool bSucceeded;
@@ -8783,6 +8797,11 @@ bool UMaterialExpressionStaticSwitch::IsResultMaterialAttributes(int32 OutputInd
 	}
 }
 
+EMaterialValueType UMaterialExpressionStaticSwitch::GetOutputValueType(int32 OutputIndex)
+{
+	return GetCompositeOutput(OutputIndex);
+}
+
 FExpressionInput* UMaterialExpressionStaticSwitch::GetEffectiveInput(class FMaterialCompiler* Compiler)
 {
 	bool bValue = DefaultValue;
diff --git a/Engine/Source/Runtime/Engine/Public/Materials/MaterialExpression.h b/Engine/Source/Runtime/Engine/Public/Materials/MaterialExpression.h
--- a/Engine/Source/Runtime/Engine/Public/Materials/MaterialExpression.h
+++ b/Engine/Source/Runtime/Engine/Public/Materials/MaterialExpression.h
@@ -457,6 +457,11 @@ class UMaterialExpression : public UObject
 	 */
 	virtual bool IsResultSubstrateMaterial(int32 OutputIndex) { return false; }
 
+	/**
+	 * Returns specific EMaterialValueType when the result is MaterialAttributes or Substrate Material, MCT_Unknown otherwise
+	 */
+	EMaterialValueType GetCompositeOutput(int32 OutputIndex);
+
 	/**
 	 * Recursively parse nodes outputing Substrate material in order to gather all the possible shading models used in a material graph output a Substrate material.
 	 */
diff --git a/Engine/Source/Runtime/Engine/Public/Materials/MaterialExpressionStaticSwitch.h b/Engine/Source/Runtime/Engine/Public/Materials/MaterialExpressionStaticSwitch.h
--- a/Engine/Source/Runtime/Engine/Public/Materials/MaterialExpressionStaticSwitch.h
+++ b/Engine/Source/Runtime/Engine/Public/Materials/MaterialExpressionStaticSwitch.h
@@ -37,6 +37,7 @@ public:
 	virtual void GetCaption(TArray<FString>& OutCaptions) const override;
 	virtual FName GetInputName(int32 InputIndex) const override;
 	virtual bool IsResultMaterialAttributes(int32 OutputIndex) override;
+	virtual EMaterialValueType GetOutputValueType(int32 OutputIndex) override;
 	virtual EMaterialValueType GetInputValueType(int32 InputIndex) override;
 	virtual bool IsStaticExpression() const override { return true; }
 
diff --git a/Engine/Source/Runtime/Engine/Public/Materials/MaterialExpressionStaticSwitchParameter.h b/Engine/Source/Runtime/Engine/Public/Materials/MaterialExpressionStaticSwitchParameter.h
--- a/Engine/Source/Runtime/Engine/Public/Materials/MaterialExpressionStaticSwitchParameter.h
+++ b/Engine/Source/Runtime/Engine/Public/Materials/MaterialExpressionStaticSwitchParameter.h
@@ -31,7 +31,7 @@ public:
 	virtual void GetCaption(TArray<FString>& OutCaptions) const override;
 	virtual bool IsResultMaterialAttributes(int32 OutputIndex) override;
 	virtual EMaterialValueType GetInputValueType(int32 InputIndex) override {return MCT_Unknown;}
-	virtual EMaterialValueType GetOutputValueType(int32 OutputIndex) override { return UMaterialExpression::GetOutputValueType(OutputIndex); }
+	virtual EMaterialValueType GetOutputValueType(int32 OutputIndex) override;
 	virtual bool IsStaticExpression() const override { return true; }
 
 	virtual bool IsResultSubstrateMaterial(int32 OutputIndex) override;
@@ -45,4 +45,3 @@ protected:
 	FExpressionInput* GetEffectiveInput(class FMaterialCompiler* Compiler);
 #endif
 };
-
