Fix for Android HDR cubemap crash onAdreno GPUs

diff --git a/Engine/Source/Runtime/Engine/Private/Components/ReflectionCaptureComponent.cpp b/Engine/Source/Runtime/Engine/Private/Components/ReflectionCaptureComponent.cpp
--- a/Engine/Source/Runtime/Engine/Private/Components/ReflectionCaptureComponent.cpp
+++ b/Engine/Source/Runtime/Engine/Private/Components/ReflectionCaptureComponent.cpp
@@ -759,7 +759,7 @@ void UReflectionCaptureComponent::OnRegister()
 		FReflectionCaptureMapBuildData* MapBuildData = GetMapBuildData();
 
 		// If the MapBuildData is valid, update it. If it is not we will use the cached values, if there are any
-		if (MapBuildData)
+		if (MapBuildData && MapBuildData->CubemapSize > 0)
 		{
 			if (!EncodedHDRCubemapTexture)
 			{
