[VisualLogger] store ObjectNames struct in local reference to access the container only once

#rnx

diff --git a/Engine/Source/Runtime/Engine/Private/VisualLogger/VisualLogger.cpp b/Engine/Source/Runtime/Engine/Private/VisualLogger/VisualLogger.cpp
--- a/Engine/Source/Runtime/Engine/Private/VisualLogger/VisualLogger.cpp
+++ b/Engine/Source/Runtime/Engine/Private/VisualLogger/VisualLogger.cpp
@@ -737,7 +737,8 @@ void FVisualLogger::FlushEntry(FVisualLogEntry& Entry, const FObjectKey& ObjectK
 	const UObject* OwnerObject = ObjectKey.ResolveObjectPtrEvenIfGarbage();
 	for (FVisualLogDevice* Device : OutputDevices)
 	{
-		Device->Serialize(OwnerObject, ObjectToNamesMap[ObjectKey].Name, ObjectToNamesMap[ObjectKey].DisplayName, ObjectToNamesMap[ObjectKey].ClassName, Entry);
+		const FObjectNames& Names = ObjectToNamesMap[ObjectKey];
+		Device->Serialize(OwnerObject, Names.Name, Names.DisplayName, Names.ClassName, Entry);
 	}
 	Entry.Reset();
 }
