New DeferredAtmospherePass show flag to preserve atmosphere lighting effects but disable the main deferred sky atmosphere render pass.

Useful for show-only actor scene captures, which need to render primitives on a translucent background while preserving accurate lighting.

diff --git a/Engine/Source/Editor/DetailCustomizations/Private/SceneCaptureDetails.cpp b/Engine/Source/Editor/DetailCustomizations/Private/SceneCaptureDetails.cpp
--- a/Engine/Source/Editor/DetailCustomizations/Private/SceneCaptureDetails.cpp
+++ b/Engine/Source/Editor/DetailCustomizations/Private/SceneCaptureDetails.cpp
@@ -139,6 +139,7 @@ void FSceneCaptureDetails::CustomizeDetails( IDetailLayoutBuilder& DetailLayout
 	TArray<FEngineShowFlags::EShowFlag> ShowFlagsToAllowForCaptures;
 
 	ShowFlagsToAllowForCaptures.Add(FEngineShowFlags::EShowFlag::SF_Atmosphere);
+	ShowFlagsToAllowForCaptures.Add(FEngineShowFlags::EShowFlag::SF_DeferredAtmospherePass);
 	ShowFlagsToAllowForCaptures.Add(FEngineShowFlags::EShowFlag::SF_Cloud);
 	ShowFlagsToAllowForCaptures.Add(FEngineShowFlags::EShowFlag::SF_BSP);
 	ShowFlagsToAllowForCaptures.Add(FEngineShowFlags::EShowFlag::SF_Decals);
diff --git a/Engine/Source/Runtime/Engine/Public/ShowFlagsValues.inl b/Engine/Source/Runtime/Engine/Public/ShowFlagsValues.inl
--- a/Engine/Source/Runtime/Engine/Public/ShowFlagsValues.inl
+++ b/Engine/Source/Runtime/Engine/Public/ShowFlagsValues.inl
@@ -307,6 +307,8 @@ SHOWFLAG_ALWAYS_ACCESSIBLE(LightShafts, SFG_LightingFeatures, NSLOCTEXT("UnrealE
 SHOWFLAG_FIXED_IN_SHIPPING(1, PostProcessMaterial, SFG_PostProcess, NSLOCTEXT("UnrealEd", "PostProcessMaterialSF", "Post Process Material"))
 /** Render Sky and Atmospheric lighting, for now SHOWFLAG_ALWAYS_ACCESSIBLE because it's exposed in SceneCapture */
 SHOWFLAG_ALWAYS_ACCESSIBLE(Atmosphere, SFG_Normal, NSLOCTEXT("UnrealEd", "AtmosphereSF", "Atmosphere"))
+/** Disable the (main) deferred sky atmosphere render pass while preserving atmospheric lighting. */
+SHOWFLAG_ALWAYS_ACCESSIBLE(DeferredAtmospherePass, SFG_Advanced, NSLOCTEXT("UnrealEd", "DeferredAtmospherePassSF", "Deferred Atmosphere Pass"))
 /** Render Cloud, for now SHOWFLAG_ALWAYS_ACCESSIBLE because it's exposed in SceneCapture */
 SHOWFLAG_ALWAYS_ACCESSIBLE(Cloud, SFG_Normal, NSLOCTEXT("UnrealEd", "CloudSF", "Cloud"))
 /** Render safe frames bars*/
diff --git a/Engine/Source/Runtime/Renderer/Private/DeferredShadingRenderer.cpp b/Engine/Source/Runtime/Renderer/Private/DeferredShadingRenderer.cpp
--- a/Engine/Source/Runtime/Renderer/Private/DeferredShadingRenderer.cpp
+++ b/Engine/Source/Runtime/Renderer/Private/DeferredShadingRenderer.cpp
@@ -3384,7 +3384,7 @@ void FDeferredShadingSceneRenderer::Render(FRDGBuilder& GraphBuilder, const FSce
 			}
 
 			// Draw the sky atmosphere
-			if (!bHasRayTracedOverlay && bShouldRenderSkyAtmosphere && !IsForwardShadingEnabled(ShaderPlatform))
+			if (!bHasRayTracedOverlay && bShouldRenderSkyAtmosphere && !IsForwardShadingEnabled(ShaderPlatform) && ViewFamily.EngineShowFlags.DeferredAtmospherePass)
 			{
 				SCOPE_CYCLE_COUNTER(STAT_FDeferredShadingSceneRenderer_RenderSkyAtmosphere);
 				RenderSkyAtmosphere(GraphBuilder, SceneTextures);
diff --git a/Engine/Source/Runtime/Renderer/Private/SkyAtmosphereRendering.cpp b/Engine/Source/Runtime/Renderer/Private/SkyAtmosphereRendering.cpp
--- a/Engine/Source/Runtime/Renderer/Private/SkyAtmosphereRendering.cpp
+++ b/Engine/Source/Runtime/Renderer/Private/SkyAtmosphereRendering.cpp
@@ -2066,6 +2066,7 @@ void FSceneRenderer::RenderSkyAtmosphere(FRDGBuilder& GraphBuilder, const FMinim
 	check(!IsMobilePlatform(Scene->GetShaderPlatform()));
 
 	check(ShouldRenderSkyAtmosphere(Scene, ViewFamily.EngineShowFlags)); // This should not be called if we should not render SkyAtmosphere
+	check(ViewFamily.EngineShowFlags.DeferredAtmospherePass);
 
 	RDG_EVENT_SCOPE_STAT(GraphBuilder, SkyAtmosphere, "SkyAtmosphere");
 	RDG_GPU_STAT_SCOPE(GraphBuilder, SkyAtmosphere);
