[Audio Golden Path] Hot Wave File Reloading

#rnx
#changelist validated

diff --git a/Engine/Source/Runtime/AudioExtensions/Public/IAudioDecoderOverride.h b/Engine/Source/Runtime/AudioExtensions/Public/IAudioDecoderOverride.h
new file mode 100644
--- /dev/null
+++ b/Engine/Source/Runtime/AudioExtensions/Public/IAudioDecoderOverride.h
@@ -0,0 +1,23 @@
+// Copyright Epic Games, Inc. All Rights Reserved.
+
+#pragma once
+
+#include "CoreMinimal.h"
+
+class ICompressedAudioInfo;
+
+/**
+ * Allows you to define a factory that can provide an alternative audio file decoder at runtime.
+ */
+class UE_EXPERIMENTAL(5.8, "Decoder overrides are experimental and the API may be altered or removed") IAudioDecoderOverride
+{
+public:
+	/** Returns an overriden decoder for this sound, or nullptr to allow the default decoder to be created. */
+	virtual ICompressedAudioInfo* CreateOverrideDecoder(const FName SoundAssetName) = 0;
+
+	/** Test if an overriden decoder will be created for this sound. */
+	virtual bool WantToOverrideDecoder(const FName SoundAssetName) = 0;
+
+protected:
+	virtual ~IAudioDecoderOverride() = default;
+};
diff --git a/Engine/Source/Runtime/AudioMixer/Private/AudioMixerBuffer.cpp b/Engine/Source/Runtime/AudioMixer/Private/AudioMixerBuffer.cpp
--- a/Engine/Source/Runtime/AudioMixer/Private/AudioMixerBuffer.cpp
+++ b/Engine/Source/Runtime/AudioMixer/Private/AudioMixerBuffer.cpp
@@ -286,7 +286,7 @@ namespace Audio
 
 		FSoundQualityInfo QualityInfo = { 0 };
 		
-		Buffer->DecompressionState = IAudioInfoFactoryRegistry::Get().Create(InWave->GetRuntimeFormat());
+		Buffer->DecompressionState = IAudioInfoFactoryRegistry::Get().Create(InWave->GetRuntimeFormat(), InWave->GetFName());
 
 		// Get the header information of our compressed format
 		if (Buffer->DecompressionState && Buffer->DecompressionState->StreamCompressedInfo(InWave, &QualityInfo))
diff --git a/Engine/Source/Runtime/Engine/Classes/Sound/SoundWave.h b/Engine/Source/Runtime/Engine/Classes/Sound/SoundWave.h
--- a/Engine/Source/Runtime/Engine/Classes/Sound/SoundWave.h
+++ b/Engine/Source/Runtime/Engine/Classes/Sound/SoundWave.h
@@ -1558,6 +1558,9 @@ public:
 		return (ESoundWavePrecacheState)PrecacheState.GetValue();
 	}
 
+	// Call this to artificially cause any consumers of the wave data to refresh their state.
+	UE_EXPERIMENTAL(5.8, "May be altered or removed")
+	ENGINE_API void TriggerWaveDataRefresh();
 
 private:
 	TSharedRef<TDataSnapshotProducer<FSoundWaveData>> SoundWaveData = MakeShared<TDataSnapshotProducer<FSoundWaveData>>();
diff --git a/Engine/Source/Runtime/Engine/Private/AudioDecompress.cpp b/Engine/Source/Runtime/Engine/Private/AudioDecompress.cpp
--- a/Engine/Source/Runtime/Engine/Private/AudioDecompress.cpp
+++ b/Engine/Source/Runtime/Engine/Private/AudioDecompress.cpp
@@ -3,10 +3,11 @@
 
 #include "AudioDecompress.h"
 #include "AudioDevice.h"
+#include "AudioStreamingCache.h"
 #include "Engine/Engine.h"
 #include "Hash/xxhash.h"
+#include "IAudioDecoderOverride.h"
 #include "Interfaces/IAudioFormat.h"
-#include "AudioStreamingCache.h"
 #include "Misc/CoreStats.h"
 #include "Misc/ScopeRWLock.h"
 #include "Stats/StatsTrace.h"
@@ -935,6 +936,7 @@ IAudioInfoFactoryRegistry& IAudioInfoFactoryRegistry::Get()
 	{
 		mutable FRWLock FactoriesRWLock;
 		TMap<FName, IAudioInfoFactory*> Factories;
+		TArray<IAudioDecoderOverride*> Overrides;
 
 		void Register(IAudioInfoFactory* InFactory, FName InFormatName) override
 		{
@@ -958,6 +960,41 @@ IAudioInfoFactoryRegistry& IAudioInfoFactoryRegistry::Get()
 			}
 			return nullptr;
 		}
+
+		ICompressedAudioInfo* Create(FName InFormat, FName InAssetName) const override
+		{
+			if (!Overrides.IsEmpty() && !InAssetName.IsNone())
+			{
+				// Check if we want to create an override decoder for this asset
+				FRWScopeLock Lock(FactoriesRWLock, FRWScopeLockType::SLT_ReadOnly);
+				for (IAudioDecoderOverride* Override : Overrides)
+				{
+					ICompressedAudioInfo* Info = Override->CreateOverrideDecoder(InAssetName);
+					if (Info)
+					{
+						return Info;
+					}
+				}
+			}
+			if (IAudioInfoFactory* Factory = Find(InFormat))
+			{
+				return Factory->Create();
+			}
+			return nullptr;
+		}
+
+		virtual void RegisterOverride(IAudioDecoderOverride* Override) override
+		{
+			FRWScopeLock Lock(FactoriesRWLock, FRWScopeLockType::SLT_Write);
+			Overrides.Add(Override);
+		}
+
+		virtual void UnregisterOverride(IAudioDecoderOverride* Override) override
+		{
+			FRWScopeLock Lock(FactoriesRWLock, FRWScopeLockType::SLT_Write);
+			Overrides.Remove(Override);
+		}
+
 	} sInstance;
 	return sInstance;
 }
diff --git a/Engine/Source/Runtime/Engine/Private/SoundWave.cpp b/Engine/Source/Runtime/Engine/Private/SoundWave.cpp
--- a/Engine/Source/Runtime/Engine/Private/SoundWave.cpp
+++ b/Engine/Source/Runtime/Engine/Private/SoundWave.cpp
@@ -4433,6 +4433,11 @@ bool USoundWave::GetChunkData(int32 ChunkIndex, uint8** OutChunkData, bool bMake
 	return true;
 }
 
+void USoundWave::TriggerWaveDataRefresh()
+{
+	SoundWaveDataPtr.Fork();
+	Proxy = CreateSoundWaveProxy();
+}
 
 uint32 USoundWave::GetInterpolatedCookedFFTDataForTimeInternal(float InTime, uint32 StartingIndex, TArray<FSoundWaveSpectralData>& OutData, bool bLoop)
 {
diff --git a/Engine/Source/Runtime/Engine/Private/SoundWaveProxyReader.cpp b/Engine/Source/Runtime/Engine/Private/SoundWaveProxyReader.cpp
--- a/Engine/Source/Runtime/Engine/Private/SoundWaveProxyReader.cpp
+++ b/Engine/Source/Runtime/Engine/Private/SoundWaveProxyReader.cpp
@@ -390,15 +390,8 @@ bool FSoundWaveProxyReader::InitializeDecoder(float InStartTimeInSeconds)
 	TRACE_CPUPROFILER_EVENT_SCOPE(FSoundWaveProxyReader::InitializeDecoder);
 
 	FName Format = WaveProxy->GetRuntimeFormat();
-	IAudioInfoFactory* Factory = IAudioInfoFactoryRegistry::Get().Find(Format);
-	if (!ensure(Factory))
-	{
-		UE_LOG(LogAudio, Error, TEXT("FSoundWaveProxyReader::InitializeDecoder: Failed to create CompressedAudioInfo for wave (package: %s). Unable to find AudioInfoFactory for format: %s"), *WaveProxy->GetPackageName().ToString(), *Format.ToString());
-		return false;
-	}
-
 	TUniquePtr<ICompressedAudioInfo> InfoInstance;
-	InfoInstance.Reset(Factory->Create());
+	InfoInstance.Reset(IAudioInfoFactoryRegistry::Get().Create(Format, WaveProxy->GetFName()));
 
 	if (!ensure(InfoInstance.IsValid()))
 	{
diff --git a/Engine/Source/Runtime/Engine/Public/AudioDecompress.h b/Engine/Source/Runtime/Engine/Public/AudioDecompress.h
--- a/Engine/Source/Runtime/Engine/Public/AudioDecompress.h
+++ b/Engine/Source/Runtime/Engine/Public/AudioDecompress.h
@@ -23,6 +23,7 @@ constexpr uint32 MONO_PCM_BUFFER_SIZE = MONO_PCM_BUFFER_SAMPLES * MONO_PCM_SAMPL
 
 struct FSoundQualityInfo;
 class FStreamedAudioChunkSeekTable;
+class IAudioDecoderOverride;
 
 /**
  * Interface class to decompress various types of audio data
@@ -398,16 +399,12 @@ struct IAudioInfoFactoryRegistry
 	virtual void Register(IAudioInfoFactory*, FName) = 0;
 	virtual void Unregister(IAudioInfoFactory*, FName) = 0;
 	virtual IAudioInfoFactory* Find(FName InName) const = 0;
-	
-	// Convenience helper.
-	ICompressedAudioInfo* Create(FName InName) const
-	{
-		if (IAudioInfoFactory* Factory = Find(InName))
-		{
-			return Factory->Create();
-		}
-		return nullptr;
-	}
+	virtual ICompressedAudioInfo* Create(FName InName, FName InAssetName = NAME_None) const = 0;
+
+	UE_EXPERIMENTAL(5.8, "Decoder overriding is experimental and the API may be changed or removed")
+	virtual void RegisterOverride(IAudioDecoderOverride* Override) = 0;
+	UE_EXPERIMENTAL(5.8, "Decoder overriding is experimental and the API may be changed or removed")
+	virtual void UnregisterOverride(IAudioDecoderOverride* Override) = 0;
 };
 
 class FSimpleAudioInfoFactory : public IAudioInfoFactory
