[VisualLogger] deprecate header

diff --git a/Engine/Source/Runtime/Engine/Public/VisualLogger/VisualLoggerCustomVersion.h b/Engine/Source/Runtime/Engine/Public/VisualLogger/VisualLoggerCustomVersion.h
--- a/Engine/Source/Runtime/Engine/Public/VisualLogger/VisualLoggerCustomVersion.h
+++ b/Engine/Source/Runtime/Engine/Public/VisualLogger/VisualLoggerCustomVersion.h
@@ -1,7 +1,10 @@
 // Copyright Epic Games, Inc. All Rights Reserved.
 #pragma once
 
+#if UE_ENABLE_INCLUDE_ORDER_DEPRECATED_IN_5_8
 #include "CoreMinimal.h"
+#endif // UE_ENABLE_INCLUDE_ORDER_DEPRECATED_IN_5_8
+
 #include "Misc/Guid.h"
 
 // version for vlog binary file format
