Add function to modify bSetOwner

diff --git a/Engine/Source/Runtime/Engine/Classes/Components/ChildActorComponent.h b/Engine/Source/Runtime/Engine/Classes/Components/ChildActorComponent.h
--- a/Engine/Source/Runtime/Engine/Classes/Components/ChildActorComponent.h
+++ b/Engine/Source/Runtime/Engine/Classes/Components/ChildActorComponent.h
@@ -111,6 +111,12 @@ class UChildActorComponent : public USceneComponent
 
 	TSubclassOf<AActor> GetChildActorClass() const { return ChildActorClass; }
 
+	/**
+	 * Sets whether or not the owner of the child actor will be set to the owner of this component 
+	 * @param bInSetChildActorOwner Will the owner of this component be set as the owner of the child actor
+	 */
+	void SetChildActorOwnerOnCreation(bool bInSetChildActorOwner) { bSetOwner = bInSetChildActorOwner; }
+
 	friend class FChildActorComponentDetails;
 
 private:
