Dstm: Add remote object guard for ViewTarget when updating NoPawnPlayerController position.

diff --git a/Engine/Source/Runtime/Engine/Private/PlayerController.cpp b/Engine/Source/Runtime/Engine/Private/PlayerController.cpp
--- a/Engine/Source/Runtime/Engine/Private/PlayerController.cpp
+++ b/Engine/Source/Runtime/Engine/Private/PlayerController.cpp
@@ -15,7 +15,6 @@
 #include "EngineStats.h"
 #include "Engine/Canvas.h"
 #include "GameFramework/PlayerStart.h"
-#include "GameFramework/Character.h"
 #include "Components/SkeletalMeshComponent.h"
 #include "Components/AudioComponent.h"
 #include "Components/ForceFeedbackComponent.h"
@@ -5032,25 +5031,32 @@ static void UpdateNetConnectionControllerPosition(UNetConnection* NetConnection)
 	bool bNoPawnPlayerControllerPositionSet = false;
 	if (ControllerActor)
 	{
-		if (APlayerController* PlayerController = Cast<APlayerController>(NetConnection->ViewTarget))
-		{
-			FVector TrueViewLocation;
-			FRotator ViewRotation;
-			PlayerController->GetPlayerViewPoint(TrueViewLocation, ViewRotation);
-			ControllerActor->SetActorLocation(TrueViewLocation);
-			UE_LOG(LogPlayerController, Verbose, TEXT("[%s] UpdateNetConnectionControllerPosition: Setting NoPawnPlayerController %s position to %s using position from PlayerController %s"),
-				*FRemoteServerId::GetLocalServerId().ToString(), 
-				*ControllerActor->GetName(), *(TrueViewLocation.ToString()), *PlayerController->GetName());
-			bNoPawnPlayerControllerPositionSet = true;
-		}
-		else if (ACharacter* PlayerCharacter = Cast<ACharacter>(NetConnection->ViewTarget))
-		{
-			FVector TrueViewLocation = PlayerCharacter->GetActorLocation();
-			ControllerActor->SetActorLocation(TrueViewLocation);
-			UE_LOG(LogPlayerController, Verbose, TEXT("[%s] UpdateNetConnectionControllerPosition: Setting NoPawnPlayerController %s position to %s using position from PlayerCharacter %s"),
-				*FRemoteServerId::GetLocalServerId().ToString(),
-				*ControllerActor->GetName(), *(TrueViewLocation.ToString()), *PlayerCharacter->GetName());
-			bNoPawnPlayerControllerPositionSet = true;
+		if (NetConnection->ViewTarget.IsRemote())
+		{
+			UE_LOG(LogPlayerController, Warning, TEXT("UpdateNetConnectionControllerPosition: Incoming ViewTarget is remote."));
+		}
+		else
+		{
+			if (APlayerController* PlayerController = Cast<APlayerController>(NetConnection->ViewTarget))
+			{
+				FVector TrueViewLocation;
+				FRotator ViewRotation;
+				PlayerController->GetPlayerViewPoint(TrueViewLocation, ViewRotation);
+				ControllerActor->SetActorLocation(TrueViewLocation);
+				UE_LOG(LogPlayerController, Verbose, TEXT("[%s] UpdateNetConnectionControllerPosition: Setting NoPawnPlayerController %s position to %s using position from PlayerController %s"),
+					*FRemoteServerId::GetLocalServerId().ToString(),
+					*ControllerActor->GetName(), *(TrueViewLocation.ToString()), *PlayerController->GetName());
+				bNoPawnPlayerControllerPositionSet = true;
+			}
+			else if (AActor* TargetActor = Cast<AActor>(NetConnection->ViewTarget))
+			{
+				FVector TrueViewLocation = TargetActor->GetActorLocation();
+				ControllerActor->SetActorLocation(TrueViewLocation);
+				UE_LOG(LogPlayerController, Verbose, TEXT("[%s] UpdateNetConnectionControllerPosition: Setting NoPawnPlayerController %s position to %s using position from PlayerCharacter %s"),
+					*FRemoteServerId::GetLocalServerId().ToString(),
+					*ControllerActor->GetName(), *(TrueViewLocation.ToString()), *TargetActor->GetName());
+				bNoPawnPlayerControllerPositionSet = true;
+			}
 		}
 	}
 	if (!bNoPawnPlayerControllerPositionSet)
