MeshSimplificationSettings: manually load MeshUtilities before MeshReductionInterface

MeshUtilities and MeshReductionInterface have a cyclical dependency, and
MeshUtilities StartupModule does not work right when MeshUtilities is loaded
from inside MeshReductionInterface startup (GetStaticMeshReductionInterface
returns nullptr since MeshReductionInterface isn't completed yet, which in turn
causes VersionString to get set to the wrong value).

diff --git a/Engine/Source/Runtime/Engine/Private/MeshSimplificationSettings.cpp b/Engine/Source/Runtime/Engine/Private/MeshSimplificationSettings.cpp
--- a/Engine/Source/Runtime/Engine/Private/MeshSimplificationSettings.cpp
+++ b/Engine/Source/Runtime/Engine/Private/MeshSimplificationSettings.cpp
@@ -45,6 +45,16 @@ void UMeshSimplificationSettings::PostInitProperties()
 	Super::PostInitProperties(); 
 
 #if WITH_EDITOR
+	// As of this writing (Oct 28, 2025), MeshReductionInterface and MeshUtilities have a
+	// circular dependency. Each tries to load the other, but it turns out that
+	// MeshUtilities loading MeshReductionInterface works out properly, but MeshReductionInterface
+	// loading MeshUtilities does not quite work, since FMeshUtilities::StartupModule()
+	// sets its VersionString incorrectly when invoked from inside MeshReductionInterface
+	// startup.
+	//
+	// Explicitly load MeshUtilities first for now to sidestep this mess.
+	FModuleManager::Get().LoadModule("MeshUtilities");
+
 	IMeshReductionManagerModule& MeshReductionModule = FModuleManager::Get().LoadModuleChecked<IMeshReductionManagerModule>("MeshReductionInterface");
 	if(IsTemplate())
 	{
