When flagging Additive Animations as raw-data only (thus clearing its compressed data) also start a new compression request

diff --git a/Engine/Source/Runtime/Engine/Private/Animation/AnimSequence.cpp b/Engine/Source/Runtime/Engine/Private/Animation/AnimSequence.cpp
--- a/Engine/Source/Runtime/Engine/Private/Animation/AnimSequence.cpp
+++ b/Engine/Source/Runtime/Engine/Private/Animation/AnimSequence.cpp
@@ -2541,6 +2541,7 @@ void UAnimSequence::FlagDependentAnimationsAsRawDataOnly() const
 		if (Seq->HasAllFlags(RF_LoadCompleted) && Seq->RefPoseSeq == this)
 		{
 			Seq->ClearAllCompressionData();
+			Seq->BeginCacheDerivedDataForCurrentPlatform();
 		}
 	}
 }
