Added cvar to silently fail on missing download dependencies.

diff --git a/Engine/Source/Runtime/Engine/Classes/Engine/StreamableManager.h b/Engine/Source/Runtime/Engine/Classes/Engine/StreamableManager.h
--- a/Engine/Source/Runtime/Engine/Classes/Engine/StreamableManager.h
+++ b/Engine/Source/Runtime/Engine/Classes/Engine/StreamableManager.h
@@ -678,6 +678,8 @@ struct FStreamableDownloadParams
 	bool bInstallSoftReferences = false;
 	/** If set, request will complete after download, no assets will be loaded */
 	bool bDownloadOnly = false;
+
+	bool bAllowMissingDependencies = false;
 };
 
 /** Parameters for an async load */
diff --git a/Engine/Source/Runtime/Engine/Private/StreamableManager.cpp b/Engine/Source/Runtime/Engine/Private/StreamableManager.cpp
--- a/Engine/Source/Runtime/Engine/Private/StreamableManager.cpp
+++ b/Engine/Source/Runtime/Engine/Private/StreamableManager.cpp
@@ -433,6 +433,7 @@ namespace UE::StreamableManager::Private
 		FStreamableDownloadPriority Priority = FStreamableManager::DownloadDefaultPriority;
 		bool bInstallSoftReferences = false;
 		bool bDownloadOnly = false;
+		bool bAllowMissingDependencies = false;
 
 		bool IsStarted() const
 		{
@@ -504,6 +505,10 @@ namespace UE::StreamableManager::Private
 			{
 				InstallArgs.Options |= UE::IoStore::EOnDemandInstallOptions::InstallSoftReferences;
 			}
+			if (bAllowMissingDependencies)
+			{
+				InstallArgs.Options |= UE::IoStore::EOnDemandInstallOptions::AllowMissingDependencies;
+			}
 
 			InstallArgs.Priority = Priority;
 
@@ -1122,11 +1127,11 @@ void FStreamableHandle::CompleteLoad()
 	{
 		if (GStreamableLogStreamableFailuresAsErrors)
 		{
-			UE_LOG(LogStreamableManager, Error, TEXT("Encounted error: %s"), *Error->GetErrorMessage(true).ToString());
+			UE_LOG(LogStreamableManager, Error, TEXT("Encountered error: %s"), *Error->GetErrorMessage(true).ToString());
 		}
 		else
 		{
-			UE_LOG(LogStreamableManager, Verbose, TEXT("Encounted error: %s"), *Error->GetErrorMessage(true).ToString());
+			UE_LOG(LogStreamableManager, Verbose, TEXT("Encountered error: %s"), *Error->GetErrorMessage(true).ToString());
 		}
 	}
 	// Only complete if it's still active
@@ -1956,6 +1961,7 @@ TSharedPtr<FStreamableHandle> FStreamableManager::RequestAsyncLoad(FStreamableAs
 		NewRequest->DownloadContext = MakeUnique<UE::StreamableManager::Private::FDownloadContext>();
 		NewRequest->DownloadContext->Priority = Params.DownloadParams->Priority;
 		NewRequest->DownloadContext->bInstallSoftReferences = Params.DownloadParams->bInstallSoftReferences;
+		NewRequest->DownloadContext->bAllowMissingDependencies = Params.DownloadParams->bAllowMissingDependencies;
 		NewRequest->DownloadContext->bDownloadOnly = Params.DownloadParams->bDownloadOnly;
 	}
 
diff --git a/Engine/Source/Runtime/Experimental/IoStore/OnDemand/Private/OnDemandContentInstaller.cpp b/Engine/Source/Runtime/Experimental/IoStore/OnDemand/Private/OnDemandContentInstaller.cpp
--- a/Engine/Source/Runtime/Experimental/IoStore/OnDemand/Private/OnDemandContentInstaller.cpp
+++ b/Engine/Source/Runtime/Experimental/IoStore/OnDemand/Private/OnDemandContentInstaller.cpp
@@ -56,6 +56,13 @@ namespace UE::IoStore
 			TEXT("When the content installer runs out of tasks, it will flush any pending iochunks to disk.")
 			);
 
+		static bool bAllowAllMissingDependencies = false;
+		static FAutoConsoleVariableRef CVar_AllowAllMissingDependencies(
+			TEXT("iostore.AllowAllMissingDependencies"),
+			bAllowAllMissingDependencies,
+			TEXT("Allow missing dependencies when attempting to download assets, can cause silent failures.")
+			);
+
 		static bool bContentInstallerEnableHttpCancel = true;
 		static FAutoConsoleVariableRef CVar_ContentInstallerEnableHttpCancel(
 			TEXT("iostore.ContentInstallerEnableHttpCancel"),
@@ -783,12 +790,13 @@ void FOnDemandContentInstaller::ProcessInstallRequest(FRequest& Request)
 			}
 		}
 
-		if (MissingChunkIds.IsEmpty() == false && !EnumHasAnyFlags(InstallRequest.Args.Options, EOnDemandInstallOptions::AllowMissingDependencies))
+		if (MissingChunkIds.IsEmpty() == false && !EnumHasAnyFlags(InstallRequest.Args.Options, EOnDemandInstallOptions::AllowMissingDependencies) && !IoStore::CVars::bAllowAllMissingDependencies)
 		{
 			Request.Result = MakeError(UE::UnifiedError::IoStoreOnDemand::ChunkMissingError::MakeError(
 				UE::UnifiedError::IoStoreOnDemand::FChunkMissingErrorContext
 				{
-					.ChunkIds = MoveTemp(MissingChunkIds)
+					.ChunkIds = MoveTemp(MissingChunkIds),
+					.DiscoveredChunksCount = InstallRequest.ResolvedChunks.Num()
 				}));
 			return CompleteInstallRequest(Request);
 		}
diff --git a/Engine/Source/Runtime/Experimental/IoStore/OnDemandCore/Public/IO/OnDemandError.h b/Engine/Source/Runtime/Experimental/IoStore/OnDemandCore/Public/IO/OnDemandError.h
--- a/Engine/Source/Runtime/Experimental/IoStore/OnDemandCore/Public/IO/OnDemandError.h
+++ b/Engine/Source/Runtime/Experimental/IoStore/OnDemandCore/Public/IO/OnDemandError.h
@@ -150,6 +150,7 @@ inline UE::IoStore::TResult<ResultType> MakeSnapshotError(
 struct FChunkMissingErrorContext
 {
 	TArray<FIoChunkId> ChunkIds;
+	int32 DiscoveredChunksCount;
 };
 
 struct FChunkHashMismatchErrorContext
@@ -209,9 +210,10 @@ inline void SerializeForLog(FCbWriter& Writer, const UE::UnifiedError::IoStoreOn
 {
 	Writer.BeginObject();
 	Writer.AddString(ANSITEXTVIEW("$type"), TErrorStructFeatures<UE::UnifiedError::IoStoreOnDemand::FChunkMissingErrorContext>::GetErrorContextTypeNameAsString());
-	Writer.AddString(ANSITEXTVIEW("$format"), TEXT("(MissingCount: {MissingCount}, ChunkIds: [{ChunkIds}])"));
-
+	Writer.AddString(ANSITEXTVIEW("$format"), TEXT("(DiscoveredCount: {DiscoveredCount}, MissingCount: {MissingCount}, MissingChunkIds: [{ChunkIds}])"));
 	Writer.AddInteger(ANSITEXTVIEW("MissingCount"), Ctx.ChunkIds.Num());
+	Writer.AddInteger(ANSITEXTVIEW("DiscoveredCount"), Ctx.DiscoveredChunksCount);
+
 	Writer.BeginArray(ANSITEXTVIEW("ChunkIds"));
 	for (const auto& ChunkId : Ctx.ChunkIds)
 	{
