"Visibility Owner" feature for PrimitiveComponent - this appends the existing FPrimitiveSceneProxy::Owners array with additional arbitrary actors, allowing arbitrary actors not attached to a player to make use of bOnlyOwnerSee, bOwnerNoSee, etc.This can be useful for a number of things, especially in splitscreen. Sometimes we're seeing very heavy STAT_BuildHiddenComponentList costs on low-end platforms, 1ms+ on lower end devices. This function does a full component hierarchy traversal for every actor in the HiddenActors list, every frame. These actors are essentially static, doing this traversal constantly is wasted GameThread time.

#tests PIE
[RN] "Visibility Owner" feature for PrimitiveComponent - this appends the existing FPrimitiveSceneProxy::Owners array with additional arbitrary actors, allowing arbitrary actors not attached to a player to make use of bOnlyOwnerSee, bOwnerNoSee, etc.This can be useful for a number of things, especially in splitscreen. Sometimes we're seeing very heavy STAT_BuildHiddenComponentList costs on low-end platforms, 1ms+ on lower end devices. This function does a full component hierarchy traversal for every actor in the HiddenActors list, every frame. These actors are essentially static, doing this traversal constantly is wasted GameThread time.

diff --git a/Engine/Source/Runtime/Engine/Classes/GameFramework/PrimitiveComponentUtilities.h b/Engine/Source/Runtime/Engine/Classes/GameFramework/PrimitiveComponentUtilities.h
new file mode 100644
--- /dev/null
+++ b/Engine/Source/Runtime/Engine/Classes/GameFramework/PrimitiveComponentUtilities.h
@@ -0,0 +1,47 @@
+// Copyright Epic Games, Inc. All Rights Reserved.
+
+#pragma once
+
+#include "CoreMinimal.h"
+#include "Kismet/BlueprintFunctionLibrary.h"
+#include "PrimitiveComponentUtilities.generated.h"
+
+#define UE_API ENGINE_API
+
+UCLASS(meta=(BlueprintThreadSafe, ScriptName="PrimitiveComponentUtilities"), MinimalAPI)
+class UPrimitiveComponentUtilities : public UBlueprintFunctionLibrary
+{
+	GENERATED_UCLASS_BODY()
+	
+public:
+	/** Adds an additional "owner" for visibility purposes, like OwnerNoSee or OnlyOwnerSee. */
+	UFUNCTION(BlueprintCallable, Category="Rendering")
+	static UE_API void AddVisibilityOwner(UPrimitiveComponent* Component, AActor* Owner);
+
+	/** Adds an additional "owner" for visibility purposes, like OwnerNoSee or OnlyOwnerSee. */
+	UFUNCTION(BlueprintCallable, Category="Rendering")
+	static UE_API void AddVisibilityOwnerActor(AActor* Actor, AActor* Owner, bool bRecursivelyIncludeAttachedActors);
+
+	/** Removes an actor previously added by AddVisibilityOwner. */
+	UFUNCTION(BlueprintCallable, Category="Rendering")
+	static UE_API void RemoveVisibilityOwner(UPrimitiveComponent* Component, const AActor* Owner);
+
+	/** Removes an actor previously added by AddVisibilityOwner. */
+ 	UFUNCTION(BlueprintCallable, Category="Rendering")
+	static UE_API void RemoveVisibilityOwnerActor(AActor* Actor, AActor* Owner, bool bRecursivelyIncludeAttachedActors);
+
+	UFUNCTION(BlueprintCallable, Category="Rendering")
+	static UE_API  void ClearVisibilityOwners(UPrimitiveComponent* Component);
+ 
+	/** Changes the value of bOnlyOwnerSee on each primitive component on the actor. */
+	UFUNCTION(BlueprintCallable, Category="Rendering")
+	static UE_API void SetOnlyVisibleToOwner(AActor* Actor, bool bNewOnlyOwnerSee, bool bRecursivelyIncludeAttachedActors);
+
+	static UE_API TArray<const AActor*> GetVisibilityOwners(const UPrimitiveComponent* Component);
+ 
+	static UE_API void ForEachVisibilityOwner(const UPrimitiveComponent* Component, TFunctionRef<void(TObjectPtr<const AActor>)> Func);
+
+	static UE_API void ForEachActorPrimitiveComponent(AActor* Actor, TFunctionRef<void(TObjectPtr<UPrimitiveComponent>)> Func, bool bRecursivelyIncludeAttachedActors);
+};
+
+#undef UE_API
\ No newline at end of file
diff --git a/Engine/Source/Runtime/Engine/Classes/GameFramework/UObjectAnnotationSparseActorList.h b/Engine/Source/Runtime/Engine/Classes/GameFramework/UObjectAnnotationSparseActorList.h
new file mode 100644
--- /dev/null
+++ b/Engine/Source/Runtime/Engine/Classes/GameFramework/UObjectAnnotationSparseActorList.h
@@ -0,0 +1,156 @@
+// Copyright Epic Games, Inc. All Rights Reserved.
+
+#pragma once
+
+#include "CoreMinimal.h"
+#include "GameFramework/Actor.h"
+#include "Misc/ScopeRWLock.h"
+#include "UObject/WeakObjectPtrTemplates.h"
+#include "UObject/UObjectAnnotation.h"
+
+template<typename InAllocatorType>
+struct FAnnotationSparseActorList
+{
+	bool IsDefault() const
+	{
+		return OwnerList.IsEmpty();
+	}
+
+	TArray<const AActor*, InAllocatorType> OwnerList;
+};
+
+// No automatic cleanup as we need to manually manage our reverse annotation map.
+template<typename InAllocatorType>
+class FUObjectAnnotationSparseActorList : public FUObjectAnnotationSparse<FAnnotationSparseActorList<InAllocatorType>, false>
+{
+	using Base = FUObjectAnnotationSparse<FAnnotationSparseActorList<InAllocatorType>, false>;
+
+public:
+	void AddOwner(const UObject* Object, const AActor* Owner);
+	void RemoveOwner(const UObject* Object, const AActor* Owner);
+
+protected:
+	virtual void NotifyUObjectDeleted(const UObjectBase* Object, int32 Index) override;
+
+private:
+	TMap<const AActor*, TArray<const UObject*>> ReverseAnnotationMap;
+};
+
+template<typename InAllocatorType>
+void FUObjectAnnotationSparseActorList<InAllocatorType>::AddOwner(const UObject* Object, const AActor* Owner)
+{
+	if (!Object || !Owner)
+	{
+		return;
+	}
+
+	UE::TRWScopeLock AnnotationMapLock(this->AnnotationMapCritical, SLT_Write);
+
+	if (this->AnnotationMap.IsEmpty() && ReverseAnnotationMap.IsEmpty() && !this->bRegistered)
+	{
+		GUObjectArray.AddUObjectDeleteListener(this);
+		this->bRegistered = true;
+	}
+
+	FAnnotationSparseActorList<InAllocatorType>& Entry = this->AnnotationMap.FindOrAdd(Object);
+	Entry.OwnerList.AddUnique(Owner);
+	ReverseAnnotationMap.FindOrAdd(Owner).AddUnique(Object);
+}
+
+template<typename InAllocatorType>
+void FUObjectAnnotationSparseActorList<InAllocatorType>::RemoveOwner(const UObject* Object, const AActor* Owner)
+{
+	if (!Object || !Owner)
+	{
+		return;
+	}
+
+	UE::TRWScopeLock AnnotationMapLock(this->AnnotationMapCritical, SLT_Write);
+	if (FAnnotationSparseActorList<InAllocatorType>* Entry = this->AnnotationMap.Find(Object))
+	{
+		Entry->OwnerList.Remove(Owner);
+
+		// Remove the annotation entirely if we removed the last owner
+		if (Entry->IsDefault())
+		{
+			this->AnnotationMap.Remove(Object);
+		}
+	}
+
+	if (TArray<const UObject*>* Objects = ReverseAnnotationMap.Find(Owner))
+	{
+		Objects->Remove(Object);
+		if (Objects->IsEmpty())
+		{
+			ReverseAnnotationMap.Remove(Owner);
+		}
+	}
+
+	if (this->AnnotationMap.IsEmpty() && ReverseAnnotationMap.IsEmpty() && this->bRegistered)
+	{
+		GUObjectArray.RemoveUObjectDeleteListener(this);
+		this->bRegistered = false;
+	}
+}
+
+template<typename InAllocatorType>
+void FUObjectAnnotationSparseActorList<InAllocatorType>::NotifyUObjectDeleted(const UObjectBase* Object, int32 Index)
+{
+	// UE::TRWScopeLock AnnotationMapLock
+	{
+		UE::TRWScopeLock AnnotationMapLock(this->AnnotationMapCritical, SLT_Write);
+
+		// Object being deleted is potentially a owner
+		if (const AActor* Owner = Cast<const AActor>(Object))
+		{
+			if (TArray<const UObject*>* Objects = ReverseAnnotationMap.Find(Owner))
+			{
+				for (const UObject* TempObject : *Objects)
+				{
+					if (FAnnotationSparseActorList<InAllocatorType>* Entry = this->AnnotationMap.Find(TempObject))
+					{
+						Entry->OwnerList.Remove(Owner);
+						if (Entry->IsDefault())
+						{
+							this->AnnotationMap.Remove(TempObject);
+						}
+					}
+				}
+				ReverseAnnotationMap.Remove(Owner);
+			}
+		}
+		// Object being deleted is a UObject, find the owners and remove them.
+		else
+		{
+			const UObject* TempObject = Cast<const UObject>(Object);
+			if (FAnnotationSparseActorList<InAllocatorType>* Owners = this->AnnotationMap.Find(TempObject))
+			{
+				for (TWeakObjectPtr<const AActor> ObjectOwner : Owners->OwnerList)
+				{
+					if (ObjectOwner.IsValid())
+					{
+						if (TArray<const UObject*>* Objects = ReverseAnnotationMap.Find(ObjectOwner.Get()))
+						{
+							Objects->Remove(TempObject);
+							if (Objects->IsEmpty())
+							{
+								ReverseAnnotationMap.Remove(ObjectOwner.Get());
+							}
+						}
+					}
+				}
+				this->AnnotationMap.Remove(TempObject);
+			}
+		}
+
+		if (this->AnnotationMap.IsEmpty() && ReverseAnnotationMap.IsEmpty() && this->bRegistered)
+		{
+			GUObjectArray.RemoveUObjectDeleteListener(this);
+			this->bRegistered = false;
+		}
+	}
+
+	// Needs to be at the end because in non shipping builds it will check if we cleaned up our objects. 
+	// Make sure you have released the lock before this point or you will deadlock.
+	Base::NotifyUObjectDeleted(Object, Index);
+}
\ No newline at end of file
diff --git a/Engine/Source/Runtime/Engine/Private/GameFramework/PrimitiveComponentUtilities.cpp b/Engine/Source/Runtime/Engine/Private/GameFramework/PrimitiveComponentUtilities.cpp
new file mode 100644
--- /dev/null
+++ b/Engine/Source/Runtime/Engine/Private/GameFramework/PrimitiveComponentUtilities.cpp
@@ -0,0 +1,162 @@
+// Copyright Epic Games, Inc. All Rights Reserved.
+
+#include "GameFramework/PrimitiveComponentUtilities.h"
+
+#include "Components/PrimitiveComponent.h"
+#include "GameFramework/UObjectAnnotationSparseActorList.h"
+#include "Misc/ScopeRWLock.h"
+
+#include UE_INLINE_GENERATED_CPP_BY_NAME(PrimitiveComponentUtilities)
+
+// TInlineAllocator<4> because split screen is usually 2-4 players. 
+static FUObjectAnnotationSparseActorList<TInlineAllocator<4>> VisibilityOwners;
+UPrimitiveComponentUtilities::UPrimitiveComponentUtilities(const FObjectInitializer& ObjectInitializer)
+	: Super(ObjectInitializer)
+{
+}
+
+void UPrimitiveComponentUtilities::AddVisibilityOwner(UPrimitiveComponent* Component, AActor* Owner)
+{
+	if (!Component || !Owner)
+	{
+		return;
+	}
+
+	VisibilityOwners.AddOwner(Component, Owner);
+	Component->MarkRenderStateDirty();
+}
+
+void UPrimitiveComponentUtilities::AddVisibilityOwnerActor(AActor* Actor, AActor* Owner, bool bRecursivelyIncludeAttachedActors)
+{
+	if (!Owner || !Actor)
+	{
+		return;
+	}
+
+	ForEachActorPrimitiveComponent(Actor, [Owner](TObjectPtr<UPrimitiveComponent> Component)
+	{
+		AddVisibilityOwner(Component, Owner);
+	}, bRecursivelyIncludeAttachedActors);
+}
+
+void UPrimitiveComponentUtilities::RemoveVisibilityOwner(UPrimitiveComponent* Component, const AActor* Owner)
+{
+	if (!Component || !Owner)
+	{
+		return;
+	}
+
+	VisibilityOwners.RemoveOwner(Component, Owner);
+	Component->MarkRenderStateDirty();
+}
+
+void UPrimitiveComponentUtilities::RemoveVisibilityOwnerActor(AActor* Actor, AActor* Owner, bool bRecursivelyIncludeAttachedActors)
+{
+	if (!Actor || !Owner)
+	{
+		return;
+	}
+
+	ForEachActorPrimitiveComponent(Actor, [Owner](TObjectPtr<UPrimitiveComponent> Component)
+	{
+		RemoveVisibilityOwner(Component, Owner);
+	}, bRecursivelyIncludeAttachedActors);
+}
+
+void UPrimitiveComponentUtilities::ClearVisibilityOwners(UPrimitiveComponent* Component)
+{
+	if (!Component)
+	{
+		return;
+	}
+
+	TArray<const AActor*> Owners = GetVisibilityOwners(Component);
+	for (const AActor* Owner : Owners)
+	{
+		RemoveVisibilityOwner(Component, Owner);
+	}
+}
+
+TArray<const AActor*> UPrimitiveComponentUtilities::GetVisibilityOwners(const UPrimitiveComponent* Component)
+{
+	TArray<const AActor*> Result;
+	if (!Component)
+	{
+		return Result;
+	}
+		
+	FAnnotationSparseActorList Annotation = VisibilityOwners.GetAnnotation(Component);
+	Result.Reserve(Annotation.OwnerList.Num());
+
+	for (TWeakObjectPtr<const AActor> WeakOwner : Annotation.OwnerList)
+	{
+		if (const AActor* Owner = WeakOwner.Get())
+		{
+			Result.Add(Owner);
+		}
+	}
+
+	return Result;
+}
+
+void UPrimitiveComponentUtilities::SetOnlyVisibleToOwner(AActor* Actor, bool bNewOnlyOwnerSee, bool bRecursivelyIncludeAttachedActors)
+{
+	if(!Actor)
+	{
+		return;
+	}
+
+	ForEachActorPrimitiveComponent(Actor, [bNewOnlyOwnerSee](TObjectPtr<UPrimitiveComponent> Component)
+	{
+		if (Component)
+		{
+			Component->SetOnlyOwnerSee(bNewOnlyOwnerSee);
+			Component->MarkRenderStateDirty();
+		}
+	}, bRecursivelyIncludeAttachedActors);
+}
+
+void UPrimitiveComponentUtilities::ForEachVisibilityOwner(const UPrimitiveComponent* Component, TFunctionRef<void(TObjectPtr<const AActor>)> Func)
+{
+	if (!Component)
+	{
+		return;
+	}
+
+	FAnnotationSparseActorList Annotation = VisibilityOwners.GetAnnotation(Component);
+	for (TWeakObjectPtr<const AActor> WeakOwner : Annotation.OwnerList)
+	{
+		if (const AActor* Owner = WeakOwner.Get())
+		{
+			Func(Owner);
+		}
+	}
+}
+
+void UPrimitiveComponentUtilities::ForEachActorPrimitiveComponent(AActor* Actor, TFunctionRef<void(TObjectPtr<UPrimitiveComponent>)> Func, bool bRecursivelyIncludeAttachedActors)
+{
+	if (Actor != nullptr)
+	{
+		TArray<AActor*> AttachedActors;
+		AttachedActors.Add(Actor);
+
+		if (bRecursivelyIncludeAttachedActors)
+		{
+			Actor->GetAttachedActors(AttachedActors, false, bRecursivelyIncludeAttachedActors);
+		}
+
+		for (AActor* AttachedActor : AttachedActors)
+		{
+			TInlineComponentArray<UPrimitiveComponent*> Components;
+			AttachedActor->GetComponents(Components);
+
+			for (int32 ComponentIndex = 0; ComponentIndex < Components.Num(); ComponentIndex++)
+			{
+				if (UPrimitiveComponent* PrimitiveComponent = Components[ComponentIndex])
+				{
+					Func(PrimitiveComponent);
+				}
+			}
+		}
+	}
+}
\ No newline at end of file
diff --git a/Engine/Source/Runtime/Engine/Private/PrimitiveSceneProxyDesc.cpp b/Engine/Source/Runtime/Engine/Private/PrimitiveSceneProxyDesc.cpp
--- a/Engine/Source/Runtime/Engine/Private/PrimitiveSceneProxyDesc.cpp
+++ b/Engine/Source/Runtime/Engine/Private/PrimitiveSceneProxyDesc.cpp
@@ -2,6 +2,7 @@
 
 #include "PrimitiveSceneProxyDesc.h"
 #include "Components/InstancedStaticMeshComponent.h"
+#include "GameFramework/PrimitiveComponentUtilities.h"
 #include "SceneInterface.h"
 
 #if WITH_EDITOR
@@ -148,6 +149,12 @@ void FPrimitiveSceneProxyDesc::InitializeFromPrimitiveComponent(const UPrimitive
 			{
 				ActorOwners.Add(CurrentOwner->GetUniqueID());
 			}
+			
+			// Append "visibility owners" from the component
+			UPrimitiveComponentUtilities::ForEachVisibilityOwner(InComponent, [this](TObjectPtr<const AActor> VisibilityOwner)
+			{
+				ActorOwners.Add(VisibilityOwner->GetUniqueID());
+			});
 		}
 
 		bIsOwnerEditorOnly = InComponent->GetOwner()->IsEditorOnly();
