[NetPhysics] Comply with CompactSet

diff --git a/Engine/Source/Runtime/Engine/Private/PhysicsEngine/PhysicsReplication.cpp b/Engine/Source/Runtime/Engine/Private/PhysicsEngine/PhysicsReplication.cpp
--- a/Engine/Source/Runtime/Engine/Private/PhysicsEngine/PhysicsReplication.cpp
+++ b/Engine/Source/Runtime/Engine/Private/PhysicsEngine/PhysicsReplication.cpp
@@ -505,10 +505,16 @@ void FPhysicsReplication::OnTick(float DeltaSeconds, TMap<TWeakObjectPtr<UPrimit
 		if (bRemoveItr)
 		{
 			OnTargetRestored(Itr.Key().Get(), Itr.Value());
-			Itr.RemoveCurrent();
+			PendingDeleteFromComponentsToTargets.Add(Itr.Key());
 		}
 	}
 
+	for (TWeakObjectPtr<UPrimitiveComponent>& PrimitiveComponent : PendingDeleteFromComponentsToTargets)
+	{
+		ComponentsToTargets.Remove(PrimitiveComponent);
+	}
+	PendingDeleteFromComponentsToTargets.Reset();
+
 	// PhysicsObject replication flow
 	for (FReplicatedPhysicsTarget& PhysicsTarget : ReplicatedTargetsQueue)
 	{
@@ -1301,7 +1307,7 @@ void FPhysicsReplicationAsync::ApplyTargetStatesAsync(const float DeltaSeconds)
 		{
 			ReplicatedParticleIDs.Remove(Handle->ParticleID());
 		}
-		Itr.RemoveCurrent();
+		PendingDeleteFromObjectToTarget.Add(Itr.Key());
 	};
 
 	// PhysicsObject flow
@@ -1355,6 +1361,12 @@ void FPhysicsReplicationAsync::ApplyTargetStatesAsync(const float DeltaSeconds)
 			RemoveTargetHelper(Itr, RigidHandle);
 		}
 	}
+
+	for (Chaos::FConstPhysicsObjectHandle& PhysicsObject : PendingDeleteFromObjectToTarget)
+	{
+		ObjectToTarget.Remove(PhysicsObject);
+	}
+	PendingDeleteFromObjectToTarget.Reset();
 }
 
 void FPhysicsReplicationAsync::CheckTargetResimValidity(FReplicatedPhysicsTargetAsync& Target)
diff --git a/Engine/Source/Runtime/Engine/Private/PhysicsEngine/PhysicsReplicationCache.cpp b/Engine/Source/Runtime/Engine/Private/PhysicsEngine/PhysicsReplicationCache.cpp
--- a/Engine/Source/Runtime/Engine/Private/PhysicsEngine/PhysicsReplicationCache.cpp
+++ b/Engine/Source/Runtime/Engine/Private/PhysicsEngine/PhysicsReplicationCache.cpp
@@ -240,8 +240,8 @@ void FPhysicsReplicationCacheAsync::PopulateReplicationCache_Internal()
 		if (ReplicationCacheData.GetAccessTime() < (RigidSolver->GetSolverTime() - ReplicationCacheCvars::LingerForSeconds))
 		{
 			// Unregister physics object from caching state
-			It.RemoveCurrent();
-		
+			PendingDeleteFromReplicationCache_Internal.Add(It.Key());
+
 			// Set empty entry in output since it's already allocated
 			AsyncOutput.ReplicationCache_Key_Marshal[MarshalIndex] = nullptr;
 			MarshalIndex++;
@@ -271,6 +271,12 @@ void FPhysicsReplicationCacheAsync::PopulateReplicationCache_Internal()
 		}
 	}
 
+	for (Chaos::FConstPhysicsObjectHandle& PhysicsObject : PendingDeleteFromReplicationCache_Internal)
+	{
+		ReplicationCache_Internal.Remove(PhysicsObject);
+	}
+	PendingDeleteFromReplicationCache_Internal.Reset();
+
 	bUpdateAfterRemoval = false;
 
 	// Mark async output as final, meaning it can be read on the game thread instantly instead of at the end of the physics tick
diff --git a/Engine/Source/Runtime/Engine/Public/Physics/PhysicsReplicationCache.h b/Engine/Source/Runtime/Engine/Public/Physics/PhysicsReplicationCache.h
--- a/Engine/Source/Runtime/Engine/Public/Physics/PhysicsReplicationCache.h
+++ b/Engine/Source/Runtime/Engine/Public/Physics/PhysicsReplicationCache.h
@@ -116,6 +116,7 @@ class FPhysicsReplicationCacheAsync : public Chaos::TSimCallbackObject<
 
 private:
 	TMap<Chaos::FConstPhysicsObjectHandle, FPhysicsReplicationCacheData> ReplicationCache_Internal;
+	TArray<Chaos::FConstPhysicsObjectHandle> PendingDeleteFromReplicationCache_Internal;
 	bool bUpdateAfterRemoval = false;
 };
 
diff --git a/Engine/Source/Runtime/Engine/Public/PhysicsReplication.h b/Engine/Source/Runtime/Engine/Public/PhysicsReplication.h
--- a/Engine/Source/Runtime/Engine/Public/PhysicsReplication.h
+++ b/Engine/Source/Runtime/Engine/Public/PhysicsReplication.h
@@ -215,6 +215,7 @@ private:
 	FNetworkPhysicsSettingsData SettingsCurrent;
 	FNetworkPhysicsSettingsData SettingsDefault;
 	TMap<Chaos::FConstPhysicsObjectHandle, FReplicatedPhysicsTargetAsync> ObjectToTarget;
+	TArray<Chaos::FConstPhysicsObjectHandle> PendingDeleteFromObjectToTarget;
 	TMap<Chaos::FConstPhysicsObjectHandle, TWeakPtr<const FNetworkPhysicsSettingsData>> ObjectToSettings;
 	TArray<const Chaos::Private::FPBDIsland*> ResimIslands;
 	TArray<const Chaos::FGeometryParticleHandle*> ResimIslandsParticles;
@@ -337,6 +338,7 @@ private:
 
 private:
 	TMap<TWeakObjectPtr<UPrimitiveComponent>, FReplicatedPhysicsTarget> ComponentToTargets_DEPRECATED; // This collection is keeping the legacy flow working until fully deprecated in a future release
+	TArray<TWeakObjectPtr<UPrimitiveComponent>> PendingDeleteFromComponentsToTargets;
 	TArray<FReplicatedPhysicsTarget> ReplicatedTargetsQueue;
 	FPhysScene* PhysScene;
 	TWeakObjectPtr<UNetworkPhysicsSettingsComponent> SettingsCurrent;
