Added dp.LogActiveDeviceProfile to help figure out the active device profile and its parents

diff --git a/Engine/Source/Runtime/Engine/Private/DeviceProfiles/DeviceProfileManager.cpp b/Engine/Source/Runtime/Engine/Private/DeviceProfiles/DeviceProfileManager.cpp
--- a/Engine/Source/Runtime/Engine/Private/DeviceProfiles/DeviceProfileManager.cpp
+++ b/Engine/Source/Runtime/Engine/Private/DeviceProfiles/DeviceProfileManager.cpp
@@ -36,6 +36,27 @@ static TAutoConsoleVariable<FString> CVarDeviceProfileOverride(
 	TEXT(" The commandline -dp option will override this on startup, but not when setting this at runtime\n"),
 	ECVF_Default);
 
+void OnCVarDeviceProfileLogActiveDeviceProfile(IConsoleVariable*)
+{
+	UDeviceProfileManager& DPManager = UDeviceProfileManager::Get();
+	UE_LOG(LogDeviceProfileManager, Log, TEXT("ActiveDeviceProfileName= %s"), *DPManager.GetActiveDeviceProfileName());
+
+	TArray<TTuple<FName, FString>> DPChain = DPManager.GetDPChain(DPManager.GetActiveDeviceProfileName(), GConfig);
+
+	for (const auto& DP : DPChain)
+	{
+		const FName& DPName = DP.Get<0>();
+		UE_LOG(LogDeviceProfileManager, Log, TEXT("   %s"), *DPName.ToString());
+	}
+}
+
+static TAutoConsoleVariable<bool> CVarDeviceProfileLogActiveDeviceProfile(
+	TEXT("dp.LogActiveDeviceProfile"),
+	false,
+	TEXT("Logs the result of GetActiveDeviceProfileName and GetDPChain"),
+	FConsoleVariableDelegate::CreateStatic(OnCVarDeviceProfileLogActiveDeviceProfile),
+	ECVF_Default);
+
 static TAutoConsoleVariable<int32> CVarAllowScalabilityGroupsToChangeAtRuntime(
 	TEXT("dp.AllowScalabilityGroupsToChangeAtRuntime"),
 	0,
