Add full Nanite version string to Nanite-enabled skeletal meshes instead of only NANITE_DERIVEDDATA_VER.
This prevents potential DDC corruption when any of the defines or CVars that affect Nanite building are changed.

diff --git a/Engine/Source/Runtime/Engine/Private/SkeletalMesh.cpp b/Engine/Source/Runtime/Engine/Private/SkeletalMesh.cpp
--- a/Engine/Source/Runtime/Engine/Private/SkeletalMesh.cpp
+++ b/Engine/Source/Runtime/Engine/Private/SkeletalMesh.cpp
@@ -5876,7 +5876,7 @@ FString USkeletalMesh::BuildDerivedDataKey(const ITargetPlatform* TargetPlatform
 		// Nanite skeletal mesh version
 		KeySuffix += TEXT("_NSK_WIP_1");
 
-		static FString CachedNaniteVersion = FDevSystemGuids::GetSystemGuid(FDevSystemGuids::Get().NANITE_DERIVEDDATA_VER).ToString();
+		static FString CachedNaniteVersion = *Nanite::IBuilderModule::Get().GetVersionString();
 		KeySuffix += *CachedNaniteVersion;
 	}
 	
