[WorldPartition] Fixed typos in strings used to display errors in FDataLayerInstanceNamesTest unittest.

diff --git a/Engine/Source/Runtime/Engine/Private/Tests/WorldPartition/DataLayerInstanceNamesTests.cpp b/Engine/Source/Runtime/Engine/Private/Tests/WorldPartition/DataLayerInstanceNamesTests.cpp
--- a/Engine/Source/Runtime/Engine/Private/Tests/WorldPartition/DataLayerInstanceNamesTests.cpp
+++ b/Engine/Source/Runtime/Engine/Private/Tests/WorldPartition/DataLayerInstanceNamesTests.cpp
@@ -57,7 +57,7 @@ namespace WorldPartitionTests
 			return false;
 		}
 
-		if (!TestTrue(TEXT("Expected merged ELD with No EDL to have Exteral Data Layer == 'EDL'"), MergeEDLWithNoEDL.GetExternalDataLayer() == TEXT("EDL")))
+		if (!TestTrue(TEXT("Expected merged ELD with No EDL to have External Data Layer == 'EDL'"), MergeEDLWithNoEDL.GetExternalDataLayer() == TEXT("EDL")))
 		{
 			return false;
 		}
@@ -75,7 +75,7 @@ namespace WorldPartitionTests
 			return false;
 		}
 
-		if (!TestTrue(TEXT("Expected merged ELD with No EDL to contain 1 Data Layer"), MergeEDLWithNoEDLForceEmpty.GetExternalDataLayer() == TEXT("EDL")))
+		if (!TestTrue(TEXT("Expected merged ELD with No EDL to have External Data Layer == 'EDL'"), MergeEDLWithNoEDLForceEmpty.GetExternalDataLayer() == TEXT("EDL")))
 		{
 			return false;
 		}
@@ -97,7 +97,7 @@ namespace WorldPartitionTests
 			return false;
 		}
 
-		if (!TestTrue(TEXT("Expected merged No ELD with EDL to have Exteral Data Layer == 'EDL'"), MergeNoEDLWithEDL.GetExternalDataLayer() == TEXT("EDL")))
+		if (!TestTrue(TEXT("Expected merged No ELD with EDL to have External Data Layer == 'EDL'"), MergeNoEDLWithEDL.GetExternalDataLayer() == TEXT("EDL")))
 		{
 			return false;
 		}
