Fix for removal being done from a Map inside a range-for with .Remove(<key>) instead of .RemoveCurrent() on the iterator
#rnx

diff --git a/Engine/Source/Runtime/Engine/Private/VisualLogger/VisualLogger.cpp b/Engine/Source/Runtime/Engine/Private/VisualLogger/VisualLogger.cpp
--- a/Engine/Source/Runtime/Engine/Private/VisualLogger/VisualLogger.cpp
+++ b/Engine/Source/Runtime/Engine/Private/VisualLogger/VisualLogger.cpp
@@ -1220,7 +1220,7 @@ void FVisualLogger::Cleanup(UWorld* OldWorld, const bool bReleaseMemory)
 					if (It.Value() == OldWorld)
 					{
 						FObjectKey Obj = It.Key();
-						ObjectToWorldMap.Remove(Obj);
+						It.RemoveCurrent();
 						CurrentEntryPerObject.Remove(Obj);
 
 						for (FVisualLoggerObjectEntryMap* ThreadCurrentEntryMap : ThreadCurrentEntryMaps)
