[WorldPartition] Refactored PerInstanceData to store information in FStreamingGenerationActorDescView & Fixed Streaming Generation Logging
- Logging of container now logs directly actor descs (includes editor only actors)
- Logging of container instances now only logs differences between the ActorDescView and its corresponding ActorDesc
- StreamingGeneration unittest now logs with EToStringMode::Full mode to see all properties like references, data layers, etc.
- Updated unittest WPUnitTests.log file
- Container instance clusters now also logs a clusters hash. When a duplicate is found, only log the clusters hash to reduce log size
- FStreamingGenerationActorDescView now stores the data that used to be in FContainerCollectionInstanceDescriptor's PerInstanceData
- Fixed ActorDescView now showing the actual final version of the used content (like Data layers)
- Fixed Log not showing actor desc mutations applied by ActorDescMutators

diff --git a/Engine/Source/Runtime/Engine/Private/WorldPartition/WorldPartitionActorDesc.cpp b/Engine/Source/Runtime/Engine/Private/WorldPartition/WorldPartitionActorDesc.cpp
--- a/Engine/Source/Runtime/Engine/Private/WorldPartition/WorldPartitionActorDesc.cpp
+++ b/Engine/Source/Runtime/Engine/Private/WorldPartition/WorldPartitionActorDesc.cpp
@@ -540,121 +540,72 @@ void FWorldPartitionActorDesc::RegisterActorDescDeprecator(TSubclassOf<AActor> A
 
 FString FWorldPartitionActorDesc::ToString(EToStringMode Mode) const
 {
-	TStringBuilder<1024> Result;
-	const TCHAR LineStart = (Mode == EToStringMode::ForDiff) ? TEXT('\t') : TEXT(' ');
-	const TCHAR* LineEnd = (Mode == EToStringMode::ForDiff) ? LINE_TERMINATOR : nullptr;
-	
-	Result.Appendf(TEXT("Guid:%s%s"), *Guid.ToString(), LineEnd ? LineEnd : TEXT(""));
-
-	auto Append = [&Result, LineStart, LineEnd](const TCHAR* Name, const TCHAR* Value)
-	{
-		Result.AppendChar(LineStart);
-		Result.Append(Name);
-		Result.AppendChar(TEXT(':'));
-		Result.Append(Value);
-
-		if (LineEnd)
-		{
-			Result.Append(LineEnd);
-		}
-	};
-
-	auto AppendFromString = [&Append](const TCHAR* Name, const FString& Value)
-	{
-		Append(Name, *Value);
-	};
-
-	auto AppendToString = [&Append, &AppendFromString]<typename Type>(const TCHAR* Name, const Type& Value)
-	{
-		if constexpr (std::is_same_v<Type, FBox>)
-		{
-			if (!Value.IsValid)
-			{
-				Append(Name, TEXT("IsValid=false"));
-				return;
-			}
-		}
-
-		if constexpr (std::is_same_v<Type, FWorldPartitionRelativeBounds>)
-		{
-			if (!Value.IsValid())
-			{
-				Append(Name, TEXT("IsValid=false"));
-				return;
-			}
-		}
-
-		AppendFromString(Name, Value.ToString());
-	};
-
-	auto AppendFromBool = [&Append, Mode](const TCHAR* Name, bool bValue)
-	{
-		Append(Name, *LexToString(bValue));
-	};
+	FWorldPartitionActorDescStringBuilder Result(Mode);
+	Result.AppendToString(TEXT("Guid"), Guid, false);
 
 	if (Mode >= EToStringMode::Compact)
 	{
 		if (BaseClass.IsValid())
 		{
-			AppendToString(TEXT("BaseClass"), BaseClass);
+			Result.AppendToString(TEXT("BaseClass"), BaseClass);
 		}
 
-		AppendToString(TEXT("NativeClass"), NativeClass);
-		AppendFromString(TEXT("Name"), GetActorNameString());
+		Result.AppendToString(TEXT("NativeClass"), NativeClass);
+		Result.AppendFromString(TEXT("Name"), GetActorNameString());
 
 		if (Mode >= EToStringMode::Verbose)
 		{
-			AppendToString(TEXT("ActorPackage"), ActorPackage);
-			AppendToString(TEXT("ActorPath"), ActorPath);
+			Result.AppendToString(TEXT("ActorPackage"), ActorPackage);
+			Result.AppendToString(TEXT("ActorPath"), ActorPath);
 		}
 
 		if (ComponentDescs.Num())
 		{
-			AppendFromString(TEXT("Components"), FString::JoinBy(ComponentDescs, TEXT(","), 
+			Result.AppendFromString(TEXT("Components"), FString::JoinBy(ComponentDescs, TEXT(","),
 				[&](const TUniquePtr<FWorldPartitionComponentDesc>& ComponentDesc)
 				{
 					return ComponentDesc->ToString();
 				}));
 		}
 
-		AppendToString(TEXT("Label"), GetActorLabel());
-		AppendFromBool(TEXT("SpatiallyLoaded"), bIsSpatiallyLoaded);
-		AppendToString(TEXT("EditorBounds"), EditorBounds);
-		AppendToString(TEXT("EditorBoundsRelative"), EditorBoundsRelative);
-		AppendToString(TEXT("RuntimeBounds"), RuntimeBounds);
-		AppendToString(TEXT("RuntimeBoundsRelative"), RuntimeBoundsRelative);
-		AppendToString(TEXT("RuntimeGrid"), RuntimeGrid);
-		AppendFromBool(TEXT("EditorOnly"), bActorIsEditorOnly);
-		AppendFromBool(TEXT("RuntimeOnly"), bActorIsRuntimeOnly);
-		AppendFromBool(TEXT("HLODRelevant"), bActorIsHLODRelevant);
-		AppendFromBool(TEXT("ListedInSceneOutliner"), bActorIsListedInSceneOutliner);
-		AppendFromBool(TEXT("IsMainWorldOnly"), IsMainWorldOnly());
+		Result.AppendToString(TEXT("Label"), GetActorLabel());
+		Result.AppendFromBool(TEXT("SpatiallyLoaded"), bIsSpatiallyLoaded);
+		Result.AppendToString(TEXT("EditorBounds"), EditorBounds);
+		Result.AppendToString(TEXT("EditorBoundsRelative"), EditorBoundsRelative);
+		Result.AppendToString(TEXT("RuntimeBounds"), RuntimeBounds);
+		Result.AppendToString(TEXT("RuntimeBoundsRelative"), RuntimeBoundsRelative);
+		Result.AppendToString(TEXT("RuntimeGrid"), RuntimeGrid);
+		Result.AppendFromBool(TEXT("EditorOnly"), bActorIsEditorOnly);
+		Result.AppendFromBool(TEXT("RuntimeOnly"), bActorIsRuntimeOnly);
+		Result.AppendFromBool(TEXT("HLODRelevant"), bActorIsHLODRelevant);
+		Result.AppendFromBool(TEXT("ListedInSceneOutliner"), bActorIsListedInSceneOutliner);
+		Result.AppendFromBool(TEXT("IsMainWorldOnly"), IsMainWorldOnly());
 
 		if (ParentActor.IsValid())
 		{
-			AppendToString(TEXT("Parent"), ParentActor);
+			Result.AppendToString(TEXT("Parent"), ParentActor);
 		}
 
 		if (HLODLayer.IsValid())
 		{
-			AppendToString(TEXT("HLODLayer"), HLODLayer);
+			Result.AppendToString(TEXT("HLODLayer"), HLODLayer);
 		}
 
 		if (!FolderPath.IsNone())
 		{
-			AppendToString(TEXT("FolderPath"), FolderPath);
+			Result.AppendToString(TEXT("FolderPath"), FolderPath);
 		}
 
 		if (FolderGuid.IsValid())
 		{
-			AppendToString(TEXT("FolderGuid"), FolderGuid);
+			Result.AppendToString(TEXT("FolderGuid"), FolderGuid);
 		}
 
 		if (Mode >= EToStringMode::Full)
 		{
 			if (References.Num())
 			{
-				AppendFromString(TEXT("References"), FString::JoinBy(References, TEXT(","), 
+				Result.AppendFromString(TEXT("References"), FString::JoinBy(References, TEXT(","),
 					[&](const FGuid& ReferenceGuid)
 					{
 						return ReferenceGuid.ToString();
@@ -663,7 +614,7 @@ FString FWorldPartitionActorDesc::ToString(EToStringMode Mode) const
 
 			if (EditorOnlyReferences.Num())
 			{
-				AppendFromString(TEXT("EditorOnlyReferences"), FString::JoinBy(EditorOnlyReferences, TEXT(","), 
+				Result.AppendFromString(TEXT("EditorOnlyReferences"), FString::JoinBy(EditorOnlyReferences, TEXT(","),
 					[&](const FGuid& ReferenceGuid)
 					{
 						return ReferenceGuid.ToString();
@@ -672,7 +623,7 @@ FString FWorldPartitionActorDesc::ToString(EToStringMode Mode) const
 
 			if (Tags.Num())
 			{
-				AppendFromString(TEXT("Tags"), FString::JoinBy(Tags, TEXT(","), 
+				Result.AppendFromString(TEXT("Tags"), FString::JoinBy(Tags, TEXT(","),
 					[&](const FName& TagName)
 					{
 						return TagName.ToString();
@@ -681,12 +632,12 @@ FString FWorldPartitionActorDesc::ToString(EToStringMode Mode) const
 
 			if (Properties.Num())
 			{
-				AppendToString(TEXT("Properties"), Properties);
+				Result.AppendToString(TEXT("Properties"), Properties);
 			}
 
 			if (DataLayers.Num())
 			{
-				AppendFromString(TEXT("DataLayers"), FString::JoinBy(DataLayers, TEXT(","), 
+				Result.AppendFromString(TEXT("DataLayers"), FString::JoinBy(DataLayers, TEXT(","),
 					[&](const FName& DataLayerName)
 					{
 						return DataLayerName.ToString();
@@ -695,7 +646,7 @@ FString FWorldPartitionActorDesc::ToString(EToStringMode Mode) const
 
 			if (ExternalDataLayerAsset.IsValid())
 			{
-				AppendToString(TEXT("ExternalDataLayerAsset"), ExternalDataLayerAsset);
+				Result.AppendToString(TEXT("ExternalDataLayerAsset"), ExternalDataLayerAsset);
 			}
 		}
 	}
@@ -1174,6 +1125,39 @@ void FWorldPartitionActorDesc::FixupStreamingBounds()
 	}
 }
 
+FWorldPartitionActorDescStringBuilder::FWorldPartitionActorDescStringBuilder(FWorldPartitionActorDesc::EToStringMode InMode)
+	: Mode(InMode)
+	, LineStart((Mode == FWorldPartitionActorDesc::EToStringMode::ForDiff) ? TEXT('\t') : TEXT(' '))
+	, LineEnd((Mode == FWorldPartitionActorDesc::EToStringMode::ForDiff) ? LINE_TERMINATOR : nullptr)
+{
+}
+
+void FWorldPartitionActorDescStringBuilder::AppendFromString(const TCHAR* Name, const FString& Value, bool bUseLineStart)
+{
+	AppendPair(Name, *Value, bUseLineStart);
+}
+
+void FWorldPartitionActorDescStringBuilder::AppendFromBool(const TCHAR* Name, bool bValue, bool bUseLineStart)
+{
+	AppendPair(Name, *LexToString(bValue), bUseLineStart);
+}
+
+void FWorldPartitionActorDescStringBuilder::AppendPair(const TCHAR* Name, const TCHAR* Value, bool bUseLineStart)
+{
+	if (bUseLineStart)
+	{
+		AppendChar(LineStart);
+	}
+	Append(Name);
+	AppendChar(TEXT(':'));
+	Append(Value);
+
+	if (LineEnd)
+	{
+		Append(LineEnd);
+	}
+}
+
 #undef LOCTEXT_NAMESPACE
 
 #endif
diff --git a/Engine/Source/Runtime/Engine/Private/WorldPartition/WorldPartitionStreamingGeneration.cpp b/Engine/Source/Runtime/Engine/Private/WorldPartition/WorldPartitionStreamingGeneration.cpp
--- a/Engine/Source/Runtime/Engine/Private/WorldPartition/WorldPartitionStreamingGeneration.cpp
+++ b/Engine/Source/Runtime/Engine/Private/WorldPartition/WorldPartitionStreamingGeneration.cpp
@@ -233,6 +233,11 @@ bool FStreamingGenerationActorDescView::GetIsSpatiallyLoaded() const
 		return ParentView->GetIsSpatiallyLoaded();
 	}
 
+	if (bIsSpatiallyLoaded.IsSet())
+	{
+		return bIsSpatiallyLoaded.GetValue();
+	}
+
 	return Super::GetIsSpatiallyLoaded();
 }
 
@@ -298,19 +303,103 @@ bool FStreamingGenerationActorDescView::IsEditorOnlyReference(const FGuid& Refer
 
 FString FStreamingGenerationActorDescView::ToString(FWorldPartitionActorDesc::EToStringMode Mode) const
 {
-	FString Result = Super::ToString(Mode);
-	if (Mode >= FWorldPartitionActorDesc::EToStringMode::Compact)
-	{
-		Result += FString::Printf(TEXT(" View=(RuntimeGrid=%s RuntimeHLODLayer=%s bIsForcedNonSpatiallyLoaded=%s bIsForcedNoRuntimeGrid=%s bIsForcedNoDataLayers=%s bIsForceNoHLODLayer=%s bIsForcedHLODIrrelevant=%s)"), 
-			ResolvedRuntimeGrid.IsSet() ? *ResolvedRuntimeGrid->ToString() : TEXT("none"),
-			RuntimeHLODLayer.IsSet() ? *RuntimeHLODLayer->ToString() : TEXT("none"),
-			*LexToString(bIsForcedNonSpatiallyLoaded),
-			*LexToString(bIsForcedNoRuntimeGrid),
-			*LexToString(bIsForcedNoDataLayers),
-			*LexToString(bIsForceNoHLODLayer),
-			*LexToString(bIsForcedHLODIrrelevant));
+	return ToString(FStreamingGenerationActorDescView::EToStringMode::Full, Mode, [](const TArray<FName>&, TArray<const UDataLayerInstance*>&) { return false; });
+}
+
+FString FStreamingGenerationActorDescView::ToString(FStreamingGenerationActorDescView::EToStringMode ActorDescViewMode, FWorldPartitionActorDesc::EToStringMode ActorDescMode, TFunctionRef<bool(const TArray<FName>&, TArray<const UDataLayerInstance*>&)> GetRuntimeDataLayerInstancesFromNamesFunc) const
+{
+	FWorldPartitionActorDescStringBuilder Result(ActorDescMode);
+
+	if (ActorDescMode >= FWorldPartitionActorDesc::EToStringMode::Compact)
+	{
+		if (GetIsSpatiallyLoaded() != GetActorDesc()->GetIsSpatiallyLoaded())
+		{
+			Result.AppendFromBool(TEXT("SpatiallyLoaded"), GetIsSpatiallyLoaded());
+			if (bIsForcedNonSpatiallyLoaded)
+			{
+				Result.AppendFromBool(TEXT("IsForcedNonSpatiallyLoaded"), bIsForcedNonSpatiallyLoaded);
+			}
+		}
+
+		if (!GetEditorBounds().Equals(GetActorDesc()->GetEditorBounds()))
+		{
+			Result.AppendToString(TEXT("EditorBounds"), GetEditorBounds());
+		}
+
+		if (!GetRuntimeBounds().Equals(GetActorDesc()->GetRuntimeBounds()))
+		{
+			Result.AppendToString(TEXT("RuntimeBounds"), GetRuntimeBounds());
+		}
+
+		if (GetRuntimeGrid() != GetActorDesc()->GetRuntimeGrid())
+		{
+			Result.AppendToString(TEXT("RuntimeGrid"), GetRuntimeGrid());
+			if (bIsForcedNoRuntimeGrid)
+			{
+				Result.AppendFromBool(TEXT("IsForcedNoRuntimeGrid"), bIsForcedNoRuntimeGrid);
+			}
+		}
+
+		if (GetHLODLayer() != GetActorDesc()->GetHLODLayer())
+		{
+			Result.AppendToString(TEXT("HLODLayer"), GetHLODLayer());
+			if (bIsForceNoHLODLayer)
+			{
+				Result.AppendFromBool(TEXT("IsForceNoHLODLayer"), bIsForceNoHLODLayer);
+			}
+		}
+
+		if (GetActorIsHLODRelevant() != GetActorDesc()->GetActorIsHLODRelevant())
+		{
+			Result.AppendFromBool(TEXT("HLODRelevant"), GetActorIsHLODRelevant());
+			if (bIsForcedHLODIrrelevant)
+			{
+				Result.AppendFromBool(TEXT("IsForcedHLODIrrelevant"), bIsForcedHLODIrrelevant);
+			}
+		}
+
+		if (ActorDescMode >= FWorldPartitionActorDesc::EToStringMode::Full)
+		{
+			TArray<const UDataLayerInstance*> DataLayerInstances;
+			if (GetRuntimeDataLayerInstancesFromNamesFunc(GetRuntimeDataLayerInstanceNames().ToArray(), DataLayerInstances))
+			{
+				TArray<FName> DataLayerInstanceDisplayNames;
+				Algo::Transform(DataLayerInstances, DataLayerInstanceDisplayNames, [&](const UDataLayerInstance* DataLayerInstance)
+				{
+					return FName(*DataLayerInstance->GetDataLayerFullName());
+				});
+				if ((DataLayerInstanceDisplayNames.Num() != GetActorDesc()->GetDataLayers().Num()) ||
+					(!TSet(DataLayerInstanceDisplayNames).Includes(TSet(GetActorDesc()->GetDataLayers()))))
+				{
+					Result.AppendFromString(TEXT("DataLayers"), DataLayerInstanceDisplayNames.IsEmpty() ? TEXT("<Empty>") : FString::JoinBy(DataLayerInstanceDisplayNames, TEXT(","), [&](const FName& DataLayerName) { return DataLayerName.ToString(); }));
+					if (bIsForcedNoDataLayers)
+					{
+						Result.AppendFromBool(TEXT("IsForcedNoDataLayers"), bIsForcedNoDataLayers);
+					}
+				}
+			}
+		}
 	}
-	return Result;
+
+	FWorldPartitionActorDescStringBuilder FinalResult(ActorDescMode);
+	if (ActorDescViewMode == FStreamingGenerationActorDescView::EToStringMode::ViewOnly)
+	{
+		if (Result.Len())
+		{
+			FinalResult.AppendToString(TEXT("Guid"), GetGuid(), false);
+			FinalResult.Append(Result.ToString());
+		}
+		return FinalResult.ToString();
+	}
+	
+	FinalResult.Append(Super::ToString(ActorDescMode));
+	if (Result.Len())
+	{
+		FinalResult.Append(TEXT(" View=("));
+		FinalResult.Append(Result.ToString());
+		FinalResult.Append(TEXT(" )"));
+	}
+	return FinalResult.ToString();
 }
 
 const TArray<FGuid>& FStreamingGenerationActorDescView::GetEditorReferences() const
@@ -365,6 +454,11 @@ void FStreamingGenerationActorDescView::AddForcedInvalidReference(const FStreami
 	}
 }
 
+void FStreamingGenerationActorDescView::SetIsSpatiallyLoaded(bool bValue)
+{
+	bIsSpatiallyLoaded = bValue;
+}
+
 void FStreamingGenerationActorDescView::SetForcedNonSpatiallyLoaded()
 {
 	if (!bIsForcedNonSpatiallyLoaded)
@@ -404,7 +498,8 @@ void FStreamingGenerationActorDescView::SetForcedNoDataLayers()
 
 void FStreamingGenerationActorDescView::SetRuntimeDataLayerInstanceNames(const FDataLayerInstanceNames& InRuntimeDataLayerInstanceNames)
 {
-	check(!InRuntimeDataLayerInstanceNames.IsForcedEmptyNonExternalDataLayers())
+	check(!InRuntimeDataLayerInstanceNames.IsForcedEmptyNonExternalDataLayers() ||
+		(RuntimeDataLayerInstanceNames.IsSet() && (RuntimeDataLayerInstanceNames->IsForcedEmptyNonExternalDataLayers() == InRuntimeDataLayerInstanceNames.IsForcedEmptyNonExternalDataLayers())));
 	RuntimeDataLayerInstanceNames = InRuntimeDataLayerInstanceNames;
 	if (bIsForcedNoDataLayers)
 	{
@@ -566,17 +661,15 @@ class FWorldPartitionStreamingGenerator
 					if (bContainsUnfilteredActors)
 					{
 						FActorSetInstance& ActorSetInstance = ActorSetInstances.Emplace_GetRef();
-						const FContainerCollectionInstanceDescriptor::FPerInstanceData& PerInstanceData = ContainerCollectionInstanceDescriptor.GetPerInstanceData(ReferenceActorDescView.GetGuid());
-
 						ActorSetInstance.ActorSetContainerInstance = &ActorSetContainer;
 						ActorSetInstance.ActorSet = &ActorSet;
 						ActorSetInstance.FilteredActors = FilteredActors;
 						ActorSetInstance.ContainerID = ContainerCollectionInstanceDescriptor.ID;
 						ActorSetInstance.Transform = ContainerCollectionInstanceDescriptor.Transform;
-						ActorSetInstance.bIsSpatiallyLoaded = PerInstanceData.bIsSpatiallyLoaded;
+						ActorSetInstance.bIsSpatiallyLoaded = ReferenceActorDescView.GetIsSpatiallyLoaded(); //PerInstanceData.bIsSpatiallyLoaded;
 						ActorSetInstance.ContentBundleID = ContainerCollectionInstanceDescriptor.ContentBundleID;
-						ActorSetInstance.RuntimeGrid = PerInstanceData.RuntimeGrid;
-						ActorSetInstance.DataLayers = StreamingGenerator->GetRuntimeDataLayerInstances(PerInstanceData.DataLayers);
+						ActorSetInstance.RuntimeGrid = ReferenceActorDescView.GetRuntimeGrid(); //PerInstanceData.RuntimeGrid;
+						ActorSetInstance.DataLayers = StreamingGenerator->GetRuntimeDataLayerInstances(ReferenceActorDescView.GetRuntimeDataLayerInstanceNames().ToArray());
 
 						ActorSetInstance.Bounds.Init();
 						ActorSetInstance.ForEachActor([this, &ActorSetInstance](const FGuid& ActorGuid)
@@ -687,69 +780,15 @@ class FWorldPartitionStreamingGenerator
 		/** List of actor clusters for this descriptor */
 		TArray<TArray<FGuid>> Clusters;
 
-		struct FPerInstanceData
+		struct FCombinedData
 		{
 			bool bIsSpatiallyLoaded;
 			FName RuntimeGrid;
-			TArray<FName> DataLayers;
-
-			inline bool operator==(const FPerInstanceData& Other) const
-			{
-				// Assumes DataLayers are sorted
-				return
-					(bIsSpatiallyLoaded == Other.bIsSpatiallyLoaded) &&
-					(RuntimeGrid == Other.RuntimeGrid) &&
-					(DataLayers == Other.DataLayers);
-			}
-
-			inline bool operator!=(const FPerInstanceData& Other) const
-			{
-				return !operator==(Other);
-			}
-
-			friend inline uint32 GetTypeHash(const FPerInstanceData& InPerInstanceData)
-			{
-				FHashBuilder HashBuilder;
-				HashBuilder << InPerInstanceData.bIsSpatiallyLoaded << InPerInstanceData.RuntimeGrid << InPerInstanceData.DataLayers;
-				return HashBuilder.GetHash();
-			}
+			FDataLayerInstanceNames RuntimeDataLayerInstanceNames;
 		};
 
-		void AddPerInstanceData(const FGuid& ActorGuid, const FPerInstanceData& ActorInstanceData)
-		{
-			if (ActorInstanceData != InstanceData)
-			{
-				const FSetElementId UniquePerInstanceDataId = UniquePerInstanceData.Add(ActorInstanceData);
-				PerInstanceData.Emplace(ActorGuid, UniquePerInstanceDataId);
-			}
-		}
-
-		FPerInstanceData& GetPerInstanceData(const FGuid& ActorGuid)
-		{
-			if (const FSetElementId* PerInstanceDataId = PerInstanceData.Find(ActorGuid))
-			{
-				check(UniquePerInstanceData.IsValidId(*PerInstanceDataId));
-				return UniquePerInstanceData[*PerInstanceDataId];
-			}
-
-			return InstanceData;
-		}
-
-		const FPerInstanceData& GetPerInstanceData(const FGuid& ActorGuid) const
-		{
-			if (const FSetElementId* PerInstanceDataId = PerInstanceData.Find(ActorGuid))
-			{
-				check(UniquePerInstanceData.IsValidId(*PerInstanceDataId));
-				return UniquePerInstanceData[*PerInstanceDataId];
-			}
-
-			return InstanceData;
-		}
-
-		/** Per instance data */
-		FPerInstanceData InstanceData;
-		TSet<FPerInstanceData> UniquePerInstanceData;
-		TMap<FGuid, FSetElementId> PerInstanceData;
+		/** Data that contains the result of combining parent's FCombinedData with ContainerCollectionInstanceView */
+		FCombinedData ContainerCombinedData;
 
 		/** Map of actor descriptor mutators */
 		TMap<FGuid, FActorDescViewMutator> ActorDescViewMutators;
@@ -1106,15 +1145,16 @@ class FWorldPartitionStreamingGenerator
 	void CreateActorDescriptorViewsRecursive(FContainerCollectionInstanceDescriptor&& InContainerCollectionInstanceDescriptor)
 	{
 		// Inherited parent per-instance data logic
-		auto InheritParentContainerPerInstanceData = [this](const FContainerCollectionInstanceDescriptor& ParentContainerCollectionInstanceDescriptor, const FStreamingGenerationActorDescView& InActorDescView)
+		auto InheritParentContainerData = [this](const FContainerCollectionInstanceDescriptor& ParentContainerCollectionInstanceDescriptor, const FStreamingGenerationActorDescView& InActorDescView)
 		{
-			FContainerCollectionInstanceDescriptor::FPerInstanceData ResultPerInstanceData;
+			const FContainerCollectionInstanceDescriptor::FCombinedData& InParentData = ParentContainerCollectionInstanceDescriptor.ContainerCombinedData;
+			FContainerCollectionInstanceDescriptor::FCombinedData OutCombinedData;
 
 			// Apply AND logic on spatially loaded flag
-			ResultPerInstanceData.bIsSpatiallyLoaded = InActorDescView.GetIsSpatiallyLoaded() && ParentContainerCollectionInstanceDescriptor.InstanceData.bIsSpatiallyLoaded;
+			OutCombinedData.bIsSpatiallyLoaded = InActorDescView.GetIsSpatiallyLoaded() && InParentData.bIsSpatiallyLoaded;
 
 			// Runtime grid is inherited from the main world if the actor has its runtime grid set to none.
-			ResultPerInstanceData.RuntimeGrid = (ParentContainerCollectionInstanceDescriptor.ID.IsMainContainer() || ParentContainerCollectionInstanceDescriptor.InstanceData.RuntimeGrid.IsNone()) ? InActorDescView.GetRuntimeGrid() : ParentContainerCollectionInstanceDescriptor.InstanceData.RuntimeGrid;
+			OutCombinedData.RuntimeGrid = (ParentContainerCollectionInstanceDescriptor.ID.IsMainContainer() || InParentData.RuntimeGrid.IsNone()) ? InActorDescView.GetRuntimeGrid() : InParentData.RuntimeGrid;
 
 			// Set RuntimeGrid for Custom HLOD Actor based on actor's HLOD Layer to have it directly included in the appropriate grid
 			if (InActorDescView.GetActorNativeClass()->IsChildOf<AWorldPartitionCustomHLOD>())
@@ -1123,27 +1163,19 @@ class FWorldPartitionStreamingGenerator
 				{
 					if (HLODLayer->GetLayerType() == EHLODLayerType::CustomHLODActor)
 					{
-						FName ResolvedRuntimeGrid = GetResolvedRuntimeGridForHLODLayer(ResultPerInstanceData.RuntimeGrid, InActorDescView.GetHLODLayer());
+						FName ResolvedRuntimeGrid = GetResolvedRuntimeGridForHLODLayer(OutCombinedData.RuntimeGrid, InActorDescView.GetHLODLayer());
 						if (!ResolvedRuntimeGrid.IsNone())
 						{
-							ResultPerInstanceData.RuntimeGrid = ResolvedRuntimeGrid;
+							OutCombinedData.RuntimeGrid = ResolvedRuntimeGrid;
 						}
 					}
 				}
 			}
 
 			// Data layers are accumulated down the hierarchy chain, since level instances supports data layers assignation on actors
-			ResultPerInstanceData.DataLayers = InActorDescView.GetRuntimeDataLayerInstanceNames().ToArray();
-			ResultPerInstanceData.DataLayers.Append(ParentContainerCollectionInstanceDescriptor.InstanceData.DataLayers);
-			ResultPerInstanceData.DataLayers.Sort(FNameFastLess());
-
-			if (ParentContainerCollectionInstanceDescriptor.InstanceData.DataLayers.Num())
-			{
-				// Remove potential duplicates from sorted data layers array
-				ResultPerInstanceData.DataLayers.SetNum(Algo::Unique(ResultPerInstanceData.DataLayers));
-			}
-
-			return ResultPerInstanceData;
+			OutCombinedData.RuntimeDataLayerInstanceNames = InActorDescView.GetRuntimeDataLayerInstanceNames();
+			OutCombinedData.RuntimeDataLayerInstanceNames.Append(InParentData.RuntimeDataLayerInstanceNames);
+			return OutCombinedData;
 		};
 
 		// Hold on to ID
@@ -1221,7 +1253,7 @@ class FWorldPartitionStreamingGenerator
 				// Since Content Bundles streaming generation happens in its own context, all actor set instances must have the same content bundle GUID for now, so Level Instances
 				// placed inside a Content Bundle will propagate their Content Bundle GUID to child instances.
 				SubContainerInstanceDescriptor.ContentBundleID = ContainerCollectionInstanceDescriptor.ContentBundleID;
-				SubContainerInstanceDescriptor.InstanceData = InheritParentContainerPerInstanceData(ContainerCollectionInstanceDescriptor, ContainerCollectionInstanceView);
+				SubContainerInstanceDescriptor.ContainerCombinedData = InheritParentContainerData(ContainerCollectionInstanceDescriptor, ContainerCollectionInstanceView);
 				SubContainerInstanceDescriptor.HLODLayer = ContainerCollectionInstanceView.GetHLODLayer().IsValid() ? ContainerCollectionInstanceView.GetHLODLayer() : ContainerCollectionInstanceDescriptor.HLODLayer;
 				SubContainerInstanceDescriptor.bIsHLODRelevant = ContainerCollectionInstanceView.GetActorIsHLODRelevant() && ContainerCollectionInstanceDescriptor.bIsHLODRelevant;
 
@@ -1246,12 +1278,13 @@ class FWorldPartitionStreamingGenerator
 			ParentContainerCollection.Bounds += ContainerCollectionInstanceDescriptor.Bounds;
 		}
 
-		// Apply per-instance data
-		ContainerCollectionInstanceDescriptor.PerInstanceData.Reserve(ContainerCollectionInstanceDescriptor.ActorDescViewMap->Num());
-		ContainerCollectionInstanceDescriptor.ActorDescViewMap->ForEachActorDescView([&ContainerCollectionInstanceDescriptor, &InheritParentContainerPerInstanceData](FStreamingGenerationActorDescView& ActorDescView)
+		// Apply parent container data
+		ContainerCollectionInstanceDescriptor.ActorDescViewMap->ForEachActorDescView([&ContainerCollectionInstanceDescriptor, &InheritParentContainerData](FStreamingGenerationActorDescView& ActorDescView)
 		{
-			const FContainerCollectionInstanceDescriptor::FPerInstanceData PerInstanceData = InheritParentContainerPerInstanceData(ContainerCollectionInstanceDescriptor, ActorDescView);
-			ContainerCollectionInstanceDescriptor.AddPerInstanceData(ActorDescView.GetGuid(), PerInstanceData);
+			const FContainerCollectionInstanceDescriptor::FCombinedData CombinedData = InheritParentContainerData(ContainerCollectionInstanceDescriptor, ActorDescView);
+			ActorDescView.SetIsSpatiallyLoaded(CombinedData.bIsSpatiallyLoaded);
+			ActorDescView.SetRuntimeGrid(CombinedData.RuntimeGrid);
+			ActorDescView.SetRuntimeDataLayerInstanceNames(CombinedData.RuntimeDataLayerInstanceNames);
 		});
 	}
 
@@ -1267,7 +1300,7 @@ class FWorldPartitionStreamingGenerator
 		MainContainerCollectionInstance.ClusterMode = EContainerClusterMode::Partitioned;
 		MainContainerCollectionInstance.OwnerName = TEXT("MainContainer");
 		MainContainerCollectionInstance.ContentBundleID = InContainerInstanceCollection.GetContentBundleGuid();
-		MainContainerCollectionInstance.InstanceData.bIsSpatiallyLoaded = true; // Since we apply AND logic on spatially loaded flag recursively, startup value must be true
+		MainContainerCollectionInstance.ContainerCombinedData.bIsSpatiallyLoaded = true; // Since we apply AND logic on spatially loaded flag recursively, startup value must be true
 		MainContainerCollectionInstance.HLODLayer = WorldPartitionContext ? FSoftObjectPath(WorldPartitionContext->GetDefaultHLODLayer()) : FSoftObjectPath();
 		MainContainerCollectionInstance.bIsHLODRelevant = true;
 
@@ -1743,19 +1776,18 @@ class FWorldPartitionStreamingGenerator
 
 					for (const auto& [ActorGuid, ActorDescViewMutator] : ActorDescViewMutators)
 					{
-						FContainerCollectionInstanceDescriptor::FPerInstanceData PerInstanceData = ContainerCollectionInstanceDescriptor.GetPerInstanceData(ActorGuid);
-
-						if (ActorDescViewMutator.bIsSpatiallyLoaded.IsSet())
+						if (FStreamingGenerationActorDescView* ActorDescView = ContainerCollectionInstanceDescriptor.ActorDescViewMap->FindByGuid(ActorGuid))
 						{
-							PerInstanceData.bIsSpatiallyLoaded = ActorDescViewMutator.bIsSpatiallyLoaded.GetValue();
-						}
+							if (ActorDescViewMutator.bIsSpatiallyLoaded.IsSet())
+							{
+								ActorDescView->SetIsSpatiallyLoaded(ActorDescViewMutator.bIsSpatiallyLoaded.GetValue());
+							}
 
-						if (ActorDescViewMutator.RuntimeGrid.IsSet())
-						{
-							PerInstanceData.RuntimeGrid = ActorDescViewMutator.RuntimeGrid.GetValue();
+							if (ActorDescViewMutator.RuntimeGrid.IsSet())
+							{
+								ActorDescView->SetRuntimeGrid(ActorDescViewMutator.RuntimeGrid.GetValue());
+							}
 						}
-
-						ContainerCollectionInstanceDescriptor.AddPerInstanceData(ActorGuid, PerInstanceData);
 					}
 
 					TArray<FActorContainerID> ChildContainersIDs;
@@ -1809,18 +1841,16 @@ class FWorldPartitionStreamingGenerator
 
 			ContainerCollectionInstanceDescriptor.ActorDescViewMap->ForEachActorDescView([this, &ContainerCollectionInstanceDescriptor, &NbErrorsDetected, PassType](FStreamingGenerationActorDescView& ActorDescView)
 			{
-				FContainerCollectionInstanceDescriptor::FPerInstanceData& PerInstanceData = ContainerCollectionInstanceDescriptor.GetPerInstanceData(ActorDescView.GetGuid());
-
 				// Report invalid grid usage by spatially loaded actors
-				if (PerInstanceData.bIsSpatiallyLoaded && !IsValidGrid(PerInstanceData.RuntimeGrid, ActorDescView.GetActorNativeClass()))
+				if (ActorDescView.GetIsSpatiallyLoaded() && !IsValidGrid(ActorDescView.GetRuntimeGrid(), ActorDescView.GetActorNativeClass()))
 				{
 					if (PassType == EPassType::ErrorReporting)
 					{
-						ErrorHandler->OnInvalidRuntimeGrid(ActorDescView, PerInstanceData.RuntimeGrid);
+						ErrorHandler->OnInvalidRuntimeGrid(ActorDescView, ActorDescView.GetRuntimeGrid());
 					}
 					else
 					{
-						PerInstanceData.RuntimeGrid = NAME_None;
+						ActorDescView.SetRuntimeGrid(NAME_None);
 					}
 
 					NbErrorsDetected++;
@@ -1830,7 +1860,7 @@ class FWorldPartitionStreamingGenerator
 				// Validate that the HLOD layer + runtime grid combination is expected.
 				if (IWorldPartitionEditorModule::Get().GetRequireExplicitHLODLayerPartitionAssignation())
 				{
-				    if (ActorDescView.GetIsSpatiallyLoaded() && ActorDescView.GetHLODLayer().IsValid() && !IsValidHLODLayer(PerInstanceData.RuntimeGrid, ActorDescView.GetHLODLayer()))
+				    if (ActorDescView.GetIsSpatiallyLoaded() && ActorDescView.GetHLODLayer().IsValid() && !IsValidHLODLayer(ActorDescView.GetRuntimeGrid(), ActorDescView.GetHLODLayer()))
 				    {
 					    if (PassType == EPassType::ErrorReporting)
 					    {
@@ -2044,6 +2074,9 @@ public:
 	{
 		TRACE_CPUPROFILER_EVENT_SCOPE(FWorldPartitionStreamingGenerator::DumpStateLog);
 
+		// Force mode to EToStringMode::Full when running automated tests
+		FWorldPartitionActorDesc::EToStringMode Mode = GIsAutomationTesting ? FWorldPartitionActorDesc::EToStringMode::Full : FWorldPartitionActorDesc::EToStringMode::Compact;
+
 		// Build the containers tree representation
 		TMultiMap<FActorContainerID, FActorContainerID> InvertedContainersHierarchy;
 		for (auto& [ContainerID, ContainerCollectionInstanceDescriptor] : ContainerCollectionInstanceDescriptorsMap)
@@ -2069,47 +2102,32 @@ public:
 
 			UE_SCOPED_INDENT_LOG_ARCHIVE(Ar.PrintfIndent(TEXT("Container: %s"), *ContainerPackageName));
 
-			if (ContainerCollectionInstanceDescriptor.ActorDescViewMap->Num())
+			const UActorDescContainerInstance* Container = ContainerCollectionInstanceDescriptor.ContainerInstanceCollection->GetBaseContainerInstance();
+			if (Container->GetContainer()->GetActorDescCount())
 			{
 				UE_SCOPED_INDENT_LOG_ARCHIVE(Ar.PrintfIndent(TEXT("ActorDescs:")));
 
-				TMap<FGuid, FStreamingGenerationActorDescView*> SortedActorDescViewMap = ContainerCollectionInstanceDescriptor.ActorDescViewMap->ActorDescViewsByGuid;
-				SortedActorDescViewMap.KeySort([](const FGuid& GuidA, const FGuid& GuidB) { return GuidA < GuidB; });
-
-				for (auto& [ActorGuid, ActorDescView] : SortedActorDescViewMap)
-				{
-					Ar.Print(*ActorDescView->ToString(FWorldPartitionActorDesc::EToStringMode::Compact));
-				}
-			}
-
-			if (ContainerCollectionInstanceDescriptor.Clusters.Num())
-			{
-				UE_SCOPED_INDENT_LOG_ARCHIVE(Ar.PrintfIndent(TEXT("Clusters:")));
-
-				TArray<TArray<FGuid>> SortedClusters = ContainerCollectionInstanceDescriptor.Clusters;
-				for (TArray<FGuid>& ActorGuids : SortedClusters)
-				{
-					ActorGuids.Sort([](const FGuid& GuidA, const FGuid& GuidB) { return GuidA < GuidB; });
-				}
-				SortedClusters.Sort([](const TArray<FGuid>& GuidA, const TArray<FGuid>& GuidB) { return GuidA[0] < GuidB[0]; });
+				FActorDescList::FGuidActorDescMap SortedActorsByGuid = Container->GetContainer()->GetActorsByGuid();
+				SortedActorsByGuid.KeySort([](const FGuid& GuidA, const FGuid& GuidB) { return GuidA < GuidB; });
 
-				int ClusterIndex = 0;
-				for (TArray<FGuid>& ActorGuids : SortedClusters)
+				for (auto& [ActorGuid, ActorDesc] : SortedActorsByGuid)
 				{
-					UE_SCOPED_INDENT_LOG_ARCHIVE(Ar.PrintfIndent(TEXT("[%3d]"), ClusterIndex++));
-					for (const FGuid& ActorGuid : ActorGuids)
-					{
-						const FStreamingGenerationActorDescView& ActorDescView = ContainerCollectionInstanceDescriptor.ActorDescViewMap->FindByGuidChecked(ActorGuid);
-						Ar.Print(*ActorDescView.ToString(FWorldPartitionActorDesc::EToStringMode::Compact));
-					}
+					Ar.Print(*((*ActorDesc)->ToString(Mode)));
 				}
 			}
 		}
 
+		TSet<uint64> UniqueClusters;
+		auto GetRuntimeDataLayerInstancesFromNames = [this](const TArray<FName>& InDataLayerInstanceNames, TArray<const UDataLayerInstance*>& OutDataLayerInstances)
+		{
+			OutDataLayerInstances = GetRuntimeDataLayerInstances(InDataLayerInstanceNames);
+			return true;
+		};
+
 		Ar.Printf(TEXT("ContainerInstances:"));
-		auto DumpContainerInstances = [this, &InvertedContainersHierarchy, &Ar](const FActorContainerID& ContainerID)
+		auto DumpContainerInstances = [this, &InvertedContainersHierarchy, &Ar, Mode, &GetRuntimeDataLayerInstancesFromNames, &UniqueClusters](const FActorContainerID& ContainerID)
 		{
-			auto DumpContainerInstancesRecursive = [this, &InvertedContainersHierarchy, &Ar](const FActorContainerID& ContainerID, auto& RecursiveFunc) -> void
+			auto DumpContainerInstancesRecursive = [this, &InvertedContainersHierarchy, &Ar, Mode, &GetRuntimeDataLayerInstancesFromNames, &UniqueClusters](const FActorContainerID& ContainerID, auto& RecursiveFunc) -> void
 			{
 				const FContainerCollectionInstanceDescriptor& ContainerInstanceDescriptor = ContainerCollectionInstanceDescriptorsMap.FindChecked(ContainerID);
 
@@ -2122,6 +2140,68 @@ public:
 					Ar.Printf(TEXT("   Transform: %s"), *ContainerInstanceDescriptor.Transform.ToString());
 					Ar.Printf(TEXT("   Container: %s"), *ContainerInstanceDescriptor.ContainerInstanceCollection->GetBaseContainerInstancePackageName().ToString());
 					Ar.Printf(TEXT("HLODRelevant: %s"), ContainerInstanceDescriptor.bIsHLODRelevant ? TEXT("true") : TEXT("false"));
+
+					if (ContainerInstanceDescriptor.Clusters.Num())
+					{
+						TArray<TArray<FGuid>> SortedClusters = ContainerInstanceDescriptor.Clusters;
+						for (TArray<FGuid>& ActorGuids : SortedClusters)
+						{
+							ActorGuids.Sort([](const FGuid& GuidA, const FGuid& GuidB) { return GuidA < GuidB; });
+						}
+						SortedClusters.Sort([](const TArray<FGuid>& GuidA, const TArray<FGuid>& GuidB) { return GuidA[0] < GuidB[0]; });
+
+						FHashBuilder HashBuilder;
+						for (TArray<FGuid>& ActorGuids : SortedClusters)
+						{
+							for (const FGuid& ActorGuid : ActorGuids)
+							{
+								HashBuilder << ActorGuid;
+							}
+						}
+						const uint64 ClustersHash = HashBuilder.GetHash();
+						bool bIsAlreadySet = false;
+						UniqueClusters.Add(ClustersHash, &bIsAlreadySet);
+						
+						UE_SCOPED_INDENT_LOG_ARCHIVE(Ar.PrintfIndent(TEXT("Clusters (Hash:%llX)%s"), ClustersHash, bIsAlreadySet ? TEXT("") : TEXT(":")));
+
+						if (!bIsAlreadySet)
+						{
+							int ClusterIndex = 0;
+							for (TArray<FGuid>& ActorGuids : SortedClusters)
+							{
+								UE_SCOPED_INDENT_LOG_ARCHIVE(Ar.PrintfIndent(TEXT("[%3d]"), ClusterIndex++));
+								for (const FGuid& ActorGuid : ActorGuids)
+								{
+									const FStreamingGenerationActorDescView& ActorDescView = ContainerInstanceDescriptor.ActorDescViewMap->FindByGuidChecked(ActorGuid);
+									Ar.Print(*ActorDescView.ToString(FWorldPartitionActorDesc::EToStringMode::Guid));
+								}
+							}
+						}
+					}
+
+					TArray<FString> ActorDescViews;
+					if (ContainerInstanceDescriptor.ActorDescViewMap->Num())
+					{
+						TMap<FGuid, FStreamingGenerationActorDescView*> SortedActorDescViewMap = ContainerInstanceDescriptor.ActorDescViewMap->ActorDescViewsByGuid;
+						SortedActorDescViewMap.KeySort([](const FGuid& GuidA, const FGuid& GuidB) { return GuidA < GuidB; });
+						for (auto& [ActorGuid, ActorDescView] : SortedActorDescViewMap)
+						{
+							FString ActorDescViewResult = ActorDescView->ToString(FStreamingGenerationActorDescView::EToStringMode::ViewOnly, Mode, GetRuntimeDataLayerInstancesFromNames);
+							if (ActorDescViewResult.Len())
+							{
+								ActorDescViews.Add(MoveTemp(ActorDescViewResult));
+							}
+						}
+					}
+
+					if (ActorDescViews.Num())
+					{
+						UE_SCOPED_INDENT_LOG_ARCHIVE(Ar.PrintfIndent(TEXT("ActorDescViews:")));
+						for (const FString& ActorDescView : ActorDescViews)
+						{
+							Ar.Print(*ActorDescView);
+						}
+					}
 				}
 
 				TArray<FActorContainerID> ChildContainersIDs;
diff --git a/Engine/Source/Runtime/Engine/Public/WorldPartition/WorldPartitionActorDesc.h b/Engine/Source/Runtime/Engine/Public/WorldPartition/WorldPartitionActorDesc.h
--- a/Engine/Source/Runtime/Engine/Public/WorldPartition/WorldPartitionActorDesc.h
+++ b/Engine/Source/Runtime/Engine/Public/WorldPartition/WorldPartitionActorDesc.h
@@ -633,3 +633,45 @@ private:
 #endif // WITH_EDITOR
 };
 
+#if WITH_EDITOR
+
+class FWorldPartitionActorDescStringBuilder : public TStringBuilder<1024>
+{
+public:
+	FWorldPartitionActorDescStringBuilder(FWorldPartitionActorDesc::EToStringMode InMode);
+	void AppendFromString(const TCHAR* Name, const FString& Value, bool bUseLineStart = true);
+	void AppendFromBool(const TCHAR* Name, bool bValue, bool bUseLineStart = true);
+
+	template <typename Type>
+	void AppendToString(const TCHAR* Name, const Type& Value, bool bUseLineStart = true)
+	{
+		if constexpr (std::is_same_v<Type, FBox>)
+		{
+			if (!Value.IsValid)
+			{
+				AppendPair(Name, TEXT("IsValid=false"), bUseLineStart);
+				return;
+			}
+		}
+
+		if constexpr (std::is_same_v<Type, FWorldPartitionRelativeBounds>)
+		{
+			if (!Value.IsValid())
+			{
+				AppendPair(Name, TEXT("IsValid=false"), bUseLineStart);
+				return;
+			}
+		}
+
+		AppendFromString(Name, Value.ToString(), bUseLineStart);
+	}
+
+private:
+	void AppendPair(const TCHAR* Name, const TCHAR* Value, bool bUseLineStart);
+
+	FWorldPartitionActorDesc::EToStringMode Mode;
+	const TCHAR LineStart;
+	const TCHAR* LineEnd;
+};
+
+#endif // WITH_EDITOR
\ No newline at end of file
diff --git a/Engine/Source/Runtime/Engine/Public/WorldPartition/WorldPartitionStreamingGeneration.h b/Engine/Source/Runtime/Engine/Public/WorldPartition/WorldPartitionStreamingGeneration.h
--- a/Engine/Source/Runtime/Engine/Public/WorldPartition/WorldPartitionStreamingGeneration.h
+++ b/Engine/Source/Runtime/Engine/Public/WorldPartition/WorldPartitionStreamingGeneration.h
@@ -44,6 +44,12 @@ public:
 		FTopLevelAssetPath NativeClass;
 	};
 
+	enum class EToStringMode : uint8
+	{
+		ViewOnly,
+		Full
+	};
+
 	FStreamingGenerationActorDescView(const FStreamingGenerationActorDescViewMap& InActorDescViewMap, const FWorldPartitionActorDescInstance* InActorDescInstance, bool bInUnsaved = false)
 		: FStreamingGenerationActorDescView(InActorDescInstance)
 	{
@@ -72,6 +78,7 @@ public:
 
 	ENGINE_API void AddForcedInvalidReference(const FStreamingGenerationActorDescView* ReferenceView);
 	ENGINE_API const FInvalidReference* GetInvalidReference(const FGuid& InGuid) const;
+	ENGINE_API void SetIsSpatiallyLoaded(bool bValue);
 	ENGINE_API void SetForcedNonSpatiallyLoaded();
 	ENGINE_API void SetForcedNoRuntimeGrid();
 	ENGINE_API void SetForcedNoDataLayers();
@@ -99,6 +106,8 @@ public:
 	}
 
 private:
+	ENGINE_API FString ToString(EToStringMode ActorDescViewMode, FWorldPartitionActorDesc::EToStringMode ActorDescMode, TFunctionRef<bool(const TArray<FName>&, TArray<const UDataLayerInstance*>&)> GetRuntimeDataLayerInstancesFromNamesFunc) const;
+
 	const FStreamingGenerationActorDescViewMap* ActorDescViewMap;
 	const FStreamingGenerationActorDescView* ParentView;
 
@@ -118,6 +127,7 @@ private:
 	TOptional<FName> ResolvedRuntimeGrid;
 	TOptional<FDataLayerInstanceNames> ResolvedDataLayerInstanceNames;
 	TOptional<FDataLayerInstanceNames> RuntimeDataLayerInstanceNames;
+	TOptional<bool> bIsSpatiallyLoaded;
 	TOptional<TArray<FGuid>> RuntimeReferences;
 	TOptional<FSoftObjectPath> RuntimeHLODLayer;
 	TArray<FGuid> EditorReferences;
