Add option to world settings to enable new  navigation on new UGC maps

diff --git a/Engine/Source/Runtime/Engine/Classes/GameFramework/WorldSettings.h b/Engine/Source/Runtime/Engine/Classes/GameFramework/WorldSettings.h
--- a/Engine/Source/Runtime/Engine/Classes/GameFramework/WorldSettings.h
+++ b/Engine/Source/Runtime/Engine/Classes/GameFramework/WorldSettings.h
@@ -392,6 +392,10 @@ protected:
 	uint8 bEnableAISystem:1;
 
 public:
+	/** If set to true, use new UGC navigation (added through scene graph component). */
+	UPROPERTY(EditAnywhere, BlueprintReadOnly, config, Category=AI, AdvancedDisplay, meta=(DisplayName = "Use Automatic UGC Navigation"))
+	uint8 bUseAutomaticUGCNavigation:1;
+	
 	/** 
 	 * Enables tools for composing a tiled world. 
 	 * Level has to be saved and all sub-levels removed before enabling this option.
diff --git a/Engine/Source/Runtime/Engine/Private/WorldSettings.cpp b/Engine/Source/Runtime/Engine/Private/WorldSettings.cpp
--- a/Engine/Source/Runtime/Engine/Private/WorldSettings.cpp
+++ b/Engine/Source/Runtime/Engine/Private/WorldSettings.cpp
@@ -101,6 +101,7 @@ PRAGMA_ENABLE_DEPRECATION_WARNINGS
 	NavigationSystemConfig = nullptr;
 	bEnableAISystem = true;
 	AISystemClass = UAISystemBase::GetAISystemClassName();
+	bUseAutomaticUGCNavigation = false;
 	bEnableWorldComposition = false;
 	bEnableWorldOriginRebasing = false;
 #if WITH_EDITORONLY_DATA
