[Waveform Editor] - Transformation array has duplicate item button but doesn't allow duplicates

diff --git a/Engine/Source/Runtime/Engine/Classes/Sound/SoundWave.h b/Engine/Source/Runtime/Engine/Classes/Sound/SoundWave.h
--- a/Engine/Source/Runtime/Engine/Classes/Sound/SoundWave.h
+++ b/Engine/Source/Runtime/Engine/Classes/Sound/SoundWave.h
@@ -1093,7 +1093,7 @@ public:
 	 * Waveform edits to be applied to this SoundWave on cook (editing transformations will trigger a cook)
 	 * Transformations is a TArray because changing it to TSet will delete user data
 	 */
-	UPROPERTY(EditAnywhere, Instanced, Category = "Waveform Processing")
+	UPROPERTY(EditAnywhere, Instanced, Category = "Waveform Processing", meta = (NoElementDuplicate))
 	TArray<TObjectPtr<class UWaveformTransformationBase>> Transformations;
 #endif
 
