[CustomHLOD] Don't consider Custom HLOD Placeholder cells during IsStreamingCompleted query as they have no Level Streaming object

diff --git a/Engine/Source/Runtime/Engine/Private/WorldPartition/WorldPartitionStreamingPolicy.cpp b/Engine/Source/Runtime/Engine/Private/WorldPartition/WorldPartitionStreamingPolicy.cpp
--- a/Engine/Source/Runtime/Engine/Private/WorldPartition/WorldPartitionStreamingPolicy.cpp
+++ b/Engine/Source/Runtime/Engine/Private/WorldPartition/WorldPartitionStreamingPolicy.cpp
@@ -1149,6 +1149,12 @@ bool UWorldPartitionStreamingPolicy::IsStreamingCompleted(EWorldPartitionRuntime
 					bSkipCell = Cell->GetIsHLOD();
 				}
 
+				// Don't consider Custom HLOD Placeholder cells as they have no level streaming object
+				if (!bSkipCell && Cell->GetIsCustomHLODPlaceholderCell())
+				{
+					bSkipCell = true;
+				}
+
 				if (!bSkipCell && Cell->HasDataLayers())
 				{
 					if (!bAreQueryDLsInTargetState)
