Support for mapping an FNetRefHandle to a remote object and transactionalizing RPCs in Iris.

diff --git a/Engine/Plugins/Online/OnlineSubsystemUtils/Source/OnlineSubsystemUtils/Private/OnlineBeaconClient.cpp b/Engine/Plugins/Online/OnlineSubsystemUtils/Source/OnlineSubsystemUtils/Private/OnlineBeaconClient.cpp
--- a/Engine/Plugins/Online/OnlineSubsystemUtils/Source/OnlineSubsystemUtils/Private/OnlineBeaconClient.cpp
+++ b/Engine/Plugins/Online/OnlineSubsystemUtils/Source/OnlineSubsystemUtils/Private/OnlineBeaconClient.cpp
@@ -143,20 +143,21 @@ FUniqueNetIdRepl GetFirstPlayerAccountId(UWorld* World)
 }
 }
 
-bool AOnlineBeaconClient::InitClient(FURL& URL)
+bool AOnlineBeaconClient::InitClient(const FOnlineBeaconClientInitParams& Params)
 {
 	QUICK_SCOPE_CYCLE_COUNTER(STAT_OnlineBeaconClient_InitClient);
 
 	bool bSuccess = false;
 
-	if(URL.Valid)
+	if(Params.URL.Valid)
 	{
 		if (InitBase() && NetDriver)
 		{
 			NetDriver->SetWorld(GetWorld());
+			NetDriver->SetReplicateTransactionally(Params.bReplicateTransactionally);
 
 			FString Error;
-			if (NetDriver->InitConnect(this, URL, Error))
+			if (NetDriver->InitConnect(this, Params.URL, Error))
 			{
 				check(NetDriver->ServerConnection);
 				UWorld* World = GetWorld();
@@ -241,6 +242,13 @@ bool AOnlineBeaconClient::InitClient(FURL& URL)
 	return bSuccess;
 }
 
+bool AOnlineBeaconClient::InitClient(FURL& URL)
+{
+	FOnlineBeaconClientInitParams Params;
+	Params.URL = URL;
+	return InitClient(Params);
+}
+
 void AOnlineBeaconClient::Tick(float DeltaTime)
 {
 	if (NetDriver != nullptr && NetDriver->ServerConnection != nullptr)
diff --git a/Engine/Plugins/Online/OnlineSubsystemUtils/Source/OnlineSubsystemUtils/Public/OnlineBeaconClient.h b/Engine/Plugins/Online/OnlineSubsystemUtils/Source/OnlineSubsystemUtils/Public/OnlineBeaconClient.h
--- a/Engine/Plugins/Online/OnlineSubsystemUtils/Source/OnlineSubsystemUtils/Public/OnlineBeaconClient.h
+++ b/Engine/Plugins/Online/OnlineSubsystemUtils/Source/OnlineSubsystemUtils/Public/OnlineBeaconClient.h
@@ -27,6 +27,15 @@ enum class EBeaconConnectionState : uint8
 	Open = 3,	// Connection is open.
 };
 
+/**
+ * Parameters used when initializing AOnlineBeaconClient.
+ */
+struct FOnlineBeaconClientInitParams
+{
+	FURL URL;
+	bool bReplicateTransactionally = true;
+};
+
 /**
  * Delegate triggered on failures to connect to a host beacon
  */
@@ -61,6 +70,16 @@ class AOnlineBeaconClient : public AOnlineBeacon
 	UE_API virtual void OnFailure() override;
 	UE_API virtual void DestroyBeacon() override;
 	//~ End OnlineBeacon Interface
+	
+	/**
+	 * Initialize the client beacon with connection endpoint
+	 *	Creates the net driver and attempts to connect with the destination
+	 *
+	 * @param Params Parameters used for initialization of the client.
+	 *
+	 * @return true if connection is being attempted, false otherwise
+	 */
+	UE_API bool InitClient(const FOnlineBeaconClientInitParams& Params);
 
 	/**
 	 * Initialize the client beacon with connection endpoint
diff --git a/Engine/Plugins/Runtime/MultiServerReplication/Source/MultiServerReplication/Private/MultiServerBeaconClient.cpp b/Engine/Plugins/Runtime/MultiServerReplication/Source/MultiServerReplication/Private/MultiServerBeaconClient.cpp
--- a/Engine/Plugins/Runtime/MultiServerReplication/Source/MultiServerReplication/Private/MultiServerBeaconClient.cpp
+++ b/Engine/Plugins/Runtime/MultiServerReplication/Source/MultiServerReplication/Private/MultiServerBeaconClient.cpp
@@ -97,15 +97,17 @@ void AMultiServerBeaconClient::ConnectToServer(const FString& ConnectInfo)
 {
 	bool bSuccess = false;
 
-	FURL ConnectURL(NULL, *ConnectInfo, TRAVEL_Absolute);
-	if (InitClient(ConnectURL))
+	FOnlineBeaconClientInitParams Params;
+	Params.URL = FURL(nullptr, *ConnectInfo, TRAVEL_Absolute);
+	Params.bReplicateTransactionally = false;
+
+	if (InitClient(Params))
 	{
-		NetDriver->SetReplicateTransactionally(false);
 		bSuccess = true;
 	}
 	else
 	{
-		UE_LOG(LogMultiServerBeacon, Warning, TEXT("ConnectToRelay: Failure to init client beacon with %s."), *ConnectURL.ToString());
+		UE_LOG(LogMultiServerBeacon, Warning, TEXT("ConnectToRelay: Failure to init client beacon with %s."), *Params.URL.ToString());
 	}
 
 	if (!bSuccess)
diff --git a/Engine/Source/Runtime/Engine/Private/ActorReplication.cpp b/Engine/Source/Runtime/Engine/Private/ActorReplication.cpp
--- a/Engine/Source/Runtime/Engine/Private/ActorReplication.cpp
+++ b/Engine/Source/Runtime/Engine/Private/ActorReplication.cpp
@@ -1105,7 +1105,14 @@ void AActor::UpdateOwningNetConnection()
 	{
 		if (AActor* Actor = EveryChildren.Pop(EAllowShrinking::No))
 		{
-			EveryChildren.Append(Actor->Children);
+			for (TObjectPtr<AActor>& ActorChildPtr : Actor->Children)
+			{
+				if (!ActorChildPtr.IsRemote())
+				{
+					AActor* ActorChild = ActorChildPtr;
+					EveryChildren.Push(ActorChild);
+				}
+			}
 
 			if (Actor->GetIsReplicated())
 			{
diff --git a/Engine/Source/Runtime/Engine/Private/Net/Experimental/Iris/DataStreamChannel.cpp b/Engine/Source/Runtime/Engine/Private/Net/Experimental/Iris/DataStreamChannel.cpp
--- a/Engine/Source/Runtime/Engine/Private/Net/Experimental/Iris/DataStreamChannel.cpp
+++ b/Engine/Source/Runtime/Engine/Private/Net/Experimental/Iris/DataStreamChannel.cpp
@@ -46,6 +46,8 @@ namespace UE::Net::Private
 	static int32 GReplicationWriterMaxAllowedPacketsIfNotHugeObject = 3;
 	static FAutoConsoleVariableRef CVarReplicationWriterMaxAllowedPacketsIfNotHugeObject(TEXT("net.Iris.ReplicationWriterMaxAllowedPacketsIfNotHugeObject"), GReplicationWriterMaxAllowedPacketsIfNotHugeObject,
 		TEXT("Allow ReplicationWriter to overcommit data if we have more data to write."));
+
+	static FName ReceivedBunchTransactionName("UDataStreamChannel::ReceivedBunch");
 }
 
 UDataStreamChannel::UDataStreamChannel(const FObjectInitializer& ObjectInitializer)
@@ -132,7 +134,19 @@ void UDataStreamChannel::ReceivedBunch(FInBunch& Bunch)
 	// For packet stats
 	SerializationContext.SetTraceCollector(Connection->GetInTraceCollector());
 
-	DataStreamManager->ReadData(SerializationContext);
+#if UE_WITH_REMOTE_OBJECT_HANDLE && UE_AUTORTFM
+	if (Connection->Driver->ShouldReplicateTransactionally())
+	{
+		UE::RemoteExecutor::ExecuteTransactional(UE::Net::Private::ReceivedBunchTransactionName, [this, &SerializationContext]()
+		{
+			DataStreamManager->ReadData(SerializationContext);
+		});
+	}
+	else
+#endif
+	{
+		DataStreamManager->ReadData(SerializationContext);
+	}
 
 	// Set the bunch at the end
 	Bunch.SetAtEnd();
diff --git a/Engine/Source/Runtime/Engine/Private/NetDriver.cpp b/Engine/Source/Runtime/Engine/Private/NetDriver.cpp
--- a/Engine/Source/Runtime/Engine/Private/NetDriver.cpp
+++ b/Engine/Source/Runtime/Engine/Private/NetDriver.cpp
@@ -7505,25 +7505,29 @@ void UNetDriver::CreateReplicationSystem(bool bInitAsClient)
 		}
 
 #if UE_WITH_REMOTE_OBJECT_HANDLE
-		Params.ObjectToNetRefHandleSerial.BindLambda([](const TObjectPtr<const UObject>& ObjectPtr)
+		if (ProxyType == UE::Net::EProxyType::None)
 		{
-			const FRemoteObjectId RemoteObjectId = ObjectPtr.GetRemoteId();
-			const uint64 RemoteObjectIdValue = RemoteObjectId.GetIdNumber();
+			Params.ObjectToNetRefHandleSerial.BindLambda([](const TObjectPtr<const UObject>& ObjectPtr)
+			{
+				const FRemoteObjectId RemoteObjectId = ObjectPtr.GetRemoteId();
+				const uint64 RemoteObjectIdValue = RemoteObjectId.GetIdNumber();
 
-			// Shift away the reserved bit in FRemoteObjectId so we're just left with the server id and serial number
-			// which will be embedded in the NetRefHandle serial.
-			const uint64 Serial = (RemoteObjectIdValue >> 1);
+				// Shift away the reserved bit in FRemoteObjectId so we're just left with the server id and serial number
+				// which will be embedded in the NetRefHandle serial.
+				const uint64 Serial = (RemoteObjectIdValue >> 1);
 
-			return Serial;
-		});
+				return Serial;
+			});
 
-		// Disabled until the calling code in Iris has been transactionalized.
-		/*Params.NetRefHandleSerialToObject.BindLambda([](const uint64 Serial)
-		{
-			const FRemoteObjectId RemoteObjectId = FRemoteObjectId::CreateFromInt(Serial).GetLocalized();
-			TObjectPtr<UObject> RemoteObjectPtr = TObjectPtr<UObject>(FObjectPtr(RemoteObjectId));
-			return RemoteObjectPtr;
-		});*/
+			Params.NetRefHandleSerialToObject.BindLambda([this](const uint64 Serial)
+			{
+				// Shift left to made room for the reserved bit on FRemoteObjectId.
+				const uint64 RemoteObjectIdValue = (Serial << 1);
+				const FRemoteObjectId RemoteObjectId = FRemoteObjectId::CreateFromInt(RemoteObjectIdValue).GetLocalized();
+
+				return TObjectPtr<UObject>(FObjectPtr(RemoteObjectId));
+			});
+		}
 #endif
 
 		SetReplicationSystem(UE::Net::FReplicationSystemFactory::CreateReplicationSystem(Params));
diff --git a/Engine/Source/Runtime/Net/Iris/Private/Iris/ReplicationSystem/ObjectReferenceCache.cpp b/Engine/Source/Runtime/Net/Iris/Private/Iris/ReplicationSystem/ObjectReferenceCache.cpp
--- a/Engine/Source/Runtime/Net/Iris/Private/Iris/ReplicationSystem/ObjectReferenceCache.cpp
+++ b/Engine/Source/Runtime/Net/Iris/Private/Iris/ReplicationSystem/ObjectReferenceCache.cpp
@@ -1186,16 +1186,18 @@ ENetObjectReferenceResolveResult FObjectReferenceCache::ResolveObjectReference(c
 			{
 				UObject* ReplicatedOuter = nullptr;
 
-				if (NetRefHandleSerialToObject.IsBound())
-				{
-					ReplicatedOuter = NetRefHandleSerialToObject.Execute(Reference.GetRefHandle().GetSerial());
-				}
-
 				if (ReplicatedOuter == nullptr)
 				{
 					bool bLocalMustBeMapped = false;
 					ReplicatedOuter = ResolveObjectReferenceHandleInternal(Reference.GetRefHandle(), ResolveContext, bLocalMustBeMapped);
 				}
+				
+				if (ReplicatedOuter == nullptr && NetRefHandleSerialToObject.IsBound())
+				{
+					const uint64 Serial = Reference.GetRefHandle().GetSerial();
+					ReplicatedOuter = NetRefHandleSerialToObject.Execute(Serial);
+				}
+
 				ResolvedObject = ReplicatedOuter ? StaticFindObject(UObject::StaticClass(), ReplicatedOuter, ToCStr(ObjectPath), EFindObjectFlags::None) : nullptr;
 
 				if (ResolvedObject == nullptr)
@@ -1231,6 +1233,12 @@ ENetObjectReferenceResolveResult FObjectReferenceCache::ResolveObjectReference(c
 	else
 	{
 		ResolvedObject = ResolveObjectReferenceHandleInternal(Reference.GetRefHandle(), ResolveContext, bMustBeMapped);
+
+		if (ResolvedObject == nullptr && Reference.GetRefHandle().IsValid() && NetRefHandleSerialToObject.IsBound())
+		{
+			const uint64 Serial = Reference.GetRefHandle().GetSerial();
+			ResolvedObject = NetRefHandleSerialToObject.Execute(Reference.GetRefHandle().GetSerial());
+		}
 	}
 	
 	ENetObjectReferenceResolveResult Result = ENetObjectReferenceResolveResult::None;
diff --git a/Engine/Source/Runtime/Net/Iris/Public/Iris/ReplicationSystem/ReplicationSystem.h b/Engine/Source/Runtime/Net/Iris/Public/Iris/ReplicationSystem/ReplicationSystem.h
--- a/Engine/Source/Runtime/Net/Iris/Public/Iris/ReplicationSystem/ReplicationSystem.h
+++ b/Engine/Source/Runtime/Net/Iris/Public/Iris/ReplicationSystem/ReplicationSystem.h
@@ -192,7 +192,9 @@ public:
 		UE::Net::FObjectToNetRefHandleSerial ObjectToNetRefHandleSerial;
 
 		/** 
-		 * Map a unique serial field from a NetRefHandle to a UObject pointer.
+		 * Map a unique serial field from a NetRefHandle serial to a UObject pointer.
+		 *
+		 * Will always be passed a serial field that is greater than 0. 
 		 *
 		 * Assumes there's a one-to-one mapping between UObjects registered
 		 * for replication and the serial field in NetRefHandle.
@@ -757,6 +759,8 @@ public:
 	 */
 	IRISCORE_API UObject* GetConnectionUserData(uint32 ConnectionId) const;
 
+	bool IsReplicatingTransactionally() const { return bReplicateTransactionally; }
+
 	int32 GetPIEInstanceID() const { return PIEInstanceID; }
 
 	/** Set the cull distance for an object temporarily instead of the default cull distance value associated with the object (ex: @see AActor::NetCullDistanceSquared) */
