[CustomHLOD] Ensure that Custom HLOD Actors assigned to always loaded HLOD layers are correctly included in the persistent level actor set

diff --git a/Engine/Source/Runtime/Engine/Private/WorldPartition/WorldPartitionStreamingGeneration.cpp b/Engine/Source/Runtime/Engine/Private/WorldPartition/WorldPartitionStreamingGeneration.cpp
--- a/Engine/Source/Runtime/Engine/Private/WorldPartition/WorldPartitionStreamingGeneration.cpp
+++ b/Engine/Source/Runtime/Engine/Private/WorldPartition/WorldPartitionStreamingGeneration.cpp
@@ -42,6 +42,7 @@
 #include "WorldPartition/HLOD/StandaloneHLODActor.h"
 #include "WorldPartition/HLOD/StandaloneHLODSubsystem.h"
 #include "WorldPartition/RuntimeHashSet/WorldPartitionRuntimeHashSet.h"
+#include "WorldPartition/RuntimeHashSet/RuntimePartitionPersistent.h"
 #include "WorldPartition/WorldPartitionSubsystem.h"
 #include "WorldPartition/WorldPartitionRuntimeContainerResolving.h"
 #include "HAL/FileManager.h"
@@ -782,6 +783,11 @@ class FWorldPartitionStreamingGenerator
 			{
 				if (const URuntimePartition* RuntimePartition = HashSet->ResolveRuntimePartitionForHLODLayer(RuntimeGrid, HLODLayerPath))
 				{
+					if (RuntimePartition->IsA<URuntimePartitionPersistent>())
+					{
+						return NAME_None;
+					}
+					
 					FName GridName = RuntimeGrid.IsNone() ? DefaultGrid : RuntimeGrid;
 					return FName(GridName.ToString() + TEXT(":") + RuntimePartition->Name.ToString());
 				}
@@ -957,6 +963,10 @@ class FWorldPartitionStreamingGenerator
 								{
 									CustomHLODPlaceholderActor->SetRuntimeGrid(ResolvedRuntimeGrid);
 								}
+								else
+								{
+									CustomHLODPlaceholderActor->SetIsSpatiallyLoaded(false);
+								}
 							}
 						}
 					}
@@ -1128,6 +1138,10 @@ class FWorldPartitionStreamingGenerator
 						{
 							ResultPerInstanceData.RuntimeGrid = ResolvedRuntimeGrid;
 						}
+						else
+						{
+							ResultPerInstanceData.bIsSpatiallyLoaded = false;
+						}
 					}
 				}
 			}
