Fix slowdown when entering long console commands.
When entering long console commands, especially those with many parameters, the main tread would begin to hitch with each keypress.  This change fixes the hitching.

diff --git a/Engine/Source/Runtime/Engine/Classes/Engine/Console.h b/Engine/Source/Runtime/Engine/Classes/Engine/Console.h
--- a/Engine/Source/Runtime/Engine/Classes/Engine/Console.h
+++ b/Engine/Source/Runtime/Engine/Classes/Engine/Console.h
@@ -213,8 +213,8 @@ class UConsole
 	/** Virtual function to allow subclasses of UConsole to add additional commands */
 	ENGINE_API virtual void AugmentRuntimeAutoCompleteList(TArray<FAutoCompleteCommand>& List);
 	
-	/** Update the auto complete indices from the typed string */
-	ENGINE_API virtual void UpdateCompleteIndices();
+	/** Update the auto complete indices from the typed string, optionally using a hint about whether our search term has just gotten more specific (i.e. narrowing down the search) */
+	ENGINE_API virtual void UpdateCompleteIndices(bool bIsSearchNarrowing = false);
 
 	/**
 	 * Process a character input event (typing) routed through unrealscript from another object. This method is assigned as the value for the
diff --git a/Engine/Source/Runtime/Engine/Private/UserInterface/Console.cpp b/Engine/Source/Runtime/Engine/Private/UserInterface/Console.cpp
--- a/Engine/Source/Runtime/Engine/Private/UserInterface/Console.cpp
+++ b/Engine/Source/Runtime/Engine/Private/UserInterface/Console.cpp
@@ -476,13 +476,21 @@ void CommandToStringArray(const FAutoCompleteCommand& Command, OUT TArray< FStri
 	StringArray.Add(Command.Command);
 }
 
-void UConsole::UpdateCompleteIndices()
+void UConsole::UpdateCompleteIndices(bool bIsSearchNarrowing /*= false*/)
 {
 	if (!bIsRuntimeAutoCompleteUpToDate)
 	{
 		BuildRuntimeAutoCompleteList(true);
 	}
 
+	// If we've just narrowed our search by appending characters to the end, yet we've already determined that
+	// no console commands match our entry, then early out.  It's impossible to discover any new suggestions
+	// by appending text.
+	if (bIsSearchNarrowing && AutoComplete.IsEmpty())
+	{
+		return;
+	}
+
 	AutoComplete.Empty();
 	AutoCompleteIndex = 0;
 	AutoCompleteCursor = 0;
@@ -809,23 +817,34 @@ bool UConsole::ProcessShiftKey(FKey Key, EInputEvent Event)
 
 void UConsole::AppendInputText(const FString& Text)
 {
+	// If we already have text, and our insertion point is at the end, then we're appending this text.
+	// Later, we'll pass this hint along to the autocomplete code, who uses it to potentially bypass a LOT of expensive code.
+	const bool bIsAppending = (TypedStrPos > 0 && TypedStrPos == TypedStr.Len());
+
 	FString TextMod = Text;
-	while (TextMod.Len() > 0)
+	int32 CurrentPos = 0;
+
+	while (CurrentPos < TextMod.Len())
 	{
-		int32 Character = **TextMod.Left(1);
-		TextMod.MidInline(1, MAX_int32, EAllowShrinking::No);
+		const int32 Character = static_cast<int32>(TextMod[CurrentPos]);
 
-		// Discard control characters, and allow more than extended ASCII (0xff) for Arabic etc
-		if (Character >= 0x20)
+		// Omit all control characters
+		if (Character < 0x20)
 		{
-			TCHAR Temp[2];
-			Temp[0] = Character;
-			Temp[1] = 0;
-			SetInputText(FString::Printf(TEXT("%s%s%s"), *TypedStr.Left(TypedStrPos), Temp, *TypedStr.Right(TypedStr.Len() - TypedStrPos)));
-			SetCursorPos(TypedStrPos + 1);
+			TextMod.RemoveAt(CurrentPos, EAllowShrinking::No);	// No reallocs
 		}
-	};
-	UpdateCompleteIndices();
+		else
+		{
+			++CurrentPos;
+		}
+	}
+
+	SetInputText(
+		FString::Format(TEXT("{0}{1}{2}"), { FStringView(TypedStr).Left(TypedStrPos), TextMod, FStringView(TypedStr).Right(TypedStr.Len() - TypedStrPos) }));
+
+	SetCursorPos(TypedStrPos + TextMod.Len());
+
+	UpdateCompleteIndices(bIsAppending);
 	UpdatePrecompletedInputLine();
 }
 
@@ -1109,9 +1128,15 @@ bool UConsole::InputKey_InputLine(FInputDeviceId DeviceId, FKey Key, EInputEvent
 					NewPos = TypedStrPos - 1;
 				}
 
-				SetInputText(FString::Printf(TEXT("%s%s"), *TypedStr.Left(NewPos), *TypedStr.Right(TypedStr.Len() - TypedStrPos)));
+				SetInputText(
+					FString::Format(TEXT("{0}{1}"), { FStringView(TypedStr).Left(NewPos), FStringView(TypedStr).Right(TypedStr.Len() - TypedStrPos) }));
+
 				SetCursorPos(NewPos);
 
+				// Update autocomplete
+				UpdateCompleteIndices();
+				UpdatePrecompletedInputLine();
+
 				// unlock auto-complete (@todo - track the lock position so we don't bother unlocking under bogus cases)
 				bAutoCompleteLocked = false;
 			}
@@ -1133,7 +1158,12 @@ bool UConsole::InputKey_InputLine(FInputDeviceId DeviceId, FKey Key, EInputEvent
 					RightStart = TypedStrPos + 1;
 				}
 
-				SetInputText(FString::Printf(TEXT("%s%s"), *TypedStr.Left(TypedStrPos), *TypedStr.Right(TypedStr.Len() - RightStart)));
+				SetInputText(
+					FString::Format(TEXT("{0}{1}"), { FStringView(TypedStr).Left(TypedStrPos), FStringView(TypedStr).Right(TypedStr.Len() - RightStart) }));
+
+				// Update autocomplete
+				UpdateCompleteIndices();
+				UpdatePrecompletedInputLine();
 			}
 			return true;
 		}
