Fix issue with attached actors moving to the wrong location when mirroring

- Fix similar issue with rotation

diff --git a/Engine/Source/Runtime/Engine/Private/ActorEditor.cpp b/Engine/Source/Runtime/Engine/Private/ActorEditor.cpp
--- a/Engine/Source/Runtime/Engine/Private/ActorEditor.cpp
+++ b/Engine/Source/Runtime/Engine/Private/ActorEditor.cpp
@@ -948,12 +948,12 @@ void AActor::EditorApplyMirror(const FVector& MirrorScale, const FVector& PivotL
 
 	if( RootComponent != NULL )
 	{
-		GetRootComponent()->SetRelativeRotationExact( NewRot.Rotator() );
+		GetRootComponent()->SetWorldRotation(NewRot.Rotator());
 		FVector Loc = GetActorLocation();
 		Loc -= PivotLocation;
 		Loc *= MirrorScale;
 		Loc += PivotLocation;
-		GetRootComponent()->SetRelativeLocation( Loc );
+		GetRootComponent()->SetWorldLocation(Loc);
 
 		FVector Scale3D = GetRootComponent()->GetRelativeScale3D();
 		Scale3D.X = -Scale3D.X;
