[Net] Do not delay first handshake package for a full resend interval, try to send on the next frame as soon as the socket becomes available

#rnx

diff --git a/Engine/Source/Runtime/Engine/Private/PacketHandlers/StatelessConnectHandlerComponent.cpp b/Engine/Source/Runtime/Engine/Private/PacketHandlers/StatelessConnectHandlerComponent.cpp
--- a/Engine/Source/Runtime/Engine/Private/PacketHandlers/StatelessConnectHandlerComponent.cpp
+++ b/Engine/Source/Runtime/Engine/Private/PacketHandlers/StatelessConnectHandlerComponent.cpp
@@ -725,12 +725,16 @@ void StatelessConnectHandlerComponent::SendToServer(EHandshakeVersion HandshakeV
 				FOutPacketTraits Traits;
 
 				Driver->ServerConnection->LowLevelSend(Packet.GetData(), Packet.GetNumBits(), Traits);
+				LastClientSendTimestamp = Driver->GetElapsedTime();
+			}
+			else
+			{
+				// ignore resent interval, send initial packet on "next frame" as the socket will be ready by then
+				LastClientSendTimestamp = Driver->GetElapsedTime() - UE::Net::HandshakeResendInterval + 0.01f;
 			}
 		}
 
 		Handler->SetRawSend(false);
-
-		LastClientSendTimestamp = Driver->GetElapsedTime();
 	}
 }
 
