SoundWaveProxyPlayer

- Fix SoundWaveProxyPlayer not resetting Resampler properly
- Add MonoUpmixMethod to ProxyPlayerSettings to provide option to

#rnx

diff --git a/Engine/Source/Runtime/Engine/Classes/Sound/SoundWaveProxyPlayer.h b/Engine/Source/Runtime/Engine/Classes/Sound/SoundWaveProxyPlayer.h
--- a/Engine/Source/Runtime/Engine/Classes/Sound/SoundWaveProxyPlayer.h
+++ b/Engine/Source/Runtime/Engine/Classes/Sound/SoundWaveProxyPlayer.h
@@ -5,6 +5,7 @@
 #include "AudioDecompress.h"
 #include "Containers/Array.h"
 #include "Containers/ArrayView.h"
+#include "DSP/ChannelMap.h"
 #include "DSP/ConvertDeinterleave.h"
 #include "DSP/RuntimeResampler.h"
 #include "HAL/Platform.h"
@@ -27,6 +28,7 @@ public:
 			, OutputNumChannels(InOutputNumChannels)
 			, NumFramesPerBlock(InNumFramesPerBlock)
 		{
+			check(NumFramesPerBlock > 0);
 			check(OutputSampleRate > 0.0f);
 			check(OutputNumChannels > 0);
 			check(OutputNumChannels <= 8);
@@ -34,6 +36,8 @@ public:
 		
 		float OutputSampleRate = 0.0f;
 		int32 OutputNumChannels = 0;
+		// Method for upmixing mono audio to multichannel output
+		Audio::EChannelMapMonoUpmixMethod MonoUpmixMethod = Audio::EChannelMapMonoUpmixMethod::FullVolume;
 		int32 NumFramesPerBlock = 0;
 		int32 MaxDecodeSizeFrames = 1024;
 		bool bMaintainAudioSync = false;
diff --git a/Engine/Source/Runtime/Engine/Private/SoundWaveProxyPlayer.cpp b/Engine/Source/Runtime/Engine/Private/SoundWaveProxyPlayer.cpp
--- a/Engine/Source/Runtime/Engine/Private/SoundWaveProxyPlayer.cpp
+++ b/Engine/Source/Runtime/Engine/Private/SoundWaveProxyPlayer.cpp
@@ -64,25 +64,23 @@ bool FSoundWaveProxyPlayer::SetSoundWave(FSoundWaveProxyPtr InWaveProxy)
 	SoundWave = InWaveProxy;
 	SourceFrameIndex = 0;
 	AudioSyncFrameIndex = 0;
-	// update the loop with the cached loop settings to update the new
+	// update the loop with the cached loop settings to update the loop start/end frames
 	SetLoop(bIsLooping, CurrentLoopStartTime, CurrentLoopDuration);
 	bFallbackSeekMethodWarningLogged = false;
-	bRequiresResampling = false;
 	bHasNearlyFinished = false;
 	SortedCuePoints = SoundWave->GetCuePointsAndLoopRegions();
 	Algo::SortBy(SortedCuePoints, &FSoundWaveCuePoint::FramePosition);
 
+	bRequiresResampling = false;
 	RuntimeResampler.Reset(NumChannels);
+	ResampledAudio.Reset();
+	SourceAudio.Reset();
 	
 	// Create algorithm for channel conversion and deinterleave 
 	Audio::FConvertDeinterleaveParams ConvertDeinterleaveParams;
 	ConvertDeinterleaveParams.NumInputChannels = NumChannels;
 	ConvertDeinterleaveParams.NumOutputChannels = Settings.OutputNumChannels;
-	// Original implementation of MetaSound WavePlayer upmixed 
-	// mono using FullVolume. In the future, the mono upmix 
-	// method may be exposed as a node input to facilitate 
-	// better control.
-	ConvertDeinterleaveParams.MonoUpmixMethod = Audio::EChannelMapMonoUpmixMethod::FullVolume;
+	ConvertDeinterleaveParams.MonoUpmixMethod = Settings.MonoUpmixMethod;
 	ConvertDeinterleave = Audio::IConvertDeinterleave::Create(ConvertDeinterleaveParams);
 
 	return true;
@@ -666,6 +664,7 @@ bool FSoundWaveProxyPlayer::SeekToFrame(int32 InFrame)
 	bRequiresResampling = false;
 	RuntimeResampler.Reset(NumChannels);
 	ResampledAudio.Reset();
+	SourceAudio.Reset();
 
 	// reset the "NearlyFinished" flag
 	bHasNearlyFinished = false;
