Minor fixes for focus summary.

Missed initializing timer after refactor and added comment for clarity on our focus path iteration.

diff --git a/Engine/Source/Runtime/Engine/Private/Analytics/EditorAnalyticsSessionSummary.cpp b/Engine/Source/Runtime/Engine/Private/Analytics/EditorAnalyticsSessionSummary.cpp
--- a/Engine/Source/Runtime/Engine/Private/Analytics/EditorAnalyticsSessionSummary.cpp
+++ b/Engine/Source/Runtime/Engine/Private/Analytics/EditorAnalyticsSessionSummary.cpp
@@ -58,6 +58,7 @@ FEditorAnalyticsSessionSummary::FEditorAnalyticsSessionSummary(TSharedPtr<IAnaly
 	: FEngineAnalyticsSessionSummary(MoveTemp(Store), MonitorProcessId)
 	, LastUserActivityTimeSecs(FPlatformTime::Seconds())
 	, AccountedUserIdleSecs(0)
+	, TimeOfLastUpdate(FPlatformTime::Seconds())
 {
 	EditorAnalyticsProperties::UserInteractionCount.Set(GetStore(), 0);
 	EditorAnalyticsProperties::Idle1Min.Set(GetStore(), 0);
@@ -163,6 +164,7 @@ namespace EditorAnalyticsSessionSummaryHelpers
 
 		const FWeakWidgetPath& WeakFocusPath = SlateUser->GetWeakFocusPath();
 
+		// intentionally don't check 0 index as that'll be the window
 		for (int i = WeakFocusPath.Widgets.Num() - 1; i > 0; --i)
 		{
 			TWeakPtr<SWidget> WidgetInPathWeakPtr = WeakFocusPath.Widgets[i];
@@ -205,6 +207,7 @@ void FEditorAnalyticsSessionSummary::UpdateUserFocusTimes(double CurrTimeSecs)
 	TSharedPtr<const FSlateUser> CurrentUser = SlateApplication.GetCursorUser();
 
 	double PrevTimeOfLastUpdate = TimeOfLastUpdate.load();
+
 	// similar to our other timer updates, we want to ensure only a single thread performs this update
 	if (TimeOfLastUpdate.compare_exchange_strong(PrevTimeOfLastUpdate, CurrTimeSecs))
 	{
