[Audio golden path] Fix wave file hot reloading by overriding SoundWaveData

#tests Tested in cooked build
#rnx

diff --git a/Engine/Plugins/Runtime/Metasound/Source/MetasoundEngine/Private/MetasoundWavePlayerNode.cpp b/Engine/Plugins/Runtime/Metasound/Source/MetasoundEngine/Private/MetasoundWavePlayerNode.cpp
--- a/Engine/Plugins/Runtime/Metasound/Source/MetasoundEngine/Private/MetasoundWavePlayerNode.cpp
+++ b/Engine/Plugins/Runtime/Metasound/Source/MetasoundEngine/Private/MetasoundWavePlayerNode.cpp
@@ -827,6 +827,9 @@ namespace Metasound
 
 			if (WaveProxy.IsValid())
 			{
+				// Make sure we have up-to-date SoundWaveData
+				WaveProxy->TryGetLatest();
+
 				UE_LOG(LogMetaSound, Verbose, TEXT("Starting Sound: '%s'"), *CurrentWaveAsset->GetFName().ToString());
 
 				// Create local sorted copy of cue points and loop regions.
diff --git a/Engine/Source/Runtime/AudioExtensions/Public/IAudioDecoderOverride.h b/Engine/Source/Runtime/AudioExtensions/Public/IAudioDecoderOverride.h
--- a/Engine/Source/Runtime/AudioExtensions/Public/IAudioDecoderOverride.h
+++ b/Engine/Source/Runtime/AudioExtensions/Public/IAudioDecoderOverride.h
@@ -2,7 +2,7 @@
 
 #pragma once
 
-#include "CoreMinimal.h"
+#include "UObject/NameTypes.h"
 
 class ICompressedAudioInfo;
 
diff --git a/Engine/Source/Runtime/Engine/Classes/Sound/SoundWave.h b/Engine/Source/Runtime/Engine/Classes/Sound/SoundWave.h
--- a/Engine/Source/Runtime/Engine/Classes/Sound/SoundWave.h
+++ b/Engine/Source/Runtime/Engine/Classes/Sound/SoundWave.h
@@ -1558,9 +1558,9 @@ public:
 		return (ESoundWavePrecacheState)PrecacheState.GetValue();
 	}
 
-	// Call this to artificially cause any consumers of the wave data to refresh their state.
+	// Call this if you intend to supply an overriden decoder at playback time.
 	UE_EXPERIMENTAL(5.8, "May be altered or removed")
-	ENGINE_API void TriggerWaveDataRefresh();
+	ENGINE_API void OverrideWaveData(int32 InNumChannels, int32 InNumFrames, int32 InSampleRate);
 
 private:
 	TSharedRef<TDataSnapshotProducer<FSoundWaveData>> SoundWaveData = MakeShared<TDataSnapshotProducer<FSoundWaveData>>();
@@ -1702,7 +1702,7 @@ private:
 	ESoundWaveLoadingBehavior LoadingBehavior = ESoundWaveLoadingBehavior::Uninitialized;
 
 #if WITH_EDITORONLY_DATA
-	// Set by CacheInheritedLoadingBehavior after traversing the Soundclass heirarchy 
+	// Set by CacheInheritedLoadingBehavior after traversing the Soundclass hierarchy 
 	FPerPlatformFloat SizeOfFirstAudioChunkInSeconds = 0.f;
 #endif //WITH_EDITORONLY_DATA
 	
diff --git a/Engine/Source/Runtime/Engine/Private/SoundWave.cpp b/Engine/Source/Runtime/Engine/Private/SoundWave.cpp
--- a/Engine/Source/Runtime/Engine/Private/SoundWave.cpp
+++ b/Engine/Source/Runtime/Engine/Private/SoundWave.cpp
@@ -4433,10 +4433,27 @@ bool USoundWave::GetChunkData(int32 ChunkIndex, uint8** OutChunkData, bool bMake
 	return true;
 }
 
-void USoundWave::TriggerWaveDataRefresh()
+void USoundWave::OverrideWaveData(int32 InNumChannels, int32 InNumFrames, int32 InSampleRate)
 {
-	SoundWaveDataPtr.Fork();
-	Proxy = CreateSoundWaveProxy();
+	check(InNumChannels > 0);
+	check(InSampleRate > 0);
+
+	if (!SoundWaveDataPtr.IsValid())
+	{
+		return;
+	}
+
+	if (InNumChannels != SoundWaveDataPtr->NumChannels || InNumFrames != SoundWaveDataPtr->NumFrames || InSampleRate != SoundWaveDataPtr->SampleRate)
+	{
+		// Replace our SoundWaveData with one with the bare minimum information needed to use it for playback.
+		SoundWaveDataPtr.Fork();
+		SoundWaveDataPtr->NumChannels = InNumChannels;
+		SoundWaveDataPtr->NumFrames = InNumFrames;
+		SoundWaveDataPtr->SampleRate = InSampleRate;
+		SoundWaveDataPtr->Duration = InNumFrames / float(InSampleRate);
+		SoundWaveDataPtr->NameCached = GetFName();
+		Proxy = CreateSoundWaveProxy();
+	}
 }
 
 uint32 USoundWave::GetInterpolatedCookedFFTDataForTimeInternal(float InTime, uint32 StartingIndex, TArray<FSoundWaveSpectralData>& OutData, bool bLoop)
