Do not auto compute texture group quality if no quality level is set for that group

diff --git a/Engine/Source/Runtime/Engine/Private/TextureLODSettings.cpp b/Engine/Source/Runtime/Engine/Private/TextureLODSettings.cpp
--- a/Engine/Source/Runtime/Engine/Private/TextureLODSettings.cpp
+++ b/Engine/Source/Runtime/Engine/Private/TextureLODSettings.cpp
@@ -106,7 +106,10 @@ void FTextureLODGroup::SetupGroup()
 		{
 			int32 MaxSizes[4];
 			ComputeMaxSizesForQualityLevels(MaxSizes);
-			MaxLODSize = MaxSizes[TextureGroupMaxQualityLevel];
+			if (MaxSizes[TextureGroupMaxQualityLevel] != -1)
+			{
+				MaxLODSize = MaxSizes[TextureGroupMaxQualityLevel];
+			}
 		}
 	}
 
@@ -155,6 +158,16 @@ bool FTextureLODGroup::operator==(const FTextureLODGroup& Other) const
 
 void FTextureLODGroup::ComputeMaxSizesForQualityLevels(int32 OutMaxSizes[4]) const
 {
+	// We don't want to implicitly generate quality if no quality was set for this group
+	if (MaxLODSize_Low == -1 &&
+		MaxLODSize_Medium == -1 &&
+		MaxLODSize_High == -1 &&
+		MaxLODSize_Epic == -1)
+	{
+		OutMaxSizes[0] = OutMaxSizes[1] = OutMaxSizes[2] = OutMaxSizes[3] = -1;
+		return;
+	}
+
 	// Copy or compute fallbacks for sizes
 	OutMaxSizes[0] = MaxLODSize_Low == -1 ? OptionalMaxLODSize : MaxLODSize_Low;
 	// In case none of the values are set
@@ -297,7 +310,10 @@ int32 UTextureLODSettings::CalculateLODBias(int32 Width, int32 Height, int32 Max
 
 			int32 MaxSizes[4];
 			LODGroupInfo.ComputeMaxSizesForQualityLevels(MaxSizes);
-			MaxLODSize = MaxSizes[QualityLevel];
+			if (MaxSizes[QualityLevel] != -1)
+			{
+				MaxLODSize = MaxSizes[QualityLevel];
+			}
 		}
 	}
 
@@ -331,16 +347,15 @@ int32 UTextureLODSettings::CalculateNumOptionalMips(int32 LODGroup, const int32
 		return 0;
 	}
 
-	int32 OptionalMaxLODSize;
+	int32 OptionalMaxLODSize = LODGroupInfo.OptionalMaxLODSize;
 	if (TextureGroupOptionalQualityLevel >= 0 && TextureGroupOptionalQualityLevel < 4)
 	{
 		int32 MaxSizes[4];
 		LODGroupInfo.ComputeMaxSizesForQualityLevels(MaxSizes);
-		OptionalMaxLODSize = MaxSizes[TextureGroupOptionalQualityLevel];
-	}
-	else
-	{
-		OptionalMaxLODSize = LODGroupInfo.OptionalMaxLODSize;
+		if (MaxSizes[TextureGroupOptionalQualityLevel] != -1)
+		{
+			OptionalMaxLODSize = MaxSizes[TextureGroupOptionalQualityLevel];
+		}
 	}
 
 	int32 OptionalLOD = FMath::Min<int32>(FMath::CeilLogTwo(OptionalMaxLODSize) + 1, NumMips);
