Move assignment before we check it when debug checks are enabled

diff --git a/Engine/Source/Runtime/Engine/Public/Animation/AnimCurveUtils.h b/Engine/Source/Runtime/Engine/Public/Animation/AnimCurveUtils.h
--- a/Engine/Source/Runtime/Engine/Public/Animation/AnimCurveUtils.h
+++ b/Engine/Source/Runtime/Engine/Public/Animation/AnimCurveUtils.h
@@ -172,15 +172,14 @@ public:
 		if(InFilter != nullptr && !InFilter->IsEmpty())
 		{
 			BuildSortedFiltered(OutCurve, InNumElements, InNamePredicate, InValuePredicate, *InFilter);
-			OutCurve.CheckSorted();
 		}
 		else
 		{
 			BuildLinearUnfiltered(OutCurve, InNumElements, InNamePredicate, InValuePredicate);
-			OutCurve.CheckSorted();
 		}
 
 		OutCurve.bSorted = true;
+		OutCurve.CheckSorted();
 	}
 
 	// Helper function for building curves
@@ -310,8 +309,8 @@ public:
 			};
 			
 			BuildSortedFiltered(OutCurve, InNumElements, GetSortedName, GetSortedValue, *InFilter);
-			OutCurve.CheckSorted();
 			OutCurve.bSorted = true;
+			OutCurve.CheckSorted();
 		}
 		else
 		{
