[World Partition]

Added a new hook to UWorldPartitionRuntimeCellTransformers which allow them to modify the UWorldPartitionRuntimeCells as they are being generated during streaming generation.
This hook can be used for example to strip data layers from the runtime cells after they are generated.

diff --git a/Engine/Source/Runtime/Engine/Private/WorldPartition/RuntimeHashSet/WorldPartitionRuntimeHashSetStreamingGeneration.cpp b/Engine/Source/Runtime/Engine/Private/WorldPartition/RuntimeHashSet/WorldPartitionRuntimeHashSetStreamingGeneration.cpp
--- a/Engine/Source/Runtime/Engine/Private/WorldPartition/RuntimeHashSet/WorldPartitionRuntimeHashSetStreamingGeneration.cpp
+++ b/Engine/Source/Runtime/Engine/Private/WorldPartition/RuntimeHashSet/WorldPartitionRuntimeHashSetStreamingGeneration.cpp
@@ -208,6 +208,8 @@ bool UWorldPartitionRuntimeHashSet::GenerateStreaming(UWorldPartitionStreamingPo
 
 				PopulateRuntimeCell(RuntimeCell, CellActorInstances, OutPackagesToGenerate);
 
+				WorldPartition->ApplyRuntimeCellsTransformerStack(RuntimeCell);
+
 				// Save HLOD actor GUID -> Cell GUID mapping to use it for Standalone HLOD source cell override
 				for (const IStreamingGenerationContext::FActorInstance& ActorInstance : CellActorInstances)
 				{
@@ -355,4 +357,4 @@ void UWorldPartitionRuntimeHashSet::DumpStateLog(FHierarchicalLogArchive& Ar) co
 
 	Ar.Printf(TEXT(""));
 }
-#endif
\ No newline at end of file
+#endif
diff --git a/Engine/Source/Runtime/Engine/Private/WorldPartition/WorldPartition.cpp b/Engine/Source/Runtime/Engine/Private/WorldPartition/WorldPartition.cpp
--- a/Engine/Source/Runtime/Engine/Private/WorldPartition/WorldPartition.cpp
+++ b/Engine/Source/Runtime/Engine/Private/WorldPartition/WorldPartition.cpp
@@ -2493,11 +2493,35 @@ FBox UWorldPartition::GetRuntimeWorldBounds() const
 	return FBox();
 }
 
+void UWorldPartition::ApplyRuntimeCellsTransformerStack(UWorldPartitionRuntimeCell* InCell)
+{
+	if (!RuntimeCellsTransformerStack.IsEmpty())
+	{
+		TRACE_CPUPROFILER_EVENT_SCOPE(ApplyRuntimeCellsTransformerStack_RuntimeCell);
+
+		for (const FRuntimeCellTransformerInstance& TransformerInstance : RuntimeCellsTransformerStack)
+		{
+			if (TransformerInstance.Class)
+			{
+				TRACE_CPUPROFILER_EVENT_SCOPE_TEXT(*TransformerInstance.Class->GetName());
+
+				TPair<double, int32>& TotalStats = RuntimeCellsTransformerStackTimes.FindOrAdd(TransformerInstance.Class);
+				TotalStats.Key -= FPlatformTime::Seconds();
+
+				TransformerInstance.TransformRuntimeCell(InCell);
+
+				TotalStats.Key += FPlatformTime::Seconds();
+				TotalStats.Value++;
+			}
+		}
+	}
+}
+
 void UWorldPartition::ApplyRuntimeCellsTransformerStack(ULevel* InLevel)
 {
 	if (!RuntimeCellsTransformerStack.IsEmpty())
 	{
-		TRACE_CPUPROFILER_EVENT_SCOPE(ApplyRuntimeCellsTransformerStack);
+		TRACE_CPUPROFILER_EVENT_SCOPE(ApplyRuntimeCellsTransformerStack_Level);
 
 		auto ApplyTransformPhase = [this, InLevel](TFunction<void(const FRuntimeCellTransformerInstance& TransformerInstance, ULevel* Level)> Func)
 		{
diff --git a/Engine/Source/Runtime/Engine/Private/WorldPartition/WorldPartitionRuntimeCellTransformer.cpp b/Engine/Source/Runtime/Engine/Private/WorldPartition/WorldPartitionRuntimeCellTransformer.cpp
--- a/Engine/Source/Runtime/Engine/Private/WorldPartition/WorldPartitionRuntimeCellTransformer.cpp
+++ b/Engine/Source/Runtime/Engine/Private/WorldPartition/WorldPartitionRuntimeCellTransformer.cpp
@@ -39,6 +39,15 @@ void UWorldPartitionRuntimeCellTransformer::PostTransform(ULevel* InLevel)
 {
 }
 
+bool UWorldPartitionRuntimeCellTransformer::IsCellTransformable(UWorldPartitionRuntimeCell* InCell) const
+{
+	return true;
+}
+
+void UWorldPartitionRuntimeCellTransformer::TransformRuntimeCell(UWorldPartitionRuntimeCell* InCell)
+{
+}
+
 void UWorldPartitionRuntimeCellTransformer::ForEachIgnoredComponentClass(TFunctionRef<bool(const TSubclassOf<UActorComponent>&)> Func) const
 {
 	for (const TSubclassOf<UActorComponent>& IgnoredComponentClass : GetDefault<UWorldPartitionRuntimeCellTransformerSettings>()->IgnoredComponentClasses)
diff --git a/Engine/Source/Runtime/Engine/Public/WorldPartition/WorldPartition.h b/Engine/Source/Runtime/Engine/Public/WorldPartition/WorldPartition.h
--- a/Engine/Source/Runtime/Engine/Public/WorldPartition/WorldPartition.h
+++ b/Engine/Source/Runtime/Engine/Public/WorldPartition/WorldPartition.h
@@ -127,6 +127,8 @@ struct FRuntimeCellTransformerInstance
 	inline void Transform(ULevel* InLevel) const { if (CanTransform(InLevel)) { Instance->Transform(InLevel); } }
 	inline void PostTransform(ULevel* InLevel) const { if (CanTransform(InLevel)) { Instance->PostTransform(InLevel); } }
 
+	inline void TransformRuntimeCell(UWorldPartitionRuntimeCell* InCell) const { if (CanTransform(InCell)) { Instance->TransformRuntimeCell(InCell); } }
+
 	/** Runtime cell transformer class */
 	UPROPERTY(EditAnywhere, Category = WorldPartitionSetup, AdvancedDisplay)
 	TSubclassOf<UWorldPartitionRuntimeCellTransformer> Class;
@@ -137,6 +139,7 @@ struct FRuntimeCellTransformerInstance
 
 private:
 	inline bool CanTransform(ULevel* InLevel) const { return Instance && Instance->IsEnabled() && Instance->IsLevelTransformable(InLevel); }
+	inline bool CanTransform(UWorldPartitionRuntimeCell* InCell) const { return Instance && Instance->IsEnabled() && Instance->IsCellTransformable(InCell); }
 #endif
 };
 
@@ -350,6 +353,7 @@ public:
 	ENGINE_API FBox GetRuntimeWorldBounds() const;
 
 	inline TConstArrayView<FRuntimeCellTransformerInstance> GetRuntimeCellsTransformerStack() const { return RuntimeCellsTransformerStack; }
+	ENGINE_API void ApplyRuntimeCellsTransformerStack(UWorldPartitionRuntimeCell* InCell);
 	ENGINE_API void ApplyRuntimeCellsTransformerStack(ULevel* InLevel);
 	ENGINE_API bool HasRuntimeCellTransformerOfType(const UClass* InCellTransformerClass);
 	ENGINE_API void AddRuntimeCellTransformerOfType(TSubclassOf<UWorldPartitionRuntimeCellTransformer> InCellTransformerClass, int32 PositionToInsert = -1);
diff --git a/Engine/Source/Runtime/Engine/Public/WorldPartition/WorldPartitionRuntimeCellTransformer.h b/Engine/Source/Runtime/Engine/Public/WorldPartition/WorldPartitionRuntimeCellTransformer.h
--- a/Engine/Source/Runtime/Engine/Public/WorldPartition/WorldPartitionRuntimeCellTransformer.h
+++ b/Engine/Source/Runtime/Engine/Public/WorldPartition/WorldPartitionRuntimeCellTransformer.h
@@ -5,6 +5,8 @@
 #include "Components/ActorComponent.h"
 #include "WorldPartitionRuntimeCellTransformer.generated.h"
 
+class UWorldPartitionRuntimeCell;
+
 #define UE_API ENGINE_API
 
 UCLASS(MinimalAPI, Config=Engine)
@@ -38,6 +40,15 @@ public:
 	UE_API virtual void PreTransform(ULevel* InLevel);
 	UE_API virtual void Transform(ULevel* InLevel);
 	UE_API virtual void PostTransform(ULevel* InLevel);
+
+	/** Whether the cell can be process by the transformer. */
+	UE_API virtual bool IsCellTransformable(UWorldPartitionRuntimeCell* InCell) const;
+	/**
+	 * Hook to transform the runtime cell structure itself.
+	 * Note: Modification of the cell contents must be done in the ULevel* hook.
+	 */
+	UE_API virtual void TransformRuntimeCell(UWorldPartitionRuntimeCell* InCell);
+
 	UE_API virtual void ForEachIgnoredComponentClass(TFunctionRef<bool(const TSubclassOf<UActorComponent>&)> Func) const;
 	UE_API virtual void ForEachIgnoredExactComponentClass(TFunctionRef<bool(const TSubclassOf<UActorComponent>&)> Func) const;
 #endif
