[WorldPartition] Streaming generation log now also logs containers with External Data Layers.

diff --git a/Engine/Source/Runtime/Engine/Private/WorldPartition/WorldPartitionStreamingGeneration.cpp b/Engine/Source/Runtime/Engine/Private/WorldPartition/WorldPartitionStreamingGeneration.cpp
--- a/Engine/Source/Runtime/Engine/Private/WorldPartition/WorldPartitionStreamingGeneration.cpp
+++ b/Engine/Source/Runtime/Engine/Private/WorldPartition/WorldPartitionStreamingGeneration.cpp
@@ -2108,29 +2108,30 @@ public:
 		UE_SCOPED_INDENT_LOG_ARCHIVE(Ar.PrintfIndent(TEXT("Containers:")));
 		for (auto& [ContainerID, ContainerCollectionInstanceDescriptor] : ContainerCollectionInstanceDescriptorsMap)
 		{
-			FString ContainerPackageName = ContainerCollectionInstanceDescriptor.ContainerInstanceCollection->GetBaseContainerInstancePackageName().ToString();
-			bool bIsAlreadySet = false;
-			UniqueContainerNames.Add(ContainerPackageName, &bIsAlreadySet);
-			if (bIsAlreadySet)
+			ContainerCollectionInstanceDescriptor.ContainerInstanceCollection->ForEachActorDescContainerInstance([&UniqueContainerNames, &Ar, Mode](const UActorDescContainerInstance* InActorDescContainerInstance)
 			{
-				continue;
-			}
-
-			UE_SCOPED_INDENT_LOG_ARCHIVE(Ar.PrintfIndent(TEXT("Container: %s"), *ContainerPackageName));
+				const FString ContainerPackageName = InActorDescContainerInstance->GetContainerPackage().ToString();
+				bool bIsAlreadySet = false;
+				UniqueContainerNames.Add(ContainerPackageName, &bIsAlreadySet);
+				if (!bIsAlreadySet)
+				{
+					UE_SCOPED_INDENT_LOG_ARCHIVE(Ar.PrintfIndent(TEXT("Container: %s"), *ContainerPackageName));
 
-			const UActorDescContainerInstance* Container = ContainerCollectionInstanceDescriptor.ContainerInstanceCollection->GetBaseContainerInstance();
-			if (Container->GetContainer()->GetActorDescCount())
-			{
-				UE_SCOPED_INDENT_LOG_ARCHIVE(Ar.PrintfIndent(TEXT("ActorDescs:")));
+					const UActorDescContainer* Container = InActorDescContainerInstance->GetContainer();
+					if (Container->GetActorDescCount())
+					{
+						UE_SCOPED_INDENT_LOG_ARCHIVE(Ar.PrintfIndent(TEXT("ActorDescs:")));
 
-				FActorDescList::FGuidActorDescMap SortedActorsByGuid = Container->GetContainer()->GetActorsByGuid();
-				SortedActorsByGuid.KeySort([](const FGuid& GuidA, const FGuid& GuidB) { return GuidA < GuidB; });
+						FActorDescList::FGuidActorDescMap SortedActorsByGuid = Container->GetActorsByGuid();
+						SortedActorsByGuid.KeySort([](const FGuid& GuidA, const FGuid& GuidB) { return GuidA < GuidB; });
 
-				for (auto& [ActorGuid, ActorDesc] : SortedActorsByGuid)
-				{
-					Ar.Print(*((*ActorDesc)->ToString(Mode)));
+						for (auto& [ActorGuid, ActorDesc] : SortedActorsByGuid)
+						{
+							Ar.Print(*((*ActorDesc)->ToString(Mode)));
+						}
+					}
 				}
-			}
+			});
 		}
 
 		TSet<uint64> UniqueClusters;
