FNewSpline: Fix assumption that FSplineCurves interp curves have identical element counts

diff --git a/Engine/Source/Runtime/Engine/Private/Curves/Spline.cpp b/Engine/Source/Runtime/Engine/Private/Curves/Spline.cpp
--- a/Engine/Source/Runtime/Engine/Private/Curves/Spline.cpp
+++ b/Engine/Source/Runtime/Engine/Private/Curves/Spline.cpp
@@ -16,6 +16,8 @@ PRAGMA_DISABLE_EXPERIMENTAL_WARNINGS
 #include "Splines/MultiSpline.h"
 PRAGMA_ENABLE_EXPERIMENTAL_WARNINGS
 
+DEFINE_LOG_CATEGORY_STATIC(LogSpline, Log, All);
+
 #include UE_INLINE_GENERATED_CPP_BY_NAME(Spline)
 
 PRAGMA_DISABLE_EXPERIMENTAL_WARNINGS
@@ -1855,14 +1857,31 @@ FNewSpline::FNewSpline(const FSplineCurves& Other)
 	ResetScale();
 	
     // setup rotation/scale attributes
-    for (int32 Idx = 0; Idx < GetSpline().GetNumPoints(); Idx++)
-    {
-		const FQuat& Rotation = Other.Rotation.Points[Idx].OutVal;
-		const FVector& Scale = Other.Scale.Points[Idx].OutVal;
+	for (int32 Idx = 0; Idx < Other.Position.Points.Num(); Idx++)
+	{
+		FQuat Rotation = FQuat::Identity;
+		if (Other.Rotation.Points.IsValidIndex(Idx))
+		{
+			Rotation = Other.Rotation.Points[Idx].OutVal;
+		}
+		else
+		{
+			UE_LOG(LogSpline, Warning, TEXT("Number of rotation values (%u) different than the number of position values (%u)!"), Other.Rotation.Points.Num(), Other.Position.Points.Num());
+		}
+    	
+		FVector Scale = FVector::OneVector;
+		if (Other.Scale.Points.IsValidIndex(Idx))
+		{
+			Scale = Other.Scale.Points[Idx].OutVal;
+		}
+		else
+		{
+			UE_LOG(LogSpline, Warning, TEXT("Number of scale values (%u) different than the number of position values (%u)!"), Other.Scale.Points.Num(), Other.Position.Points.Num());
+		}
     	
 		SetAttributeValue<FQuat>(RotationAttrName, Rotation, Idx);
 		SetAttributeValue<FVector>(ScaleAttrName, Scale, Idx);
-    }
+	}
 
 	ValidateRotScale();
 	
