Add filter so that the EDL does not cook more than once but does generate EDLs for the DLC cooks.
Add commandline to DLC cooks to restrict the generators appropriately.

Add new commandline option for describing which generators should not be saved:
-EDLFilterGenerator=Value1+Value2+Value3

diff --git a/Engine/Source/Runtime/Engine/Private/WorldCookPackageSplitter.cpp b/Engine/Source/Runtime/Engine/Private/WorldCookPackageSplitter.cpp
--- a/Engine/Source/Runtime/Engine/Private/WorldCookPackageSplitter.cpp
+++ b/Engine/Source/Runtime/Engine/Private/WorldCookPackageSplitter.cpp
@@ -140,7 +140,7 @@ ICookPackageSplitter::FGenerationManifest FWorldCookPackageSplitter::ReportGener
 			}
 
 			// Any sub-splitter can disable save of the map
-			TOptional<bool> SkipSave = WorldCookPackageSplitter->SkipSaveGeneratorPackage();
+			TOptional<bool> SkipSave = WorldCookPackageSplitter->SkipSaveGeneratorPackage(OwnerPackage);
 			if (SkipSave.IsSet())
 			{
 				if (SkipSaveGeneratorPackage.IsSet() && SkipSave.GetValue() != SkipSaveGeneratorPackage.GetValue())
diff --git a/Engine/Source/Runtime/Engine/Private/WorldPartition/WorldPartitionCookPackageSplitter.cpp b/Engine/Source/Runtime/Engine/Private/WorldPartition/WorldPartitionCookPackageSplitter.cpp
--- a/Engine/Source/Runtime/Engine/Private/WorldPartition/WorldPartitionCookPackageSplitter.cpp
+++ b/Engine/Source/Runtime/Engine/Private/WorldPartition/WorldPartitionCookPackageSplitter.cpp
@@ -104,10 +104,10 @@ bool FWorldPartitionCookPackageSplitter::PopulateGeneratorPackage(ICookPackageSp
 	return bIsSuccess;
 }
 
-TOptional<bool> FWorldPartitionCookPackageSplitter::SkipSaveGeneratorPackage()
+TOptional<bool> FWorldPartitionCookPackageSplitter::SkipSaveGeneratorPackage(const UPackage* OwnerPackage)
 {
 	// If we have EDL cook filters, we are trying to cook the EDLs in isolation. Do not save the map even though it was loaded.
-	if (UWorldPartitionSettings::GetEDLCookFilters().Num() > 0)
+	if (!UWorldPartitionSettings::ShouldCookEDLGenerator(OwnerPackage))
 	{
 		return true;
 	}
diff --git a/Engine/Source/Runtime/Engine/Private/WorldPartition/WorldPartitionCookPackageSplitter.h b/Engine/Source/Runtime/Engine/Private/WorldPartition/WorldPartitionCookPackageSplitter.h
--- a/Engine/Source/Runtime/Engine/Private/WorldPartition/WorldPartitionCookPackageSplitter.h
+++ b/Engine/Source/Runtime/Engine/Private/WorldPartition/WorldPartitionCookPackageSplitter.h
@@ -15,7 +15,7 @@ public:
 	virtual TOptional<bool> RequiresEmptyPackageBeforePopulate(const UPackage* Package) override;
 	virtual bool PopulateGeneratedPackage(ICookPackageSplitter::FPopulateContext& PopulateContext) override;
 	virtual bool PopulateGeneratorPackage(ICookPackageSplitter::FPopulateContext& PopulateContext) override;
-	virtual TOptional<bool> SkipSaveGeneratorPackage() override;
+	virtual TOptional<bool> SkipSaveGeneratorPackage(const UPackage* OwnerPackage) override;
 	virtual void Teardown(ICookPackageSplitter::ETeardown Status) override;
 	//~End FWorldCookPackageSplitter::ISubSplitter interface
 
diff --git a/Engine/Source/Runtime/Engine/Private/WorldPartition/WorldPartitionSettings.cpp b/Engine/Source/Runtime/Engine/Private/WorldPartition/WorldPartitionSettings.cpp
--- a/Engine/Source/Runtime/Engine/Private/WorldPartition/WorldPartitionSettings.cpp
+++ b/Engine/Source/Runtime/Engine/Private/WorldPartition/WorldPartitionSettings.cpp
@@ -12,48 +12,69 @@ UWorldPartitionSettings::UWorldPartitionSettings(const FObjectInitializer& Objec
 {}
 
 #if WITH_EDITOR
-
-const TSet<FTopLevelAssetPath>& UWorldPartitionSettings::GetEDLCookFilters()
+UWorldPartitionSettings::EDLCookFilterCommandLineLoader::EDLCookFilterCommandLineLoader()
 {
-	// Caches the current EDL cooking filters from the commandline and config
-	struct EDLCookFilterCommandLineLoader
+	// First parse EDL filters passed through commandline like -EDLCookFilters=Value1+Value2+Value3
+	FString SwitchValue;
+	if (FParse::Value(FCommandLine::Get(), TEXT("-EDLCookFilters="), SwitchValue))
 	{
-		EDLCookFilterCommandLineLoader()
+		TArray<FString> Tokens;
+		SwitchValue.ParseIntoArray(Tokens, TEXT("+"));
+		for (const FString& Filter : Tokens)
 		{
-			// First parse EDL filters passed through commandline like -EDLCookFilters=Value1+Value2+Value3
-			FString SwitchValue;
-			if (FParse::Value(FCommandLine::Get(), TEXT("-EDLCookFilters="), SwitchValue))
+			FTopLevelAssetPath Path;
+			if (Path.TrySetPath(Filter))
 			{
-				TArray<FString> Tokens;
-				SwitchValue.ParseIntoArray(Tokens, TEXT("+"));
-				for (const FString& Filter : Tokens)
-				{
-					FTopLevelAssetPath Path;
-					if (Path.TrySetPath(Filter))
-					{
-						Filters.Add(Path);
-					}
-				}
+				EDLFilters.Add(Path);
 			}
+		}
+	}
 
-			// Then append filters found in config, e.g. in [Base/Default/etc.]Engine.ini
-			// [/Script/Engine.WorldPartitionSettings]
-			// +EDLCookFilters=/MyGFP/DataLayer/MyEDL.MyEDL
-			for (const TSoftObjectPtr<UExternalDataLayerAsset>& EDL : GetDefault<UWorldPartitionSettings>()->EDLCookFilters)
-			{
-				FSoftObjectPath Path = EDL.ToSoftObjectPath();
-				if (Path.IsValid())
-				{
-					Filters.Add(Path.GetAssetPath());
-				}
-			}
+	// Second parse generators filters passed through commandline like -EDLFilterGenerator=Value1+Value2+Value3
+	// These restrict if the generator should be cooked more than once. When specified the generator will not be saved again.
+	FString EDLFilterGeneratorString;
+	if (FParse::Value(FCommandLine::Get(), TEXT("-EDLFilterGenerator="), EDLFilterGeneratorString))
+	{
+		TArray<FString> Tokens;
+		EDLFilterGeneratorString.ParseIntoArray(Tokens, TEXT("+"));
+		for (const FString& Filter : Tokens)
+		{
+			EDLGeneratorFilters.Add(FName(Filter));
 		}
+	}
 
-		TSet<FTopLevelAssetPath> Filters;
-	};
+	// Then append filters found in config, e.g. in [Base/Default/etc.]Engine.ini
+	// [/Script/Engine.WorldPartitionSettings]
+	// +EDLCookFilters=/MyGFP/DataLayer/MyEDL.MyEDL
+	for (const TSoftObjectPtr<UExternalDataLayerAsset>& EDL : GetDefault<UWorldPartitionSettings>()->EDLCookFilters)
+	{
+		FSoftObjectPath Path = EDL.ToSoftObjectPath();
+		if (Path.IsValid())
+		{
+			EDLFilters.Add(Path.GetAssetPath());
+		}
+	}
+}
+
+const UWorldPartitionSettings::EDLCookFilterCommandLineLoader& UWorldPartitionSettings::GetEDLCookFilterLoader()
+{
+	static UWorldPartitionSettings::EDLCookFilterCommandLineLoader Loader;
+	return Loader;
+}
+
+const TSet<FTopLevelAssetPath>& UWorldPartitionSettings::GetEDLCookFilters()
+{
+	return GetEDLCookFilterLoader().EDLFilters;
+}
+
+const bool UWorldPartitionSettings::ShouldCookEDLGenerator(const UPackage* Generator)
+{
+	if (!UWorldPartitionSettings::GetEDLCookFilters().IsEmpty())
+	{
+		return false;
+	}
 
-	static EDLCookFilterCommandLineLoader Loader;
-	return Loader.Filters;
+	return !GetEDLCookFilterLoader().EDLGeneratorFilters.Contains(Generator->GetFName());
 }
 
 #endif // #if WITH_EDITOR
diff --git a/Engine/Source/Runtime/Engine/Public/WorldCookPackageSplitter.h b/Engine/Source/Runtime/Engine/Public/WorldCookPackageSplitter.h
--- a/Engine/Source/Runtime/Engine/Public/WorldCookPackageSplitter.h
+++ b/Engine/Source/Runtime/Engine/Public/WorldCookPackageSplitter.h
@@ -30,7 +30,7 @@ public:
 		 * This can return true to disable the save of the generator package, for fine-tuned output control. Return unset by default.
 		 * All splitters should return the same value, or leave unset (as only one sub-splitter should be responsible for each package), otherwise the resulting mismatch will log an error.
 		 */
-		virtual TOptional<bool> SkipSaveGeneratorPackage() = 0;
+		virtual TOptional<bool> SkipSaveGeneratorPackage(const UPackage* OwnerPackage) = 0;
 		virtual void Teardown(ICookPackageSplitter::ETeardown Status) = 0;
 	};
 
diff --git a/Engine/Source/Runtime/Engine/Public/WorldPartition/WorldPartitionSettings.h b/Engine/Source/Runtime/Engine/Public/WorldPartition/WorldPartitionSettings.h
--- a/Engine/Source/Runtime/Engine/Public/WorldPartition/WorldPartitionSettings.h
+++ b/Engine/Source/Runtime/Engine/Public/WorldPartition/WorldPartitionSettings.h
@@ -56,6 +56,7 @@ public:
 	 * If there are paths in the list however, only these EDLs are cooked, excluding other EDLs not in the list, and most importantly excluding the base map
 	 */
 	static ENGINE_API const TSet<FTopLevelAssetPath>& GetEDLCookFilters();
+	static ENGINE_API const bool ShouldCookEDLGenerator(const UPackage* Generator);
 #endif
 
 protected:
@@ -98,6 +99,17 @@ protected:
 	UPROPERTY(EditAnywhere, Config, Category = WorldPartition, meta=(ToolTip="EDLs to cook in exclusion from base maps and other EDLs"))
 	TSet<TSoftObjectPtr<UExternalDataLayerAsset>> EDLCookFilters;
 #endif
+
+private:
+	struct EDLCookFilterCommandLineLoader
+	{
+		EDLCookFilterCommandLineLoader();
+
+		TSet<FTopLevelAssetPath> EDLFilters;
+		TSet<FName> EDLGeneratorFilters;
+	};
+
+	static const EDLCookFilterCommandLineLoader& GetEDLCookFilterLoader();
 };
 
 #undef UE_API
