[HWRT] Fix GetDynamicRayTracingInstances(...) outputting invalid mesh batches.
- those batches caused issues during dynamic ray tracing geometry update.

diff --git a/Engine/Source/Runtime/Engine/Private/SplineMeshSceneProxy.cpp b/Engine/Source/Runtime/Engine/Private/SplineMeshSceneProxy.cpp
--- a/Engine/Source/Runtime/Engine/Private/SplineMeshSceneProxy.cpp
+++ b/Engine/Source/Runtime/Engine/Private/SplineMeshSceneProxy.cpp
@@ -305,7 +305,7 @@ void FSplineMeshSceneProxy::GetDynamicRayTracingInstances(FRayTracingInstanceCol
 		{
 			for (int32 SectionIndex = 0; SectionIndex < RayTracingLOD.Sections->Num(); SectionIndex++)
 			{
-				FMeshBatch& MeshBatch = CachedRayTracingMaterials.AddDefaulted_GetRef();
+				FMeshBatch MeshBatch;
 
 				bool bResult;
 				if (RenderData->RayTracingProxy->bUsingRenderingLODs)
@@ -319,15 +319,14 @@ void FSplineMeshSceneProxy::GetDynamicRayTracingInstances(FRayTracingInstanceCol
 					bResult = GetRayTracingMeshElement(LODIndex, BatchIndex, SectionIndex, PrimitiveDPG, MeshBatch);
 				}
 
-				if (!bResult)
+				if (bResult)
 				{
-					// Hidden material
-					MeshBatch.MaterialRenderProxy = UMaterial::GetDefaultMaterial(MD_Surface)->GetRenderProxy();
-					MeshBatch.VertexFactory = &(*RenderData->RayTracingProxy->LODVertexFactories)[LODIndex].VertexFactory;
+					MeshBatch.ReverseCulling = bReverseCulling;
+					MeshBatch.SegmentIndex = SectionIndex;
+					MeshBatch.MeshIdInPrimitive = SectionIndex;
+
+					CachedRayTracingMaterials.Add(MoveTemp(MeshBatch));
 				}
-				MeshBatch.ReverseCulling = bReverseCulling;
-				MeshBatch.SegmentIndex = SectionIndex;
-				MeshBatch.MeshIdInPrimitive = SectionIndex;
 			}
 		}
 
@@ -384,7 +383,7 @@ void FSplineMeshSceneProxy::GetDynamicRayTracingInstances(FRayTracingInstanceCol
 	}
 
 	check(CachedRayTracingMaterials.Num() == RayTracingInstance.GetMaterials().Num());
-	checkf(RayTracingInstance.Geometry->Initializer.Segments.Num() == CachedRayTracingMaterials.Num(), TEXT("Segments/Materials mismatch. Number of segments: %d. Number of Materials: %d. LOD Index: %d"),
+	checkf(RayTracingInstance.Geometry->Initializer.Segments.Num() >= CachedRayTracingMaterials.Num(), TEXT("Segments/Materials mismatch. Number of segments: %d. Number of Materials: %d. LOD Index: %d"),
 		RayTracingInstance.Geometry->Initializer.Segments.Num(),
 		CachedRayTracingMaterials.Num(),
 		LODIndex);
diff --git a/Engine/Source/Runtime/Engine/Private/StaticMeshSceneProxy.cpp b/Engine/Source/Runtime/Engine/Private/StaticMeshSceneProxy.cpp
--- a/Engine/Source/Runtime/Engine/Private/StaticMeshSceneProxy.cpp
+++ b/Engine/Source/Runtime/Engine/Private/StaticMeshSceneProxy.cpp
@@ -2263,7 +2263,7 @@ void FStaticMeshSceneProxy::GetDynamicRayTracingInstances(FRayTracingInstanceCol
 			{
 				for (int32 SectionIndex = 0; SectionIndex < RayTracingLOD.Sections->Num(); SectionIndex++)
 				{
-					FMeshBatch &MeshBatch = CachedRayTracingMaterials.AddDefaulted_GetRef();
+					FMeshBatch MeshBatch;
 
 					bool bResult;
 					if (RenderData->RayTracingProxy->bUsingRenderingLODs)
@@ -2277,15 +2277,14 @@ void FStaticMeshSceneProxy::GetDynamicRayTracingInstances(FRayTracingInstanceCol
 						bResult = GetRayTracingMeshElement(LODIndex, BatchIndex, SectionIndex, PrimitiveDPG, MeshBatch);
 					}
 
-					if (!bResult)
+					if (bResult)
 					{
-						// Hidden material
-						MeshBatch.MaterialRenderProxy = UMaterial::GetDefaultMaterial(MD_Surface)->GetRenderProxy();
-						MeshBatch.VertexFactory = &(*RenderData->RayTracingProxy->LODVertexFactories)[LODIndex].VertexFactory;
+						MeshBatch.ReverseCulling = bReverseCulling; // overwrite what came from GetMeshElement as DXR only needs the user driven flag, not the flipping implied by the transform
+						MeshBatch.SegmentIndex = SectionIndex;
+						MeshBatch.MeshIdInPrimitive = SectionIndex;
+
+						CachedRayTracingMaterials.Add(MoveTemp(MeshBatch));
 					}
-					MeshBatch.ReverseCulling = bReverseCulling; // overwrite what came from GetMeshElement as DXR only needs the user driven flag, not the flipping implied by the transform
-					MeshBatch.SegmentIndex = SectionIndex;
-					MeshBatch.MeshIdInPrimitive = SectionIndex;
 				}
 			}
 			
@@ -2330,7 +2329,7 @@ void FStaticMeshSceneProxy::GetDynamicRayTracingInstances(FRayTracingInstanceCol
 		}
 
 		check(CachedRayTracingMaterials.Num() == RayTracingInstance.GetMaterials().Num());
-		checkf(RayTracingInstance.Geometry->Initializer.Segments.Num() == CachedRayTracingMaterials.Num(), TEXT("Segments/Materials mismatch. Number of segments: %d. Number of Materials: %d. LOD Index: %d"), 
+		checkf(RayTracingInstance.Geometry->Initializer.Segments.Num() >= CachedRayTracingMaterials.Num(), TEXT("Segments/Materials mismatch. Number of segments: %d. Number of Materials: %d. LOD Index: %d"),
 			RayTracingInstance.Geometry->Initializer.Segments.Num(), 
 			CachedRayTracingMaterials.Num(), 
 			LODIndex);
