[HLOD] Fix an issue where the main editor world package could be marked as transient when building HLODs.

diff --git a/Engine/Source/Runtime/Engine/Private/WorldPartition/HLOD/HLODSourceActorsFromCell.cpp b/Engine/Source/Runtime/Engine/Private/WorldPartition/HLOD/HLODSourceActorsFromCell.cpp
--- a/Engine/Source/Runtime/Engine/Private/WorldPartition/HLOD/HLODSourceActorsFromCell.cpp
+++ b/Engine/Source/Runtime/Engine/Private/WorldPartition/HLOD/HLODSourceActorsFromCell.cpp
@@ -131,6 +131,7 @@ bool UWorldPartitionHLODSourceActorsFromCell::LoadSourceActors(bool& bOutDirty,
 
 	if (FWorldPartitionLevelHelper::LoadActors(MoveTemp(Params)))
 	{
+		TSet<UPackage*> LevelPackagesMadeTransient;
 		TArray<UObject*> ActorsToResetLoaders;
 		ActorsToResetLoaders.Reserve(ActorsCopy.Num());
 		for (const FWorldPartitionRuntimeCellObjectMapping& ActorMapping : ActorsCopy)
@@ -145,7 +146,11 @@ bool UWorldPartitionHLODSourceActorsFromCell::LoadSourceActors(bool& bOutDirty,
 				{
 					if (UPackage* ActorLevelPackage = ActorLevel->GetPackage())
 					{
-						ActorLevelPackage->SetFlags(RF_Transient);
+						if (IsValid(ActorLevelPackage) && !ActorLevelPackage->HasAnyFlags(RF_Transient))
+						{
+							ActorLevelPackage->SetFlags(RF_Transient);
+							LevelPackagesMadeTransient.Add(ActorLevelPackage);
+						}
 					}
 				}
 			}
@@ -155,6 +160,11 @@ bool UWorldPartitionHLODSourceActorsFromCell::LoadSourceActors(bool& bOutDirty,
 		TArray<UPackage*> ModifiedPackages;
 		FWorldPartitionLevelHelper::MoveExternalActorsToLevel(ActorsCopy, TargetWorld->PersistentLevel, ModifiedPackages);
 
+		for (UPackage* Package : LevelPackagesMadeTransient)
+		{
+			Package->ClearFlags(RF_Transient);
+		}
+
 		return true;
 	}
 
