LFV - added comments.

#rnx
#tests PC editor

diff --git a/Engine/Source/Runtime/Engine/Classes/Engine/LocalFogVolume.h b/Engine/Source/Runtime/Engine/Classes/Engine/LocalFogVolume.h
--- a/Engine/Source/Runtime/Engine/Classes/Engine/LocalFogVolume.h
+++ b/Engine/Source/Runtime/Engine/Classes/Engine/LocalFogVolume.h
@@ -32,7 +32,7 @@ public:
 	ULocalFogVolumeComponent* GetComponent() const { return LocalFogVolumeVolume; }
 	
 #if WITH_EDITOR
-	//Allow LocalFogVolume instances to be spatially loaded
+	//Allow LocalFogVolume instances to be spatially loaded (this flag enabled the streaming UI properties), and ULocalFogVolumeComponent is a USceneComponent which should provide streaming bounds.
 	virtual bool CanChangeIsSpatiallyLoadedFlag() const override { return true; }
 #endif
 
