Prevent PhysicalAnimationComponent crash if physics blending is enabled mid-tick

diff --git a/Engine/Source/Runtime/Engine/Private/PhysicsEngine/PhysicalAnimationComponent.cpp b/Engine/Source/Runtime/Engine/Private/PhysicsEngine/PhysicalAnimationComponent.cpp
--- a/Engine/Source/Runtime/Engine/Private/PhysicsEngine/PhysicalAnimationComponent.cpp
+++ b/Engine/Source/Runtime/Engine/Private/PhysicsEngine/PhysicalAnimationComponent.cpp
@@ -267,6 +267,23 @@ void UPhysicalAnimationComponent::UpdateTargetActors(ETeleportType TeleportType)
 		// various anim ticks, before buffers are flipped (which happens in the skel mesh component's post-physics tick)
 		const TArray<FTransform>& SpaceBases = SkeletalMeshComponent->GetEditableComponentSpaceTransforms();
 
+		if (SpaceBases.IsEmpty())
+		{
+			// Our tick depends on the skeletal mesh primary tick. When physics is disabled (i.e.
+			// physics blend = 0), this might have been allowed to run post-physics. Thus, if
+			// physics blending is enabled after the ticks have started processing (this can happen
+			// from the actor's Blueprint code) then for one tick we might find ourselves running in
+			// a post-physics tick group.This can result in this function running at the same time
+			// as the secondary EndPhysics SKM tick function that performs a parallel blend back
+			// from physics into the SKM component transforms. When this happens, the component
+			// transforms will have been moved temporarily into the animation context. Preventing
+			// this from happening is difficult (as it would need the tick functions to be updated
+			// whilst they are getting processed) so we simply return if the component transforms
+			// are missing. Note that setting CVarAnimationDelaysEndGroup
+			// ("tick.AnimationDelaysEndGroup") to 0 should prevent us from getting here, as it
+			// stops the primary animation tick from being allowed to run after PrePhysics.
+			return;
+		}
 		FPhysicsCommand::ExecuteWrite(SkeletalMeshComponent, [&]()
 		{
 			TArray<FTransform> LocalTransforms = SkeletalMeshComponent->GetBoneSpaceTransforms();
