Add 'r.psoprecache.DefaultMaterialDelayStratCanIgnorePSOPrioBoost' which prevents proxy creation boosting PSO precache priority, leaving it to r.PSOPrecache.DrawnComponentBoostStrategy to perform the priority boost when required.
Fix possible leak of UMaterialInterface in FMaterialPSOPrecacheCollectionTask
Fix bug in MarkCompilationComplete, in which MarkCompilationComplete could incorrectly process a retired request.
Add CSV profiler stat to count the number of primitive proxies that have drawn with default materials and are waiting for their boosted PSOs to complete.

diff --git a/Engine/Source/Runtime/Engine/Private/Components/DecalComponent.cpp b/Engine/Source/Runtime/Engine/Private/Components/DecalComponent.cpp
--- a/Engine/Source/Runtime/Engine/Private/Components/DecalComponent.cpp
+++ b/Engine/Source/Runtime/Engine/Private/Components/DecalComponent.cpp
@@ -493,7 +493,7 @@ FDeferredDecalProxy* UDecalComponent::CreateSceneProxy()
 	LLM_SCOPE(ELLMTag::SceneRender);
 
 #if UE_WITH_PSO_PRECACHING
-	if (LatestPSOPrecacheJobSetCompleted != LatestPSOPrecacheJobSet && GetPSOPrecacheProxyCreationStrategy() == EPSOPrecacheProxyCreationStrategy::DelayUntilPSOPrecached)
+	if (LatestPSOPrecacheJobSetCompleted != LatestPSOPrecacheJobSet && GetPSOPrecacheProxyCreationStrategy() != EPSOPrecacheProxyCreationStrategy::AlwaysCreate)
 	{
 		return nullptr;
 	}
diff --git a/Engine/Source/Runtime/Engine/Private/Components/PrimitiveComponent.cpp b/Engine/Source/Runtime/Engine/Private/Components/PrimitiveComponent.cpp
--- a/Engine/Source/Runtime/Engine/Private/Components/PrimitiveComponent.cpp
+++ b/Engine/Source/Runtime/Engine/Private/Components/PrimitiveComponent.cpp
@@ -5253,6 +5253,16 @@ struct FPSOPrecacheFinishedTask
 	TWeakObjectPtr<UPrimitiveComponent> WeakPrimitiveComponent;
 	int32 JobSetThatJustCompleted;
 };
+
+namespace PrimitiveComponentCVars
+{
+	static bool bDefaultMaterialDelayStratCanIgnorePSOPrioBoost = false;
+	static FAutoConsoleVariableRef CVarDefaultMaterialDelayStratCanIgnorePSOPrioBoost(
+		TEXT("r.psoprecache.DefaultMaterialDelayStratCanIgnorePSOPrioBoost"),
+		bDefaultMaterialDelayStratCanIgnorePSOPrioBoost,
+		TEXT("false: PSO precache requests are boosted during proxy create. (default) \n")
+		TEXT("true: When r.PSOPrecache.ProxyCreationDelayStrategy and r.PSOPrecache.DrawnComponentBoostStrategy are in use, priority boost during proxy create will be ignored."));
+}
 #endif
 
 void UPrimitiveComponent::RequestRecreateRenderStateWhenPSOPrecacheFinished(const FGraphEventArray& PSOPrecacheCompileEvents)
@@ -5311,7 +5321,9 @@ bool UPrimitiveComponent::CheckPSOPrecachingAndBoostPriority(EPSOPrecachePriorit
 	ensure(!IsComponentPSOPrecachingEnabled() || bPSOPrecacheCalled);
 	check(NewPSOPrecachePriority == EPSOPrecachePriority::High || NewPSOPrecachePriority == EPSOPrecachePriority::Highest);
 
-	if (bPrecacheStillRunning && PSOPrecacheRequestPriority < NewPSOPrecachePriority)
+	const bool bIgnoringBoostDueToDefaultMaterialStrat = PrimitiveComponentCVars::bDefaultMaterialDelayStratCanIgnorePSOPrioBoost && GetPSOPrecacheProxyCreationStrategy() == EPSOPrecacheProxyCreationStrategy::UseDefaultMaterialUntilPSOPrecached && ShouldBoostPSOPrecachePriorityOnDraw();
+
+	if (bPrecacheStillRunning && PSOPrecacheRequestPriority < NewPSOPrecachePriority && !bIgnoringBoostDueToDefaultMaterialStrat)
 	{
 		BoostPSOPriority(NewPSOPrecachePriority, MaterialPSOPrecacheRequestIDs);
 		PSOPrecacheRequestPriority = NewPSOPrecachePriority;
diff --git a/Engine/Source/Runtime/Engine/Private/PSOPrecacheMaterial.cpp b/Engine/Source/Runtime/Engine/Private/PSOPrecacheMaterial.cpp
--- a/Engine/Source/Runtime/Engine/Private/PSOPrecacheMaterial.cpp
+++ b/Engine/Source/Runtime/Engine/Private/PSOPrecacheMaterial.cpp
@@ -94,7 +94,7 @@ class FMaterialPSOPrecacheCollectionTask
 {
 public:
 	explicit FMaterialPSOPrecacheCollectionTask(
-		TStrongObjectPtr<UMaterialInterface>* InMaterialInterface,
+		TStrongObjectPtr<UMaterialInterface> InMaterialInterface,
 		const FMaterialPSOPrecacheParams& InPrecacheParams,
 		FGraphEventRef& InCollectionGraphEvent,
 		uint32 InRequestLifecycleID)
@@ -114,7 +114,7 @@ public:
 
 public:
 
-	TStrongObjectPtr<UMaterialInterface>* MaterialInterface;
+	TStrongObjectPtr<UMaterialInterface> MaterialInterface;
 	FMaterialPSOPrecacheParams PrecacheParams;
 	FGraphEventRef CollectionGraphEvent;
 	uint32 RequestLifecycleID;	
@@ -265,9 +265,9 @@ public:
 						// Create task the clear mark fully complete in the cache when done
 						uint32 RequestLifecycleID = LifecycleID;
 						FFunctionGraphTask::CreateAndDispatchWhenReady(
-							[this, Params, RequestLifecycleID]
+							[this, Params, RequestLifecycleID, RequestID]
 							{
-								MarkCompilationComplete(Params, RequestLifecycleID);
+								MarkCompilationComplete(Params, RequestLifecycleID, RequestID);
 							},
 							TStatId{}, CollectionGraphEvent
 						);
@@ -279,7 +279,7 @@ public:
 			if (bUseBackgroundTask)
 			{
 				// Make sure the material instance isn't garbage collected or destroyed yet (create TStrongObjectPtr which will be destroyed on the GT when the collection is done)
-				TStrongObjectPtr<UMaterialInterface>* MaterialInterface = new TStrongObjectPtr<UMaterialInterface>(Params.Material->GetMaterialInterface());
+				TStrongObjectPtr<UMaterialInterface> MaterialInterface(Params.Material->GetMaterialInterface());
 
 				FGraphEventArray Prereqs;
 				// Create and kick off the PSO collection task.
@@ -408,6 +408,36 @@ public:
 		MaterialPSORequests[MaterialPSORequestID] = FMaterialPSOPrecacheParams();
 	}
 
+	FGraphEventArray GetPrecacheCompletePrereq(FMaterialPSOPrecacheRequestID MaterialPSORequestID)
+	{
+		FGraphEventArray FoundEvents;
+		FRWScopeLock ReadLock(RWLock, SLT_ReadOnly);
+
+		if (MaterialPSORequests.IsValidIndex(MaterialPSORequestID))
+		{
+			const FMaterialPSOPrecacheParams& Params = MaterialPSORequests[MaterialPSORequestID];
+			if (FPrecacheData* FindResult = MaterialPSORequestData.Find(Params))
+			{				
+				if (FindResult->State != EState::Completed)
+				{
+					if (!FindResult->CollectionGraphEvent.IsValid())
+					{
+						for (int32 Index = 0; Index < FindResult->ActivePSOPrecacheRequests.Num(); ++Index)
+						{
+							FoundEvents.Add(FindResult->ActivePSOPrecacheRequests[Index].AsyncCompileEvent);
+						}
+					}
+					else
+					{
+						FoundEvents.Add(FindResult->CollectionGraphEvent);
+					}
+				}
+			}
+		}
+		return FoundEvents;
+	}
+
+
 	void BoostPriority(EPSOPrecachePriority NewPri, FMaterialPSOPrecacheRequestID MaterialPSORequestID)
 	{
 		LLM_SCOPE(ELLMTag::PSO);
@@ -562,12 +592,13 @@ private:
 		return (PrecacheData.State != EState::Completed);
 	}	
 
-	void MarkCompilationComplete(const FMaterialPSOPrecacheParams& Params, uint32 RequestLifecycleID)
+	void MarkCompilationComplete(const FMaterialPSOPrecacheParams& Params, uint32 RequestLifecycleID, FMaterialPSOPrecacheRequestID RequestID)
 	{
 		FRWScopeLock WriteLock(RWLock, SLT_Write);
 		FPrecacheData* FindResult = MaterialPSORequestData.Find(Params);
-		if (FindResult && RequestLifecycleID == LifecycleID)
+		if (FindResult && RequestLifecycleID == LifecycleID && RequestID == FindResult->RequestID)
 		{
+			check(!FindResult->CollectionGraphEvent|| FindResult->CollectionGraphEvent->IsCompleted());
 			verifyf(!CheckCompilingPSOs(*FindResult, false /*bBoostPriority*/),
 				TEXT("CheckCompilingPSOs should not be active: EState=%d, ActivePSOPrecacheRequests.Num()=%d"),
 				(int32)FindResult->State,
@@ -600,7 +631,7 @@ void FMaterialPSOPrecacheCollectionTask::DoTask(ENamedThreads::Type CurrentThrea
 	if (RequestLifecycleID != GMaterialPSORequestManager.GetLifecycleID())
 	{
 		CollectionGraphEvent->Unlock();
-		MaterialInterface->Reset();
+		MaterialInterface.Reset();
 		return;
 	}
 
@@ -620,7 +651,7 @@ void FMaterialPSOPrecacheCollectionTask::DoTask(ENamedThreads::Type CurrentThrea
 	GMaterialPSORequestManager.MarkCollectionComplete(PrecacheParams, PSOPrecacheData, PrecacheResults, RequestLifecycleID);
 
 	// Won't touch the material interface anymore - PSO compile jobs take refs to all RHI resources while creating the task
-	MaterialInterface->Reset();
+	MaterialInterface.Reset();
 
 	// Extend MyCompletionGraphEvent to wait for all the async compile events
 	if (PrecacheResults.Num() > 0)
@@ -731,6 +762,66 @@ void ReleasePSOPrecacheData(const TArray<FMaterialPSOPrecacheRequestID>& Materia
 	}
 }
 
+FGraphEventArray GetPrecacheRequestsCompletePrereqs(const TArray<FMaterialPSOPrecacheRequestID>& MaterialPSORequestIDs)
+{
+	TRACE_CPUPROFILER_EVENT_SCOPE(GetPrecacheRequestsCompletePrereqs);
+
+	FGraphEventArray Prereqs;
+	for (FMaterialPSOPrecacheRequestID RequestID : MaterialPSORequestIDs)
+	{
+		Prereqs.Append(GMaterialPSORequestManager.GetPrecacheCompletePrereq(RequestID));
+	}
+	return Prereqs;
+}
+
+#if CSV_PROFILER_STATS
+CSV_DECLARE_CATEGORY_MODULE_EXTERN(RHI_API, PSOPrecacheCompiling);
+
+struct FPSOPrecacheStats
+{
+	TCsvPersistentCustomStat<int>* ProxiesWaitingForMaterialPSOs = FCsvProfiler::Get()->GetOrCreatePersistentCustomStatInt(TEXT("ProxiesWaitingForMaterialPSOs"), CSV_CATEGORY_INDEX(PSOPrecacheCompiling));
+
+	static FPSOPrecacheStats& GetStats()
+	{
+		static FPSOPrecacheStats PersistentStats;
+		return PersistentStats;
+	}
+
+	static bool AreEnabled()
+	{
+		return FCsvProfiler::Get()->IsCategoryEnabled(CSV_CATEGORY_INDEX(PSOPrecacheCompiling));
+	}
+};
+#endif
+
+void BoostPSOPriorityOnDraw(const TArray<FMaterialPSOPrecacheRequestID>& MaterialPSORequestIDs)
+{
+	TRACE_CPUPROFILER_EVENT_SCOPE(BoostPSOPriorityOnDraw);
+
+	constexpr EPSOPrecachePriority NewPri = EPSOPrecachePriority::Highest;
+
+	// if any of the requests are still in progress then we can consider this to be default material usage.
+	FGraphEventArray AllRequests = GetPrecacheRequestsCompletePrereqs(MaterialPSORequestIDs);
+	if(!AllRequests.IsEmpty())
+	{ 
+#if CSV_PROFILER_STATS
+		if (FPSOPrecacheStats::AreEnabled())
+		{
+			FPSOPrecacheStats::GetStats().ProxiesWaitingForMaterialPSOs->Add(1);
+
+			FFunctionGraphTask::CreateAndDispatchWhenReady(
+				[]
+				{
+					FPSOPrecacheStats::GetStats().ProxiesWaitingForMaterialPSOs->Add(-1);
+				},
+				TStatId{}, &AllRequests
+			);
+		}
+#endif
+ 		BoostPSOPriority(NewPri, MaterialPSORequestIDs);
+	}
+}
+
 void BoostPSOPriority(EPSOPrecachePriority NewPri, const TArray<FMaterialPSOPrecacheRequestID>& MaterialPSORequestIDs)
 {
 	TRACE_CPUPROFILER_EVENT_SCOPE(BoostPSOPriority);
diff --git a/Engine/Source/Runtime/Engine/Private/PrimitiveSceneProxy.cpp b/Engine/Source/Runtime/Engine/Private/PrimitiveSceneProxy.cpp
--- a/Engine/Source/Runtime/Engine/Private/PrimitiveSceneProxy.cpp
+++ b/Engine/Source/Runtime/Engine/Private/PrimitiveSceneProxy.cpp
@@ -606,7 +606,7 @@ FPrimitiveSceneProxy::FPrimitiveSceneProxy(const FPrimitiveSceneProxyDesc& InPro
 	if (UPrimitiveComponent* PrimitiveComponent = Cast<UPrimitiveComponent>(InProxyDesc.Component))
 	{
 		check(!PrimitiveComponent->ShouldRenderProxyFallbackToDefaultMaterial() || !PrimitiveComponent->MaterialPSOPrecacheRequestIDs.IsEmpty());
-		if (ShouldBoostPSOPrecachePriorityOnDraw() && PrimitiveComponent->IsPSOPrecaching() && PrimitiveComponent->PSOPrecacheRequestPriority == EPSOPrecachePriority::High)
+		if (ShouldBoostPSOPrecachePriorityOnDraw() && PrimitiveComponent->IsPSOPrecaching())
 		{
 			SetPSORequestsToBoostOnDraw(PrimitiveComponent->MaterialPSOPrecacheRequestIDs);
 		}
@@ -1840,8 +1840,8 @@ void FPrimitiveSceneProxy::BoostPrecachedPSORequestsOnDraw()
 {
 	if (!PSOPrecacheRequestsToBoostOnDraw.IsEmpty())
 	{
-		BoostPSOPriority(EPSOPrecachePriority::Highest, PSOPrecacheRequestsToBoostOnDraw);
-		PSOPrecacheRequestsToBoostOnDraw.Empty();
+		TArray<FMaterialPSOPrecacheRequestID> LocalPSOPrecacheRequestsToBoostOnDraw = MoveTemp(PSOPrecacheRequestsToBoostOnDraw);
+		BoostPSOPriorityOnDraw(LocalPSOPrecacheRequestsToBoostOnDraw);
 	}
 }
 
diff --git a/Engine/Source/Runtime/Engine/Public/PSOPrecacheMaterial.h b/Engine/Source/Runtime/Engine/Public/PSOPrecacheMaterial.h
--- a/Engine/Source/Runtime/Engine/Public/PSOPrecacheMaterial.h
+++ b/Engine/Source/Runtime/Engine/Public/PSOPrecacheMaterial.h
@@ -147,6 +147,11 @@ extern ENGINE_API void ReleasePSOPrecacheData(const TArray<FMaterialPSOPrecacheR
  */
 extern ENGINE_API void BoostPSOPriority(EPSOPrecachePriority NewPri, const TArray<FMaterialPSOPrecacheRequestID>& MaterialPSORequestIDs);
 
+/**
+ * Boost priority to highest when the PSOs are required for drawing. 
+ */
+extern ENGINE_API void BoostPSOPriorityOnDraw(const TArray<FMaterialPSOPrecacheRequestID>& MaterialPSORequestIDs);
+
 /**
  * Invalidate & clear all the current material PSO requests
  */
diff --git a/Engine/Source/Runtime/RHI/Private/PipelineStateCache.cpp b/Engine/Source/Runtime/RHI/Private/PipelineStateCache.cpp
--- a/Engine/Source/Runtime/RHI/Private/PipelineStateCache.cpp
+++ b/Engine/Source/Runtime/RHI/Private/PipelineStateCache.cpp
@@ -42,7 +42,7 @@ namespace PipelineStateCache
 #define PIPELINESTATECACHE_VERIFYTHREADSAFE (!UE_BUILD_SHIPPING && !UE_BUILD_TEST)
 
 CSV_DECLARE_CATEGORY_EXTERN(PSO);
-CSV_DEFINE_CATEGORY(PSOPrecacheCompiling, false);
+CSV_DEFINE_CATEGORY_MODULE(RHI_API, PSOPrecacheCompiling, false);
 
 DEFINE_LOG_CATEGORY_STATIC(LogPSOHitching, Log, All);
 
@@ -847,7 +847,7 @@ class FPSOCSVStatTracker
 		TCsvPersistentCustomStat<int>* MinPriPSOPrecompileQueueLength		= FCsvProfiler::Get()->GetOrCreatePersistentCustomStatInt(TEXT("MinPriPSOPrecompileQueueLength"), CSV_CATEGORY_INDEX(PSOPrecacheCompiling));
 	};
 
-	FGraphicsPipelineStateInitializer::EPSOPrecacheCompileType Pri = FGraphicsPipelineStateInitializer::EPSOPrecacheCompileType::NotSet;
+	std::atomic_int Pri = (int)FGraphicsPipelineStateInitializer::EPSOPrecacheCompileType::NotSet;
 
 	static FPersistentStats& GetStats()
 	{
@@ -892,15 +892,15 @@ public:
 			return;
 		}
 
-		if (bPrecache && InPri != Pri)
+		FGraphicsPipelineStateInitializer::EPSOPrecacheCompileType PriorValue = (FGraphicsPipelineStateInitializer::EPSOPrecacheCompileType)Pri.exchange((int)InPri);
+		if (bPrecache && PriorValue != InPri)
 		{
 			// unset the previous priority
-			if (Pri != FGraphicsPipelineStateInitializer::EPSOPrecacheCompileType::NotSet)
+			if (PriorValue != FGraphicsPipelineStateInitializer::EPSOPrecacheCompileType::NotSet)
 			{
-				GetPriStat(Pri, false)->Add(-1);
+				GetPriStat(PriorValue, false)->Add(-1);
 			}
-			Pri = InPri;
-			GetPriStat(Pri, false)->Add(1);
+			GetPriStat(InPri, false)->Add(1);
 		}
 	}
 
@@ -940,10 +940,11 @@ public:
 				PersistentStats.PSOComputePrecompilesInProgress->Add(-1);
 			}
 
-			if (ensure(Pri != FGraphicsPipelineStateInitializer::EPSOPrecacheCompileType::NotSet))
+			FGraphicsPipelineStateInitializer::EPSOPrecacheCompileType MyPri = (FGraphicsPipelineStateInitializer::EPSOPrecacheCompileType)Pri.load();
+			if (ensure(MyPri != FGraphicsPipelineStateInitializer::EPSOPrecacheCompileType::NotSet))
 			{
-				GetPriStat(Pri, true)->Add(-1);
-				GetPriStat(Pri, false)->Add(-1);
+				GetPriStat(MyPri, true)->Add(-1);
+				GetPriStat(MyPri, false)->Add(-1);
 			}
 		}
 	}
