DrawDebugLibrary: Enabled thread-safe debug drawing in Anim Instance via DrawDebugLibrary

diff --git a/Engine/Plugins/Experimental/DrawDebugLibrary/Source/DrawDebugLibrary/Private/DrawDebugLibrary.cpp b/Engine/Plugins/Experimental/DrawDebugLibrary/Source/DrawDebugLibrary/Private/DrawDebugLibrary.cpp
--- a/Engine/Plugins/Experimental/DrawDebugLibrary/Source/DrawDebugLibrary/Private/DrawDebugLibrary.cpp
+++ b/Engine/Plugins/Experimental/DrawDebugLibrary/Source/DrawDebugLibrary/Private/DrawDebugLibrary.cpp
@@ -7,6 +7,7 @@
 #include "SkeletalDebugRendering.h"
 #include "HitProxies.h"
 #include "DrawDebugHelpers.h"
+#include "Animation/AnimInstance.h"
 #include "Animation/AnimInstanceProxy.h"
 #include "VisualLogger/VisualLogger.h"
 #include "Components/SkinnedMeshComponent.h"
@@ -419,6 +420,11 @@ FDebugDrawer::FDebugDrawer(UWorld* InWorld) : FDebugDrawer()
 	World = InWorld;
 }
 
+FDebugDrawer::FDebugDrawer(UAnimInstance* InAnimInstance) : FDebugDrawer()
+{
+	AnimInstance = InAnimInstance;
+}
+
 FDebugDrawer::FDebugDrawer(FAnimInstanceProxy* InAnimInstanceProxy) : FDebugDrawer()
 {
 	AnimInstanceProxy = InAnimInstanceProxy;
@@ -449,6 +455,11 @@ FDebugDrawer FDebugDrawer::MakeDebugDrawer(UWorld* World)
 	return FDebugDrawer(World);
 }
 
+FDebugDrawer FDebugDrawer::MakeDebugDrawer(UAnimInstance* AnimInstance)
+{
+	return FDebugDrawer(AnimInstance);
+}
+
 FDebugDrawer FDebugDrawer::MakeDebugDrawer(FAnimInstanceProxy* AnimInstanceProxy)
 {
 	return FDebugDrawer(AnimInstanceProxy);
@@ -505,6 +516,7 @@ FDebugDrawer FDebugDrawer::MakeMergedDebugDrawer(const TArrayView<const FDebugDr
 		if (SubDrawer.PDI) { Drawer.PDI = SubDrawer.PDI; }
 		if (SubDrawer.RigVM) { Drawer.RigVM = SubDrawer.RigVM; }
 		if (SubDrawer.World) { Drawer.World = SubDrawer.World; }
+		if (SubDrawer.AnimInstance) { Drawer.AnimInstance = SubDrawer.AnimInstance; }
 		if (SubDrawer.AnimInstanceProxy) { Drawer.AnimInstanceProxy = SubDrawer.AnimInstanceProxy; }
 		if (SubDrawer.VisualLoggerObject)
 		{
@@ -555,6 +567,14 @@ void FDebugDrawer::DrawPoints(const TArrayView<const FVector> Locations, const F
 		}
 	}
 
+	if (AnimInstance)
+	{
+		for (int32 Idx = 0; Idx < Num; Idx++)
+		{
+			AnimInstance->AnimDrawDebugLine(Locations[Idx], Locations[Idx], Color.ToFColor(true), false, -1.0f, Thickness, bDepthTest ? SDPG_World : SDPG_Foreground);
+		}
+	}
+
 	if (AnimInstanceProxy)
 	{
 		for (int32 Idx = 0; Idx < Num; Idx++)
@@ -660,6 +680,14 @@ void FDebugDrawer::DrawLines(const TArrayView<const TPair<FVector, FVector>> Seg
 		}
 	}
 
+	if (AnimInstance)
+	{
+		for (int32 Idx = 0; Idx < Num; Idx++)
+		{
+			AnimInstance->AnimDrawDebugLine(Segments[Idx].Key, Segments[Idx].Value, Color.ToFColor(true), false, -1.0f, Thickness, bDepthTest ? SDPG_World : SDPG_Foreground);
+		}
+	}
+
 	if (AnimInstanceProxy)
 	{
 		for (int32 Idx = 0; Idx < Num; Idx++)
@@ -753,6 +781,11 @@ FDebugDrawer UDrawDebugLibrary::MakeWorldDebugDrawer(UWorld* World)
 	return FDebugDrawer::MakeDebugDrawer(World);
 }
 
+FDebugDrawer UDrawDebugLibrary::MakeAnimInstanceDebugDrawer(UAnimInstance* AnimInstance)
+{
+	return FDebugDrawer::MakeDebugDrawer(AnimInstance);
+}
+
 FDebugDrawer UDrawDebugLibrary::MakeAnimInstanceProxyDebugDrawer(FAnimInstanceProxy* AnimInstanceProxy)
 {
 	return FDebugDrawer::MakeDebugDrawer(AnimInstanceProxy);
diff --git a/Engine/Plugins/Experimental/DrawDebugLibrary/Source/DrawDebugLibrary/Public/DrawDebugLibrary.h b/Engine/Plugins/Experimental/DrawDebugLibrary/Source/DrawDebugLibrary/Public/DrawDebugLibrary.h
--- a/Engine/Plugins/Experimental/DrawDebugLibrary/Source/DrawDebugLibrary/Public/DrawDebugLibrary.h
+++ b/Engine/Plugins/Experimental/DrawDebugLibrary/Source/DrawDebugLibrary/Public/DrawDebugLibrary.h
@@ -476,6 +476,9 @@ public: // Constructors
 	/** Make a debug drawer for a UWorld */
 	UE_API FDebugDrawer(UWorld* InWorld);
 
+	/** Make a debug drawer for an AnimInstance */
+	UE_API FDebugDrawer(UAnimInstance* InAnimInstance);
+
 	/** Make a debug drawer for an AnimInstanceProxy */
 	UE_API FDebugDrawer(FAnimInstanceProxy* InAnimInstanceProxy);
 
@@ -494,6 +497,9 @@ public: // Constructors
 	/** Make a debug drawer for a UWorld */
 	static UE_API FDebugDrawer MakeDebugDrawer(UWorld* World);
 
+	/** Make a debug drawer for an AnimInstance */
+	static UE_API FDebugDrawer MakeDebugDrawer(UAnimInstance* AnimInstance);
+
 	/** Make a debug drawer for an AnimInstanceProxy */
 	static UE_API FDebugDrawer MakeDebugDrawer(FAnimInstanceProxy* AnimInstanceProxy);
 
@@ -542,6 +548,10 @@ private:
 	UPROPERTY()
 	TObjectPtr<UWorld> World = nullptr;
 
+	/** Anim Instance Object - used for debug drawing in an AnimBlueprints. */
+	UPROPERTY()
+	TObjectPtr<UAnimInstance> AnimInstance = nullptr;
+
 	/** Anim Instance Proxy Object - used for debug drawing in an AnimNode. */
 	FAnimInstanceProxy* AnimInstanceProxy = nullptr;
 
@@ -592,6 +602,10 @@ public: // Drawers
 	/** Make a debug drawer for a UWorld */
 	static UE_API FDebugDrawer MakeWorldDebugDrawer(UWorld* World);
 
+	/** Make a debug drawer for an AnimInstance */
+	UFUNCTION(BlueprintPure, Category = "Draw Debug Library", meta = (DevelopmentOnly, DefaultToSelf = "AnimInstance"))
+	static UE_API FDebugDrawer MakeAnimInstanceDebugDrawer(UAnimInstance* AnimInstance);
+
 	/** Make a debug drawer for an AnimInstanceProxy */
 	static UE_API FDebugDrawer MakeAnimInstanceProxyDebugDrawer(FAnimInstanceProxy* AnimInstanceProxy);
 
@@ -642,7 +656,7 @@ public: // Visual Debugger Logging
 public: // Convenience functions
 
 	/** Convenience function for making an array of linearly spaced float values */
-	UFUNCTION(BlueprintPure = true, Category = "Draw Debug Library", meta = (NotBlueprintThreadSafe))
+	UFUNCTION(BlueprintPure = true, Category = "Draw Debug Library")
 	static UE_API void MakeLinearlySpacedFloatArray(TArray<float>& OutValues, const float Start = 0.0f, const float Stop = 1.0f, const int32 Num = 10);
 
 	/** Convenience function that adds a float to the end of an array, popping values from the front once the max is reached */
diff --git a/Engine/Source/Runtime/Engine/Classes/Animation/AnimInstance.h b/Engine/Source/Runtime/Engine/Classes/Animation/AnimInstance.h
--- a/Engine/Source/Runtime/Engine/Classes/Animation/AnimInstance.h
+++ b/Engine/Source/Runtime/Engine/Classes/Animation/AnimInstance.h
@@ -14,6 +14,9 @@
 #include "Animation/AnimInertializationRequest.h"
 #include "AnimInstance.generated.h"
 
+// Disable debugging information for shipping and test builds.
+#define ENABLE_ANIM_DRAW_DEBUG (1 && !(UE_BUILD_SHIPPING || UE_BUILD_TEST))
+
 // Post Compile Validation requires WITH_EDITOR
 #define ANIMINST_PostCompileValidation WITH_EDITOR
 
@@ -935,6 +938,34 @@ public:
 		return static_cast<SubsystemType&>(*Subsystem);
 	}
 
+public:
+
+#if ENABLE_ANIM_DRAW_DEBUG
+		ENGINE_API void AnimDrawDebugOnScreenMessage(const FString& DebugMessage, const FColor& Color, const FVector2D& TextScale = FVector2D::UnitVector, ESceneDepthPriorityGroup DepthPriority = SDPG_World);
+		ENGINE_API void AnimDrawDebugInWorldMessage(const FString& DebugMessage, const FVector& TextLocation, const FColor& Color, float TextScale);
+		ENGINE_API void AnimDrawDebugLine(const FVector& StartLoc, const FVector& EndLoc, const FColor& Color, bool bPersistentLines = false, float LifeTime = -1.f, float Thickness = 0.f, ESceneDepthPriorityGroup DepthPriority = SDPG_World);
+		ENGINE_API void AnimDrawDebugDirectionalArrow(const FVector& LineStart, const FVector& LineEnd, float ArrowSize, const FColor& Color, bool bPersistentLines = false, float LifeTime = -1.f, float Thickness = 0.f, ESceneDepthPriorityGroup DepthPriority = SDPG_World);
+		ENGINE_API void AnimDrawDebugSphere(const FVector& Center, float Radius, int32 Segments, const FColor& Color, bool bPersistentLines = false, float LifeTime = -1.f, float Thickness = 0.f, ESceneDepthPriorityGroup DepthPriority = SDPG_World);
+		ENGINE_API void AnimDrawDebugCoordinateSystem(FVector const& AxisLoc, FRotator const& AxisRot, float Scale = 1.f, bool bPersistentLines = false, float LifeTime = -1.f, float Thickness = 0.f, ESceneDepthPriorityGroup DepthPriority = SDPG_World);
+		ENGINE_API void AnimDrawDebugPlane(const FTransform& BaseTransform, float Radii, const FColor& Color, bool bPersistentLines = false, float LifeTime = -1.f, float Thickness = 0.f, ESceneDepthPriorityGroup DepthPriority = SDPG_World);
+		ENGINE_API void AnimDrawDebugPoint(const FVector& Loc, float Size, const FColor& Color, bool bPersistentLines = false, float LifeTime = -1.f, ESceneDepthPriorityGroup DepthPriority = SDPG_World);
+		ENGINE_API void AnimDrawDebugCircle(const FVector& Center, float Radius, int32 Segments, const FColor& Color, const FVector& UpVector = FVector::UpVector, bool bPersistentLines = false, float LifeTime = -1.f, ESceneDepthPriorityGroup DepthPriority = SDPG_World, float Thickness = 0.f);
+		ENGINE_API void AnimDrawDebugCone(const FVector& Center, float Radius, const FVector& Direction, float AngleWidth, float AngleHeight, int32 Segments, const FColor& Color, bool bPersistentLines = false, float LifeTime = -1.f, ESceneDepthPriorityGroup DepthPriority = SDPG_World, float Thickness = 0.f);
+		ENGINE_API void AnimDrawDebugCapsule(const FVector& Center, float HalfHeight, float Radius, const FRotator& Rotation, const FColor& Color, bool bPersistentLines = false, float LifeTime = -1.f, float Thickness = 0.f, ESceneDepthPriorityGroup DepthPriority = SDPG_World);
+#else
+		void AnimDrawDebugOnScreenMessage(const FString& DebugMessage, const FColor& Color, const FVector2D& TextScale = FVector2D::UnitVector, ESceneDepthPriorityGroup DepthPriority = SDPG_World) {}
+		void AnimDrawDebugInWorldMessage(const FString& DebugMessage, const FVector& TextLocation, const FColor& Color, float TextScale) {}
+		void AnimDrawDebugLine(const FVector& StartLoc, const FVector& EndLoc, const FColor& Color, bool bPersistentLines = false, float LifeTime = -1.f, float Thickness = 0.f, ESceneDepthPriorityGroup DepthPriority = SDPG_World) {}
+		void AnimDrawDebugDirectionalArrow(const FVector& LineStart, const FVector& LineEnd, float ArrowSize, const FColor& Color, bool bPersistentLines = false, float LifeTime = -1.f, float Thickness = 0.f, ESceneDepthPriorityGroup DepthPriority = SDPG_World) {}
+		void AnimDrawDebugSphere(const FVector& Center, float Radius, int32 Segments, const FColor& Color, bool bPersistentLines = false, float LifeTime = -1.f, float Thickness = 0.f, ESceneDepthPriorityGroup DepthPriority = SDPG_World) {}
+		void AnimDrawDebugCoordinateSystem(FVector const& AxisLoc, FRotator const& AxisRot, float Scale = 1.f, bool bPersistentLines = false, float LifeTime = -1.f, float Thickness = 0.f, ESceneDepthPriorityGroup DepthPriority = SDPG_World) {}
+		void AnimDrawDebugPlane(const FTransform& BaseTransform, float Radii, const FColor& Color, bool bPersistentLines = false, float LifeTime = -1.f, float Thickness = 0.f, ESceneDepthPriorityGroup DepthPriority = SDPG_World) {}
+		void AnimDrawDebugPoint(const FVector& Loc, float Size, const FColor& Color, bool bPersistentLines = false, float LifeTime = -1.f, ESceneDepthPriorityGroup DepthPriority = SDPG_World) {}
+		void AnimDrawDebugCircle(const FVector& Center, float Radius, int32 Segments, const FColor& Color, const FVector& UpVector = FVector::UpVector, bool bPersistentLines = false, float LifeTime = -1.f, ESceneDepthPriorityGroup DepthPriority = SDPG_World, float Thickness = 0.f) {}
+		void AnimDrawDebugCone(const FVector& Center, float Radius, const FVector& Direction, float AngleWidth, float AngleHeight, int32 Segments, const FColor& Color, bool bPersistentLines = false, float LifeTime = -1.f, ESceneDepthPriorityGroup DepthPriority = SDPG_World, float Thickness = 0.f) {}
+		void AnimDrawDebugCapsule(const FVector& Center, float HalfHeight, float Radius, const FRotator& Rotation, const FColor& Color, bool bPersistentLines = false, float LifeTime = -1.f, float Thickness = 0.f, ESceneDepthPriorityGroup DepthPriority = SDPG_World) {}
+#endif // ENABLE_ANIM_DRAW_DEBUG
+
 private:
 	/** Helper function to perform layer overlay actions (set, clear) */
 	ENGINE_API void PerformLinkedLayerOverlayOperation(TSubclassOf<UAnimInstance> InClass, TFunctionRef<UClass*(UClass*, FAnimNode_LinkedAnimLayer*)> InClassSelectorFunction, bool bInDeferSubGraphInitialization = false);
diff --git a/Engine/Source/Runtime/Engine/Private/Animation/AnimInstance.cpp b/Engine/Source/Runtime/Engine/Private/Animation/AnimInstance.cpp
--- a/Engine/Source/Runtime/Engine/Private/Animation/AnimInstance.cpp
+++ b/Engine/Source/Runtime/Engine/Private/Animation/AnimInstance.cpp
@@ -4033,6 +4033,65 @@ FAnimSubsystemInstance* UAnimInstance::FindSubsystem(UScriptStruct* InSubsystemT
 	return nullptr;
 }
 
+#if ENABLE_ANIM_DRAW_DEBUG
+
+void UAnimInstance::AnimDrawDebugOnScreenMessage(const FString& DebugMessage, const FColor& Color, const FVector2D& TextScale, ESceneDepthPriorityGroup DepthPriority)
+{
+	GetProxyOnAnyThread<FAnimInstanceProxy>().AnimDrawDebugOnScreenMessage(DebugMessage, Color, TextScale, DepthPriority);
+}
+
+void UAnimInstance::AnimDrawDebugInWorldMessage(const FString& DebugMessage, const FVector& TextLocation, const FColor& Color, float TextScale)
+{
+	GetProxyOnAnyThread<FAnimInstanceProxy>().AnimDrawDebugInWorldMessage(DebugMessage, TextLocation, Color, TextScale);
+}
+
+void UAnimInstance::AnimDrawDebugLine(const FVector& StartLoc, const FVector& EndLoc, const FColor& Color, bool bPersistentLines, float LifeTime, float Thickness, ESceneDepthPriorityGroup DepthPriority)
+{
+	GetProxyOnAnyThread<FAnimInstanceProxy>().AnimDrawDebugLine(StartLoc, EndLoc, Color, bPersistentLines, LifeTime, Thickness, DepthPriority);
+}
+
+void UAnimInstance::AnimDrawDebugDirectionalArrow(const FVector& LineStart, const FVector& LineEnd, float ArrowSize, const FColor& Color, bool bPersistentLines, float LifeTime, float Thickness, ESceneDepthPriorityGroup DepthPriority)
+{
+	GetProxyOnAnyThread<FAnimInstanceProxy>().AnimDrawDebugDirectionalArrow(LineStart, LineEnd, ArrowSize, Color, bPersistentLines, LifeTime, Thickness, DepthPriority);
+}
+
+void UAnimInstance::AnimDrawDebugSphere(const FVector& Center, float Radius, int32 Segments, const FColor& Color, bool bPersistentLines, float LifeTime, float Thickness, ESceneDepthPriorityGroup DepthPriority)
+{
+	GetProxyOnAnyThread<FAnimInstanceProxy>().AnimDrawDebugSphere(Center, Radius, Segments, Color, bPersistentLines, LifeTime, Thickness, DepthPriority);
+}
+
+void UAnimInstance::AnimDrawDebugCoordinateSystem(FVector const& AxisLoc, FRotator const& AxisRot, float Scale, bool bPersistentLines, float LifeTime, float Thickness, ESceneDepthPriorityGroup DepthPriority)
+{
+	GetProxyOnAnyThread<FAnimInstanceProxy>().AnimDrawDebugCoordinateSystem(AxisLoc, AxisRot, Scale, bPersistentLines, LifeTime, Thickness, DepthPriority);
+}
+
+void UAnimInstance::AnimDrawDebugPlane(const FTransform& BaseTransform, float Radii, const FColor& Color, bool bPersistentLines, float LifeTime, float Thickness, ESceneDepthPriorityGroup DepthPriority)
+{
+	GetProxyOnAnyThread<FAnimInstanceProxy>().AnimDrawDebugPlane(BaseTransform, Radii, Color, bPersistentLines, LifeTime, Thickness, DepthPriority);
+}
+
+void UAnimInstance::AnimDrawDebugPoint(const FVector& Loc, float Size, const FColor& Color, bool bPersistentLines, float LifeTime, ESceneDepthPriorityGroup DepthPriority)
+{
+	GetProxyOnAnyThread<FAnimInstanceProxy>().AnimDrawDebugPoint(Loc, Size, Color, bPersistentLines, LifeTime, DepthPriority);
+}
+
+void UAnimInstance::AnimDrawDebugCircle(const FVector& Center, float Radius, int32 Segments, const FColor& Color, const FVector& UpVector, bool bPersistentLines, float LifeTime, ESceneDepthPriorityGroup DepthPriority, float Thickness)
+{
+	GetProxyOnAnyThread<FAnimInstanceProxy>().AnimDrawDebugCircle(Center, Radius, Segments, Color, UpVector, bPersistentLines, LifeTime, DepthPriority, Thickness);
+}
+
+void UAnimInstance::AnimDrawDebugCone(const FVector& Center, float Radius, const FVector& Direction, float AngleWidth, float AngleHeight, int32 Segments, const FColor& Color, bool bPersistentLines, float LifeTime, ESceneDepthPriorityGroup DepthPriority, float Thickness)
+{
+	GetProxyOnAnyThread<FAnimInstanceProxy>().AnimDrawDebugCone(Center, Radius, Direction, AngleWidth, AngleHeight, Segments, Color, bPersistentLines, LifeTime, DepthPriority, Thickness);
+}
+
+void UAnimInstance::AnimDrawDebugCapsule(const FVector& Center, float HalfHeight, float Radius, const FRotator& Rotation, const FColor& Color, bool bPersistentLines, float LifeTime, float Thickness, ESceneDepthPriorityGroup DepthPriority)
+{
+	GetProxyOnAnyThread<FAnimInstanceProxy>().AnimDrawDebugCapsule(Center, HalfHeight, Radius, Rotation, Color, bPersistentLines, LifeTime, Thickness, DepthPriority);
+}
+
+#endif // ENABLE_ANIM_DRAW_DEBUG
+
 UAnimInstance* UAnimInstance::GetLinkedAnimLayerInstanceByGroup(FName InGroup) const
 {
 	if (IAnimClassInterface* AnimBlueprintClass = IAnimClassInterface::GetFromClass(GetClass()))
