Avoid unnecessary array copies in Physical Animation Component

diff --git a/Engine/Source/Runtime/Engine/Classes/Components/SkeletalMeshComponent.h b/Engine/Source/Runtime/Engine/Classes/Components/SkeletalMeshComponent.h
--- a/Engine/Source/Runtime/Engine/Classes/Components/SkeletalMeshComponent.h
+++ b/Engine/Source/Runtime/Engine/Classes/Components/SkeletalMeshComponent.h
@@ -426,20 +426,15 @@ public:
 	UPROPERTY(EditAnywhere, BlueprintReadWrite, Category=Animation, meta=(ShowOnlyInnerProperties, EditCondition = bEnableAnimation))
 	struct FSingleAnimationPlayData AnimationData;
 
-	// this is explicit copy because this buffer is reused during evaluation
-	// we want to have reference and emptied during evaluation
+	// Blocks on the evaluation task (if necessary) and then returns a copy of the bone space
+	// transforms. A copy is returned because the original array can/will be modified during the
+	// evaluation, which can be running in a parallel task.
 	ENGINE_API TArray<FTransform> GetBoneSpaceTransforms();
 
-	/** Get the bone space transforms as array view. */
+	// Blocks on the evaluation task (if necessary) and then returns bone space transforms as an
+	// array view.
 	ENGINE_API TArrayView<const FTransform> GetBoneSpaceTransformsView();
 
-	/** 
-	 * Temporary array of local-space (relative to parent bone) rotation/translation for each bone. 
-	 * This property is not safe to access during evaluation, so we created wrapper.
-	 */
-	UE_DEPRECATED(4.23, "Direct access to this property is deprecated, please use GetBoneSpaceTransforms instead. We will move to private in the future.")
-	TArray<FTransform> BoneSpaceTransforms;
-
 public:
 	/** Offset of the root bone from the reference pose. Used to offset bounding box. */
 	UPROPERTY(transient)
@@ -456,6 +451,12 @@ public:
 	ENGINE_API const TArray<FTransform>& GetCachedComponentSpaceTransforms() const;
 
 private:
+	/**
+	 * Temporary array of local-space (relative to parent bone) rotation/translation for each bone.
+	 * This property is not safe to access directly during evaluation - see GetBoneSpaceTransforms() etc
+	 */
+	TArray<FTransform> BoneSpaceTransforms;
+
 	/** Any running linked anim instances */
 	UPROPERTY(transient)
 	TArray<TObjectPtr<UAnimInstance>> LinkedInstances;
diff --git a/Engine/Source/Runtime/Engine/Private/Components/SkeletalMeshComponent.cpp b/Engine/Source/Runtime/Engine/Private/Components/SkeletalMeshComponent.cpp
--- a/Engine/Source/Runtime/Engine/Private/Components/SkeletalMeshComponent.cpp
+++ b/Engine/Source/Runtime/Engine/Private/Components/SkeletalMeshComponent.cpp
@@ -5411,7 +5411,7 @@ PRAGMA_ENABLE_DEPRECATION_WARNINGS
 TArray<FTransform> USkeletalMeshComponent::GetBoneSpaceTransforms() 
 {
 	// We may be doing parallel evaluation on the current anim instance
-	// Calling this here with true will block this init till that thread completes
+	// Calling this here with true will block until that thread completes
 	// and it is safe to continue
 	const bool bBlockOnTask = true; // wait on evaluation task so it is safe to swap the buffers
 	const bool bPerformPostAnimEvaluation = true; // Do PostEvaluation so we make sure to swap the buffers back. 
@@ -5422,11 +5422,10 @@ TArray<FTransform> USkeletalMeshComponent::GetBoneSpaceTransforms()
 	PRAGMA_ENABLE_DEPRECATION_WARNINGS
 }
 
-
 TArrayView<const FTransform> USkeletalMeshComponent::GetBoneSpaceTransformsView()
 {
 	// We may be doing parallel evaluation on the current anim instance
-	// Calling this here with true will block this init till that thread completes
+	// Calling this here with true will block until that thread completes
 	// and it is safe to continue
 	const bool bBlockOnTask = true; // wait on evaluation task so it is safe to swap the buffers
 	const bool bPerformPostAnimEvaluation = true; // Do PostEvaluation so we make sure to swap the buffers back. 
diff --git a/Engine/Source/Runtime/Engine/Private/PhysicsEngine/PhysicalAnimationComponent.cpp b/Engine/Source/Runtime/Engine/Private/PhysicsEngine/PhysicalAnimationComponent.cpp
--- a/Engine/Source/Runtime/Engine/Private/PhysicsEngine/PhysicalAnimationComponent.cpp
+++ b/Engine/Source/Runtime/Engine/Private/PhysicsEngine/PhysicalAnimationComponent.cpp
@@ -210,15 +210,22 @@ FTransform UPhysicalAnimationComponent::GetBodyTargetTransform(FName BodyName) c
 	return FTransform::Identity;
 }
 
-FTransform ComputeWorldSpaceTargetTM(const USkeletalMeshComponent& SkeletalMeshComponent, const TArray<FTransform>& SpaceBases, int32 BoneIndex)
+FTransform ComputeWorldSpaceTargetTM(
+	const USkeletalMeshComponent& SkeletalMeshComponent, 
+	const TArrayView<const FTransform> ComponentSpaceTransforms, 
+	int32 BoneIndex)
 {
-	return SpaceBases[BoneIndex] * SkeletalMeshComponent.GetComponentToWorld();
+	return ComponentSpaceTransforms[BoneIndex] * SkeletalMeshComponent.GetComponentToWorld();
 }
 
-FTransform ComputeLocalSpaceTargetTM(const USkeletalMeshComponent& SkeletalMeshComponent, const UPhysicsAsset& PhysAsset, const TArray<FTransform>& LocalTransforms, int32 BoneIndex)
+FTransform ComputeLocalSpaceTargetTM(
+	const USkeletalMeshComponent& SkeletalMeshComponent, 
+	const UPhysicsAsset& PhysAsset, 
+	const TArrayView<const FTransform> BoneSpaceTransforms,
+	int32 BoneIndex)
 {
 	const FReferenceSkeleton& RefSkeleton = SkeletalMeshComponent.GetSkeletalMeshAsset()->GetRefSkeleton();
-	FTransform AccumulatedDelta = LocalTransforms[BoneIndex];
+	FTransform AccumulatedDelta = BoneSpaceTransforms[BoneIndex];
 	int32 CurBoneIdx = BoneIndex;
 	while ((CurBoneIdx = RefSkeleton.GetParentIndex(CurBoneIdx)) != INDEX_NONE)
 	{
@@ -245,15 +252,23 @@ FTransform ComputeLocalSpaceTargetTM(const USkeletalMeshComponent& SkeletalMeshC
 			}
 		}
 
-		AccumulatedDelta = AccumulatedDelta * LocalTransforms[CurBoneIdx];
+		AccumulatedDelta = AccumulatedDelta * BoneSpaceTransforms[CurBoneIdx];
 	}
 
 	return FTransform::Identity;
 }
 
-FTransform ComputeTargetTM(const FPhysicalAnimationData& PhysAnimData, const USkeletalMeshComponent& SkeletalMeshComponent, const UPhysicsAsset& PhysAsset, const TArray<FTransform>& LocalTransforms, const TArray<FTransform>& SpaceBases, int32 BoneIndex)
+FTransform ComputeTargetTM(
+	const FPhysicalAnimationData& PhysAnimData, 
+	const USkeletalMeshComponent& SkeletalMeshComponent, 
+	const UPhysicsAsset& PhysAsset, 
+	const TArrayView<const FTransform> BoneSpaceTransforms,
+	const TArrayView<const FTransform> ComponentSpaceTransforms, 
+	int32 BoneIndex)
 {
-	return PhysAnimData.bIsLocalSimulation ? ComputeLocalSpaceTargetTM(SkeletalMeshComponent, PhysAsset, LocalTransforms, BoneIndex) : ComputeWorldSpaceTargetTM(SkeletalMeshComponent, SpaceBases, BoneIndex);
+	return PhysAnimData.bIsLocalSimulation 
+		? ComputeLocalSpaceTargetTM(SkeletalMeshComponent, PhysAsset, BoneSpaceTransforms, BoneIndex)
+		: ComputeWorldSpaceTargetTM(SkeletalMeshComponent, ComponentSpaceTransforms, BoneIndex);
 }
 
 void UPhysicalAnimationComponent::UpdateTargetActors(ETeleportType TeleportType)
@@ -265,9 +280,9 @@ void UPhysicalAnimationComponent::UpdateTargetActors(ETeleportType TeleportType)
 
 		// Note we use GetEditableComponentSpaceTransforms because we need to update target actors in the midst of the 
 		// various anim ticks, before buffers are flipped (which happens in the skel mesh component's post-physics tick)
-		const TArray<FTransform>& SpaceBases = SkeletalMeshComponent->GetEditableComponentSpaceTransforms();
+		const TArrayView<const FTransform> ComponentSpaceTransforms = SkeletalMeshComponent->GetEditableComponentSpaceTransforms();
 
-		if (SpaceBases.IsEmpty())
+		if (ComponentSpaceTransforms.IsEmpty())
 		{
 			// Our tick depends on the skeletal mesh primary tick. When physics is disabled (i.e.
 			// physics blend = 0), this might have been allowed to run post-physics. Thus, if
@@ -286,7 +301,7 @@ void UPhysicalAnimationComponent::UpdateTargetActors(ETeleportType TeleportType)
 		}
 		FPhysicsCommand::ExecuteWrite(SkeletalMeshComponent, [&]()
 		{
-			TArray<FTransform> LocalTransforms = SkeletalMeshComponent->GetBoneSpaceTransforms();
+			const TArrayView<const FTransform> BoneSpaceTransforms = SkeletalMeshComponent->GetBoneSpaceTransformsView();
 			for (int32 DataIdx = 0; DataIdx < DriveData.Num(); ++DataIdx)
 			{
 				const FPhysicalAnimationData& PhysAnimData = DriveData[DataIdx];
@@ -295,10 +310,16 @@ void UPhysicalAnimationComponent::UpdateTargetActors(ETeleportType TeleportType)
 				if (FPhysicsActorHandle TargetActor = InstanceData.TargetActor)
 				{
 					const int32 BoneIdx = RefSkeleton.FindBoneIndex(PhysAnimData.BodyName);
-					if (BoneIdx != INDEX_NONE)	//It's possible the skeletal mesh has changed out from under us. In that case we should probably reset, but at the very least don't do work on non-existent bones
+					// It's possible the skeletal mesh has changed out from under us. In that case
+					// we should probably reset, but at the very least don't do work on non-existent bones
+					if (BoneIdx != INDEX_NONE)
 					{
-						const FTransform TargetTM = ComputeTargetTM(PhysAnimData, *SkeletalMeshComponent, *PhysAsset, LocalTransforms, SpaceBases, BoneIdx);
-						// todo(chaos): Calling SetKinematicTarget_AssumesLocked before SetGlobalPose_AssumesLocked is unnessary for chaos. We should fix this when PhysX is removed.
+						const FTransform TargetTM = ComputeTargetTM(
+							PhysAnimData, *SkeletalMeshComponent, *PhysAsset, 
+							BoneSpaceTransforms, ComponentSpaceTransforms, BoneIdx);
+						// todo(chaos): Calling SetKinematicTarget_AssumesLocked before
+						// SetGlobalPose_AssumesLocked is unnecessary for chaos. We should fix this
+						// when PhysX is removed.
 						FPhysicsInterface::SetKinematicTarget_AssumesLocked(TargetActor, TargetTM);
 
 						if (TeleportType == ETeleportType::TeleportPhysics)
