Removed code that was being overly aggressive with checking for material usage that has led to many materials having MATUSAGE_ParticleMeshes flagged incorrectly.  We already check the flag on the materials being rendered.

diff --git a/Engine/Source/Runtime/Engine/Private/Particles/ParticleEmitterInstances.cpp b/Engine/Source/Runtime/Engine/Private/Particles/ParticleEmitterInstances.cpp
--- a/Engine/Source/Runtime/Engine/Private/Particles/ParticleEmitterInstances.cpp
+++ b/Engine/Source/Runtime/Engine/Private/Particles/ParticleEmitterInstances.cpp
@@ -3909,15 +3909,8 @@ bool FParticleMeshEmitterInstance::FillReplayData( FDynamicEmitterReplayDataBase
 	FDynamicMeshEmitterReplayData* NewReplayData =
 		static_cast< FDynamicMeshEmitterReplayData* >( &OutData );
 
-	UMaterialInterface* RenderMaterial = CurrentMaterial;
-	if (RenderMaterial == NULL  || (RenderMaterial->CheckMaterialUsage_Concurrent(MATUSAGE_MeshParticles) == false))
-	{
-		RenderMaterial = UMaterial::GetDefaultMaterial(MD_Surface);
-	}
-	NewReplayData->MaterialInterface = RenderMaterial;
-	CurrentMaterial = RenderMaterial;
-
 	// Mesh settings
+	NewReplayData->MaterialInterface = nullptr;
 	NewReplayData->bScaleUV = LODLevel->RequiredModule->bScaleUV;
 	NewReplayData->SubUVInterpMethod = LODLevel->RequiredModule->InterpolationMethod;
 	NewReplayData->SubUVDataOffset = SubUVDataOffset;
