Basic support for transactional latent actions
- This allows latent actions to run transactionally when compiled with remote objects.
- Note this feature needs more work.  For example, ideally latent actions will migrate with their counterpart actor
- Also, the processing should be changed to allow progress to be made while blocked on migrating dependencies, similar to ticking

diff --git a/Engine/Source/Runtime/Engine/Private/LatentActionManager.cpp b/Engine/Source/Runtime/Engine/Private/LatentActionManager.cpp
--- a/Engine/Source/Runtime/Engine/Private/LatentActionManager.cpp
+++ b/Engine/Source/Runtime/Engine/Private/LatentActionManager.cpp
@@ -177,6 +177,8 @@ void FLatentActionManager::ProcessLatentActions(UObject* InObject, float DeltaTi
 		SCOPE_CYCLE_COUNTER(STAT_RemoveLatentActions);
 		for (FActionsForObject::TIterator It(ActionsToRemoveMap); It; ++It)
 		{
+			auto LatentActionWork = [this, &It]()
+			{
 			FObjectActions* ObjectActions = GetActionsForObject(It->Key);
 			TSharedPtr<TArray<FUuidAndAction>> ActionToRemoveListPtr = It->Value;
 			if (ActionToRemoveListPtr.IsValid() && ObjectActions)
@@ -195,6 +197,21 @@ void FLatentActionManager::ProcessLatentActions(UObject* InObject, float DeltaTi
 				// Notify listeners that latent actions for this object were removed
 				LatentActionsChangedDelegate.Broadcast(It->Key.Get(), ELatentActionChangeType::ActionsRemoved);
 			}
+			};
+
+#if UE_WITH_REMOTE_OBJECT_HANDLE
+			if (!AutoRTFM::IsClosed())
+			{
+				static FName TransactionalWorkName("FLatentActionManager::ProcessLatentActions_Remove");
+				UE::RemoteExecutor::ExecuteTransactional(TransactionalWorkName, [LatentActionWork]{
+						LatentActionWork();
+					});
+			}
+			else
+#endif
+			{
+				LatentActionWork();
+			}
 		}
 		ActionsToRemoveMap.Reset();
 	}
@@ -220,6 +237,8 @@ void FLatentActionManager::ProcessLatentActions(UObject* InObject, float DeltaTi
 		SCOPE_CYCLE_COUNTER(STAT_TickLatentActions);
 		for (FObjectToActionListMap::TIterator ObjIt(ObjectToActionListMap); ObjIt; ++ObjIt)
 		{	
+			auto LatentActionWork = [this, DeltaTime, &ObjIt]()
+			{
 			UObject* Object = ObjIt.Key().Get();
 			FObjectActions* ObjectActions = ObjIt.Value().Get();
 			check(ObjectActions);
@@ -257,6 +276,21 @@ void FLatentActionManager::ProcessLatentActions(UObject* InObject, float DeltaTi
 			{
 				ObjIt.RemoveCurrent();
 			}
+			};
+
+#if UE_WITH_REMOTE_OBJECT_HANDLE
+			if (!AutoRTFM::IsClosed())
+			{
+				static FName TransactionalWorkName("FLatentActionManager::ProcessLatentActions");
+				UE::RemoteExecutor::ExecuteTransactional(TransactionalWorkName, [LatentActionWork]{
+						LatentActionWork();
+					});
+			}
+			else
+#endif
+			{
+				LatentActionWork();
+			}
 		}
 	}
 
