SoundWaveProxyPlayer

Fix Simulate Seek Forward to properly update its buffers

check SoundWave for updates when playing a new sound

#rnx

diff --git a/Engine/Source/Runtime/Engine/Private/SoundWaveProxyPlayer.cpp b/Engine/Source/Runtime/Engine/Private/SoundWaveProxyPlayer.cpp
--- a/Engine/Source/Runtime/Engine/Private/SoundWaveProxyPlayer.cpp
+++ b/Engine/Source/Runtime/Engine/Private/SoundWaveProxyPlayer.cpp
@@ -46,6 +46,9 @@ bool FSoundWaveProxyPlayer::SetSoundWave(FSoundWaveProxyPtr InWaveProxy)
 		Reset(Settings);
 		return false;
 	}
+
+	// make sure we have the most up-to-date data for the Sound Wave
+	InWaveProxy->TryGetLatest();
 	
 	if (!IsValidSoundWaveProxy(InWaveProxy.ToSharedRef()))
 	{
@@ -472,8 +475,9 @@ int32 FSoundWaveProxyPlayer::SimulateSeekForward(int32 NumFrames)
 {
 	int32 NumFramesRemaining = NumFrames;
 	int32 SampleOffset = 0;
-	uint32 MaxBufferSize = NumFrames * NumChannels * DecodeBuffer.GetTypeSize();
-	uint32 BufferSize = FMath::Min(DecodeBuffer.Num() * DecodeBuffer.GetTypeSize(), MaxBufferSize);
+	const uint32 MaxBufferSize = DecodeBuffer.Num() * DecodeBuffer.GetTypeSize();
+	uint32 DesiredBufferSize = NumFrames * NumChannels * DecodeBuffer.GetTypeSize();
+	uint32 BufferSize = FMath::Min(DesiredBufferSize, MaxBufferSize);
 	uint8* ByteBuffer = (uint8*)DecodeBuffer.GetData();
 
 	bool bFinished = false;
@@ -500,7 +504,8 @@ int32 FSoundWaveProxyPlayer::SimulateSeekForward(int32 NumFrames)
 
 		SampleOffset += NumSamplesStreamed;
 		NumFramesRemaining -= NumFramesStreamed;
-		BufferSize -= NumBytesStreamed;
+		DesiredBufferSize -= NumBytesStreamed;
+		BufferSize = FMath::Min(DesiredBufferSize, MaxBufferSize);
 	}
 
 	int32 NumFramesGenerated = NumFrames - NumFramesRemaining;
