Added Component Providers in Component Interface
Use TComponentInterfaceProvider<type>::Provides(FromObj) to acquire a component interface from a potential uobject providing those interfaces

diff --git a/Engine/Source/Runtime/Engine/Classes/Components/ComponentInterfaceIterator.h b/Engine/Source/Runtime/Engine/Classes/Components/ComponentInterfaceIterator.h
--- a/Engine/Source/Runtime/Engine/Classes/Components/ComponentInterfaceIterator.h
+++ b/Engine/Source/Runtime/Engine/Classes/Components/ComponentInterfaceIterator.h
@@ -20,7 +20,7 @@ public:
 	{
 		TArray<UObject*>	ObjectArray;
 
-		for (TComponentInterfaceImplementation<T>& Implementation : T::Implementers)
+		for (TComponentInterfaceImplementation<T>& Implementation : T::GetImplementers())
 		{
 			if (Implementation.Class)
 			{
diff --git a/Engine/Source/Runtime/Engine/Classes/Components/ComponentInterfaces.h b/Engine/Source/Runtime/Engine/Classes/Components/ComponentInterfaces.h
--- a/Engine/Source/Runtime/Engine/Classes/Components/ComponentInterfaces.h
+++ b/Engine/Source/Runtime/Engine/Classes/Components/ComponentInterfaces.h
@@ -5,6 +5,8 @@
 #include "PrimitiveComponentId.h"
 #include "Templates/RefCounting.h"
 #include "Templates/Function.h"
+#include "UObject/Object.h"
+#include "Misc/Iteration.h"
 
 class ULevel;
 class UWorld;
@@ -55,6 +57,31 @@ struct FPrimitiveStats
 	TArray<FPrimitiveLODStats> LODStats;
 };
 
+// A ComponentInterface Provider bridges an owner/container/implementer object to a ComponentInterface, many providers are allowed to exist for the same interface instance
+template<typename T>
+struct TComponentInterfaceProvider
+{
+	TComponentInterfaceProvider(UClass* InClass, TFunction<T* (UObject*)> InProvider)
+		: Class(InClass)
+		, Provider(InProvider)
+	{}
+
+	TComponentInterfaceProvider(UClass* InClass, TFunction<TArray<T*>(UObject*)> InCollectionProvider)
+		: Class(InClass)
+		, CollectionProvider(InCollectionProvider)
+	{}
+
+
+	UClass* Class = nullptr;
+	TFunction<T* (UObject*)> Provider;
+	TFunction<TArray<T*> (UObject*)> CollectionProvider;
+
+	static ENGINE_API T* Provides(UObject* SourceObject);
+	static ENGINE_API void Provides(UObject* SourceObject, TArray<T*>& OutResults);
+	static void Provides(UObject* SourceObject, TFunctionRef<EIteration (T*)>);
+};
+
+// A ComponentInterface Implementation is used to enumerate the various existing ComponentInterfaces, only one Implementation per interface instance is allowed
 template<typename T>
 struct TComponentInterfaceImplementation
 {
@@ -74,7 +101,28 @@ struct TComponentInterfaceImplementation
 };
 typedef TComponentInterfaceImplementation<void> FComponentInterfaceImplementation;
 
-class IPrimitiveComponent
+template<class T>
+class TComponentInterfaceIterator;
+	
+
+template<class T> 
+class TComponentInterfaceBase
+{
+public:
+	static ENGINE_API void AddImplementer(const TComponentInterfaceImplementation<T>& Implementer);
+	static ENGINE_API void RemoveImplementer(const UClass* ImplementerClass);
+	static ENGINE_API void AddProvider(const TComponentInterfaceProvider<T>& Provider);
+	static ENGINE_API void RemoveProvider(const UClass* ProviderClass);
+
+protected:
+	
+	static ENGINE_API TArray<TComponentInterfaceImplementation<T>>& GetImplementers();
+
+	friend TComponentInterfaceIterator<T>;
+	friend TComponentInterfaceProvider<T>;
+};
+
+class IPrimitiveComponent : public TComponentInterfaceBase<IPrimitiveComponent>
 {
 public:
 	virtual bool IsRenderStateCreated() const = 0;
@@ -141,22 +189,9 @@ public:
 	virtual HHitProxy* CreateMeshHitProxy(int32 SectionIndex, int32 MaterialIndex) = 0;
 #endif
 	virtual HHitProxy* CreatePrimitiveHitProxies(TArray<TRefCountPtr<HHitProxy> >& OutHitProxies) = 0;
-	
-	ENGINE_API static void AddImplementer(const TComponentInterfaceImplementation<IPrimitiveComponent>& Implementer);
-	ENGINE_API static void RemoveImplementer(const UClass* ImplementerClass);
-
-	UE_DEPRECATED(5.8, "Use templated version")
-	ENGINE_API static void AddImplementer(const FComponentInterfaceImplementation& Implementer);
-
-protected:
-
-	template<class T> 
-	friend class TComponentInterfaceIterator;
-
-	ENGINE_API static TArray<TComponentInterfaceImplementation<IPrimitiveComponent>> Implementers;
 };
 
-class IStaticMeshComponent
+class IStaticMeshComponent : public TComponentInterfaceBase<IStaticMeshComponent>
 {
 public:
 #if WITH_EDITOR
@@ -166,29 +201,15 @@ public:
 #endif
 	virtual UStaticMesh* GetStaticMesh() const = 0;
 
-	virtual IPrimitiveComponent* GetPrimitiveComponentInterface() = 0;		
+	virtual IPrimitiveComponent* GetPrimitiveComponentInterface() = 0;
 	inline const IPrimitiveComponent* GetPrimitiveComponentInterface() const
 	{
 		// use the non-const version and return it as a const object to avoid duplicating the code in implementers
 		return (const_cast<IStaticMeshComponent*>(this))->GetPrimitiveComponentInterface();
 	}
-	
-	ENGINE_API static void AddImplementer(const TComponentInterfaceImplementation<IStaticMeshComponent>& Implementer);
-	ENGINE_API static void RemoveImplementer(const UClass* ImplementerClass);
-
-	UE_DEPRECATED(5.8, "Use templated version")
-	ENGINE_API static void AddImplementer(const FComponentInterfaceImplementation& Implementer);
-
-protected:
-
-	template<class T> 
-	friend class TComponentInterfaceIterator;
-
-	ENGINE_API static TArray<TComponentInterfaceImplementation<IStaticMeshComponent>> Implementers;
-	
 };
 
-class ISkinnedMeshComponent
+class ISkinnedMeshComponent : public TComponentInterfaceBase<ISkinnedMeshComponent>
 {
 public:
 #if WITH_EDITOR
@@ -203,19 +224,9 @@ public:
 		// use the non-const version and return it as a const object to avoid duplicating the code in implementers
 		return (const_cast<ISkinnedMeshComponent*>(this))->GetPrimitiveComponentInterface();
 	}
-	
-	ENGINE_API static void AddImplementer(const TComponentInterfaceImplementation<ISkinnedMeshComponent>& Implementer);
-	ENGINE_API static void RemoveImplementer(const UClass* ImplementerClass);
-
-	UE_DEPRECATED(5.8, "Use templated version")
-	ENGINE_API static void AddImplementer(const FComponentInterfaceImplementation& Implementer);
-protected:
-	template<class T> 
-	friend class TComponentInterfaceIterator;
-
-	ENGINE_API static TArray<TComponentInterfaceImplementation<ISkinnedMeshComponent>> Implementers;
 };
 
+
 #pragma region HelperMacros
 
 // These macros are intended to allow implementing an interface with same memory footprint/performance as inheriting from 
diff --git a/Engine/Source/Runtime/Engine/Private/Components/ComponentInterfaces.cpp b/Engine/Source/Runtime/Engine/Private/Components/ComponentInterfaces.cpp
--- a/Engine/Source/Runtime/Engine/Private/Components/ComponentInterfaces.cpp
+++ b/Engine/Source/Runtime/Engine/Private/Components/ComponentInterfaces.cpp
@@ -5,21 +5,24 @@
 #include "HAL/IConsoleManager.h"
 #include "UnrealEngine.h"
 
-TArray<TComponentInterfaceImplementation<IPrimitiveComponent>> IPrimitiveComponent::Implementers;
 
-void IPrimitiveComponent::AddImplementer(const TComponentInterfaceImplementation<IPrimitiveComponent>& Implementer)
+template<class T> 
+struct TComponentInterfaceBaseData
 {
-	Implementers.Add(Implementer);
-}
+	static inline TArray<TComponentInterfaceImplementation<T>> Implementers;
+	static inline TArray<TComponentInterfaceProvider<T>> Providers;
+};
 
-void IPrimitiveComponent::AddImplementer(const FComponentInterfaceImplementation& Implementer)
-{
-	AddImplementer(reinterpret_cast<const TComponentInterfaceImplementation<IPrimitiveComponent>&>(Implementer));
+template<class T>
+void TComponentInterfaceBase<T>::AddImplementer(const TComponentInterfaceImplementation<T>& Implementer)
+{	
+	TComponentInterfaceBaseData<T>::Implementers.Add(Implementer);
 }
 
-void IPrimitiveComponent::RemoveImplementer(const UClass* ImplementerClass)
-{
-	for (TArray<TComponentInterfaceImplementation<IPrimitiveComponent>>::TIterator ImplementerIt(Implementers); ImplementerIt; ++ImplementerIt)
+template<class T>
+void TComponentInterfaceBase<T>::RemoveImplementer(const UClass* ImplementerClass)
+{	
+	for (typename TArray<TComponentInterfaceImplementation<T>>::TIterator ImplementerIt(TComponentInterfaceBaseData<T>::Implementers); ImplementerIt; ++ImplementerIt)
 	{
 		if (ImplementerIt->Class == ImplementerClass)
 		{
@@ -28,48 +31,145 @@ void IPrimitiveComponent::RemoveImplementer(const UClass* ImplementerClass)
 	}
 }
 
-TArray<TComponentInterfaceImplementation<IStaticMeshComponent>> IStaticMeshComponent::Implementers;
-
-void IStaticMeshComponent::AddImplementer(const TComponentInterfaceImplementation<IStaticMeshComponent>& Implementer)
+template<class T>
+TArray<TComponentInterfaceImplementation<T>>& TComponentInterfaceBase<T>::GetImplementers()
 {
-	Implementers.Add(Implementer);
+	return TComponentInterfaceBaseData<T>::Implementers;
 }
 
-void IStaticMeshComponent::AddImplementer(const FComponentInterfaceImplementation& Implementer)
+template<class T>
+void TComponentInterfaceBase<T>::AddProvider(const TComponentInterfaceProvider<T>& Provider)
 {
-	AddImplementer(reinterpret_cast<const TComponentInterfaceImplementation<IStaticMeshComponent>&>(Implementer));
+	TComponentInterfaceBaseData<T>::Providers.Add(Provider);
 }
 
-void IStaticMeshComponent::RemoveImplementer(const UClass* ImplementerClass)
+template<class T>
+void TComponentInterfaceBase<T>::RemoveProvider(const UClass* ProviderClass)
 {
-	for (TArray<TComponentInterfaceImplementation<IStaticMeshComponent>>::TIterator ImplementerIt(Implementers); ImplementerIt; ++ImplementerIt)
+	for (typename TArray<TComponentInterfaceProvider<T>>::TIterator ProviderIt(TComponentInterfaceBaseData<T>::Providers); ProviderIt; ++ProviderIt)
 	{
-		if (ImplementerIt->Class == ImplementerClass)
+		if (ProviderIt->Class == ProviderClass)
 		{
-			ImplementerIt.RemoveCurrent();
+			ProviderIt.RemoveCurrent();
 		}
 	}
 }
 
-TArray<TComponentInterfaceImplementation<ISkinnedMeshComponent>> ISkinnedMeshComponent::Implementers;
-
-void ISkinnedMeshComponent::AddImplementer(const TComponentInterfaceImplementation<ISkinnedMeshComponent>& Implementer)
+template<class T> 
+T* TComponentInterfaceProvider<T>::Provides(UObject* SourceObject)
 {
-	Implementers.Add(Implementer);
+	T* Result = nullptr;
+	Provides(SourceObject, [&Result](T* Interface) { Result = Interface; return EIteration::Break; });
+	
+	return Result;
 }
 
-void ISkinnedMeshComponent::AddImplementer(const FComponentInterfaceImplementation& Implementer)
+template<class T> 
+void TComponentInterfaceProvider<T>::Provides(UObject* SourceObject, TArray<T*>& OutResults)
 {
-	AddImplementer(reinterpret_cast<const TComponentInterfaceImplementation<ISkinnedMeshComponent>&>(Implementer));
+	Provides(SourceObject, [&OutResults](T* Interface) { OutResults.Add(Interface); return EIteration::Continue; });
 }
 
-void ISkinnedMeshComponent::RemoveImplementer(const UClass* ImplementerClass)
+template<class T> 
+void TComponentInterfaceProvider<T>::Provides(UObject* SourceObject, TFunctionRef<EIteration(T*)> ProvidesFn)
 {
-	for (TArray<TComponentInterfaceImplementation<ISkinnedMeshComponent>>::TIterator ImplementerIt(Implementers); ImplementerIt; ++ImplementerIt)
+	if(!SourceObject)
 	{
-		if (ImplementerIt->Class == ImplementerClass)
+		return;
+	}
+
+	// This could be accelerated by caching all classes in a TMap<> of derived classes of the Provider class to the FComponentInterfaceProvider* 
+	for (TComponentInterfaceProvider<T>& Provider : TComponentInterfaceBaseData<T>::Providers)
+	{
+		if (Provider.Class)
 		{
-			ImplementerIt.RemoveCurrent();
+			if (SourceObject->IsA(Provider.Class))
+			{
+				if (Provider.Provider)
+				{
+					// If we want to specialization of the provider by derived classes ensure providers are sorted with derived classes first 
+					if (T* Interface = Provider.Provider(SourceObject))
+					{
+						if (ProvidesFn(Interface) == EIteration::Continue)
+						{
+						}
+						else
+						{
+							return;
+						}
+					}
+				}
+				else if (Provider.CollectionProvider)
+				{
+					TArray<T*> InterfaceCollection = Provider.CollectionProvider(SourceObject);
+				
+					for (T* Interface : InterfaceCollection)
+					{
+						if (Interface)
+						{
+							if (ProvidesFn(Interface) == EIteration::Continue)
+							{
+							}
+							else
+							{
+								return;
+							}
+						}
+					}
+				}
+			}
 		}
 	}
+	
+	// We also consider interface implementers as providers
+	for (TComponentInterfaceImplementation<T>& Implementation : TComponentInterfaceBaseData<T>::Implementers)
+	{
+		// Collection resolvers potentially provide > 1 object, and need a different interface
+		if (Implementation.Class)
+		{
+			if (SourceObject->IsA(Implementation.Class))
+			{
+				if (Implementation.Resolver)
+				{
+					if (T* Interface = Implementation.Resolver(SourceObject))
+					{
+						if (ProvidesFn(Interface) == EIteration::Continue)
+						{
+						}
+						else
+						{
+							return;
+						}
+					}
+				}
+				else if (Implementation.CollectionResolver)
+				{
+					TArray<T*> InterfaceCollection = Implementation.CollectionResolver(SourceObject);
+				
+					for (T* Interface : InterfaceCollection)
+					{
+						if (Interface)
+						{
+							if (ProvidesFn(Interface) == EIteration::Continue)
+							{
+							}
+							else
+							{
+								return;
+							}
+						}
+					}
+				}
+			}
+		}
+	}	
 }
+
+#define UE_COMPONENT_INTERFACE_INSTANTIATION(Interface)\
+	template struct TComponentInterfaceBaseData<Interface>;\
+	template class TComponentInterfaceBase<Interface>;\
+	template struct TComponentInterfaceProvider<Interface>;
+
+UE_COMPONENT_INTERFACE_INSTANTIATION(IPrimitiveComponent)
+UE_COMPONENT_INTERFACE_INSTANTIATION(IStaticMeshComponent)
+UE_COMPONENT_INTERFACE_INSTANTIATION(ISkinnedMeshComponent)
\ No newline at end of file
diff --git a/Engine/Source/Runtime/Engine/Private/UnrealEngine.cpp b/Engine/Source/Runtime/Engine/Private/UnrealEngine.cpp
--- a/Engine/Source/Runtime/Engine/Private/UnrealEngine.cpp
+++ b/Engine/Source/Runtime/Engine/Private/UnrealEngine.cpp
@@ -343,6 +343,18 @@ void FEngineModule::StartupModule()
 	IStaticMeshComponent::AddImplementer(TComponentInterfaceImplementation<IStaticMeshComponent>{UStaticMeshComponent::StaticClass(), [](UObject* Obj){return Cast<UStaticMeshComponent>(Obj)->GetStaticMeshComponentInterface();}});
 	ISkinnedMeshComponent::AddImplementer(TComponentInterfaceImplementation<ISkinnedMeshComponent>{USkinnedMeshComponent::StaticClass(), [](UObject* Obj) {return Cast<USkinnedMeshComponent>(Obj)->GetSkinnedMeshComponentInterface();}});
 
+	IPrimitiveComponent::AddProvider({AActor::StaticClass(),  [](UObject* Obj) -> IPrimitiveComponent*
+	{
+		if (AActor* Actor = Cast<AActor>(Obj))
+		{
+			if (UPrimitiveComponent* Primitive = Actor->GetComponentByClass<UPrimitiveComponent>())
+			{
+				return Primitive->GetPrimitiveComponentInterface();
+			}
+		}
+		return nullptr;
+	}});
+
 #if OBJECT_TRACE_ENABLED
 	FObjectTrace::Init();
 #endif
@@ -398,6 +410,8 @@ void FEngineModule::ShutdownModule()
 {
 	IPrimitiveComponent::RemoveImplementer(UPrimitiveComponent::StaticClass());
 	IStaticMeshComponent::RemoveImplementer(UStaticMeshComponent::StaticClass());
+	ISkinnedMeshComponent::RemoveImplementer(USkinnedMeshComponent::StaticClass());
+	IPrimitiveComponent::RemoveProvider(AActor::StaticClass());
 
 	UE::Net::ShutdownEngineNetObjectFactories();
 
