Add NotifyInstanceID to FAnimNotifyEventReference

diff --git a/Engine/Source/Runtime/Engine/Private/Animation/AnimInstance.cpp b/Engine/Source/Runtime/Engine/Private/Animation/AnimInstance.cpp
--- a/Engine/Source/Runtime/Engine/Private/Animation/AnimInstance.cpp
+++ b/Engine/Source/Runtime/Engine/Private/Animation/AnimInstance.cpp
@@ -117,6 +117,23 @@ namespace MontageCVars
 		TEXT("Whether we should flush all completed montages IMMEDIATELY when a new montage stomps the group. Use this to prevent accumulating montages when animation tick is paused, and firing all at once when we unpause."));
 }
 
+namespace UE::Anim::Private
+{
+	static int32 NextNotifyInstanceID = 0;
+	static int32 GetNextNotifyInstanceID()
+	{
+		if (NextNotifyInstanceID == TNumericLimits<int32>::Max())
+		{
+			NextNotifyInstanceID = 0;
+			return NextNotifyInstanceID;
+		}
+		else
+		{
+			return NextNotifyInstanceID++;
+		}
+	}
+}
+
 /** Used to determine notify dispatch behaviour when sources (montage, anim graph) may tick at different times within a frame. */
 enum class EAnimNotifyDispatchPolicy : int32
 {
@@ -1755,11 +1772,17 @@ void UAnimInstance::TriggerAnimNotifies(float DeltaSeconds, ENotifyTriggerMode M
 				if (const bool bIsNotifyStateAlreadyActive = ActiveAnimNotifyState.Find(*AnimNotifyEvent, ExistingItemIndex))
 				{
 					check(ActiveAnimNotifyState.Num() == ActiveAnimNotifyEventReference.Num());
+
+					// Notify already exists, so we should re-use the same Instance ID.
+					NotifyQueue.AnimNotifies[Index].SetNotifyInstanceID(ActiveAnimNotifyEventReference[ExistingItemIndex].GetNotifyInstanceID());
 					ActiveAnimNotifyState.RemoveAtSwap(ExistingItemIndex, EAllowShrinking::No);
 					ActiveAnimNotifyEventReference.RemoveAtSwap(ExistingItemIndex, EAllowShrinking::No);
 				}
 				else
 				{
+					// New notify. Generate unique ID for this instance.
+					NotifyQueue.AnimNotifies[Index].SetNotifyInstanceID(UE::Anim::Private::GetNextNotifyInstanceID());
+
 					NotifyStateBeginEvent.Add(AnimNotifyEvent);
 					NotifyStateBeginEventReference.Add(&NotifyQueue.AnimNotifies[Index]);
 				}
@@ -1891,6 +1914,7 @@ void UAnimInstance::TriggerSingleAnimNotify(FAnimNotifyEventReference& EventRefe
 			{
 				// Implemented notify: just call Notify. UAnimNotify will forward this to the event which will do the work.
 				TRACE_ANIM_NOTIFY(this, *AnimNotifyEvent, Event);
+				EventReference.SetNotifyInstanceID(UE::Anim::Private::GetNextNotifyInstanceID());
 				AnimNotifyEvent->Notify->Notify(GetSkelMeshComponent(), Cast<UAnimSequenceBase>(AnimNotifyEvent->Notify->GetOuter()), EventReference);
 			}
 		}
diff --git a/Engine/Source/Runtime/Engine/Private/Animation/AnimNotifyLibrary.cpp b/Engine/Source/Runtime/Engine/Private/Animation/AnimNotifyLibrary.cpp
--- a/Engine/Source/Runtime/Engine/Private/Animation/AnimNotifyLibrary.cpp
+++ b/Engine/Source/Runtime/Engine/Private/Animation/AnimNotifyLibrary.cpp
@@ -20,6 +20,16 @@ bool UAnimNotifyLibrary::NotifyStateReachedEnd(const FAnimNotifyEventReference&
 	return false;
 }
 
+float UAnimNotifyLibrary::GetNotifyStartTime(const FAnimNotifyEventReference& EventReference)
+{
+	if (const FAnimNotifyEvent* NotifyEvent = EventReference.GetNotify())
+	{
+		return NotifyEvent->GetTriggerTime();
+	}
+
+	return 0.0f;
+}
+
 float UAnimNotifyLibrary::GetCurrentAnimationTime(const FAnimNotifyEventReference& EventReference)
 {
 	return EventReference.GetCurrentAnimationTime();
diff --git a/Engine/Source/Runtime/Engine/Public/Animation/AnimNotifyLibrary.h b/Engine/Source/Runtime/Engine/Public/Animation/AnimNotifyLibrary.h
--- a/Engine/Source/Runtime/Engine/Public/Animation/AnimNotifyLibrary.h
+++ b/Engine/Source/Runtime/Engine/Public/Animation/AnimNotifyLibrary.h
@@ -26,6 +26,15 @@ public:
 	UFUNCTION(BlueprintPure, Category = "Utilities|Animation|Notifies" , meta = (ScriptMethod))
 	static ENGINE_API bool NotifyStateReachedEnd(const FAnimNotifyEventReference& EventReference);
 
+	/**
+	* Get the start time for this notify. For a notify this is the trigger time.
+	*
+	* @param EventReference		The event to inspect
+	* @return the time in at which this notify started, relative to the owning asset i.e. play length.
+	*/
+	UFUNCTION(BlueprintPure, Category = "Utilities|Animation|Notifies" , meta = (ScriptMethod))
+	static ENGINE_API float GetNotifyStartTime(const FAnimNotifyEventReference& EventReference);
+
 	/**
 	 * Get the current anim notify time in seconds for when this notify was fired
 	 *
diff --git a/Engine/Source/Runtime/Engine/Public/Animation/AnimNotifyQueue.h b/Engine/Source/Runtime/Engine/Public/Animation/AnimNotifyQueue.h
--- a/Engine/Source/Runtime/Engine/Public/Animation/AnimNotifyQueue.h
+++ b/Engine/Source/Runtime/Engine/Public/Animation/AnimNotifyQueue.h
@@ -115,6 +115,11 @@ struct FAnimNotifyEventReference
 		return bActiveContext;
 	}
 
+	// NotifyInstanceID is generated when we dispatch the NotifyBegin/Notify call.
+	// In the case of a non-unique/existing Notify, we re-use/propagate the ID from the previous notify. 
+	int32 GetNotifyInstanceID() const { return NotifyInstanceID; }
+	void SetNotifyInstanceID(int32 InNewInstanceID) { NotifyInstanceID = InNewInstanceID; }
+
 private:
 	// Context data gleaned from the tick record
 	TSharedPtr<TArray<TUniquePtr<const UE::Anim::IAnimNotifyEventContextDataInterface>>> ContextData;
@@ -131,6 +136,10 @@ private:
 	// The recorded time from the tick record that this notify event was fired at
 	float CurrentAnimTime = 0.0f;
 
+	// Unique per-notify instance ID. This gets populated before firing the notifies and should only be used after when notifies are actually fired.
+	UPROPERTY(transient, Category=Notifies, BlueprintReadOnly, meta=(AllowPrivateAccess="true"))
+	int32 NotifyInstanceID = INDEX_NONE;
+
 	// Whether the context this notify was fired from is active or not (active == not blending out).
 	bool bActiveContext = false;
 };
