Fixed crash when using ISKM's with multiple LOD's due to the scene proxy selecting the GPU lod path when it's not enabled.

diff --git a/Engine/Source/Runtime/Engine/Private/InstancedSkinnedMeshSceneProxy.h b/Engine/Source/Runtime/Engine/Private/InstancedSkinnedMeshSceneProxy.h
--- a/Engine/Source/Runtime/Engine/Private/InstancedSkinnedMeshSceneProxy.h
+++ b/Engine/Source/Runtime/Engine/Private/InstancedSkinnedMeshSceneProxy.h
@@ -150,7 +150,14 @@ public:
 
 	virtual FDesiredLODLevel GetDesiredLODLevel_RenderThread(const FSceneView* View) const override
 	{
-		return FDesiredLODLevel::CreateFirst(GetCurrentFirstLODIdx_Internal());
+		if (bUseGpuLodSelection)
+		{
+			return FDesiredLODLevel::CreateFirst(GetCurrentFirstLODIdx_Internal());
+		}
+		else
+		{
+			return Super::GetDesiredLODLevel_RenderThread(View);
+		}
 	}
 
 private:
