Removed checkNoEntry to avoid crash when ray tracing is enabled.

diff --git a/Engine/Source/Runtime/Engine/Private/InstancedSkinnedMeshSceneProxy.h b/Engine/Source/Runtime/Engine/Private/InstancedSkinnedMeshSceneProxy.h
--- a/Engine/Source/Runtime/Engine/Private/InstancedSkinnedMeshSceneProxy.h
+++ b/Engine/Source/Runtime/Engine/Private/InstancedSkinnedMeshSceneProxy.h
@@ -105,7 +105,6 @@ public:
 
 	virtual void GetDynamicRayTracingInstances(FRayTracingInstanceCollector& Collector) override
 	{
-		checkNoEntry();
 	}
 #endif
 
