Add UE_LOG_CONTEXT to animation evaluation for skeletal mesh component and anim blueprint

diff --git a/Engine/Source/Runtime/Engine/Private/Animation/AnimInstanceProxy.cpp b/Engine/Source/Runtime/Engine/Private/Animation/AnimInstanceProxy.cpp
--- a/Engine/Source/Runtime/Engine/Private/Animation/AnimInstanceProxy.cpp
+++ b/Engine/Source/Runtime/Engine/Private/Animation/AnimInstanceProxy.cpp
@@ -1227,6 +1227,7 @@ void FAnimInstanceProxy::UpdateAnimation()
 	FMemMark Mark(FMemStack::Get());
 	FAnimationUpdateSharedContext SharedContext;
 	FAnimationUpdateContext Context(this, CurrentDeltaSeconds, &SharedContext);
+	UE_LOG_CONTEXT("AnimBlueprint", FTopLevelAssetPath(GetAnimBlueprint()));
 
 	if(AnimClassInterface && AnimClassInterface->GetAnimBlueprintFunctions().Num() > 0)
 	{
diff --git a/Engine/Source/Runtime/Engine/Private/Components/SkeletalMeshComponent.cpp b/Engine/Source/Runtime/Engine/Private/Components/SkeletalMeshComponent.cpp
--- a/Engine/Source/Runtime/Engine/Private/Components/SkeletalMeshComponent.cpp
+++ b/Engine/Source/Runtime/Engine/Private/Components/SkeletalMeshComponent.cpp
@@ -347,6 +347,7 @@ public:
 	{
 		if (USkeletalMeshComponent* Comp = SkeletalMeshComponent.Get())
 		{
+			UE_LOG_CONTEXT("SkeletalMeshComponent", FSoftObjectPath(Comp));
 			FScopedFTZFloatMode FTZ;
 
 			FScopeCycleCounterUObject ContextScope(Comp);
