FNewSpline: Fix issues building RotationCurve

diff --git a/Engine/Source/Runtime/Engine/Private/Curves/Spline.cpp b/Engine/Source/Runtime/Engine/Private/Curves/Spline.cpp
--- a/Engine/Source/Runtime/Engine/Private/Curves/Spline.cpp
+++ b/Engine/Source/Runtime/Engine/Private/Curves/Spline.cpp
@@ -2620,7 +2620,9 @@ void FNewSpline::RebuildLegacyCurves() const
     PositionCurve.bIsLooped = IsClosedLoop();
     PositionCurve.LoopKeyOffset = 1.f;
     RotationCurve.bIsLooped = IsClosedLoop();
+	RotationCurve.LoopKeyOffset = 1.f;
     ScaleCurve.bIsLooped = IsClosedLoop();
+	ScaleCurve.LoopKeyOffset = 1.f;
 	
     for (int32 i = 0; i < GetSpline().GetNumPoints(); ++i)
     {
@@ -2629,8 +2631,14 @@ void FNewSpline::RebuildLegacyCurves() const
     	// Additionally populate interp curves here so we don't double loop
     	PositionCurve.Points.Emplace(AttributeParameter, GetSpline().GetValue(i), GetSpline().GetTangentIn(i), GetSpline().GetTangentOut(i), ConvertTangentModeToInterpCurveMode(GetSpline().GetTangentModes()[i]));
     	RotationCurve.AddPoint(AttributeParameter, GetAttributeValue<FQuat>(RotationAttrName, i));
+    	RotationCurve.Points.Last().InterpMode = CIM_CurveAuto;
     	ScaleCurve.AddPoint(AttributeParameter, GetAttributeValue<FVector>(ScaleAttrName, i));
+    	ScaleCurve.Points.Last().InterpMode = CIM_CurveAuto;
     }
+
+	PositionCurve.AutoSetTangents(0.f, false);
+	RotationCurve.AutoSetTangents(0.f, false);
+	ScaleCurve.AutoSetTangents(0.f, false);
     
     LegacyCurvesVersion = LegacyCurvesNextVersion.GetValue();
 }
