Fix use of editor-only function in log context

diff --git a/Engine/Source/Runtime/Engine/Private/Animation/AnimInstanceProxy.cpp b/Engine/Source/Runtime/Engine/Private/Animation/AnimInstanceProxy.cpp
--- a/Engine/Source/Runtime/Engine/Private/Animation/AnimInstanceProxy.cpp
+++ b/Engine/Source/Runtime/Engine/Private/Animation/AnimInstanceProxy.cpp
@@ -1227,7 +1227,7 @@ void FAnimInstanceProxy::UpdateAnimation()
 	FMemMark Mark(FMemStack::Get());
 	FAnimationUpdateSharedContext SharedContext;
 	FAnimationUpdateContext Context(this, CurrentDeltaSeconds, &SharedContext);
-	UE_LOG_CONTEXT("AnimBlueprint", FTopLevelAssetPath(GetAnimBlueprint()));
+	UE_LOG_CONTEXT("AnimBlueprint", FTopLevelAssetPath(GetAnimInstanceObject() ? GetAnimInstanceObject()->GetClass() : nullptr));
 
 	if(AnimClassInterface && AnimClassInterface->GetAnimBlueprintFunctions().Num() > 0)
 	{
