Update GetPSOPrecacheVertexFetchElements to address D3DDebug spam

diff --git a/Engine/Source/Runtime/Engine/Private/LocalVertexFactory.cpp b/Engine/Source/Runtime/Engine/Private/LocalVertexFactory.cpp
--- a/Engine/Source/Runtime/Engine/Private/LocalVertexFactory.cpp
+++ b/Engine/Source/Runtime/Engine/Private/LocalVertexFactory.cpp
@@ -346,16 +346,16 @@ void FLocalVertexFactory::GetPSOPrecacheVertexFetchElements(EVertexInputStreamTy
 			Elements.Add(FVertexElement(1, 0, VET_UInt, 13, sizeof(uint32), true));
 			if (!RHISupportsManualVertexFetch(GMaxRHIShaderPlatform))
 			{
-				Elements.Add(FVertexElement(2, 0, VET_Float3, 1, 36, false));
-				Elements.Add(FVertexElement(2, 12, VET_Float3, 2, 36, false));
-				Elements.Add(FVertexElement(2, 24, VET_Float3, 3, 36, false));
-				Elements.Add(FVertexElement(3, 0, VET_Float3, 4, 48, false));
-				Elements.Add(FVertexElement(3, 12, VET_Float3, 5, 48, false));
-				Elements.Add(FVertexElement(3, 24, VET_Float3, 6, 48, false));
-				Elements.Add(FVertexElement(3, 36, VET_Float3, 7, 48, false));
+				Elements.Add(FVertexElement(2, 0, VET_PackedNormal, 1, 8, false));
+				Elements.Add(FVertexElement(2, 4, VET_PackedNormal, 2, 8, false));
+				Elements.Add(FVertexElement(3, 0, VET_Color, 3, 4, false));
+				Elements.Add(FVertexElement(4, 0, VET_Half4, 4, 32, false));
+				Elements.Add(FVertexElement(4, 8, VET_Half4, 5, 32, false));
+				Elements.Add(FVertexElement(4, 16, VET_Half4, 6, 32, false));
+				Elements.Add(FVertexElement(4, 24, VET_Half4, 7, 32, false));
 				if (IsGPUSkinPassThroughSupported(GMaxRHIShaderPlatform))
 				{
-					Elements.Add(FVertexElement(4, 0, VET_Float3, 14, sizeof(float)*3u, false));
+					Elements.Add(FVertexElement(5, 0, VET_Float3, 14, sizeof(float)*3u, false));
 				}
 			}
 			break;
