Fixed so includes are placed properly in relation to defines

diff --git a/Engine/Source/Runtime/Engine/Classes/GameFramework/CheatManagerDefines.h b/Engine/Source/Runtime/Engine/Classes/GameFramework/CheatManagerDefines.h
--- a/Engine/Source/Runtime/Engine/Classes/GameFramework/CheatManagerDefines.h
+++ b/Engine/Source/Runtime/Engine/Classes/GameFramework/CheatManagerDefines.h
@@ -2,6 +2,7 @@
 
 #pragma once
 
+#include "Misc/Build.h"
 
 /** If set to 0, cheat manager functionality will be entirely disabled. If cheats are desired in shipping mode this can be overridden in the target definitions */
 #ifndef UE_WITH_CHEAT_MANAGER
