Remove incorrect lookup on shader compiler output for identifying if a shader uses Anisotropy when Substrate is enabled.

diff --git a/Engine/Source/Runtime/Engine/Private/Materials/MaterialIRValueAnalyzer.cpp b/Engine/Source/Runtime/Engine/Private/Materials/MaterialIRValueAnalyzer.cpp
--- a/Engine/Source/Runtime/Engine/Private/Materials/MaterialIRValueAnalyzer.cpp
+++ b/Engine/Source/Runtime/Engine/Private/Materials/MaterialIRValueAnalyzer.cpp
@@ -381,7 +381,7 @@ static void AnalyzeScreenTexture(FMaterialIRValueAnalyzer& Analyzer, MIR::FScree
 					{
 						bool bHasNormalConnected =
 							Substrate::IsSubstrateEnabled() ?
-							FSubstrateMaterialInfo::HasPropertyConnected(Analyzer.Material->GetCachedExpressionData().PropertyConnectedMask, MP_Normal) :
+							Analyzer.Material->IsPropertyConnected(MP_Normal) :
 							Analyzer.Material->HasNormalConnected();
 
 						if (bHasNormalConnected)
diff --git a/Engine/Source/Runtime/Engine/Private/Materials/MaterialShared.cpp b/Engine/Source/Runtime/Engine/Private/Materials/MaterialShared.cpp
--- a/Engine/Source/Runtime/Engine/Private/Materials/MaterialShared.cpp
+++ b/Engine/Source/Runtime/Engine/Private/Materials/MaterialShared.cpp
@@ -2180,7 +2180,7 @@ bool FMaterialResource::HasMaterialPropertyConnected(EMaterialProperty In) const
 			return Material->HasAmbientOcclusionConnected();
 		}
 		// Substrate material traversal is cached as this is an expensive operation
-		return FSubstrateMaterialInfo::HasPropertyConnected(Material->GetCachedExpressionData().PropertyConnectedMask, In);
+		return Material->IsPropertyConnected(In);
 	}
 	else
 	{
@@ -5802,7 +5802,7 @@ FMaterialShaderParametersBase::FMaterialShaderParametersBase(const FMaterial* In
 		bHasMetallicConnected = InMaterial->HasMaterialPropertyConnected(MP_Metallic);
 		bHasEmissiveColorConnected = InMaterial->HasMaterialPropertyConnected(MP_EmissiveColor);
 		bHasAmbientOcclusionConnected = InMaterial->HasMaterialPropertyConnected(MP_AmbientOcclusion);
-		bHasAnisotropyConnected = InMaterial->HasMaterialPropertyConnected(MP_Anisotropy) || EnumHasAnyFlags(InMaterial->MaterialGetSubstrateMaterialBsdfFeatures_GameThread(), ESubstrateBsdfFeature::Anisotropy);
+		bHasAnisotropyConnected = InMaterial->HasMaterialPropertyConnected(MP_Anisotropy);
 	}
 	else
 	{
