PoseAsset: Avoid copy when retrieving curve values

#rnx

diff --git a/Engine/Source/Runtime/Engine/Classes/Animation/PoseAsset.h b/Engine/Source/Runtime/Engine/Classes/Animation/PoseAsset.h
--- a/Engine/Source/Runtime/Engine/Classes/Animation/PoseAsset.h
+++ b/Engine/Source/Runtime/Engine/Classes/Animation/PoseAsset.h
@@ -274,7 +274,7 @@ public:
 	
 	ENGINE_API const TArray<FName> GetCurveFNames() const;
 	ENGINE_API const TArray<FAnimCurveBase>& GetCurveData() const;
-	ENGINE_API const TArray<float> GetCurveValues(const int32 PoseIndex) const;
+	ENGINE_API const TArray<float>& GetCurveValues(const int32 PoseIndex) const;
 
 	/** Find index of a track with a given bone name. Returns INDEX_NONE if not found. */
 	ENGINE_API const int32 GetTrackIndexByName(const FName& InTrackName) const;
@@ -363,6 +363,7 @@ public:
 
 	// util to return transform of a bone from the pose asset in local space */
 	ENGINE_API const FTransform& GetLocalSpaceTransform(FName BoneName, int32 PoseIndex = 0) const;
+	ENGINE_API const FTransform& GetLocalSpaceTransform(int32 TrackIndex, int32 PoseIndex = 0) const;
 
 	ENGINE_API int32 DeletePoses(TArray<FName> PoseNamesToDelete);
 	ENGINE_API int32 DeleteCurves(TArray<FName> CurveNamesToDelete);
diff --git a/Engine/Source/Runtime/Engine/Private/Animation/PoseAsset.cpp b/Engine/Source/Runtime/Engine/Private/Animation/PoseAsset.cpp
--- a/Engine/Source/Runtime/Engine/Private/Animation/PoseAsset.cpp
+++ b/Engine/Source/Runtime/Engine/Private/Animation/PoseAsset.cpp
@@ -1116,16 +1116,15 @@ const TArray<FAnimCurveBase>& UPoseAsset::GetCurveData() const
 	return PoseContainer.Curves;
 }
 
-const TArray<float> UPoseAsset::GetCurveValues(const int32 PoseIndex) const
+const TArray<float>& UPoseAsset::GetCurveValues(const int32 PoseIndex) const
 {
-	TArray<float> ResultCurveData;
-
 	if (PoseContainer.Poses.IsValidIndex(PoseIndex))
 	{
-		ResultCurveData = PoseContainer.Poses[PoseIndex].CurveData;
+		return PoseContainer.Poses[PoseIndex].CurveData;
 	}
 
-	return ResultCurveData;
+	static const TArray<float> EmptyCurveValues;
+	return EmptyCurveValues;
 }
 
 bool UPoseAsset::GetCurveValue(const int32 PoseIndex, const int32 CurveIndex, float& OutValue) const
@@ -1743,6 +1742,19 @@ const FTransform& UPoseAsset::GetLocalSpaceTransform(FName BoneName, int32 PoseI
 	return PoseContainer.Poses[PoseIndex].SourceLocalSpacePose[TrackIdx];
 }
 
+const FTransform& UPoseAsset::GetLocalSpaceTransform(int32 TrackIndex, int32 PoseIndex) const
+{
+	if (!PoseContainer.Poses.IsValidIndex(PoseIndex))
+	{
+		return FTransform::Identity;
+	}
+	if (!PoseContainer.Poses[PoseIndex].SourceLocalSpacePose.IsValidIndex(TrackIndex))
+	{
+		return FTransform::Identity;
+	}
+	return PoseContainer.Poses[PoseIndex].SourceLocalSpacePose[TrackIndex];
+}
+
 bool UPoseAsset::ConvertSpace(bool bNewAdditivePose, int32 NewBasePoseIndex)
 {
 	// first convert to full pose first
