Added an editor-only flag bForceDuplicateInPIE on AActor to force the actor to be duplicated during PIE.
This is particularly useful for PCG load as preview resources that are modified without being dirtied, as it ensures those modifications are visible in PIE.

diff --git a/Engine/Source/Runtime/Engine/Classes/GameFramework/Actor.h b/Engine/Source/Runtime/Engine/Classes/GameFramework/Actor.h
--- a/Engine/Source/Runtime/Engine/Classes/GameFramework/Actor.h
+++ b/Engine/Source/Runtime/Engine/Classes/GameFramework/Actor.h
@@ -305,6 +305,7 @@ public:
 
 #if WITH_EDITOR
 	virtual bool ShouldDuplicateInPIE() const { return !bIgnoreInPIE; }
+	virtual bool ShouldForceDuplicateInPIE() const { return bForceDuplicateInPIE; }
 #endif
 
 private:
@@ -360,6 +361,9 @@ public:
 	UPROPERTY(Transient)
 	uint8 bIgnoreInPIE : 1;
 
+	UPROPERTY(Transient)
+	uint8 bForceDuplicateInPIE : 1;
+
 private:
 #endif
 
diff --git a/Engine/Source/Runtime/Engine/Private/WorldPartition/WorldPartitionStreamingGeneration.cpp b/Engine/Source/Runtime/Engine/Private/WorldPartition/WorldPartitionStreamingGeneration.cpp
--- a/Engine/Source/Runtime/Engine/Private/WorldPartition/WorldPartitionStreamingGeneration.cpp
+++ b/Engine/Source/Runtime/Engine/Private/WorldPartition/WorldPartitionStreamingGeneration.cpp
@@ -1047,7 +1047,7 @@ class FWorldPartitionStreamingGenerator
 					}
 
 					// Dirty actors
-					if (bShouldHandleUnsavedActors && (bIsTempContainerPackage || (Actor->GetPackage()->IsDirty() || Actor->HasAllFlags(RF_Transient | RF_NonPIEDuplicateTransient))))
+					if (bShouldHandleUnsavedActors && (bIsTempContainerPackage || (Actor->GetPackage()->IsDirty() || Actor->ShouldForceDuplicateInPIE() || Actor->HasAllFlags(RF_Transient | RF_NonPIEDuplicateTransient))))
 					{
 						// Dirty, unsaved actor for PIE
 						TUniquePtr<FStreamingGenerationTempActorDescInstance>& UnsavedDirtyRef = OutUnsavedDirtyInstances.Add_GetRef(FStreamingGenerationTempActorDescInstance::Create(Actor, InActorDescCollection));
