World Partition: fix streaming disabled in PIE for non-spatially loaded level instances referencing a standalone WP world.

diff --git a/Engine/Source/Runtime/Engine/Private/LevelInstance/LevelInstanceLevelStreaming.cpp b/Engine/Source/Runtime/Engine/Private/LevelInstance/LevelInstanceLevelStreaming.cpp
--- a/Engine/Source/Runtime/Engine/Private/LevelInstance/LevelInstanceLevelStreaming.cpp
+++ b/Engine/Source/Runtime/Engine/Private/LevelInstance/LevelInstanceLevelStreaming.cpp
@@ -566,22 +566,25 @@ void ULevelStreamingLevelInstance::OnLevelLoadedChanged(ULevel* InLevel)
 					OuterWorldPartition->OnActorReplacedEvent.AddUObject(this, &ULevelStreamingLevelInstance::OnActorReplacedEvent);
 				}
 
-				if (UWorldPartition* OwningWorldPartition = GetWorld()->GetWorldPartition(); OwningWorldPartition && OwningWorldPartition->IsStreamingEnabled())
+				if (!GetWorld()->IsGameWorld())
 				{
-					if (GDisableLevelInstanceEditorPartialLoading)
+					if (UWorldPartition* OwningWorldPartition = GetWorld()->GetWorldPartition(); OwningWorldPartition && OwningWorldPartition->IsStreamingEnabled())
 					{
-						OuterWorldPartition->bOverrideEnableStreamingInEditor = false;
+						if (GDisableLevelInstanceEditorPartialLoading)
+						{
+							OuterWorldPartition->bOverrideEnableStreamingInEditor = false;
+						}
+						else if (ILevelInstanceInterface* LevelInstance = LevelInstanceSubsystem->GetLevelInstance(LevelInstanceID))
+						{
+							OuterWorldPartition->bOverrideEnableStreamingInEditor = LevelInstance->SupportsPartialEditorLoading();
+						}
 					}
-					else if (ILevelInstanceInterface* LevelInstance = LevelInstanceSubsystem->GetLevelInstance(LevelInstanceID))
+					else
 					{
-						OuterWorldPartition->bOverrideEnableStreamingInEditor = LevelInstance->SupportsPartialEditorLoading();
+						// Do not enable Streaming in Editor if Level Instance is not part of a World Partition/Streaming world
+						OuterWorldPartition->bOverrideEnableStreamingInEditor = false;
 					}
 				}
-				else
-				{
-					// Do not enable Streaming in Editor if Level Instance is not part of a World Partition/Streaming world
-					OuterWorldPartition->bOverrideEnableStreamingInEditor = false;
-				}
 			}
 #endif
 		}
