Control Rig: Fix not being able to choose a cooked control rig for the default rig in a skeletal mesh component
#rnx

diff --git a/Engine/Plugins/Animation/ControlRig/Source/ControlRig/Public/Settings/ControlRigSettings.h b/Engine/Plugins/Animation/ControlRig/Source/ControlRig/Public/Settings/ControlRigSettings.h
--- a/Engine/Plugins/Animation/ControlRig/Source/ControlRig/Public/Settings/ControlRigSettings.h
+++ b/Engine/Plugins/Animation/ControlRig/Source/ControlRig/Public/Settings/ControlRigSettings.h
@@ -66,7 +66,7 @@ public:
 	UPROPERTY(EditAnywhere, config, Category = Shapes)
 	TSoftObjectPtr<UControlRigShapeLibrary> DefaultShapeLibrary;
 
-	UPROPERTY(EditAnywhere, config, Category = ModularRigging, meta=(AllowedClasses="/Script/ControlRigDeveloper.ControlRigBlueprint"))
+	UPROPERTY(EditAnywhere, config, Category = ModularRigging, meta=(AllowedClasses="/Script/ControlRigDeveloper.ControlRigBlueprint, /Script/ControlRig.ControlRigBlueprintGeneratedClass"))
 	FSoftObjectPath DefaultRootModule;
 #endif
 
diff --git a/Engine/Source/Runtime/Engine/Classes/Components/SkeletalMeshComponent.h b/Engine/Source/Runtime/Engine/Classes/Components/SkeletalMeshComponent.h
--- a/Engine/Source/Runtime/Engine/Classes/Components/SkeletalMeshComponent.h
+++ b/Engine/Source/Runtime/Engine/Classes/Components/SkeletalMeshComponent.h
@@ -2755,7 +2755,7 @@ private:
 #if WITH_EDITORONLY_DATA
 
 	/** Default Animating Rig to Use if bOverrideDefaultAnimatingRig is true */
-	UPROPERTY(EditAnywhere, Category = AnimationRig, BlueprintGetter = GetDefaultAnimatingRigOverride, BlueprintSetter = SetDefaultAnimatingRigOverride, meta = (AllowedClasses = "/Script/ControlRigDeveloper.ControlRigBlueprint"), meta = (editcondition = "bOverrideDefaultAnimatingRig"))
+	UPROPERTY(EditAnywhere, Category = AnimationRig, BlueprintGetter = GetDefaultAnimatingRigOverride, BlueprintSetter = SetDefaultAnimatingRigOverride, meta = (AllowedClasses = "/Script/ControlRigDeveloper.ControlRigBlueprint, /Script/ControlRig.ControlRigBlueprintGeneratedClass"), meta = (editcondition = "bOverrideDefaultAnimatingRig"))
 	TSoftObjectPtr<UObject> DefaultAnimatingRigOverride;
 
 #endif
