[WorldPartition] Added support for mutating actor descriptor's HLODLayer and bIsHLODRelevant properties.

diff --git a/Engine/Plugins/Editor/WorldPartitionHLODUtilities/Source/Private/WorldPartition/HLOD/Utilities/WorldPartitionHLODUtilities.cpp b/Engine/Plugins/Editor/WorldPartitionHLODUtilities/Source/Private/WorldPartition/HLOD/Utilities/WorldPartitionHLODUtilities.cpp
--- a/Engine/Plugins/Editor/WorldPartitionHLODUtilities/Source/Private/WorldPartition/HLOD/Utilities/WorldPartitionHLODUtilities.cpp
+++ b/Engine/Plugins/Editor/WorldPartitionHLODUtilities/Source/Private/WorldPartition/HLOD/Utilities/WorldPartitionHLODUtilities.cpp
@@ -717,13 +717,13 @@ void GatherOutputStats(AWorldPartitionHLOD* InHLODActor)
 }
 
 // Iterate over the source actors and retrieve HLOD relevant components using GetHLODRelevantComponents()
-static TArray<UActorComponent*> GatherHLODRelevantComponents(UWorld* InWorld)
+static TArray<UActorComponent*> GatherHLODRelevantComponents(UWorld* InWorld, const UWorldPartitionHLODSourceActors& InSourceActors)
 {
 	TSet<UActorComponent*> HLODRelevantComponents;
 
 	for (TActorIterator<AActor> It(InWorld); It; ++It)
 	{
-		if (!It->IsHLODRelevant())
+		if (!InSourceActors.IsHLODRelevant(*It))
 		{
 			continue;
 		}
@@ -838,7 +838,9 @@ uint32 FWorldPartitionHLODUtilities::BuildHLOD(AWorldPartitionHLOD* InHLODActor)
 		LoadTimeMS = FMath::RoundToInt(LoadTimeScope.GetTime() * 1000);
 	}
 
-	TArray<UActorComponent*> HLODRelevantComponents = GatherHLODRelevantComponents(TempWorld);
+	const UWorldPartitionHLODSourceActors* SourceActors = InHLODActor->GetSourceActors();
+	check(SourceActors);
+	TArray<UActorComponent*> HLODRelevantComponents = GatherHLODRelevantComponents(TempWorld, *SourceActors);
 	FHLODHashBuilder HashBuilder;
 	::ComputeHLODHash(HashBuilder, InHLODActor, HLODRelevantComponents);
 
@@ -1084,7 +1086,9 @@ uint32 FWorldPartitionHLODUtilities::ComputeHLODHash(const AWorldPartitionHLOD*
 		return 0;
 	}
 
-	TArray<UActorComponent*> HLODRelevantComponents = GatherHLODRelevantComponents(TempWorld);
+	const UWorldPartitionHLODSourceActors* SourceActors = InHLODActor->GetSourceActors();
+	check(SourceActors);
+	TArray<UActorComponent*> HLODRelevantComponents = GatherHLODRelevantComponents(TempWorld, *SourceActors);
 
 	FHLODHashBuilder HashBuilder;
 	::ComputeHLODHash(HashBuilder, InHLODActor, HLODRelevantComponents);
diff --git a/Engine/Source/Runtime/Engine/Private/WorldPartition/HLOD/HLODSourceActorsFromCell.cpp b/Engine/Source/Runtime/Engine/Private/WorldPartition/HLOD/HLODSourceActorsFromCell.cpp
--- a/Engine/Source/Runtime/Engine/Private/WorldPartition/HLOD/HLODSourceActorsFromCell.cpp
+++ b/Engine/Source/Runtime/Engine/Private/WorldPartition/HLOD/HLODSourceActorsFromCell.cpp
@@ -184,6 +184,16 @@ uint32 UWorldPartitionHLODSourceActorsFromCell::GetSourceActorsHash(const TArray
 	return Ar.GetCrc();
 }
 
+bool UWorldPartitionHLODSourceActorsFromCell::IsHLODRelevant(AActor* InActor) const
+{
+	check(InActor);
+	const FGuid& ActorInstanceGuid = InActor->GetActorInstanceGuid();
+	return !!Actors.FindByPredicate([&ActorInstanceGuid](const FWorldPartitionRuntimeCellObjectMapping& Actor)
+	{
+		return Actor.ActorInstanceGuid == ActorInstanceGuid;
+	});
+}
+
 void UWorldPartitionHLODSourceActorsFromCell::ComputeHLODHash(FHLODHashBuilder& InHashBuilder) const
 {
 	Super::ComputeHLODHash(InHashBuilder);
diff --git a/Engine/Source/Runtime/Engine/Private/WorldPartition/HLOD/HLODSourceActorsFromLevel.cpp b/Engine/Source/Runtime/Engine/Private/WorldPartition/HLOD/HLODSourceActorsFromLevel.cpp
--- a/Engine/Source/Runtime/Engine/Private/WorldPartition/HLOD/HLODSourceActorsFromLevel.cpp
+++ b/Engine/Source/Runtime/Engine/Private/WorldPartition/HLOD/HLODSourceActorsFromLevel.cpp
@@ -42,6 +42,11 @@ void UWorldPartitionHLODSourceActorsFromLevel::ComputeHLODHash(FHLODHashBuilder&
 	InHashBuilder.HashField(SourceLevel.GetUniqueID().GetLongPackageName(), GET_MEMBER_NAME_CHECKED(UWorldPartitionHLODSourceActorsFromLevel, SourceLevel));
 }
 
+bool UWorldPartitionHLODSourceActorsFromLevel::IsHLODRelevant(AActor* InActor) const
+{
+	check(InActor);
+	return InActor->IsHLODRelevant();
+}
 
 void UWorldPartitionHLODSourceActorsFromLevel::SetSourceLevel(const UWorld* InSourceLevel)
 {
diff --git a/Engine/Source/Runtime/Engine/Private/WorldPartition/WorldPartitionStreamingGeneration.cpp b/Engine/Source/Runtime/Engine/Private/WorldPartition/WorldPartitionStreamingGeneration.cpp
--- a/Engine/Source/Runtime/Engine/Private/WorldPartition/WorldPartitionStreamingGeneration.cpp
+++ b/Engine/Source/Runtime/Engine/Private/WorldPartition/WorldPartitionStreamingGeneration.cpp
@@ -264,6 +264,11 @@ bool FStreamingGenerationActorDescView::GetActorIsHLODRelevant() const
 		return false;
 	}
 
+	if (bIsHLODRelevant.IsSet())
+	{
+		return bIsHLODRelevant.GetValue();
+	}
+
 	return Super::GetActorIsHLODRelevant();
 }
 
@@ -536,6 +541,7 @@ void FStreamingGenerationActorDescView::SetForcedNoHLODLayer()
 
 void FStreamingGenerationActorDescView::SetRuntimeHLODLayer(const FSoftObjectPath& InHLODLayer)
 {
+	bIsForceNoHLODLayer = false;
 	RuntimeHLODLayer = InHLODLayer;
 }
 
@@ -544,6 +550,12 @@ void FStreamingGenerationActorDescView::SetForcedHLODIrrelevant()
 	bIsForcedHLODIrrelevant = true;
 }
 
+void FStreamingGenerationActorDescView::SetIsHLODRelevant(bool bValue)
+{
+	bIsForcedHLODIrrelevant = false;
+	bIsHLODRelevant = bValue;
+}
+
 const FDataLayerInstanceNames& FStreamingGenerationActorDescView::GetRuntimeDataLayerInstanceNames() const
 {
 	if (!RuntimeDataLayerInstanceNames.IsSet())
@@ -669,9 +681,9 @@ class FWorldPartitionStreamingGenerator
 						ActorSetInstance.FilteredActors = FilteredActors;
 						ActorSetInstance.ContainerID = ContainerCollectionInstanceDescriptor.ID;
 						ActorSetInstance.Transform = ContainerCollectionInstanceDescriptor.Transform;
-						ActorSetInstance.bIsSpatiallyLoaded = ReferenceActorDescView.GetIsSpatiallyLoaded(); //PerInstanceData.bIsSpatiallyLoaded;
+						ActorSetInstance.bIsSpatiallyLoaded = ReferenceActorDescView.GetIsSpatiallyLoaded();
 						ActorSetInstance.ContentBundleID = ContainerCollectionInstanceDescriptor.ContentBundleID;
-						ActorSetInstance.RuntimeGrid = ReferenceActorDescView.GetRuntimeGrid(); //PerInstanceData.RuntimeGrid;
+						ActorSetInstance.RuntimeGrid = ReferenceActorDescView.GetRuntimeGrid();
 						ActorSetInstance.DataLayers = StreamingGenerator->GetRuntimeDataLayerInstances(ReferenceActorDescView.GetRuntimeDataLayerInstanceNames().ToArray());
 
 						ActorSetInstance.Bounds.Init();
@@ -1768,8 +1780,22 @@ class FWorldPartitionStreamingGenerator
 				FContainerCollectionInstanceDescriptor& ContainerCollectionInstanceDescriptor = ContainerCollectionInstanceDescriptorsMap.FindChecked(ActorDescMutatorInstance.ContainerId);
 				FActorDescViewMutator& ActorDescViewMutator = ContainerCollectionInstanceDescriptor.ActorDescViewMutators.FindOrAdd(ActorDescMutatorInstance.ActorGuid);
 
-				ActorDescViewMutator.bIsSpatiallyLoaded = ActorDescMutatorInstance.bIsSpatiallyLoaded;
-				ActorDescViewMutator.RuntimeGrid = ActorDescMutatorInstance.RuntimeGrid;
+				if (ActorDescMutatorInstance.bIsSpatiallyLoaded.IsSet())
+				{
+					ActorDescViewMutator.bIsSpatiallyLoaded = ActorDescMutatorInstance.bIsSpatiallyLoaded;
+				}
+				if (ActorDescMutatorInstance.RuntimeGrid.IsSet())
+				{
+					ActorDescViewMutator.RuntimeGrid = ActorDescMutatorInstance.RuntimeGrid;
+				}
+				if (ActorDescMutatorInstance.bIsHLODRelevant.IsSet())
+				{
+					ActorDescViewMutator.bIsHLODRelevant = ActorDescMutatorInstance.bIsHLODRelevant;
+				}
+				if (ActorDescMutatorInstance.HLODLayer.IsSet())
+				{
+					ActorDescViewMutator.HLODLayer = ActorDescMutatorInstance.HLODLayer;
+				}
 			}
 
 			// Build the containers tree representation
@@ -1803,6 +1829,16 @@ class FWorldPartitionStreamingGenerator
 							{
 								ActorDescView->SetRuntimeGrid(ActorDescViewMutator.RuntimeGrid.GetValue());
 							}
+
+							if (ActorDescViewMutator.bIsHLODRelevant.IsSet())
+							{
+								ActorDescView->SetIsHLODRelevant(ActorDescViewMutator.bIsHLODRelevant.GetValue());
+							}
+
+							if (ActorDescViewMutator.HLODLayer.IsSet() && ActorDescView->GetIsSpatiallyLoaded() && ActorDescView->GetActorIsHLODRelevant())
+							{
+								ActorDescView->SetRuntimeHLODLayer(ActorDescViewMutator.HLODLayer.GetValue());
+							}
 						}
 					}
 
diff --git a/Engine/Source/Runtime/Engine/Public/WorldPartition/HLOD/HLODSourceActors.h b/Engine/Source/Runtime/Engine/Public/WorldPartition/HLOD/HLODSourceActors.h
--- a/Engine/Source/Runtime/Engine/Public/WorldPartition/HLOD/HLODSourceActors.h
+++ b/Engine/Source/Runtime/Engine/Public/WorldPartition/HLOD/HLODSourceActors.h
@@ -19,6 +19,7 @@ public:
 #if WITH_EDITOR
 	virtual bool LoadSourceActors(bool& bOutDirty, UWorld* TargetWorld) const PURE_VIRTUAL(UWorldPartitionHLODSourceActors::LoadSourceActors, return false; );
 	ENGINE_API virtual void ComputeHLODHash(FHLODHashBuilder& HashBuilder) const;
+	virtual bool IsHLODRelevant(AActor* InActor) const PURE_VIRTUAL(UWorldPartitionHLODSourceActors::IsHLODRelevant, return false; );
 
 	ENGINE_API void SetHLODLayer(const UHLODLayer* HLODLayer);
 	ENGINE_API const UHLODLayer* GetHLODLayer() const;
diff --git a/Engine/Source/Runtime/Engine/Public/WorldPartition/HLOD/HLODSourceActorsFromCell.h b/Engine/Source/Runtime/Engine/Public/WorldPartition/HLOD/HLODSourceActorsFromCell.h
--- a/Engine/Source/Runtime/Engine/Public/WorldPartition/HLOD/HLODSourceActorsFromCell.h
+++ b/Engine/Source/Runtime/Engine/Public/WorldPartition/HLOD/HLODSourceActorsFromCell.h
@@ -16,6 +16,7 @@ public:
 #if WITH_EDITOR
 	ENGINE_API virtual bool LoadSourceActors(bool& bOutDirty, UWorld* TargetWorld) const override;
 	ENGINE_API virtual void ComputeHLODHash(FHLODHashBuilder& InHashBuilder) const override;
+	ENGINE_API virtual bool IsHLODRelevant(AActor* InActor) const override;
 
 	ENGINE_API static uint32 GetSourceActorsHash(const TArray<FWorldPartitionRuntimeCellObjectMapping>& InSourceActors);
 
diff --git a/Engine/Source/Runtime/Engine/Public/WorldPartition/HLOD/HLODSourceActorsFromLevel.h b/Engine/Source/Runtime/Engine/Public/WorldPartition/HLOD/HLODSourceActorsFromLevel.h
--- a/Engine/Source/Runtime/Engine/Public/WorldPartition/HLOD/HLODSourceActorsFromLevel.h
+++ b/Engine/Source/Runtime/Engine/Public/WorldPartition/HLOD/HLODSourceActorsFromLevel.h
@@ -15,6 +15,7 @@ public:
 #if WITH_EDITOR
 	ENGINE_API virtual bool LoadSourceActors(bool& bOutDirty, UWorld* TargetWorld) const override;
 	ENGINE_API virtual void ComputeHLODHash(FHLODHashBuilder& InHashBuilder) const override;
+	ENGINE_API virtual bool IsHLODRelevant(AActor* InActor) const override;
 
 	ENGINE_API void SetSourceLevel(const UWorld* InSourceLevel);
 	ENGINE_API const TSoftObjectPtr<UWorld>& GetSourceLevel() const;
diff --git a/Engine/Source/Runtime/Engine/Public/WorldPartition/WorldPartitionStreamingGeneration.h b/Engine/Source/Runtime/Engine/Public/WorldPartition/WorldPartitionStreamingGeneration.h
--- a/Engine/Source/Runtime/Engine/Public/WorldPartition/WorldPartitionStreamingGeneration.h
+++ b/Engine/Source/Runtime/Engine/Public/WorldPartition/WorldPartitionStreamingGeneration.h
@@ -93,6 +93,7 @@ public:
 
 	ENGINE_API void SetForcedNoHLODLayer();
 	ENGINE_API void SetRuntimeHLODLayer(const FSoftObjectPath& InHLODLayer);
+	ENGINE_API void SetIsHLODRelevant(bool bValue);
 	ENGINE_API void SetForcedHLODIrrelevant();
 
 	bool operator==(const FStreamingGenerationActorDescView& Other) const
@@ -128,6 +129,7 @@ private:
 	TOptional<FDataLayerInstanceNames> ResolvedDataLayerInstanceNames;
 	TOptional<FDataLayerInstanceNames> RuntimeDataLayerInstanceNames;
 	TOptional<bool> bIsSpatiallyLoaded;
+	TOptional<bool> bIsHLODRelevant;
 	TOptional<TArray<FGuid>> RuntimeReferences;
 	TOptional<FSoftObjectPath> RuntimeHLODLayer;
 	TArray<FGuid> EditorReferences;
@@ -304,6 +306,8 @@ struct FActorDescViewMutator
 {
 	TOptional<bool> bIsSpatiallyLoaded;
 	TOptional<FName> RuntimeGrid;
+	TOptional<bool> bIsHLODRelevant;
+	TOptional<FSoftObjectPath> HLODLayer;
 };
 
 struct FActorDescViewMutatorInstance : FActorDescViewMutator
