Make it so, OnUnregister, the component will always get removed from the nav octree.

diff --git a/Engine/Source/Runtime/Engine/Classes/Components/PrimitiveComponent.h b/Engine/Source/Runtime/Engine/Classes/Components/PrimitiveComponent.h
--- a/Engine/Source/Runtime/Engine/Classes/Components/PrimitiveComponent.h
+++ b/Engine/Source/Runtime/Engine/Classes/Components/PrimitiveComponent.h
@@ -724,6 +724,9 @@ public:
 #endif
 
 #if WITH_EDITOR
+	/** If true, then this component was added to the navigation octree when it was registered. */
+	uint8 bRegisteredForNavigation : 1;
+
 	uint8 bAlwaysAllowTranslucentSelect : 1;
 
 	uint8 SelectionOutlineColorIndex;
diff --git a/Engine/Source/Runtime/Engine/Private/Components/PrimitiveComponent.cpp b/Engine/Source/Runtime/Engine/Private/Components/PrimitiveComponent.cpp
--- a/Engine/Source/Runtime/Engine/Private/Components/PrimitiveComponent.cpp
+++ b/Engine/Source/Runtime/Engine/Private/Components/PrimitiveComponent.cpp
@@ -419,6 +419,7 @@ UPrimitiveComponent::UPrimitiveComponent(const FObjectInitializer& ObjectInitial
 
 #if WITH_EDITOR
 	bAlwaysAllowTranslucentSelect = false;
+	bRegisteredForNavigation = false;
 
 	OverlayColor = FColor(ForceInitToZero);
 	
@@ -675,6 +676,9 @@ void UPrimitiveComponent::OnRegister()
 		if (bNavRelevant)
 		{
 			FNavigationSystem::OnComponentRegistered(*this);
+#if WITH_EDITOR
+			bRegisteredForNavigation = true;
+#endif
 		}
 	}
 	else
@@ -715,11 +719,38 @@ void UPrimitiveComponent::OnUnregister()
 		IStreamingManager::Get().NotifyPrimitiveDetached(this);
 	}
 
-	if (bCanEverAffectNavigation)
+
+#if WITH_EDITOR
+
+	// This is necessary because, due to the SimpleConstructionScript system, it will register this component with default property values before having the correct serialized property values applied.
+	// This was leading to this component registering itself with the navigation octree when it was not supposed to and never knowing to remove itself.
+	// When this bCanEverAffectNavigation was being set, the order of events was:
+	// 1. OnUnregister the object with the old property value.
+	// 2. Change the property value.
+	// 3. OnRegister the object with the new property value.
+	// 4. UActorComponent::ConsolidatedPostEditChange gets called which...
+	//     a. Calls OnUnregister on and destroys the old component.
+	//     b. Creates a new component with default property values.
+	//     c. Calls OnRegister on the new component with the default property values. <=== Source of the bug.
+	//     d. Serializes in the correct property values.
+	//     e. Calls OnUnregister with the correct property value, which has lost what happened during the OnRegister.
+	//     f. Calls OnRegister again with the correct property value.
+	// This was written up in UE-350551.
+
+	const bool bUnregisterNavigation = (bCanEverAffectNavigation || bRegisteredForNavigation);
+#else
+	const bool bUnregisterNavigation = bCanEverAffectNavigation;
+#endif
+
+	if (bUnregisterNavigation)
 	{
 		FNavigationSystem::OnComponentUnregistered(*this);
 	}
 
+#if WITH_EDITOR
+	bRegisteredForNavigation = false;
+#endif
+
 #if WITH_EDITOR
 	FStaticLightingSystemInterface::OnPrimitiveComponentUnregistered.Broadcast(this);
 #endif
