World Partition - HLOD: Added an ensure to try to understand the source of a rare crash during replication of the destruction state

diff --git a/Engine/Source/Runtime/Engine/Private/WorldPartition/HLOD/DestructibleHLODComponent.cpp b/Engine/Source/Runtime/Engine/Private/WorldPartition/HLOD/DestructibleHLODComponent.cpp
--- a/Engine/Source/Runtime/Engine/Private/WorldPartition/HLOD/DestructibleHLODComponent.cpp
+++ b/Engine/Source/Runtime/Engine/Private/WorldPartition/HLOD/DestructibleHLODComponent.cpp
@@ -117,13 +117,13 @@ void FWorldPartitionDestructibleHLODState::PostReplicatedAdd(const TArrayView<in
 
 void FWorldPartitionDestructibleHLODState::PostReplicatedChange(const TArrayView<int32>& ChangedIndices, int32 FinalSize)
 {
-	check(IsClient());
-
 	// Component may still be null if replication occurs before the component BeginPlay() is called, 
 	// in which case FWorldPartitionDestructibleHLODState::Initialize() will not have been called.
 	// DamagedActors will be process directly in Initialize() if that's the case
 	if (OwnerComponent)
 	{
+		check(IsClient());
+
 		for (int32 ChangedIndex : ChangedIndices)
 		{			
 			ApplyDamagedActorState(ChangedIndex);
@@ -212,7 +212,10 @@ void UWorldPartitionDestructibleHLODComponent::ApplyDamagedActorState(int32 Acto
 		DestructibleHLODInstancesMappingData.ForEachActorInstancingMapping(ActorIndex, [](UHLODInstancedStaticMeshComponent* HLODComponent, uint32 InstancesStart, uint32 InstancesCount)
 		{
 			// Clear all instances associated with the actor
-			HLODComponent->BatchUpdateInstancesTransform(InstancesStart, InstancesCount, ZeroedTransform, false, true, false);
+			if (ensure(IsValid(HLODComponent)))
+			{
+				HLODComponent->BatchUpdateInstancesTransform(InstancesStart, InstancesCount, ZeroedTransform, false, true, false);
+			}
 		});
 	}
 }
