PoseAsset: Make retargeting accessors public

Follow up to 47844827

#rnx

diff --git a/Engine/Source/Runtime/Engine/Classes/Animation/PoseAsset.h b/Engine/Source/Runtime/Engine/Classes/Animation/PoseAsset.h
--- a/Engine/Source/Runtime/Engine/Classes/Animation/PoseAsset.h
+++ b/Engine/Source/Runtime/Engine/Classes/Animation/PoseAsset.h
@@ -444,9 +444,8 @@ public:
 	ENGINE_API void UpdateRetargetSourceAssetData();
 #endif
 
-private:
-	const TArray<FTransform>& GetRetargetTransforms() const;
-	FName GetRetargetTransformsSourceName() const;
+	ENGINE_API const TArray<FTransform>& GetRetargetTransforms() const;
+	ENGINE_API FName GetRetargetTransformsSourceName() const;
 
 	friend class FPoseAssetDetails;
 };
