SoundWaveProxyPlayer

Don't reset the Resampler when we seek. Apparently this results in the wrong behavior.

Also, instead of storing the "LastSourceSamples" for the Resampler in a separate array instead of the one we're actively using to render to.

#rnx

diff --git a/Engine/Source/Runtime/Engine/Classes/Sound/SoundWaveProxyPlayer.h b/Engine/Source/Runtime/Engine/Classes/Sound/SoundWaveProxyPlayer.h
--- a/Engine/Source/Runtime/Engine/Classes/Sound/SoundWaveProxyPlayer.h
+++ b/Engine/Source/Runtime/Engine/Classes/Sound/SoundWaveProxyPlayer.h
@@ -320,14 +320,24 @@ private:
 	 */
 	bool InitializeDecoder(const FSoundWaveProxyRef& InWaveProxy, int32 MaxDecodeSizeInFrames);
 
+	void InitializeResampler(int32 NumChannels);
+
 	FSoundWaveProxyPtr SoundWave;
 	FSettings Settings;
 	TUniquePtr<ICompressedAudioInfo> CompressedAudioInfo;
 	TUniquePtr<Audio::IConvertDeinterleave> ConvertDeinterleave;
 	Audio::FRuntimeResampler RuntimeResampler;
 
+	// working "decode buffer" to write decoded audio to before converting it to floats
 	TArray<int16, Audio::FAudioBufferAlignedAllocator> DecodeBuffer;
+
+	// "source audio" buffer to write "interleaved" audio to before de-interleaving
 	Audio::FAlignedFloatBuffer SourceAudio;
+	
+	// residual samples to pass back in to the resampler in the next render
+	Audio::FAlignedFloatBuffer LastSourceSamples;
+
+	// "resampled audio" buffer to render "Resampled" audio to before de-interleaving
 	Audio::FAlignedFloatBuffer ResampledAudio;
 
 	// source wave data
diff --git a/Engine/Source/Runtime/Engine/Private/SoundWaveProxyPlayer.cpp b/Engine/Source/Runtime/Engine/Private/SoundWaveProxyPlayer.cpp
--- a/Engine/Source/Runtime/Engine/Private/SoundWaveProxyPlayer.cpp
+++ b/Engine/Source/Runtime/Engine/Private/SoundWaveProxyPlayer.cpp
@@ -74,10 +74,7 @@ bool FSoundWaveProxyPlayer::SetSoundWave(FSoundWaveProxyPtr InWaveProxy)
 	SortedCuePoints = SoundWave->GetCuePointsAndLoopRegions();
 	Algo::SortBy(SortedCuePoints, &FSoundWaveCuePoint::FramePosition);
 
-	bRequiresResampling = false;
-	RuntimeResampler.Reset(NumChannels);
-	ResampledAudio.Reset();
-	SourceAudio.Reset();
+	InitializeResampler(NumChannels);
 	
 	// Create algorithm for channel conversion and deinterleave 
 	Audio::FConvertDeinterleaveParams ConvertDeinterleaveParams;
@@ -307,7 +304,6 @@ void FSoundWaveProxyPlayer::RenderMultiChannelAudio(Audio::FMultichannelBufferVi
 		}
 		else
 		{
-			
 			SourceAudio.SetNum(OutputNumFrames * NumChannels);
 			GenerateSourceAudio(SourceAudio);
 			ConvertDeinterleave->ProcessAudio(SourceAudio, OutMultiChannelAudio);
@@ -324,8 +320,9 @@ void FSoundWaveProxyPlayer::RenderMultiChannelAudio(Audio::FMultichannelBufferVi
 		int32 SourceNumFrames = RuntimeResampler.GetNumInputFramesNeededToProduceOutputFrames(OutputNumFrames);
 
 		// include the samples we used from the last render
-		int32 LastNumSamples = SourceAudio.Num();
+		int32 LastNumSamples = LastSourceSamples.Num();
 		SourceAudio.SetNum(SourceNumFrames * NumChannels);
+		FMemory::Memcpy(SourceAudio.GetData(), LastSourceSamples.GetData(), LastNumSamples * sizeof(float));
 		GenerateSourceAudio(MakeArrayView(SourceAudio.GetData() + LastNumSamples, SourceAudio.Num() - LastNumSamples));
 
 		int32 NumFramesConsumed = 0;
@@ -349,13 +346,16 @@ void FSoundWaveProxyPlayer::RenderMultiChannelAudio(Audio::FMultichannelBufferVi
 		{
 			// move the last frame of samples to the start of the SourceBuffer for future processing
 			int32 NumSamplesCopy = SourceAudio.Num() - NumSamplesConsumed;
-			FMemory::Memcpy(SourceAudio.GetData(), SourceAudio.GetData() + NumSamplesConsumed, NumSamplesCopy * sizeof(float));
-			SourceAudio.SetNum(NumSamplesCopy);
+			check(NumSamplesCopy == NumChannels);
+			LastSourceSamples.SetNum(NumSamplesCopy);
+			FMemory::Memcpy(LastSourceSamples.GetData(), SourceAudio.GetData() + NumSamplesConsumed, NumSamplesCopy * sizeof(float));
 		}
 		else
 		{
-			SourceAudio.SetNum(0);
+			LastSourceSamples.Reset();
 		}
+
+		SourceAudio.Reset();
 	}
 }
 
@@ -665,12 +665,6 @@ bool FSoundWaveProxyPlayer::SeekToFrame(int32 InFrame)
 		return true;
 	}
 
-	// if we seek, then we should reset the resampler
-	bRequiresResampling = false;
-	RuntimeResampler.Reset(NumChannels);
-	ResampledAudio.Reset();
-	SourceAudio.Reset();
-
 	// reset the "NearlyFinished" flag
 	bHasNearlyFinished = false;
 
@@ -787,4 +781,12 @@ bool FSoundWaveProxyPlayer::InitializeDecoder(const FSoundWaveProxyRef& InWavePr
 	check(DecodeSize > 0);
 	DecodeBuffer.SetNum(DecodeSize * NumChannels, EAllowShrinking::Yes);
 	return true;
-}
\ No newline at end of file
+}
+
+void FSoundWaveProxyPlayer::InitializeResampler(int32 NumInputChannels)
+{
+	bRequiresResampling = false;
+	RuntimeResampler.Reset(NumInputChannels);
+	ResampledAudio.Reset();
+	LastSourceSamples.Empty(NumInputChannels);
+}
