RemoteObjects:  During migration, the ActorComponent should be aware that its owner Actor could have destroyed it during the same PostMigrate step.  We should guard against destroying it twice.

#rnx

diff --git a/Engine/Source/Runtime/Engine/Private/Components/ActorComponent.cpp b/Engine/Source/Runtime/Engine/Private/Components/ActorComponent.cpp
--- a/Engine/Source/Runtime/Engine/Private/Components/ActorComponent.cpp
+++ b/Engine/Source/Runtime/Engine/Private/Components/ActorComponent.cpp
@@ -777,8 +777,15 @@ void UActorComponent::PostMigrate(const struct FUObjectMigrationContext& Migrati
 
 		if (RecvType == EObjectMigrationRecvType::ReturnedLoan)
 		{
-			FLoanedState LoanedState = RemoteObjects_LoanedStates.GetAndRemoveAnnotation(this);
-			if (NewState != LoanedState)
+			const FLoanedState LoanedState = RemoteObjects_LoanedStates.GetAndRemoveAnnotation(this);
+
+			const AActor* Owner = GetOwner();
+			if (ensure(Owner) && Owner->IsActorBeingDestroyed())
+			{
+				// If the Actor is being destroyed, then this Component would have been destroyed through
+				// the Actor's Destroy call, so let's ignore the state changes
+			}
+			else if (NewState != LoanedState)
 			{
 				const bool bNewlyDestroyed = IsBeingDestroyed() && !LoanedState.bWasDestroyed;
 				UE_CLOG(bNewlyDestroyed, LogActorComponent, Warning, TEXT("%hs: Destroying %s back on owning server (it was destroyed while loaned-out)"), __func__, *GetName());
