NetTest
* Added TickUntil function that works with the configured delta time
* Made FindReplicatedObjectOnClient work without passing a ClientIndex. Will use the first client in the array by default now.
* FindReplicatedObjectOnClient will automatically deduce the template type from the first param.

#rnx

diff --git a/Engine/Source/Runtime/Engine/Private/Tests/NetTestHelpers.cpp b/Engine/Source/Runtime/Engine/Private/Tests/NetTestHelpers.cpp
--- a/Engine/Source/Runtime/Engine/Private/Tests/NetTestHelpers.cpp
+++ b/Engine/Source/Runtime/Engine/Private/Tests/NetTestHelpers.cpp
@@ -548,6 +548,18 @@ UObject* FTestWorlds::FindReplicatedObjectOnClient(UObject* ServerObject, uint32
 {
 	using namespace UE::Net;
 
+	if (ClientIndex == UINT32_MAX)
+	{
+		if (Clients.IsEmpty())
+		{
+			ensureMsgf(false, TEXT("%hs failed. No client world exist"), __FUNCTION__);
+			return nullptr;
+		}
+
+		// Use the first client
+		ClientIndex = 0;
+	}
+
 	if (!Clients.IsValidIndex(ClientIndex))
 	{
 		ensureMsgf(false, TEXT("FTestWorlds::FindReplicatedObjectOnClient received invalid ClientIndex: %u"), ClientIndex);
diff --git a/Engine/Source/Runtime/Engine/Public/Tests/NetTestHelpers.h b/Engine/Source/Runtime/Engine/Public/Tests/NetTestHelpers.h
--- a/Engine/Source/Runtime/Engine/Public/Tests/NetTestHelpers.h
+++ b/Engine/Source/Runtime/Engine/Public/Tests/NetTestHelpers.h
@@ -228,6 +228,13 @@ struct FTestWorlds
 	template<class PredicateT>
 	bool TickAllUntil(const PredicateT& Predicate, float DeltaSeconds = 0.0166f, int32 MaxTicks = 60);
 
+	/**
+	 * Ticks all server & client worlds until Predicate returns true, or MaxTicks is reached using the default Delta
+	 * Returns true if Predicate did, false if it didn't within MaxTicks.
+	 */
+	template<class PredicateT>
+	bool TickUntil(const PredicateT& Predicate, int32 MaxTicks = 60);
+
 	/** Ticks all server & client worlds until the passed in client world has a valid client PlayerController. */
 	ENGINE_API bool WaitForClientConnect(FTestWorldInstance& Client);
 
@@ -240,23 +247,23 @@ struct FTestWorlds
 	/**
 	 * Find the remote instance of a replicated object and ensures if it can't find it
 	 * @param ServerObject The server object that you want the remote version of.
-	 * @param ClientIndex The client you want a remote instance of
+	 * @param ClientIndex The client you want a remote instance of. If not set the function will use the first client available.
 	 * @return Return the remote (client) instance of the same object if it exists.
 	 */
-	ENGINE_API UObject* FindReplicatedObjectOnClient(UObject* ServerObject, uint32 ClientIndex) const;
+	ENGINE_API UObject* FindReplicatedObjectOnClient(UObject* ServerObject, uint32 ClientIndex = UINT32_MAX) const;
 
 	template<typename T>
-	T* FindReplicatedObjectOnClient(UObject* ServerObject, uint32 ClientIndex) const
+	T* FindReplicatedObjectOnClient(T* ServerObject, uint32 ClientIndex = UINT32_MAX) const
 	{
-		return Cast<T>(FindReplicatedObjectOnClient(ServerObject, ClientIndex));
+		return Cast<T>(FindReplicatedObjectOnClient(Cast<UObject>(ServerObject), ClientIndex));
 	}
 
 	/**
 	 * Returns if the remote instance of a replicated object exists on the client or not.
 	 * @param ServerObject The server object you want to see if it is replicated
-	 * @param ClientIndex The client you want to test on
+	 * @param ClientIndex The client you want to test on. If not set the function will use the first client available.
 	 */
-	ENGINE_API bool DoesReplicatedObjectExistOnClient(UObject* ServerObject, uint32 ClientIndex) const;
+	ENGINE_API bool DoesReplicatedObjectExistOnClient(UObject* ServerObject, uint32 ClientIndex = UINT32_MAX) const;
 
 	/**
 	 * Returns if the object is registered for replication on the server.
@@ -318,6 +325,26 @@ inline bool FTestWorlds::TickAllUntil(const PredicateT& Predicate, float DeltaSe
 	return bPredicateResult;
 }
 
+template<class PredicateT>
+inline bool FTestWorlds::TickUntil(const PredicateT& Predicate, int32 MaxTicks)
+{
+	PreTickAllDelegate.Broadcast();
+
+	int32 TickCount = 0;
+	bool bPredicateResult = Predicate();
+
+	while (!bPredicateResult && TickCount < MaxTicks)
+	{
+		TickServer();
+		TickClients();
+		TickCount++;
+		GFrameCounter++;
+		bPredicateResult = Predicate();
+	}
+
+	return bPredicateResult;
+}
+
 ENGINE_API bool IsClientConnected(FTestWorldInstance& Instance);
 ENGINE_API bool IsServerConnected(FTestWorldInstance& Instance);
 
