Fix guid determinism, avoid creating a guid for all template objects rather than just CDOs and native subobjects

#rnx

diff --git a/Engine/Source/Runtime/Engine/Private/Components/PostProcessComponent.cpp b/Engine/Source/Runtime/Engine/Private/Components/PostProcessComponent.cpp
--- a/Engine/Source/Runtime/Engine/Private/Components/PostProcessComponent.cpp
+++ b/Engine/Source/Runtime/Engine/Private/Components/PostProcessComponent.cpp
@@ -45,9 +45,9 @@ void UPostProcessComponent::PostInitProperties()
 	Super::PostInitProperties();
 
 #if WITH_EDITOR
-	if (!HasAnyFlags(EObjectFlags::RF_ClassDefaultObject | EObjectFlags::RF_DefaultSubObject) && !VolumeGuid.IsValid())
+	if (!IsTemplate() && !VolumeGuid.IsValid())
 	{
-		VolumeGuid = FGuid::NewGuid();
+		VolumeGuid = FGuid::NewDeterministicGuid(GetPathName());
 	}
 #endif
 }
diff --git a/Engine/Source/Runtime/Engine/Private/PostProcessVolume.cpp b/Engine/Source/Runtime/Engine/Private/PostProcessVolume.cpp
--- a/Engine/Source/Runtime/Engine/Private/PostProcessVolume.cpp
+++ b/Engine/Source/Runtime/Engine/Private/PostProcessVolume.cpp
@@ -134,9 +134,9 @@ void APostProcessVolume::PostInitProperties()
 	Super::PostInitProperties();
 
 #if WITH_EDITOR
-	if (!HasAnyFlags(EObjectFlags::RF_ClassDefaultObject) && !VolumeGuid.IsValid())
+	if (!IsTemplate() && !VolumeGuid.IsValid())
 	{
-		VolumeGuid = FGuid::NewGuid();
+		VolumeGuid = FGuid::NewDeterministicGuid(GetPathName());
 	}
 #endif
 }
