Texture: We also hide chromaticity coordinates on advanced texture settings to match the updated media source support.

diff --git a/Engine/Source/Runtime/Engine/Classes/Engine/Texture.h b/Engine/Source/Runtime/Engine/Classes/Engine/Texture.h
--- a/Engine/Source/Runtime/Engine/Classes/Engine/Texture.h
+++ b/Engine/Source/Runtime/Engine/Classes/Engine/Texture.h
@@ -1189,19 +1189,19 @@ struct FTextureSourceColorSettings
 	ETextureColorSpace ColorSpace;
 
 	/** Red chromaticity coordinate of the source color space. */
-	UPROPERTY(EditAnywhere, BlueprintReadWrite, Category = ColorManagement, meta = (EditCondition = "ColorSpace == ETextureColorSpace::TCS_Custom"))
+	UPROPERTY(EditAnywhere, BlueprintReadWrite, Category = ColorManagement, meta = (EditCondition = "ColorSpace == ETextureColorSpace::TCS_Custom", EditConditionHides))
 	FVector2D RedChromaticityCoordinate;
 
 	/** Green chromaticity coordinate of the source color space. */
-	UPROPERTY(EditAnywhere, BlueprintReadWrite, Category = ColorManagement, meta = (EditCondition = "ColorSpace == ETextureColorSpace::TCS_Custom"))
+	UPROPERTY(EditAnywhere, BlueprintReadWrite, Category = ColorManagement, meta = (EditCondition = "ColorSpace == ETextureColorSpace::TCS_Custom", EditConditionHides))
 	FVector2D GreenChromaticityCoordinate;
 
 	/** Blue chromaticity coordinate of the source color space. */
-	UPROPERTY(EditAnywhere, BlueprintReadWrite, Category = ColorManagement, meta = (EditCondition = "ColorSpace == ETextureColorSpace::TCS_Custom"))
+	UPROPERTY(EditAnywhere, BlueprintReadWrite, Category = ColorManagement, meta = (EditCondition = "ColorSpace == ETextureColorSpace::TCS_Custom", EditConditionHides))
 	FVector2D BlueChromaticityCoordinate;
 
 	/** White chromaticity coordinate of the source color space. */
-	UPROPERTY(EditAnywhere, BlueprintReadWrite, Category = ColorManagement, meta = (EditCondition = "ColorSpace == ETextureColorSpace::TCS_Custom"))
+	UPROPERTY(EditAnywhere, BlueprintReadWrite, Category = ColorManagement, meta = (EditCondition = "ColorSpace == ETextureColorSpace::TCS_Custom", EditConditionHides))
 	FVector2D WhiteChromaticityCoordinate;
 
 	/** Chromatic adaption method applied if the source white point differs from the working color space white point. */
