TriggerSphere Radius is scaled by any gizmo axis, not just X

diff --git a/Engine/Source/Runtime/Engine/Private/TriggerActors.cpp b/Engine/Source/Runtime/Engine/Private/TriggerActors.cpp
--- a/Engine/Source/Runtime/Engine/Private/TriggerActors.cpp
+++ b/Engine/Source/Runtime/Engine/Private/TriggerActors.cpp
@@ -211,8 +211,16 @@ void ATriggerSphere::EditorApplyScale(const FVector& DeltaScale, const FVector*
 {
 	const FVector ModifiedScale = DeltaScale * ( AActor::bUsePercentageBasedScaling ? 500.0f : 5.0f );
 
+	float ScaleToApply = ModifiedScale.GetMax();
+	const float MinScale = ModifiedScale.GetMin();
+
+	// Use scale with the greater length
+	if (FMath::Abs(MinScale) > ScaleToApply)
+	{
+		ScaleToApply = MinScale;
+	}
+
 	USphereComponent * SphereComponent = CastChecked<USphereComponent>(GetRootComponent());
-	SphereComponent->SetSphereRadius(FMath::Max<float>(0.0f, SphereComponent->GetUnscaledSphereRadius() + ModifiedScale.X));
+	SphereComponent->SetSphereRadius(FMath::Max<float>(0.0f, SphereComponent->GetUnscaledSphereRadius() + ScaleToApply));
 }
 #endif	//WITH_EDITOR
-
