Fixed bug that caused incorrect precomputed data to be used with the instance hierarchy. This can happen when a sublevel has a transform, as there is no mechansim for notifying the primitive when this causes a transform to change.

diff --git a/Engine/Source/Runtime/Engine/Private/InstanceData/InstanceDataManager.cpp b/Engine/Source/Runtime/Engine/Private/InstanceData/InstanceDataManager.cpp
--- a/Engine/Source/Runtime/Engine/Private/InstanceData/InstanceDataManager.cpp
+++ b/Engine/Source/Runtime/Engine/Private/InstanceData/InstanceDataManager.cpp
@@ -28,6 +28,25 @@ static TAutoConsoleVariable<bool> CVarInstanceDataManagerBuildOptimizedForTestin
 	TEXT("."));
 
 
+
+/**
+ * Helper function to hash the local to world that was used to build the precomputed data as we can't rely on it being unchanged, 
+ * for example when sublevels with transforms are used.
+ */
+inline uint32 HashPrimitiveLocalToWorld(const FRenderTransform& M, const FVector& PrimitiveWorldSpaceOffset)
+{
+	// round to reduce risk of floating point precision causing otherwise same transforms from being considered different
+	uint32 Hash = HashCombineFast(
+		GetTypeHash(FMath::RoundToInt(M.TransformRows[0][0])), GetTypeHash(FMath::RoundToInt(M.TransformRows[0][1])), GetTypeHash(FMath::RoundToInt(M.TransformRows[0][2])),
+		GetTypeHash(FMath::RoundToInt(M.TransformRows[1][0])), GetTypeHash(FMath::RoundToInt(M.TransformRows[1][1])), GetTypeHash(FMath::RoundToInt(M.TransformRows[1][2])),
+		GetTypeHash(FMath::RoundToInt(M.TransformRows[2][0])), GetTypeHash(FMath::RoundToInt(M.TransformRows[2][1])), GetTypeHash(FMath::RoundToInt(M.TransformRows[2][2])),
+		GetTypeHash(FMath::RoundToInt(M.Origin[0])), GetTypeHash(FMath::RoundToInt(M.Origin[1])), GetTypeHash(FMath::RoundToInt(M.Origin[2])),
+		GetTypeHash(FMath::RoundToInt(PrimitiveWorldSpaceOffset[0])), GetTypeHash(FMath::RoundToInt(PrimitiveWorldSpaceOffset[1])), GetTypeHash(FMath::RoundToInt(PrimitiveWorldSpaceOffset[2]))
+	);
+
+	return Hash;
+}
+
 template <typename IndexRemapType>
 void ApplyDataChanges(FInstanceUpdateChangeSet& ChangeSet, const IndexRemapType& IndexRemap, int32 PostUpdateNumInstances, FInstanceSceneDataBuffers::FWriteView& ProxyData)
 {
@@ -753,9 +772,16 @@ bool FInstanceDataManager::FlushChanges(FInstanceDataManagerSourceDataDesc&& Com
 			LOG_INST_DATA(TEXT("Full Build %s"), TEXT(""));
 			if (PrecomputedOptimizationData.IsValid())
 			{
-				LOG_INST_DATA(TEXT("  Optimized Build (%s)"), PrecomputedOptimizationData.IsValid() ? TEXT("Precomputed") : TEXT(""));
-				ChangeSet.PrecomputedOptimizationData = PrecomputedOptimizationData;
-				SuccessorTrackingState = ETrackingState::Optimized;
+				if (PrecomputedOptimizationData->PrimitiveLocalToWorldHash == HashPrimitiveLocalToWorld(ChangeSet.PrimitiveToRelativeWorld, ChangeSet.PrimitiveWorldSpaceOffset))
+				{
+					LOG_INST_DATA(TEXT("  Optimized Build (%s)"), PrecomputedOptimizationData.IsValid() ? TEXT("Precomputed") : TEXT(""));
+					ChangeSet.PrecomputedOptimizationData = PrecomputedOptimizationData;
+					SuccessorTrackingState = ETrackingState::Optimized;
+				}
+				else
+				{
+					PrecomputedOptimizationData.Reset();
+				}
 			}
 			DispatchInstanceDataUpdateTask(bIsUnattached, InstanceDataProxy, InstanceDataBufferHeader, [ChangeSet = MoveTemp(ChangeSet), InstanceDataProxy = InstanceDataProxy, NumHeaderInstances = InstanceDataBufferHeader.NumInstances] () mutable 
 			{
@@ -1060,7 +1086,7 @@ FPrecomputedInstanceSpatialHashData FInstanceDataManager::PrecomputeOptimization
 	{
 		Result.ProxyIndexToComponentIndexRemap.Reset();
 	}
-
+	Result.PrimitiveLocalToWorldHash = HashPrimitiveLocalToWorld(ChangeSet.PrimitiveToRelativeWorld, ChangeSet.PrimitiveWorldSpaceOffset);
 	return Result;
 }
 
@@ -1088,6 +1114,7 @@ void FInstanceDataManager::WriteCookedRenderData(FArchive& Ar, FInstanceDataMana
 			// Serialize the stuff we need.
 			bHasCookedData = true;
 			Ar << bHasCookedData;
+			Ar << OptData.PrimitiveLocalToWorldHash;
 
 			OptData.Hashes.BulkSerialize(Ar);
 			OptData.ProxyIndexToComponentIndexRemap.BulkSerialize(Ar);
@@ -1110,6 +1137,7 @@ void FInstanceDataManager::ReadCookedRenderData(FArchive& Ar)
 	if (bHasCookedData)
 	{
 		FPrecomputedInstanceSpatialHashData Tmp;
+		Ar << Tmp.PrimitiveLocalToWorldHash;
 
 		// TODO: Pack the data representation to far fewer bits
 		Tmp.Hashes.BulkSerialize(Ar);
diff --git a/Engine/Source/Runtime/Engine/Private/InstancedStaticMesh/ISMInstanceDataManager.cpp b/Engine/Source/Runtime/Engine/Private/InstancedStaticMesh/ISMInstanceDataManager.cpp
--- a/Engine/Source/Runtime/Engine/Private/InstancedStaticMesh/ISMInstanceDataManager.cpp
+++ b/Engine/Source/Runtime/Engine/Private/InstancedStaticMesh/ISMInstanceDataManager.cpp
@@ -41,6 +41,24 @@ static TAutoConsoleVariable<int32> CVarInstanceDataMinInstanceCountToOptimize(
 	2,
 	TEXT("Minimum number of instances to perform optimized build for (if enabled), can be used to disable optimized build for small ISMs as there is some overhead from doing so."));
 
+
+/**
+ * Helper function to hash the local to world that was used to build the precomputed data as we can't rely on it being unchanged, 
+ * for example when sublevels with transforms are used.
+ */
+static uint32 HashLocalToWorld(const FMatrix& M)
+{
+	// round to reduce risk of floating point precision causing otherwise same transforms from being considered different
+	uint32 Hash = HashCombineFast(
+		GetTypeHash(FMath::RoundToInt64(M.M[0][0])), GetTypeHash(FMath::RoundToInt64(M.M[0][1])), GetTypeHash(FMath::RoundToInt64(M.M[0][2])),
+		GetTypeHash(FMath::RoundToInt64(M.M[1][0])), GetTypeHash(FMath::RoundToInt64(M.M[1][1])), GetTypeHash(FMath::RoundToInt64(M.M[1][2])),
+		GetTypeHash(FMath::RoundToInt64(M.M[2][0])), GetTypeHash(FMath::RoundToInt64(M.M[2][1])), GetTypeHash(FMath::RoundToInt64(M.M[2][2])),
+		GetTypeHash(FMath::RoundToInt64(M.M[3][0])), GetTypeHash(FMath::RoundToInt64(M.M[3][1])), GetTypeHash(FMath::RoundToInt64(M.M[3][2]))
+	);
+
+	return Hash;
+}
+
 FPrimitiveInstanceDataManager::FPrimitiveInstanceDataManager(UPrimitiveComponent* InPrimitiveComponent) 
 	: bNumCustomDataChanged(false)
 	, bBakedLightingDataChanged(false)
@@ -682,9 +700,17 @@ bool FPrimitiveInstanceDataManager::FlushChanges(FInstanceUpdateComponentDesc &&
 				check(ComponentData.ComponentMobility == EComponentMobility::Static || ComponentData.ComponentMobility == EComponentMobility::Stationary);
 
 				LOG_INST_DATA(TEXT("  Optimized Build (%s)"), PrecomputedOptimizationData.IsValid() ? TEXT("Precomputed") : TEXT(""));
-				Proxy->bBuildOptimized = true;
-				Proxy->PrecomputedOptimizationData = PrecomputedOptimizationData;
-				SuccessorTrackingState = ETrackingState::Optimized;
+
+				if (PrecomputedOptimizationData->PrimitiveLocalToWorldHash == HashLocalToWorld(ComponentData.PrimitiveLocalToWorld))
+				{
+					Proxy->bBuildOptimized = true;
+					Proxy->PrecomputedOptimizationData = PrecomputedOptimizationData;
+					SuccessorTrackingState = ETrackingState::Optimized;
+				}
+				else
+				{
+					PrecomputedOptimizationData.Reset();
+				}
 			}
 
 			DispatchInstanceDataUpdateTask(bIsUnattached, Proxy, InstanceDataBufferHeader, [ChangeSet = MoveTemp(ChangeSet), Proxy = Proxy, NumHeaderInstances = InstanceDataBufferHeader.NumInstances] () mutable 
@@ -1106,7 +1132,7 @@ FISMPrecomputedSpatialHashData FPrimitiveInstanceDataManager::PrecomputeOptimiza
 	{
 		Result.ProxyIndexToComponentIndexRemap.Reset();
 	}
-
+	Result.PrimitiveLocalToWorldHash = HashLocalToWorld(ComponentData.PrimitiveLocalToWorld);
 	return Result;
 }
 
@@ -1134,6 +1160,7 @@ void FPrimitiveInstanceDataManager::WriteCookedRenderData(FArchive& Ar, FInstanc
 			// Serialize the stuff we need.
 			bHasCookedData = true;
 			Ar << bHasCookedData;
+			Ar << OptData.PrimitiveLocalToWorldHash;
 
 			OptData.Hashes.BulkSerialize(Ar);
 			OptData.ProxyIndexToComponentIndexRemap.BulkSerialize(Ar);
@@ -1156,6 +1183,8 @@ void FPrimitiveInstanceDataManager::ReadCookedRenderData(FArchive& Ar)
 	{
 		FISMPrecomputedSpatialHashData Tmp;
 
+		Ar << Tmp.PrimitiveLocalToWorldHash;
+
 		// TODO: Pack the data representation to far fewer bits
 		Tmp.Hashes.BulkSerialize(Ar);
 		// TODO: RLE-compress
diff --git a/Engine/Source/Runtime/Engine/Public/InstanceData/InstanceUpdateChangeSet.h b/Engine/Source/Runtime/Engine/Public/InstanceData/InstanceUpdateChangeSet.h
--- a/Engine/Source/Runtime/Engine/Public/InstanceData/InstanceUpdateChangeSet.h
+++ b/Engine/Source/Runtime/Engine/Public/InstanceData/InstanceUpdateChangeSet.h
@@ -285,6 +285,7 @@ class FPrecomputedInstanceSpatialHashData
 public:
 	TArray<FInstanceSceneDataBuffers::FCompressedSpatialHashItem> Hashes;
 	TArray<int32> ProxyIndexToComponentIndexRemap;
+	uint32 PrimitiveLocalToWorldHash = 0u;
 };
 
 using FPrecomputedInstanceSpatialHashDataPtr = TSharedPtr<const FPrecomputedInstanceSpatialHashData, ESPMode::ThreadSafe>;
diff --git a/Engine/Source/Runtime/Engine/Public/InstancedStaticMesh/ISMInstanceDataSceneProxy.h b/Engine/Source/Runtime/Engine/Public/InstancedStaticMesh/ISMInstanceDataSceneProxy.h
--- a/Engine/Source/Runtime/Engine/Public/InstancedStaticMesh/ISMInstanceDataSceneProxy.h
+++ b/Engine/Source/Runtime/Engine/Public/InstancedStaticMesh/ISMInstanceDataSceneProxy.h
@@ -18,11 +18,8 @@ DECLARE_LOG_CATEGORY_EXTERN(LogInstanceProxy, Log, All)
 /**
  * Precomputed optimization data that descrives the spatial hashes and reordering needed.
  */
-struct FISMPrecomputedSpatialHashData
-{
-	TArray<FInstanceSceneDataBuffers::FCompressedSpatialHashItem> Hashes;
-	TArray<int32> ProxyIndexToComponentIndexRemap;
-};
+class FPrecomputedInstanceSpatialHashData;
+using FISMPrecomputedSpatialHashData = FPrecomputedInstanceSpatialHashData;
 
 /**
  * Proxy class that represents scene instance data to the renderer.
