Substrate - fixed missing test with SLW causing shader compilation error instead of a clear editor error.

#tests PC editor
#rnx

diff --git a/Engine/Source/Runtime/Engine/Private/Materials/MaterialShared.cpp b/Engine/Source/Runtime/Engine/Private/Materials/MaterialShared.cpp
--- a/Engine/Source/Runtime/Engine/Private/Materials/MaterialShared.cpp
+++ b/Engine/Source/Runtime/Engine/Private/Materials/MaterialShared.cpp
@@ -6069,6 +6069,8 @@ void FinalCompileValidation(
 	TArray<FString>& OutErrors)
 {
 	const bool bSubstrateEnabled = Substrate::IsSubstrateEnabled();
+	const bool bSubstrateHiddenConversionEnabled = Substrate::IsHiddenMaterialAssetConversionEnabled();
+	const bool bSubstrateFrontMaterialPluggedIn = bHasFrontMaterialExpr;
 	const EMaterialDomain Domain = Material->MaterialDomain;
 
 	// If Substrate is enabled or if this is a Substrate material cooked/used in non-Substrate mode, 
@@ -6164,7 +6166,7 @@ void FinalCompileValidation(
 		}
 
 		if (!Material->HasAnyExpressionsInMaterialAndFunctionsOfType<UMaterialExpressionSingleLayerWaterMaterialOutput>()
-			&& !bSubstrateEnabled
+			&& (!bSubstrateEnabled || (bSubstrateHiddenConversionEnabled && !bSubstrateFrontMaterialPluggedIn))
 			)
 		{
 			OutErrors.Add(TEXT("SingleLayerWater materials requires the use of SingleLayerWaterMaterial output node."));
