Iris
* Add utility class FNetRootObjectAdapter
* This class can be used to easily integrate autonomous replication support to an UObject base class.
* This allows non-AActor derived classes to be replicated without being added as subobjects.
* Adapter class currently only supports either AlwaysRelevant objects or NotRouted objects relevant via dependency links.

diff --git a/Engine/Source/Runtime/Engine/Private/Net/Iris/ReplicationSystem/NetRootObjectAdapter.cpp b/Engine/Source/Runtime/Engine/Private/Net/Iris/ReplicationSystem/NetRootObjectAdapter.cpp
new file mode 100644
--- /dev/null
+++ b/Engine/Source/Runtime/Engine/Private/Net/Iris/ReplicationSystem/NetRootObjectAdapter.cpp
@@ -0,0 +1,426 @@
+// Copyright Epic Games, Inc. All Rights Reserved.
+
+#include "Net/Iris/ReplicationSystem/NetRootObjectAdapter.h"
+
+#include "Engine/Engine.h"
+#include "Engine/Level.h"
+#include "Engine/World.h"
+#include "Engine/NetDriver.h"
+
+#include "Iris/Core/IrisLog.h"
+
+#include "Iris/ReplicationSystem/ReplicationSystem.h"
+#include "Iris/ReplicationSystem/NetObjectFactoryRegistry.h"
+
+#include "Net/Iris/ReplicationSystem/EngineReplicationBridge.h"
+#include "Net/Iris/ReplicationSystem/NetRootObjectFactory.h"
+#include "Net/Iris/ReplicationSystem/ReplicationSystemUtil.h"
+
+namespace UE::Net::Private
+{
+	static const FName NotRoutedFilter(TEXT("NotRouted"));
+}
+
+namespace UE::Net
+{
+
+FNetRootObjectAdapter::FNetRootObjectAdapter(const FRootObjectSettings& Settings)
+{
+	Configure(Settings);
+}
+
+void FNetRootObjectAdapter::InitAdapter(UObject* ReplicatedObject)
+{
+	WeakReplicatedObject = ReplicatedObject;
+	check(WeakReplicatedObject.IsValid());
+	ensure(ReplicatedObject->IsSupportedForNetworking());
+
+	bIsInitialized = true;
+}
+
+void FNetRootObjectAdapter::DeinitAdapter()
+{
+	ensureMsgf(!bIsReplicating, TEXT("DeinitAdapter on %s called while the object is still replicating. The object will stay in the replication system!"), *GetNameSafe(WeakReplicatedObject.Get()));
+
+	WeakReplicatedObject.Reset();
+	WeakLevel.Reset();
+
+	bIsInitialized = false;
+}
+
+void FNetRootObjectAdapter::Configure(const FRootObjectSettings& Settings)
+{
+	if (bIsReplicating)
+	{
+		ensureMsgf(!bIsReplicating, TEXT("Cannot configure: %s after it started replicating"), *GetNameSafe(WeakReplicatedObject.Get()));
+		return;
+	}
+
+	ensureMsgf(Settings.bIsAlwaysRelevant != Settings.bIsNotRouted, TEXT("Wrong config for: %s. Setting must be either exclusively AlwaysRelevant or bNotRouted."), *GetNameSafe(WeakReplicatedObject.Get()));
+
+	bIsAlwaysRelevant = Settings.bIsAlwaysRelevant;
+	bIsNotRouted = Settings.bIsNotRouted;
+	bOnlyReplicateWhenLinked = Settings.bOnlyReplicateWhenLinked;
+
+	ExplicitNetFactoryName = Settings.FactoryName;
+}
+
+void FNetRootObjectAdapter::SetAttachedLevel(ULevel* Level)
+{
+	if (bIsReplicating)
+	{
+		ensureMsgf(!bIsReplicating, TEXT("Cannot set attached level after object: %s is replicating"), *GetNameSafe(WeakReplicatedObject.Get()));
+		return;
+	}
+
+	WeakLevel = Level;
+	ensureMsgf(Level == nullptr || WeakLevel.IsValid(), TEXT("Cannot attach to an invalid level"));
+}
+
+FName FNetRootObjectAdapter::GetNetFactoryName() const
+{
+	// If no name was set, use the default RootObjectFactory
+	return ExplicitNetFactoryName.IsNone() ? UNetRootObjectFactory::GetFactoryName() : ExplicitNetFactoryName;
+}
+
+void FNetRootObjectAdapter::SetNetFactoryName(FName InFactoryName)
+{
+	if (ensureMsgf(!bIsReplicating, TEXT("Cannot change factories after replication started")))
+	{
+		ExplicitNetFactoryName = InFactoryName;
+	}
+}
+
+void FNetRootObjectAdapter::StartReplication(ULevel* Level)
+{
+	if (bIsReplicating)
+	{
+		// Already replicating, nothing to do
+		return;
+	}
+
+	if (!bIsInitialized)
+	{
+		ensureMsgf(bIsInitialized, TEXT("FNetRootObjectAdapter must be initialized before starting replication"));
+		return;
+	}
+
+	UObject* ReplicatedObject = WeakReplicatedObject.Get();
+	if (!ReplicatedObject)
+	{
+		ensureMsgf(WeakReplicatedObject.IsValid(), TEXT("FNetRootObjectAdapter cannot start replication without a valid object"));
+		return;
+	}
+
+	// If no level was passed in, see if one was already attached to the object
+	if (Level == nullptr)
+	{
+		Level = WeakLevel.Get();
+	}
+
+	if (!IsValid(Level))
+	{
+		ensureMsgf(Level, TEXT("Adapter for %s needs a level to StartReplication. Pass a valid level or associate one early via SetAttachedLevel"), *GetNameSafe(ReplicatedObject));
+		return;
+	}
+
+	UWorld* World = Level->GetWorld();
+	if (!World)
+	{
+		ensureMsgf(World, TEXT("Level: %s does not have an owning world"), *GetNameSafe(Level));
+		return;
+	}
+	
+	if (World->bIsTearingDown)
+	{
+		UE_LOG(LogIris, Warning, TEXT("%hs Cannot start replication for: %s. The world is being torn down."), __FUNCTION__, *GetNameSafe(ReplicatedObject));
+		return;
+	}
+
+	UE_LOG(LogIris, Verbose, TEXT("%hs for: %s (level: %s)"), __FUNCTION__, *GetNameSafe(ReplicatedObject), *GetNameSafe(Level));
+
+	// Make sure Iris is enabled in the GameNetDriver
+	if (World->GetNetDriver() && !World->GetNetDriver()->IsUsingIrisReplication())
+	{
+		ensureMsgf(false, TEXT("NetRootObjectAdapter is only compatible with the Iris replication system."));
+		return;
+	}
+
+	// Attach the level to the root object if that wasn't done already
+	if (WeakLevel.IsExplicitlyNull())
+	{
+		SetAttachedLevel(Level);
+	}
+	else
+	{
+		ensureMsgf(WeakLevel == Level, TEXT("StartReplication for %s called with level: %s. But he was already associated to level: %s"), *GetNameSafe(WeakReplicatedObject.Get()), *GetNameSafe(Level), *GetNameSafe(WeakLevel.Get()));
+	}
+
+	const FName NetFactoryName = GetNetFactoryName();
+	const FNetObjectFactoryId NetFactoryId = FNetObjectFactoryRegistry::GetFactoryIdFromName(NetFactoryName);
+	if (UNLIKELY(NetFactoryId == InvalidNetObjectFactoryId))
+	{
+		ensureMsgf(NetFactoryId != InvalidNetObjectFactoryId, TEXT("NetFactory: %s is not a valid factory name"), *NetFactoryName.ToString());
+		return;
+	}
+
+	FReplicationSystemUtil::ForEachServerBridge(World, [&](UEngineReplicationBridge* Bridge)
+	{
+		UObjectReplicationBridge::FRootObjectReplicationParams Params 
+		{
+			.bUseClassConfigDynamicFilter = false,
+		};
+
+		if (bIsNotRouted)
+		{
+			Params.bUseExplicitDynamicFilter = true;
+			Params.ExplicitDynamicFilterName = UE::Net::Private::NotRoutedFilter;
+		}
+		else if (bIsAlwaysRelevant)
+		{
+			Params.bUseExplicitDynamicFilter = true;
+			Params.ExplicitDynamicFilterName = NAME_None;
+		}
+		else
+		{
+			//TODO: Custom filter support
+			//      Spatial filter support
+		}
+
+		const FNetRefHandle NetRefHandle = Bridge->StartReplicatingRootObject(ReplicatedObject, Params, NetFactoryId);
+
+		if (!NetRefHandle.IsValid())
+		{
+			UE_LOG(LogIris, Error, TEXT("%hs Error starting replication for RootObject: %s"), __FUNCTION__, *GetNameSafe(ReplicatedObject));
+			ensure(false);
+			return;
+		}
+
+		bIsReplicating = true;
+
+		// Setup level filtering
+		Bridge->AddRootObjectToLevelGroup(ReplicatedObject, Level);
+
+		//TODO: Support a subobject list in the adapter ?
+	});
+}
+
+void FNetRootObjectAdapter::StopReplication()
+{
+	if (!bIsReplicating)
+	{
+		// Was never replicating, nothing to do
+		return;
+	}
+
+	ULevel* Level = WeakLevel.Get();
+	if (!Level)
+	{
+		ensureMsgf(Level, TEXT("StopReplication for %s called after his level was destroyed"), *GetNameSafe(WeakReplicatedObject.Get()));
+		return;
+	}
+
+	UWorld* World = Level->GetWorld();
+	if (!World)
+	{
+		ensureMsgf(World, TEXT("Level: %s does not have an owning world"), *GetNameSafe(Level));
+		return;
+	}
+
+	UObject* ReplicatedObject = WeakReplicatedObject.GetEvenIfUnreachable();
+	if (!ReplicatedObject)
+	{
+		UE_CLOG(!WeakReplicatedObject.IsExplicitlyNull(), LogIris, Warning, TEXT("%hs failed because the root object was destroyed before StopReplication"), __FUNCTION__);
+		
+		return;
+	}
+
+	UE_LOG(LogIris, Verbose, TEXT("%hs for: %s (level: %s)"), __FUNCTION__, *GetNameSafe(ReplicatedObject), *GetNameSafe(Level));
+
+	// If we were still attached to dependents
+	if (DependentRefCount > 0)
+	{
+		UE_LOG(LogIris, Verbose, TEXT("%hs for: %s while still dependent to %d root objects. Dependencies will be cleared."), __FUNCTION__, *GetNameSafe(ReplicatedObject), DependentRefCount);
+
+		// Iris Dependencies will be destroyed by StopReplicatingNetRefHandle so let's clear the refcount here
+		DependentRefCount = 0;
+	}
+
+	FReplicationSystemUtil::ForEachServerBridge(World, [&](UEngineReplicationBridge* Bridge)
+	{
+		FNetRefHandle RefHandle = Bridge->GetReplicatedRefHandle(ReplicatedObject, EGetRefHandleFlags::EvenIfGarbage);
+		if (!RefHandle.IsValid())
+		{
+			// Already not replicated
+			return;
+		}
+
+		//TODO: Figure out what EndReplication flags to use in our different cases.
+		// 		Might be time to refactor this enum it's both too implicit and too explicit.
+		//		Ex: Do Dynamic objects need to pass Destroy or not ? 
+		//			Does Destroy do anything on stables ?
+		//			Why should the caller need to think about all these million use cases; the UEngineReplicationBridge::StopReplicatingActor does abstract most of it. Should we just add a Bridge one for RootObjects too ?
+		//			Should flush be an object trait instead ?
+
+		// Stop replicating the entity
+		Bridge->StopReplicatingNetRefHandle(RefHandle, EEndReplicationFlags::None);
+	});
+
+	bIsReplicating = false;
+}
+
+void FNetRootObjectAdapter::RelevantWith(UObject* OtherRootObject)
+{
+	check(OtherRootObject);
+
+	if (!bIsInitialized)
+	{
+		ensureMsgf(bIsInitialized, TEXT("RootObjectAdapter cannot replicate unless it's initialized to an object"));
+		return;
+	}
+
+	ULevel* Level = WeakLevel.Get();
+	if (!Level)
+	{
+		ensureMsgf(Level, TEXT("RootObjectAdapter cannot replicated without a valid level associated to it: %s."), *GetNameSafe(WeakReplicatedObject.Get()));
+		return;
+	}
+
+	UObject* ReplicatedObject = WeakReplicatedObject.Get();
+	if (!ReplicatedObject)
+	{
+		ensureMsgf(ReplicatedObject, TEXT("RootObjectAdapter assigned to a destroyed object. Something went wrong"));
+		return;
+	}
+
+	if (ReplicatedObject == OtherRootObject)
+	{
+		ensureMsgf(ReplicatedObject != OtherRootObject, TEXT("Cannot call RelevantWith to yourself: %s"), *GetNameSafe(OtherRootObject));
+		return;
+	}
+
+	UE_LOG(LogIris, Verbose, TEXT("%hs for: %s | Attached to %s | (level: %s)"), __FUNCTION__, *GetNameSafe(ReplicatedObject), *GetNameSafe(OtherRootObject), *GetNameSafe(Level));
+
+	DependentRefCount++;
+
+	// Make sure our object is replicating when it first gets associated with another root object.
+	if (!bIsReplicating && DependentRefCount == 1)
+	{
+		StartReplication(Level);
+	}
+
+	if (!bIsReplicating)
+	{
+		// Something went wrong
+		return;
+	}
+
+	UWorld* World = Level->GetWorld();
+	check(World);
+
+	FReplicationSystemUtil::ForEachServerBridge(World, [&](UObjectReplicationBridge* Bridge)
+	{
+		// Find our handle
+		const FNetRefHandle MyRefHandle = Bridge->GetReplicatedRefHandle(ReplicatedObject);
+		if (!MyRefHandle.IsValid())
+		{
+			ensureMsgf(false, TEXT("No netrefhandle assigned to %s"), *GetNameSafe(ReplicatedObject));
+			return;
+		}
+
+		// Find the handle of the object we will be dependent to
+		const FNetRefHandle OtherRefHandle = Bridge->GetReplicatedRefHandle(OtherRootObject);
+		if (!OtherRefHandle.IsValid())
+		{
+			UE_LOG(LogIris, Error, TEXT("%hs Cannot replicate: %s with: %s because the other root object is not replicating"), __FUNCTION__, *GetNameSafe(ReplicatedObject), *GetNameSafe(OtherRootObject));
+			ensure(false);
+			return;
+		}
+
+		// Make sure it's a root object
+		const bool bIsRootObject = Bridge->IsRootObject(OtherRefHandle);
+		if (!bIsRootObject)
+		{
+			UE_LOG(LogIris, Error, TEXT("%hs Cannot replicate %s with %s because it's a subobject"), __FUNCTION__, *GetNameSafe(ReplicatedObject), *GetNameSafe(OtherRootObject));
+			ensure(false);
+			return;
+		}
+
+		Bridge->AddDependentObject(OtherRefHandle, MyRefHandle);
+	});
+}
+
+void FNetRootObjectAdapter::RemoveRelevantWith(UObject* OtherRootObject)
+{
+	UObject* ReplicatedObject = WeakReplicatedObject.Get();
+	if (!ReplicatedObject)
+	{
+		ensureMsgf(ReplicatedObject, TEXT("RootObjectAdapter assigned to a destroyed object. Something went wrong"));
+		return;
+	}
+
+	if (ReplicatedObject == OtherRootObject)
+	{
+		ensureMsgf(ReplicatedObject != OtherRootObject, TEXT("Cannot StopReplicatingWith yourself: %s"), *GetNameSafe(OtherRootObject));
+		return;
+	}
+
+	ULevel* Level = WeakLevel.Get();
+	if (!Level)
+	{
+		ensureMsgf(Level, TEXT("RootObjectAdapter cannot be replicated without a valid level associated to it: %s."), *GetNameSafe(WeakReplicatedObject.Get()));
+		return;
+	}
+
+	UE_LOG(LogIris, Verbose, TEXT("%hs for: %s | Dettaching from: %s | (level: %s)"), __FUNCTION__, *GetNameSafe(ReplicatedObject), *GetNameSafe(OtherRootObject), *GetNameSafe(Level));
+
+	UWorld* World = Level->GetWorld();
+	check(World);
+
+	if (DependentRefCount > 0)
+	{
+		DependentRefCount--;
+	}
+	else
+	{
+		ensureMsgf(DependentRefCount > 0, TEXT("Error found in dependency refcount between %s and %s."), *GetNameSafe(ReplicatedObject), *GetNameSafe(OtherRootObject));
+		// Let's try to remove this dependency anyway in case the error was in a previous call.
+	}
+
+	
+
+	if (!bIsReplicating)
+	{
+		// Nothing more to do if we stopped replicating already
+		return;
+	}
+
+	FReplicationSystemUtil::ForEachServerBridge(World, [&](UObjectReplicationBridge* Bridge)
+	{
+		// Find our handle
+		const FNetRefHandle MyRefHandle = Bridge->GetReplicatedRefHandle(ReplicatedObject);
+		if (!MyRefHandle.IsValid())
+		{
+			ensureMsgf(false, TEXT("No netrefhandle assigned to %s"), *GetNameSafe(ReplicatedObject));
+			return;
+		}
+
+		// Find the handle of the object we will be dependent to
+		const FNetRefHandle OtherRefHandle = Bridge->GetReplicatedRefHandle(OtherRootObject);
+		if (!OtherRefHandle.IsValid())
+		{
+			UE_LOG(LogIris, Warning, TEXT("%hs ignored for: %s with: %s because the other root object is not replicating"), __FUNCTION__, *GetNameSafe(ReplicatedObject), *GetNameSafe(OtherRootObject));
+			return;
+		}
+
+		Bridge->RemoveDependentObject(OtherRefHandle, MyRefHandle);
+	});
+
+	// If we are no longer replicated via another object, remove ourself from the replication system.
+	if (DependentRefCount == 0 && bOnlyReplicateWhenLinked)
+	{
+		StopReplication();
+	}
+}
+
+} // end namespace UE::Net
\ No newline at end of file
diff --git a/Engine/Source/Runtime/Engine/Public/Net/Iris/ReplicationSystem/NetRootObjectAdapter.h b/Engine/Source/Runtime/Engine/Public/Net/Iris/ReplicationSystem/NetRootObjectAdapter.h
new file mode 100644
--- /dev/null
+++ b/Engine/Source/Runtime/Engine/Public/Net/Iris/ReplicationSystem/NetRootObjectAdapter.h
@@ -0,0 +1,147 @@
+// Copyright Epic Games, Inc. All Rights Reserved.
+
+#pragma once
+
+#include "UObject/WeakObjectPtr.h"
+#include "Misc/NotNull.h"
+
+class UWorld;
+class ULevel;
+
+namespace UE::Net
+{
+
+struct FRootObjectSettings
+{
+	/** Set this true so the object will be relevant to everyone. */
+	bool bIsAlwaysRelevant = false;
+
+	/** Set this true when you don't want the root object to be relevant to anyone by default. Useful when you want to make the object relevant only via another root object. */
+	bool bIsNotRouted = false;
+
+	/** When set to true, the object will be registered/unregistered in the ReplicationSystem when it has an active dependency to another object. Use this to reduce the amount of objects handled by the replication system if the object is not always active.  */
+	bool bOnlyReplicateWhenLinked = false;
+
+	//TODO: Add bRelevantForOwner support
+
+	/** Optional name used when the object is handled by a factory different from the NetRootObjectFactory */
+	FName FactoryName;
+};
+
+
+/**
+ * Class that is added by composition to any UObject that wants to be a replicated root object
+ * The expectation if for the UObject class to:
+ * 1) Initialize the adapter early via InitAdapter()
+ * 2) Configure the replication settings via Configure() prior to starting replication
+ * 3) Provide the ULevel the UObject is a part of
+ * 4) Make the object replicated via StartReplication or RelevantWith
+ */
+class FNetRootObjectAdapter
+{
+public:
+
+	FNetRootObjectAdapter() = default;
+	ENGINE_API FNetRootObjectAdapter(const FRootObjectSettings& Settings);
+
+	/** Initialize the adapter with the object it will be responsible to replicate */
+	ENGINE_API void InitAdapter(UObject* ReplicatedObject);
+
+	/** Deinitialize the adapter and clear references */
+	ENGINE_API void DeinitAdapter();
+
+	/** 
+	 * Configure how the object will be replicated. @See FRootObjectSettings for available options. 
+	 * The default behavior is the object will be Always Relevant.
+	 */
+	ENGINE_API void Configure(const FRootObjectSettings& Settings);
+
+	/** 
+	 * Every object needs to be assigned to a ULevel to know which UWorld->ReplicationSystem it will be added to. 
+	 * The level will also be used to set level filtering logic to replicate the object only to clients that loaded the level locally.
+	 * This must be called early if your object will be replicated conditionnally via RelevantWith
+	 * Note: if you don't have a specific level you can use the World->PersistentLevel.
+	 */
+	ENGINE_API void SetAttachedLevel(ULevel* Level);
+
+	/** 
+	 * Start replicating the object.
+	 * @param Level: Level the object is a part of. Can be set to null if SetAttachedLevel was previously set.
+	 */
+	ENGINE_API void StartReplication(ULevel* Level);
+
+	/** Stop replicating the object. */
+	ENGINE_API void StopReplication();
+
+	/**
+	 * Start replicating the object but make it relevant only via another root object.
+	 * You can have relevancy links to multiple different objects. As long as of those is relevant to a client then this object will be too.
+	 * @OtherRootObject:  A replicated root object we want to be tie our relevancy status to.
+	 */
+	ENGINE_API void RelevantWith(UObject* OtherRootObject);
+
+	/**
+	 * Remove the relevancy dependency between our object and another root object
+	 * @OtherRootObject: A replicated root object that was previously in a RelevantWith call.
+	 */
+	ENGINE_API void RemoveRelevantWith(UObject* OtherRootObject);
+
+	/** The netobject factory to use with this rootobject. By default it will use the UNetRootObjectFactory. */
+	ENGINE_API FName GetNetFactoryName() const;
+
+	/** Set to use a netobject factory different from the default. Must be set before replication starts. */
+	ENGINE_API void SetNetFactoryName(FName InFactoryName);
+
+	bool IsInitialized() const
+	{
+		return bIsInitialized;
+	}
+
+	bool IsReplicating() const
+	{
+		return bIsReplicating;
+	}
+
+	bool HasAttachedLevel() const
+	{
+		return WeakLevel.IsValid();
+	}
+
+	bool IsAlwaysRelevant() const
+	{
+		return bIsAlwaysRelevant;
+	}
+
+	bool IsNotRouted() const
+	{
+		return bIsNotRouted;
+	}
+
+private:
+
+	FName ExplicitNetFactoryName;	
+
+	FWeakObjectPtr WeakReplicatedObject;
+	TWeakObjectPtr<ULevel> WeakLevel;
+
+	/** Number of objects we are dependent with. */
+	uint32 DependentRefCount = 0;
+
+	// Status flags
+
+	/** Is the adapter correctly setup */
+	bool bIsInitialized:1 = false;
+	/** Server only variable set when the object is part of the replication system */
+	bool bIsReplicating:1 = false;
+
+	// Configuration flags 
+
+	/** By default objects will be always relevant. Use FRootObjectSettings to customize the object. */
+	bool bIsAlwaysRelevant:1 = true;
+	/** When true the object will not be relevant until an object it is replicating with is also relevant. */
+	bool bIsNotRouted:1 = false;
+	/** When true we automatically stop replication when the object does not have any objects to replicate with anymore. */
+	bool bOnlyReplicateWhenLinked:1 = false;
+};
+
+} // end namespace UE::Net
\ No newline at end of file
