Fix abandoned RenderAssetUpdate crash

#rnx

diff --git a/Engine/Source/Runtime/Engine/Private/Streaming/RenderAssetUpdate.cpp b/Engine/Source/Runtime/Engine/Private/Streaming/RenderAssetUpdate.cpp
--- a/Engine/Source/Runtime/Engine/Private/Streaming/RenderAssetUpdate.cpp
+++ b/Engine/Source/Runtime/Engine/Private/Streaming/RenderAssetUpdate.cpp
@@ -225,17 +225,11 @@ void FRenderAssetUpdate::Tick(EThreadType InCurrentThread)
 			TaskState = TS_Locked;
 
 			ETaskState TickResult;
-			do // Iterate as longs as there is progress
-			{
-				LLM_SCOPE_DYNAMIC_STAT_OBJECTPATH(StreamableAsset, ELLMTagSet::Assets);
-				LLM_SCOPE_DYNAMIC_STAT_OBJECTPATH(StreamableAsset->GetClass(), ELLMTagSet::AssetClasses);
-				UE_TRACE_METADATA_SCOPE_ASSET(StreamableAsset, StreamableAsset->GetClass());
-
-				// Only test for suspension the first time and in normal progress.
-				// When cancelled, we want the update to complete without interruptions, allowing reference to be freed.
-				TickResult = TickInternal(InCurrentThread, !bWasAlreadyLocked && !bIsCancelled);
-			} 
-			while (TickResult == TS_Locked);
+#if !UE_BUILD_SHIPPING
+			TickResult = IterateInternalTickWithProfiling(InCurrentThread, bWasAlreadyLocked);
+#else
+			TickResult = IterateInternalTick(InCurrentThread, bWasAlreadyLocked);
+#endif
 
 			// We do this to prevent updating the TaskState while in the Lock.
 			if (!bWasAlreadyLocked)
@@ -247,6 +241,40 @@ void FRenderAssetUpdate::Tick(EThreadType InCurrentThread)
 	}
 }
 
+auto FRenderAssetUpdate::IterateInternalTick(EThreadType InCurrentThread, bool bWasAlreadyLocked) -> ETaskState
+{
+	ETaskState TickResult;
+	do // Iterate as long as there is progress
+	{
+		// Only test for suspension the first time and in normal progress.
+		// When cancelled, we want the update to complete without interruptions, allowing reference to be freed.
+		TickResult = TickInternal(InCurrentThread, !bWasAlreadyLocked && !bIsCancelled);
+	} while (TickResult == TS_Locked);
+
+	return TickResult;
+}
+
+#if !UE_BUILD_SHIPPING
+auto FRenderAssetUpdate::IterateInternalTickWithProfiling(EThreadType InCurrentThread, bool bWasAlreadyLocked) -> ETaskState
+{
+	ETaskState TickResult;
+	if (StreamableAsset) // StreamableAsset can be nullptr because the asset update was abandoned (see r.Streaming.EnableAssetUpdateAbandons)
+	{
+		LLM_SCOPE_DYNAMIC_STAT_OBJECTPATH(StreamableAsset, ELLMTagSet::Assets);
+		LLM_SCOPE_DYNAMIC_STAT_OBJECTPATH(StreamableAsset->GetClass(), ELLMTagSet::AssetClasses);
+		UE_TRACE_METADATA_SCOPE_ASSET(StreamableAsset, StreamableAsset->GetClass());
+
+		TickResult = IterateInternalTick(InCurrentThread, bWasAlreadyLocked);
+	}
+	else
+	{
+		TickResult = IterateInternalTick(InCurrentThread, bWasAlreadyLocked);
+	}
+
+	return TickResult;
+}
+#endif
+
 class FRenderAssetUpdateTickGTTask
 {
 public:
diff --git a/Engine/Source/Runtime/Engine/Public/RenderAssetUpdate.h b/Engine/Source/Runtime/Engine/Public/RenderAssetUpdate.h
--- a/Engine/Source/Runtime/Engine/Public/RenderAssetUpdate.h
+++ b/Engine/Source/Runtime/Engine/Public/RenderAssetUpdate.h
@@ -229,6 +229,14 @@ protected:
 
 private:
 
+	/** Iterate calling internal tick as long as there is progress. */
+	ETaskState IterateInternalTick(EThreadType InCurrentThread, bool bWasAlreadyLocked);
+
+#if !UE_BUILD_SHIPPING
+	/** Iterate calling internal tick as long as there is progress with added profiling. */
+	ETaskState IterateInternalTickWithProfiling(EThreadType InCurrentThread, bool bWasAlreadyLocked);
+#endif
+
 	/** Ref couting **/
 	mutable FThreadSafeCounter NumRefs;
 };
