[HLOD] Expose BuildHLODForActors and BuildHLODForVolume functions to Editor Utils

diff --git a/Engine/Source/Runtime/Engine/Private/WorldPartition/HLOD/HLODBlueprintLibrary.cpp b/Engine/Source/Runtime/Engine/Private/WorldPartition/HLOD/HLODBlueprintLibrary.cpp
new file mode 100644
--- /dev/null
+++ b/Engine/Source/Runtime/Engine/Private/WorldPartition/HLOD/HLODBlueprintLibrary.cpp
@@ -0,0 +1,49 @@
+// Copyright Epic Games, Inc. All Rights Reserved.
+
+#include "WorldPartition/HLOD/HLODBlueprintLibrary.h"
+
+#if WITH_EDITOR
+#include "Editor.h"
+#include "WorldPartition/HLOD/HLODCreationFilter.h"
+#include "WorldPartition/WorldPartitionHLODsBuilder.h"
+#endif
+
+#include UE_INLINE_GENERATED_CPP_BY_NAME(HLODBlueprintLibrary)
+
+bool UWorldPartitionHLODBlueprintLibrary::BuildHLODForActors(const TArray<AActor*>& InActors)
+{
+#if WITH_EDITOR
+	if (UWorld* World = GEditor->GetEditorWorldContext().World())
+	{
+		TArray<FBox> FilterVolumes;
+		for (const AActor* Actor : InActors)
+		{
+			if (IsValid(Actor))
+			{
+				FBox RuntimeBounds;
+				FBox EditorBounds;
+				Actor->GetStreamingBounds(RuntimeBounds, EditorBounds);
+				FilterVolumes.Add(RuntimeBounds);
+			}
+		}
+
+		if (FilterVolumes.Num() > 0)
+		{
+			return UWorldPartitionHLODsBuilder::BuildWithFilters(World, { MakeShared<FHLODCreationVolumeFilter>(MoveTemp(FilterVolumes)) });
+		}
+	}
+#endif
+	return false;
+}
+
+bool UWorldPartitionHLODBlueprintLibrary::BuildHLODForVolume(const FBox& InBox)
+{
+#if WITH_EDITOR
+	if (UWorld* World = GEditor->GetEditorWorldContext().World())
+	{
+		TArray<FBox> FilterVolumes = { InBox };
+		return UWorldPartitionHLODsBuilder::BuildWithFilters(World, { MakeShared<FHLODCreationVolumeFilter>(MoveTemp(FilterVolumes)) });
+	}
+#endif
+	return false;
+}
\ No newline at end of file
diff --git a/Engine/Source/Runtime/Engine/Public/WorldPartition/HLOD/HLODBlueprintLibrary.h b/Engine/Source/Runtime/Engine/Public/WorldPartition/HLOD/HLODBlueprintLibrary.h
new file mode 100644
--- /dev/null
+++ b/Engine/Source/Runtime/Engine/Public/WorldPartition/HLOD/HLODBlueprintLibrary.h
@@ -0,0 +1,20 @@
+// Copyright Epic Games, Inc. All Rights Reserved.
+
+#pragma once
+
+#include "Kismet/BlueprintFunctionLibrary.h"
+#include "HLODBlueprintLibrary.generated.h"
+
+class AActor;
+
+UCLASS(MinimalAPI)
+class UWorldPartitionHLODBlueprintLibrary : public UBlueprintFunctionLibrary
+{
+	GENERATED_BODY()
+
+	UFUNCTION(BlueprintCallable, Category="HLOD")
+	static bool BuildHLODForActors(const TArray<AActor*>& InActors);
+
+	UFUNCTION(BlueprintCallable, Category="HLOD")
+	static bool BuildHLODForVolume(const FBox& InBox);
+};
\ No newline at end of file
