Added GeneratedClassPreSaveRoot/GeneratedClassPostSaveRoot in UBlueprint to allow derived classes to react to it.
Changed UEntityPrefab to handle this and forward it to the the CDO entity.

diff --git a/Engine/Source/Runtime/Engine/Classes/Engine/Blueprint.h b/Engine/Source/Runtime/Engine/Classes/Engine/Blueprint.h
--- a/Engine/Source/Runtime/Engine/Classes/Engine/Blueprint.h
+++ b/Engine/Source/Runtime/Engine/Classes/Engine/Blueprint.h
@@ -991,6 +991,9 @@ public:
 	//~ End UObject Interface
 
 #if WITH_EDITORONLY_DATA
+	ENGINE_API virtual void GeneratedClassPreSaveRoot(FObjectPreSaveRootContext ObjectSaveContext);
+	ENGINE_API virtual void GeneratedClassPostSaveRoot(FObjectPostSaveRootContext ObjectSaveContext);
+
 	/** Get the Blueprint object that generated the supplied class */
 	static ENGINE_API UBlueprint* GetBlueprintFromClass(const UClass* InClass);
 
diff --git a/Engine/Source/Runtime/Engine/Private/Blueprint.cpp b/Engine/Source/Runtime/Engine/Private/Blueprint.cpp
--- a/Engine/Source/Runtime/Engine/Private/Blueprint.cpp
+++ b/Engine/Source/Runtime/Engine/Private/Blueprint.cpp
@@ -2384,6 +2384,14 @@ int32 UBlueprint::RemoveExtension(const TObjectPtr<UBlueprintExtension>& InExten
 
 PRAGMA_ENABLE_DEPRECATION_WARNINGS
 
+void UBlueprint::GeneratedClassPreSaveRoot(FObjectPreSaveRootContext ObjectSaveContext)
+{
+}
+
+void UBlueprint::GeneratedClassPostSaveRoot(FObjectPostSaveRootContext ObjectSaveContext)
+{
+}
+
 #endif //WITH_EDITOR
 
 #if WITH_EDITORONLY_DATA
diff --git a/Engine/Source/Runtime/Engine/Private/BlueprintGeneratedClass.cpp b/Engine/Source/Runtime/Engine/Private/BlueprintGeneratedClass.cpp
--- a/Engine/Source/Runtime/Engine/Private/BlueprintGeneratedClass.cpp
+++ b/Engine/Source/Runtime/Engine/Private/BlueprintGeneratedClass.cpp
@@ -2279,6 +2279,11 @@ void UBlueprintGeneratedClass::PreSaveRoot(FObjectPreSaveRootContext ObjectSaveC
 	{
 		PurgeCookedMetaData();
 	}
+	
+	if (UBlueprint* BP = Cast<UBlueprint>(ClassGeneratedBy))
+	{
+		BP->GeneratedClassPreSaveRoot(ObjectSaveContext);
+	}
 }
 
 void UBlueprintGeneratedClass::PostSaveRoot(FObjectPostSaveRootContext ObjectSaveContext)
@@ -2288,6 +2293,11 @@ void UBlueprintGeneratedClass::PostSaveRoot(FObjectPostSaveRootContext ObjectSav
 	CookedPropertyGuids.Reset();
 
 	PurgeCookedMetaData();
+
+	if (UBlueprint* BP = Cast<UBlueprint>(ClassGeneratedBy))
+	{
+		BP->GeneratedClassPostSaveRoot(ObjectSaveContext);
+	}	
 }
 #endif
 
