[Actor]
- Fix LiveCoding issue with PostInitProperties being called from multiple threads now that TEDS is thread-safe
- Keep the update in PostLoad in the loading case for now, but it should be safe to also cleanup that path later

diff --git a/Engine/Source/Runtime/Engine/Private/Actor.cpp b/Engine/Source/Runtime/Engine/Private/Actor.cpp
--- a/Engine/Source/Runtime/Engine/Private/Actor.cpp
+++ b/Engine/Source/Runtime/Engine/Private/Actor.cpp
@@ -537,15 +537,12 @@ void AActor::PostInitProperties()
 	Super::PostInitProperties();
 
 #if WITH_EDITOR
-	if (IsInGameThread())
+	// This will be executed in PostLoad instead during loading to avoid registering objects too soon
+	// before they are ready.
+	if (!IsInAsyncLoadingThread())
 	{
-		// This will be executed in PostLoad instead during loading
 		UEngineElementsLibrary::CreateEditorActorElement(this);
 	}
-	else
-	{
-		check(IsInAsyncLoadingThread());
-	}
 #endif	// WITH_EDITOR
 
 	RemoteRole = (bReplicates ? ROLE_SimulatedProxy : ROLE_None);
