Fix crash on XB1 related to GpuPassthroughVertexFactory pre-skinning SRV not properly set to a null resource when GpuSkinCache is disabled and ManualVertexFetch is disabled.

#tests on console

diff --git a/Engine/Source/Runtime/Engine/Private/GPUSkinVertexFactory.cpp b/Engine/Source/Runtime/Engine/Private/GPUSkinVertexFactory.cpp
--- a/Engine/Source/Runtime/Engine/Private/GPUSkinVertexFactory.cpp
+++ b/Engine/Source/Runtime/Engine/Private/GPUSkinVertexFactory.cpp
@@ -1516,7 +1516,14 @@ void FGPUSkinPassthroughVertexFactory::UpdateLooseUniformBuffer(FRHICommandListB
 	Parameters.FrameNumber = InFrameNumber;
 	Parameters.PositionBuffer = PositionSRV;
 	Parameters.PreviousPositionBuffer = PrevPositionSRV;
-	Parameters.PreSkinnedTangentBuffer = PreSkinnedTangentSRV;
+	if (PreSkinnedTangentSRV)
+	{
+		Parameters.PreSkinnedTangentBuffer = PreSkinnedTangentSRV;
+	}
+	else
+	{
+		Parameters.PreSkinnedTangentBuffer = GNullVertexBuffer.VertexBufferSRV;
+	}
 	LooseParametersUniformBuffer.UpdateUniformBufferImmediate(RHICmdList, Parameters);
 }
 
