Splines: Make default parameterization policy CVar configurable

This also includes a minor change which forwards the knot vector after reparameterization to the child splines. This should have existed in the first place, it is done everywhere else we touch the knot vector.

diff --git a/Engine/Source/Runtime/Engine/Private/Curves/Spline.cpp b/Engine/Source/Runtime/Engine/Private/Curves/Spline.cpp
--- a/Engine/Source/Runtime/Engine/Private/Curves/Spline.cpp
+++ b/Engine/Source/Runtime/Engine/Private/Curves/Spline.cpp
@@ -39,6 +39,13 @@ FAutoConsoleVariableRef CVarSplineImplementation(
 	,ECVF_ReadOnly		// In non-editor builds, this CVar should exist so that GImplementation is properly set but it should not be changed by the user.
 #endif
 );
+
+int32 GParameterizationPolicy = 2;
+FAutoConsoleVariableRef CVarParameterizationPolicy(
+	TEXT("Spline.ParameterizationPolicy"),
+	GParameterizationPolicy,
+	TEXT("0) Uniform - 1) Chord Length - 2) Centripetal")
+);
 	
 bool GApproximateTangents = false;
 FAutoConsoleVariableRef CVarApproximateTangents(
@@ -99,6 +106,16 @@ FAutoConsoleVariableRef CVarValidateRotScale(
 	
 }
 
+UE::Geometry::Spline::EParameterizationPolicy GetParameterizationPolicy()
+{
+	switch (UE::Spline::GParameterizationPolicy)
+	{
+	default:	// fallthrough
+	case 0: return UE::Geometry::Spline::EParameterizationPolicy::Uniform;
+	case 1: return UE::Geometry::Spline::EParameterizationPolicy::ChordLength;
+	case 2: return UE::Geometry::Spline::EParameterizationPolicy::Centripetal;
+	}
+}
 
 /**
  * Converts from EInterpCurveMode to ETangentMode
@@ -314,6 +331,8 @@ private:
 	// The alternative is to redefine a bunch of MultiSpline functions to do the remapping internally, but this is a bit of a pain.
 	friend struct FSpline;
 	
+	void Reparameterize();
+
 	/** Convert from [0, 1] to [0, NumSegments] */
 	float FromNormalizedSpace(float Parameter) const;
 	
@@ -1916,10 +1935,11 @@ bool FNewSpline::Serialize(FArchive& Ar)
     TMultiSpline::Serialize(Ar);
     Ar << ReparamStepsPerSegment;
 
-	// Do not serialize legacy curves, only generate if loading.
 	if (Ar.IsLoading())
 	{
-		GetSpline().Reparameterize();
+		Reparameterize();
+
+		// Do not serialize legacy curves, only generate if loading.
 		RebuildLegacyCurves();
 	}
 	
@@ -2494,13 +2514,13 @@ FVector FNewSpline::EvaluateScale(float Parameter) const
 
 void FNewSpline::UpdateSpline()
 {
-    GetSpline().UpdateTangents();			// updates the tangents on our version of the points
-	
-	// We do this so that parameterization remains proportional to the square root of segment chord lengths.
-	GetSpline().Reparameterize();
-	
+	// Update our parameterization, forward it to the rotation and scale channels, and then recompute tangents.
+	Reparameterize();
+
     MarkReparamTableDirty();
-    RebuildLegacyCurves();					// updates legacy curves based on points, never evaluates internal spline.
+
+	// Updates legacy curves based on points, never evaluates internal spline.
+    RebuildLegacyCurves();					
 }
 
 void FNewSpline::UpdateSpline(const FSpline::FUpdateSplineParams& Params)
@@ -2514,6 +2534,32 @@ void FNewSpline::UpdateSpline(const FSpline::FUpdateSplineParams& Params)
     UpdateSpline();
 }
 
+void FNewSpline::Reparameterize()
+{
+	GetSpline().Reparameterize(GetParameterizationPolicy());
+	GetSpline().UpdateTangents();
+
+	if (UE::Geometry::Spline::TTangentBezierSpline<FVector>* ScaleSpline = GetTypedAttributeChannel<UE::Geometry::Spline::TTangentBezierSpline<FVector>>(ScaleAttrName))
+	{
+		ScaleSpline->SetKnotVector(GetSpline().GetKnotVector());
+		ScaleSpline->UpdateTangents();
+	}
+	else
+	{
+		UE_LOG(LogSpline, Warning, TEXT("Failed to get scale spline while reparameterizing!"));
+	}
+
+	if (UE::Geometry::Spline::TTangentBezierSpline<FQuat>* RotSpline = GetTypedAttributeChannel<UE::Geometry::Spline::TTangentBezierSpline<FQuat>>(RotationAttrName))
+	{
+		RotSpline->SetKnotVector(GetSpline().GetKnotVector());
+		RotSpline->UpdateTangents();
+	}
+	else
+	{
+		UE_LOG(LogSpline, Warning, TEXT("Failed to get rotation spline while reparameterizing!"));
+	}
+}
+
 float FNewSpline::FromNormalizedSpace(float Parameter) const
 {
 	return GetNumberOfSegments() * Parameter;
