Merging CL#47866696

Disable AnimDyn + RBAN on a specifed character part on all active player characters
- Added a per anim node enable flag to AnimDynamics node to mirror the one in RBAN
- Added a Disable AnimDynamics flag to the SkeletalMesh to mirror the one for RBAN
- Added a set of BP callable functions that allow enabling/disabling of AnimDynamics, RBAN and Cloth physics on any character part on one target animinstance or all active players
- Added and a console function to allow enabling/disabling of AnimDynamics, RBAN and Cloth physics on any character part on all active players

diff --git a/Engine/Source/Runtime/AnimGraphRuntime/Private/BoneControllers/AnimNode_AnimDynamics.cpp b/Engine/Source/Runtime/AnimGraphRuntime/Private/BoneControllers/AnimNode_AnimDynamics.cpp
--- a/Engine/Source/Runtime/AnimGraphRuntime/Private/BoneControllers/AnimNode_AnimDynamics.cpp
+++ b/Engine/Source/Runtime/AnimGraphRuntime/Private/BoneControllers/AnimNode_AnimDynamics.cpp
@@ -227,6 +227,7 @@ FAnimNode_AnimDynamics::FAnimNode_AnimDynamics()
 #if WITH_EDITOR
 , bDoPhysicsUpdateInEditor(true)
 #endif
+, bEnabled(true)
 #if ENABLE_ANIM_DRAW_DEBUG
 , FilteredBoneIndex(INDEX_NONE)
 #endif
@@ -439,7 +440,7 @@ void FAnimNode_AnimDynamics::EvaluateSkeletalControl_AnyThread(FComponentSpacePo
 	SCOPE_CYCLE_COUNTER(STAT_AnimDynamicsOverall);
 	CSV_SCOPED_TIMING_STAT(Animation, AnimDynamicsEval);
 
-	if (IsAnimDynamicsSystemEnabledFor(Output.AnimInstanceProxy->GetLODLevel()))
+	if (IsAnimDynamicsSystemEnabledFor(Output.AnimInstanceProxy->GetLODLevel()) && bEnabled)
 	{
 		const int bScalingEnabled = CVarEnableScaling.GetValueOnAnyThread();
 
@@ -1353,20 +1354,6 @@ void FAnimNode_AnimDynamics::UpdateLimits(FComponentSpacePoseContext& Output)
 	}
 }
 
-bool FAnimNode_AnimDynamics::HasPreUpdate() const
-{
-	if(CVarEnableDynamics.GetValueOnAnyThread() == 1)
-	{
-		return (CVarEnableWind.GetValueOnAnyThread() == 1 && (bEnableWind || bWindWasEnabled))
-#if ENABLE_ANIM_DRAW_DEBUG
-				|| (CVarShowDebug.GetValueOnAnyThread() == 1 && !CVarDebugBone.GetValueOnAnyThread().IsEmpty())
-#endif
-				;
-	}
-
-	return false;
-}
-
 void FAnimNode_AnimDynamics::PreUpdate(const UAnimInstance* InAnimInstance)
 {
 	// If dynamics are disabled, skip all this work as it'll never get used
@@ -1389,6 +1376,8 @@ void FAnimNode_AnimDynamics::PreUpdate(const UAnimInstance* InAnimInstance)
 		return;
 	}
 
+	bEnabled = SkelComp->GetAllowAnimDynamicsAnimNode();
+
 	const UWorld* World = SkelComp->GetWorld();
 
 	if(CVarEnableWind.GetValueOnAnyThread() == 1 && bEnableWind)
diff --git a/Engine/Source/Runtime/AnimGraphRuntime/Public/BoneControllers/AnimNode_AnimDynamics.h b/Engine/Source/Runtime/AnimGraphRuntime/Public/BoneControllers/AnimNode_AnimDynamics.h
--- a/Engine/Source/Runtime/AnimGraphRuntime/Public/BoneControllers/AnimNode_AnimDynamics.h
+++ b/Engine/Source/Runtime/AnimGraphRuntime/Public/BoneControllers/AnimNode_AnimDynamics.h
@@ -562,7 +562,7 @@ struct FAnimNode_AnimDynamics : public FAnimNode_SkeletalControlBase
 	ANIMGRAPHRUNTIME_API virtual void UpdateInternal(const FAnimationUpdateContext& Context) override;
 	ANIMGRAPHRUNTIME_API virtual void EvaluateSkeletalControl_AnyThread(FComponentSpacePoseContext& Output, TArray<FBoneTransform>& OutBoneTransforms) override;
 	ANIMGRAPHRUNTIME_API virtual void GatherDebugData(FNodeDebugData& DebugData) override;
-	ANIMGRAPHRUNTIME_API virtual bool HasPreUpdate() const override;
+	ANIMGRAPHRUNTIME_API virtual bool HasPreUpdate() const override { return true; }
 	ANIMGRAPHRUNTIME_API virtual void PreUpdate(const UAnimInstance* InAnimInstance) override;
 	virtual bool NeedsDynamicReset() const override { return true; }
 	virtual void ResetDynamics(ETeleportType InTeleportType) override { RequestInitialise(InTeleportType); }
@@ -682,6 +682,9 @@ private:
 	// Previous angular velocity to resolve world angular acceleration when not using world space simulation
 	FVector PreviousSimSpaceAngularVelocity;
 
+	// When true, physics simulation is enabled.
+	uint8 bEnabled : 1;
+
 	//////////////////////////////////////////////////////////////////////////
 	// Live debug
 	//////////////////////////////////////////////////////////////////////////
diff --git a/Engine/Source/Runtime/Engine/Classes/Components/SkeletalMeshComponent.h b/Engine/Source/Runtime/Engine/Classes/Components/SkeletalMeshComponent.h
--- a/Engine/Source/Runtime/Engine/Classes/Components/SkeletalMeshComponent.h
+++ b/Engine/Source/Runtime/Engine/Classes/Components/SkeletalMeshComponent.h
@@ -696,10 +696,15 @@ public:
 	int32 DeferredKinematicUpdateIndex;
 
 private:
+
 	/** Disable rigid body animation nodes and play original animation without simulation */
 	UPROPERTY(EditAnywhere, Category = Physics)
 	uint8 bDisableRigidBodyAnimNode:1;
 
+	/** Disable Anim Dynamics animation nodes and play original animation without simulation */
+	UPROPERTY(EditAnywhere, Category = Physics)
+	uint8 bDisableAnimDynamicsAnimNode:1;
+
 	/** Disable animation curves for this component. If this is set true, no curves will be processed */
 	UPROPERTY(EditAnywhere, AdvancedDisplay, Category = SkeletalMesh)
 	uint8 bAllowAnimCurveEvaluation : 1;
@@ -1362,6 +1367,21 @@ public:
 	UFUNCTION(BlueprintCallable, Category = "Components|SkeletalMesh")
 	bool GetAllowRigidBodyAnimNode() const { return !bDisableRigidBodyAnimNode; }
 
+	/**
+	 * Sets whether or not to allow Anim Dynamics animation nodes for this component
+	 */
+	UFUNCTION(BlueprintCallable, Category = "Components|SkeletalMesh")
+	ENGINE_API void SetAllowAnimDynamicsAnimNode(bool bInAllow, bool bReinitAnim = true);
+
+	UFUNCTION(BlueprintCallable, Category = "Components|SkeletalMesh")
+	bool GetAllowAnimDynamicsAnimNode() const { return !bDisableAnimDynamicsAnimNode; }
+
+	/**
+	 * Initialise animation nodes.
+	 */
+	UFUNCTION(BlueprintCallable, Category = "Components|SkeletalMesh")
+	ENGINE_API void ReinitializeAnimNodes();
+
 	/**
 	* Sets whether or not to force tick component in order to update animation and refresh transform for this component
 	* This is supported only in the editor
diff --git a/Engine/Source/Runtime/Engine/Private/Components/SkeletalMeshComponent.cpp b/Engine/Source/Runtime/Engine/Private/Components/SkeletalMeshComponent.cpp
--- a/Engine/Source/Runtime/Engine/Private/Components/SkeletalMeshComponent.cpp
+++ b/Engine/Source/Runtime/Engine/Private/Components/SkeletalMeshComponent.cpp
@@ -491,6 +491,7 @@ USkeletalMeshComponent::USkeletalMeshComponent(const FObjectInitializer& ObjectI
 
 	bAllowClothActors = true;
 	bPostEvaluatingAnimation = false;
+	bDisableAnimDynamicsAnimNode = false;
 	bAllowAnimCurveEvaluation = true;
 	bDisablePostProcessBlueprint = false;
 
@@ -4853,23 +4854,45 @@ void USkeletalMeshComponent::UnbindClothFromLeaderPoseComponent(bool bRestoreSim
 	}
 }
 
+void USkeletalMeshComponent::ReinitializeAnimNodes()
+{
+	if (bRegistered && GetSkeletalMeshAsset() != nullptr)
+	{
+		if (AnimScriptInstance)
+		{
+			AnimScriptInstance->InitializeAnimation();
+		}
+		if (PostProcessAnimInstance)
+		{
+			PostProcessAnimInstance->InitializeAnimation();
+		}
+	}
+}
+
 void USkeletalMeshComponent::SetAllowRigidBodyAnimNode(bool bInAllow, bool bReinitAnim)
 {
-	if(bDisableRigidBodyAnimNode == bInAllow)
+	if (bDisableRigidBodyAnimNode == bInAllow)
 	{
 		bDisableRigidBodyAnimNode = !bInAllow;
 
-		if(bReinitAnim && bRegistered && GetSkeletalMeshAsset() != nullptr)
+		if(bReinitAnim)
 		{
 			// need to reinitialize rigid body nodes for new setting to take effect
-			if (AnimScriptInstance)
-			{
-				AnimScriptInstance->InitializeAnimation();
-			}
-			if (PostProcessAnimInstance)
-			{
-				PostProcessAnimInstance->InitializeAnimation();
-			}
+			ReinitializeAnimNodes();
+		}
+	}
+}
+
+void USkeletalMeshComponent::SetAllowAnimDynamicsAnimNode(bool bInAllow, bool bReinitAnim)
+{
+	if (bDisableAnimDynamicsAnimNode == bInAllow)
+	{
+		bDisableAnimDynamicsAnimNode = !bInAllow;
+
+		if (bReinitAnim)
+		{
+			// need to reinitialize rigid body nodes for new setting to take effect
+			ReinitializeAnimNodes();
 		}
 	}
 }
