Change ensure condition to prevent false warning: "Expression 'DependentRefCount > 0' is always false"

#rnx

diff --git a/Engine/Source/Runtime/Engine/Private/Net/Iris/ReplicationSystem/NetRootObjectAdapter.cpp b/Engine/Source/Runtime/Engine/Private/Net/Iris/ReplicationSystem/NetRootObjectAdapter.cpp
--- a/Engine/Source/Runtime/Engine/Private/Net/Iris/ReplicationSystem/NetRootObjectAdapter.cpp
+++ b/Engine/Source/Runtime/Engine/Private/Net/Iris/ReplicationSystem/NetRootObjectAdapter.cpp
@@ -383,7 +383,7 @@ void FNetRootObjectAdapter::RemoveRelevantWith(UObject* OtherRootObject)
 	}
 	else
 	{
-		ensureMsgf(DependentRefCount > 0, TEXT("Error found in dependency refcount between %s and %s."), *GetNameSafe(ReplicatedObject), *GetNameSafe(OtherRootObject));
+		ensureMsgf(false, TEXT("Error with mismatched RelevantWith/RemoveRelevantWith calls. DependentRefCount is already at 0 when removing relevancy between %s and %s."), *GetNameSafe(ReplicatedObject), *GetNameSafe(OtherRootObject));
 		// Let's try to remove this dependency anyway in case the error was in a previous call.
 	}
 
