Fix an issue where the AttachParent was getting nulled out during cooking.

#tests Created a minimal code only build and tested cooking a project that reproed the issue in VKEdit and with VKPlay using a private code.
#tests Ran change locally with additional logging to verity the component still had a valid AttachParent after cooking.
#rnx

diff --git a/Engine/Source/Runtime/Engine/Private/Level.cpp b/Engine/Source/Runtime/Engine/Private/Level.cpp
--- a/Engine/Source/Runtime/Engine/Private/Level.cpp
+++ b/Engine/Source/Runtime/Engine/Private/Level.cpp
@@ -141,6 +141,12 @@ static FAutoConsoleVariableRef CVarAllowIncrementalPreUnregisterComponents(
 	GAllowIncrementalPreUnregisterComponents,
 	TEXT("Whether to allow level streaming incremental pre-unregistration of actor components."));
 
+bool GAlwaysCallDetachWhenRemovingActors = false;
+static FAutoConsoleVariableRef CVarAlwaysCallDetachWhenRemovingActors(
+	TEXT("LevelStreaming.AlwaysCallDetachWhenRemovingActors"),
+	GAlwaysCallDetachWhenRemovingActors,
+	TEXT("Whether to always call DetachFromActor when removing loaded actors from the level."));
+
 #if WITH_EDITOR
 
 TArray<ULevel::FLevelExternalActorsPathsProviderDelegate> ULevel::LevelExternalActorsPathsProviders;
@@ -1200,7 +1206,13 @@ void ULevel::RemoveLoadedActors(const TArray<AActor*>& ActorList, const FTransfo
 		// UnregisterAllComponents can destroy child actors
 		if (IsValid(Actor))
 		{
-			Actor->DetachFromActor(FDetachmentTransformRules(EDetachmentRule::KeepWorld, /*bCallModify=*/false));
+			//Prevent the DetachFromActor call when cooking.
+			//While validating content, actors are added and removed from world partition.
+			//If we detach from our parent that change can be saved out in the cooked data causing the actor to be permanently detached.
+			if (GAlwaysCallDetachWhenRemovingActors || !IsRunningCookCommandlet())
+			{
+				Actor->DetachFromActor(FDetachmentTransformRules(EDetachmentRule::KeepWorld, /*bCallModify=*/false));
+			}
 			Actor->UnregisterAllComponents();
 			Actor->RegisterAllActorTickFunctions(false, true);
 			Actor->OnLoadedActorRemovedFromLevel();
