[WorldPartition] Fix overwirting key in map leading to inability to find/remove elements
* Fixed up some coding standard violations

#tests Replace actor reference that was causing the issue to begin with

diff --git a/Engine/Source/Runtime/Engine/Public/ExternalDirtyActorsTracker.h b/Engine/Source/Runtime/Engine/Public/ExternalDirtyActorsTracker.h
--- a/Engine/Source/Runtime/Engine/Public/ExternalDirtyActorsTracker.h
+++ b/Engine/Source/Runtime/Engine/Public/ExternalDirtyActorsTracker.h
@@ -19,7 +19,8 @@ class TExternalDirtyActorsTracker : public FTickableEditorObject
 {
 public:
 	using Super = TExternalDirtyActorsTracker<StoreType>;
-	using MapType = TMap<TWeakObjectPtr<AActor>, typename StoreType::Type>;
+	using ValueType = typename StoreType::Type;
+	using MapType = TMap<TWeakObjectPtr<AActor>, ValueType>;
 
 	TExternalDirtyActorsTracker(const ULevel* InLevel, typename StoreType::OwnerType* InOwner)
 		: Level(InLevel)
@@ -41,8 +42,8 @@ public:
 
 	/** Returns false if tracking this actor in unnecessary */
 	virtual bool OnAddDirtyActor(const TWeakObjectPtr<AActor> InActor) { return true; }
-	virtual void OnRemoveInvalidDirtyActor(const TWeakObjectPtr<AActor> InActor, typename StoreType::Type& InValue) {}
-	virtual void OnRemoveNonDirtyActor(const TWeakObjectPtr<AActor> InActor, typename StoreType::Type& InValue) {}
+	virtual void OnRemoveInvalidDirtyActor(const TWeakObjectPtr<AActor> InActor, ValueType& InValue) {}
+	virtual void OnRemoveNonDirtyActor(const TWeakObjectPtr<AActor> InActor, ValueType& InValue) {}
 
 protected:
 	void OnPackageDirtyStateChanged(UPackage* InPackage)
@@ -62,7 +63,7 @@ protected:
 					}
 					else
 					{
-						typename StoreType::Type Value;
+						ValueType Value;
 						if (DirtyActors.RemoveAndCopyValue(Actor, Value))
 						{
 							if (IsValid(Actor))
@@ -82,11 +83,11 @@ protected:
 
 	void OnObjectsReplaced(const TMap<UObject*, UObject*>& OldToNewObjectMap)
 	{
-		for (auto [OldObject, NewObject] : OldToNewObjectMap)
+		for (const TPair<UObject*, UObject*>& Itr : OldToNewObjectMap)
 		{
-			if (AActor* OldActor = Cast<AActor>(OldObject))
+			if (AActor* OldActor = Cast<AActor>(Itr.Key))
 			{
-				if (AActor* NewActor = Cast<AActor>(NewObject))
+				if (AActor* NewActor = Cast<AActor>(Itr.Value))
 				{
 					OnEditorActorReplaced(OldActor, NewActor);
 				}
@@ -96,12 +97,10 @@ protected:
 
 	void OnEditorActorReplaced(AActor* InOldActor, AActor* InNewActor)
 	{
-		for (auto& [WeakActor, Value] : DirtyActors)
+		ValueType ValueToCopy;
+		if (DirtyActors.RemoveAndCopyValue(InOldActor, ValueToCopy))
 		{
-			if (WeakActor.IsValid() && WeakActor.Get() == InOldActor)
-			{
-				WeakActor = InNewActor;
-			}
+			DirtyActors.Add(InNewActor, ValueToCopy);
 		}
 	}
 
