Substrate - fix for material preview on nodes.
Make and Set material attribute nodes now have a preview.
Fixes material function output to no have any compilation errors and work with substrate as well as with material attributes.

#rnx
#tests PC editor

diff --git a/Engine/Source/Runtime/Engine/Private/Materials/MaterialExpressions.cpp b/Engine/Source/Runtime/Engine/Private/Materials/MaterialExpressions.cpp
--- a/Engine/Source/Runtime/Engine/Private/Materials/MaterialExpressions.cpp
+++ b/Engine/Source/Runtime/Engine/Private/Materials/MaterialExpressions.cpp
@@ -6178,6 +6178,18 @@ EMaterialValueType UMaterialExpressionMakeMaterialAttributes::GetInputValueType(
 	}
 }
 
+int32 UMaterialExpressionMakeMaterialAttributes::CompilePreview(class FMaterialCompiler* Compiler, int32 OutputIndex)
+{
+	// Material attributes "Compile" will return the value of the current attribute at the top of the compiler stack. By default it is EmissiveColor, so can be just black.
+	// Instead we use BaseColor as a better preview representative for material attributes.
+	FExpressionInput BaseColorTraced = BaseColor.GetTracedInput();
+	UMaterialExpression* BaseColorExpression = BaseColor.Expression;
+	if (BaseColorExpression)
+	{
+		return BaseColor.Compile(Compiler); 
+	}
+	return Compiler->Constant3(0, 0, 0);
+}
 #endif // WITH_EDITOR
 
 // -----
@@ -6980,6 +6992,47 @@ void UMaterialExpressionSetMaterialAttributes::PostEditChangeProperty(FPropertyC
 
 	Super::PostEditChangeProperty( PropertyChangedEvent );
 }
+
+int32 UMaterialExpressionSetMaterialAttributes::CompilePreview(class FMaterialCompiler* Compiler, int32 OutputIndex)
+{
+	// Use the new base color as preview if specified
+	int32 InputIndex = 0;
+	for (const FGuid& AttributeType : AttributeSetTypes)
+	{
+		EMaterialProperty MaterialProperty = FMaterialAttributeDefinitionMap::GetProperty(AttributeType);
+		if (MaterialProperty == MP_BaseColor && Inputs.Num() > (InputIndex + 1))
+		{
+			FExpressionInput BaseColorInput = Inputs[InputIndex + 1]; // +1 because first element is always material attribute input
+
+			// Trace the base color expression to use as preview
+			FExpressionInput BaseColorTraced = BaseColorInput.GetTracedInput();
+			UMaterialExpression* BaseColorExpression = BaseColorTraced.Expression;
+			if (BaseColorExpression)
+			{
+				return BaseColorInput.Compile(Compiler);
+			}
+		}
+		InputIndex++;
+	}
+
+	// Otherwise display the one from the material attribute input
+	if(Inputs.Num() > 0)
+	{
+		FExpressionInput MaterialAttributeInput = Inputs[0];
+		FExpressionInput MaterialAttributeTraced = MaterialAttributeInput.GetTracedInput();
+		UMaterialExpression* MaterialAttributeExpression = MaterialAttributeTraced.Expression;
+		if (MaterialAttributeExpression)
+		{
+			// Material attributes "Compile" will return the value of the current attribute at the top of the compiler stack. By default it is EmissiveColor, so can be just black.
+			// Instead we use BaseColor as a better preview representative for material attributes.
+			return MaterialAttributeExpression->CompilePreview(Compiler, MaterialAttributeInput.OutputIndex);
+		}
+	}
+
+	// Otherwise black color
+	return Compiler->Constant3(0, 0, 0);
+}
+
 #endif // WITH_EDITOR
 
 // -----
@@ -16843,19 +16896,27 @@ int32 UMaterialExpressionFunctionOutput::CompilePreview(class FMaterialCompiler*
 	FExpressionInput ATraced = A.GetTracedInput();
 	UMaterialExpression* AExpression = ATraced.Expression;
 
-	if (Substrate::IsSubstrateEnabled() && AExpression)
+	// If this logic is changed, one need to test 
+	//  - material function preview on nodes in materials, 
+	//  - and in material functions with output node being SubstrateData (a slab) or a color,
+	//  - and for material with Surface or UI domain preview.
+	//  - This for r.substrate 1 and r.Substrate.ProjectGBufferFormat 0/1.
+	if (Substrate::IsSubstrateEnabled() && AExpression && AExpression->IsResultSubstrateMaterial(ATraced.OutputIndex))
 	{
 		// Compile the SubstrateData output.
 		int32 SubstrateDataCodeChunk = Compile(Compiler, ATraced.OutputIndex);
-		if (!AExpression->IsResultSubstrateMaterial(ATraced.OutputIndex))
-		{
-			SubstrateDataCodeChunk = UMaterialExpressionSubstrateConvertMaterialAttributes::CompileDefaultSlab(Compiler, SubstrateDataCodeChunk);			
-		}
 		// Convert the SubstrateData to a preview color.
 		return Compiler->SubstrateCompilePreview(SubstrateDataCodeChunk);
 	}
 
-	// Compile the preview value, outputting a float type
+	// If this is a material attribute node, let's have it decide of the preview 
+	// (instead of using current attribute on the compiler stack which is not a good representative, default to EmissiveColor).
+	if (AExpression && AExpression->IsResultMaterialAttributes(ATraced.OutputIndex))
+	{
+		return AExpression->CompilePreview(Compiler, ATraced.OutputIndex);
+	}
+
+	// Default preview value compilation.
 	return Compile(Compiler, OutputIndex);
 }
 #endif // WITH_EDITOR
diff --git a/Engine/Source/Runtime/Engine/Public/Materials/MaterialExpressionMakeMaterialAttributes.h b/Engine/Source/Runtime/Engine/Public/Materials/MaterialExpressionMakeMaterialAttributes.h
--- a/Engine/Source/Runtime/Engine/Public/Materials/MaterialExpressionMakeMaterialAttributes.h
+++ b/Engine/Source/Runtime/Engine/Public/Materials/MaterialExpressionMakeMaterialAttributes.h
@@ -90,6 +90,7 @@ public:
 	virtual void GetCaption(TArray<FString>& OutCaptions) const override;
 	virtual bool IsResultMaterialAttributes(int32 OutputIndex) override {return true;}
 	virtual EMaterialValueType GetInputValueType(int32 InputIndex) override;
+	virtual int32 CompilePreview(class FMaterialCompiler* Compiler, int32 OutputIndex) override;
 
 	uint64 GetConnectedInputs() const;
 #endif // WITH_EDITOR
diff --git a/Engine/Source/Runtime/Engine/Public/Materials/MaterialExpressionSetMaterialAttributes.h b/Engine/Source/Runtime/Engine/Public/Materials/MaterialExpressionSetMaterialAttributes.h
--- a/Engine/Source/Runtime/Engine/Public/Materials/MaterialExpressionSetMaterialAttributes.h
+++ b/Engine/Source/Runtime/Engine/Public/Materials/MaterialExpressionSetMaterialAttributes.h
@@ -43,6 +43,7 @@ class UMaterialExpressionSetMaterialAttributes : public UMaterialExpression
 	virtual EMaterialValueType GetInputValueType(int32 InputIndex) override;
 	virtual bool IsInputConnectionRequired(int32 InputIndex) const override {return true;}
 	virtual bool IsResultMaterialAttributes(int32 OutputIndex) override {return true;}
+	virtual int32 CompilePreview(class FMaterialCompiler* Compiler, int32 OutputIndex) override;
 
 	ENGINE_API int32 CreateOrGetInputAttribute(EMaterialProperty Attribute);
 	ENGINE_API bool ConnectInputAttribute(EMaterialProperty Attribute, UMaterialExpression* Expression, int32 OutputIndex = 0);
