Fix comments for SetBoneTrackKeys in IAnimationDataController

diff --git a/Engine/Source/Runtime/Engine/Classes/Animation/AnimData/IAnimationDataController.h b/Engine/Source/Runtime/Engine/Classes/Animation/AnimData/IAnimationDataController.h
--- a/Engine/Source/Runtime/Engine/Classes/Animation/AnimData/IAnimationDataController.h
+++ b/Engine/Source/Runtime/Engine/Classes/Animation/AnimData/IAnimationDataController.h
@@ -224,7 +224,7 @@ public:
 	virtual void RemoveAllBoneTracks(bool bShouldTransact = true) = 0;
 
 	/**
-	* Removes an existing bone animation track with the provided name. Broadcasts a EAnimDataModelNotifyType::TrackChanged notify if successful.
+	* Sets keys for an existing bone animation track with the provided name. Broadcasts a EAnimDataModelNotifyType::TrackChanged notify if successful.
 	* The provided number of keys provided is expected to match for each component, and be non-zero.
 	*
 	* @param	BoneName			Bone name of the track for which the keys should be set
@@ -238,7 +238,7 @@ public:
 	virtual bool SetBoneTrackKeys(FName BoneName, const TArray<FVector3f>& PositionalKeys, const TArray<FQuat4f>& RotationalKeys, const TArray<FVector3f>& ScalingKeys, bool bShouldTransact = true) = 0;
 
 	/**
-	* Removes an existing bone animation track with the provided name. Broadcasts a EAnimDataModelNotifyType::TrackChanged notify if successful.
+	* Sets keys for an existing bone animation track with the provided name. Broadcasts a EAnimDataModelNotifyType::TrackChanged notify if successful.
 	* The provided number of keys provided is expected to match for each component, and be non-zero.
 	*
 	* @param	BoneName			Bone name of the track for which the keys should be set
