Fix double truncation warnings.

diff --git a/Engine/Source/Runtime/Engine/Public/HeterogeneousVolumeInterface.h b/Engine/Source/Runtime/Engine/Public/HeterogeneousVolumeInterface.h
--- a/Engine/Source/Runtime/Engine/Public/HeterogeneousVolumeInterface.h
+++ b/Engine/Source/Runtime/Engine/Public/HeterogeneousVolumeInterface.h
@@ -51,12 +51,12 @@ public:
 		: PrimitiveSceneProxy(SceneProxy)
 		, InstanceToLocal(FMatrix::Identity)
 		, VoxelResolution(FIntVector::ZeroValue)
-		, MinimumVoxelSize(0.1)
-		, StepFactor(1.0)
-		, ShadowStepFactor(8.0)
-		, ShadowBiasFactor(0.0)
-		, LightingDownsampleFactor(1.0)
-		, MipBias(0.0)
+		, MinimumVoxelSize(0.1f)
+		, StepFactor(1.0f)
+		, ShadowStepFactor(8.0f)
+		, ShadowBiasFactor(0.0f)
+		, LightingDownsampleFactor(1.0f)
+		, MipBias(0.0f)
 		, bPivotAtCentroid(false)
 		, bHoldout(false)
 	{}
@@ -65,12 +65,12 @@ public:
 		: PrimitiveSceneProxy(SceneProxy)
 		, InstanceToLocal(FMatrix::Identity)
 		, VoxelResolution(FIntVector::ZeroValue)
-		, MinimumVoxelSize(0.1)
-		, StepFactor(1.0)
-		, ShadowStepFactor(8.0)
-		, ShadowBiasFactor(0.0)
-		, LightingDownsampleFactor(1.0)
-		, MipBias(0.0)
+		, MinimumVoxelSize(0.1f)
+		, StepFactor(1.0f)
+		, ShadowStepFactor(8.0f)
+		, ShadowBiasFactor(0.0f)
+		, LightingDownsampleFactor(1.0f)
+		, MipBias(0.0f)
 		, bPivotAtCentroid(false)
 		, bHoldout(false)
 #if ACTOR_HAS_LABELS
