Pose Asset: Adapting API to allow interaction with the PoseAsset outside of AnimBP

#rnx

diff --git a/Engine/Source/Runtime/Engine/Classes/Animation/PoseAsset.h b/Engine/Source/Runtime/Engine/Classes/Animation/PoseAsset.h
--- a/Engine/Source/Runtime/Engine/Classes/Animation/PoseAsset.h
+++ b/Engine/Source/Runtime/Engine/Classes/Animation/PoseAsset.h
@@ -78,6 +78,18 @@ struct FPoseAssetInfluences
 	TArray<FPoseAssetInfluence> Influences;
 };
 
+USTRUCT()
+struct FPoseAssetInfluenceForPose
+{
+	GENERATED_USTRUCT_BODY()
+
+	UPROPERTY()
+	int32 TrackIndex = INDEX_NONE;
+
+	UPROPERTY()
+	int32 BoneTransformIndex = INDEX_NONE;
+};
+
 /**
 * Pose data container
 * 
@@ -235,7 +247,7 @@ public:
 	*/
 	UE_DEPRECATED(4.26, "Use GetAnimationPose with other signature")
 	ENGINE_API bool GetAnimationPose(struct FCompactPose& OutPose, FBlendedCurve& OutCurve, const FAnimExtractContext& ExtractionContext) const;
-	ENGINE_API bool GetAnimationPose(struct FAnimationPoseData& OutAnimationPoseData, const FAnimExtractContext& ExtractionContext, const bool bIgnoreSourceAnimationCurves = false) const;
+	ENGINE_API bool GetAnimationPose(struct FAnimationPoseData& OutAnimationPoseData, const FAnimExtractContext& ExtractionContext, const bool bIgnoreSourceAnimationCurves = false, const bool bAllowNegativeWeights = false) const;
 
 	UE_DEPRECATED(4.26, "Use GetBaseAnimationPose with other signature")
 	ENGINE_API void GetBaseAnimationPose(struct FCompactPose& OutPose, FBlendedCurve& OutCurve) const;
@@ -276,6 +288,14 @@ public:
 	ENGINE_API const TArray<FAnimCurveBase>& GetCurveData() const;
 	ENGINE_API const TArray<float>& GetCurveValues(const int32 PoseIndex) const;
 
+#if WITH_EDITOR
+	// Return full / absolute curve values. Will do conversion if PoseAsset is Additive. 
+	ENGINE_API const TArray<float>& GetFullCurveValues(const int32 PoseIndex) const;
+#endif
+	// Return relative / additive curve values. Will give access to raw data and won't do conversion.
+	// This is the same as GetCurveValues
+	ENGINE_API const TArray<float>& GetRawCurveValues(const int32 PoseIndex) const;
+
 	/** Find index of a track with a given bone name. Returns INDEX_NONE if not found. */
 	ENGINE_API const int32 GetTrackIndexByName(const FName& InTrackName) const;
 
@@ -357,7 +377,14 @@ public:
 	virtual void PostEditChangeProperty(FPropertyChangedEvent& PropertyChangedEvent) override;
 	// Return full (local space, non additive) pose. Will do conversion if PoseAsset is Additive. 
 	ENGINE_API bool GetFullPose(int32 PoseIndex, TArray<FTransform>& OutTransforms) const;
-	
+	ENGINE_API const TArray<FTransform>& GetFullPose(int32 PoseIndex) const;
+#endif
+
+	// Return local (local space, additive) pose. Will give access to raw data and won't do conversion 
+	ENGINE_API bool GetRawPose(int32 PoseIndex, TArray<FTransform>& OutTransforms) const;
+	ENGINE_API const TArray<FTransform>& GetRawPose(int32 PoseIndex) const;
+
+#if WITH_EDITOR
 	// util to return transform of a bone from the pose asset in component space, by walking up tracks in pose asset */
 	ENGINE_API FTransform GetComponentSpaceTransform(FName BoneName, const TArray<FTransform>& LocalTransforms) const;
 
@@ -375,6 +402,13 @@ public:
 	ENGINE_API const int32 GetPoseIndexByName(const FName& InBasePoseName) const;
 	ENGINE_API const int32 GetCurveIndexByName(const FName& InCurveName) const;
 
+	ENGINE_API bool GetBasePoseTransform(TArray<FTransform>& OutBasePose, TArray<float>& OutBaseCurve) const;
+	ENGINE_API bool GetTrackInfluencesForPose(const int32 PoseIndex, TArray<FPoseAssetInfluenceForPose>& OutInfluences) const;
+	ENGINE_API bool GetTrackInfluencesForPoses(const TArrayView<const FPoseCurve>& InPoses, TArray<FPoseAssetInfluenceForPose>& OutInfluences) const;
+	ENGINE_API bool ConvertToFullPose(const int32 PoseIndex, TArray<FTransform>& OutFullPose, TArray<FPoseAssetInfluenceForPose>& OutInfluences, TArray<float>& OutFullCurves) const;
+	ENGINE_API bool ConvertToFullPose(const int32 PoseIndex, TArray<FTransform>& OutFullPose, TArray<FPoseAssetInfluenceForPose>& OutInfluences) const;
+	ENGINE_API bool ConvertToFullPose(const int32 PoseIndex, TArray<float>& OutFullCurves) const;
+
 #if WITH_EDITOR
 private: 
 	DECLARE_MULTICAST_DELEGATE(FOnPoseListChangedMulticaster)
@@ -401,7 +435,7 @@ public:
 	
 	ENGINE_API static FName GetUniquePoseName(UPoseAsset* PoseAsset);
 	ENGINE_API FGuid GetSourceAnimationGuid() const;
-
+	
 protected:
 	virtual void RemapTracksToNewSkeleton(USkeleton* NewSkeleton, bool bConvertSpaces) override;
 private: 
@@ -411,7 +445,6 @@ private:
 
 	void ConvertToFullPose();
 	void ConvertToAdditivePose(int32 NewBasePoseIndex);
-	bool GetBasePoseTransform(TArray<FTransform>& OutBasePose, TArray<float>& OutCurve) const;
 	void Reinitialize();
 
 	// After any update to SourceLocalPoses, this does update runtime data
diff --git a/Engine/Source/Runtime/Engine/Private/Animation/PoseAsset.cpp b/Engine/Source/Runtime/Engine/Private/Animation/PoseAsset.cpp
--- a/Engine/Source/Runtime/Engine/Private/Animation/PoseAsset.cpp
+++ b/Engine/Source/Runtime/Engine/Private/Animation/PoseAsset.cpp
@@ -604,7 +604,7 @@ bool UPoseAsset::GetAnimationPose(struct FCompactPose& OutPose, FBlendedCurve& O
 	return GetAnimationPose(OutPoseData, ExtractionContext);
 }
 
-bool UPoseAsset::GetAnimationPose(struct FAnimationPoseData& OutAnimationPoseData, const FAnimExtractContext& ExtractionContext, const bool bIgnoreSourceAnimationCurves) const
+bool UPoseAsset::GetAnimationPose(struct FAnimationPoseData& OutAnimationPoseData, const FAnimExtractContext& ExtractionContext, const bool bIgnoreSourceAnimationCurves, const bool bAllowNegativeWeights) const
 {
 	ANIM_MT_SCOPE_CYCLE_COUNTER(PoseAssetGetAnimationPose, !IsInGameThread());
 
@@ -641,9 +641,10 @@ bool UPoseAsset::GetAnimationPose(struct FAnimationPoseData& OutAnimationPoseDat
 				const FPoseData& Pose = PoseContainer.Poses[PoseIndex];
 				// Clamp weight for non-additive pose assets rather than normalizing the weight
 				const float Weight = bAdditivePose ? Curve.Value : FMath::Clamp(Curve.Value, 0.f, 1.f);
+				const float WeightToCheck = bAllowNegativeWeights ? FMath::Abs(Weight) : Weight;
 
 				// Only generate pose if the single weight is actually relevant
-				if(FAnimWeight::IsRelevant(Weight))
+				if(FAnimWeight::IsRelevant(WeightToCheck))
 				{
 					if (!bIgnoreSourceAnimationCurves)
 					{
@@ -753,7 +754,8 @@ bool UPoseAsset::GetAnimationPose(struct FAnimationPoseData& OutAnimationPoseDat
 
 					// we only add to the list if it's not additive Or if it's additive, we don't want to add base pose index
 					// and has weight
-					if ((!bAdditivePose || PoseIndex != BasePoseIndex) && FAnimationRuntime::HasWeight(Value))
+					const float WeightToCheck = bAllowNegativeWeights ? FMath::Abs(Value) : Value;
+					if ((!bAdditivePose || PoseIndex != BasePoseIndex) && FAnimationRuntime::HasWeight(WeightToCheck))
 					{
 						TotalWeight += Value;
 
@@ -837,7 +839,7 @@ bool UPoseAsset::GetAnimationPose(struct FAnimationPoseData& OutAnimationPoseDat
 							if(WeightedPoses[PoseIndex])
 							{
 								const float& Weight = PoseWeights[PoseIndex];
-								TotalLocalWeight += Weight;
+								TotalLocalWeight += FMath::Abs(Weight);
 							
 								const FPoseData& PoseData = PoseContainer.Poses[PoseIndex];
 								const FTransform& BonePose = PoseData.LocalSpacePose[BonePoseIndex];
@@ -1117,6 +1119,26 @@ const TArray<FAnimCurveBase>& UPoseAsset::GetCurveData() const
 }
 
 const TArray<float>& UPoseAsset::GetCurveValues(const int32 PoseIndex) const
+{
+	return GetRawCurveValues(PoseIndex);
+}
+
+#if WITH_EDITOR
+
+const TArray<float>& UPoseAsset::GetFullCurveValues(const int32 PoseIndex) const
+{
+	if (PoseContainer.Poses.IsValidIndex(PoseIndex))
+	{
+		return PoseContainer.Poses[PoseIndex].SourceCurveData;
+	}
+
+	static const TArray<float> EmptyCurveValues;
+	return EmptyCurveValues;
+}
+
+#endif
+
+const TArray<float>& UPoseAsset::GetRawCurveValues(const int32 PoseIndex) const
 {
 	if (PoseContainer.Poses.IsValidIndex(PoseIndex))
 	{
@@ -1703,6 +1725,44 @@ bool UPoseAsset::GetFullPose(int32 PoseIndex, TArray<FTransform>& OutTransforms)
 	return true;
 }
 
+const TArray<FTransform>& UPoseAsset::GetFullPose(int32 PoseIndex) const
+{
+	if (PoseContainer.Poses.IsValidIndex(PoseIndex))
+	{
+		return PoseContainer.Poses[PoseIndex].SourceLocalSpacePose;
+	}
+
+	static const TArray<FTransform> EmptyTransforms;
+	return EmptyTransforms;
+}
+
+#endif
+
+bool UPoseAsset::GetRawPose(int32 PoseIndex, TArray<FTransform>& OutTransforms) const
+{
+	if (!PoseContainer.Poses.IsValidIndex(PoseIndex))
+	{
+		return false;
+	}
+
+	// just return source data
+	OutTransforms = PoseContainer.Poses[PoseIndex].LocalSpacePose;
+	return true;
+}
+
+const TArray<FTransform>& UPoseAsset::GetRawPose(int32 PoseIndex) const
+{
+	if (PoseContainer.Poses.IsValidIndex(PoseIndex))
+	{
+		return PoseContainer.Poses[PoseIndex].LocalSpacePose;
+	}
+
+	static const TArray<FTransform> EmptyTransforms;
+	return EmptyTransforms;
+}
+
+#if WITH_EDITOR
+
 FTransform UPoseAsset::GetComponentSpaceTransform(FName BoneName, const TArray<FTransform>& LocalTransforms) const
 {
 	const FReferenceSkeleton& RefSkel = GetSkeleton()->GetReferenceSkeleton();
@@ -2009,7 +2069,9 @@ void UPoseAsset::ReplaceReferredAnimations(const TMap<UAnimationAsset*, UAnimati
 	}
 }
 
-bool UPoseAsset::GetBasePoseTransform(TArray<FTransform>& OutBasePose, TArray<float>& OutCurve) const
+#endif
+
+bool UPoseAsset::GetBasePoseTransform(TArray<FTransform>& OutBasePose, TArray<float>& OutBaseCurve) const
 {
 	int32 TotalNumTrack = PoseContainer.Tracks.Num();
 	OutBasePose.Reset(TotalNumTrack);
@@ -2036,21 +2098,152 @@ bool UPoseAsset::GetBasePoseTransform(TArray<FTransform>& OutBasePose, TArray<fl
 		}
 
 		// add zero curves
-		OutCurve.AddZeroed(PoseContainer.Curves.Num());
+		OutBaseCurve.AddZeroed(PoseContainer.Curves.Num());
 		check(OutBasePose.Num() == TotalNumTrack);
 		return true;
 	}
 	else if (PoseContainer.Poses.IsValidIndex(BasePoseIndex))
 	{
+#if WITH_EDITOR
 		OutBasePose = PoseContainer.Poses[BasePoseIndex].SourceLocalSpacePose;
-		OutCurve = PoseContainer.Poses[BasePoseIndex].SourceCurveData;
+		OutBaseCurve = PoseContainer.Poses[BasePoseIndex].SourceCurveData;
+#else
+		TArray<FPoseAssetInfluenceForPose> Influences;
+		ConvertToFullPose(GetBasePoseIndex(), OutBasePose, Influences, OutBaseCurve);
+#endif
 		check(OutBasePose.Num() == TotalNumTrack);
 		return true;
 	}
 
 	return false;
 }
-#endif // WITH_EDITOR
+
+bool UPoseAsset::GetTrackInfluencesForPose(const int32 PoseIndex, TArray<FPoseAssetInfluenceForPose>& OutInfluences) const
+{
+	OutInfluences.Reset();
+
+	// todo: is this right? we may need to do this properly
+	if (PoseIndex == BasePoseIndex && bAdditivePose)
+	{
+		return true;
+	}
+	
+	if (!PoseContainer.Poses.IsValidIndex(PoseIndex))
+	{
+		return false;
+	}
+	for (int32 TrackIndex = 0; TrackIndex < PoseContainer.TrackPoseInfluenceIndices.Num(); ++TrackIndex)
+	{
+		const FPoseAssetInfluences& Influences = PoseContainer.TrackPoseInfluenceIndices[TrackIndex];
+		for (const FPoseAssetInfluence& Influence : Influences.Influences)
+		{
+			if (Influence.PoseIndex == PoseIndex)
+			{
+				const int32 TransformIndex = OutInfluences.Num();
+				OutInfluences.Emplace(TrackIndex, TransformIndex);
+				break;
+			}
+		}
+	}
+	return true;
+}
+
+bool UPoseAsset::GetTrackInfluencesForPoses(const TArrayView<const FPoseCurve>& InPoses, TArray<FPoseAssetInfluenceForPose>& OutInfluences) const
+{
+	OutInfluences.Reset();
+
+	TArray<bool> IncludePose;
+	IncludePose.AddZeroed(GetNumPoses());
+	for (const FPoseCurve& Pose : InPoses)
+	{
+		if (Pose.PoseIndex == BasePoseIndex)
+		{
+			continue;
+		}
+		if (!PoseContainer.Poses.IsValidIndex(Pose.PoseIndex))
+		{
+			continue;
+		}
+		if (!FAnimWeight::IsRelevant(FMath::Abs(Pose.Value)))
+		{
+			continue;
+		}
+		IncludePose[Pose.PoseIndex] = true;
+	}
+
+	for (int32 TrackIndex = 0; TrackIndex < PoseContainer.TrackPoseInfluenceIndices.Num(); ++TrackIndex)
+	{
+		const FPoseAssetInfluences& Influences = PoseContainer.TrackPoseInfluenceIndices[TrackIndex];
+		for (const FPoseAssetInfluence& Influence : Influences.Influences)
+		{
+			if (IncludePose.IsValidIndex(Influence.PoseIndex) && IncludePose[Influence.PoseIndex])
+			{
+				const int32 TransformIndex = OutInfluences.Num();
+				OutInfluences.Emplace(TrackIndex, TransformIndex);
+				break;
+			}
+		}
+	}
+	return true;
+}
+
+bool UPoseAsset::ConvertToFullPose(const int32 PoseIndex, TArray<FTransform>& OutFullPose, TArray<FPoseAssetInfluenceForPose>& OutInfluences, TArray<float>& OutFullCurves) const
+{
+	return
+		ConvertToFullPose(PoseIndex, OutFullPose, OutInfluences) &&
+		ConvertToFullPose(PoseIndex, OutFullCurves);
+}
+
+bool UPoseAsset::ConvertToFullPose(const int32 PoseIndex, TArray<FTransform>& OutFullPose, TArray<FPoseAssetInfluenceForPose>& OutInfluences) const
+{
+	if (!GetTrackInfluencesForPose(PoseIndex, OutInfluences))
+	{
+		return false;
+	}
+
+	if (PoseIndex == BasePoseIndex || BasePoseIndex == INDEX_NONE || !bAdditivePose)
+	{
+		OutFullPose = PoseContainer.Poses[PoseIndex].LocalSpacePose;
+	}
+	else
+	{
+		const TArray<FTransform> BasePose = PoseContainer.Poses[BasePoseIndex].LocalSpacePose;
+		const TArray<FTransform> AdditivePose = PoseContainer.Poses[PoseIndex].LocalSpacePose;
+		OutFullPose.Reset();
+		OutFullPose.Reserve(OutInfluences.Num());
+
+		for (const FPoseAssetInfluenceForPose& Influence : OutInfluences)
+		{
+			FTransform& BoneTransform = OutFullPose.Add_GetRef(BasePose[Influence.TrackIndex]);
+			BoneTransform.AccumulateWithAdditiveScale(AdditivePose[Influence.BoneTransformIndex], ScalarRegister(1));
+		}
+	}
+
+	return true;
+}
+
+bool UPoseAsset::ConvertToFullPose(const int32 PoseIndex, TArray<float>& OutFullCurves) const
+{
+	if (PoseIndex == BasePoseIndex || BasePoseIndex == INDEX_NONE || !bAdditivePose)
+	{
+		OutFullCurves = PoseContainer.Poses[PoseIndex].CurveData;
+	}
+	else
+	{
+		const TArray<float> BaseCurves = PoseContainer.Poses[BasePoseIndex].CurveData;
+		const TArray<float> AdditiveCurves = PoseContainer.Poses[PoseIndex].CurveData;
+		check(BaseCurves.Num() == AdditiveCurves.Num());
+		OutFullCurves.Reset();
+		OutFullCurves.Reserve(BaseCurves.Num());
+
+		for (int32 CurveIndex = 0; CurveIndex < BaseCurves.Num(); ++CurveIndex)
+		{
+			OutFullCurves.Add(BaseCurves[CurveIndex] + AdditiveCurves[CurveIndex]);
+		}
+	}
+
+	return true;
+}
 
 #undef LOCTEXT_NAMESPACE 
 
