WorldPartition: Add Initial Runtime State to Data Layer Instance AR tags

diff --git a/Engine/Source/Runtime/Engine/Private/WorldPartition/DataLayer/DataLayerInstance.cpp b/Engine/Source/Runtime/Engine/Private/WorldPartition/DataLayer/DataLayerInstance.cpp
--- a/Engine/Source/Runtime/Engine/Private/WorldPartition/DataLayer/DataLayerInstance.cpp
+++ b/Engine/Source/Runtime/Engine/Private/WorldPartition/DataLayer/DataLayerInstance.cpp
@@ -84,6 +84,7 @@ namespace DataLayerInstance
 	static const FName NAME_DataLayerInstanceIsIncludedInActorFilterDefault("DataLayerInstanceIsIncludedInActorFilterDefault");
 	static const FName NAME_DataLayerInstancePrivateDataLayerSupportsActorFilter("DataLayerInstancePrivateDataLayerSupportsActorFilter");
 	static const FName NAME_DataLayerInstancePrivateShortName("DataLayerInstancePrivateShortName");
+	static const FName NAME_DataLayerInstanceInitialRuntimeState("DataLayerInstanceInitialRuntimeState");
 };
 
 void UDataLayerInstance::GetAssetRegistryTags(TArray<FAssetRegistryTag>& OutTags) const
@@ -116,6 +117,7 @@ void UDataLayerInstance::GetAssetRegistryTags(FAssetRegistryTagsContext Context)
 		Context.AddTag(FAssetRegistryTag(DataLayerInstance::NAME_DataLayerInstanceIsIncludedInActorFilterDefault, IsIncludedInActorFilterDefault() ? TEXT("1") : TEXT("0"), FAssetRegistryTag::TT_Hidden));
 		Context.AddTag(FAssetRegistryTag(DataLayerInstance::NAME_DataLayerInstancePrivateDataLayerSupportsActorFilter, SupportsActorFilters() ? TEXT("1") : TEXT("0"), FAssetRegistryTag::TT_Hidden));
 		Context.AddTag(FAssetRegistryTag(DataLayerInstance::NAME_DataLayerInstancePrivateShortName, *GetDataLayerShortName(), FAssetRegistryTag::TT_Hidden));
+		Context.AddTag(FAssetRegistryTag(DataLayerInstance::NAME_DataLayerInstanceInitialRuntimeState, *StaticEnum<EDataLayerRuntimeState>()->GetNameStringByValue(static_cast<int64>(InitialRuntimeState)), FAssetRegistryTag::TT_Hidden));
 	}
 }
 
@@ -166,6 +168,10 @@ bool UDataLayerInstance::GetAssetRegistryInfoFromPackage(const FAssetData& InAss
 	{
 		OutDataLayerInstanceDesc.PrivateShortName = Value;
 	}
+	if (InAsset.GetTagValue(DataLayerInstance::NAME_DataLayerInstanceInitialRuntimeState, Value))
+	{
+		OutDataLayerInstanceDesc.InitialRuntimeState = static_cast<EDataLayerRuntimeState>(StaticEnum<EDataLayerRuntimeState>()->GetValueByNameString(Value));
+	}
 	return !OutDataLayerInstanceDesc.Name.IsNone();
 }
 
