SplineTests: exploring functional tests

#virtualized

diff --git a/Engine/Source/Runtime/Engine/Classes/Components/SplineComponent.h b/Engine/Source/Runtime/Engine/Classes/Components/SplineComponent.h
--- a/Engine/Source/Runtime/Engine/Classes/Components/SplineComponent.h
+++ b/Engine/Source/Runtime/Engine/Classes/Components/SplineComponent.h
@@ -1076,7 +1076,7 @@ PRAGMA_ENABLE_DEPRECATION_WARNINGS
 	ELastAuthority LastAuthority = ELastAuthority::Unset;
 	
 	/** Checks for consistency between Spline and SplineCurves. */
-	bool Validate() const;
+	ENGINE_API bool Validate() const;
 	
 	/** Set the SplineCurves with the default shape i.e. a spline of two points (Used by default constructor) */
 	ENGINE_API void SetDefaultSpline();
@@ -1091,6 +1091,11 @@ PRAGMA_ENABLE_DEPRECATION_WARNINGS
 
 	// FSplineComponentVisualizer will access some private members when attempting to call NotifyPropertiesModified
 	friend class FSplineComponentVisualizer;
+
+#if WITH_DEV_AUTOMATION_TESTS
+	// AFunctionalSplineTestBase accesses private methods for functional testing.
+	friend class AFunctionalSplineTestBase;
+#endif
 };
 
 /** Used to store spline data during RerunConstructionScripts */
