SplineComponent: Make AddSplinePoint's usage of FSpline CVar configurable.

diff --git a/Engine/Source/Runtime/Engine/Private/Components/SplineComponent.cpp b/Engine/Source/Runtime/Engine/Private/Components/SplineComponent.cpp
--- a/Engine/Source/Runtime/Engine/Private/Components/SplineComponent.cpp
+++ b/Engine/Source/Runtime/Engine/Private/Components/SplineComponent.cpp
@@ -87,6 +87,13 @@ FAutoConsoleVariableRef CVarFailedValidationEnsures(
 	TEXT("When true, all validation checks will ensure if not passed.")
 );
 
+int32 GAddSplinePointImplementation = 0;
+FAutoConsoleVariableRef CVarAddSplinePointImplementation(
+	TEXT("SplineComponent.AddSplinePointImplementation"),
+	GAddSplinePointImplementation,
+	TEXT("0) Uses FSpline::AddPoint - 1) Uses FSpline::InsertPoint")
+);
+
 /**
  * This option exists because the tangent computation code for CurveClamped points is just not consistent.
  * I suspect the FSplineCurves logic has a bug, but this has not yet been confirmed.
@@ -1373,7 +1380,12 @@ void USplineComponent::AddSplinePoint(const FVector& Position, ESplineCoordinate
 	NewPoint.Scale = FVector(1.f);
 	NewPoint.Type = ConvertInterpCurveModeToSplinePointType(CIM_CurveAuto);
 
-	Spline.AddPoint(NewPoint);
+	switch (UE::SplineComponent::GAddSplinePointImplementation)
+	{
+	default: // fallthrough
+	case 0: Spline.AddPoint(NewPoint); break;
+	case 1: Spline.InsertPoint(NewPoint, GetNumberOfSplinePoints()); break;
+	}
 
 	WarninglessSplineCurves().Position.Points.Emplace(InKey, TransformedPosition, FVector::ZeroVector, FVector::ZeroVector, CIM_CurveAuto);
 	WarninglessSplineCurves().Rotation.Points.Emplace(InKey, FQuat::Identity, FQuat::Identity, FQuat::Identity, CIM_CurveAuto);
