Update commandline variable to -GeneratorPackageFilterString
This should be more clear as we are filtering the generators not EDLs.

diff --git a/Engine/Source/Runtime/Engine/Private/WorldPartition/WorldPartitionSettings.cpp b/Engine/Source/Runtime/Engine/Private/WorldPartition/WorldPartitionSettings.cpp
--- a/Engine/Source/Runtime/Engine/Private/WorldPartition/WorldPartitionSettings.cpp
+++ b/Engine/Source/Runtime/Engine/Private/WorldPartition/WorldPartitionSettings.cpp
@@ -12,7 +12,7 @@ UWorldPartitionSettings::UWorldPartitionSettings(const FObjectInitializer& Objec
 {}
 
 #if WITH_EDITOR
-UWorldPartitionSettings::EDLCookFilterCommandLineLoader::EDLCookFilterCommandLineLoader()
+UWorldPartitionSettings::CookFilterCommandLineLoader::CookFilterCommandLineLoader()
 {
 	// First parse EDL filters passed through commandline like -EDLCookFilters=Value1+Value2+Value3
 	FString SwitchValue;
@@ -30,16 +30,16 @@ UWorldPartitionSettings::EDLCookFilterCommandLineLoader::EDLCookFilterCommandLin
 		}
 	}
 
-	// Second parse generators filters passed through commandline like -EDLFilterGenerator=Value1+Value2+Value3
+	// Second parse generators filters passed through commandline like -GeneratorPackageFilter=Value1+Value2+Value3
 	// These restrict if the generator should be cooked more than once. When specified the generator will not be saved again.
-	FString EDLFilterGeneratorString;
-	if (FParse::Value(FCommandLine::Get(), TEXT("-EDLFilterGenerator="), EDLFilterGeneratorString))
+	FString GeneratorPackageFilterString;
+	if (FParse::Value(FCommandLine::Get(), TEXT("-GeneratorPackageFilter="), GeneratorPackageFilterString))
 	{
 		TArray<FString> Tokens;
-		EDLFilterGeneratorString.ParseIntoArray(Tokens, TEXT("+"));
+		GeneratorPackageFilterString.ParseIntoArray(Tokens, TEXT("+"));
 		for (const FString& Filter : Tokens)
 		{
-			EDLGeneratorFilters.Add(FName(Filter));
+			GeneratorPackageFilters.Add(FName(Filter));
 		}
 	}
 
@@ -56,15 +56,15 @@ UWorldPartitionSettings::EDLCookFilterCommandLineLoader::EDLCookFilterCommandLin
 	}
 }
 
-const UWorldPartitionSettings::EDLCookFilterCommandLineLoader& UWorldPartitionSettings::GetEDLCookFilterLoader()
+const UWorldPartitionSettings::CookFilterCommandLineLoader& UWorldPartitionSettings::GetCookFilterCommandLineLoader()
 {
-	static UWorldPartitionSettings::EDLCookFilterCommandLineLoader Loader;
+	static UWorldPartitionSettings::CookFilterCommandLineLoader Loader;
 	return Loader;
 }
 
 const TSet<FTopLevelAssetPath>& UWorldPartitionSettings::GetEDLCookFilters()
 {
-	return GetEDLCookFilterLoader().EDLFilters;
+	return GetCookFilterCommandLineLoader().EDLFilters;
 }
 
 const bool UWorldPartitionSettings::ShouldCookEDLGenerator(const UPackage* Generator)
@@ -74,7 +74,7 @@ const bool UWorldPartitionSettings::ShouldCookEDLGenerator(const UPackage* Gener
 		return false;
 	}
 
-	return !GetEDLCookFilterLoader().EDLGeneratorFilters.Contains(Generator->GetFName());
+	return !GetCookFilterCommandLineLoader().GeneratorPackageFilters.Contains(Generator->GetFName());
 }
 
 #endif // #if WITH_EDITOR
diff --git a/Engine/Source/Runtime/Engine/Public/WorldPartition/WorldPartitionSettings.h b/Engine/Source/Runtime/Engine/Public/WorldPartition/WorldPartitionSettings.h
--- a/Engine/Source/Runtime/Engine/Public/WorldPartition/WorldPartitionSettings.h
+++ b/Engine/Source/Runtime/Engine/Public/WorldPartition/WorldPartitionSettings.h
@@ -101,15 +101,15 @@ protected:
 #endif
 
 private:
-	struct EDLCookFilterCommandLineLoader
+	struct CookFilterCommandLineLoader
 	{
-		EDLCookFilterCommandLineLoader();
+		CookFilterCommandLineLoader();
 
 		TSet<FTopLevelAssetPath> EDLFilters;
-		TSet<FName> EDLGeneratorFilters;
+		TSet<FName> GeneratorPackageFilters;
 	};
 
-	static const EDLCookFilterCommandLineLoader& GetEDLCookFilterLoader();
+	static const CookFilterCommandLineLoader& GetCookFilterCommandLineLoader();
 };
 
 #undef UE_API
