Streaming manager API for invalidating internal file cache

diff --git a/Engine/Source/Runtime/Engine/Private/Streaming/StreamingManagerTexture.cpp b/Engine/Source/Runtime/Engine/Private/Streaming/StreamingManagerTexture.cpp
--- a/Engine/Source/Runtime/Engine/Private/Streaming/StreamingManagerTexture.cpp
+++ b/Engine/Source/Runtime/Engine/Private/Streaming/StreamingManagerTexture.cpp
@@ -221,9 +221,7 @@ FRenderAssetStreamingManager::FRenderAssetStreamingManager()
 
 	OnPakFileMounted2Handle = FCoreDelegates::GetOnPakFileMounted2().AddLambda([this](const IPakFile& PakFile)
 	{
-		FScopeLock ScopeLock(&MountedStateDirtyFilesCS);
-		bRecacheAllFiles = true;
-		MountedStateDirtyFiles.Empty();
+		InvalidateFileCache();
 	});
 
 	NewFileAddedDelegateHandle = FCoreDelegates::NewFileAddedDelegate.AddLambda([this](const FString& FileName)
@@ -1247,6 +1245,12 @@ void FRenderAssetStreamingManager::SetRenderAssetsRemovedTimestamp(const FRemove
 	}
 }
 
+void FRenderAssetStreamingManager::InvalidateFileCache()
+{
+	FScopeLock ScopeLock(&MountedStateDirtyFilesCS);
+	bRecacheAllFiles = true;
+	MountedStateDirtyFiles.Empty();
+}
 
 void FRenderAssetStreamingManager::NotifyPrimitiveUpdated( const IPrimitiveComponent* Primitive )
 {
diff --git a/Engine/Source/Runtime/Engine/Private/Streaming/StreamingManagerTexture.h b/Engine/Source/Runtime/Engine/Private/Streaming/StreamingManagerTexture.h
--- a/Engine/Source/Runtime/Engine/Private/Streaming/StreamingManagerTexture.h
+++ b/Engine/Source/Runtime/Engine/Private/Streaming/StreamingManagerTexture.h
@@ -195,6 +195,9 @@ struct FRenderAssetStreamingManager final : public IRenderAssetStreamingManager
 	 */
 	void SetRenderAssetsRemovedTimestamp(const FRemovedRenderAssetArray& RemovedRenderAssets);
 
+	/* Schedules file recache */
+	virtual void InvalidateFileCache() override;
+
 	/** Whether it is allowed for this manager to use more cores using ParallelFor when doing the incremental update. */
 	static bool IsParallelForAllowedDuringIncrementalUpdate(int32 NumElementsToProcess, int32& OutParallelForMinBatchSize);
 
diff --git a/Engine/Source/Runtime/Engine/Public/ContentStreaming.h b/Engine/Source/Runtime/Engine/Public/ContentStreaming.h
--- a/Engine/Source/Runtime/Engine/Public/ContentStreaming.h
+++ b/Engine/Source/Runtime/Engine/Public/ContentStreaming.h
@@ -503,6 +503,9 @@ struct IRenderAssetStreamingManager : public IStreamingManager
 	virtual void AddRenderedTextureStats(TMap<FString, FRenderedTextureStats>& InOutRenderedTextureAssets) = 0;
 
 	virtual void OnRenderAssetUpdateAbandoned(UStreamableRenderAsset* OwningRenderAsset, TRefCountPtr<FRenderAssetUpdate>& RenderAssetUpdate) = 0;
+
+	/* Schedules file recache */
+	virtual void InvalidateFileCache() = 0;
 };
 
 enum class EAudioChunkLoadResult : uint8
