[PostProcessing] Disable more editor-only showflags when VisualizeColorGrading is enabled, to get a viewport that is more representative of the game view.

diff --git a/Engine/Source/Runtime/Engine/Private/ShowFlags.cpp b/Engine/Source/Runtime/Engine/Private/ShowFlags.cpp
--- a/Engine/Source/Runtime/Engine/Private/ShowFlags.cpp
+++ b/Engine/Source/Runtime/Engine/Private/ShowFlags.cpp
@@ -710,6 +710,13 @@ void EngineShowFlagOverride(EShowFlagInitMode ShowFlagInitMode, EViewModeIndex V
 	if (EngineShowFlags.VisualizeColorGrading)
 	{
 		EngineShowFlags.SetCompositeEditorPrimitives(false);
+		EngineShowFlags.SetCompositeDebugPrimitives(false);
+		EngineShowFlags.SetAudioRadius(false);
+		EngineShowFlags.SetEditor(false);
+		EngineShowFlags.SetGrid(false);
+		EngineShowFlags.SetSplines(false);
+		EngineShowFlags.SetSelection(false);
+		EngineShowFlags.SetLightRadius(false);
 	}
 
 #if !(UE_BUILD_SHIPPING || UE_BUILD_TEST)
