Added console command to reload a package in the editor, including in -game

obj reloadpkg package=/Game/Foo

diff --git a/Engine/Source/Runtime/Engine/Private/UnrealEngine.cpp b/Engine/Source/Runtime/Engine/Private/UnrealEngine.cpp
--- a/Engine/Source/Runtime/Engine/Private/UnrealEngine.cpp
+++ b/Engine/Source/Runtime/Engine/Private/UnrealEngine.cpp
@@ -151,6 +151,7 @@ UnrealEngine.cpp: Implements the UEngine class and helpers.
 #include "UniversalObjectLocators/AnimInstanceLocatorFragment.h"
 
 #if WITH_EDITOR
+#include "PackageTools.h"
 #include "Settings/LevelEditorPlaySettings.h"
 #include "LandscapeSubsystem.h"
 #include "TextureCompiler.h"
@@ -9165,6 +9166,29 @@ bool UEngine::HandleObjCommand( const TCHAR* Cmd, FOutputDevice& Ar )
 		}
 		return true;
 	}
+#if WITH_EDITOR
+	else if (FParse::Command(&Cmd, TEXT("RELOADPKG")) || FParse::Command(&Cmd, TEXT("RELOADPACKAGE")))
+	{
+		if (UPackage* Pkg = nullptr;
+			ParseObject<UPackage>(Cmd, TEXT("PACKAGE="), Pkg, nullptr))
+		{
+			if (FText ReloadPkgError;
+				UPackageTools::ReloadPackages({ Pkg }, ReloadPkgError, EReloadPackagesInteractionMode::AssumePositive))
+			{
+				UE_SUPPRESS(LogExec, Log, Ar.Log(TEXT("Package reloaded successfully.")));
+			}
+			else
+			{
+				UE_SUPPRESS(LogExec, Error, Ar.Logf(TEXT("Package failed to reload: %s."), *ReloadPkgError.ToString()));
+			}
+		}
+		else
+		{
+			UE_SUPPRESS(LogExec, Error, Ar.Log(TEXT("Package wasn't found.")));
+		}
+		return true;
+	}
+#endif	// #if WITH_EDITOR
 #if !UE_BUILD_SHIPPING
 	else if (FParse::Command(&Cmd,TEXT("LIST2")))
 	{			
