[WorldPartition]

Removed checks which prevent reassigning runtime cell data layers. With the new support for mutating the runtime cells within cell transformers we now expect that this can happen.

diff --git a/Engine/Source/Runtime/Engine/Private/WorldPartition/WorldPartitionRuntimeCell.cpp b/Engine/Source/Runtime/Engine/Private/WorldPartition/WorldPartitionRuntimeCell.cpp
--- a/Engine/Source/Runtime/Engine/Private/WorldPartition/WorldPartitionRuntimeCell.cpp
+++ b/Engine/Source/Runtime/Engine/Private/WorldPartition/WorldPartitionRuntimeCell.cpp
@@ -81,14 +81,6 @@ bool UWorldPartitionRuntimeCell::NeedsActorToCellRemapping() const
 
 void UWorldPartitionRuntimeCell::SetDataLayers(const TArray<const UDataLayerInstance*>& InDataLayerInstances)
 {
-	check(DataLayers.IsEmpty());
-	check(!ExternalDataLayerAsset);
-
-	if (InDataLayerInstances.IsEmpty())
-	{
-		return;
-	}
-
 	// Validate that we have maximum 1 External Data Layer
 	check(Algo::CountIf(InDataLayerInstances, [](const UDataLayerInstance* DataLayerInstance) { return DataLayerInstance->IsA<UExternalDataLayerInstance>(); }) <= 1);
 	// Validate that all Data Layers are Runtime
@@ -113,10 +105,13 @@ void UWorldPartitionRuntimeCell::SetDataLayers(const TArray<const UDataLayerInst
 	TArray<FName> SortedDataLayerInstanceNames;
 	bool bIsFirstDataLayerExternal = false;
 	Algo::Transform(SortedDataLayerInstances, SortedDataLayerInstanceNames, [](const UDataLayerInstance* DataLayerInstance) { return DataLayerInstance->GetDataLayerFName(); });
-	if (const UExternalDataLayerInstance* ExternalDataLayerInstance = Cast<UExternalDataLayerInstance>(SortedDataLayerInstances[0]))
+	if (SortedDataLayerInstances.Num() > 0)
 	{
-		bIsFirstDataLayerExternal = true;
-		ExternalDataLayerAsset = ExternalDataLayerInstance->GetExternalDataLayerAsset();
+		if (const UExternalDataLayerInstance* ExternalDataLayerInstance = Cast<UExternalDataLayerInstance>(SortedDataLayerInstances[0]))
+		{
+			bIsFirstDataLayerExternal = true;
+			ExternalDataLayerAsset = ExternalDataLayerInstance->GetExternalDataLayerAsset();
+		}
 	}
 	DataLayers = FDataLayerInstanceNames(SortedDataLayerInstanceNames, bIsFirstDataLayerExternal);
 }
