Add the option to display trajectory velocity in debug log

diff --git a/Engine/Plugins/Experimental/UAF/UAFPoseSearch/Source/UAFPoseSearch/Private/RigUnit_DebugDrawTrajectory.cpp b/Engine/Plugins/Experimental/UAF/UAFPoseSearch/Source/UAFPoseSearch/Private/RigUnit_DebugDrawTrajectory.cpp
--- a/Engine/Plugins/Experimental/UAF/UAFPoseSearch/Source/UAFPoseSearch/Private/RigUnit_DebugDrawTrajectory.cpp
+++ b/Engine/Plugins/Experimental/UAF/UAFPoseSearch/Source/UAFPoseSearch/Private/RigUnit_DebugDrawTrajectory.cpp
@@ -26,6 +26,6 @@ FRigUnit_DebugDrawTrajectory_Execute()
 	// into a module component that can be used to access GT data where needed. 
 	UAnimNextComponent* Component = Cast<UAnimNextComponent>(ModuleInstance->GetObject());
 
-	UTransformTrajectoryBlueprintLibrary::DebugDrawTrajectory(Trajectory, Component, LogPoseSearch, ELogVerbosity::Display, DebugThickness, DebugOffset);
+	UTransformTrajectoryBlueprintLibrary::DebugDrawTrajectory(Trajectory, Component, LogPoseSearch, ELogVerbosity::Display, DebugThickness, DebugOffset, DisplayVelocity);
 #endif
 }
diff --git a/Engine/Plugins/Experimental/UAF/UAFPoseSearch/Source/UAFPoseSearch/Private/RigUnit_DebugDrawTrajectory.h b/Engine/Plugins/Experimental/UAF/UAFPoseSearch/Source/UAFPoseSearch/Private/RigUnit_DebugDrawTrajectory.h
--- a/Engine/Plugins/Experimental/UAF/UAFPoseSearch/Source/UAFPoseSearch/Private/RigUnit_DebugDrawTrajectory.h
+++ b/Engine/Plugins/Experimental/UAF/UAFPoseSearch/Source/UAFPoseSearch/Private/RigUnit_DebugDrawTrajectory.h
@@ -35,6 +35,10 @@ struct FRigUnit_DebugDrawTrajectory: public FRigUnit_AnimNextBase
 	UPROPERTY(EditAnywhere, Category = "Debug Draw - Trajectory", meta = (Input))
 	bool Enabled = false;
 
+	// Whether to display the velocity of the samples in on screen text
+	UPROPERTY(EditAnywhere, Category = "Debug Draw - Trajectory", meta = (Input))
+	bool DisplayVelocity = false;
+
 	// The execution result
 	UPROPERTY(EditAnywhere, DisplayName = "Execute", Category = "BeginExecution", meta = (Input, Output))
 	FAnimNextExecuteContext ExecuteContext;
diff --git a/Engine/Source/Runtime/Engine/Private/Animation/TrajectoryTypes.cpp b/Engine/Source/Runtime/Engine/Private/Animation/TrajectoryTypes.cpp
--- a/Engine/Source/Runtime/Engine/Private/Animation/TrajectoryTypes.cpp
+++ b/Engine/Source/Runtime/Engine/Private/Animation/TrajectoryTypes.cpp
@@ -164,7 +164,7 @@ void UTransformTrajectoryBlueprintLibrary::DebugDrawTrajectory(const FTransformT
 	}
 }
 
-void UTransformTrajectoryBlueprintLibrary::DebugDrawTrajectory(const FTransformTrajectory& Trajectory, const UObject* Owner, const FLogCategoryBase& Category, ELogVerbosity::Type Verbosity, const float DebugThickness, float HeightOffset)
+void UTransformTrajectoryBlueprintLibrary::DebugDrawTrajectory(const FTransformTrajectory& Trajectory, const UObject* Owner, const FLogCategoryBase& Category, ELogVerbosity::Type Verbosity, const float DebugThickness, float HeightOffset, bool bDisplayVelocity)
 {
 #if ENABLE_VISUAL_LOG
 
@@ -194,7 +194,19 @@ void UTransformTrajectoryBlueprintLibrary::DebugDrawTrajectory(const FTransformT
 
 			const bool IsPast = Trajectory.Samples[Index].TimeInSeconds <= 0.f;
 
-			CurrentEntry->AddSegment(Pos, Pos + X * Scale, Category.GetCategoryName(), Verbosity, IsPast ? FColor::Red : FColor::Blue, TEXT(""), DebugThickness);
+			FVector Velocity = FVector::ZeroVector;
+			FString VelocityDebugStr;
+			if (Index > 0 && bDisplayVelocity)
+			{
+				float DeltaTime = Trajectory.Samples[Index].TimeInSeconds - Trajectory.Samples[Index - 1].TimeInSeconds;
+				if (DeltaTime > 0.f)
+				{
+					Velocity = (Trajectory.Samples[Index].Position - Trajectory.Samples[Index - 1].Position) / DeltaTime;
+					VelocityDebugStr = FString::Printf(TEXT("Velocity: %.3f"), Velocity.Size());
+				}
+			}
+
+			CurrentEntry->AddSegment(Pos, Pos + X * Scale, Category.GetCategoryName(), Verbosity, IsPast ? FColor::Red : FColor::Blue, VelocityDebugStr, DebugThickness);
 			CurrentEntry->AddSegment(Pos, Pos + Y * Scale, Category.GetCategoryName(), Verbosity, IsPast ? FColor::Orange : FColor::Turquoise, TEXT(""), DebugThickness);
 
 			if (Index == LastIndex)
diff --git a/Engine/Source/Runtime/Engine/Public/Animation/TrajectoryTypes.h b/Engine/Source/Runtime/Engine/Public/Animation/TrajectoryTypes.h
--- a/Engine/Source/Runtime/Engine/Public/Animation/TrajectoryTypes.h
+++ b/Engine/Source/Runtime/Engine/Public/Animation/TrajectoryTypes.h
@@ -70,7 +70,7 @@ public:
 
 #if ENABLE_ANIM_DEBUG
 	ENGINE_API static void DebugDrawTrajectory(const FTransformTrajectory& Trajectory, FAnimInstanceProxy& AnimInstanceProxy, const float DebugThickness = 0.f, float HeightOffset = 0.f, int MaxHistorySamples = -1, int MaxPredictionSamples = -1);
-	ENGINE_API static void DebugDrawTrajectory(const FTransformTrajectory& Trajectory, const UObject* Owner, const FLogCategoryBase& Category, ELogVerbosity::Type Verbosity, const float DebugThickness = 0.f, float HeightOffset = 0.f);
+	ENGINE_API static void DebugDrawTrajectory(const FTransformTrajectory& Trajectory, const UObject* Owner, const FLogCategoryBase& Category, ELogVerbosity::Type Verbosity, const float DebugThickness = 0.f, float HeightOffset = 0.f, bool bDisplayVelocity = false);
 #endif
 };
 
