Add dump PSO precache stats to CSV file option

diff --git a/Engine/Source/Runtime/Engine/Private/PSOPrecacheValidation.cpp b/Engine/Source/Runtime/Engine/Private/PSOPrecacheValidation.cpp
--- a/Engine/Source/Runtime/Engine/Private/PSOPrecacheValidation.cpp
+++ b/Engine/Source/Runtime/Engine/Private/PSOPrecacheValidation.cpp
@@ -323,6 +323,13 @@ bool PSOCollectorStats::FPrecacheStatsCollector::IsPrecached(uint64 PrecacheStat
 	return (Value && Value->bPrecached);
 }
 
+bool PSOCollectorStats::FPrecacheStatsCollector::IsUsed(uint64 PrecacheStateHash)
+{
+	FScopeLock Lock(&StateMapLock);
+	FShaderStateUsage* Value = HashedStateMap.Find(PrecacheStateHash);
+	return (Value && Value->bUsed);
+}
+
 #if PSO_PRECACHING_TRACKING
 
 bool PSOCollectorStats::FPrecacheStatsCollector::GetPrecacheData(uint64 PrecacheStateHash, FString& OutMaterialName, int32& OutPSOCollectorIndex, const FVertexFactoryType*& OutVertexFactoryType)
@@ -1101,4 +1108,230 @@ void LogPSOMissInfo(
 	UE_LOG(LogEngine, Log, TEXT("%s\n"), StringBuilder.ToString());
 }
 
+
+void PSOCollectorStats::DumpPSOPrecacheValidationStats()
+{
+	const FString Filename = FString::Printf(TEXT("%sPSOPrecacheStats-%s.csv"), *FPaths::ProfilingDir(), *FDateTime::Now().ToString());
+	FArchive* CSVFile = IFileManager::Get().CreateFileWriter(*Filename, FILEWRITE_AllowRead);
+	if (CSVFile == nullptr)
+	{
+		return;
+	}
+
+	enum class ETrackingType
+	{
+		Missed,
+		Precached,
+		Used,
+		TooLate,
+		Untracked,
+		Num,
+	};
+
+	enum class EStatsType
+	{
+		ShaderOnly,
+		Minimal,
+		Full,
+		Num,
+	};
+	
+	struct FPerTypeStats
+	{
+		FPerTypeStats()
+		{
+			FMemory::Memzero(TrackedStats, (int32)EStatsType::Num * (int32)ETrackingType::Num * sizeof(int32));
+		}
+
+		int32 GetStat(EStatsType StatsType, ETrackingType TrackingType) const
+		{
+			return TrackedStats[(int32)StatsType][(int32)TrackingType];
+		}
+		
+		void SetStat(EStatsType StatsType, ETrackingType TrackingType, int32 Value)
+		{
+			TrackedStats[(int32)StatsType][(int32)TrackingType] = Value;
+		}
+
+		FString Name;
+		int32 TrackedStats[(int32)EStatsType::Num][(int32)ETrackingType::Num];
+	};	
+	
+	auto AddStatsToCSVFile = [CSVFile](const TCHAR* TypeName, bool bSort, TArray<FPerTypeStats>& TypeStats)
+	{		
+		if (bSort)
+		{
+			Algo::Sort(TypeStats, [](FPerTypeStats& LHS, FPerTypeStats& RHS)
+				{
+					return LHS.GetStat(EStatsType::Full, ETrackingType::Precached) > RHS.GetStat(EStatsType::Full, ETrackingType::Precached);
+				});
+		}
+
+		FString Header = FString::Printf(TEXT("%s,Full Precached,Full Used,Full Usage,Full Missed,Full TooLate,Full Untracked,ShaderOnly Precached,ShaderOnly Used,ShaderOnly Usage,ShaderOnly Missed,Minimal Precached,Minimal Used,Minimal Usage,Minimal Missed\n"), TypeName);		
+		CSVFile->Serialize(TCHAR_TO_ANSI(*Header), Header.Len());
+				
+		for (auto Stats : TypeStats)
+		{
+			auto AddStats = [CSVFile, Stats](EStatsType StatsType)
+			{				
+				int32 Precached = Stats.GetStat(StatsType, ETrackingType::Precached);
+				int32 Used = Stats.GetStat(StatsType, ETrackingType::Used);
+				float Usage = Precached > 0 ? float(Used) / Precached : 0.0f;
+				int32 Missed = Stats.GetStat(StatsType, ETrackingType::Missed);
+
+				FString Row = FString::Printf(TEXT(",%d,%d,%0.2f,%d"), Precached, Used, Usage, Missed);
+				CSVFile->Serialize(TCHAR_TO_ANSI(*Row), Row.Len());
+			};
+			
+			int32 FullPrecached = Stats.GetStat(EStatsType::Full, ETrackingType::Precached);
+			int32 FullUsed = Stats.GetStat(EStatsType::Full, ETrackingType::Used);
+			int32 FullMissed = Stats.GetStat(EStatsType::Full, ETrackingType::Missed);
+
+			if (FullPrecached != 0 || FullUsed != 0 || FullMissed != 0)
+			{
+				FString NameEntry = FString::Printf(TEXT("%s"), *Stats.Name);
+				CSVFile->Serialize(TCHAR_TO_ANSI(*NameEntry), NameEntry.Len());
+
+				AddStats(EStatsType::Full);
+				
+				FString Row = FString::Printf(TEXT(",%d,%d"), Stats.GetStat(EStatsType::Full, ETrackingType::TooLate), Stats.GetStat(EStatsType::Full, ETrackingType::Untracked));
+				CSVFile->Serialize(TCHAR_TO_ANSI(*Row), Row.Len());
+
+				AddStats(EStatsType::ShaderOnly);
+				AddStats(EStatsType::Minimal);
+				
+				FString NewLine(TEXT("\n"));
+				CSVFile->Serialize(TCHAR_TO_ANSI(*NewLine), NewLine.Len());
+			}
+		}
+				
+		const TCHAR* Footer = TEXT("\n\n\n");
+		CSVFile->Serialize(TCHAR_TO_ANSI(Footer), FPlatformString::Strlen(Footer));
+	};
+
+	const PSOCollectorStats::FPrecacheStats& FullPSOPrecacheStats = FullPSOPrecacheStatsCollector.GetStats();
+	const PSOCollectorStats::FPrecacheStats& ShaderOnlyPSOPrecacheStats = ShadersOnlyPSOPrecacheStatsCollector.GetStats();
+	const PSOCollectorStats::FPrecacheStats& MinimalPSOPrecacheStats = MinimalPSOPrecacheStatsCollector.GetStats();
+
+	// Totals
+	{
+		TArray<FPerTypeStats> TotalStats;
+
+		FPerTypeStats PerTypeStats;
+		PerTypeStats.Name = TEXT("");
+		auto AddTotalStats = [&PerTypeStats](EStatsType StatsType, const PSOCollectorStats::FPrecacheStats& Stats)
+		{			
+			PerTypeStats.SetStat(StatsType, ETrackingType::Missed, Stats.MissData.GetTotalCount());
+			PerTypeStats.SetStat(StatsType, ETrackingType::Precached, Stats.PrecacheData.GetTotalCount());
+			PerTypeStats.SetStat(StatsType, ETrackingType::Used, Stats.UsageData.GetTotalCount());
+			PerTypeStats.SetStat(StatsType, ETrackingType::TooLate, Stats.TooLateData.GetTotalCount());
+			PerTypeStats.SetStat(StatsType, ETrackingType::Untracked, Stats.UntrackedData.GetTotalCount());
+		};		
+		AddTotalStats(EStatsType::Full, FullPSOPrecacheStats);
+		AddTotalStats(EStatsType::Minimal, MinimalPSOPrecacheStats);
+		AddTotalStats(EStatsType::ShaderOnly, ShaderOnlyPSOPrecacheStats);
+
+		TotalStats.Add(PerTypeStats);
+
+		AddStatsToCSVFile(TEXT("Totals"), false, TotalStats);
+	}
+	
+	// Collect Global stats
+	{	
+		TArray<FPerTypeStats> GlobalTypeStats;
+
+		for (int32 PSOCollectorIndex = 0; PSOCollectorIndex < FGlobalPSOCollectorManager::GetPSOCollectorCount(); ++PSOCollectorIndex)
+		{
+			FPerTypeStats PerTypeStats;
+			PerTypeStats.Name = FGlobalPSOCollectorManager::GetName(PSOCollectorIndex);
+			PerTypeStats.SetStat(EStatsType::Full, ETrackingType::Missed, FullPSOPrecacheStats.MissData.PerGlobalTypeCount[PSOCollectorIndex]);
+			PerTypeStats.SetStat(EStatsType::Full, ETrackingType::Precached, FullPSOPrecacheStats.PrecacheData.PerGlobalTypeCount[PSOCollectorIndex]);
+			PerTypeStats.SetStat(EStatsType::Full, ETrackingType::Used, FullPSOPrecacheStats.UsageData.PerGlobalTypeCount[PSOCollectorIndex]);			
+			PerTypeStats.SetStat(EStatsType::Full, ETrackingType::TooLate, FullPSOPrecacheStats.TooLateData.PerGlobalTypeCount[PSOCollectorIndex]);
+			PerTypeStats.SetStat(EStatsType::Full, ETrackingType::Untracked, FullPSOPrecacheStats.UntrackedData.PerGlobalTypeCount[PSOCollectorIndex]);
+			GlobalTypeStats.Add(PerTypeStats);
+		}
+		
+		AddStatsToCSVFile(TEXT("GlobalType"), true, GlobalTypeStats);
+	}
+	
+	// Collect MeshPass stats
+	{	
+		TArray<FPerTypeStats> PerMeshPassStats;
+
+		EShadingPath ShadingPath = GetFeatureLevelShadingPath(GMaxRHIFeatureLevel);
+		for (int32 PSOCollectorIndex = 0; PSOCollectorIndex < FPSOCollectorCreateManager::GetPSOCollectorCount(ShadingPath); ++PSOCollectorIndex)
+		{
+			FPerTypeStats MeshPassStats;
+			MeshPassStats.Name = FPSOCollectorCreateManager::GetName(ShadingPath, PSOCollectorIndex);
+						
+			MeshPassStats.SetStat(EStatsType::Full, ETrackingType::Missed, FullPSOPrecacheStats.MissData.PerMeshPassCount[PSOCollectorIndex]);
+			MeshPassStats.SetStat(EStatsType::Full, ETrackingType::Precached, FullPSOPrecacheStats.PrecacheData.PerMeshPassCount[PSOCollectorIndex]);
+			MeshPassStats.SetStat(EStatsType::Full, ETrackingType::Used, FullPSOPrecacheStats.UsageData.PerMeshPassCount[PSOCollectorIndex]);
+			MeshPassStats.SetStat(EStatsType::Full, ETrackingType::TooLate, FullPSOPrecacheStats.TooLateData.PerMeshPassCount[PSOCollectorIndex]);
+			MeshPassStats.SetStat(EStatsType::Full, ETrackingType::Untracked, FullPSOPrecacheStats.UntrackedData.PerMeshPassCount[PSOCollectorIndex]);
+
+			MeshPassStats.SetStat(EStatsType::ShaderOnly, ETrackingType::Missed, ShaderOnlyPSOPrecacheStats.MissData.PerMeshPassCount[PSOCollectorIndex]);
+			MeshPassStats.SetStat(EStatsType::ShaderOnly, ETrackingType::Precached, ShaderOnlyPSOPrecacheStats.PrecacheData.PerMeshPassCount[PSOCollectorIndex]);
+			MeshPassStats.SetStat(EStatsType::ShaderOnly, ETrackingType::Used, ShaderOnlyPSOPrecacheStats.UsageData.PerMeshPassCount[PSOCollectorIndex]);
+
+			MeshPassStats.SetStat(EStatsType::Minimal, ETrackingType::Missed, MinimalPSOPrecacheStats.MissData.PerMeshPassCount[PSOCollectorIndex]);
+			MeshPassStats.SetStat(EStatsType::Minimal, ETrackingType::Precached, MinimalPSOPrecacheStats.PrecacheData.PerMeshPassCount[PSOCollectorIndex]);
+			MeshPassStats.SetStat(EStatsType::Minimal, ETrackingType::Used, MinimalPSOPrecacheStats.UsageData.PerMeshPassCount[PSOCollectorIndex]);
+
+			PerMeshPassStats.Add(MeshPassStats);
+		}
+				
+		AddStatsToCSVFile(TEXT("PassName"), true, PerMeshPassStats);
+	}
+
+	// Collect VF stats
+	{
+		TMap<FString, FPerTypeStats> VertexFactoryStatMap;
+		auto AddVertexFactoryStats = [&VertexFactoryStatMap](EStatsType StatsType, ETrackingType TrackingType, PSOCollectorStats::VertexFactoryCountTableType PerVertexFactoryCounts)
+		{
+			for (auto Iter : PerVertexFactoryCounts)
+			{
+				FString VertexFactoryTypeName = Iter.Key->GetName();
+				FPerTypeStats& Stats = VertexFactoryStatMap.FindOrAdd(VertexFactoryTypeName);
+				Stats.Name = VertexFactoryTypeName;
+				int32 CurrentValue = Stats.GetStat(StatsType, TrackingType);
+				Stats.SetStat(StatsType, TrackingType, CurrentValue + Iter.Value);
+			}
+		};
+		AddVertexFactoryStats(EStatsType::Full, ETrackingType::Precached, FullPSOPrecacheStats.PrecacheData.PerVertexFactoryCount);
+		AddVertexFactoryStats(EStatsType::Full, ETrackingType::Used, FullPSOPrecacheStats.UsageData.PerVertexFactoryCount);
+		AddVertexFactoryStats(EStatsType::Full, ETrackingType::Missed, FullPSOPrecacheStats.MissData.PerVertexFactoryCount);
+		AddVertexFactoryStats(EStatsType::Full, ETrackingType::TooLate, FullPSOPrecacheStats.TooLateData.PerVertexFactoryCount);
+		AddVertexFactoryStats(EStatsType::Full, ETrackingType::Untracked, FullPSOPrecacheStats.UntrackedData.PerVertexFactoryCount);
+		
+		AddVertexFactoryStats(EStatsType::ShaderOnly, ETrackingType::Precached, ShaderOnlyPSOPrecacheStats.PrecacheData.PerVertexFactoryCount);
+		AddVertexFactoryStats(EStatsType::ShaderOnly, ETrackingType::Used, ShaderOnlyPSOPrecacheStats.UsageData.PerVertexFactoryCount);
+		AddVertexFactoryStats(EStatsType::ShaderOnly, ETrackingType::Missed, ShaderOnlyPSOPrecacheStats.MissData.PerVertexFactoryCount);
+		
+		AddVertexFactoryStats(EStatsType::Minimal, ETrackingType::Precached, MinimalPSOPrecacheStats.PrecacheData.PerVertexFactoryCount);
+		AddVertexFactoryStats(EStatsType::Minimal, ETrackingType::Used, MinimalPSOPrecacheStats.UsageData.PerVertexFactoryCount);
+		AddVertexFactoryStats(EStatsType::Minimal, ETrackingType::Missed, MinimalPSOPrecacheStats.MissData.PerVertexFactoryCount);
+
+		TArray<FPerTypeStats> VertexFactoryStats;
+		for (auto StatEntry : VertexFactoryStatMap)
+		{
+			VertexFactoryStats.Add(StatEntry.Value);
+		}
+
+		AddStatsToCSVFile(TEXT("VertexFactoryType"), true, VertexFactoryStats);
+	}
+
+	delete CSVFile;
+	CSVFile = nullptr;
+}
+
+static FAutoConsoleCommand GDumpPSOPrecacheValidationStatsCmd(
+	TEXT("r.PSOPrecache.DumpStats"),
+	TEXT("Dumps all currently collected PSO precache validation stats to csv file"),
+	FConsoleCommandWithArgsDelegate::CreateStatic([](const TArray<FString>& Args)
+{		
+	PSOCollectorStats::DumpPSOPrecacheValidationStats();
+}));
+
 #endif // PSO_PRECACHING_VALIDATE && UE_WITH_PSO_PRECACHING
diff --git a/Engine/Source/Runtime/Engine/Public/PSOPrecacheValidation.h b/Engine/Source/Runtime/Engine/Public/PSOPrecacheValidation.h
--- a/Engine/Source/Runtime/Engine/Public/PSOPrecacheValidation.h
+++ b/Engine/Source/Runtime/Engine/Public/PSOPrecacheValidation.h
@@ -150,9 +150,8 @@ namespace PSOCollectorStats
 	 * Track a PSO precache stat's total count, and optionally also track the counts at
 	 * the mesh pass type and vertex factory type granularity.
 	 */
-	class FPrecacheUsageData
+	struct FPrecacheUsageData
 	{
-	public:
 		FPrecacheUsageData(FName StatFName = FName())
 			: StatFName(StatFName)
 		{
@@ -168,7 +167,6 @@ namespace PSOCollectorStats
 
 		ENGINE_API void UpdateStats(EPSOPrecacheType PrecacheType, int32 PSOCollectorIndex, const FVertexFactoryType* VFType);
 		
-	private:
 		bool ShouldRecordFullStats(EPSOPrecacheType PrecacheType, int32 PSOCollectorIndex, const FVertexFactoryType* VFType) const
 		{
 			if (PrecacheType == EPSOPrecacheType::MeshPass)
@@ -306,6 +304,7 @@ namespace PSOCollectorStats
 		}
 
 		ENGINE_API bool IsPrecached(uint64 PrecacheStateHash);
+		ENGINE_API bool IsUsed(uint64 PrecacheStateHash);
 #if PSO_PRECACHING_TRACKING
 		bool GetPrecacheData(uint64 PrecacheStateHash, FString& OutMaterialName, int32& OutPSOCollectorIndex, const FVertexFactoryType*& OutVertexFactoryType);
 #endif
@@ -325,6 +324,8 @@ namespace PSOCollectorStats
 	extern ENGINE_API FPrecacheStatsCollector& GetShadersOnlyPSOPrecacheStatsCollector();
 	extern ENGINE_API FPrecacheStatsCollector& GetMinimalPSOPrecacheStatsCollector();
 	extern ENGINE_API FPrecacheStatsCollector& GetFullPSOPrecacheStatsCollector();
+
+	extern ENGINE_API void DumpPSOPrecacheValidationStats();
 }
 
 #endif // PSO_PRECACHING_VALIDATE
