Make vertex elements list more stable for static and gpu skinned vertex factories so it reduces the amount of PSOs to precache and runtime use on platforms which don't support manual vertex fetch:
- Always add all texture coordinate elements as half4 even though they might not all be used (lots of permutations otherwise due to single texure coord entry and adding last element to list)
- Don't add light map coordinates when static lighting isn't enabled (offset can change depending on lightmap index causing unneeded extra permutations)
- Fixup default vertex element for vertex factory to match the most common usage to reduce misses when using WorldGridMaterial and fix them for compilation when manual vertex fetch is not enabled

diff --git a/Engine/Plugins/FX/Niagara/Source/NiagaraVertexFactories/Private/NiagaraMeshVertexFactory.cpp b/Engine/Plugins/FX/Niagara/Source/NiagaraVertexFactories/Private/NiagaraMeshVertexFactory.cpp
--- a/Engine/Plugins/FX/Niagara/Source/NiagaraVertexFactories/Private/NiagaraMeshVertexFactory.cpp
+++ b/Engine/Plugins/FX/Niagara/Source/NiagaraVertexFactories/Private/NiagaraMeshVertexFactory.cpp
@@ -13,6 +13,7 @@
 #include "MeshMaterialShader.h"
 #include "GlobalRenderResources.h"
 #include "DataDrivenShaderPlatformInfo.h"
+#include "LocalVertexFactory.h"
 
 // Main control for GPU scene on all platforms
 #define NIAGARA_ENABLE_GPU_SCENE_MESHES 1
@@ -127,23 +128,8 @@ void FNiagaraMeshVertexFactory::GetVertexElements(ERHIFeatureLevel::Type Feature
 		if (Data.TextureCoordinates.Num())
 		{
 			const int32 BaseTexCoordAttribute = 4;
-			for (int32 CoordinateIndex = 0; CoordinateIndex < Data.TextureCoordinates.Num(); CoordinateIndex++)
-			{
-				Elements.Add(AccessStreamComponent(
-					Data.TextureCoordinates[CoordinateIndex],
-					BaseTexCoordAttribute + CoordinateIndex,
-					InOutStreams
-				));
-			}
-
-			for (int32 CoordinateIndex = Data.TextureCoordinates.Num(); CoordinateIndex < MAX_TEXCOORDS; CoordinateIndex++)
-			{
-				Elements.Add(AccessStreamComponent(
-					Data.TextureCoordinates[Data.TextureCoordinates.Num() - 1],
-					BaseTexCoordAttribute + CoordinateIndex,
-					InOutStreams
-				));
-			}
+			uint8 TexCoordStreamIndex = AddVertexStream(Data.TextureCoordinates[0], InOutStreams);
+			AddStaticMeshTextureCoordinateElements(TexCoordStreamIndex, BaseTexCoordAttribute, Data.TextureCoordinates, Elements);
 		}
 	}
 }
@@ -211,12 +197,34 @@ void FNiagaraMeshVertexFactory::ModifyCompilationEnvironment(const FVertexFactor
 
 void FNiagaraMeshVertexFactory::GetPSOPrecacheVertexFetchElements(EVertexInputStreamType VertexInputStreamType, FVertexDeclarationElementList& Elements)
 {
-	Elements.Add(FVertexElement(0, 0, VET_Float3, 0, sizeof(float)*3u, false));
+	uint8 StreamIndex = 0;
+	Elements.Add(FVertexElement(StreamIndex++, 0, VET_Float3, 0, 0, false));
 
-	if (IsGPUSceneEnabled(GMaxRHIShaderPlatform, GMaxRHIFeatureLevel) && !PlatformGPUSceneUsesUniformBufferView(GMaxRHIShaderPlatform))
+	if (!RHISupportsManualVertexFetch(GMaxRHIShaderPlatform))
 	{
-		Elements.Add(FVertexElement(1, 0, VET_UInt, 13, sizeof(uint32), true));
+		// Normals
+		Elements.Add(FVertexElement(StreamIndex, 0, VET_PackedNormal, 1, 0, false));
+		Elements.Add(FVertexElement(StreamIndex++, 4, VET_PackedNormal, 2, 0, false));
+
+		// Color
+		Elements.Add(FVertexElement(StreamIndex++, 0, VET_Color, 3, 0, false));
+
+		// TexCoords
+		int8 UVStreamIndex = StreamIndex;
+		Elements.Add(FVertexElement(StreamIndex, 0, VET_Half4, 4, 0, false));
+		Elements.Add(FVertexElement(StreamIndex, 8, VET_Half4, 5, 0, false));
+		Elements.Add(FVertexElement(StreamIndex, 16, VET_Half4, 6, 0, false));
+		Elements.Add(FVertexElement(StreamIndex++, 24, VET_Half4, 7, 0, false));
 	}
+	
+#if NIAGARA_ENABLE_GPU_SCENE_MESHES
+	// Primitive ID stream
+	if (UseGPUScene(GMaxRHIShaderPlatform, GMaxRHIFeatureLevel)
+		&& !PlatformGPUSceneUsesUniformBufferView(GMaxRHIShaderPlatform))
+	{
+		Elements.Add(FVertexElement(StreamIndex, 0, VET_UInt, 13, 0, true));
+	}
+#endif
 }
 
 void FNiagaraMeshVertexFactory::GetVertexElements(ERHIFeatureLevel::Type FeatureLevel, bool bSupportsManualVertexFetch, FStaticMeshDataType& Data, FVertexDeclarationElementList& Elements)
diff --git a/Engine/Plugins/Runtime/HairStrands/Source/HairStrandsCore/Private/HairCardsVertexFactory.cpp b/Engine/Plugins/Runtime/HairStrands/Source/HairStrandsCore/Private/HairCardsVertexFactory.cpp
--- a/Engine/Plugins/Runtime/HairStrands/Source/HairStrandsCore/Private/HairCardsVertexFactory.cpp
+++ b/Engine/Plugins/Runtime/HairStrands/Source/HairStrandsCore/Private/HairCardsVertexFactory.cpp
@@ -340,8 +340,9 @@ void FHairCardsVertexFactory::GetPSOPrecacheVertexFetchElements(EVertexInputStre
 
 		Elements.Add(FVertexElement(2+NumStreams, 0, FHairCardsMaterialFormat::VertexElementType, 4, FHairCardsMaterialFormat::SizeInByte, false)); // Material
 		Elements.Add(FVertexElement(2+NumStreams, 0, FHairCardsMaterialFormat::VertexElementType, 5, FHairCardsMaterialFormat::SizeInByte, false)); // VertexColor
-
-		Elements.Add(FVertexElement(3+NumStreams, 0, FHairCardsUVFormat::VertexElementType, 6, FHairCardsUVFormat::SizeInByte, false));
+		
+		Elements.Add(FVertexElement(3+NumStreams, 0, FHairCardsUVFormat::VertexElementType, 6, FHairCardsUVFormat::SizeInByte, false));		
+		Elements.Add(FVertexElement(3+NumStreams, 0, FHairCardsUVFormat::VertexElementType, 7, FHairCardsUVFormat::SizeInByte, false));
 	}
 }
 
diff --git a/Engine/Source/Runtime/Engine/Private/GPUSkinVertexFactory.cpp b/Engine/Source/Runtime/Engine/Private/GPUSkinVertexFactory.cpp
--- a/Engine/Source/Runtime/Engine/Private/GPUSkinVertexFactory.cpp
+++ b/Engine/Source/Runtime/Engine/Private/GPUSkinVertexFactory.cpp
@@ -855,45 +855,53 @@ template <GPUSkinBoneInfluenceType BoneInfluenceType>
 void TGPUSkinVertexFactory<BoneInfluenceType>::GetPSOPrecacheVertexFetchElements(EVertexInputStreamType VertexInputStreamType, FVertexDeclarationElementList& Elements)
 {
 	check(VertexInputStreamType == EVertexInputStreamType::Default);
+		
+	uint32 StreamIndex = 0;
 
 	// Position
-	Elements.Add(FVertexElement(0, 0, VET_Float3, 0, 0, false));
+	Elements.Add(FVertexElement(StreamIndex++, 0, VET_Float3, 0, 0, false));
 
 	// Normals
-	Elements.Add(FVertexElement(1, 0, VET_PackedNormal, 1, 0, false));
-	Elements.Add(FVertexElement(2, 0, VET_PackedNormal, 2, 0, false));
+	Elements.Add(FVertexElement(StreamIndex, 0, VET_PackedNormal, 1, 0, false));
+	Elements.Add(FVertexElement(StreamIndex++, 4, VET_PackedNormal, 2, 0, false));
+	
+	// Texcoords
+	Elements.Add(FVertexElement(StreamIndex, 0, VET_Half4, 5, 0, false));
+	Elements.Add(FVertexElement(StreamIndex, 8, VET_Half4, 6, 0, false));	
+	Elements.Add(FVertexElement(StreamIndex, 16, VET_Half4, 7, 0, false));
+	Elements.Add(FVertexElement(StreamIndex++, 24, VET_Half4, 8, 0, false));
 	
+	// Color
+	Elements.Add(FVertexElement(StreamIndex++, 0, VET_Color, 13, 0, false));
+
 	// Bone data
-	uint32 BaseStreamIndex = 3;
 	if (BoneInfluenceType == UnlimitedBoneInfluence)
 	{
 		// Blend offset count
-		Elements.Add(FVertexElement(BaseStreamIndex++, 0, VET_UInt, 3, 0, false));
+		Elements.Add(FVertexElement(StreamIndex++, 0, VET_UInt, 3, 0, false));
 	}
 	else
 	{
-		// Blend indices
-		Elements.Add(FVertexElement(BaseStreamIndex++, 0, VET_UByte4, 3, 0, false));
-		Elements.Add(FVertexElement(BaseStreamIndex++, 0, VET_UByte4, 14, 0, false));
-
-		// Blend weights
-		Elements.Add(FVertexElement(BaseStreamIndex++, 0, VET_UByte4N, 4, 0, false));
-		Elements.Add(FVertexElement(BaseStreamIndex++, 0, VET_UByte4N, 15, 0, false));
+		// Blend indices & weight
+		Elements.Add(FVertexElement(StreamIndex, 0, VET_UByte4, 3, 0, false));
+		Elements.Add(FVertexElement(StreamIndex, 4, VET_UByte4N, 4, 0, false));
+		
+		// Extra blend indices & weights
+		Elements.Add(FVertexElement(StreamIndex, 0, VET_UByte4, 14, 0, false));
+		Elements.Add(FVertexElement(StreamIndex++, 4, VET_UByte4N, 15, 0, false));
 	}
 
-	// Texcoords
-	Elements.Add(FVertexElement(BaseStreamIndex++, 0, VET_Half4, 5, 0, false));
-	Elements.Add(FVertexElement(BaseStreamIndex++, 0, VET_Half4, 6, 0, false));
-
-	// Color
-	Elements.Add(FVertexElement(BaseStreamIndex++, 0, VET_Color, 13, 0, false));
-
-	// Attribute ID
-	Elements.Add(FVertexElement(BaseStreamIndex++, 0, VET_UInt, 16, 0, true));
-
 	// Morph blend data
-	Elements.Add(FVertexElement(Elements.Num(), 0, VET_Float3, 9, 0, false));
-	Elements.Add(FVertexElement(Elements.Num(), 0, VET_Float3, 10, 0, false));
+	Elements.Add(FVertexElement(StreamIndex, 0, VET_Float3, 9, 0, false));
+	Elements.Add(FVertexElement(StreamIndex++, 12, VET_Float3, 10, 0, false));
+	
+	// Primitive ID stream
+	if (UseGPUScene(GMaxRHIShaderPlatform, GMaxRHIFeatureLevel)
+		&& !PlatformGPUSceneUsesUniformBufferView(GMaxRHIShaderPlatform))
+	{
+		// Attribute ID
+		Elements.Add(FVertexElement(StreamIndex++, 0, VET_UInt, 16, 0, true));
+	}
 }
 
 template <GPUSkinBoneInfluenceType BoneInfluenceType>
@@ -911,22 +919,9 @@ void TGPUSkinVertexFactory<BoneInfluenceType>::GetVertexElements(ERHIFeatureLeve
 	// Texture coordinates
 	if (GPUSkinData.TextureCoordinates.Num())
 	{
-		const uint8 BaseTexCoordAttribute = 5;
-		for (int32 CoordinateIndex = 0; CoordinateIndex < GPUSkinData.TextureCoordinates.Num(); ++CoordinateIndex)
-		{
-			OutElements.Add(AccessStreamComponent(
-				GPUSkinData.TextureCoordinates[CoordinateIndex],
-				BaseTexCoordAttribute + CoordinateIndex, InOutStreams
-			));
-		}
-
-		for (int32 CoordinateIndex = GPUSkinData.TextureCoordinates.Num(); CoordinateIndex < MAX_TEXCOORDS; ++CoordinateIndex)
-		{
-			OutElements.Add(AccessStreamComponent(
-				GPUSkinData.TextureCoordinates[GPUSkinData.TextureCoordinates.Num() - 1],
-				BaseTexCoordAttribute + CoordinateIndex, InOutStreams
-			));
-		}
+		const uint8 BaseTexCoordAttribute = 5;				
+		uint8 TexCoordStreamIndex = AddVertexStream(GPUSkinData.TextureCoordinates[0], InOutStreams);
+		AddStaticMeshTextureCoordinateElements(TexCoordStreamIndex, BaseTexCoordAttribute, GPUSkinData.TextureCoordinates, OutElements);
 	}
 
 	if (GPUSkinData.ColorComponentsSRV == nullptr)
diff --git a/Engine/Source/Runtime/Engine/Private/GeometryCacheVertexFactory.cpp b/Engine/Source/Runtime/Engine/Private/GeometryCacheVertexFactory.cpp
--- a/Engine/Source/Runtime/Engine/Private/GeometryCacheVertexFactory.cpp
+++ b/Engine/Source/Runtime/Engine/Private/GeometryCacheVertexFactory.cpp
@@ -230,24 +230,9 @@ void FGeometryCacheVertexVertexFactory::GetVertexElements(
 
 	if (StreamComponentData.TextureCoordinates.Num())
 	{
-		const int32 BaseTexCoordAttribute = 5;
-		for (int32 CoordinateIndex = 0; CoordinateIndex < StreamComponentData.TextureCoordinates.Num(); CoordinateIndex++)
-		{
-			OutElements.Add(AccessStreamComponent(
-				StreamComponentData.TextureCoordinates[CoordinateIndex],
-				BaseTexCoordAttribute + CoordinateIndex, 
-				InOutStreams
-			));
-		}
-
-		for (int32 CoordinateIndex = StreamComponentData.TextureCoordinates.Num(); CoordinateIndex < MAX_STATIC_TEXCOORDS / 2; CoordinateIndex++)
-		{
-			OutElements.Add(AccessStreamComponent(
-				StreamComponentData.TextureCoordinates[StreamComponentData.TextureCoordinates.Num() - 1],
-				BaseTexCoordAttribute + CoordinateIndex,
-				InOutStreams
-			));
-		}
+		const int32 BaseTexCoordAttribute = 5;	
+		uint8 TexCoordStreamIndex = AddVertexStream(StreamComponentData.TextureCoordinates[0], InOutStreams);
+		AddStaticMeshTextureCoordinateElements(TexCoordStreamIndex, BaseTexCoordAttribute, StreamComponentData.TextureCoordinates, OutElements);
 	}
 }
 
@@ -278,7 +263,7 @@ void FGeometryCacheVertexVertexFactory::GetPSOPrecacheVertexFetchElements(EVerte
 				OutElements.Add(FVertexElement(4, 0, VET_Float3, 4, sizeof(FVector3f), false));
 				for (int32 CoordinateIndex = 0; CoordinateIndex < MAX_STATIC_TEXCOORDS / 2; CoordinateIndex++)
 				{
-					OutElements.Add(FVertexElement(5 + CoordinateIndex, 0, VET_Float2, 5 + CoordinateIndex, sizeof(FVector2f), false));
+					OutElements.Add(FVertexElement(5 + CoordinateIndex, 0, VET_Half4, 5 + CoordinateIndex, sizeof(FVector2f), false));
 				}
 
 				OutElements.Add(FVertexElement(9, 0, VET_UInt, 13, sizeof(uint32), true));
diff --git a/Engine/Source/Runtime/Engine/Private/GeometryCollectionRendering.cpp b/Engine/Source/Runtime/Engine/Private/GeometryCollectionRendering.cpp
--- a/Engine/Source/Runtime/Engine/Private/GeometryCollectionRendering.cpp
+++ b/Engine/Source/Runtime/Engine/Private/GeometryCollectionRendering.cpp
@@ -93,38 +93,45 @@ void FGeometryCollectionVertexFactory::ValidateCompiledResult(const FVertexFacto
 
 void FGeometryCollectionVertexFactory::GetPSOPrecacheVertexFetchElements(EVertexInputStreamType VertexInputStreamType, FVertexDeclarationElementList& Elements)
 {
-	Elements.Add(FVertexElement(0, 0, VET_Float3, 0, sizeof(float)*3u, false));
-
+	uint8 StreamIndex = 0;
+	Elements.Add(FVertexElement(StreamIndex++, 0, VET_Float3, 0, 0, false));
+		
 	if (VertexInputStreamType == EVertexInputStreamType::PositionAndNormalOnly)
 	{
 		// 2-axis TangentBasis components in a single buffer, hence *2u
-		Elements.Add(FVertexElement(1, 0, VET_PackedNormal, 2, sizeof(FPackedNormal)*2u, false));
+		Elements.Add(FVertexElement(StreamIndex++, 4, VET_PackedNormal, 2, 0, false));
 	}
 
+	// Primitive ID stream
 	if (UseGPUScene(GMaxRHIShaderPlatform, GMaxRHIFeatureLevel)
 		&& !PlatformGPUSceneUsesUniformBufferView(GMaxRHIShaderPlatform))
 	{
-		switch (VertexInputStreamType)
-		{
-		case EVertexInputStreamType::Default:
-		{
-			Elements.Add(FVertexElement(1, 0, VET_UInt, 13, sizeof(uint32), true));
-			break;
-		}
-		case EVertexInputStreamType::PositionOnly:
+		if (VertexInputStreamType == EVertexInputStreamType::Default)
 		{
-			Elements.Add(FVertexElement(1, 0, VET_UInt, 1, sizeof(uint32), true));
-			break;
+			Elements.Add(FVertexElement(StreamIndex++, 0, VET_UInt, 13, 0, true));
 		}
-		case EVertexInputStreamType::PositionAndNormalOnly:
+		else
 		{
-			Elements.Add(FVertexElement(2, 0, VET_UInt, 1, sizeof(uint32), true));
-			break;
-		}
-		default:
-			checkNoEntry();
+			Elements.Add(FVertexElement(StreamIndex++, 0, VET_UInt, 1, 0, true));
 		}
 	}
+
+	if (VertexInputStreamType == EVertexInputStreamType::Default && !RHISupportsManualVertexFetch(GMaxRHIShaderPlatform))
+	{
+		// Normals
+		Elements.Add(FVertexElement(StreamIndex, 0, VET_PackedNormal, 1, 0, false));
+		Elements.Add(FVertexElement(StreamIndex++, 4, VET_PackedNormal, 2, 0, false));
+
+		// Color
+		Elements.Add(FVertexElement(StreamIndex++, 0, VET_Color, 3, 0, false));
+
+		// TexCoords
+		int8 UVStreamIndex = StreamIndex;
+		Elements.Add(FVertexElement(StreamIndex, 0, VET_Half4, 4, 0, false));
+		Elements.Add(FVertexElement(StreamIndex, 8, VET_Half4, 5, 0, false));
+		Elements.Add(FVertexElement(StreamIndex, 16, VET_Half4, 6, 0, false));
+		Elements.Add(FVertexElement(StreamIndex++, 24, VET_Half4, 7, 0, false));
+	}
 }
 
 void FGeometryCollectionVertexFactory::InitRHI(FRHICommandListBase& RHICmdList)
@@ -204,33 +211,23 @@ void FGeometryCollectionVertexFactory::InitRHI(FRHICommandListBase& RHICmdList)
 
 	if (Data.TextureCoordinates.Num())
 	{
-		const int32 BaseTexCoordAttribute = 4;
-		for (int32 CoordinateIndex = 0; CoordinateIndex < Data.TextureCoordinates.Num(); ++CoordinateIndex)
+		const int32 BaseTexCoordAttribute = 4;			
+		uint8 TexCoordStreamIndex = AddVertexStream(Data.TextureCoordinates[0], Streams);
+		AddStaticMeshTextureCoordinateElements(TexCoordStreamIndex, BaseTexCoordAttribute, Data.TextureCoordinates, Elements);
+	}
+
+	if (IsStaticLightingAllowed())
+	{
+		if (Data.LightMapCoordinateComponent.VertexBuffer)
 		{
-			Elements.Add(AccessStreamComponent(
-				Data.TextureCoordinates[CoordinateIndex],
-				BaseTexCoordAttribute + CoordinateIndex
-			));
+			Elements.Add(AccessStreamComponent(Data.LightMapCoordinateComponent, 15));
 		}
-
-		for (int32 CoordinateIndex = Data.TextureCoordinates.Num(); CoordinateIndex < MAX_STATIC_TEXCOORDS / 2; ++CoordinateIndex)
+		else if (Data.TextureCoordinates.Num())
 		{
-			Elements.Add(AccessStreamComponent(
-				Data.TextureCoordinates[Data.TextureCoordinates.Num() - 1],
-				BaseTexCoordAttribute + CoordinateIndex
-			));
+			Elements.Add(AccessStreamComponent(Data.TextureCoordinates[0], 15));
 		}
 	}
 
-	if (Data.LightMapCoordinateComponent.VertexBuffer)
-	{
-		Elements.Add(AccessStreamComponent(Data.LightMapCoordinateComponent, 15));
-	}
-	else if (Data.TextureCoordinates.Num())
-	{
-		Elements.Add(AccessStreamComponent(Data.TextureCoordinates[0], 15));
-	}
-
 	check(Streams.Num() > 0);
 
 	InitDeclaration(Elements);
diff --git a/Engine/Source/Runtime/Engine/Private/InstancedStaticMesh.cpp b/Engine/Source/Runtime/Engine/Private/InstancedStaticMesh.cpp
--- a/Engine/Source/Runtime/Engine/Private/InstancedStaticMesh.cpp
+++ b/Engine/Source/Runtime/Engine/Private/InstancedStaticMesh.cpp
@@ -853,24 +853,9 @@ void FInstancedStaticMeshVertexFactory::GetVertexElements(
 
 		if (Data.TextureCoordinates.Num())
 		{
-			const int32 BaseTexCoordAttribute = 4;
-			for (int32 CoordinateIndex = 0; CoordinateIndex < Data.TextureCoordinates.Num(); CoordinateIndex++)
-			{
-				Elements.Add(AccessStreamComponent(
-					Data.TextureCoordinates[CoordinateIndex],
-					BaseTexCoordAttribute + CoordinateIndex,
-					Streams
-				));
-			}
-
-			for (int32 CoordinateIndex = Data.TextureCoordinates.Num(); CoordinateIndex < (InstancedStaticMeshMaxTexCoord + 1) / 2; CoordinateIndex++)
-			{
-				Elements.Add(AccessStreamComponent(
-					Data.TextureCoordinates[Data.TextureCoordinates.Num() - 1],
-					BaseTexCoordAttribute + CoordinateIndex,
-					Streams
-				));
-			}
+			const int32 BaseTexCoordAttribute = 4;				
+			uint8 TexCoordStreamIndex = AddVertexStream(Data.TextureCoordinates[0], Streams);
+			AddStaticMeshTextureCoordinateElements(TexCoordStreamIndex, BaseTexCoordAttribute, Data.TextureCoordinates, Elements);
 		}
 
 		// PreSkinPosition attribute is only used for GPUSkinPassthrough variation of local vertex factory.
@@ -881,13 +866,16 @@ void FInstancedStaticMeshVertexFactory::GetVertexElements(
 			Elements.Add(AccessStreamComponent(NullComponent, 14, Streams));
 		}
 
-		if (Data.LightMapCoordinateComponent.VertexBuffer)
+		if (IsStaticLightingAllowed())
 		{
-			Elements.Add(AccessStreamComponent(Data.LightMapCoordinateComponent, 15, Streams));
-		}
-		else if (Data.TextureCoordinates.Num())
-		{
-			Elements.Add(AccessStreamComponent(Data.TextureCoordinates[0], 15, Streams));
+			if (Data.LightMapCoordinateComponent.VertexBuffer)
+			{
+				Elements.Add(AccessStreamComponent(Data.LightMapCoordinateComponent, 15, Streams));
+			}
+			else if (Data.TextureCoordinates.Num())
+			{
+				Elements.Add(AccessStreamComponent(Data.TextureCoordinates[0], 15, Streams));
+			}
 		}
 	}
 
diff --git a/Engine/Source/Runtime/Engine/Private/LocalVertexFactory.cpp b/Engine/Source/Runtime/Engine/Private/LocalVertexFactory.cpp
--- a/Engine/Source/Runtime/Engine/Private/LocalVertexFactory.cpp
+++ b/Engine/Source/Runtime/Engine/Private/LocalVertexFactory.cpp
@@ -20,6 +20,7 @@
 #include "GPUSkinVertexFactory.h"
 #include "RenderUtils.h"
 #include "SceneInterface.h"
+#include "Rendering/StaticMeshVertexBuffer.h"
 
 IMPLEMENT_TYPE_LAYOUT(FLocalVertexFactoryShaderParametersBase);
 IMPLEMENT_TYPE_LAYOUT(FLocalVertexFactoryShaderParameters);
@@ -324,54 +325,72 @@ void FLocalVertexFactory::ValidateCompiledResult(const FVertexFactoryType* Type,
 }
 
 /**
-* Return the vertex elements used when manual vertex fetch is used
+* Return the vertex elements when no vertex declaration is provided during PSO precaching
 */
 void FLocalVertexFactory::GetPSOPrecacheVertexFetchElements(EVertexInputStreamType VertexInputStreamType, FVertexDeclarationElementList& Elements)
 {
-	Elements.Add(FVertexElement(0, 0, VET_Float3, 0, sizeof(float)*3u, false));
+	uint8 StreamIndex = 0;
+	Elements.Add(FVertexElement(StreamIndex++, 0, VET_Float3, 0, 0, false));
 
-	if (VertexInputStreamType == EVertexInputStreamType::PositionAndNormalOnly)
+	switch (VertexInputStreamType)
 	{
-		// 2-axis TangentBasis components in a single buffer, hence *2u
-		Elements.Add(FVertexElement(1, 4, VET_PackedNormal, 2, sizeof(FPackedNormal)*2u, false));
-	}
-
-	if (UseGPUScene(GMaxRHIShaderPlatform, GMaxRHIFeatureLevel)
-		&& !PlatformGPUSceneUsesUniformBufferView(GMaxRHIShaderPlatform))
-	{
-		switch (VertexInputStreamType)
-		{
 		case EVertexInputStreamType::Default:
 		{
-			Elements.Add(FVertexElement(1, 0, VET_UInt, 13, sizeof(uint32), true));
 			if (!RHISupportsManualVertexFetch(GMaxRHIShaderPlatform))
 			{
-				Elements.Add(FVertexElement(2, 0, VET_PackedNormal, 1, 8, false));
-				Elements.Add(FVertexElement(2, 4, VET_PackedNormal, 2, 8, false));
-				Elements.Add(FVertexElement(3, 0, VET_Color, 3, 4, false));
-				Elements.Add(FVertexElement(4, 0, VET_Half4, 4, 32, false));
-				Elements.Add(FVertexElement(4, 8, VET_Half4, 5, 32, false));
-				Elements.Add(FVertexElement(4, 16, VET_Half4, 6, 32, false));
-				Elements.Add(FVertexElement(4, 24, VET_Half4, 7, 32, false));
+				// Normals
+				Elements.Add(FVertexElement(StreamIndex, 0, VET_PackedNormal, 1, 0, false));
+				Elements.Add(FVertexElement(StreamIndex++, 4, VET_PackedNormal, 2, 0, false));
+
+				// Color
+				Elements.Add(FVertexElement(StreamIndex++, 0, VET_Color, 3, 0, false));
+
+				// TexCoords
+				int8 UVStreamIndex = StreamIndex;
+				Elements.Add(FVertexElement(StreamIndex, 0, VET_Half4, 4, 0, false));
+				Elements.Add(FVertexElement(StreamIndex, 8, VET_Half4, 5, 0, false));
+				Elements.Add(FVertexElement(StreamIndex, 16, VET_Half4, 6, 0, false));
+				Elements.Add(FVertexElement(StreamIndex++, 24, VET_Half4, 7, 0, false));
+
+				// Pre skin data
 				if (IsGPUSkinPassThroughSupported(GMaxRHIShaderPlatform))
 				{
-					Elements.Add(FVertexElement(5, 0, VET_Float3, 14, sizeof(float)*3u, false));
+					Elements.Add(FVertexElement(StreamIndex++, 0, VET_Float3, 14, 0, false));
+				}
+
+				// Light map data
+				if (IsStaticLightingAllowed())
+				{
+					Elements.Add(FVertexElement(UVStreamIndex, 0, VET_Half2, 15, 0, false));
 				}
 			}
 			break;
 		}
 		case EVertexInputStreamType::PositionOnly:
 		{
-			Elements.Add(FVertexElement(1, 0, VET_UInt, 1, sizeof(uint32), true));
 			break;
 		}
 		case EVertexInputStreamType::PositionAndNormalOnly:
 		{
-			Elements.Add(FVertexElement(2, 0, VET_UInt, 1, sizeof(uint32), true));
+			// 2-axis TangentBasis components in a single buffer, hence *2u
+			Elements.Add(FVertexElement(StreamIndex++, 4, VET_PackedNormal, 2, 0, false));
 			break;
 		}
 		default:
 			checkNoEntry();
+	}
+	
+	// Primitive ID stream
+	if (UseGPUScene(GMaxRHIShaderPlatform, GMaxRHIFeatureLevel)
+		&& !PlatformGPUSceneUsesUniformBufferView(GMaxRHIShaderPlatform))
+	{
+		if (VertexInputStreamType == EVertexInputStreamType::Default)
+		{
+			Elements.Add(FVertexElement(StreamIndex, 0, VET_UInt, 13, 0, true));
+		}
+		else
+		{
+			Elements.Add(FVertexElement(StreamIndex, 0, VET_UInt, 1, 0, true));
 		}
 	}
 }
@@ -385,7 +404,14 @@ void FLocalVertexFactory::GetVertexElements(ERHIFeatureLevel::Type FeatureLevel,
 	if (UseGPUScene(GMaxRHIShaderPlatform, GMaxRHIFeatureLevel) 
 		&& !PlatformGPUSceneUsesUniformBufferView(GMaxRHIShaderPlatform))
 	{
-		Elements.Add(FVertexElement(VertexStreams.Num(), 0, VET_UInt, 13, sizeof(uint32), true));
+		if (InputStreamType == EVertexInputStreamType::Default)
+		{
+			Elements.Add(FVertexElement(VertexStreams.Num(), 0, VET_UInt, 13, 0, true));
+		}
+		else
+		{
+			Elements.Add(FVertexElement(VertexStreams.Num(), 0, VET_UInt, 1, 0, true));
+		}
 	}
 }
 
@@ -474,6 +500,52 @@ void FLocalVertexFactory::InitRHI(FRHICommandListBase& RHICmdList)
 	check(IsValidRef(GetDeclaration()));
 }
 
+void AddStaticMeshTextureCoordinateElements(
+	uint8 TexCoordStreamIndex, 
+	int32 BaseTexCoordAttribute, 
+	const TArrayView<FVertexStreamComponent>& TextureCoordinates, 
+	FVertexDeclarationElementList& Elements)
+{
+	if (TextureCoordinates.IsEmpty())
+	{
+		return; 
+	}
+
+	// Validate all texture coordinates share the same stream data
+	for (int32 TextureCoordIndex = 1; TextureCoordIndex < TextureCoordinates.Num(); ++TextureCoordIndex)
+	{
+		check(TextureCoordinates[0].VertexBuffer == TextureCoordinates[TextureCoordIndex].VertexBuffer);
+		check(TextureCoordinates[0].Stride == TextureCoordinates[TextureCoordIndex].Stride);
+		check(TextureCoordinates[0].StreamOffset == TextureCoordinates[TextureCoordIndex].StreamOffset);
+		check(TextureCoordinates[0].VertexStreamUsage == TextureCoordinates[TextureCoordIndex].VertexStreamUsage);
+	}
+
+	EVertexElementType FirstElementType = TextureCoordinates[0].Type;
+			
+	EVertexElementType UVDoubleWideVertexElementType = VET_None;
+	uint32 DoubleUVSizeInBytes = 0;
+	if (FirstElementType == VET_Float2 || FirstElementType == VET_Float4)
+	{
+		DoubleUVSizeInBytes = sizeof(TStaticMeshVertexUVsTypeSelector<EStaticMeshVertexUVType::HighPrecision>::UVsTypeT) * 2;
+		UVDoubleWideVertexElementType = VET_Float4;
+	}
+	else
+	{
+		DoubleUVSizeInBytes = sizeof(TStaticMeshVertexUVsTypeSelector<EStaticMeshVertexUVType::Default>::UVsTypeT) * 2;
+		UVDoubleWideVertexElementType = VET_Half4;
+	}
+			
+	for (int32 CoordinateIndex = 0; CoordinateIndex < MAX_STATIC_TEXCOORDS / 2; ++CoordinateIndex)
+	{
+		Elements.Add(FVertexElement(
+			TexCoordStreamIndex, 
+			DoubleUVSizeInBytes * CoordinateIndex, 
+			UVDoubleWideVertexElementType, 
+			BaseTexCoordAttribute + CoordinateIndex, 
+			TextureCoordinates[0].Stride));
+	}
+}
+
 void FLocalVertexFactory::GetVertexElements(
 	ERHIFeatureLevel::Type FeatureLevel, 
 	EVertexInputStreamType InputStreamType,
@@ -530,23 +602,8 @@ void FLocalVertexFactory::GetVertexElements(
 		if (Data.TextureCoordinates.Num())
 		{
 			const int32 BaseTexCoordAttribute = 4;
-			for (int32 CoordinateIndex = 0; CoordinateIndex < Data.TextureCoordinates.Num(); ++CoordinateIndex)
-			{
-				Elements.Add(AccessStreamComponent(
-					Data.TextureCoordinates[CoordinateIndex],
-					BaseTexCoordAttribute + CoordinateIndex,
-					InOutStreams
-				));
-			}
-
-			for (int32 CoordinateIndex = Data.TextureCoordinates.Num(); CoordinateIndex < MAX_STATIC_TEXCOORDS / 2; ++CoordinateIndex)
-			{
-				Elements.Add(AccessStreamComponent(
-					Data.TextureCoordinates[Data.TextureCoordinates.Num() - 1],
-					BaseTexCoordAttribute + CoordinateIndex,
-					InOutStreams
-				));
-			}
+			uint8 TexCoordStreamIndex = AddVertexStream(Data.TextureCoordinates[0], InOutStreams);
+			AddStaticMeshTextureCoordinateElements(TexCoordStreamIndex, BaseTexCoordAttribute, Data.TextureCoordinates, Elements);
 		}
 
 		// TODO: should also check if VFType supports 'SupportsGPUSkinPassThrough'
@@ -557,13 +614,16 @@ void FLocalVertexFactory::GetVertexElements(
 			Elements.Add(AccessStreamComponent(Data.PreSkinPositionComponent.VertexBuffer ? Data.PreSkinPositionComponent : NullComponent, 14, InOutStreams));
 		}
 
-		if (Data.LightMapCoordinateComponent.VertexBuffer)
+		if (IsStaticLightingAllowed())
 		{
-			Elements.Add(AccessStreamComponent(Data.LightMapCoordinateComponent, 15, InOutStreams));
-		}
-		else if (Data.TextureCoordinates.Num())
-		{
-			Elements.Add(AccessStreamComponent(Data.TextureCoordinates[0], 15, InOutStreams));
+			if (Data.LightMapCoordinateComponent.VertexBuffer)
+			{
+				Elements.Add(AccessStreamComponent(Data.LightMapCoordinateComponent, 15, InOutStreams));
+			}
+			else if (Data.TextureCoordinates.Num())
+			{
+				Elements.Add(AccessStreamComponent(Data.TextureCoordinates[0], 15, InOutStreams));
+			}
 		}
 	}
 }
diff --git a/Engine/Source/Runtime/Engine/Private/MeshParticleVertexFactory.cpp b/Engine/Source/Runtime/Engine/Private/MeshParticleVertexFactory.cpp
--- a/Engine/Source/Runtime/Engine/Private/MeshParticleVertexFactory.cpp
+++ b/Engine/Source/Runtime/Engine/Private/MeshParticleVertexFactory.cpp
@@ -12,6 +12,7 @@
 #include "DataDrivenShaderPlatformInfo.h"
 #include "MeshUVChannelInfo.h"
 #include "Misc/DelayedAutoRegister.h"
+#include "LocalVertexFactory.h"
 
 class FMeshParticleVertexFactoryShaderParameters : public FVertexFactoryShaderParameters
 {
@@ -172,24 +173,9 @@ void FMeshParticleVertexFactory::GetVertexElements(ERHIFeatureLevel::Type Featur
 
 	if(Data.TextureCoordinates.Num())
 	{
-		const int32 BaseTexCoordAttribute = 4;
-		for(int32 CoordinateIndex = 0;CoordinateIndex < Data.TextureCoordinates.Num();CoordinateIndex++)
-		{
-			Elements.Add(AccessStreamComponent(
-				Data.TextureCoordinates[CoordinateIndex],
-				BaseTexCoordAttribute + CoordinateIndex,
-				InOutStreams
-			));
-		}
-
-		for(int32 CoordinateIndex = Data.TextureCoordinates.Num();CoordinateIndex < MAX_TEXCOORDS;CoordinateIndex++)
-		{
-			Elements.Add(AccessStreamComponent(
-				Data.TextureCoordinates[Data.TextureCoordinates.Num() - 1],
-				BaseTexCoordAttribute + CoordinateIndex,
-				InOutStreams
-			));
-		}
+		const int32 BaseTexCoordAttribute = 4;		
+		uint8 TexCoordStreamIndex = AddVertexStream(Data.TextureCoordinates[0], InOutStreams);
+		AddStaticMeshTextureCoordinateElements(TexCoordStreamIndex, BaseTexCoordAttribute, Data.TextureCoordinates, Elements);
 	}
 }
 
@@ -309,16 +295,16 @@ void FMeshParticleVertexFactory::GetPSOPrecacheVertexFetchElements(EVertexInputS
 
 		// Normals
 		Elements.Add(FVertexElement(1, 0, VET_PackedNormal, 1, 0, false));
-		Elements.Add(FVertexElement(2, 0, VET_PackedNormal, 2, 0, false));
+		Elements.Add(FVertexElement(2, 4, VET_PackedNormal, 2, 0, false));
 
 		// Color
 		Elements.Add(FVertexElement(3, 0, VET_Color, 3, 0, false));
 
 		// Texcoords
-		Elements.Add(FVertexElement(4, 0, VET_Half2, 4, 0, false));
-		Elements.Add(FVertexElement(5, 0, VET_Half2, 5, 0, false));
-		Elements.Add(FVertexElement(6, 0, VET_Half2, 6, 0, false));
-		Elements.Add(FVertexElement(7, 0, VET_Half2, 7, 0, false));
+		Elements.Add(FVertexElement(4, 0, VET_Half4, 4, 0, false));
+		Elements.Add(FVertexElement(5, 4, VET_Half4, 5, 0, false));
+		Elements.Add(FVertexElement(6, 8, VET_Half4, 6, 0, false));
+		Elements.Add(FVertexElement(7, 12, VET_Half4, 7, 0, false));
 	}
 
 	// Per instance data
diff --git a/Engine/Source/Runtime/Engine/Public/LocalVertexFactory.h b/Engine/Source/Runtime/Engine/Public/LocalVertexFactory.h
--- a/Engine/Source/Runtime/Engine/Public/LocalVertexFactory.h
+++ b/Engine/Source/Runtime/Engine/Public/LocalVertexFactory.h
@@ -45,6 +45,12 @@ extern ENGINE_API TUniformBufferRef<FLocalVertexFactoryUniformShaderParameters>
 	int32 PreSkinBaseVertexIndex
 	);
 
+extern ENGINE_API void AddStaticMeshTextureCoordinateElements(
+	uint8 TexCoordStreamIndex, 
+	int32 BaseTexCoordAttribute, 
+	const TArrayView<FVertexStreamComponent>& TextureCoordinates, 
+	FVertexDeclarationElementList& Elements);
+
 /**
  * A vertex factory which simply transforms explicit vertex attributes from local to world space.
  */
@@ -175,7 +181,7 @@ public:
 
 	inline const class UStaticMesh* GetStaticMesh() const { return Data.StaticMesh; }
 #endif
-
+		
 protected:
 	friend class FLocalVertexFactoryShaderParameters;
 	friend class FSkeletalMeshSceneProxy;
diff --git a/Engine/Source/Runtime/RenderCore/Public/VertexFactory.h b/Engine/Source/Runtime/RenderCore/Public/VertexFactory.h
--- a/Engine/Source/Runtime/RenderCore/Public/VertexFactory.h
+++ b/Engine/Source/Runtime/RenderCore/Public/VertexFactory.h
@@ -794,16 +794,22 @@ protected:
 	 */
 	template<typename VertexStreamListType>
 	static FVertexElement AccessStreamComponent(const FVertexStreamComponent& Component, uint8 AttributeIndex, VertexStreamListType& InOutStreams)
+	{
+		uint8 StreamIndex = AddVertexStream(Component, InOutStreams);
+		return FVertexElement(StreamIndex, Component.Offset, Component.Type, AttributeIndex, Component.Stride, EnumHasAnyFlags(EVertexStreamUsage::Instancing, Component.VertexStreamUsage));
+	}
+
+	template<typename VertexStreamListType>
+	static uint8 AddVertexStream(const FVertexStreamComponent& Component, VertexStreamListType& InOutStreams)
 	{
 		FVertexStream VertexStream;
 		VertexStream.VertexBuffer = Component.VertexBuffer;
 		VertexStream.Stride = Component.Stride;
 		VertexStream.Offset = Component.StreamOffset;
 		VertexStream.VertexStreamUsage = Component.VertexStreamUsage;
-		return FVertexElement((uint8)InOutStreams.AddUnique(VertexStream), Component.Offset, Component.Type, AttributeIndex, VertexStream.Stride, EnumHasAnyFlags(EVertexStreamUsage::Instancing, VertexStream.VertexStreamUsage));
+		return (uint8)InOutStreams.AddUnique(VertexStream);
 	}
 
-
 	/**
 	 * Initializes the vertex declaration.
 	 * @param Elements - The elements of the vertex declaration.
@@ -832,7 +838,7 @@ protected:
 	};
 
 	typedef TArray<FVertexStream, TInlineAllocator<8> > FVertexStreamList;
-		
+	
 	/** The vertex streams used to render the factory. */
 	FVertexStreamList Streams;
 	
