[NetPhysics] Fix input issue for sim-proxies not merging inputs in network physics component
- And fix issue where jump would not be handled correctly in Input Decay or Input Interpolation.

diff --git a/Engine/Plugins/Experimental/Mover/Source/Mover/Private/MoverDataModelTypes.cpp b/Engine/Plugins/Experimental/Mover/Source/Mover/Private/MoverDataModelTypes.cpp
--- a/Engine/Plugins/Experimental/Mover/Source/Mover/Private/MoverDataModelTypes.cpp
+++ b/Engine/Plugins/Experimental/Mover/Source/Mover/Private/MoverDataModelTypes.cpp
@@ -123,11 +123,14 @@ void FCharacterDefaultInputs::Interpolate(const FMoverDataStructBase& From, cons
 {
 	const FCharacterDefaultInputs& TypedFrom = static_cast<const FCharacterDefaultInputs&>(From);
 	const FCharacterDefaultInputs& TypedTo = static_cast<const FCharacterDefaultInputs&>(To);
-	
-	// Note, this ignores movement base as this is not used by the physics mover
-	const FCharacterDefaultInputs* ClosestInputs = Pct < 0.5f ? &TypedFrom : &TypedTo;
-	bIsJumpJustPressed = ClosestInputs->bIsJumpJustPressed;
+
+	const FCharacterDefaultInputs* ClosestInputs = Pct + UE_SMALL_NUMBER < 0.5f ? &TypedFrom : &TypedTo;
+
+	// Set bIsJumpJustPressed if previous input was not jumping while next input is jumping but not just pressed (just pressed needs to have happened between the two inputs)
+	bIsJumpJustPressed = (!TypedFrom.bIsJumpJustPressed && !TypedFrom.bIsJumpPressed) && (!ClosestInputs->bIsJumpJustPressed && ClosestInputs->bIsJumpPressed);
 	bIsJumpPressed = ClosestInputs->bIsJumpPressed;
+
+	// Note, this ignores movement base as this is not used by the physics mover
 	SuggestedMovementMode = ClosestInputs->SuggestedMovementMode;
 
 	SetMoveInput(ClosestInputs->GetMoveInputType(), FMath::Lerp(TypedFrom.GetMoveInput(), TypedTo.GetMoveInput(), Pct));
@@ -153,9 +156,6 @@ void FCharacterDefaultInputs::Decay(float DecayAmount)
 	DecayAmount *= CharacterDefaultInputsDecayAmountMultiplier;
 
 	MoveInput *= (1.0f - DecayAmount);
-
-	// Single use inputs
-	bIsJumpJustPressed = FMath::IsNearlyZero(DecayAmount) ? bIsJumpJustPressed : false;
 }
 
 // FMoverDefaultSyncState //////////////////////////////////////////////////////////////
diff --git a/Engine/Source/Runtime/Engine/Private/PhysicsEngine/NetworkPhysicsComponent.cpp b/Engine/Source/Runtime/Engine/Private/PhysicsEngine/NetworkPhysicsComponent.cpp
--- a/Engine/Source/Runtime/Engine/Private/PhysicsEngine/NetworkPhysicsComponent.cpp
+++ b/Engine/Source/Runtime/Engine/Private/PhysicsEngine/NetworkPhysicsComponent.cpp
@@ -76,7 +76,7 @@ namespace PhysicsReplicationCVars
 		bool bApplyDataInsteadOfMergeData = false;
 		static FAutoConsoleVariableRef CVarResimApplyDataInsteadOfMergeData(TEXT("np2.Resim.ApplyDataInsteadOfMergeData"), bApplyDataInsteadOfMergeData, TEXT("When true, call ApplyData for each data instead of MergeData when having to use multiple data entries in one frame."));
 		bool bAllowInputExtrapolation = true;
-		static FAutoConsoleVariableRef CVarResimAllowInputExtrapolation(TEXT("np2.Resim.AllowInputExtrapolation"), bAllowInputExtrapolation, TEXT("When true and not locally controlled, allow inputs to be extrapolated from last known and if there is a gap allow interpolation between two known inputs."));
+		static FAutoConsoleVariableRef CVarResimAllowInputExtrapolation(TEXT("np2.Resim.AllowInputExtrapolation"), bAllowInputExtrapolation, TEXT("When true, allow inputs to be extrapolated from last known on the server and if there is a gap allow interpolation between two known inputs."));
 		bool bValidateDataOnGameThread = false;
 		static FAutoConsoleVariableRef CVarResimValidateDataOnGameThread(TEXT("np2.Resim.ValidateDataOnGameThread"), bValidateDataOnGameThread, TEXT("When true, perform server-side input validation through FNetworkPhysicsData::ValidateData on the Game Thread, note that LocalFrame will be the same as ServerFrame on Game Thread. If false, perform the call on the Physics Thread."));
 		bool bApplySimProxyStateAtRuntime = false;
@@ -1213,7 +1213,7 @@ void UNetworkPhysicsComponent::OnRep_SetReplicatedInputs()
 
 	if (!NetworkPhysicsComponent_Internal || !InputHelper || !ReplicatedInputs.History)
 	{
-		return ;
+		return;
 	}
 	ensure(bIsUsingLegacyData);
 
@@ -1240,7 +1240,7 @@ void UNetworkPhysicsComponent::OnRep_SetReplicatedRemoteInputs()
 
 	if (!NetworkPhysicsComponent_Internal || !InputHelper || !ReplicatedRemoteInputs.History)
 	{
-		return ;
+		return;
 	}
 	ensure(bIsUsingLegacyData);
 
@@ -2113,22 +2113,27 @@ void FAsyncNetworkPhysicsComponent::OnPreProcessInputs_Internal(const int32 Phys
 		if ((bApplyServerInput || bApplySimProxyInput || bIsSolverResim) && InputHistory && InputData)
 		{
 			FNetworkPhysicsPayload* PhysicsData = InputData.Get();
-			int32 NextExpectedLocalFrame = bIsSolverResim ? PhysicsStep : PhysicsData->LocalFrame + 1;
+			int32 NextExpectedInputFrame = PhysicsStep;
+			if (!bIsSolverResim)
+			{
+				NextExpectedInputFrame = bApplySimProxyInput ? LastInputFrameApplied : PhysicsData->LocalFrame;
+				NextExpectedInputFrame++; // Iterate one to get the next frame number we expect to be used
+			}
 
 			// There are important inputs earlier than upcoming input to apply
-			if (!bIsSolverResim && NewImportantInputFrame < NextExpectedLocalFrame)
+			if (!bIsSolverResim && NewImportantInputFrame < NextExpectedInputFrame)
 			{
 				if (ComponentSettings.GetApplyDataInsteadOfMergeData())
 				{
 #if DEBUG_NETWORK_PHYSICS
-					UE_LOG(LogChaos, Log, TEXT("		Non-Determinism: Reapplying multiple data due to receiving an important data that was previously missed. FromFrame: %d | ToFrame: %d | IsLocallyControlled = %d"), NewImportantInputFrame, (NextExpectedLocalFrame - 1), IsLocallyControlled());
+					UE_LOG(LogChaos, Log, TEXT("		Non-Determinism: Reapplying multiple data due to receiving an important data that was previously missed. FromFrame: %d | ToFrame: %d | IsLocallyControlled = %d"), NewImportantInputFrame, (NextExpectedInputFrame - 1), IsLocallyControlled());
 #endif
 					if (bIsUsingLegacyData)
 					{
 						if (ensure(ActorComponent.IsValid() && ActorComponent.Get()->IsBeingDestroyed() == false))
 						{
 							// Apply all inputs in range
-							InputHistory->ApplyDataRange(NewImportantInputFrame, NextExpectedLocalFrame - 1, ActorComponent.Get(), /*bOnlyImportant*/false);
+							InputHistory->ApplyDataRange(NewImportantInputFrame, NextExpectedInputFrame - 1, ActorComponent.Get(), /*bOnlyImportant*/false);
 						}
 					}
 /* // Importance is not implemented into the new network flow (yet?)
@@ -2136,7 +2141,7 @@ void FAsyncNetworkPhysicsComponent::OnPreProcessInputs_Internal(const int32 Phys
 					{
 						if (InputHelper && ensure(ImplementationInterface_Internal))
 						{
-							InputHelper->ApplyDataRange(InputHistory.Get(), NewImportantInputFrame, NextExpectedLocalFrame - 1, *ImplementationInterface_Internal);
+							InputHelper->ApplyDataRange(InputHistory.Get(), NewImportantInputFrame, NextExpectedInputFrame - 1, *ImplementationInterface_Internal);
 						}
 					}
 */
@@ -2144,15 +2149,18 @@ void FAsyncNetworkPhysicsComponent::OnPreProcessInputs_Internal(const int32 Phys
 				else
 				{
 					// Merge all inputs from earliest new important
-					NextExpectedLocalFrame = NewImportantInputFrame;
+					NextExpectedInputFrame = NewImportantInputFrame;
 #if DEBUG_NETWORK_PHYSICS
-					UE_LOG(LogChaos, Log, TEXT("		Non-Determinism: Prepare to reapply multiple data through MergeData due to receiving an important data that was previously missed. FromFrame: %d | ToFrame: %d | IsLocallyControlled = %d"), NewImportantInputFrame, (NextExpectedLocalFrame - 1), IsLocallyControlled());
+					UE_LOG(LogChaos, Log, TEXT("		Non-Determinism: Prepare to reapply multiple data through MergeData due to receiving an important data that was previously missed. FromFrame: %d | ToFrame: %d | IsLocallyControlled = %d"), NewImportantInputFrame, (NextExpectedInputFrame - 1), IsLocallyControlled());
 #endif
 				}
 			}
 
-			if (InputHistory->ExtractData(PhysicsStep, bIsSolverReset, PhysicsData, /*bExactFrame*/(ComponentSettings.GetAllowInputExtrapolation() == false)))
+			const bool RequireExactFrame = bIsServer && !ComponentSettings.GetAllowInputExtrapolation();
+			if (InputHistory->ExtractData(PhysicsStep, bIsSolverReset, PhysicsData, RequireExactFrame))
 			{
+				const int32 CurrentInputFrame = bApplySimProxyInput ? InputHistory->GetLatestFrame() : PhysicsData->LocalFrame;
+
 				// Check if the extracted data was altered and if we have a hole in the buffer
 				if (bIsServer && PhysicsData->bDataAltered)
 				{
@@ -2180,18 +2188,20 @@ void FAsyncNetworkPhysicsComponent::OnPreProcessInputs_Internal(const int32 Phys
 					}
 				}
 				// Check if we have a gap between last used input and current input
-				else if (PhysicsData->LocalFrame > NextExpectedLocalFrame)
+				else if (CurrentInputFrame > NextExpectedInputFrame)
 				{
 					if (ComponentSettings.GetApplyDataInsteadOfMergeData())
 					{
 #if DEBUG_NETWORK_PHYSICS
-						UE_LOG(LogChaos, Log, TEXT("		Non-Determinism: Applying multiple data instead of merging, from LocalFrame %d into LocalFrame %d | IsLocallyControlled = %d"), NextExpectedLocalFrame, PhysicsData->LocalFrame, IsLocallyControlled());
+						UE_LOG(LogChaos, Log, TEXT("		Non-Determinism: Applying multiple data instead of merging, from LocalFrame %d into LocalFrame %d | IsLocallyControlled = %d"), NextExpectedInputFrame, PhysicsData->LocalFrame, IsLocallyControlled());
 #endif
 						// Iterate over each input and call ApplyData, except on the last, it will get handled by the normal ApplyData call further down
-						const int32 LastFrame = PhysicsData->LocalFrame;
-						for (; NextExpectedLocalFrame <= LastFrame; NextExpectedLocalFrame++)
+						for (; NextExpectedInputFrame <= CurrentInputFrame; NextExpectedInputFrame++)
 						{
-							if (InputHistory->ExtractData(NextExpectedLocalFrame, bIsSolverReset, PhysicsData, true) && NextExpectedLocalFrame < LastFrame)
+							/** Note, extract the last frame but don't apply it here, it will be applied later in the flow
+							* Also note that we allow extracting non-exact data on the last frame, unless required by settings, to get the same data extracted as we did in the original ExtractData above */
+							bool bLastInput = NextExpectedInputFrame == CurrentInputFrame;
+							if (InputHistory->ExtractData(NextExpectedInputFrame, bIsSolverReset, PhysicsData, (RequireExactFrame || !bLastInput)) && !bLastInput)
 							{
 								// Perform input decay on sim-proxy runtime (non-resim) input
 								if (bApplySimProxyInput && ComponentSettings.GetApplyInputDecaySimProxyInputAtRuntime())
@@ -2220,10 +2230,10 @@ void FAsyncNetworkPhysicsComponent::OnPreProcessInputs_Internal(const int32 Phys
 					else
 					{
 #if DEBUG_NETWORK_PHYSICS
-						UE_LOG(LogChaos, Log, TEXT("		Non-Determinism: Merging inputs from LocalFrame %d into LocalFrame %d | IsLocallyControlled = %d"), NextExpectedLocalFrame, PhysicsData->LocalFrame, IsLocallyControlled());
+						UE_LOG(LogChaos, Log, TEXT("		Non-Determinism: Merging inputs from LocalFrame %d into LocalFrame %d | IsLocallyControlled = %d"), NextExpectedInputFrame, PhysicsData->LocalFrame, IsLocallyControlled());
 #endif
 						// Merge all inputs since last used input
-						InputHistory->MergeData(NextExpectedLocalFrame, PhysicsData);
+						InputHistory->MergeData(NextExpectedInputFrame, PhysicsData);
 					}
 				}
 
@@ -2258,6 +2268,8 @@ void FAsyncNetworkPhysicsComponent::OnPreProcessInputs_Internal(const int32 Phys
 					}
 				}
 
+				LastInputFrameApplied = InputHistory->GetLatestFrame();
+
 #if DEBUG_NETWORK_PHYSICS
 				{
 					UE_LOG(LogChaos, Log, TEXT("			Applying extracted input from history | LocalFrame = %d | ServerFrame = %d | bDataAltered = %d | IsResim = %d | IsLocallyControlled = %d | InputDecay = %f | Data: %s")
diff --git a/Engine/Source/Runtime/Engine/Public/Physics/NetworkPhysicsComponent.h b/Engine/Source/Runtime/Engine/Public/Physics/NetworkPhysicsComponent.h
--- a/Engine/Source/Runtime/Engine/Public/Physics/NetworkPhysicsComponent.h
+++ b/Engine/Source/Runtime/Engine/Public/Physics/NetworkPhysicsComponent.h
@@ -2036,6 +2036,9 @@ private:
 	// Send last N number of states each replication call
 	uint16 StatesToNetwork = 1;
 
+	// Store input frame that was applied from the input history, used for simulated proxies which apply inputs "too late"
+	int32 LastInputFrameApplied = -1;
+
 	// Properties for dynamic scaling of buffered inputs
 	float InputBufferAverageMs = 90.0;
 	float InputBufferTargetSizeMs = 90.0;
