Fix LineTrace_Geom and FPhysInterface_Chaos::Sweep_Geom for FBody_Instance that are not owned by a Component or are not part of a Physic scene.

This problem was reported by a license on EPS (case 21685573)

diff --git a/Engine/Source/Runtime/Engine/Classes/PhysicsEngine/BodyInstance.h b/Engine/Source/Runtime/Engine/Classes/PhysicsEngine/BodyInstance.h
--- a/Engine/Source/Runtime/Engine/Classes/PhysicsEngine/BodyInstance.h
+++ b/Engine/Source/Runtime/Engine/Classes/PhysicsEngine/BodyInstance.h
@@ -706,6 +706,8 @@ public:
 	/** Owner containing this body (when not a UPrimitiveComponent). */
 	IPhysicsBodyInstanceOwner* BodyInstanceOwner;
 
+	ENGINE_API UWorld* GetWorld() const;
+
 	/** Source object for this body. If OwnerComponent is not null, this will typically be the owning actor. */
 	TWeakObjectPtr<UObject> SourceObject;
 
diff --git a/Engine/Source/Runtime/Engine/Private/PhysicsEngine/BodyInstance.cpp b/Engine/Source/Runtime/Engine/Private/PhysicsEngine/BodyInstance.cpp
--- a/Engine/Source/Runtime/Engine/Private/PhysicsEngine/BodyInstance.cpp
+++ b/Engine/Source/Runtime/Engine/Private/PhysicsEngine/BodyInstance.cpp
@@ -1129,6 +1129,18 @@ void FBodyInstance::UpdatePhysicsFilterData()
 	UpdateInterpolateWhenSubStepping();
 }
 
+UWorld* FBodyInstance::GetWorld() const
+{
+	const UPrimitiveComponent* OwnerComponentInst = OwnerComponent.Get();
+	if (OwnerComponentInst)
+	{
+		return OwnerComponentInst->GetWorld();
+	}
+	UObject* BodyInstanceObject = BodyInstanceOwner ? BodyInstanceOwner->GetSourceObject() : nullptr;
+	return BodyInstanceObject ? BodyInstanceObject->GetWorld() : nullptr;
+};
+
+
 TAutoConsoleVariable<int32> CDisableQueryOnlyActors(TEXT("p.DisableQueryOnlyActors"), 0, TEXT("If QueryOnly is used, actors are marked as simulation disabled. This is NOT compatible with origin shifting at the moment."));
 
 #if USE_BODYINSTANCE_DEBUG_NAMES
diff --git a/Engine/Source/Runtime/Engine/Private/PhysicsEngine/Experimental/PhysInterface_Chaos.cpp b/Engine/Source/Runtime/Engine/Private/PhysicsEngine/Experimental/PhysInterface_Chaos.cpp
--- a/Engine/Source/Runtime/Engine/Private/PhysicsEngine/Experimental/PhysInterface_Chaos.cpp
+++ b/Engine/Source/Runtime/Engine/Private/PhysicsEngine/Experimental/PhysInterface_Chaos.cpp
@@ -987,8 +987,9 @@ bool FPhysInterface_Chaos::LineTrace_Geom(FHitResult& OutHit, const FBodyInstanc
 						const Chaos::Filter::FQueryFilterData QueryFilter = Chaos::Filter::FQueryTraceFilterBuilder::BuildOverlapAll();
 
 						FTransform StartTM(WorldStart);
-						const UPrimitiveComponent* OwnerComponentInst = InInstance->OwnerComponent.Get();
-						ConvertQueryImpactHit(OwnerComponentInst ? OwnerComponentInst->GetWorld() : nullptr, BestHit, OutHit, DeltaMag, QueryFilter, WorldStart, WorldEnd, nullptr, StartTM, true, bExtractPhysMaterial);
+						const UWorld* World = InInstance->GetWorld();
+
+						ConvertQueryImpactHit(World, BestHit, OutHit, DeltaMag, QueryFilter, WorldStart, WorldEnd, nullptr, StartTM, true, bExtractPhysMaterial);
 						bHitSomething = true;
 					}
 				}
@@ -1018,7 +1019,7 @@ bool FPhysInterface_Chaos::Sweep_Geom(FHitResult& OutHit, const FBodyInstance* I
 
 		FPhysicsCommand::ExecuteRead(TargetInstance->GetPhysicsActor(), [&](const FPhysicsActorHandle& Actor)
 		{
-			if (Actor && InInstance->OwnerComponent.Get())
+			if (Actor)
 			{
 				FPhysicsShapeAdapter ShapeAdapter(InShapeRotation, InShape);
 
@@ -1027,9 +1028,7 @@ bool FPhysInterface_Chaos::Sweep_Geom(FHitResult& OutHit, const FBodyInstance* I
 				
 				const FTransform ActorTM(Actor->GetGameThreadAPI().R(), Actor->GetGameThreadAPI().X());
 
-				UPrimitiveComponent* OwnerComponentInst = InInstance->OwnerComponent.Get();
 				FTransform StartTM(ShapeAdapter.GetGeomOrientation(), InStart);
-				FTransform CompTM(OwnerComponentInst->GetComponentTransform());
 
 				const Chaos::FVec3 Dir = DeltaMag > 0 ? (Delta / DeltaMag) : Chaos::FVec3(1, 0, 0);
 
@@ -1098,7 +1097,9 @@ bool FPhysInterface_Chaos::Sweep_Geom(FHitResult& OutHit, const FBodyInstance* I
 							SetFlags(Hit, EHitFlags::Distance | EHitFlags::Normal | EHitFlags::Position | EHitFlags::FaceIndex);
 
 							FTransform StartTransform(InStart);
-							ConvertQueryImpactHit(OwnerComponentInst->GetWorld(), Hit, OutHit, DeltaMag, QueryFilter, InStart, InEnd, &QueryGeometry, StartTransform, true, false);
+							const UWorld* World = InInstance->GetWorld();
+
+							ConvertQueryImpactHit(World, Hit, OutHit, DeltaMag, QueryFilter, InStart, InEnd, &QueryGeometry, StartTransform, true, false);
 							bSweepHit = true;
 						}
 					}
