Fix for mf preview.
TODO test attribute with substrate off

#tests PC editor
#rnx

diff --git a/Engine/Source/Runtime/Engine/Private/Materials/Material.cpp b/Engine/Source/Runtime/Engine/Private/Materials/Material.cpp
--- a/Engine/Source/Runtime/Engine/Private/Materials/Material.cpp
+++ b/Engine/Source/Runtime/Engine/Private/Materials/Material.cpp
@@ -5494,6 +5494,13 @@ void UMaterial::RebuildShadingModelField()
 		FSubstrateMaterialInfo SubstrateMaterialInfo = GetSubstrateMaterialInfo(EditorOnly);
 		UpdatePropertyConnectedMask(SubstrateMaterialInfo, this->CachedExpressionData.Get());
 
+#if WITH_EDITOR
+		extern bool IsPreviewUnlitMaterial(const UMaterial * InMaterial);
+		const bool bIsMaterialFunctionPreview = IsPreviewUnlitMaterial(this);
+#else
+		const bool bIsMaterialFunctionPreview = false;
+#endif
+
 		if (SubstrateMaterialInfo.IsValid() && SubstrateMaterialInfo.GetSubstrateTreeOutOfStackDepthOccurred())
 		{
 			SubstrateMaterialInfo.AddShadingModel(SSM_Unlit);
@@ -5695,8 +5702,13 @@ void UMaterial::RebuildShadingModelField()
 		}
 		else
 		{
+			if (bIsMaterialFunctionPreview)
+			{
+				// Keep material domain selected from the preview window and set unlit for preview to work.
+				ShadingModel = MSM_Unlit;
+			}
 			// Now derive some properties from the material into the legacy fields
-			if (SubstrateMaterialInfo.HasOnlyShadingModel(SSM_Unlit))
+			else if (SubstrateMaterialInfo.HasOnlyShadingModel(SSM_Unlit))
 			{
 				MaterialDomain = EMaterialDomain::MD_Surface;
 				ShadingModel = MSM_Unlit;
diff --git a/Engine/Source/Runtime/Engine/Private/Materials/MaterialExpressions.cpp b/Engine/Source/Runtime/Engine/Private/Materials/MaterialExpressions.cpp
--- a/Engine/Source/Runtime/Engine/Private/Materials/MaterialExpressions.cpp
+++ b/Engine/Source/Runtime/Engine/Private/Materials/MaterialExpressions.cpp
@@ -326,6 +326,7 @@
 #include "Serialization/ShaderKeyGenerator.h"
 #include "SubstrateMaterial.h"
 #include "PostProcess/PostProcessMaterialInputs.h"
+#include "MaterialEditor/PreviewMaterial.h"
 #else
 #include "Materials/MaterialExpressionVertexInterpolator.h"
 #include "Materials/MaterialParameterCollection.h"
@@ -2132,6 +2133,15 @@ void UMaterialExpression::UpdateParameterGuid(bool bForceGeneration, bool bAllow
 
 #if WITH_EDITOR
 
+// A non public function to know if a material is a preview material that needs to be unlit.
+// This is required for instance when previewing a material function using the UI domain. It needs unlit to fetch the right data and Substrate also uses unlit behind the scene to display that domain.
+// Also this is used to maintain the correct MaterialDomain since Substrate can convert that otherwise in RebuildShadingModelField. 
+// When in fact we must keep that for the ExpressionPreviewMaterial created from MaterialEditor.cpp for MAterialFunction asset preview for instance.
+bool IsPreviewUnlitMaterial(const UMaterial* InMaterial)
+{
+	return InMaterial && InMaterial->MaterialDomain != MD_Surface && InMaterial->bIsPreviewMaterial && InMaterial->IsA(UPreviewMaterial::StaticClass());
+}
+
 void UMaterialExpression::ConnectToPreviewMaterial(UMaterial* InMaterial, int32 OutputIndex)
 {
 	// This is used when a node is right clicked and "Start previewing node" is used.
@@ -2141,7 +2151,14 @@ void UMaterialExpression::ConnectToPreviewMaterial(UMaterial* InMaterial, int32
 		{
 			if (IsResultSubstrateMaterial(0))
 			{
-				InMaterial->SetShadingModel(MSM_DefaultLit);
+				if (IsPreviewUnlitMaterial(InMaterial))
+				{
+					InMaterial->SetShadingModel(MSM_Unlit);
+				}
+				else
+				{
+					InMaterial->SetShadingModel(MSM_DefaultLit);
+				}
 				InMaterial->bUseMaterialAttributes = false;
 				FExpressionInput* MaterialInput = InMaterial->GetExpressionInputForProperty(MP_FrontMaterial);
 				check(MaterialInput);
@@ -2168,7 +2185,14 @@ void UMaterialExpression::ConnectToPreviewMaterial(UMaterial* InMaterial, int32
 				UMaterialExpressionSubstrateConvertMaterialAttributes* ConvertAttributeNode = NewObject<UMaterialExpressionSubstrateConvertMaterialAttributes>(this);
 				ConvertAttributeNode->Material = InMaterial;
 				ConvertAttributeNode->MaterialAttributes.Connect(OutputIndex, this);
-				ConvertAttributeNode->ShadingModelOverride = MSM_DefaultLit;
+				if (IsPreviewUnlitMaterial(InMaterial))
+				{
+					ConvertAttributeNode->ShadingModelOverride = MSM_Unlit;
+				}
+				else
+				{
+					ConvertAttributeNode->ShadingModelOverride = MSM_DefaultLit;
+				}
 
 				// Connect substrate data into material FrontMaterial input
 				if (UMaterialEditorOnlyData* MaterialEditorOnlyData = InMaterial->GetEditorOnlyData())
