Fixing multi trace results not being sorted correctly when the trace involves cluster unions.
Adding explicit sort after appending cluster union hit results

diff --git a/Engine/Source/Runtime/Engine/Private/Collision/SceneQuery.cpp b/Engine/Source/Runtime/Engine/Private/Collision/SceneQuery.cpp
--- a/Engine/Source/Runtime/Engine/Private/Collision/SceneQuery.cpp
+++ b/Engine/Source/Runtime/Engine/Private/Collision/SceneQuery.cpp
@@ -620,6 +620,7 @@ bool TSceneCastCommonImpWithRetryRequest(const UWorld* World, typename Traits::O
 						TArray<int32> ClusterUnionIndices;
 						TArray<AActor*, TInlineAllocator<1>> ClusterUnionActorsToIgnoreIfRetry;
 						bBlockingHit = false;
+						bool bResortHits = false;
 
 						for (int32 Index = 0; Index < OutHits.Num(); ++Index)
 						{
@@ -669,13 +670,20 @@ bool TSceneCastCommonImpWithRetryRequest(const UWorld* World, typename Traits::O
 							{
 								// No shrinking since we're going to be adding more elements shortly.
 								OutHits.RemoveAtSwap(ClusterUnionIndices[Index], EAllowShrinking::No);
+								bResortHits = true;
 							}
 	
 							if (Params.bReplaceHitWithSubComponents)
 							{
 								OutHits.Append(AllNewHits);
+								bResortHits = true;
 							}
 						}
+
+						if (bResortHits)
+						{
+							OutHits.Sort(FCompareFHitResultTime());
+						}
 					}
 					else
 					{
