[AsyncLoading2]

- Adds ImportRequestId tracking for AsyncPackages to know which requestid instigated the load of the package via an import
- Changes the ordering for how package Completion callbacks are invoked. We now provide a guarantee that callbacks for imported packages will run before their importing package's competion callback (this is also the case for circular dependencies). Previously, callbacks would only be invoked for the package being flushed exclusively.
- Adds tests for loading packages with hard references to ensure we are loading packages in the correct order for the correct phases.
-- For tests loading packages with circular dependencies we only verify that the expected imported packages are loaded and ignores order as we currently do not guarantee an ordering for those cases.
- Adds test for case in UE-222543 where a package sync loaded from a postload would not be flushed until the next syncloadevents call allowing for A loading B as part of its postload to have a callback ordering of [CallbackA , CallbackB], rather than what you'd expect, [CallbackB , CallbackA]

#rnx

diff --git a/Engine/Source/Runtime/CoreUObject/Private/Serialization/AsyncLoading2.cpp b/Engine/Source/Runtime/CoreUObject/Private/Serialization/AsyncLoading2.cpp
--- a/Engine/Source/Runtime/CoreUObject/Private/Serialization/AsyncLoading2.cpp
+++ b/Engine/Source/Runtime/CoreUObject/Private/Serialization/AsyncLoading2.cpp
@@ -29,6 +29,7 @@
 #include "Misc/ConfigCacheIni.h"
 #include "Misc/MessageDialog.h"
 #include "Misc/PackageName.h"
+#include "Misc/ReverseIterate.h"
 #include "Misc/StringBuilder.h"
 #include "UObject/ObjectResource.h"
 #include "UObject/SoftObjectPath.h"
@@ -4145,15 +4146,23 @@ private:
 	EAsyncLoadingResult::Type	LoadStatus = EAsyncLoadingResult::Succeeded;
 	/** True if this package was created by this async package */
 	bool						bCreatedLinkerRoot = false;
-	/** List of all request handles. */
-	TArray<int32> RequestIDs;
 
 	UE::FMutex ConstructedObjectsMutex;
 	/** List of ConstructedObjects = Exports + UPackage + ObjectsCreatedFromExports */
 	TArray<UObject*> ConstructedObjects;
+
+	typedef TArray<int32, TInlineAllocator<1>> FRequestIDs;
+	/** List of all request handles. */
+	FRequestIDs RequestIDs;
+
 	/** Detects if the constructed objects are improperly accessed by different threads at the same time. */
 	UE_MT_DECLARE_MRSW_RECURSIVE_ACCESS_DETECTOR(ConstructedObjectsAccessDetector);
 	TArray<FExternalReadCallback> ExternalReadDependencies;
+
+	// The fields below are only valid until the package is marked as completed
+
+	/** List of all request ids which caused the import of this package. */
+	FRequestIDs ImportRequestIDs;
 	/** Callbacks called when we finished loading this package */
 	TArray<TUniquePtr<FLoadPackageAsyncDelegate>, TInlineAllocator<2>> CompletionCallbacks;
 	/** Callbacks called for the different loading phase of this package */
@@ -4468,9 +4477,10 @@ private:
 		EAsyncLoadingResult::Type Result = EAsyncLoadingResult::Succeeded;
 		UPackage* UPackage = nullptr;
 		FAsyncPackage2* AsyncPackage = nullptr;
+		FAsyncPackage2::FRequestIDs RequestIDs;
 		TArray<TUniquePtr<FLoadPackageAsyncDelegate>, TInlineAllocator<2>> CompletionCallbacks;
 		TArray<TSharedPtr<FLoadPackageAsyncProgressDelegate>, TInlineAllocator<2>> ProgressCallbacks;
-		TArray<int32, TInlineAllocator<2>> RequestIDs;
+		FAsyncPackage2::FRequestIDs ImportRequestIDs;
 
 		static FCompletedPackageRequest FromUnreslovedPackage(
 			const FAsyncPackageDesc2& Desc,
@@ -4496,11 +4506,12 @@ private:
 				.Result = Package->GetLoadStatus(),
 				.UPackage = Package->LinkerRoot,
 				.AsyncPackage = Package,
+				.RequestIDs = Package->RequestIDs, // must be a copy as this array is accessed after the package marked as completed
 				.CompletionCallbacks = MoveTemp(Package->CompletionCallbacks),
-				.ProgressCallbacks = MoveTemp(Package->ProgressCallbacks)
+				.ProgressCallbacks = MoveTemp(Package->ProgressCallbacks),
+				.ImportRequestIDs = MoveTemp(Package->ImportRequestIDs)
 			};
 
-			Request.RequestIDs.Append(Package->RequestIDs);
 			return Request;
 		}
 
@@ -4870,7 +4881,7 @@ public:
 	/**
 	 * [ASYNC/GAME THREAD] Removes a request ID from the list of pending requests
 	 */
-	void RemovePendingRequests(FAsyncLoadingThreadState2& ThreadState, TConstArrayView<int32> RequestIDs)
+	void RemovePendingRequests(FAsyncLoadingThreadState2& ThreadState, const FAsyncPackage2::FRequestIDs& RequestIDs)
 	{
 		int32 RemovedCount = 0;
 		{
@@ -5178,6 +5189,7 @@ private:
 	*/
 	EAsyncPackageState::Type ProcessLoadedPackagesFromGameThread(FAsyncLoadingThreadState2& ThreadState, bool& bDidSomething, TConstArrayView<int32> FlushRequestIDs = {});
 
+	void GatherCompletedPackagesToFlush(FAsyncLoadingThreadState2& ThreadState, TConstArrayView<int32> FlushRequestIDs, TArray<FCompletedPackageRequest>& OutRequestsToProcess);
 	void IncludePackageInSyncLoadContextRecursive(FAsyncLoadingThreadState2& ThreadState, uint64 ContextId, FAsyncPackage2* Package);
 	void UpdateSyncLoadContext(FAsyncLoadingThreadState2& ThreadState, bool bAutoHandleSyncLoadContext = true);
 	void CleanupSyncLoadContext(FAsyncLoadingThreadState2& ThreadState);
@@ -5754,6 +5766,10 @@ FAsyncPackage2* FAsyncLoadingThread2::FindOrInsertPackage(FAsyncLoadingThreadSta
 				{
 					Package->AddRequestID(ThreadState, Desc.RequestID);
 				}
+				if (ImportedByPackage)
+				{
+					Package->ImportRequestIDs.Append(ImportedByPackage->RequestIDs);
+				}
 				if (Desc.Priority > Package->Desc.Priority)
 				{
 					TRACE_CPUPROFILER_EVENT_SCOPE(UpdatePackagePriority);
@@ -10494,22 +10510,11 @@ EAsyncPackageState::Type FAsyncLoadingThread2::ProcessLoadedPackagesFromGameThre
 		}
 		LocalCompletedAsyncPackages.Reset();
 
-		TArray<FCompletedPackageRequest> RequestsToProcess;
-
 		// Move CompletedPackageRequest out of the global collection to prevent it from changing from within the callbacks.
 		// If we're flushing a specific request only call callbacks for that request
-		for (int32 CompletedPackageRequestIndex = CompletedPackageRequests.Num() - 1; CompletedPackageRequestIndex >= 0; --CompletedPackageRequestIndex)
-		{
-			FCompletedPackageRequest& CompletedPackageRequest = CompletedPackageRequests[CompletedPackageRequestIndex];
-			if (FlushRequestIDs.Num() == 0
-				|| Algo::AnyOf(FlushRequestIDs, [&CompletedPackageRequest](int32 FlushRequestID) { return CompletedPackageRequest.RequestIDs.Contains(FlushRequestID); }))
-			{
-				RemovePendingRequests(ThreadState, CompletedPackageRequest.RequestIDs);
-				RequestsToProcess.Emplace(MoveTemp(CompletedPackageRequest));
-				CompletedPackageRequests.RemoveAt(CompletedPackageRequestIndex);
-				bLocalDidSomething = true;
-			}
-		}
+		TArray<FCompletedPackageRequest> RequestsToProcess;
+		GatherCompletedPackagesToFlush(ThreadState, FlushRequestIDs, RequestsToProcess);
+		bLocalDidSomething |= !!RequestsToProcess.Num();
 
 		// Call callbacks in a batch in a stack-local array after all other work has been done to handle
 		// callbacks that may call FlushAsyncLoading
@@ -10545,6 +10550,122 @@ EAsyncPackageState::Type FAsyncLoadingThread2::ProcessLoadedPackagesFromGameThre
 	return Result;
 }
 
+void FAsyncLoadingThread2::GatherCompletedPackagesToFlush(FAsyncLoadingThreadState2& ThreadState, TConstArrayView<int32> FlushRequestIds, TArray<FCompletedPackageRequest>& OutRequestsToProcess)
+{
+	if (CompletedPackageRequests.IsEmpty())
+	{
+		return;
+	}
+
+	if (FlushRequestIds.Num())
+	{
+		TRACE_CPUPROFILER_EVENT_SCOPE(GatherCompletedPackagesToFlush);
+
+		// Gather all completed packages affected by the flush request (i.e. the package being flushed
+		// and any completed imports loaded as a result). To do so, we build an adjacency list from the 
+		// CompletedRequests, and perform a depth-first sort for each package being flushed. We will append 
+		// the DFS results from each flush while maintaining a 'visited' set. This set will prevent loopiung on cycles
+		// but also allow us to prevent duplication when we first flush a package that is a sub-tree of some other package
+		// we are also going to flush later.
+		TMap<int32, int32, TInlineSetAllocator<64>>  RequestIdToCompletePackageIndex;
+		for (int32 CompletedPackageRequestIndex = 0; CompletedPackageRequestIndex < CompletedPackageRequests.Num(); ++CompletedPackageRequestIndex)
+		{
+			FCompletedPackageRequest& CompletedPackageRequest = CompletedPackageRequests[CompletedPackageRequestIndex];
+			for (int32 RequestId : CompletedPackageRequest.RequestIDs)
+			{
+				RequestIdToCompletePackageIndex.Add(RequestId, CompletedPackageRequestIndex);
+			}
+		}
+
+		// Adjacency list of completed packages
+		typedef TSet<int32, DefaultKeyFuncs<int32>, TInlineSetAllocator<64>> FInlineIndexSet;
+		TMap<int32, FInlineIndexSet, TInlineSetAllocator<64>> IndexToImportIndices;
+		for (int32 CompletedPackageRequestIndex = 0; CompletedPackageRequestIndex < CompletedPackageRequests.Num(); ++CompletedPackageRequestIndex)
+		{
+			FCompletedPackageRequest& CompletedPackageRequest = CompletedPackageRequests[CompletedPackageRequestIndex];
+			for (int32 ImportRequestId : CompletedPackageRequest.ImportRequestIDs)
+			{
+				int32* ImportCompletedPackageRequestIndex = RequestIdToCompletePackageIndex.Find(ImportRequestId);
+				if (ImportCompletedPackageRequestIndex && *ImportCompletedPackageRequestIndex != CompletedPackageRequestIndex)
+				{
+					FInlineIndexSet& Imports = IndexToImportIndices.FindOrAdd(*ImportCompletedPackageRequestIndex);
+					Imports.Add(CompletedPackageRequestIndex);
+				}
+			}
+		}
+
+		FInlineIndexSet Visited;
+		TArray<int32, TInlineAllocator<64>> Stack;
+		TArray<int32, TInlineAllocator<64>> DFS;
+		TArray<int32, TInlineAllocator<64>> CompletedPackageIndicesToProcess;
+		// For each root (flushed) package, depth-first sort
+		for (int32 FlushRequestId : FlushRequestIds)
+		{
+			int32* CompletedPackageIndexToFlush = RequestIdToCompletePackageIndex.Find(FlushRequestId);
+			if (!CompletedPackageIndexToFlush || Visited.Contains(*CompletedPackageIndexToFlush))
+			{
+				continue;
+			}
+			Stack.Push(*CompletedPackageIndexToFlush);
+			DFS.Reset();
+			do
+			{
+				int32 CompletedPackageIndex = Stack.Pop();
+				bool bAlreadyInserted = false;
+				Visited.Add(CompletedPackageIndex, &bAlreadyInserted);
+
+				// Cycles are common in loading graphs so we must defend against them
+				if (!bAlreadyInserted)
+				{
+					DFS.Add(CompletedPackageIndex);
+				}
+
+				FInlineIndexSet* ImportIndices = IndexToImportIndices.Find(CompletedPackageIndex);
+				if (ImportIndices)
+				{
+					for (int32 ImportIndex : *ImportIndices)
+					{
+						if (!Visited.Contains(ImportIndex))
+						{
+							Stack.Push(ImportIndex);
+						}
+					}
+				}
+			} while (!Stack.IsEmpty());
+
+			for (int32 Index : ReverseIterate(DFS))
+			{
+				CompletedPackageIndicesToProcess.Add(Index);
+			}
+		}
+
+		for (int32 ProcessIndex : CompletedPackageIndicesToProcess)
+		{
+			FCompletedPackageRequest& CompletedPackageRequest = CompletedPackageRequests[ProcessIndex];
+			RemovePendingRequests(ThreadState, CompletedPackageRequest.RequestIDs);
+			OutRequestsToProcess.Push(MoveTemp(CompletedPackageRequest));
+		}
+
+		// Sort the indices to remove elements from the end of the array first 
+		// so our indices are not invalidated as we perform removes and reduce the memmoving involved
+		CompletedPackageIndicesToProcess.Sort([](int32 A, int32 B) { return A > B; });
+		for (int32 ProcessIndex : CompletedPackageIndicesToProcess)
+		{
+			CompletedPackageRequests.RemoveAt(ProcessIndex);
+		}
+	}
+	else
+	{
+		OutRequestsToProcess.Reserve(CompletedPackageRequests.Num());
+		for (FCompletedPackageRequest& CompletedPackageRequest : CompletedPackageRequests)
+		{
+			RemovePendingRequests(ThreadState, CompletedPackageRequest.RequestIDs);
+			OutRequestsToProcess.Emplace(MoveTemp(CompletedPackageRequest));
+		}
+		CompletedPackageRequests.Empty();
+	}
+}
+
 EAsyncPackageState::Type FAsyncLoadingThread2::TickAsyncLoadingFromGameThread(FAsyncLoadingThreadState2& ThreadState, bool bUseTimeLimit, bool bUseFullTimeLimit, double TimeLimit, TConstArrayView<int32> FlushRequestIDs, bool& bDidSomething)
 {
 	SCOPE_CYCLE_COUNTER(STAT_FAsyncPackage_TickAsyncLoadingGameThread);
@@ -11639,7 +11760,8 @@ void FAsyncPackage2::AddRequestID(FAsyncLoadingThreadState2& ThreadState, int32
 			Desc.RequestID = Id;
 		}
 
-		RequestIDs.Add(Id);
+		// We expect the amount of aliasing to be very low so add unique to reduce memory
+		RequestIDs.AddUnique(Id);
 
 		// The Id is most likely already present because it's added as soon as the request is created.
 		AsyncLoadingThread.AddPendingRequest(Id);
@@ -12035,7 +12157,7 @@ void FAsyncLoadingThread2::QueueUnresolvedPackage(FAsyncLoadingThreadState2& Thr
 	}
 	else
 	{
-		RemovePendingRequests(ThreadState, TArrayView<int32>(&PackageDesc.RequestID, 1));
+		RemovePendingRequests(ThreadState, { PackageDesc.RequestID });
 		--PackagesWithRemainingWorkCounter;
 		TRACE_COUNTER_SET(AsyncLoadingPackagesWithRemainingWork, PackagesWithRemainingWorkCounter);
 	}
@@ -13182,4 +13304,4 @@ bool FPublicExportMapPreInsertPublicExportsTest::RunTest(const FString& Paramete
 	return true;
 }
 
-#endif // WITH_LOW_LEVEL_TESTS
\ No newline at end of file
+#endif // WITH_LOW_LEVEL_TESTS
diff --git a/Engine/Source/Runtime/Engine/Private/Tests/Loading/AsyncLoadingTests_Flushes.cpp b/Engine/Source/Runtime/Engine/Private/Tests/Loading/AsyncLoadingTests_Flushes.cpp
--- a/Engine/Source/Runtime/Engine/Private/Tests/Loading/AsyncLoadingTests_Flushes.cpp
+++ b/Engine/Source/Runtime/Engine/Private/Tests/Loading/AsyncLoadingTests_Flushes.cpp
@@ -1,7 +1,9 @@
 // Copyright Epic Games, Inc. All Rights Reserved.
 
 #include "AsyncLoadingTests_Shared.h"
+#include "Algo/AllOf.h"
 #include "Misc/AutomationTest.h"
+#include "Misc/EnumClassFlags.h"
 #include "Misc/PackageName.h"
 #include "Async/ManualResetEvent.h"
 #include "UObject/CoreRedirects.h"
@@ -64,4 +66,408 @@ bool FLoadingTests_Flush_ValidFromCallback::RunTest(const FString& Parameters)
 	return true;
 }
 
+IMPLEMENT_CUSTOM_SIMPLE_AUTOMATION_TEST(
+	FLoadingTests_Flush_LoadOrder_LoadFromPostLoad,
+	FLoadingTests_Flush_Base,
+	TEXT("System.Engine.Loading.Flush.LoadOrder_LoadFromPostLoad"),
+	EAutomationTestFlags::EditorContext | EAutomationTestFlags::EngineFilter
+)
+bool FLoadingTests_Flush_LoadOrder_LoadFromPostLoad::RunTest(const FString& Parameters)
+{
+	TArray<const TCHAR*> CompletionCallbackOrder;
+	FLoadPackageAsyncDelegate LoadCompletedFn; LoadCompletedFn.BindLambda([&CompletionCallbackOrder](const FName& PackageName, UPackage* LoadedPackage, EAsyncLoadingResult::Type Result)
+		{
+			if (PackageName == FLoadingTestsScope::PackagePath1)
+			{
+				CompletionCallbackOrder.AddUnique(FLoadingTestsScope::PackagePath1);
+			}
+			else if (PackageName == FLoadingTestsScope::PackagePath2)
+			{
+				CompletionCallbackOrder.AddUnique(FLoadingTestsScope::PackagePath2);
+			}
+		});
+	UAsyncLoadingTests_Shared::OnPostLoad.BindLambda([&CompletionCallbackOrder, &LoadCompletedFn](UAsyncLoadingTests_Shared* Object)
+		{
+			UPackage* Package = Object->GetPackage();
+			if (!Package)
+			{
+				return;
+			}
+			FString PackageName = Package->GetFName().ToString();
+			if (PackageName.Equals(FLoadingTestsScope::PackagePath1))
+			{
+				FlushAsyncLoading(LoadPackageAsync(FLoadingTestsScope::PackagePath2, LoadCompletedFn));
+			}
+		});
+
+	FLoadingTestsScope LoadingTestScope(this);
+	FlushAsyncLoading(LoadPackageAsync(FLoadingTestsScope::PackagePath1, LoadCompletedFn));
+
+	// If we sync load a package while loading an existing one, we expect the callbacks to be ordered as packages are completed
+	const TArray<const TCHAR*> ExpectedOrder = { FLoadingTestsScope::PackagePath2, FLoadingTestsScope::PackagePath1 };
+	TestEqual(TEXT("Expected LoadCompletion callback order to be identical"), CompletionCallbackOrder, ExpectedOrder);
+	return true;
+}
+
+enum class EVerifyOrder : uint8
+{
+	None			= 0,
+	Serialize		= 1 << 0,
+	PostLoad		= 1 << 1,
+	LoadCompleted	= 1 << 2,
+	All				= (uint8) ~0
+};
+ENUM_CLASS_FLAGS(EVerifyOrder);
+
+template <typename TMutateObjectsFn>
+void TestLoadPhases(FAutomationTestBase* AutomationTest, TMutateObjectsFn& MutateObjectsFn, const TArray<const TCHAR*>& PackagesToFlush, const TArray<const TCHAR*>& ExpectedOrder, EVerifyOrder VerifyOrder)
+{
+	TArray<const TCHAR*> SerializedPackages; SerializedPackages.Reserve(3);
+	TArray<const TCHAR*> PostLoadedPackages; PostLoadedPackages.Reserve(3);
+	TArray<const TCHAR*> LoadCompletedPackages; LoadCompletedPackages.Reserve(3);
+
+	auto CheckNameFn = [](UObject* Object, TArray<const TCHAR*>& OutPackagesFound)
+		{
+			UPackage* Package = Cast<UPackage>(Object);
+			if (!Package)
+			{
+				Package = Object->GetPackage();
+			}
+			FString PackageName = Package->GetFName().ToString();
+			
+			if (PackageName == FLoadingTestsScope::PackagePath1)
+			{
+				OutPackagesFound.AddUnique(FLoadingTestsScope::PackagePath1);
+			}
+			else if (PackageName == FLoadingTestsScope::PackagePath2)
+			{
+				OutPackagesFound.AddUnique(FLoadingTestsScope::PackagePath2);
+			}
+			else if (PackageName == FLoadingTestsScope::PackagePath3)
+			{
+				OutPackagesFound.AddUnique(FLoadingTestsScope::PackagePath3);
+			}
+			else if (PackageName == FLoadingTestsScope::PackagePath4)
+			{
+				OutPackagesFound.AddUnique(FLoadingTestsScope::PackagePath4);
+			}
+		};
+
+	UAsyncLoadingTests_Shared::OnSerialize.BindLambda([&SerializedPackages, &CheckNameFn](FArchive& Ar, UAsyncLoadingTests_Shared* Object)
+		{
+			CheckNameFn(Object, SerializedPackages);
+		});
+	UAsyncLoadingTests_Shared::OnPostLoad.BindLambda([&PostLoadedPackages, &CheckNameFn](UAsyncLoadingTests_Shared* Object)
+		{
+			CheckNameFn(Object, PostLoadedPackages);
+		});
+
+	FLoadPackageAsyncDelegate LoadCompletedFn; LoadCompletedFn.BindLambda([&LoadCompletedPackages, &CheckNameFn](const FName& PackageName, UPackage* LoadedPackage, EAsyncLoadingResult::Type Result)
+		{
+			if (Result != EAsyncLoadingResult::Type::Succeeded)
+			{
+				return;
+			}
+			CheckNameFn(LoadedPackage, LoadCompletedPackages);
+		});
+
+	FLoadingTestsScope LoadingTestScope(AutomationTest, MutateObjectsFn);
+	LoadPackageAsync(FLoadingTestsScope::PackagePath1, LoadCompletedFn);
+	LoadPackageAsync(FLoadingTestsScope::PackagePath2, LoadCompletedFn);
+	LoadPackageAsync(FLoadingTestsScope::PackagePath3, LoadCompletedFn);
+	LoadPackageAsync(FLoadingTestsScope::PackagePath4, LoadCompletedFn);
+
+	// Flush the packages requested, but use a new load request so that we can also verify 
+	// we handle aliasing request ids for packages in our ordering validation
+	TArray<int32, TInlineAllocator<4>> FlushRequests;
+	for (const TCHAR* PackagePath : PackagesToFlush)
+	{
+		FlushRequests.Add(LoadPackageAsync(PackagePath, LoadCompletedFn));
+	}
+	FlushAsyncLoading(FlushRequests);
+
+	auto VerifyLoads = [AutomationTest](const TCHAR* OrderType, const TArray<const TCHAR*>& Actual, const TArray<const TCHAR*>& Expected, bool bVerifyExactOrder)
+		{
+			bool bMatch = false;
+			if (bVerifyExactOrder)
+			{
+				bMatch = Actual == Expected;
+			}
+			else
+			{
+				bMatch = Actual.Num() == Expected.Num() && Algo::AllOf(Actual, [&Expected](const TCHAR* Element)
+				{
+					return Expected.Contains(Element);
+				});
+			}
+
+			if (!bMatch)
+			{
+				TStringBuilder<256> Builder;
+				Builder.Appendf(TEXT("Expected to %s packages with %s order:\n\t"), OrderType,
+					(bVerifyExactOrder ? TEXT("exact") : TEXT("any")));
+				if (!Expected.IsEmpty())
+				{
+					for (const TCHAR* Path : Expected)
+					{
+						Builder.Append(Path);
+						Builder.Append(TEXT(", "));
+					}
+					Builder.RemoveSuffix(2);
+				}
+				Builder.Append(TEXT("\nActual order:\n\t"));
+				if (!Actual.IsEmpty())
+				{
+					for (const TCHAR* Path : Actual)
+					{
+						Builder.Append(Path);
+						Builder.Append(TEXT(", "));
+					}
+					Builder.RemoveSuffix(2);
+				}
+				AutomationTest->TestTrue(Builder.ToString(), false);
+			}
+		};
+
+	// Currently we do not provide any Serialize order guarantee with respect to dependencies.
+	// We likely never want this so we can enable parallel preloading opportunities. However
+	// if we should want preload ordering, re-enable to ensure this is done.
+	VerifyLoads(TEXT("Serialize"), SerializedPackages, ExpectedOrder, false /*!!(VerifyOrder& EVerifyOrder::Serialize)*/);
+	VerifyLoads(TEXT("PostLoad"), PostLoadedPackages, ExpectedOrder, !!(VerifyOrder & EVerifyOrder::PostLoad));
+	VerifyLoads(TEXT("LoadCompletedCallbacks"), LoadCompletedPackages, ExpectedOrder, !!(VerifyOrder & EVerifyOrder::LoadCompleted));
+}
+
+IMPLEMENT_CUSTOM_SIMPLE_AUTOMATION_TEST(
+	FLoadingTests_Flush_LoadOrder_A,
+	FLoadingTests_Flush_Base,
+	TEXT("System.Engine.Loading.Flush.LoadOrder_A"),
+	EAutomationTestFlags::EditorContext | EAutomationTestFlags::EngineFilter
+)
+bool FLoadingTests_Flush_LoadOrder_A::RunTest(const FString& Parameters)
+{
+	// Package 1 has a hard ref to package 2 which has a hard ref to package 3.
+	// 1 -> 2 -> 3 -> 4
+	auto MutateObjects =
+		[](FLoadingTestsScope& Scope)
+		{
+			Scope.Object1->HardReference = Scope.Object2;
+			Scope.Object2->HardReference = Scope.Object3;
+			Scope.Object3->HardReference = Scope.Object4;
+		};
+
+	bool bVerifyExactOrder = true;
+	const TArray<const TCHAR*> PackagesToFlush = { FLoadingTestsScope::PackagePath1 };
+	const TArray<const TCHAR*> ExpectedOrder = { FLoadingTestsScope::PackagePath4, FLoadingTestsScope::PackagePath3, FLoadingTestsScope::PackagePath2, FLoadingTestsScope::PackagePath1 };
+	TestLoadPhases(this, MutateObjects, PackagesToFlush, ExpectedOrder, EVerifyOrder::All);
+	return true;
+}
+
+IMPLEMENT_CUSTOM_SIMPLE_AUTOMATION_TEST(
+	FLoadingTests_Flush_LoadOrder_B,
+	FLoadingTests_Flush_Base,
+	TEXT("System.Engine.Loading.Flush.LoadOrder_B"),
+	EAutomationTestFlags::EditorContext | EAutomationTestFlags::EngineFilter
+)
+bool FLoadingTests_Flush_LoadOrder_B::RunTest(const FString& Parameters)
+{
+	// Flush a sub-tree then flush the rest of the tree
+	auto MutateObjects =
+		[](FLoadingTestsScope& Scope)
+		{
+			Scope.Object1->HardReference = Scope.Object2;
+			Scope.Object2->HardReference = Scope.Object3;
+			Scope.Object3->HardReference = Scope.Object4;
+		};
+
+	bool bVerifyExactOrder = true;
+	const TArray<const TCHAR*> PackagesToFlush = { FLoadingTestsScope::PackagePath3, FLoadingTestsScope::PackagePath1 };
+	const TArray<const TCHAR*> ExpectedOrder = { FLoadingTestsScope::PackagePath4, FLoadingTestsScope::PackagePath3, FLoadingTestsScope::PackagePath2, FLoadingTestsScope::PackagePath1 };
+	TestLoadPhases(this, MutateObjects, PackagesToFlush, ExpectedOrder, EVerifyOrder::All);
+	return true;
+}
+
+IMPLEMENT_CUSTOM_SIMPLE_AUTOMATION_TEST(
+	FLoadingTests_Flush_LoadOrder_C,
+	FLoadingTests_Flush_Base,
+	TEXT("System.Engine.Loading.Flush.LoadOrder_C"),
+	EAutomationTestFlags::EditorContext | EAutomationTestFlags::EngineFilter
+)
+bool FLoadingTests_Flush_LoadOrder_C::RunTest(const FString& Parameters)
+{
+	// Flush a sub-tree via distinct root then flush the rest of the tree
+	//      3 -> 4
+	// 1 -> 2 -> 4
+	auto MutateObjects =
+		[](FLoadingTestsScope& Scope)
+		{
+			Scope.Object1->HardReference = Scope.Object2;
+			Scope.Object2->HardReference = Scope.Object4;
+
+			Scope.Object3->HardReference = Scope.Object4;
+		};
+
+	{
+		bool bVerifyExactOrder = true;
+		const TArray<const TCHAR*> PackagesToFlush = { FLoadingTestsScope::PackagePath3, FLoadingTestsScope::PackagePath1 };
+		const TArray<const TCHAR*> ExpectedOrder = { FLoadingTestsScope::PackagePath4, FLoadingTestsScope::PackagePath3, FLoadingTestsScope::PackagePath2, FLoadingTestsScope::PackagePath1 };
+		TestLoadPhases(this, MutateObjects, PackagesToFlush, ExpectedOrder, EVerifyOrder::LoadCompleted);
+	}
+
+	{
+		// Swap flush order. We should still see imports before their referencers when loading
+		const TArray<const TCHAR*> PackagesToFlush = { FLoadingTestsScope::PackagePath1, FLoadingTestsScope::PackagePath3 };
+		const TArray<const TCHAR*> ExpectedOrder = { FLoadingTestsScope::PackagePath4, FLoadingTestsScope::PackagePath2, FLoadingTestsScope::PackagePath1, FLoadingTestsScope::PackagePath3 };
+		TestLoadPhases(this, MutateObjects, PackagesToFlush, ExpectedOrder, EVerifyOrder::LoadCompleted);
+	}
+	return true;
+}
+
+// Circular dependencies are currently chaotic in the loader and the order for processing is a mixture of request
+// order, import load state, and import table ordering. The tests below confirm that load completion callbacks
+// are exactly ordered, and all other load orderings are only verified to ensure that we see load events (Serialize/Postload)
+// from all packages but we do not expect an exact ordering.
+IMPLEMENT_CUSTOM_SIMPLE_AUTOMATION_TEST(
+	FLoadingTests_Flush_LoadOrder_CircularDependencies_A,
+	FLoadingTests_Flush_Base,
+	TEXT("System.Engine.Loading.Flush.LoadOrder_CircularDependencies_A"),
+	EAutomationTestFlags::EditorContext | EAutomationTestFlags::EngineFilter
+)
+bool FLoadingTests_Flush_LoadOrder_CircularDependencies_A::RunTest(const FString& Parameters)
+{
+	// Circular reference
+	// 1 -> 2 -> 3 -> 4 -> 1
+	auto MutateObjects =
+		[](FLoadingTestsScope& Scope)
+		{
+			Scope.Object1->HardReference = Scope.Object2;
+			Scope.Object2->HardReference = Scope.Object3;
+			Scope.Object3->HardReference = Scope.Object4;
+			// Circular ref back to Package1
+			Scope.Object4->HardReference = Scope.Object1;
+		};
+
+	const TArray<const TCHAR*> PackagesToFlush = { FLoadingTestsScope::PackagePath1 };
+	const TArray<const TCHAR*> ExpectedOrder = { FLoadingTestsScope::PackagePath4, FLoadingTestsScope::PackagePath3, FLoadingTestsScope::PackagePath2, FLoadingTestsScope::PackagePath1 };
+	TestLoadPhases(this, MutateObjects, PackagesToFlush, ExpectedOrder, EVerifyOrder::LoadCompleted);
+	return true;
+}
+
+IMPLEMENT_CUSTOM_SIMPLE_AUTOMATION_TEST(
+	FLoadingTests_Flush_LoadOrder_CircularDependencies_B,
+	FLoadingTests_Flush_Base,
+	TEXT("System.Engine.Loading.Flush.LoadOrder_CircularDependencies_B"),
+	EAutomationTestFlags::EditorContext | EAutomationTestFlags::EngineFilter
+)
+bool FLoadingTests_Flush_LoadOrder_CircularDependencies_B::RunTest(const FString& Parameters)
+{
+	// Circular reference
+	// 1 -> 2 -> 3 <-> 4
+	auto MutateObjects =
+		[](FLoadingTestsScope& Scope)
+		{
+			Scope.Object1->HardReference = Scope.Object2;
+			Scope.Object2->HardReference = Scope.Object3;
+			Scope.Object3->HardReference = Scope.Object4;
+			// Circular ref back to Package3
+			Scope.Object4->HardReference = Scope.Object3;
+		};
+
+	const TArray<const TCHAR*> PackagesToFlush = { FLoadingTestsScope::PackagePath1 };
+	const TArray<const TCHAR*> ExpectedOrder = { FLoadingTestsScope::PackagePath4, FLoadingTestsScope::PackagePath3, FLoadingTestsScope::PackagePath2, FLoadingTestsScope::PackagePath1 };
+	TestLoadPhases(this, MutateObjects, PackagesToFlush, ExpectedOrder, EVerifyOrder::LoadCompleted);
+	return true;
+}
+
+IMPLEMENT_CUSTOM_SIMPLE_AUTOMATION_TEST(
+	FLoadingTests_Flush_LoadOrder_CircularDependencies_C,
+	FLoadingTests_Flush_Base,
+	TEXT("System.Engine.Loading.Flush.LoadOrder_CircularDependencies_C"),
+	EAutomationTestFlags::EditorContext | EAutomationTestFlags::EngineFilter
+)
+bool FLoadingTests_Flush_LoadOrder_CircularDependencies_C::RunTest(const FString& Parameters)
+{
+	// Two circular references
+	// 1 <-> 2 <-> 3 -> 4
+	auto MutateObjects =
+		[](FLoadingTestsScope& Scope)
+		{
+			Scope.Object1->HardReference = Scope.Object2;
+			// Circular ref back to Package1
+			Scope.Object2->HardReference = Scope.Object1;
+			Scope.Object2->HardReference2 = Scope.Object3;
+			// Circular ref back to Package2
+			Scope.Object3->HardReference = Scope.Object2;
+			Scope.Object3->HardReference2 = Scope.Object4;
+		};
+
+	const TArray<const TCHAR*> PackagesToFlush = { FLoadingTestsScope::PackagePath1 };
+	const TArray<const TCHAR*> ExpectedOrder = { FLoadingTestsScope::PackagePath4, FLoadingTestsScope::PackagePath3, FLoadingTestsScope::PackagePath2, FLoadingTestsScope::PackagePath1 };
+	TestLoadPhases(this, MutateObjects, PackagesToFlush, ExpectedOrder, EVerifyOrder::LoadCompleted);
+	return true;
+}
+
+IMPLEMENT_CUSTOM_SIMPLE_AUTOMATION_TEST(
+	FLoadingTests_Flush_LoadOrder_CircularDependencies_D,
+	FLoadingTests_Flush_Base,
+	TEXT("System.Engine.Loading.Flush.LoadOrder_CircularDependencies_D"),
+	EAutomationTestFlags::EditorContext | EAutomationTestFlags::EngineFilter
+)
+bool FLoadingTests_Flush_LoadOrder_CircularDependencies_D::RunTest(const FString& Parameters)
+{
+	// One circular reference. 3<->4, but 1 is also able to reach the cycle
+	// 1 -> 2 -> 3 
+	// |         ^
+	// |         |
+	// |         v
+	// +-------> 4
+	auto MutateObjects =
+		[](FLoadingTestsScope& Scope)
+		{
+			Scope.Object1->HardReference = Scope.Object2;
+			Scope.Object1->HardReference2 = Scope.Object4;
+			Scope.Object2->HardReference = Scope.Object3;
+			Scope.Object3->HardReference = Scope.Object4;
+			// Circular ref back to Package3
+			Scope.Object4->HardReference = Scope.Object3;
+		};
+
+	const TArray<const TCHAR*> PackagesToFlush = { FLoadingTestsScope::PackagePath1 };
+	const TArray<const TCHAR*> ExpectedOrder = { FLoadingTestsScope::PackagePath4, FLoadingTestsScope::PackagePath3, FLoadingTestsScope::PackagePath2, FLoadingTestsScope::PackagePath1 };
+	TestLoadPhases(this, MutateObjects, PackagesToFlush, ExpectedOrder, EVerifyOrder::LoadCompleted);
+	return true;
+}
+
+IMPLEMENT_CUSTOM_SIMPLE_AUTOMATION_TEST(
+	FLoadingTests_Flush_LoadOrder_CircularDependencies_E,
+	FLoadingTests_Flush_Base,
+	TEXT("System.Engine.Loading.Flush.LoadOrder_CircularDependencies_E"),
+	EAutomationTestFlags::EditorContext | EAutomationTestFlags::EngineFilter
+)
+bool FLoadingTests_Flush_LoadOrder_CircularDependencies_E::RunTest(const FString& Parameters)
+{
+	// One circular reference. 3<->2, but 1 and 4 can both reach the cycle
+	// 1 -> 3 -> 4
+	// |    ^    |
+	// v    |    |
+	// 2 <--+<---+
+	auto MutateObjects =
+		[](FLoadingTestsScope& Scope)
+		{
+			Scope.Object1->HardReference = Scope.Object2;
+			Scope.Object1->HardReference2 = Scope.Object3;
+			Scope.Object2->HardReference = Scope.Object3;
+			// Circular ref back to Package2
+			Scope.Object3->HardReference = Scope.Object2;
+			Scope.Object3->HardReference2 = Scope.Object4;
+			Scope.Object4->HardReference = Scope.Object2;
+		};
+
+	// Note the ordering for nodes at the same depth is coincidentally numerically ordered, we only care that depth is sorted 
+	// but for simplicity omit that as the sorting used in the implementation is stable and it's unlikely the save order will 
+	// change affecting this test. If it does, update the test to adjust to the new order or relax the ordering constraint within each depth.
+	const TArray<const TCHAR*> PackagesToFlush = { FLoadingTestsScope::PackagePath1 };
+	const TArray<const TCHAR*> ExpectedOrder = { FLoadingTestsScope::PackagePath4, FLoadingTestsScope::PackagePath3, FLoadingTestsScope::PackagePath2, FLoadingTestsScope::PackagePath1 };
+	TestLoadPhases(this, MutateObjects, PackagesToFlush, ExpectedOrder, EVerifyOrder::LoadCompleted);
+	return true;
+}
 #endif // WITH_DEV_AUTOMATION_TESTS
diff --git a/Engine/Source/Runtime/Engine/Private/Tests/Loading/AsyncLoadingTests_RecursiveLoads.cpp b/Engine/Source/Runtime/Engine/Private/Tests/Loading/AsyncLoadingTests_RecursiveLoads.cpp
--- a/Engine/Source/Runtime/Engine/Private/Tests/Loading/AsyncLoadingTests_RecursiveLoads.cpp
+++ b/Engine/Source/Runtime/Engine/Private/Tests/Loading/AsyncLoadingTests_RecursiveLoads.cpp
@@ -457,4 +457,90 @@ bool FLoadingTests_RecursiveLoads_FromBothSerializeAndPostLoad::RunTest(const FS
 	return true;
 }
 
+/**
+ * These tests validate loading an object that loads itself synchronously during postload.
+ */
+IMPLEMENT_CUSTOM_SIMPLE_AUTOMATION_TEST(
+	FLoadingTests_RecursiveLoads_LoadSelfFromPostLoad_Direct,
+	FLoadingTests_RecursiveLoads_Base,
+	TEXT("System.Engine.Loading.RecursiveLoads.LoadSelfFromPostLoad_Direct"),
+	EAutomationTestFlags::EditorContext | EAutomationTestFlags::EngineFilter
+)
+bool FLoadingTests_RecursiveLoads_LoadSelfFromPostLoad_Direct::RunTest(const FString& Parameters)
+{
+	FLoadingTestsScope LoadingTestScope(this,
+		[](FLoadingTestsScope& TestScope)
+		{
+			TestScope.Object1->SoftReference = TestScope.Object1;
+		}
+	);
+
+	UAsyncLoadingTests_Shared::OnPostLoad.BindLambda(
+		[this](UAsyncLoadingTests_Shared* Object)
+		{
+			if (UObject* Obj = Object->SoftReference.LoadSynchronous())
+			{
+				TestFalse(TEXT("Recursive loads in postload should be fully loaded"), Obj->HasAnyFlags(RF_NeedLoad | RF_NeedPostLoad));
+			}
+		}
+	);
+
+	LoadPackage(nullptr, LoadingTestScope.PackagePath1, LOAD_None);
+
+	UAsyncLoadingTests_Shared* LocalObject1 = FindObject<UAsyncLoadingTests_Shared>(nullptr, LoadingTestScope.ObjectPath1);
+	TestTrue(TEXT("The object should have been properly loaded recursively"), LocalObject1 != nullptr);
+	if (LocalObject1)
+	{
+		TestFalse(TEXT("The object should have been properly loaded recursively"), LocalObject1->HasAnyFlags(RF_NeedLoad | RF_NeedPostLoad));
+	}
+
+	return true;
+}
+
+IMPLEMENT_CUSTOM_SIMPLE_AUTOMATION_TEST(
+	FLoadingTests_RecursiveLoads_LoadSelfFromPostLoad_Indirect,
+	FLoadingTests_RecursiveLoads_Base,
+	TEXT("System.Engine.Loading.RecursiveLoads.LoadSelfFromPostLoad_Indirect"),
+	EAutomationTestFlags::EditorContext | EAutomationTestFlags::EngineFilter
+)
+bool FLoadingTests_RecursiveLoads_LoadSelfFromPostLoad_Indirect::RunTest(const FString& Parameters)
+{
+	FLoadingTestsScope LoadingTestScope(this,
+		[](FLoadingTestsScope& TestScope)
+		{
+			// Soft-circular reference back to Object1
+			TestScope.Object1->SoftReference = TestScope.Object2;
+			TestScope.Object2->HardReference = TestScope.Object1;
+		}
+	);
+
+	UAsyncLoadingTests_Shared::OnPostLoad.BindLambda(
+		[this](UAsyncLoadingTests_Shared* Object)
+		{
+			if (UObject* Obj = Object->SoftReference.LoadSynchronous())
+			{
+				TestFalse(TEXT("Recursive loads in postload should be fully loaded"), Obj->HasAnyFlags(RF_NeedLoad | RF_NeedPostLoad));
+			}
+		}
+	);
+
+	LoadPackage(nullptr, LoadingTestScope.PackagePath1, LOAD_None);
+
+	UAsyncLoadingTests_Shared* LocalObject1 = FindObject<UAsyncLoadingTests_Shared>(nullptr, LoadingTestScope.ObjectPath1);
+	TestTrue(TEXT("The object should have been properly loaded recursively"), LocalObject1 != nullptr);
+	if (LocalObject1)
+	{
+		TestFalse(TEXT("The object should have been properly loaded recursively"), LocalObject1->HasAnyFlags(RF_NeedLoad | RF_NeedPostLoad));
+	}
+
+	UAsyncLoadingTests_Shared* LocalObject2 = FindObject<UAsyncLoadingTests_Shared>(nullptr, LoadingTestScope.ObjectPath2);
+	TestTrue(TEXT("The object should have been properly loaded recursively"), LocalObject2 != nullptr);
+	if (LocalObject2)
+	{
+		TestFalse(TEXT("The object should have been properly loaded recursively"), LocalObject2->HasAnyFlags(RF_NeedLoad | RF_NeedPostLoad));
+	}
+
+	return true;
+}
+
 #endif // WITH_DEV_AUTOMATION_TESTS
diff --git a/Engine/Source/Runtime/Engine/Private/Tests/Loading/AsyncLoadingTests_Shared.cpp b/Engine/Source/Runtime/Engine/Private/Tests/Loading/AsyncLoadingTests_Shared.cpp
--- a/Engine/Source/Runtime/Engine/Private/Tests/Loading/AsyncLoadingTests_Shared.cpp
+++ b/Engine/Source/Runtime/Engine/Private/Tests/Loading/AsyncLoadingTests_Shared.cpp
@@ -26,6 +26,9 @@ void FLoadingTestsScope::CreateObjects()
 	
 	Package3 = CreatePackage();
 	Object3 = NewObject<UAsyncLoadingTests_Shared>(Package3, ObjectName, RF_Public | RF_Standalone);
+
+	Package4 = CreatePackage();
+	Object4 = NewObject<UAsyncLoadingTests_Shared>(Package4, ObjectName, RF_Public | RF_Standalone);
 }
 
 void FLoadingTestsScope::DefaultMutateObjects()
diff --git a/Engine/Source/Runtime/Engine/Private/Tests/Loading/AsyncLoadingTests_Shared.h b/Engine/Source/Runtime/Engine/Private/Tests/Loading/AsyncLoadingTests_Shared.h
--- a/Engine/Source/Runtime/Engine/Private/Tests/Loading/AsyncLoadingTests_Shared.h
+++ b/Engine/Source/Runtime/Engine/Private/Tests/Loading/AsyncLoadingTests_Shared.h
@@ -21,6 +21,9 @@ public:
 	UPROPERTY()
 	TObjectPtr<UObject> HardReference;
 
+	UPROPERTY()
+	TObjectPtr<UObject> HardReference2;
+
 	// Those delegates allow to easily change behavior between tests
 	// Avoid using thread-safe delegates since calling them currently is an exclusive lock and can cause deadlocks in the test when serialize is multithreaded.
 	// Registration to these delegates happens safely when the test begins before threads are involved.
@@ -88,14 +91,18 @@ public:
 	static constexpr const TCHAR* ObjectPath2 = TEXT("/Engine/LoadingTestsScope_Package2.TestObject");
 	static constexpr const TCHAR* PackagePath3 = TEXT("/Engine/LoadingTestsScope_Package3");
 	static constexpr const TCHAR* ObjectPath3 = TEXT("/Engine/LoadingTestsScope_Package3.TestObject");
+	static constexpr const TCHAR* PackagePath4 = TEXT("/Engine/LoadingTestsScope_Package4");
+	static constexpr const TCHAR* ObjectPath4 = TEXT("/Engine/LoadingTestsScope_Package4.TestObject");
 
 	UPackage* Package1 = nullptr;
 	UPackage* Package2 = nullptr;
 	UPackage* Package3 = nullptr;
+	UPackage* Package4 = nullptr;
 
 	UAsyncLoadingTests_Shared* Object1 = nullptr;
 	UAsyncLoadingTests_Shared* Object2 = nullptr;
 	UAsyncLoadingTests_Shared* Object3 = nullptr;
+	UAsyncLoadingTests_Shared* Object4 = nullptr;
 
 	UPackage* CreatePackage()
 	{
