SplineComponent: Change default value of SplineComponent.AddSplinePointImplementation to 1

It turns out AddSplinePoint's usage of FSpline::AddPoint was unrelaible. InsertPoint is perfectly acceptable and this change addresses runtime issues on certain consoles. We will address the issue in FSpline::AddPoint separately

diff --git a/Engine/Source/Runtime/Engine/Private/Components/SplineComponent.cpp b/Engine/Source/Runtime/Engine/Private/Components/SplineComponent.cpp
--- a/Engine/Source/Runtime/Engine/Private/Components/SplineComponent.cpp
+++ b/Engine/Source/Runtime/Engine/Private/Components/SplineComponent.cpp
@@ -87,7 +87,7 @@ FAutoConsoleVariableRef CVarFailedValidationEnsures(
 	TEXT("When true, all validation checks will ensure if not passed.")
 );
 
-int32 GAddSplinePointImplementation = 0;
+int32 GAddSplinePointImplementation = 1;
 FAutoConsoleVariableRef CVarAddSplinePointImplementation(
 	TEXT("SplineComponent.AddSplinePointImplementation"),
 	GAddSplinePointImplementation,
