Enable Local Vertex Factory for materials tagged with skeletal mesh usage when using the skin cache.
#rnx

diff --git a/Engine/Source/Runtime/Engine/Private/LocalVertexFactory.cpp b/Engine/Source/Runtime/Engine/Private/LocalVertexFactory.cpp
--- a/Engine/Source/Runtime/Engine/Private/LocalVertexFactory.cpp
+++ b/Engine/Source/Runtime/Engine/Private/LocalVertexFactory.cpp
@@ -281,6 +281,12 @@ bool FLocalVertexFactory::ShouldCompilePermutation(const FVertexFactoryShaderPer
 		return true;
 	}
 
+	// We can use the local vertex factory in pass through mode when rendering SkeletalMesh through the SkinCache
+	if (Parameters.MaterialParameters.bIsUsedWithSkeletalMesh && IsGPUSkinPassThroughSupported(Parameters.Platform))
+	{
+		return true;
+	}
+
 	// Otherwise compile when the Material has selected that it can be used with StaticMesh
 	return Parameters.MaterialParameters.bIsUsedWithStaticMesh == 1;
 }
