Cleaning up comments.

diff --git a/Engine/Source/Runtime/Engine/Private/Streaming/Texture2DMipAllocator_AsyncReallocate.h b/Engine/Source/Runtime/Engine/Private/Streaming/Texture2DMipAllocator_AsyncReallocate.h
--- a/Engine/Source/Runtime/Engine/Private/Streaming/Texture2DMipAllocator_AsyncReallocate.h
+++ b/Engine/Source/Runtime/Engine/Private/Streaming/Texture2DMipAllocator_AsyncReallocate.h
@@ -1,9 +1,5 @@
 // Copyright Epic Games, Inc. All Rights Reserved.
 
-/*=============================================================================
-Texture2DMipAllocator_AsyncReallocate.h: Implementation of FTextureMipAllocator using RHIAsyncReallocateTexture2D
-=============================================================================*/
-
 #pragma once
 
 #include "CoreMinimal.h"
@@ -15,7 +11,7 @@ class UTexture2D;
 
 /**
 * FTexture2DMipAllocator_AsyncReallocate is the default implementation of  FTextureMipAllocator.
-* It uses RHIAsyncReallocateTexture2D which creates a new texture on the renderthread and copies the mips if required.
+* It uses AsyncReallocateTexture2D which creates a new texture on the render thread and copies the mips if required.
 */
 class FTexture2DMipAllocator_AsyncReallocate : public FTextureMipAllocator
 {
diff --git a/Engine/Source/Runtime/Engine/Private/Streaming/TextureMipAllocator.h b/Engine/Source/Runtime/Engine/Private/Streaming/TextureMipAllocator.h
--- a/Engine/Source/Runtime/Engine/Private/Streaming/TextureMipAllocator.h
+++ b/Engine/Source/Runtime/Engine/Private/Streaming/TextureMipAllocator.h
@@ -13,7 +13,7 @@ TextureMipAllocator.h: Base class for implementing a mip allocation strategy use
 /**
 * FTextureMipAllocator defines the update steps and interface to implement the mip allocation strategy
 * used in FTextureStreamIn. It allows to decouples the texture creation steps from the texture update.
-* Typical implementations are using RHIAsyncCreateTexture2D, RHIAsyncReallocateTexture2D or virtual textures.
+* Typical implementations are using RHIAsyncCreateTexture2D, AsyncReallocateTexture2D or virtual textures.
 */
 class FTextureMipAllocator
 {
