Added asset debugging code for streamable manager

diff --git a/Engine/Source/Runtime/Engine/Private/StreamableManager.cpp b/Engine/Source/Runtime/Engine/Private/StreamableManager.cpp
--- a/Engine/Source/Runtime/Engine/Private/StreamableManager.cpp
+++ b/Engine/Source/Runtime/Engine/Private/StreamableManager.cpp
@@ -83,6 +83,14 @@ static FAutoConsoleVariableRef CVarStreamableErrorOnRequest(
 		TEXT("If this is set and the errorOnRequest string partially matches the current debug name of the asset being loaded, then an error will be thrown instead of completion."),
 		ECVF_Default
 	);
+
+static FString GStreamableDebugOnRequest = "";
+static FAutoConsoleVariableRef CVarStreamableDebugOnRequest(
+		TEXT("s.debugOnRequest"),
+		GStreamableDebugOnRequest,
+		TEXT("If this is set then a log will fire when any requested asset string matches in part this variable, set breakpoint to debug the caller."),
+		ECVF_Default
+	);
 #endif
 FStreamableDownloadCachePin::FStreamableDownloadCachePin() = default;
 FStreamableDownloadCachePin::~FStreamableDownloadCachePin() = default;
@@ -1875,6 +1883,20 @@ TSharedPtr<FStreamableHandle> FStreamableManager::RequestAsyncLoad(FStreamableAs
 	TSet<FSoftObjectPath> TargetSet;
 	TargetSet.Reserve(NumValidRequests);
 
+#if !UE_BUILD_SHIPPING
+	if (!GStreamableDebugOnRequest.IsEmpty())
+	{
+
+		for (const FSoftObjectPath& RequestedAsset : NewRequest->RequestedAssets)
+		{
+			if (RequestedAsset.ToString().Contains(GStreamableDebugOnRequest))
+			{
+				UE_LOG(LogStreamableManager, Warning, TEXT("Debug on request hit on asset:%s, debug request:%s, place breakpoint here to debug"), *RequestedAsset.ToString(), *GStreamableDebugOnRequest);
+			}
+		}
+	}
+#endif
+
 	for (const FSoftObjectPath& TargetName : NewRequest->RequestedAssets)
 	{
 		if (TargetName.IsNull())
