Fix potential thread-unsafe access to skeleton remapping data.  In that case of FReferenceSkeleton, limit when the remapping data is refreshed to only when it's a skeleton and not a mesh that is having it's ref skeleton rebuilt.

diff --git a/Engine/Source/Runtime/Engine/Private/Animation/ReferenceSkeleton.cpp b/Engine/Source/Runtime/Engine/Private/Animation/ReferenceSkeleton.cpp
--- a/Engine/Source/Runtime/Engine/Private/Animation/ReferenceSkeleton.cpp
+++ b/Engine/Source/Runtime/Engine/Private/Animation/ReferenceSkeleton.cpp
@@ -532,16 +532,16 @@ void FReferenceSkeleton::RebuildRefSkeleton(const USkeleton* Skeleton, bool bReb
 		}
 	}
 
-	// Full rebuild of all compatible with this and with ones we are compatible with.
-	UE::Anim::FSkeletonRemappingRegistry::Get().RefreshMappings(Skeleton);
+	if (Skeleton && (this == &Skeleton->GetReferenceSkeleton()))
+	{
+		// Full rebuild of all compatible with this and with ones we are compatible with.
+		UE::Anim::FSkeletonRemappingRegistry::Get().RefreshMappings(Skeleton);
 
 #if WITH_EDITOR
-	if (Skeleton)
-	{
 		// Notify the skeleton that we changed
 		Skeleton->HandleReferenceSkeletonChange();
-	}
 #endif
+	}
 }
 
 void FReferenceSkeleton::RemoveDuplicateBones(const UObject* Requester, TArray<FBoneIndexType> & DuplicateBones)
diff --git a/Engine/Source/Runtime/Engine/Private/Animation/Skeleton.cpp b/Engine/Source/Runtime/Engine/Private/Animation/Skeleton.cpp
--- a/Engine/Source/Runtime/Engine/Private/Animation/Skeleton.cpp
+++ b/Engine/Source/Runtime/Engine/Private/Animation/Skeleton.cpp
@@ -2015,8 +2015,6 @@ void USkeleton::HandleVirtualBoneChanges()
 	constexpr bool bRebuildNameMap = false;
 	ReferenceSkeleton.RebuildRefSkeleton(this, bRebuildNameMap);
 
-	UE::Anim::FSkeletonRemappingRegistry::Get().RefreshMappings(this);
-
 	// store skeletal meshes that are also transacting to avoid re-registering the component here
 	// as it will be done later in USkeletalMesh::PostEditUndo()
 	TArray<USkeletalMesh*> SkeletalMeshTransacting;
diff --git a/Engine/Source/Runtime/Engine/Private/Animation/SkeletonRemappingRegistry.cpp b/Engine/Source/Runtime/Engine/Private/Animation/SkeletonRemappingRegistry.cpp
--- a/Engine/Source/Runtime/Engine/Private/Animation/SkeletonRemappingRegistry.cpp
+++ b/Engine/Source/Runtime/Engine/Private/Animation/SkeletonRemappingRegistry.cpp
@@ -111,12 +111,10 @@ const FSkeletonRemapping& FSkeletonRemappingRegistry::GetRemapping(const USkelet
 void FSkeletonRemappingRegistry::RefreshMappings(const USkeleton* InSkeleton)
 {
 	TArray<TSharedPtr<FSkeletonRemapping>> ExistingMappings;
-	{
-		UE::TReadScopeLock ReadLock(MappingsLock);
-		PerSkeletonMappings.MultiFind(InSkeleton, ExistingMappings);
-	}
 	{
 		UE::TWriteScopeLock WriteLock(MappingsLock);
+		PerSkeletonMappings.MultiFind(InSkeleton, ExistingMappings);
+		
 		for(const TSharedPtr<FSkeletonRemapping>& ExistingMapping : ExistingMappings)
 		{
 			ExistingMapping->RegenerateMapping();
