Changed the selection outline rendering to make use of the per instance selection data, rather than relying on internals of the ISM VF.

diff --git a/Engine/Shaders/Private/InstanceCulling/BuildInstanceDrawCommands.usf b/Engine/Shaders/Private/InstanceCulling/BuildInstanceDrawCommands.usf
--- a/Engine/Shaders/Private/InstanceCulling/BuildInstanceDrawCommands.usf
+++ b/Engine/Shaders/Private/InstanceCulling/BuildInstanceDrawCommands.usf
@@ -59,6 +59,7 @@ uint DynamicInstanceIdOffset;
 uint DynamicInstanceIdMax;
 uint NumCullingViews;
 uint CurrentBatchProcessingMode;
+uint bDrawOnlySelected;
 
 RWStructuredBuffer<uint> InstanceIdsBufferOut;
 RWStructuredBuffer<float4> InstanceIdsBufferOutMobile;
@@ -113,7 +114,7 @@ float RoundUpF16(float DeviceZ)
 	return f16tof32(f32tof16(DeviceZ) + 1);
 }
 
-bool IsInstanceVisible(FPrimitiveSceneData PrimitiveData, FInstanceSceneData InstanceData, uint InstanceId, uint ViewIdIndex, bool bAllowOcclusionCulling, FDrawCommandDesc DrawCommandDesc, inout uint CullingFlags)
+bool IsInstanceVisible(FPrimitiveSceneData PrimitiveData, FInstanceSceneData InstanceData, uint InstanceId, uint ViewIdIndex, bool bAllowOcclusionCulling, bool bDrawOnlySelected, FDrawCommandDesc DrawCommandDesc, inout uint CullingFlags)
 {
 	CullingFlags = INSTANCE_CULLING_FLAGS_DEFAULT;
 
@@ -124,6 +125,13 @@ bool IsInstanceVisible(FPrimitiveSceneData PrimitiveData, FInstanceSceneData Ins
 		return false;
 	}
 
+#if USE_EDITOR_SHADERS
+	bool bHasPerInstanceEditorData = (InstanceData.Flags & INSTANCE_SCENE_DATA_FLAG_HAS_EDITOR_DATA) != 0u;
+	if (bDrawOnlySelected && bHasPerInstanceEditorData && !InstanceData.EditorData.bIsSelected)
+	{
+		return false;
+	}
+#endif
 	// TODO: The test for dot(InstanceData.LocalBoundsExtent, InstanceData.LocalBoundsExtent) <= 0.0f is just a workaround since the FDynamicMeshBuilder::GetMesh
 	//       seems to just set empty bounds (and FLineBatcherSceneProxy pretends everything is at the origin). In the future these should compute reasonable bounds and 
 	//       this should be removed.
@@ -252,6 +260,7 @@ void InstanceCullBuildInstanceIdBufferCS(uint3 GroupId : SV_GroupID, int GroupTh
 	BatchInfo.DynamicInstanceIdMax = DynamicInstanceIdMax;
 	// Note: for the unbatched case, the permutation will control HZB test, so we set to true
 	BatchInfo.bAllowOcclusionCulling = true;
+	BatchInfo.bDrawOnlySelected = bDrawOnlySelected;
 #endif // ENABLE_BATCH_MODE
 
 	FInstanceCullingSetup InstanceCullingSetup = LoadInstanceCullingSetup(GroupId, GroupThreadIndex, BatchInfo.DynamicInstanceIdOffset, BatchInfo.DynamicInstanceIdMax, GetItemDataOffset(BatchInfo, CurrentBatchProcessingMode));
@@ -299,8 +308,8 @@ void InstanceCullBuildInstanceIdBufferCS(uint3 GroupId : SV_GroupID, int GroupTh
 #if STEREO_CULLING_MODE
 	uint CullingFlags0 = 0;
 	uint CullingFlags1 = 0;
-	const bool bVisible = IsInstanceVisible(PrimitiveData, InstanceData, InstanceId, BatchInfo.ViewIdsOffset + 0U, BatchInfo.bAllowOcclusionCulling, DrawCommandDesc, CullingFlags0) ||
-					      IsInstanceVisible(PrimitiveData, InstanceData, InstanceId, BatchInfo.ViewIdsOffset + 1U, BatchInfo.bAllowOcclusionCulling, DrawCommandDesc, CullingFlags1);
+	const bool bVisible = IsInstanceVisible(PrimitiveData, InstanceData, InstanceId, BatchInfo.ViewIdsOffset + 0U, BatchInfo.bAllowOcclusionCulling, BatchInfo.bDrawOnlySelected, DrawCommandDesc, CullingFlags0) ||
+					      IsInstanceVisible(PrimitiveData, InstanceData, InstanceId, BatchInfo.ViewIdsOffset + 1U, BatchInfo.bAllowOcclusionCulling, BatchInfo.bDrawOnlySelected, DrawCommandDesc, CullingFlags1);
 	const uint CullingFlags = CullingFlags0 | CullingFlags1;
 	NumVisibleInstances += bVisible ? 2 : 0;
 	
@@ -324,7 +333,7 @@ void InstanceCullBuildInstanceIdBufferCS(uint3 GroupId : SV_GroupID, int GroupTh
 	{
 		// Culling is disabled for multi-view
 		uint CullingFlags = 0;
-		bool bVisible = IsInstanceVisible(PrimitiveData, InstanceData, InstanceId, BatchInfo.ViewIdsOffset + ViewIdIndex, BatchInfo.bAllowOcclusionCulling, DrawCommandDesc, CullingFlags);
+		bool bVisible = IsInstanceVisible(PrimitiveData, InstanceData, InstanceId, BatchInfo.ViewIdsOffset + ViewIdIndex, BatchInfo.bAllowOcclusionCulling, BatchInfo.bDrawOnlySelected, DrawCommandDesc, CullingFlags);
 
 		NumVisibleInstances += bVisible ? 1 : 0;
 		
diff --git a/Engine/Shaders/Private/InstanceCulling/InstanceCullingCommon.ush b/Engine/Shaders/Private/InstanceCulling/InstanceCullingCommon.ush
--- a/Engine/Shaders/Private/InstanceCulling/InstanceCullingCommon.ush
+++ b/Engine/Shaders/Private/InstanceCulling/InstanceCullingCommon.ush
@@ -42,7 +42,7 @@ struct FContextBatchInfoPacked
 	uint PayloadOffset;
 	uint CompactionDataOffset;
 	uint ViewIdsOffset;
-	uint NumViewIds_bAllowOcclusionCulling;
+	uint NumViewIds_Flags;
 	uint DynamicInstanceIdOffset;
 	uint DynamicInstanceIdMax;
 	// Must be equal to ItemDataOffset[BATCH_PROCESSING_MODE_NUM]
@@ -60,6 +60,7 @@ struct FContextBatchInfo
 	uint ViewIdsOffset;
 	uint NumViewIds;
 	bool bAllowOcclusionCulling;
+	bool bDrawOnlySelected;
 	uint DynamicInstanceIdOffset;
 	uint DynamicInstanceIdMax;
 	// Must be equal to ItemDataOffset[BATCH_PROCESSING_MODE_NUM]
@@ -89,8 +90,9 @@ FContextBatchInfo LoadBatchInfo(uint GroupIndex)
 	ContextBatchInfo.PayloadOffset = ContextBatchInfoPacked.PayloadOffset;
 	ContextBatchInfo.CompactionDataOffset = ContextBatchInfoPacked.CompactionDataOffset;
 	ContextBatchInfo.ViewIdsOffset = ContextBatchInfoPacked.ViewIdsOffset;
-	ContextBatchInfo.NumViewIds = ContextBatchInfoPacked.NumViewIds_bAllowOcclusionCulling >> 1u;
-	ContextBatchInfo.bAllowOcclusionCulling = (ContextBatchInfoPacked.NumViewIds_bAllowOcclusionCulling & 1u) != 0u;
+	ContextBatchInfo.NumViewIds = ContextBatchInfoPacked.NumViewIds_Flags >> INSTANCE_CULLING_FLAGS_BITS;
+	ContextBatchInfo.bAllowOcclusionCulling = (ContextBatchInfoPacked.NumViewIds_Flags & INSTANCE_CULLING_FLAGS_ALLOW_OCCLUSION_CULLING) != 0u;
+	ContextBatchInfo.bDrawOnlySelected		= (ContextBatchInfoPacked.NumViewIds_Flags & INSTANCE_CULLING_FLAGS_DRAW_ONLY_SELECTED) != 0u;
 	ContextBatchInfo.DynamicInstanceIdOffset = ContextBatchInfoPacked.DynamicInstanceIdOffset;
 	ContextBatchInfo.DynamicInstanceIdMax = ContextBatchInfoPacked.DynamicInstanceIdMax;
 	ContextBatchInfo.ItemDataOffset0 = ContextBatchInfoPacked.ItemDataOffset0;
diff --git a/Engine/Shaders/Shared/InstanceCullingDefinitions.h b/Engine/Shaders/Shared/InstanceCullingDefinitions.h
--- a/Engine/Shaders/Shared/InstanceCullingDefinitions.h
+++ b/Engine/Shaders/Shared/InstanceCullingDefinitions.h
@@ -6,3 +6,7 @@
 #define INSTANCE_CULLING_DYNAMIC_INSTANCE_DATA_OFFSET_BIT_MASK 2U
 
 #define INSTANCE_CULLING_PAYLOAD_NUM_COMMON_BITS 2U
+
+#define INSTANCE_CULLING_FLAGS_BITS (2u)
+#define INSTANCE_CULLING_FLAGS_ALLOW_OCCLUSION_CULLING	(1u << 0u)
+#define INSTANCE_CULLING_FLAGS_DRAW_ONLY_SELECTED		(1u << 1u)
\ No newline at end of file
diff --git a/Engine/Source/Runtime/Engine/Classes/Engine/InstancedStaticMesh.h b/Engine/Source/Runtime/Engine/Classes/Engine/InstancedStaticMesh.h
--- a/Engine/Source/Runtime/Engine/Classes/Engine/InstancedStaticMesh.h
+++ b/Engine/Source/Runtime/Engine/Classes/Engine/InstancedStaticMesh.h
@@ -544,8 +544,6 @@ protected:
 
 	/** LOD transition info. */
 	FInstancingUserData UserData_AllInstances;
-	FInstancingUserData UserData_SelectedInstances;
-	FInstancingUserData UserData_DeselectedInstances;
 
 	/** LOD distance scale from component. */
 	float InstanceLODDistanceScale;
diff --git a/Engine/Source/Runtime/Engine/Private/HierarchicalInstancedStaticMesh.cpp b/Engine/Source/Runtime/Engine/Private/HierarchicalInstancedStaticMesh.cpp
--- a/Engine/Source/Runtime/Engine/Private/HierarchicalInstancedStaticMesh.cpp
+++ b/Engine/Source/Runtime/Engine/Private/HierarchicalInstancedStaticMesh.cpp
@@ -816,12 +816,7 @@ FHierarchicalStaticMeshSceneProxy::FHierarchicalStaticMeshSceneProxy(UHierarchic
 	// Store LODDistanceScale so it can be used in FInstancedStaticMeshVertexFactoryShaderParameters::GetElementShaderBindings when dither LOD transitions are enabled
 	float LODDistanceScale = InComponent->InstanceLODDistanceScale * CVarFoliageLODDistanceScale.GetValueOnGameThread();
 	UserData_AllInstances.LODDistanceScale = LODDistanceScale;
-	UserData_SelectedInstances.LODDistanceScale = LODDistanceScale;
-	UserData_DeselectedInstances.LODDistanceScale = LODDistanceScale;
-
-	UserData_AllInstances.AverageInstancesScale = 
-		UserData_SelectedInstances.AverageInstancesScale = 
-		UserData_DeselectedInstances.AverageInstancesScale = InComponent->GetAverageScale();
+	UserData_AllInstances.AverageInstancesScale = InComponent->GetAverageScale();
 }
 
 void FHierarchicalStaticMeshSceneProxy::SetupOcclusion(UHierarchicalInstancedStaticMeshComponent* InComponent)
@@ -1212,12 +1207,11 @@ void FHierarchicalStaticMeshSceneProxy::Traverse(const FFoliageCullInstanceParam
 
 struct FFoliageElementParams
 {
-	const FInstancingUserData* PassUserData[2];
-	int32 NumSelectionGroups;
+	const FInstancingUserData* PassUserData;
 	const FSceneView* View;
 	int32 ViewIndex;
 	bool bSelectionRenderEnabled;
-	bool BatchRenderSelection[2];
+	bool bBatchRenderSelection;
 	bool bIsWireframe;
 	bool bUseHoveredMaterial;
 	bool bUseInstanceRuns;
@@ -1254,9 +1248,8 @@ void FHierarchicalStaticMeshSceneProxy::FillDynamicMeshElements(const FSceneView
 		const FStaticMeshLODResources& LODModel = RenderData->LODResources[LODIndex];
 		const int32 TotalNumSections = LODModel.Sections.Num();
 
-		for (int32 SelectionGroupIndex = 0; SelectionGroupIndex < ElementParams.NumSelectionGroups; SelectionGroupIndex++)
 		{
-			FInstancedStaticMeshVFLooseUniformShaderParametersRef LooseUniformBuffer = CreateLooseUniformBuffer(View, ElementParams.PassUserData[SelectionGroupIndex], InstancedLODRange, LODIndex, EUniformBufferUsage::UniformBuffer_SingleFrame);
+			FInstancedStaticMeshVFLooseUniformShaderParametersRef LooseUniformBuffer = CreateLooseUniformBuffer(View, ElementParams.PassUserData, InstancedLODRange, LODIndex, EUniformBufferUsage::UniformBuffer_SingleFrame);
 
 			for (int32 SectionIndex = 0; SectionIndex < TotalNumSections; SectionIndex++)
 			{
@@ -1273,20 +1266,20 @@ void FHierarchicalStaticMeshSceneProxy::FillDynamicMeshElements(const FSceneView
 					FMeshBatch& MeshBatch = Collector.AllocateMesh();
 					INC_DWORD_STAT(STAT_FoliageMeshBatches);
 
-					if (!FStaticMeshSceneProxy::GetMeshElement(LODIndex, 0, SectionIndex, GetDepthPriorityGroup(ElementParams.View), ElementParams.BatchRenderSelection[SelectionGroupIndex], true, MeshBatch))
+					if (!FStaticMeshSceneProxy::GetMeshElement(LODIndex, 0, SectionIndex, GetDepthPriorityGroup(ElementParams.View), ElementParams.bBatchRenderSelection, true, MeshBatch))
 					{
 						continue;
 					}
 
 					checkSlow(MeshBatch.GetNumPrimitives() > 0);
 					MeshBatch.bCanApplyViewModeOverrides = true;
-					MeshBatch.bUseSelectionOutline = ElementParams.BatchRenderSelection[SelectionGroupIndex];
-					MeshBatch.bUseWireframeSelectionColoring = ElementParams.BatchRenderSelection[SelectionGroupIndex];
+					MeshBatch.bUseSelectionOutline = ElementParams.bBatchRenderSelection;
+					MeshBatch.bUseWireframeSelectionColoring = ElementParams.bBatchRenderSelection;
 					MeshBatch.bUseAsOccluder = ShouldUseAsOccluder();
 					MeshBatch.VertexFactory = &InstancedRenderData.VertexFactories[LODIndex];
 
 					FMeshBatchElement& MeshBatchElement = MeshBatch.Elements[0];
-					MeshBatchElement.UserData = ElementParams.PassUserData[SelectionGroupIndex];
+					MeshBatchElement.UserData = ElementParams.PassUserData;
 					MeshBatchElement.bUserDataIsColorVertexBuffer = false;
 					MeshBatchElement.MaxScreenSize = 1.0;
 					MeshBatchElement.MinScreenSize = 0.0;
@@ -1305,16 +1298,16 @@ void FHierarchicalStaticMeshSceneProxy::FillDynamicMeshElements(const FSceneView
 						{
 							if (ElementParams.FinalCullDistance > 9.9E8)
 							{
-								UE_LOG(LogStaticMesh, Display, TEXT("lod:%1d/%1d   sel:%1d   section:%1d/%1d   runs:%4d   inst:%8d   tris:%9lld   cast shadow:%1d   cull:-NONE!!-   shadow:%1d     %s %s"),
-									LODIndex, InstancedRenderData.VertexFactories.Num(), SelectionGroupIndex, SectionIndex, TotalNumSections, RunArray.Num() / 2,
+								UE_LOG(LogStaticMesh, Display, TEXT("lod:%1d/%1d   section:%1d/%1d   runs:%4d   inst:%8d   tris:%9lld   cast shadow:%1d   cull:-NONE!!-   shadow:%1d     %s %s"),
+									LODIndex, InstancedRenderData.VertexFactories.Num(), SectionIndex, TotalNumSections, RunArray.Num() / 2,
 									TotalInstances, Tris, (int)MeshBatch.CastShadow, ElementParams.ShadowFrustum,
 									*StaticMesh->GetPathName(),
 									*MeshBatch.MaterialRenderProxy->GetIncompleteMaterialWithFallback(ElementParams.FeatureLevel).GetFriendlyName());
 							}
 							else
 							{
-								UE_LOG(LogStaticMesh, Display, TEXT("lod:%1d/%1d   sel:%1d   section:%1d/%1d   runs:%4d   inst:%8d   tris:%9lld   cast shadow:%1d   cull:%8.0f   shadow:%1d     %s %s"),
-									LODIndex, InstancedRenderData.VertexFactories.Num(), SelectionGroupIndex, SectionIndex, TotalNumSections, RunArray.Num() / 2,
+								UE_LOG(LogStaticMesh, Display, TEXT("lod:%1d/%1d   section:%1d/%1d   runs:%4d   inst:%8d   tris:%9lld   cast shadow:%1d   cull:%8.0f   shadow:%1d     %s %s"),
+									LODIndex, InstancedRenderData.VertexFactories.Num(), SectionIndex, TotalNumSections, RunArray.Num() / 2,
 									TotalInstances, Tris, (int)MeshBatch.CastShadow, ElementParams.FinalCullDistance, ElementParams.ShadowFrustum,
 									*StaticMesh->GetPathName(),
 									*MeshBatch.MaterialRenderProxy->GetIncompleteMaterialWithFallback(ElementParams.FeatureLevel).GetFriendlyName());
@@ -1375,7 +1368,7 @@ void FHierarchicalStaticMeshSceneProxy::FillDynamicMeshElements(const FSceneView
 						FMeshBatch& MeshElement = Collector.AllocateMesh();
 						INC_DWORD_STAT(STAT_FoliageMeshBatches);
 
-						if (!FStaticMeshSceneProxy::GetMeshElement(LODIndex, 0, SectionIndex, GetDepthPriorityGroup(ElementParams.View), ElementParams.BatchRenderSelection[SelectionGroupIndex], true, MeshElement))
+						if (!FStaticMeshSceneProxy::GetMeshElement(LODIndex, 0, SectionIndex, GetDepthPriorityGroup(ElementParams.View), ElementParams.bBatchRenderSelection, true, MeshElement))
 						{
 							continue;
 						}
@@ -1384,7 +1377,7 @@ void FHierarchicalStaticMeshSceneProxy::FillDynamicMeshElements(const FSceneView
 						MeshElement.VertexFactory = &InstancedRenderData.VertexFactories[LODIndex];
 						FMeshBatchElement& BatchElement0 = MeshElement.Elements[0];
 
-						BatchElement0.UserData = ElementParams.PassUserData[SelectionGroupIndex];
+						BatchElement0.UserData = ElementParams.PassUserData;
 						BatchElement0.bUserDataIsColorVertexBuffer = false;
 						BatchElement0.MaxScreenSize = 1.0;
 						BatchElement0.MinScreenSize = 0.0;
@@ -1393,8 +1386,8 @@ void FHierarchicalStaticMeshSceneProxy::FillDynamicMeshElements(const FSceneView
 						BatchElement0.PrimitiveUniformBuffer = GetUniformBuffer();
 						BatchElement0.LooseParametersUniformBuffer = LooseUniformBuffer;
 						MeshElement.bCanApplyViewModeOverrides = true;
-						MeshElement.bUseSelectionOutline = ElementParams.BatchRenderSelection[SelectionGroupIndex];
-						MeshElement.bUseWireframeSelectionColoring = ElementParams.BatchRenderSelection[SelectionGroupIndex];
+						MeshElement.bUseSelectionOutline = ElementParams.bBatchRenderSelection;
+						MeshElement.bUseWireframeSelectionColoring = ElementParams.bBatchRenderSelection;
 						MeshElement.bUseAsOccluder = ShouldUseAsOccluder();
 
 						if (!bDidStats)
@@ -1407,16 +1400,16 @@ void FHierarchicalStaticMeshSceneProxy::FillDynamicMeshElements(const FSceneView
 							{
 								if (ElementParams.FinalCullDistance > 9.9E8)
 								{
-									UE_LOG(LogStaticMesh, Display, TEXT("lod:%1d/%1d   sel:%1d   section:%1d/%1d   runs:%4d   inst:%8d   tris:%9lld   cast shadow:%1d   cull:-NONE!!-   shadow:%1d     %s %s"),
-										LODIndex, InstancedRenderData.VertexFactories.Num(), SelectionGroupIndex, SectionIndex, TotalNumSections, RunArray.Num() / 2,
+									UE_LOG(LogStaticMesh, Display, TEXT("lod:%1d/%1d   section:%1d/%1d   runs:%4d   inst:%8d   tris:%9lld   cast shadow:%1d   cull:-NONE!!-   shadow:%1d     %s %s"),
+										LODIndex, InstancedRenderData.VertexFactories.Num(), SectionIndex, TotalNumSections, RunArray.Num() / 2,
 										RemainingInstances, Tris, (int)MeshElement.CastShadow, ElementParams.ShadowFrustum,
 										*StaticMesh->GetPathName(),
 										*MeshElement.MaterialRenderProxy->GetIncompleteMaterialWithFallback(ElementParams.FeatureLevel).GetFriendlyName());
 								}
 								else
 								{
-									UE_LOG(LogStaticMesh, Display, TEXT("lod:%1d/%1d   sel:%1d   section:%1d/%1d   runs:%4d   inst:%8d   tris:%9lld   cast shadow:%1d   cull:%8.0f   shadow:%1d     %s %s"),
-										LODIndex, InstancedRenderData.VertexFactories.Num(), SelectionGroupIndex, SectionIndex, TotalNumSections, RunArray.Num() / 2,
+									UE_LOG(LogStaticMesh, Display, TEXT("lod:%1d/%1d   section:%1d/%1d   runs:%4d   inst:%8d   tris:%9lld   cast shadow:%1d   cull:%8.0f   shadow:%1d     %s %s"),
+										LODIndex, InstancedRenderData.VertexFactories.Num(), SectionIndex, TotalNumSections, RunArray.Num() / 2,
 										RemainingInstances, Tris, (int)MeshElement.CastShadow, ElementParams.FinalCullDistance, ElementParams.ShadowFrustum,
 										*StaticMesh->GetPathName(),
 										*MeshElement.MaterialRenderProxy->GetIncompleteMaterialWithFallback(ElementParams.FeatureLevel).GetFriendlyName());
@@ -1528,11 +1521,8 @@ void FHierarchicalStaticMeshSceneProxy::GetDynamicMeshElements(const TArray<cons
 
 			FFoliageElementParams ElementParams;
 			ElementParams.bSelectionRenderEnabled = GIsEditor && ViewFamily.EngineShowFlags.Selection;
-			ElementParams.NumSelectionGroups = (ElementParams.bSelectionRenderEnabled && bHasSelectedInstances) ? 2 : 1;
-			ElementParams.PassUserData[0] = bHasSelectedInstances && ElementParams.bSelectionRenderEnabled ? &UserData_SelectedInstances : &UserData_AllInstances;
-			ElementParams.PassUserData[1] = &UserData_DeselectedInstances;
-			ElementParams.BatchRenderSelection[0] = ElementParams.bSelectionRenderEnabled && IsSelected();
-			ElementParams.BatchRenderSelection[1] = false;
+			ElementParams.PassUserData = &UserData_AllInstances;
+			ElementParams.bBatchRenderSelection = ElementParams.bSelectionRenderEnabled && IsSelected();
 			ElementParams.bIsWireframe = ViewFamily.EngineShowFlags.Wireframe;
 			ElementParams.bUseHoveredMaterial = IsHovered();
 			ElementParams.bUseInstanceRuns = (CVarFoliageUseInstanceRuns.GetValueOnRenderThread() > 0);
@@ -2969,7 +2959,7 @@ void UHierarchicalInstancedStaticMeshComponent::SetPerInstanceLightMapAndEditorD
 			{
 				// Record if the instance is selected
 				FColor HitProxyColor(ForceInit);
-				bool bSelected = SelectedInstances.IsValidIndex(Index) && SelectedInstances[Index];
+				bool bSelected = !SelectedInstances.IsValidIndex(Index) || SelectedInstances[Index];
 
 				if (HitProxies.IsValidIndex(Index))
 				{
diff --git a/Engine/Source/Runtime/Engine/Private/InstanceData/InstanceUpdateChangeSet.cpp b/Engine/Source/Runtime/Engine/Private/InstanceData/InstanceUpdateChangeSet.cpp
--- a/Engine/Source/Runtime/Engine/Private/InstanceData/InstanceUpdateChangeSet.cpp
+++ b/Engine/Source/Runtime/Engine/Private/InstanceData/InstanceUpdateChangeSet.cpp
@@ -47,7 +47,7 @@ void FInstanceUpdateChangeSet::SetEditorData(const TArray<TRefCountPtr<HHitProxy
 	{
 		// Record if the instance is selected
 		FColor HitProxyColor(ForceInit);
-		bool bSelected = InSelectedInstances.IsValidIndex(Index) && InSelectedInstances[Index];
+		bool bSelected = !InSelectedInstances.IsValidIndex(Index) || InSelectedInstances[Index];
 		if (HitProxies.IsValidIndex(Index))
 		{
 			HitProxyColor = HitProxies[Index]->Id.GetColor();
diff --git a/Engine/Source/Runtime/Engine/Private/InstancedStaticMesh.cpp b/Engine/Source/Runtime/Engine/Private/InstancedStaticMesh.cpp
--- a/Engine/Source/Runtime/Engine/Private/InstancedStaticMesh.cpp
+++ b/Engine/Source/Runtime/Engine/Private/InstancedStaticMesh.cpp
@@ -1090,21 +1090,7 @@ void FInstancedStaticMeshSceneProxy::GetDynamicMeshElements(const TArray<const F
 
 	const bool bSelectionRenderEnabled = GIsEditor && ViewFamily.EngineShowFlags.Selection;
 
-	// If the first pass rendered selected instances only, we need to render the deselected instances in a second pass
-	const int32 NumSelectionGroups = (bSelectionRenderEnabled && bHasSelectedInstances) ? 2 : 1;
-
-	const FInstancingUserData* PassUserData[2] =
-	{
-		bHasSelectedInstances && bSelectionRenderEnabled ? &UserData_SelectedInstances : &UserData_AllInstances,
-		&UserData_DeselectedInstances
-	};
-
-	bool BatchRenderSelection[2] = 
-	{
-		bSelectionRenderEnabled && IsSelected(),
-		false
-	};
-
+	bool bBatchRenderSelection = bSelectionRenderEnabled && IsSelected();
 	const bool bIsWireframe = ViewFamily.EngineShowFlags.Wireframe;
 
 	for (int32 ViewIndex = 0; ViewIndex < Views.Num(); ViewIndex++)
@@ -1114,13 +1100,12 @@ void FInstancedStaticMeshSceneProxy::GetDynamicMeshElements(const TArray<const F
 			const FSceneView* View = Views[ViewIndex];
 			FPrimitiveDrawInterface* PDI = Collector.GetPDI(ViewIndex);
 
-			for (int32 SelectionGroupIndex = 0; SelectionGroupIndex < NumSelectionGroups; SelectionGroupIndex++)
 			{
 				const int32 FirstLODIndex = RenderData->GetFirstValidLODIdx(RenderData->CurrentFirstLODIdx);
 				const int32 LODIndex = FMath::Max(GetLOD(View), FirstLODIndex);
 				const FStaticMeshLODResources& LODModel = RenderData->LODResources[LODIndex];
 
-				FInstancedStaticMeshVFLooseUniformShaderParametersRef LooseUniformBuffer = CreateLooseUniformBuffer(View, PassUserData[SelectionGroupIndex], /*InstancedLODRange=*/0, LODIndex, EUniformBufferUsage::UniformBuffer_SingleFrame);
+				FInstancedStaticMeshVFLooseUniformShaderParametersRef LooseUniformBuffer = CreateLooseUniformBuffer(View, &UserData_AllInstances, /*InstancedLODRange=*/0, LODIndex, EUniformBufferUsage::UniformBuffer_SingleFrame);
 
 				for (int32 SectionIndex = 0; SectionIndex < LODModel.Sections.Num(); SectionIndex++)
 				{
@@ -1130,15 +1115,15 @@ void FInstancedStaticMeshSceneProxy::GetDynamicMeshElements(const TArray<const F
 					{
 						FMeshBatch& MeshElement = Collector.AllocateMesh();
 
-						if (GetMeshElement(LODIndex, BatchIndex, SectionIndex, GetDepthPriorityGroup(View), BatchRenderSelection[SelectionGroupIndex], true, MeshElement))
+						if (GetMeshElement(LODIndex, BatchIndex, SectionIndex, GetDepthPriorityGroup(View), bBatchRenderSelection, true, MeshElement))
 						{
 							//@todo-rco this is only supporting selection on the first element
-							MeshElement.Elements[0].UserData = PassUserData[SelectionGroupIndex];
+							MeshElement.Elements[0].UserData = &UserData_AllInstances;
 							MeshElement.Elements[0].bUserDataIsColorVertexBuffer = false;
 							MeshElement.Elements[0].LooseParametersUniformBuffer = LooseUniformBuffer;
 							MeshElement.bCanApplyViewModeOverrides = true;
-							MeshElement.bUseSelectionOutline = BatchRenderSelection[SelectionGroupIndex];
-							MeshElement.bUseWireframeSelectionColoring = BatchRenderSelection[SelectionGroupIndex];
+							MeshElement.bUseSelectionOutline = bBatchRenderSelection;
+							MeshElement.bUseWireframeSelectionColoring = bBatchRenderSelection;
 
 							if (View->bRenderFirstInstanceOnly)
 							{
@@ -1551,14 +1536,6 @@ void FInstancedStaticMeshSceneProxy::SetupProxy(const FInstancedStaticMeshSceneP
 	// Only used by HISM, and thus set in the descendant ctor
 	UserData_AllInstances.AverageInstancesScale = FVector::Zero();
 
-	// selected only
-	UserData_SelectedInstances = UserData_AllInstances;
-	UserData_SelectedInstances.bRenderUnselected = false;
-
-	// unselected only
-	UserData_DeselectedInstances = UserData_AllInstances;
-	UserData_DeselectedInstances.bRenderSelected = false;
-
 	const bool bUseGPUScene = UseGPUScene(GetScene().GetShaderPlatform(), GetScene().GetFeatureLevel());
 	
 	const bool bEnableGpuLodSelection = CVarGpuLodSelection.GetValueOnAnyThread() != 0;
@@ -2287,12 +2264,6 @@ void FInstancedStaticMeshSceneProxy::SetInstanceCullDistance_RenderThread(float
 	UserData_AllInstances.StartCullDistance = StartCullDistance;
 	UserData_AllInstances.EndCullDistance = EndCullDistance;
 
-	UserData_SelectedInstances.StartCullDistance = StartCullDistance;
-	UserData_SelectedInstances.EndCullDistance = EndCullDistance;
-
-	UserData_DeselectedInstances.StartCullDistance = StartCullDistance;
-	UserData_DeselectedInstances.EndCullDistance = EndCullDistance;
-
 	for (int32 i = 0; i < LODs.Num(); ++i)
 	{
 		FInstancedStaticMeshVFLooseUniformShaderParametersRef LooseUniformBuffer = CreateLooseUniformBuffer(nullptr, &UserData_AllInstances, 0, i, EUniformBufferUsage::UniformBuffer_MultiFrame);
@@ -5925,8 +5896,14 @@ void UInstancedStaticMeshComponent::SelectInstance(bool bInSelected, int32 InIns
 			if (SelectedInstances.IsValidIndex(InInstanceIndex))
 			{
 				SelectedInstances[InstanceIndex] = bInSelected;
-					}
-					}
+			}
+		}
+
+		// Reset the flags if all are false (this is also used to ensure per primitive selection works as intended since it relies on this being cleared when per instance selection is disabled).
+		if (SelectedInstances.Find(true) == INDEX_NONE)
+		{
+			SelectedInstances.Empty();
+		}
 		
 		// We need to recreate to make sure it changes to the dynamic path, which makes selection outline rendering work.
 		MarkRenderStateDirty();
diff --git a/Engine/Source/Runtime/Renderer/Private/InstanceCulling/InstanceCullingContext.cpp b/Engine/Source/Runtime/Renderer/Private/InstanceCulling/InstanceCullingContext.cpp
--- a/Engine/Source/Runtime/Renderer/Private/InstanceCulling/InstanceCullingContext.cpp
+++ b/Engine/Source/Runtime/Renderer/Private/InstanceCulling/InstanceCullingContext.cpp
@@ -603,6 +603,7 @@ public:
 
 		SHADER_PARAMETER(int32, DynamicInstanceIdOffset)
 		SHADER_PARAMETER(int32, DynamicInstanceIdMax)
+		SHADER_PARAMETER(uint32, bDrawOnlySelected)
 	END_SHADER_PARAMETER_STRUCT()
 };
 IMPLEMENT_GLOBAL_SHADER(FBuildInstanceIdBufferAndCommandsFromPrimitiveIdsCs, "/Engine/Private/InstanceCulling/BuildInstanceDrawCommands.usf", "InstanceCullBuildInstanceIdBufferCS", SF_Compute);
@@ -813,6 +814,7 @@ void FInstanceCullingContext::BuildRenderingCommandsInternal(
 	PassParametersTmp.GPUSceneParameters = GPUScene.GetShaderParameters(GraphBuilder);
 	PassParametersTmp.DynamicInstanceIdOffset = DynamicInstanceIdOffset;
 	PassParametersTmp.DynamicInstanceIdMax = DynamicInstanceIdOffset + DynamicInstanceIdNum;
+	PassParametersTmp.bDrawOnlySelected = EnumHasAnyFlags(Flags, EInstanceCullingFlags::DrawOnlySelected);
 
 	// Compaction parameters
 	PassParametersTmp.DrawCommandCompactionData = DrawCommandCompactionDataSRV;
diff --git a/Engine/Source/Runtime/Renderer/Private/InstanceCulling/InstanceCullingMergedContext.cpp b/Engine/Source/Runtime/Renderer/Private/InstanceCulling/InstanceCullingMergedContext.cpp
--- a/Engine/Source/Runtime/Renderer/Private/InstanceCulling/InstanceCullingMergedContext.cpp
+++ b/Engine/Source/Runtime/Renderer/Private/InstanceCulling/InstanceCullingMergedContext.cpp
@@ -3,6 +3,7 @@
 #include "InstanceCulling/InstanceCullingMergedContext.h"
 #include "InstanceCulling/InstanceCullingManager.h"
 #include "RenderGraphBuilder.h"
+#include "InstanceCullingDefinitions.h"
 
 FInstanceCullingMergedContext::FInstanceCullingMergedContext(EShaderPlatform InShaderPlatform, bool bInMustAddAllContexts, int32 InNumBins)
 	: ShaderPlatform(InShaderPlatform)
@@ -87,10 +88,14 @@ void FInstanceCullingMergedContext::MergeBatches()
 		}
 
 		BatchInfo.ViewIdsOffset = ViewIds.Num();
-		BatchInfo.NumViewIds_bAllowOcclusionCulling = uint32(InstanceCullingContext.ViewIds.Num()) << 1u;
+		BatchInfo.NumViewIds_Flags = uint32(InstanceCullingContext.ViewIds.Num()) << INSTANCE_CULLING_FLAGS_BITS;
 		if (InstanceCullingContext.PrevHZB.IsValid())
 		{
-			BatchInfo.NumViewIds_bAllowOcclusionCulling |= 1u;
+			BatchInfo.NumViewIds_Flags |= INSTANCE_CULLING_FLAGS_ALLOW_OCCLUSION_CULLING;
+		}
+		if (EnumHasAnyFlags(InstanceCullingContext.Flags, EInstanceCullingFlags::DrawOnlySelected))
+		{
+			BatchInfo.NumViewIds_Flags |= INSTANCE_CULLING_FLAGS_DRAW_ONLY_SELECTED;
 		}
 		ViewIds.Append(InstanceCullingContext.ViewIds);
 
diff --git a/Engine/Source/Runtime/Renderer/Private/InstanceCulling/InstanceCullingMergedContext.h b/Engine/Source/Runtime/Renderer/Private/InstanceCulling/InstanceCullingMergedContext.h
--- a/Engine/Source/Runtime/Renderer/Private/InstanceCulling/InstanceCullingMergedContext.h
+++ b/Engine/Source/Runtime/Renderer/Private/InstanceCulling/InstanceCullingMergedContext.h
@@ -28,7 +28,7 @@ public:
 		uint32 PayloadDataOffset;
 		uint32 CompactionDataOffset;
 		uint32 ViewIdsOffset;
-		uint32 NumViewIds_bAllowOcclusionCulling;
+		uint32 NumViewIds_Flags;
 		uint32 DynamicInstanceIdOffset;
 		uint32 DynamicInstanceIdMax;
 		uint32 ItemDataOffset[uint32(EBatchProcessingMode::Num)];
diff --git a/Engine/Source/Runtime/Renderer/Private/SceneRendering.cpp b/Engine/Source/Runtime/Renderer/Private/SceneRendering.cpp
--- a/Engine/Source/Runtime/Renderer/Private/SceneRendering.cpp
+++ b/Engine/Source/Runtime/Renderer/Private/SceneRendering.cpp
@@ -4772,6 +4772,13 @@ void FSceneRenderer::SetupMeshPass(FViewInfo& View, FExclusiveDepthStencil::Type
 			
 			EInstanceCullingFlags CullingFlags = EInstanceCullingFlags::None;
 
+#if WITH_EDITOR
+			if (PassType == EMeshPass::EditorSelection)
+			{
+				CullingFlags |= EInstanceCullingFlags::DrawOnlySelected;
+			}
+#endif
+
 			// TODO: Maybe this should be configured somewhere else?
 			const bool bAllowInstanceOcclusionCulling = PassType != EMeshPass::CustomDepth;
 
diff --git a/Engine/Source/Runtime/Renderer/Public/InstanceCulling/InstanceCullingContext.h b/Engine/Source/Runtime/Renderer/Public/InstanceCulling/InstanceCullingContext.h
--- a/Engine/Source/Runtime/Renderer/Public/InstanceCulling/InstanceCullingContext.h
+++ b/Engine/Source/Runtime/Renderer/Public/InstanceCulling/InstanceCullingContext.h
@@ -49,8 +49,9 @@ enum class EInstanceCullingMode
 
 enum class EInstanceCullingFlags : uint8
 {
-	None 							= 0,
-	NoInstanceOrderPreservation		= 1 << 0,
+	None 							= 0u,
+	NoInstanceOrderPreservation		= 1u << 0,
+	DrawOnlySelected				= 1u << 1,
 };
 ENUM_CLASS_FLAGS(EInstanceCullingFlags)
 
