Add detailed PSO precaching tracking for all precached PSOs with time spend compiling them as well (and if used)

diff --git a/Engine/Source/Runtime/Engine/Private/PSOPrecacheMaterial.cpp b/Engine/Source/Runtime/Engine/Private/PSOPrecacheMaterial.cpp
--- a/Engine/Source/Runtime/Engine/Private/PSOPrecacheMaterial.cpp
+++ b/Engine/Source/Runtime/Engine/Private/PSOPrecacheMaterial.cpp
@@ -373,7 +373,11 @@ public:
 		check(FindResult->ActivePSOPrecacheRequests.IsEmpty());
 		FindResult->ActivePSOPrecacheRequests = PrecacheRequestResults;
 #if PSO_PRECACHING_TRACKING
+		// TODO: remove PSOPrecachaData from FPrecacheData because it's stored already in the tracked data
+		// 	     but need to be sure that every material name is unique then.
+		// 		 Currently checked by pointer of material and not just name - but pointer can't be stored for tracked because material could be deleted already
 		FindResult->PSOPrecachaData = PrecacheData;
+		AddTrackedMaterialData(Params, PrecacheData);
 #endif
 
 		// update the state
@@ -535,6 +539,142 @@ public:
 		return FindResult->PSOPrecachaData;
 	}
 
+	void DumpPSOPrecacheMaterialStats()
+	{		
+		FRWScopeLock ReadLock(RWLock, SLT_ReadOnly);
+
+		EShadingPath ShadingPath = GetFeatureLevelShadingPath(GMaxRHIFeatureLevel);
+
+		struct FEntry
+		{
+			FString MaterialName;
+			FString VertexFactoryName;
+			FString PassName;
+			uint32 CustomDefaultDeclarationPSOPrecacheHash;
+			FPSOPrecacheParams PrecachePSOParams;
+			bool bUsed;
+			bool bGraphics;
+			bool bRequired;
+			uint64 PSOPrecacheHash;
+			float PSOCompilationTime;
+		};
+
+		TArray<FPSOPrecacheParams> AllPSOPrecacheParams;
+		TArray<FRHIVertexDeclaration*> AllCustomVertexDeclerations;
+		TArray<FEntry> Entries;
+		for (auto Iter : TrackingMaterialPSORequestData)
+		{
+			uint64 PSOPrecacheHash = Iter.Key;
+			const FTrackingMaterialPSOPrecacheParams& PrecacheParams = Iter.Value;
+				
+			bool bUsed = PSOCollectorStats::GetFullPSOPrecacheStatsCollector().IsUsed(PSOPrecacheHash);
+			const TCHAR* PassName = FPSOCollectorCreateManager::GetName(ShadingPath, PrecacheParams.PSOCollectorIndex);
+			uint64 PSOCompilationTime = PipelineStateCache::GetPrecachePSOCompilationTime(PrecacheParams.bGraphics, PSOPrecacheHash);
+								
+			FEntry& NewEntry = Entries.AddDefaulted_GetRef();
+			NewEntry.MaterialName = PrecacheParams.MaterialName;
+			NewEntry.VertexFactoryName = PrecacheParams.VertexFactoryData.VertexFactoryType ? PrecacheParams.VertexFactoryData.VertexFactoryType->GetName() : TEXT("None");
+			NewEntry.PassName = PassName != nullptr ? PassName : TEXT("None");
+			NewEntry.CustomDefaultDeclarationPSOPrecacheHash = PrecacheParams.VertexFactoryData.CustomDefaultVertexDeclaration != nullptr ? PrecacheParams.VertexFactoryData.CustomDefaultVertexDeclaration->GetPrecachePSOHash() : 0;
+			NewEntry.PrecachePSOParams = PrecacheParams.PrecachePSOParams;
+			NewEntry.bUsed = bUsed;
+			NewEntry.bGraphics = PrecacheParams.bGraphics;
+			NewEntry.bRequired = PrecacheParams.bRequired;
+			NewEntry.PSOPrecacheHash = PSOPrecacheHash;
+			NewEntry.PSOCompilationTime = FPlatformTime::ToMilliseconds64(PSOCompilationTime);
+
+			AllPSOPrecacheParams.AddUnique(PrecacheParams.PrecachePSOParams);
+			AllCustomVertexDeclerations.AddUnique(PrecacheParams.VertexFactoryData.CustomDefaultVertexDeclaration);
+		}
+
+		Algo::Sort(Entries, [](FEntry& LHS, FEntry& RHS)
+		   {
+			   if (LHS.MaterialName != RHS.MaterialName)
+			   {
+				   return LHS.MaterialName < RHS.MaterialName;
+			   }
+			   
+			   if (LHS.VertexFactoryName != RHS.VertexFactoryName)
+			   {
+				   return LHS.VertexFactoryName < RHS.VertexFactoryName;
+			   }
+			   
+			   if (LHS.PassName != RHS.PassName)
+			   {
+				   return LHS.PassName < RHS.PassName;
+			   }
+		
+			   return LHS.PrecachePSOParams.Data < RHS.PrecachePSOParams.Data;
+		   });
+
+		{
+			const FString Filename = FString::Printf(TEXT("%sPSOPrecacheMaterialStats-%s.csv"), *FPaths::ProfilingDir(), *FDateTime::Now().ToString());
+			FArchive* CSVFile = IFileManager::Get().CreateFileWriter(*Filename, FILEWRITE_AllowRead);
+			if (CSVFile != nullptr)
+			{
+				FString Header(TEXT("Material,VertexFactory,MeshPass,PrecachePSOParams,Used,Type,Required,PSOCompilationTime,CustomVertexDeclarationHash,PSOPrecacheHash\n"));
+				CSVFile->Serialize(TCHAR_TO_ANSI(*Header), Header.Len());
+				for (FEntry Entry : Entries)
+				{
+					FString Row = FString::Printf(TEXT("%s,%s,%s,%d,%s,%s,%s,%.2f,%x,%llx\n"),
+												  *Entry.MaterialName,
+												  *Entry.VertexFactoryName,
+												  *Entry.PassName,
+												  Entry.PrecachePSOParams.Data,
+												  Entry.bUsed ? TEXT("Used") : TEXT("NotUsed"),
+												  Entry.bGraphics ? TEXT("Graphics") : TEXT("Compute"),
+												  Entry.bRequired ? TEXT("Yes") : TEXT("No"),
+												  Entry.PSOCompilationTime,
+												  Entry.CustomDefaultDeclarationPSOPrecacheHash,
+												  Entry.PSOPrecacheHash);
+					CSVFile->Serialize(TCHAR_TO_ANSI(*Row), Row.Len());
+				}
+			}
+			delete CSVFile;
+			CSVFile = nullptr;
+		}
+
+		{
+			const FString Filename = FString::Printf(TEXT("%sPSOPrecacheCustomVertexDeclarations-%s.log"), *FPaths::ProfilingDir(), *FDateTime::Now().ToString());
+			FArchive* CSVFile = IFileManager::Get().CreateFileWriter(*Filename, FILEWRITE_AllowRead);
+			if (CSVFile != nullptr)
+			{				
+				for (FRHIVertexDeclaration* VertexDeclaration : AllCustomVertexDeclerations)
+				{
+					if (VertexDeclaration == nullptr)
+					{
+						continue;
+					}
+
+					FStringBuilderBase StringBuilder;
+					StringBuilder.Appendf(TEXT("VertexDeclaration: %x\n"), VertexDeclaration->GetPrecachePSOHash());
+			
+					FVertexDeclarationElementList Elements;
+					if (VertexDeclaration->GetInitializer(Elements))
+					{
+						for (FVertexElement& VertexElememt : Elements)
+						{
+							StringBuilder.Appendf(TEXT("\tStreamIndex: %2d - Offset: %4d - Type: %2d - AttributeIndex: %2d - Stride: %4d - UseInstanceIndex: %d\n"),
+												  VertexElememt.StreamIndex,
+												  VertexElememt.Offset,
+												  VertexElememt.Type.GetValue(),
+												  VertexElememt.AttributeIndex,
+												  VertexElememt.Stride,
+												  VertexElememt.bUseInstanceIndex);
+						}
+					}
+					StringBuilder.Append(TEXT("\n\n"));
+
+					FString Row = StringBuilder.ToString();
+					CSVFile->Serialize(TCHAR_TO_ANSI(*Row), Row.Len());
+				}
+			}
+			
+			delete CSVFile;
+			CSVFile = nullptr;
+		}
+	}
+
 #endif // PSO_PRECACHING_TRACKING
 
 private:
@@ -610,6 +750,50 @@ private:
 	TArray<FMaterialPSOPrecacheParams> MaterialPSORequests;	
 	TMap<FMaterialPSOPrecacheParams, FPrecacheData> MaterialPSORequestData;
 	uint32 LifecycleID = 0; //< ID to check current outstanding requests are still valid - incremented on re-precache of all current requests
+
+#if PSO_PRECACHING_TRACKING
+
+	void AddTrackedMaterialData(const FMaterialPSOPrecacheParams& InParams, const FPSOPrecacheDataArray& PrecacheDataArray)
+	{
+		for (const FPSOPrecacheData& PrecacheData : PrecacheDataArray)
+		{
+			// Track single entry for each PSO
+			uint64 PSOPrecacheHash = 0; 
+			if (PrecacheData.Type == FPSOPrecacheData::EType::Graphics)
+			{
+				PSOPrecacheHash = PSOCollectorStats::GetPSOPrecacheHash(PrecacheData.GraphicsPSOInitializer);
+			}
+			else
+			{
+				PSOPrecacheHash = PSOCollectorStats::GetPSOPrecacheHash(*PrecacheData.ComputeShader);
+			}
+
+			if (!TrackingMaterialPSORequestData.Contains(PSOPrecacheHash))
+			{				
+				FTrackingMaterialPSOPrecacheParams TrackingParams;
+				TrackingParams.MaterialName = InParams.Material->GetAssetName();
+				TrackingParams.VertexFactoryData = InParams.VertexFactoryData;
+				TrackingParams.PrecachePSOParams = InParams.PrecachePSOParams;
+				TrackingParams.PSOCollectorIndex = PrecacheData.PSOCollectorIndex;
+				TrackingParams.bGraphics = PrecacheData.Type == FPSOPrecacheData::EType::Graphics;
+				TrackingParams.bRequired = PrecacheData.bRequired;
+
+				TrackingMaterialPSORequestData.Add(PSOPrecacheHash, TrackingParams);
+			}
+		}
+	}
+	
+	struct FTrackingMaterialPSOPrecacheParams
+	{
+		FString MaterialName;
+		FPSOPrecacheVertexFactoryData VertexFactoryData;
+		FPSOPrecacheParams PrecachePSOParams;
+		int32 PSOCollectorIndex : 30;
+		bool bGraphics : 1;
+		bool bRequired : 1;
+	};
+	TMap<uint64, FTrackingMaterialPSOPrecacheParams> TrackingMaterialPSORequestData;
+#endif
 };
 
 // The global request manager - only used locally in a few global function to precache, release or boost PSO precache requests
@@ -851,6 +1035,19 @@ FPSOPrecacheDataArray GetMaterialPSOPrecacheData(FMaterialPSOPrecacheRequestID R
 	return GMaterialPSORequestManager.GetMaterialPSOPrecacheData(RequestID);
 }
 
+void DumpPSOPrecacheMaterialStats()
+{		
+	GMaterialPSORequestManager.DumpPSOPrecacheMaterialStats();
+}
+
+static FAutoConsoleCommand GDumpPSOPrecacheMaterialStatsCmd(
+	TEXT("r.PSOPrecache.DumpMaterialStats"),
+	TEXT("Dumps all currently collected material based PSO precache data to csv file"),
+	FConsoleCommandWithArgsDelegate::CreateStatic([](const TArray<FString>& Args)
+{		
+	DumpPSOPrecacheMaterialStats();
+}));
+
 #else
 
 FMaterialPSOPrecacheParams GetMaterialPSOPrecacheParams(FMaterialPSOPrecacheRequestID RequestID)
@@ -863,5 +1060,9 @@ FPSOPrecacheDataArray GetMaterialPSOPrecacheData(FMaterialPSOPrecacheRequestID R
 	return FPSOPrecacheDataArray();
 }
 
+void DumpPSOPrecacheMaterialStats()
+{
+	// Do nothing
+}
 
 #endif // PSO_PRECACHING_TRACKING
diff --git a/Engine/Source/Runtime/Engine/Public/PSOPrecacheMaterial.h b/Engine/Source/Runtime/Engine/Public/PSOPrecacheMaterial.h
--- a/Engine/Source/Runtime/Engine/Public/PSOPrecacheMaterial.h
+++ b/Engine/Source/Runtime/Engine/Public/PSOPrecacheMaterial.h
@@ -167,3 +167,9 @@ extern ENGINE_API FMaterialPSOPrecacheParams GetMaterialPSOPrecacheParams(FMater
  */
 extern ENGINE_API FPSOPrecacheDataArray GetMaterialPSOPrecacheData(FMaterialPSOPrecacheRequestID RequestID);
 
+/**
+ * Dump material PSO precache data to CSV file
+ */
+extern ENGINE_API void DumpPSOPrecacheMaterialStats();
+
+
diff --git a/Engine/Source/Runtime/RHI/Private/PipelineStateCache.cpp b/Engine/Source/Runtime/RHI/Private/PipelineStateCache.cpp
--- a/Engine/Source/Runtime/RHI/Private/PipelineStateCache.cpp
+++ b/Engine/Source/Runtime/RHI/Private/PipelineStateCache.cpp
@@ -271,17 +271,18 @@ struct FPSOCompilationDebugData
 #endif // WITH_RHI_BREADCRUMBS
 };
 
-static inline void CheckAndUpdateHitchCountStat(FPSOPrecacheRequestID::EType PSOType, bool bIsRuntimePSO, const FPSOCompilationDebugData& PSOCompilationDebugData, uint64 StartTime, EPSOPrecacheResult PSOPrecacheResult)
+static inline uint64 CheckAndUpdateHitchCountStat(FPSOPrecacheRequestID::EType PSOType, bool bIsRuntimePSO, const FPSOCompilationDebugData& PSOCompilationDebugData, uint64 StartTime, EPSOPrecacheResult PSOPrecacheResult)
 {
+	uint64 PSOCompilationTime = FPlatformTime::Cycles64() - StartTime;
 	if (!bIsRuntimePSO)
 	{
-		return;
+		return PSOCompilationTime;
 	}
 
 	RuntimePSOCreationCount.fetch_add(1, std::memory_order_relaxed);
 
 	int32 RuntimePSOCreationHitchThreshold = CVarPSORuntimeCreationHitchThreshold.GetValueOnAnyThread();
-	double PSOCreationTimeMs = FPlatformTime::ToMilliseconds64(FPlatformTime::Cycles64() - StartTime);
+	double PSOCreationTimeMs = FPlatformTime::ToMilliseconds64(PSOCompilationTime);
 	if (PSOCreationTimeMs > RuntimePSOCreationHitchThreshold)
 	{
 		if (PSOType == FPSOPrecacheRequestID::EType::Graphics)
@@ -335,6 +336,8 @@ static inline void CheckAndUpdateHitchCountStat(FPSOPrecacheRequestID::EType PSO
 				GraphicsPSOCreationHitchCount.load(std::memory_order_relaxed), ComputePSOCreationHitchCount.load(std::memory_order_relaxed), PrecachedPSOCreationHitchCount.load(std::memory_order_relaxed));
 		}
 	}
+
+	return PSOCompilationTime;
 }
 
 static int32 GPSOPrecompileThreadPoolSize = 0;
@@ -2579,6 +2582,17 @@ public:
 		return GetPrecachingStateInternal(InitializerHash, bMarkUsed, bOutMarkedUsed);
 	}
 
+	uint64 GetPSOCompilationTime(uint64 PSOPrecacheHash)
+	{
+#if !UE_BUILD_SHIPPING
+		FRWScopeLock ReadLock(PrecachePSOsRWLock, SLT_ReadOnly);
+		FPrecacheTask* FindResult = PrecachedPSOInitializerData.Find(PSOPrecacheHash);
+		return FindResult ? FindResult->PSOCompilationTime : 0;
+#else
+		return 0;
+#endif
+	}
+
 	// Retrieves the precaching state for the given PSO, and marks it as used if not already marked.
 	EPSOPrecacheResult GetPrecachingStateAndMarkUsed(const TPrecachedPSOInitializer& Initializer, bool& bOutMarkedUsed)
 	{
@@ -2620,7 +2634,7 @@ public:
 		}
 	}
 
-	void PrecacheFinished(const TPrecachedPSOInitializer& Initializer, bool bValid)
+	void PrecacheFinished(const TPrecachedPSOInitializer& Initializer, bool bValid, uint64 PSOCompilationTime)
 	{
 		uint64 InitializerHash = TPrecachePipelineCacheDerived::PipelineStateInitializerHash(Initializer);
 
@@ -2631,6 +2645,12 @@ public:
 			// Mark compiled (either succeeded or failed)
 			FPrecacheTask* FindResult = PrecachedPSOInitializerData.Find(InitializerHash);
 			check(FindResult);
+
+#if !UE_BUILD_SHIPPING
+			// Store the compilation time for stat tracking
+			FindResult->PSOCompilationTime = PSOCompilationTime;
+#endif
+
 			// We still add the 'compiling' bit here because if the task is fast enough, we can get here before the end of TryAddNewState
 			const EPSOPrecacheStateMask CompleteStateMask = bValid ? (EPSOPrecacheStateMask::Succeeded | EPSOPrecacheStateMask::Compiling) : (EPSOPrecacheStateMask::Failed | EPSOPrecacheStateMask::Compiling);
 			PreviousStateMask = FindResult->AddPSOPrecacheState(CompleteStateMask);
@@ -2821,6 +2841,9 @@ protected:
 	{
 		TPipelineState* PipelineState = nullptr;
 		FPSOPrecacheRequestID RequestID;
+#if !UE_BUILD_SHIPPING
+		uint64 PSOCompilationTime = 0;
+#endif
 
 		EPSOPrecacheStateMask AddPSOPrecacheState(EPSOPrecacheStateMask DesiredState)
 		{
@@ -3438,7 +3461,7 @@ public:
 
 		uint64 StartTime = FPlatformTime::Cycles64();
 		GfxPipeline->RHIPipeline = bSkipCreation ? nullptr : RHICreateGraphicsPipelineState(Initializer);
-		CheckAndUpdateHitchCountStat(FPSOPrecacheRequestID::EType::Graphics, !IsPrecachedPSO(Initializer), PSOCompilationDebugData, StartTime, PSOPrecacheResult);
+		uint64 PSOCompilationTime = CheckAndUpdateHitchCountStat(FPSOPrecacheRequestID::EType::Graphics, !IsPrecachedPSO(Initializer), PSOCompilationDebugData, StartTime, PSOPrecacheResult);
 
 		if (GfxPipeline->RHIPipeline)
 		{
@@ -3452,7 +3475,7 @@ public:
 		// Mark as finished when it's a precaching job
 		if (Initializer.bPSOPrecache)
 		{
-			GPrecacheGraphicsPipelineCache->PrecacheFinished(Initializer, GfxPipeline->RHIPipeline != nullptr);
+			GPrecacheGraphicsPipelineCache->PrecacheFinished(Initializer, GfxPipeline->RHIPipeline != nullptr, PSOCompilationTime);
 		}
 		else if (bPrecachedPSOMarkedUsed)
 		{
@@ -3565,7 +3588,7 @@ public:
 
 		// TODO: Send the priority to the RHI.
 		ComputePipeline->RHIPipeline = RHICreateComputePipelineState(Initializer);
-		CheckAndUpdateHitchCountStat(FPSOPrecacheRequestID::EType::Compute, !IsPrecachedPSO(Initializer), PSOCompilationDebugData, StartTime, PSOPrecacheResult);
+		uint64 PSOCompilationTime = CheckAndUpdateHitchCountStat(FPSOPrecacheRequestID::EType::Compute, !IsPrecachedPSO(Initializer), PSOCompilationDebugData, StartTime, PSOPrecacheResult);
 
 		if (!ComputePipeline->RHIPipeline)
 		{
@@ -3575,7 +3598,7 @@ public:
 		if (Initializer.bPSOPrecache)
 		{
 			bool bCSValid = ComputePipeline->RHIPipeline != nullptr && ComputePipeline->RHIPipeline->IsValid();
-			GPrecacheComputePipelineCache->PrecacheFinished(Initializer, bCSValid);
+			GPrecacheComputePipelineCache->PrecacheFinished(Initializer, bCSValid, PSOCompilationTime);
 		}
 		else if (bPrecachedPSOMarkedUsed)
 		{
@@ -3752,11 +3775,11 @@ static void InternalCreateComputePipelineState(const FComputePipelineStateInitia
 		check(GraphEvent == nullptr);
 		uint64 StartTime = FPlatformTime::Cycles64();
 		CachedState->RHIPipeline = RHICreateComputePipelineState(Initializer);
-		CheckAndUpdateHitchCountStat(FPSOPrecacheRequestID::EType::Compute, !IsPrecachedPSO(Initializer), PSOCompilationDebugData, StartTime, PSOPrecacheResult);
+		uint64 PSOCompilationTime = CheckAndUpdateHitchCountStat(FPSOPrecacheRequestID::EType::Compute, !IsPrecachedPSO(Initializer), PSOCompilationDebugData, StartTime, PSOPrecacheResult);
 
 		if (Initializer.bPSOPrecache)
 		{
-			GPrecacheComputePipelineCache->PrecacheFinished(Initializer, CachedState->RHIPipeline != nullptr);
+			GPrecacheComputePipelineCache->PrecacheFinished(Initializer, CachedState->RHIPipeline != nullptr, PSOCompilationTime);
 		}
 
 		if (!CachedState->RHIPipeline)
@@ -4237,11 +4260,11 @@ static void InternalCreateGraphicsPipelineState(const FGraphicsPipelineStateInit
 		check(GraphEvent == nullptr);
 		uint64 StartTime = FPlatformTime::Cycles64();
 		CachedState->RHIPipeline = RHICreateGraphicsPipelineState(Initializer);
-		CheckAndUpdateHitchCountStat(FPSOPrecacheRequestID::EType::Graphics, !IsPrecachedPSO(Initializer), PSOCompilationDebugData, StartTime, PSOPrecacheResult);
+		uint64 PSOCompilationTime = CheckAndUpdateHitchCountStat(FPSOPrecacheRequestID::EType::Graphics, !IsPrecachedPSO(Initializer), PSOCompilationDebugData, StartTime, PSOPrecacheResult);
 
 		if (Initializer.bPSOPrecache)
 		{
-			GPrecacheGraphicsPipelineCache->PrecacheFinished(Initializer, CachedState->RHIPipeline != nullptr);
+			GPrecacheGraphicsPipelineCache->PrecacheFinished(Initializer, CachedState->RHIPipeline != nullptr, PSOCompilationTime);
 		}
 
 		if (CachedState->RHIPipeline)
@@ -4575,6 +4598,16 @@ EPSOPrecacheResult PipelineStateCache::CheckPipelineStateInCache(FRHIComputeShad
 	return GPrecacheComputePipelineCache->GetPrecachingState(FComputePipelineStateInitializer(ComputeShader, 0));
 }
 
+uint64 PipelineStateCache::GetPrecachePSOCompilationTime(bool bGraphics, uint64 PSOPrecacheHash)
+{
+	if (!IsPSOPrecachingEnabled())
+	{
+		return 0;
+	}
+
+	return bGraphics ? GPrecacheGraphicsPipelineCache->GetPSOCompilationTime(PSOPrecacheHash) : GPrecacheComputePipelineCache->GetPSOCompilationTime(PSOPrecacheHash);
+}
+
 bool PipelineStateCache::IsPrecaching(const FPSOPrecacheRequestID& PSOPrecacheRequestID)
 {
 	if (!IsPSOPrecachingEnabled())
diff --git a/Engine/Source/Runtime/RHI/Public/PipelineStateCache.h b/Engine/Source/Runtime/RHI/Public/PipelineStateCache.h
--- a/Engine/Source/Runtime/RHI/Public/PipelineStateCache.h
+++ b/Engine/Source/Runtime/RHI/Public/PipelineStateCache.h
@@ -204,6 +204,9 @@ namespace PipelineStateCache
 
 	/* Retrieve the current PSO precache result state (slightly slower than IsPrecaching) */
 	extern RHI_API EPSOPrecacheResult		CheckPipelineStateInCache(FRHIComputeShader* ComputeShader);
+		
+	/* Retrieve the the time it took to compile this precache PSO */
+	extern RHI_API uint64					GetPrecachePSOCompilationTime(bool bGraphics, uint64 PSOPrecacheHash);
 
 	/* Is the given PSO still precaching? */
 	extern RHI_API bool						IsPrecaching(const FPSOPrecacheRequestID& PSOPrecacheRequestID);
