Build fix for when PSO_PRECACHING_VALIDATE is not enabled for build

diff --git a/Engine/Source/Runtime/Engine/Private/PSOPrecacheMaterial.cpp b/Engine/Source/Runtime/Engine/Private/PSOPrecacheMaterial.cpp
--- a/Engine/Source/Runtime/Engine/Private/PSOPrecacheMaterial.cpp
+++ b/Engine/Source/Runtime/Engine/Private/PSOPrecacheMaterial.cpp
@@ -377,8 +377,10 @@ public:
 		// 	     but need to be sure that every material name is unique then.
 		// 		 Currently checked by pointer of material and not just name - but pointer can't be stored for tracked because material could be deleted already
 		FindResult->PSOPrecachaData = PrecacheData;
+#if PSO_PRECACHING_VALIDATE
 		AddTrackedMaterialData(Params, PrecacheData);
-#endif
+#endif // PSO_PRECACHING_VALIDATE
+#endif // PSO_PRECACHING_TRACKING
 
 		// update the state
 		FindResult->State = FindResult->ActivePSOPrecacheRequests.IsEmpty() ? EState::Completed : EState::Compiling;
@@ -541,6 +543,7 @@ public:
 
 	void DumpPSOPrecacheMaterialStats()
 	{		
+#if PSO_PRECACHING_VALIDATE
 		FRWScopeLock ReadLock(RWLock, SLT_ReadOnly);
 
 		EShadingPath ShadingPath = GetFeatureLevelShadingPath(GMaxRHIFeatureLevel);
@@ -673,6 +676,7 @@ public:
 			delete CSVFile;
 			CSVFile = nullptr;
 		}
+#endif // PSO_PRECACHING_VALIDATE
 	}
 
 #endif // PSO_PRECACHING_TRACKING
@@ -751,7 +755,7 @@ private:
 	TMap<FMaterialPSOPrecacheParams, FPrecacheData> MaterialPSORequestData;
 	uint32 LifecycleID = 0; //< ID to check current outstanding requests are still valid - incremented on re-precache of all current requests
 
-#if PSO_PRECACHING_TRACKING
+#if PSO_PRECACHING_TRACKING && PSO_PRECACHING_VALIDATE
 
 	void AddTrackedMaterialData(const FMaterialPSOPrecacheParams& InParams, const FPSOPrecacheDataArray& PrecacheDataArray)
 	{
@@ -793,7 +797,7 @@ private:
 		bool bRequired : 1;
 	};
 	TMap<uint64, FTrackingMaterialPSOPrecacheParams> TrackingMaterialPSORequestData;
-#endif
+#endif // PSO_PRECACHING_TRACKING && PSO_PRECACHING_VALIDATE
 };
 
 // The global request manager - only used locally in a few global function to precache, release or boost PSO precache requests
