Added SetDescribe method to FTestWorlds be able to provide better description of unittests for NetTraces
Fixed iris natvis.

diff --git a/Engine/Source/Runtime/Engine/Private/Tests/NetTestHelpers.cpp b/Engine/Source/Runtime/Engine/Private/Tests/NetTestHelpers.cpp
--- a/Engine/Source/Runtime/Engine/Private/Tests/NetTestHelpers.cpp
+++ b/Engine/Source/Runtime/Engine/Private/Tests/NetTestHelpers.cpp
@@ -17,6 +17,7 @@
 #include "HAL/PlatformMath.h"
 #include "Kismet/GameplayStatics.h"
 #include "Net/UnrealNetwork.h"
+#include "Net/Core/Trace/NetTrace.h"
 
 #if WITH_EDITOR
 #include "Settings/LevelEditorPlaySettings.h"
@@ -389,11 +390,31 @@ bool FTestWorlds::WaitForClientConnect(FTestWorldInstance& Client)
 	if (bConnected)
 	{
 		OnClientConnected.Broadcast(Client);
+		UE_NET_TRACE_UPDATE_INSTANCE(Client.GetNetDriver()->GetNetTraceId(), false, *Description);
 	}
 
 	return bConnected;
 }
 
+void FTestWorlds::SetDescription(const FString& InDescription)
+{
+	Description = InDescription;
+
+#if UE_NET_TRACE_ENABLED
+	if (Server.GetNetDriver())
+	{
+		UE_NET_TRACE_UPDATE_INSTANCE(Server.GetNetDriver()->GetNetTraceId(), true, *Description);
+	}
+	for (FTestWorldInstance& ClientIt : Clients)
+	{
+		if (ClientIt.GetNetDriver())
+		{
+			UE_NET_TRACE_UPDATE_INSTANCE(ClientIt.GetNetDriver()->GetNetTraceId(), false, *Description);
+		}
+	}
+#endif
+}
+
 void FTestWorlds::TickAll(int32 NumTicks)
 {
 	PreTickAllDelegate.Broadcast();
@@ -585,7 +606,7 @@ UObject* FTestWorlds::FindReplicatedObjectOnClient(UObject* ServerObject, uint32
 			return nullptr;
 		}
 
-		return ClientContext.IrisRepBridge->GetReplicatedObject(NetHandle);
+		return ClientContext.IrisRepBridge ? ClientContext.IrisRepBridge->GetReplicatedObject(NetHandle) : nullptr;
 	}
 	else
 	{
diff --git a/Engine/Source/Runtime/Engine/Public/Tests/NetTestHelpers.h b/Engine/Source/Runtime/Engine/Public/Tests/NetTestHelpers.h
--- a/Engine/Source/Runtime/Engine/Public/Tests/NetTestHelpers.h
+++ b/Engine/Source/Runtime/Engine/Public/Tests/NetTestHelpers.h
@@ -199,6 +199,9 @@ struct FTestWorlds
 
 	ENGINE_API ~FTestWorlds();
 
+	ENGINE_API void SetDescription(const FString& Description);
+	FString GetDescription() const { return Description; }
+
 	ENGINE_API void SetTickInSeconds(float TickInSeconds);
 	float GetTickDeltaSeconds() const { return TickDeltaSeconds; }
 	
@@ -296,6 +299,9 @@ private:
 
 	// Restore GWorld and other PIE settings
 	FScopedNetTestPIERestoration PIERestoration;
+
+	// Testworld description.
+	FString Description;
 };
 
 //------------------------------------------------------------------------
