Remove r.Mobile.SupportsSM5MaterialNodes usage from a material HLSL translator.
Always allow SM5 features that can be supported by current configuration of a mobile renderer

diff --git a/Engine/Source/Runtime/Engine/Private/Materials/HLSLMaterialTranslator.cpp b/Engine/Source/Runtime/Engine/Private/Materials/HLSLMaterialTranslator.cpp
--- a/Engine/Source/Runtime/Engine/Private/Materials/HLSLMaterialTranslator.cpp
+++ b/Engine/Source/Runtime/Engine/Private/Materials/HLSLMaterialTranslator.cpp
@@ -8487,7 +8487,8 @@ int32 FHLSLMaterialTranslator::GetScreenAlignedUV(int32 Offset, int32 ViewportUV
 
 int32 FHLSLMaterialTranslator::SceneDepth(int32 Offset, int32 ViewportUV, bool bUseOffset)
 {
-	if (ShaderFrequency == SF_Vertex && IsMobilePlatform(Platform) && !MobileSupportsSM5MaterialNodes(Platform))
+	if (ShaderFrequency == SF_Vertex && IsMobilePlatform(Platform) 
+		&& (MobileAllowFramebufferFetch(Platform) && !MobileUsesFullDepthPrepass(Platform)))
 	{
 		// mobile currently does not support this, we need to read a separate copy of the depth, we must disable framebuffer fetch and force scene texture reads.
 		return Errorf(TEXT("Cannot read scene depth from the vertex shader with the Mobile feature level"));
@@ -8633,12 +8634,9 @@ void FHLSLMaterialTranslator::UseSceneTextureId(ESceneTextureId SceneTextureId,
 			Errorf(TEXT("Decals can only access SceneDepth, CustomDepth, CustomStencil, and WorldNormal."));
 		}
 
-		const bool bSceneTextureRequiresSM5 = SceneTextureId == PPI_WorldNormal;
-		if (bSceneTextureRequiresSM5)
+		if (SceneTextureId == PPI_WorldNormal && IsMobilePlatform(Platform) && !IsMobileDeferredShadingEnabled(Platform))
 		{
-			// Minimum requirement for mobile supporting SM5 nodes is ES3_1.
-			ERHIFeatureLevel::Type RequiredFeatureLevel = MobileSupportsSM5MaterialNodes(Platform) ? ERHIFeatureLevel::ES3_1 : ERHIFeatureLevel::SM5;
-			ErrorUnlessFeatureLevelSupported(RequiredFeatureLevel);
+			Errorf(TEXT("Decals can only access WorldNormal when deferred shading is used."));
 		}
 
 		if (SceneTextureId == PPI_WorldNormal && Material->HasNormalConnected() && !IsUsingDBuffers(Platform))
@@ -8732,9 +8730,9 @@ int32 FHLSLMaterialTranslator::SceneColor(int32 Offset, int32 ViewportUV, bool b
 		Errorf(TEXT("SceneColor lookups are only available when MaterialDomain = Surface."));
 	}
 
-	if (IsMobilePlatform(Platform) && !MobileSupportsSM5MaterialNodes(Platform))
+	if (IsMobilePlatform(Platform) && !IsMobileDeferredShadingEnabled(Platform))
 	{
-		Errorf(TEXT("Cannot read SceneColor with mobile shading (platform id %d)."), Platform);
+		Errorf(TEXT("SceneColor lookups are only available when deferred shading is used"));
 	}
 
 	MaterialCompilationOutput.SetIsSceneTextureUsed(PPI_SceneColor);
@@ -12568,8 +12566,9 @@ int32 FHLSLMaterialTranslator::DistanceFieldApproxAO(int32 PositionArg, EPositio
 
 int32 FHLSLMaterialTranslator::SamplePhysicsField(int32 PositionArg, EPositionOrigin PositionOrigin, const int32 OutputType, const int32 TargetIndex)
 {
-	if (FeatureLevel == ERHIFeatureLevel::ES3_1 && MobileSupportsSM5MaterialNodes(Platform))
+	if (FeatureLevel == ERHIFeatureLevel::ES3_1)
 	{
+		// TODO: add physics field samlping to mobile
 		if (OutputType == EFieldOutputType::Field_Output_Vector)
 		{
 			return Constant3(0.0f, 0.0f, 0.0f);
