PR #12404: Fix light actors inheritance

diff --git a/Engine/Source/Runtime/Engine/Classes/Engine/DirectionalLight.h b/Engine/Source/Runtime/Engine/Classes/Engine/DirectionalLight.h
--- a/Engine/Source/Runtime/Engine/Classes/Engine/DirectionalLight.h
+++ b/Engine/Source/Runtime/Engine/Classes/Engine/DirectionalLight.h
@@ -33,10 +33,10 @@ private:
 public:
 
 	// UObject Interface
-	virtual void PostLoad() override;
+	ENGINE_API virtual void PostLoad() override;
 #if WITH_EDITOR
-	virtual void LoadedFromAnotherClass(const FName& OldClassName) override;
-	virtual void PostEditChangeProperty(FPropertyChangedEvent& PropertyChangedEvent) override;
+	ENGINE_API virtual void LoadedFromAnotherClass(const FName& OldClassName) override;
+	ENGINE_API virtual void PostEditChangeProperty(FPropertyChangedEvent& PropertyChangedEvent) override;
 	virtual bool CanChangeIsSpatiallyLoadedFlag() const override { return false; }
 	virtual void GetStreamingBounds(FBox& OutRuntimeBounds, FBox& OutEditorBounds) const override { OutRuntimeBounds.Init(); OutEditorBounds.Init(); }
 #endif
diff --git a/Engine/Source/Runtime/Engine/Classes/Engine/PointLight.h b/Engine/Source/Runtime/Engine/Classes/Engine/PointLight.h
--- a/Engine/Source/Runtime/Engine/Classes/Engine/PointLight.h
+++ b/Engine/Source/Runtime/Engine/Classes/Engine/PointLight.h
@@ -24,13 +24,13 @@ class APointLight : public ALight
 
 #if WITH_EDITOR
 	//~ Begin AActor Interface.
-	virtual void EditorApplyScale(const FVector& DeltaScale, const FVector* PivotLocation, bool bAltDown, bool bShiftDown, bool bCtrlDown) override;
+	ENGINE_API virtual void EditorApplyScale(const FVector& DeltaScale, const FVector* PivotLocation, bool bAltDown, bool bShiftDown, bool bCtrlDown) override;
 	//~ End AActor Interface.
 #endif
 
 	//~ Begin UObject Interface.
 #if WITH_EDITOR
-	virtual void LoadedFromAnotherClass(const FName& OldClassName) override;
+	ENGINE_API virtual void LoadedFromAnotherClass(const FName& OldClassName) override;
 #endif
 	//~ End UObject Interface.
 };
\ No newline at end of file
diff --git a/Engine/Source/Runtime/Engine/Classes/Engine/RectLight.h b/Engine/Source/Runtime/Engine/Classes/Engine/RectLight.h
--- a/Engine/Source/Runtime/Engine/Classes/Engine/RectLight.h
+++ b/Engine/Source/Runtime/Engine/Classes/Engine/RectLight.h
@@ -17,7 +17,7 @@ class ARectLight : public ALight
 
 #if WITH_EDITOR
 	//~ Begin AActor Interface.
-	virtual void EditorApplyScale(const FVector& DeltaScale, const FVector* PivotLocation, bool bAltDown, bool bShiftDown, bool bCtrlDown) override;
+	ENGINE_API virtual void EditorApplyScale(const FVector& DeltaScale, const FVector* PivotLocation, bool bAltDown, bool bShiftDown, bool bCtrlDown) override;
 	//~ End AActor Interface.
 #endif
 };
\ No newline at end of file
diff --git a/Engine/Source/Runtime/Engine/Classes/Engine/SpotLight.h b/Engine/Source/Runtime/Engine/Classes/Engine/SpotLight.h
--- a/Engine/Source/Runtime/Engine/Classes/Engine/SpotLight.h
+++ b/Engine/Source/Runtime/Engine/Classes/Engine/SpotLight.h
@@ -36,15 +36,15 @@ public:
 
 #if WITH_EDITOR
 	//~ Begin AActor Interface.
-	virtual void EditorApplyScale(const FVector& DeltaScale, const FVector* PivotLocation, bool bAltDown, bool bShiftDown, bool bCtrlDown) override;
+	ENGINE_API virtual void EditorApplyScale(const FVector& DeltaScale, const FVector* PivotLocation, bool bAltDown, bool bShiftDown, bool bCtrlDown) override;
 	//~ End AActor Interface.
 #endif
 
 	//~ Begin UObject Interface
-	virtual void PostLoad() override;
+	ENGINE_API virtual void PostLoad() override;
 #if WITH_EDITOR
-	virtual void LoadedFromAnotherClass(const FName& OldClassName) override;
-	virtual void PostEditChangeProperty(FPropertyChangedEvent& PropertyChangedEvent) override;
+	ENGINE_API virtual void LoadedFromAnotherClass(const FName& OldClassName) override;
+	ENGINE_API virtual void PostEditChangeProperty(FPropertyChangedEvent& PropertyChangedEvent) override;
 #endif
 	//~ End UObject Interface
 
