PR #12426: Fixed incorrect DeltaTime values after long gameplay due to floating-point precision issues.

diff --git a/Engine/Source/Runtime/Engine/Classes/Engine/EngineBaseTypes.h b/Engine/Source/Runtime/Engine/Classes/Engine/EngineBaseTypes.h
--- a/Engine/Source/Runtime/Engine/Classes/Engine/EngineBaseTypes.h
+++ b/Engine/Source/Runtime/Engine/Classes/Engine/EngineBaseTypes.h
@@ -309,7 +309,7 @@ private:
 		 * The last world game time at which we were ticked. Game time used is dependent on bTickEvenWhenPaused
 		 * Valid only if we've been ticked at least once since having a tick interval; otherwise set to -1.f
 		 **/
-		float LastTickGameTimeSeconds;
+		double LastTickGameTimeSeconds;
 
 		/** Back pointer to the FTickTaskLevel containing this tick function if it is registered **/
 		class FTickTaskLevel*						TickTaskLevel;
@@ -421,8 +421,8 @@ public:
 	 * For actively registered tick functions with an interval, return the last time it ticked.
 	 * For unregistered or non-interval functions it will return -1;
 	 */
-	float GetLastTickGameTime() const { return (InternalData ? InternalData->LastTickGameTimeSeconds : -1.f); }
-	 
+	double GetLastTickGameTime() const { return (InternalData ? InternalData->LastTickGameTimeSeconds : -1.0); }
+
 	/** 
 	 * Correctly call ExecuteTick on a nested function, and will set the completion event properly for the duration of the tick function
 	 * @param DeltaTime - frame time to advance, in seconds
diff --git a/Engine/Source/Runtime/Engine/Private/TickTaskManager.cpp b/Engine/Source/Runtime/Engine/Private/TickTaskManager.cpp
--- a/Engine/Source/Runtime/Engine/Private/TickTaskManager.cpp
+++ b/Engine/Source/Runtime/Engine/Private/TickTaskManager.cpp
@@ -2387,7 +2387,7 @@ FTickFunction::FInternalData::FInternalData()
 	, TaskPointer(nullptr)
 	, Next(nullptr)
 	, RelativeTickCooldown(0.f)
-	, LastTickGameTimeSeconds(-1.f)
+	, LastTickGameTimeSeconds(-1.0)
 	, TickTaskLevel(nullptr)
 {
 }
@@ -2451,7 +2451,7 @@ void FTickFunction::SetTickFunctionEnable(bool bInEnabled)
 
 		if (TickState == ETickState::Disabled)
 		{
-			InternalData->LastTickGameTimeSeconds = -1.f;
+			InternalData->LastTickGameTimeSeconds = -1.0;
 		}
 	}
 	else
@@ -2828,14 +2828,14 @@ float FTickFunction::CalculateDeltaTime(float DeltaTime, const UWorld* TickingWo
 	{
 		// No tick interval. Return the world delta seconds, and make sure to mark that
 		// we're not tracking last-tick-time for this object.
-		InternalData->LastTickGameTimeSeconds = -1.f;
+		InternalData->LastTickGameTimeSeconds = -1.0;
 	}
 	else
 	{
 		// We've got a tick interval. Mark last-tick-time. If we already had last-tick-time, return
 		// the time since then; otherwise, return the world delta seconds.
-		const float CurrentWorldTime = (bTickEvenWhenPaused ? TickingWorld->GetUnpausedTimeSeconds() : TickingWorld->GetTimeSeconds());
-		if (InternalData->LastTickGameTimeSeconds >= 0.f)
+		const double CurrentWorldTime = (bTickEvenWhenPaused ? TickingWorld->GetUnpausedTimeSeconds() : TickingWorld->GetTimeSeconds());
+		if (InternalData->LastTickGameTimeSeconds >= 0.0)
 		{
 			DeltaTime = CurrentWorldTime - InternalData->LastTickGameTimeSeconds;
 		}
