Referenced Default Textures Optimization
Ensure default material textures which are not used for rendering are not loaded at runtime.
The way this is done is by removing the dependency to the array which contains all the default textures.
That member is now moved to Editor Only Data and a new member is added to the Material Interface which holds only the textures that specific material uses, if it is a Material or a Material Instance with static parameters.
Material Instances without static parameters get their textures from their parents.
Everything else is built to support this by leveraging better encapsulation.

RemoveUnusedDefaultTextures() contains the functionality to remove the unused textures.
This is done at cook time, based on the Shader Compilation output for the platform for which the cook is being done.

GetReferencedTextures() in FMaterial is now deprecated to allow for that better encapsulation of the actual Textures array and specific functions were added for each specific request.

GetDefaultTextureIdx(), GetDefaultTextureAtIdx(), GetDefaultTextures() have been added to the UMaterialInterface for this.
The GetMat variants of these methods have been added to the FMaterial, which in the end call the UMaterialInterface methods.

GetDefaultTextures and GetMatDefaultTextures has been added as a variant which replicates old behavior by returning an array with all default textures. This is used only in editor behaviors,

Common properties like CachedMaterialResourcesForCooking have been moved from UMaterial and UMaterialInstance to UMaterialInterface.

Common functionalities between UMaterial and UMaterialInstance that could be refactored out to clean out the code and help with this task, have been moved to UMaterialInterface, like:
GetMaterialResources, GetTextureExpressionValues, BuildQualityLevelRange, BuildShaderPlatformRange.

[RN] medium | Rendering Architecture

diff --git a/Engine/Source/Developer/MaterialBaking/Private/ExportMaterialProxy.h b/Engine/Source/Developer/MaterialBaking/Private/ExportMaterialProxy.h
--- a/Engine/Source/Developer/MaterialBaking/Private/ExportMaterialProxy.h
+++ b/Engine/Source/Developer/MaterialBaking/Private/ExportMaterialProxy.h
@@ -225,7 +225,7 @@ public:
 	{
 		SetQualityLevelProperties(GMaxRHIShaderPlatform);
 		Material = InMaterialInterface->GetMaterial();
-		ReferencedTextures = InMaterialInterface->GetReferencedTextures();
+		ReferencedTextures = InMaterialInterface->GetDefaultTextures();
 		ReferencedTextureCollections = InMaterialInterface->GetReferencedTextureCollections();
 
 		const FMaterialResource* Resource = InMaterialInterface->GetMaterialResource(GMaxRHIShaderPlatform);
@@ -305,6 +305,26 @@ public:
 		return ReferencedTextures;
 	}
 
+	virtual const TArray<TObjectPtr<UObject>>& GetMatDefaultTextures() const override
+	{
+		return ReferencedTextures;
+	}
+
+	virtual int32 GetMatDefaultTextureIdx(TObjectPtr<UObject> InTexture) const override
+	{
+		return ReferencedTextures.Find(InTexture);
+	}
+
+	virtual TObjectPtr<UObject> GetMatDefaultTextureAtIdx(const int32 InIdx) const override
+	{
+		if (!ensureAlwaysMsgf(ReferencedTextures.IsValidIndex(InIdx), TEXT("Invalid texture Idx %d requested from Array with %d Elements for Material: %s"), InIdx, ReferencedTextures.Num(), *MaterialInterface->GetPathName()))
+		{
+			return nullptr;
+		}
+
+		return ReferencedTextures[InIdx];
+	}
+
 	virtual TConstArrayView<TObjectPtr<UTextureCollection>> GetReferencedTextureCollections() const override
 	{
 		return ReferencedTextureCollections;
diff --git a/Engine/Source/Developer/MaterialUtilities/Private/MaterialUtilities.cpp b/Engine/Source/Developer/MaterialUtilities/Private/MaterialUtilities.cpp
--- a/Engine/Source/Developer/MaterialUtilities/Private/MaterialUtilities.cpp
+++ b/Engine/Source/Developer/MaterialUtilities/Private/MaterialUtilities.cpp
@@ -472,7 +472,7 @@ public:
 	{
 		SetQualityLevelProperties(GMaxRHIShaderPlatform);
 		Material = InMaterialInterface->GetMaterial();
-		ReferencedTextures = InMaterialInterface->GetReferencedTextures();
+		ReferencedTextures = InMaterialInterface->GetDefaultTextures();
 		FPlatformMisc::CreateGuid(Id);
 
 		FMaterialResource* Resource = InMaterialInterface->GetMaterialResource(GMaxRHIShaderPlatform);
@@ -539,6 +539,27 @@ public:
 		return ReferencedTextures;
 	}
 
+	virtual const TArray<TObjectPtr<UObject>>& GetMatDefaultTextures() const override
+	{
+		return ReferencedTextures;
+	}
+
+	virtual int32 GetMatDefaultTextureIdx(TObjectPtr<UObject> InTexture) const override
+	{
+		return ReferencedTextures.Find(InTexture);
+	}
+
+	virtual TObjectPtr<UObject> GetMatDefaultTextureAtIdx(const int32 InIdx) const override
+	{
+		if (!ensureAlwaysMsgf(ReferencedTextures.IsValidIndex(InIdx), TEXT("Invalid texture Idx %d requested from Array with %d Elements."), InIdx, ReferencedTextures.Num()))
+		{
+			return nullptr;
+		}
+
+		return ReferencedTextures[InIdx];
+	}
+
+
 	////////////////
 	// FMaterialRenderProxy interface.
 	virtual const FMaterial* GetMaterialNoFallback(ERHIFeatureLevel::Type InFeatureLevel) const override
diff --git a/Engine/Source/Editor/MaterialEditor/Private/MaterialEditor.cpp b/Engine/Source/Editor/MaterialEditor/Private/MaterialEditor.cpp
--- a/Engine/Source/Editor/MaterialEditor/Private/MaterialEditor.cpp
+++ b/Engine/Source/Editor/MaterialEditor/Private/MaterialEditor.cpp
@@ -269,7 +269,7 @@ FMatExpressionPreview::FMatExpressionPreview(UMaterialExpression* InExpression)
 	SetQualityLevelProperties(GMaxRHIShaderPlatform);
 
 	UMaterial* BaseMaterial = InExpression->Material;
-	ReferencedTextures = InExpression->Material->GetReferencedTextures();
+	ReferencedTextures = InExpression->Material->GetDefaultTextures();
 	ReferencedTextureCollections = InExpression->Material->GetReferencedTextureCollections();
 }
 
@@ -281,10 +281,6 @@ void FMatExpressionPreview::AddReferencedObjects(FReferenceCollector& Collector)
 {
 	Collector.AddReferencedObjects(ReferencedTextures);
 	Collector.AddReferencedObjects(ReferencedTextureCollections);
-	if (CachedExpressionData)
-	{
-		CachedExpressionData->AddReferencedObjects(Collector);
-	}
 }
 
 bool FMatExpressionPreview::ShouldCache(const FShaderType* ShaderType, const FVertexFactoryType* VertexFactoryType) const
@@ -420,26 +416,33 @@ bool FMatExpressionPreview::IsUsingNewHLSLGenerator() const
 	return false;
 }
 
-TArrayView<const TObjectPtr<UObject>> FMatExpressionPreview::GetReferencedTextures() const
+const TArray<TObjectPtr<UObject>>& FMatExpressionPreview::GetMatDefaultTextures() const
 {
-	if (CachedExpressionData)
-	{
-		// Path for new HLSL translator
-		return MakeArrayView(CachedExpressionData->ReferencedTextures);
-	}
+	return ReferencedTextures;
+}
 
-	// Legacy path
-	return MakeArrayView(ReferencedTextures);
+int32 FMatExpressionPreview::GetMatDefaultTextureIdx(TObjectPtr<UObject> InTexture) const
+{
+	return ReferencedTextures.Find(InTexture);
 }
 
-TConstArrayView<TObjectPtr<UTextureCollection>> FMatExpressionPreview::GetReferencedTextureCollections() const
+TObjectPtr<UObject> FMatExpressionPreview::GetMatDefaultTextureAtIdx(const int32 InIdx) const
 {
-	if (CachedExpressionData)
+	if (!ensureAlwaysMsgf(ReferencedTextures.IsValidIndex(InIdx), TEXT("Invalid texture Idx %d requested from Array with %d Elements."), InIdx, ReferencedTextures.Num()))
 	{
-		// Path for new HLSL translator
-		return MakeArrayView(CachedExpressionData->ReferencedTextureCollections);
+		return nullptr;
 	}
 
+	return ReferencedTextures[InIdx];
+}
+
+TArrayView<const TObjectPtr<UObject>> FMatExpressionPreview::GetReferencedTextures() const
+{
+	return ReferencedTextures;
+}
+
+TConstArrayView<TObjectPtr<UTextureCollection>> FMatExpressionPreview::GetReferencedTextureCollections() const
+{
 	// Legacy path
 	return MakeArrayView(ReferencedTextureCollections);
 }
diff --git a/Engine/Source/Editor/MaterialEditor/Private/MaterialEditor.h b/Engine/Source/Editor/MaterialEditor/Private/MaterialEditor.h
--- a/Engine/Source/Editor/MaterialEditor/Private/MaterialEditor.h
+++ b/Engine/Source/Editor/MaterialEditor/Private/MaterialEditor.h
@@ -84,6 +84,10 @@ public:
 	 * @return true if the shader should be compiled
 	 */
 	virtual bool ShouldCache(const FShaderType* ShaderType, const FVertexFactoryType* VertexFactoryType) const override;
+	
+	virtual const TArray<TObjectPtr<UObject>>& GetMatDefaultTextures() const override;
+	virtual int32 GetMatDefaultTextureIdx(TObjectPtr<UObject> InTexture) const override;
+	virtual TObjectPtr<UObject> GetMatDefaultTextureAtIdx(const int32 InIdx) const override;
 
 	virtual TArrayView<const TObjectPtr<UObject>> GetReferencedTextures() const override;
 	virtual TConstArrayView<TObjectPtr<UTextureCollection>> GetReferencedTextureCollections() const override;
@@ -178,7 +182,6 @@ public:
 	float UnrelatedNodesOpacity;
 
 private:
-	TUniquePtr<FMaterialCachedExpressionData> CachedExpressionData;
 	TWeakObjectPtr<UMaterialExpression> Expression;
 	TArray<TObjectPtr<UObject>> ReferencedTextures;
 	TArray<TObjectPtr<UTextureCollection>> ReferencedTextureCollections;
diff --git a/Engine/Source/Editor/UnrealEd/Private/Lightmass/LightmassRender.cpp b/Engine/Source/Editor/UnrealEd/Private/Lightmass/LightmassRender.cpp
--- a/Engine/Source/Editor/UnrealEd/Private/Lightmass/LightmassRender.cpp
+++ b/Engine/Source/Editor/UnrealEd/Private/Lightmass/LightmassRender.cpp
@@ -287,6 +287,28 @@ public:
 		}
 	}
 
+#if WITH_EDITOR
+	virtual const TArray<TObjectPtr<UObject>>& GetMatDefaultTextures() const override
+	{
+		return ReferencedTextures;
+	}
+
+	virtual int32 GetMatDefaultTextureIdx(TObjectPtr<UObject> InTexture) const override
+	{
+		return ReferencedTextures.Find(InTexture);
+	}
+#endif
+
+	virtual TObjectPtr<UObject> GetMatDefaultTextureAtIdx(const int32 InIdx) const override
+	{
+		if (!ensureAlwaysMsgf(ReferencedTextures.IsValidIndex(InIdx), TEXT("Invalid texture Idx %d requested from Array with %d Elements."), InIdx, ReferencedTextures.Num()))
+		{
+			return nullptr;
+		}
+
+		return ReferencedTextures[InIdx];
+	}
+
 	virtual TArrayView<const TObjectPtr<UObject>> GetReferencedTextures() const override
 	{
 		return ReferencedTextures;
diff --git a/Engine/Source/Runtime/Engine/Private/Materials/HLSLMaterialTranslator.cpp b/Engine/Source/Runtime/Engine/Private/Materials/HLSLMaterialTranslator.cpp
--- a/Engine/Source/Runtime/Engine/Private/Materials/HLSLMaterialTranslator.cpp
+++ b/Engine/Source/Runtime/Engine/Private/Materials/HLSLMaterialTranslator.cpp
@@ -7625,7 +7625,7 @@ int32 FHLSLMaterialTranslator::TextureSample(
 		const FMaterialUniformExpressionTexture* TextureExpression = Expression ? Expression->GetTextureUniformExpression() : nullptr;
 		if (ensure(TextureExpression))
 		{
-			Texture = Cast<UTexture>(Material->GetReferencedTextures()[TextureExpression->GetTextureIndex()]);
+			Texture = Cast<UTexture>(Material->GetMatDefaultTextureAtIdx(TextureExpression->GetTextureIndex()));
 		}
 	}
 
@@ -8901,13 +8901,13 @@ int32 FHLSLMaterialTranslator::Switch(int32 SwitchValueInput, int32 DefaultInput
 int32 FHLSLMaterialTranslator::Texture(UTexture* InTexture, int32& TextureReferenceIndex, EMaterialSamplerType SamplerType, ESamplerSourceMode SamplerSource, ETextureMipValueMode MipValueMode)
 {
 	EMaterialValueType ShaderType = InTexture->GetMaterialType();
-	TextureReferenceIndex = Material->GetReferencedTextures().Find(InTexture);
+	TextureReferenceIndex = Material->GetMatDefaultTextureIdx(InTexture);
 
 #if DO_CHECK
 	// UE-3518: Additional pre-assert logging to help determine the cause of this failure.
 	if (TextureReferenceIndex == INDEX_NONE)
 	{
-		const auto ReferencedTextures = Material->GetReferencedTextures();
+		const TArray<TObjectPtr<UObject>>& ReferencedTextures = Material->GetMatDefaultTextures();
 		UE_LOG(LogMaterial, Error, TEXT("Compiler->Texture() failed to find texture '%s' in referenced list of size '%i':"), *InTexture->GetName(), ReferencedTextures.Num());
 		for (int32 i = 0; i < ReferencedTextures.Num(); ++i)
 		{
@@ -8940,7 +8940,7 @@ int32 FHLSLMaterialTranslator::TextureParameter(FName ParameterName, UTexture* I
 	}
 
 	EMaterialValueType ShaderType = DefaultValue->GetMaterialType();
-	TextureReferenceIndex = Material->GetReferencedTextures().Find(DefaultValue);
+	TextureReferenceIndex = Material->GetMatDefaultTextureIdx(DefaultValue);
 	if (TextureReferenceIndex == INDEX_NONE)
 	{
 		FString TextureName;
@@ -9051,7 +9051,7 @@ int32 FHLSLMaterialTranslator::VirtualTexture(URuntimeVirtualTexture* InTexture,
 		return INDEX_NONE;
 	}
 
-	TextureReferenceIndex = Material->GetReferencedTextures().Find(InTexture);
+	TextureReferenceIndex = Material->GetMatDefaultTextureIdx(InTexture);
 	checkf(TextureReferenceIndex != INDEX_NONE, TEXT("Material expression called Compiler->VirtualTexture() without implementing UMaterialExpression::GetReferencedTexture properly"));
 
 	FAddUniformExpressionScope Scope(this);
@@ -9074,7 +9074,7 @@ int32 FHLSLMaterialTranslator::VirtualTextureParameter(FName ParameterName, URun
 		DefaultValue = Meta.Value.RuntimeVirtualTexture;
 	}
 
-	TextureReferenceIndex = Material->GetReferencedTextures().Find(DefaultValue);
+	TextureReferenceIndex = Material->GetMatDefaultTextureIdx(DefaultValue);
 	checkf(TextureReferenceIndex != INDEX_NONE, TEXT("Material expression called Compiler->VirtualTexture() without implementing UMaterialExpression::GetReferencedTexture properly"));
 
 	FMaterialParameterInfo ParameterInfo = GetParameterAssociationInfo();
@@ -9219,7 +9219,7 @@ int32 FHLSLMaterialTranslator::ExternalTexture(UTexture* InTexture, int32& Textu
 		return NonPixelShaderExpressionError();
 	}
 
-	TextureReferenceIndex = Material->GetReferencedTextures().Find(InTexture);
+	TextureReferenceIndex = Material->GetMatDefaultTextureIdx(InTexture);
 	checkf(TextureReferenceIndex != INDEX_NONE, TEXT("Material expression called Compiler->ExternalTexture() without implementing UMaterialExpression::GetReferencedTexture properly"));
 
 	FAddUniformExpressionScope Scope(this);
@@ -9235,7 +9235,7 @@ int32 FHLSLMaterialTranslator::ExternalTextureParameter(FName ParameterName, UTe
 		return NonPixelShaderExpressionError();
 	}
 
-	TextureReferenceIndex = Material->GetReferencedTextures().Find(DefaultValue);
+	TextureReferenceIndex = Material->GetMatDefaultTextureIdx(DefaultValue);
 	checkf(TextureReferenceIndex != INDEX_NONE, TEXT("Material expression called Compiler->ExternalTextureParameter() without implementing UMaterialExpression::GetReferencedTexture properly"));
 	FAddUniformExpressionScope Scope(this);
 	return AddUniformExpression(Scope, new FMaterialUniformExpressionExternalTextureParameter(ParameterName, TextureReferenceIndex), MCT_TextureExternal, TEXT(""));
@@ -9264,7 +9264,7 @@ int32 FHLSLMaterialTranslator::ExternalTextureCoordinateOffset(const FGuid& Exte
 
 UObject* FHLSLMaterialTranslator::GetReferencedTexture(int32 Index)
 {
-	return Material->GetReferencedTextures()[Index];
+	return Material->GetMatDefaultTextureAtIdx(Index);
 }
 
 UTextureCollection* FHLSLMaterialTranslator::GetReferencedTextureCollection(int32 Index)
@@ -17359,7 +17359,7 @@ void FHLSLMaterialTranslator::PrepareMaterialSourceStringParameters()
 
 int32 FHLSLMaterialTranslator::SparseVolumeTexture(USparseVolumeTexture* Texture, int32& TextureReferenceIndex, EMaterialSamplerType SamplerType)
 {
-	TextureReferenceIndex = Material->GetReferencedTextures().Find(Texture);
+	TextureReferenceIndex = Material->GetMatDefaultTextureIdx(Texture);
 	checkf(TextureReferenceIndex != INDEX_NONE, TEXT("Material expression called Compiler->SparseVolumeTexture() without implementing UMaterialExpression::GetReferencedTexture properly"));
 
 	FAddUniformExpressionScope Scope(this);
@@ -17377,7 +17377,7 @@ int32 FHLSLMaterialTranslator::SparseVolumeTextureParameter(FName ParameterName,
 		DefaultTexture = Meta.Value.SparseVolumeTexture;
 	}
 
-	TextureReferenceIndex = Material->GetReferencedTextures().Find(DefaultTexture);
+	TextureReferenceIndex = Material->GetMatDefaultTextureIdx(DefaultTexture);
 	checkf(TextureReferenceIndex != INDEX_NONE, TEXT("Material expression called Compiler->SparseVolumeTextureParameter() without implementing UMaterialExpression::GetReferencedTexture properly"));
 
 	FMaterialParameterInfo ParameterInfo = GetParameterAssociationInfo();
@@ -17445,7 +17445,7 @@ int32 FHLSLMaterialTranslator::SparseVolumeTextureSamplePageTable(int32 SparseVo
 	int32 SVTReferenceIndex = UniformTextureExpressions[(uint32)EMaterialTextureParameterType::SparseVolume].Find(TextureUniformExpression);
 	check(UniformTextureExpressions[(uint32)EMaterialTextureParameterType::SparseVolume].IsValidIndex(SVTReferenceIndex));
 
-	const USparseVolumeTexture* SVTexture = Cast<USparseVolumeTexture>(Material->GetReferencedTextures()[TextureUniformExpression->GetTextureIndex()]);
+	const USparseVolumeTexture* SVTexture = Cast<USparseVolumeTexture>(Material->GetMatDefaultTextureAtIdx(TextureUniformExpression->GetTextureIndex()));
 
 	// StaticAddress mode at time of compile
 	// Similar to VirtualTextures, this may not be 100% correct, if SamplerSource is set to 'SSM_FromTextureAsset', as texture parameter may change the address mode in a derived instance
diff --git a/Engine/Source/Runtime/Engine/Private/Materials/Material.cpp b/Engine/Source/Runtime/Engine/Private/Materials/Material.cpp
--- a/Engine/Source/Runtime/Engine/Private/Materials/Material.cpp
+++ b/Engine/Source/Runtime/Engine/Private/Materials/Material.cpp
@@ -1249,113 +1249,13 @@ void UMaterial::GetUsedTextures(TArray<UTexture*>& OutTextures, TOptional<EMater
 
 	if (!FPlatformProperties::IsServerOnly())
 	{
-		FInt32Range QualityLevelRange(0, EMaterialQualityLevel::Num - 1);
-		if (InSpecificQualityLevel.IsSet())
-		{
-			EMaterialQualityLevel::Type QualityLevel = InSpecificQualityLevel.GetValue();
-			ensure(QualityLevel != EMaterialQualityLevel::Num);
-			if (QualityLevel == EMaterialQualityLevel::Num)
-			{
-				QualityLevel = GetCurrentMaterialQualityLevelChecked();
-			}
-			QualityLevelRange = FInt32Range(QualityLevel, QualityLevel);
-		}
-
-#if WITH_EDITOR
-		FInt32Range ShaderPlatformRange(0, EShaderPlatform::SP_NumPlatforms - 1);
-		if (InSpecificShaderPlatform.IsSet())
-		{
-			EShaderPlatform ShaderPlatform = InSpecificShaderPlatform.GetValue();
-			ensure(ShaderPlatform != EShaderPlatform::SP_NumPlatforms);
-			if (ShaderPlatform == EShaderPlatform::SP_NumPlatforms)
-			{
-				ShaderPlatform = GetMaxShaderPlatformChecked();
-			}
-			ShaderPlatformRange = FInt32Range(ShaderPlatform, ShaderPlatform);
-		}
-#else
-		FInt32Range ShaderPlatformRange(GMaxRHIShaderPlatform, GMaxRHIShaderPlatform);
-#endif
-
-		TArray<const FMaterialResource*, TInlineAllocator<4>> MatchedResources;
-		// Parse all relevant quality and feature levels.
-		for (int32 QualityLevelIndex = QualityLevelRange.GetLowerBoundValue(); QualityLevelIndex <= QualityLevelRange.GetUpperBoundValue(); ++QualityLevelIndex)
-		{
-			for (int32 ShaderPlatformIndex = ShaderPlatformRange.GetLowerBoundValue(); ShaderPlatformIndex <= ShaderPlatformRange.GetUpperBoundValue(); ++ShaderPlatformIndex)
-			{
-				if (!FDataDrivenShaderPlatformInfo::IsValid((EShaderPlatform)ShaderPlatformIndex))
-				{
-					continue;
-				}
-
-				const FMaterialResource* CurrentResource = FindMaterialResource(MaterialResources, (EShaderPlatform)ShaderPlatformIndex, (EMaterialQualityLevel::Type)QualityLevelIndex, true);
-				if (CurrentResource)
-				{
-					MatchedResources.AddUnique(CurrentResource);
-				}
-			}
-		}
+		const FInt32Range QualityLevelRange = BuildQualityLevelRange(InSpecificQualityLevel);
+		const FInt32Range ShaderPlatformRange = BuildShaderPlatformRange(InSpecificShaderPlatform);
 
+		TArray<const FMaterialResource*, TInlineAllocator<4>> MatchedResources = GetMaterialResources(this, QualityLevelRange, ShaderPlatformRange);
 		for (const FMaterialResource* CurrentResource : MatchedResources)
 		{
-			for (int32 TypeIndex = 0; TypeIndex < NumMaterialTextureParameterTypes; TypeIndex++)
-			{
-				// Iterate over each of the material's texture expressions.
-				for (const FMaterialTextureParameterInfo& Parameter : CurrentResource->GetUniformTextureExpressions((EMaterialTextureParameterType)TypeIndex))
-				{
-					UTexture* Texture = NULL;
-					Parameter.GetGameThreadTextureValue(this, *CurrentResource, Texture);
-					if (Texture)
-					{
-						OutTextures.AddUnique(Texture);
-					}
-				}
-			}
-
-			for (const FMaterialTextureCollectionParameterInfo& Parameter : CurrentResource->GetUniformTextureCollectionExpressions())
-			{
-				UTextureCollection* TextureCollection = nullptr;
-				Parameter.GetGameThreadTextureCollectionValue(this, *CurrentResource, TextureCollection);
-				if (TextureCollection)
-				{
-					for (UTexture* Texture : TextureCollection->Textures)
-					{
-						if (Texture)
-						{
-							OutTextures.AddUnique(Texture);
-						}
-					}
-				}
-			}
-
-#if WITH_EDITOR
-/*
-			// Also look for any scalar parameters that are acting as lookups for an atlas texture, and store the atlas texture
-			const TArrayView<const FMaterialNumericParameterInfo> AtlasExpressions[1] =
-			{
-				CurrentResource->GetUniformNumericParameterExpressions()
-			};
-
-			for (int32 TypeIndex = 0; TypeIndex < UE_ARRAY_COUNT(AtlasExpressions); TypeIndex++)
-			{
-				// Iterate over each of the material's texture expressions.
-				for (const FMaterialNumericParameterInfo& Parameter : AtlasExpressions[TypeIndex])
-				{
-					if (Parameter.ParameterType == EMaterialParameterType::Scalar)
-					{
-						bool bIsUsedAsAtlasPosition;
-						TSoftObjectPtr<UCurveLinearColor> Curve;
-						TSoftObjectPtr<UCurveLinearColorAtlas> Atlas;
-						IsScalarParameterUsedAsAtlasPosition(Parameter.ParameterInfo, bIsUsedAsAtlasPosition, Curve, Atlas);
-						if (Atlas)
-						{
-							OutTextures.AddUnique(Atlas.Get());
-						}
-					}
-				}
-			}
-*/
-#endif // WITH_EDITOR
+			GetTextureExpressionValues(CurrentResource, OutTextures, nullptr, true, false);
 		}
 	}
 }
@@ -1372,36 +1272,7 @@ void UMaterial::GetUsedTexturesAndIndices(TArray<UTexture*>& OutTextures, TArray
 		const FMaterialResource* CurrentResource = FindMaterialResource(MaterialResources, InShaderPlatform, QualityLevel, true);
 		if (CurrentResource)
 		{
-			TArrayView<const FMaterialTextureParameterInfo> ExpressionsByType[NumMaterialTextureParameterTypes];
-			uint32 NumTextures = 0u;
-			for (uint32 TypeIndex = 0u; TypeIndex < NumMaterialTextureParameterTypes; ++TypeIndex)
-			{
-				ExpressionsByType[TypeIndex] = CurrentResource->GetUniformTextureExpressions((EMaterialTextureParameterType)TypeIndex);
-				NumTextures += ExpressionsByType[TypeIndex].Num();
-			}
-
-			// Try to prevent resizing since this would be expensive.
-			OutIndices.Empty(NumTextures);
-
-			for (int32 TypeIndex = 0; TypeIndex < UE_ARRAY_COUNT(ExpressionsByType); TypeIndex++)
-			{
-				// Iterate over each of the material's texture expressions.
-				for (const FMaterialTextureParameterInfo& Parameter : ExpressionsByType[TypeIndex])
-				{
-					UTexture* Texture = NULL;
-					Parameter.GetGameThreadTextureValue(this, *CurrentResource, Texture);
-
-					if (Texture)
-					{
-						int32 InsertIndex = OutTextures.AddUnique(Texture);
-						if (InsertIndex >= OutIndices.Num())
-						{
-							OutIndices.AddDefaulted(InsertIndex - OutIndices.Num() + 1);
-						}
-						OutIndices[InsertIndex].Add(Parameter.TextureIndex);
-					}
-				}
-			}
+			GetTextureExpressionValues(CurrentResource, OutTextures, &OutIndices, false, false);
 		}
 	}
 }
@@ -2294,8 +2165,11 @@ void UMaterial::UpdateCachedExpressionData()
 
 	CachedExpressionData.Reset(LocalCachedExpressionData);
 	EditorOnlyData->CachedExpressionData = LocalCachedExpressionData->EditorOnlyData;
+	ReferencedDefaultTextures = CachedExpressionData->EditorOnlyData->EditorReferencedDefaultTextures;
 
+#if WITH_EDITOR
 	FObjectCacheEventSink::NotifyReferencedTextureChanged_Concurrent(this);
+#endif
 }
 #endif // WITH_EDITOR
 
@@ -2975,6 +2849,14 @@ void UMaterial::Serialize(FArchive& Ar)
 	Ar.UsingCustomVersion(FUE5MainStreamObjectVersion::GUID);
 	Ar.UsingCustomVersion(FUE5ReleaseStreamObjectVersion::GUID);
 
+#if WITH_EDITOR
+	if (Ar.IsSaving() && Ar.IsCooking())
+	{
+		constexpr bool bKeepProjectPrivateTextures = true;
+		RemoveUnusedDefaultTextures(bKeepProjectPrivateTextures);
+	}
+#endif
+
 	Super::Serialize(Ar);
 
 	if (Ar.UEVer() >= VER_UE4_PURGED_FMATERIAL_COMPILE_OUTPUTS)
@@ -3150,6 +3032,17 @@ void UMaterial::Serialize(FArchive& Ar)
 			SubsurfaceProfiles.Add(SubsurfaceProfile);
 		}
 	}
+
+#if WITH_EDITOR
+	if (Ar.IsSaving() && Ar.IsCooking())
+	{
+		// Undo textures removal for Multi Platform Cook
+		if (ReferencedDefaultTextures.Num() > 0)
+		{
+			ReferencedDefaultTextures = CachedExpressionData->EditorOnlyData->EditorReferencedDefaultTextures;
+		}
+	}
+#endif
 }
 
 void UMaterial::PostDuplicate(bool bDuplicateForPIE)
@@ -4667,7 +4560,7 @@ void UMaterial::PostLoad()
 		{
 			// Before caching shader resources we have to make sure all referenced textures have been post loaded
 			// as we depend on their resources being valid.
-			for (UObject* Texture : CachedExpressionData->ReferencedTextures)
+			for (UObject* Texture : ReferencedDefaultTextures)
 			{
 				if (Texture)
 				{
@@ -4744,6 +4637,22 @@ void UMaterial::DumpDebugInfo(FOutputDevice& OutputDevice) const
 #endif
 }
 
+TObjectPtr<UObject> UMaterial::GetDefaultTextureAtIdx(const int32 InIdx, TMicRecursionGuard RecursionGuard)
+{
+	if (!ReferencedDefaultTextures.IsValidIndex(InIdx))
+	{
+		static bool bWarnedOnce = false;
+		if (!bWarnedOnce)
+		{
+			UE_LOG(LogMaterial, Warning, TEXT("%s: Requesting an invalid TextureIndex! (%u / %u)"), *GetName(), InIdx, ReferencedDefaultTextures.Num());
+			bWarnedOnce = true;
+		}
+		return nullptr;
+	}
+
+	return ReferencedDefaultTextures[InIdx];
+}
+
 void UMaterial::SaveShaderStableKeys(const class ITargetPlatform* TP)
 {
 #if WITH_EDITOR
@@ -7168,6 +7077,29 @@ void UMaterial::ForceRecompileForRendering(EMaterialShaderPrecompileMode Compile
 	CacheResourceShadersForRendering(false,  CompileMode);
 }
 
+#if WITH_EDITOR
+void UMaterial::GetUsedTexturesInCookCache(TArray<TObjectPtr<UObject>>& OutTextures)
+{
+	OutTextures.Empty();
+	const FMaterialCachedExpressionData& UsedCachedData = GetCachedExpressionData();
+	TArray<const FMaterialResource*, TInlineAllocator<4>> MatchedResources = GetCookedCacheMaterialResources();
+
+	for (const FMaterialResource* CurrentResource : MatchedResources)
+	{
+		const FUniformExpressionSet& CurrentSet = CurrentResource->GetUniformExpressions();
+
+		for (int32 TypeIndex = 0; TypeIndex < NumMaterialTextureParameterTypes; TypeIndex++)
+		{
+			for (int32 TextureIndex = 0; TextureIndex < CurrentSet.GetNumTextures((EMaterialTextureParameterType)TypeIndex); TextureIndex++)
+			{
+				const FMaterialTextureParameterInfo& Parameter = CurrentSet.GetTextureParameter((EMaterialTextureParameterType)TypeIndex, TextureIndex);
+				OutTextures.AddUnique(UsedCachedData.EditorOnlyData->EditorReferencedDefaultTextures[Parameter.TextureIndex]);
+			}
+		}
+	}
+}
+#endif
+
 bool UMaterial::CheckInValidStateForCompilation(class FMaterialCompiler* Compiler) const
 {
 	bool bSuccess = true;
diff --git a/Engine/Source/Runtime/Engine/Private/Materials/MaterialCachedData.cpp b/Engine/Source/Runtime/Engine/Private/Materials/MaterialCachedData.cpp
--- a/Engine/Source/Runtime/Engine/Private/Materials/MaterialCachedData.cpp
+++ b/Engine/Source/Runtime/Engine/Private/Materials/MaterialCachedData.cpp
@@ -109,7 +109,9 @@ FMaterialCachedExpressionData::FMaterialCachedExpressionData()
 
 void FMaterialCachedExpressionData::AddReferencedObjects(FReferenceCollector& Collector)
 {
-	Collector.AddStableReferenceArray(&ReferencedTextures);
+#if WITH_EDITOR
+ 	Collector.AddStableReferenceArray(&EditorOnlyData->EditorReferencedDefaultTextures);
+#endif
 	Collector.AddStableReferenceArray(&ReferencedTextureCollections);
 	Collector.AddStableReferenceArray(&GrassTypes);
 	Collector.AddStableReferenceArray(&MaterialCacheTags);
@@ -556,7 +558,7 @@ void FMaterialCachedExpressionData::UpdateForExpressions(const FMaterialCachedEx
 
 		if (ReferencedTexture)
 		{
-			ReferencedTextures.AddUnique(ReferencedTexture);
+			EditorOnlyData->EditorReferencedDefaultTextures.AddUnique(ReferencedTexture);
 		}
 		else if (ReferencedTextureCollection)
 		{
@@ -572,7 +574,7 @@ void FMaterialCachedExpressionData::UpdateForExpressions(const FMaterialCachedEx
 			const UMaterialExpression::ReferencedTextureArray ExpressionReferencedTextures = Expression->GetReferencedTextures();
 			for (UObject* ExpressionReferencedTexture : ExpressionReferencedTextures)
 			{
-				ReferencedTextures.AddUnique(ExpressionReferencedTexture);
+				EditorOnlyData->EditorReferencedDefaultTextures.AddUnique(ExpressionReferencedTexture);
 			}
 		}
 
@@ -1133,4 +1135,11 @@ void FMaterialCachedExpressionData::GetAllGlobalParameterInfoOfType(EMaterialPar
 	}
 }
 
+#if WITH_EDITOR
+int32 FMaterialCachedExpressionData::GetDefaultTextureIdx(UObject* InTexture) const
+{
+	return EditorOnlyData->EditorReferencedDefaultTextures.Find(TObjectPtr<UObject>(InTexture));
+}
+#endif
+
 #undef LOCTEXT_NAMESPACE
diff --git a/Engine/Source/Runtime/Engine/Private/Materials/MaterialIREmitter.cpp b/Engine/Source/Runtime/Engine/Private/Materials/MaterialIREmitter.cpp
--- a/Engine/Source/Runtime/Engine/Private/Materials/MaterialIREmitter.cpp
+++ b/Engine/Source/Runtime/Engine/Private/Materials/MaterialIREmitter.cpp
@@ -1273,7 +1273,7 @@ FValueRef FEmitter::PreshaderParameter(FType Type, EPreshader2Opcode Opcode, FVa
 	FPreshaderParameter Prototype = MakePrototype<FPreshaderParameter>(Type);
 	Prototype.SourceParameter = SourceParameter;
 	Prototype.Opcode = Opcode;
-	Prototype.TextureIndex = Material->GetReferencedTextures().Find(SourceParameterTexture);
+	Prototype.TextureIndex = Material->GetDefaultTextureIdx(SourceParameterTexture);
 	Prototype.Payload = MoveTemp(Payload);
 	return EmitPrototype(*this, Prototype);
 }
diff --git a/Engine/Source/Runtime/Engine/Private/Materials/MaterialIRValueAnalyzer.cpp b/Engine/Source/Runtime/Engine/Private/Materials/MaterialIRValueAnalyzer.cpp
--- a/Engine/Source/Runtime/Engine/Private/Materials/MaterialIRValueAnalyzer.cpp
+++ b/Engine/Source/Runtime/Engine/Private/Materials/MaterialIRValueAnalyzer.cpp
@@ -181,7 +181,7 @@ static int32 RegisterTextureParameter(FMaterialIRValueAnalyzer& Analyzer, UObjec
 	const FMaterialTextureParameterInfo TextureParameterInfo
 	{
 		.ParameterInfo = ParameterInfo,
-		.TextureIndex = Analyzer.Material->GetReferencedTextures().Find(Texture),
+		.TextureIndex = Analyzer.Material->GetDefaultTextureIdx(Texture),
 		.SamplerSource = SSM_FromTextureAsset,
 		.VirtualTextureLayerIndex = static_cast<uint8>(VTLayerIndex == INDEX_NONE ? UINT8_MAX : VTLayerIndex),
 	};
diff --git a/Engine/Source/Runtime/Engine/Private/Materials/MaterialInstance.cpp b/Engine/Source/Runtime/Engine/Private/Materials/MaterialInstance.cpp
--- a/Engine/Source/Runtime/Engine/Private/Materials/MaterialInstance.cpp
+++ b/Engine/Source/Runtime/Engine/Private/Materials/MaterialInstance.cpp
@@ -1174,11 +1174,6 @@ bool UMaterialInstance::GetParameterValue(EMaterialParameterType Type, const FMe
 	GetMaterialInheritanceChain(InstanceChain);
 
 	bool bResult = false;
-	if (EnumHasAnyFlags(Flags, EMaterialGetParameterValueFlags::CheckNonOverrides))
-	{
-		bResult = InstanceChain.GetCachedExpressionData().GetParameterValue(Type, ParameterInfo, OutResult);
-	}
-
 	const bool bCheckInstanceOverrides = EnumHasAnyFlags(Flags, EMaterialGetParameterValueFlags::CheckInstanceOverrides);
 	FMemoryImageMaterialParameterInfo CurrentParameterInfo = ParameterInfo;
 	bool bHasValidParameter = true;
@@ -1210,6 +1205,12 @@ bool UMaterialInstance::GetParameterValue(EMaterialParameterType Type, const FMe
 		ParentIndex++;
 	}
 
+	// Look through the default values if we have not found an override
+	if (!bResult && EnumHasAnyFlags(Flags, EMaterialGetParameterValueFlags::CheckNonOverrides))
+	{
+		bResult = InstanceChain.GetCachedExpressionData().GetParameterValue(Type, ParameterInfo, OutResult);
+	}
+
 	return bResult;
 }
 
@@ -1311,67 +1312,9 @@ int32 UMaterialInstance::GetBlendablePriority(const UMaterial* Base) const
 	}
 }
 
-void UMaterialInstance::GetTextureExpressionValues(const FMaterialResource* MaterialResource, TArray<UTexture*>& OutTextures, TArray<TArray<int32>>* OutIndices, bool bIncludeTextureCollections) const
+void UMaterialInstance::GetTextureExpressionValues(const FMaterialResource* MaterialResource, TArray<UTexture*>& OutTextures, TArray<TArray<int32>>* OutIndices, bool bIncludeTextureCollections, bool bAllowOverrides) const
 {
-	check(MaterialResource);
-	const FUniformExpressionSet& UniformExpressions = MaterialResource->GetUniformExpressions();
-
-	if (OutIndices) // Try to prevent resizing since this would be expensive.
-	{
-		uint32 NumTextures = 0u;
-		for (uint32 TypeIndex = 0u; TypeIndex < NumMaterialTextureParameterTypes; ++TypeIndex)
-		{
-			NumTextures += UniformExpressions.GetNumTextures((EMaterialTextureParameterType)TypeIndex);
-		}
-		OutIndices->Empty(NumTextures);
-	}
-
-	for(int32 TypeIndex = 0;TypeIndex < NumMaterialTextureParameterTypes;TypeIndex++)
-	{
-		// Iterate over each of the material's texture expressions.
-		for(int32 TextureIndex = 0; TextureIndex < UniformExpressions.GetNumTextures((EMaterialTextureParameterType)TypeIndex); TextureIndex++)
-		{
-			// Evaluate the expression in terms of this material instance.
-			UTexture* Texture = NULL;
-			UniformExpressions.GetGameThreadTextureValue((EMaterialTextureParameterType)TypeIndex, TextureIndex, this,*MaterialResource,Texture, true);
-			
-			if (Texture)
-			{
-				const int32 InsertIndex = OutTextures.AddUnique(Texture);
-				if (OutIndices)
-				{
-					const FMaterialTextureParameterInfo& Parameter = UniformExpressions.GetTextureParameter((EMaterialTextureParameterType)TypeIndex, TextureIndex);
-					if (InsertIndex >= OutIndices->Num())
-					{
-						OutIndices->AddDefaulted(InsertIndex - OutIndices->Num() + 1);
-					}
-					(*OutIndices)[InsertIndex].Add(Parameter.TextureIndex);
-				}
-			}
-		}
-	}
-
-	if (bIncludeTextureCollections)
-	{
-		checkf(OutIndices == nullptr, TEXT("Texture Collections don't work with Texture Indices."));
-
-		for (int32 TextureCollectionIndex = 0; TextureCollectionIndex < UniformExpressions.GetNumTextureCollections(); TextureCollectionIndex++)
-		{
-			UTextureCollection* TextureCollection = nullptr;
-			UniformExpressions.GetGameThreadTextureCollectionValue(TextureCollectionIndex, this, *MaterialResource, TextureCollection);
-
-			if (TextureCollection)
-			{
-				for (UTexture* Texture : TextureCollection->Textures)
-				{
-					if (Texture)
-					{
-						OutTextures.AddUnique(Texture);
-					}
-				}
-			}
-		}
-	}
+	Super::GetTextureExpressionValues(MaterialResource, OutTextures, OutIndices, bIncludeTextureCollections, bAllowOverrides);
 
 #if WITH_EDITOR
 	for (int32 i = 0; i < TransientTextureParameterOverrides.Num(); ++i)
@@ -1406,33 +1349,8 @@ void UMaterialInstance::GetUsedTextures(TArray<UTexture*>& OutTextures, TOptiona
 	// Do not care if we're running dedicated server
 	if (!FPlatformProperties::IsServerOnly())
 	{
-		FInt32Range QualityLevelRange(0, EMaterialQualityLevel::Num - 1);
-		if (InSpecificQualityLevel.IsSet())
-		{
-			EMaterialQualityLevel::Type QualityLevel = InSpecificQualityLevel.GetValue();
-			ensure(QualityLevel != EMaterialQualityLevel::Num);
-			if (QualityLevel == EMaterialQualityLevel::Num)
-			{
-				QualityLevel = GetCurrentMaterialQualityLevelChecked();
-			}
-			QualityLevelRange = FInt32Range(QualityLevel, QualityLevel);
-		}
-
-#if WITH_EDITOR
-		FInt32Range ShaderPlatformRange(0, EShaderPlatform::SP_NumPlatforms - 1);
-		if (InSpecificShaderPlatform.IsSet())
-		{
-			EShaderPlatform ShaderPlatform = InSpecificShaderPlatform.GetValue();
-			ensure(ShaderPlatform != EShaderPlatform::SP_NumPlatforms);
-			if (ShaderPlatform == EShaderPlatform::SP_NumPlatforms)
-			{
-				ShaderPlatform = GetMaxShaderPlatformChecked();
-			}
-			ShaderPlatformRange = FInt32Range(ShaderPlatform, ShaderPlatform);
-		}
-#else
-		FInt32Range ShaderPlatformRange(GMaxRHIShaderPlatform, GMaxRHIShaderPlatform);
-#endif
+		const FInt32Range QualityLevelRange = BuildQualityLevelRange(InSpecificQualityLevel);
+		const FInt32Range ShaderPlatformRange = BuildShaderPlatformRange(InSpecificShaderPlatform);
 
 		const UMaterial* BaseMaterial = GetMaterial();
 		const UMaterialInstance* MaterialInstanceToUse = this;
@@ -1448,28 +1366,13 @@ void UMaterialInstance::GetUsedTextures(TArray<UTexture*>& OutTextures, TOptiona
 			// Use the uniform expressions from the lowest material instance with static parameters in the chain, if one exists
 			const UMaterialInterface* MaterialToUse = (MaterialInstanceToUse && MaterialInstanceToUse->bHasStaticPermutationResource) ? (const UMaterialInterface*)MaterialInstanceToUse : (const UMaterialInterface*)BaseMaterial;
 
-			TArray<const FMaterialResource*, TInlineAllocator<4>> MatchedResources;
-			// Parse all relevant quality and feature levels.
-			for (int32 QualityLevelIndex = QualityLevelRange.GetLowerBoundValue(); QualityLevelIndex <= QualityLevelRange.GetUpperBoundValue(); ++QualityLevelIndex)
-			{
-				for (int32 ShaderPlatformIndex = ShaderPlatformRange.GetLowerBoundValue(); ShaderPlatformIndex <= ShaderPlatformRange.GetUpperBoundValue(); ++ShaderPlatformIndex)
-				{
-					if (!FDataDrivenShaderPlatformInfo::IsValid((EShaderPlatform)ShaderPlatformIndex))
-					{
-						continue;
-					}
-
-					const FMaterialResource* MaterialResource = MaterialToUse->GetMaterialResource((EShaderPlatform)ShaderPlatformIndex, (EMaterialQualityLevel::Type)QualityLevelIndex);
-					if (MaterialResource)
-					{
-						MatchedResources.AddUnique(MaterialResource);
-					}
-				}
-			}
+			TArray<const FMaterialResource*, TInlineAllocator<4>> MatchedResources = GetMaterialResources(MaterialToUse, QualityLevelRange, ShaderPlatformRange);
 
-			for (const FMaterialResource* MaterialResource : MatchedResources)
+			for (const FMaterialResource* CurrentResource : MatchedResources)
 			{
-				GetTextureExpressionValues(MaterialResource, OutTextures, nullptr, true);
+				constexpr bool bIncludeTextureCollections = true;
+				constexpr bool bAllowOverrides = true;
+				GetTextureExpressionValues(CurrentResource, OutTextures, nullptr, bIncludeTextureCollections, bAllowOverrides);
 			}
 		}
 		else
@@ -1537,7 +1440,10 @@ void UMaterialInstance::LogMaterialsAndTextures(FOutputDevice& Ar, int32 Indent)
 				if (MaterialResource->HasValidGameThreadShaderMap())
 				{
 					TArray<UTexture*> Textures;
-					GetTextureExpressionValues(MaterialResource, Textures);
+					constexpr bool bIncludeTextureCollections = false;
+					constexpr bool bAllowOverrides = true;
+					GetTextureExpressionValues(MaterialResource, Textures, nullptr, bIncludeTextureCollections, bAllowOverrides);
+
 					for (UTexture* Texture : Textures)
 					{
 						if (Texture)
@@ -1685,7 +1591,9 @@ void UMaterialInstance::GetUsedTexturesAndIndices(TArray<UTexture*>& OutTextures
 			const FMaterialResource* CurrentResource = FindMaterialResource(MaterialInstanceToUse->StaticPermutationMaterialResources, InShaderPlatform, QualityLevel, true);
 			if (CurrentResource)
 			{
-				GetTextureExpressionValues(CurrentResource, OutTextures, &OutIndices);
+				constexpr bool bIncludeTextureCollections = false;
+				constexpr bool bAllowOverrides = true;
+				GetTextureExpressionValues(CurrentResource, OutTextures, &OutIndices, bIncludeTextureCollections, bAllowOverrides);
 			}
 		}
 		else // Use the uniform expressions from the base material
@@ -1696,7 +1604,9 @@ void UMaterialInstance::GetUsedTexturesAndIndices(TArray<UTexture*>& OutTextures
 				const FMaterialResource* MaterialResource = Material->GetMaterialResource(InShaderPlatform, QualityLevel);
 				if( MaterialResource )
 				{
-					GetTextureExpressionValues(MaterialResource, OutTextures, &OutIndices);
+					constexpr bool bIncludeTextureCollections = false;
+					constexpr bool bAllowOverrides = true;
+					GetTextureExpressionValues(MaterialResource, OutTextures, &OutIndices, bIncludeTextureCollections, bAllowOverrides);
 				}
 			}
 			else // If the material instance has no material, use the default material.
@@ -3275,6 +3185,11 @@ void UMaterialInstance::Serialize(FArchive& Ar)
 #endif
 
 #if WITH_EDITOR
+	if (Ar.IsSaving() && Ar.IsCooking())
+	{
+		constexpr bool bKeepProjectPrivateTextures = false;
+		RemoveUnusedDefaultTextures(bKeepProjectPrivateTextures);
+	}
 	// Do not serialize the overrides
 	ResetAllTextureParameterOverrides();
 #endif
@@ -3487,6 +3402,25 @@ void UMaterialInstance::Serialize(FArchive& Ar)
 	{
 		ValidateTextureOverrides(GMaxRHIShaderPlatform);
 	}
+
+#if WITH_EDITOR
+	// Undo textures removal for Multi Platform Cook
+	if (Ar.IsSaving() && Ar.IsCooking())
+	{
+		if (ReferencedDefaultTextures.Num() > 0)
+		{
+			if (CachedExpressionData)
+			{
+				ReferencedDefaultTextures = CachedExpressionData->EditorOnlyData->EditorReferencedDefaultTextures;
+			}
+			else
+			{
+				TUniquePtr<FMaterialCachedExpressionData> LocalCachedExpressionData = CreateCachedExpressionData();
+				ReferencedDefaultTextures = LocalCachedExpressionData->EditorOnlyData->EditorReferencedDefaultTextures;
+			}
+		}
+	}
+#endif
 }
 
 void UMaterialInstance::PostDuplicate(bool bDuplicateForPIE)
@@ -3658,14 +3592,11 @@ void UMaterialInstance::PostLoad()
 		}
 	}
 
-	if (CachedExpressionData != nullptr)
+	for (UObject* Texture : ReferencedDefaultTextures)
 	{
-		for (UObject* Texture : CachedExpressionData->ReferencedTextures)
+		if (Texture)
 		{
-			if (Texture)
-			{
-				Texture->ConditionalPostLoad();
-			}
+			Texture->ConditionalPostLoad();
 		}
 	}
 
@@ -4675,7 +4606,7 @@ TSoftObjectPtr<UObject> UMaterialInstance::GetEnumerationObject(int32 Index) con
 
 void UMaterialInstance::GetReferencedTexturesAndOverrides(TSet<const UTexture*>& InOutTextures) const
 {
-	for (UObject* UsedObject : GetCachedExpressionData().ReferencedTextures)
+	for (UObject* UsedObject : GetDefaultTextures())
 	{
 		if (const UTexture* UsedTexture = Cast<UTexture>(UsedObject))
 		{
@@ -4703,6 +4634,23 @@ void UMaterialInstance::UpdateCachedData()
 	// Overridden for MIC/MID
 }
 
+TUniquePtr<FMaterialCachedExpressionData> UMaterialInstance::CreateCachedExpressionData() const
+{
+	TUniquePtr<FMaterialCachedExpressionData> LocalCachedExpressionData;
+	const UMaterial* BaseMaterial = GetMaterial();
+
+	const FStaticParameterSet& LocalStaticParameters = GetStaticParameters();
+	FMaterialLayersFunctions MaterialLayers;
+	LocalStaticParameters.GetMaterialLayers(MaterialLayers);
+
+	FMaterialCachedExpressionContext Context;
+	Context.LayerOverrides = &MaterialLayers;
+	LocalCachedExpressionData.Reset(new FMaterialCachedExpressionData());
+	LocalCachedExpressionData->UpdateForExpressions(Context, BaseMaterial->GetExpressions(), GlobalParameter, INDEX_NONE);
+
+	return MoveTemp(LocalCachedExpressionData);
+}
+
 void UMaterialInstance::UpdateStaticPermutation(const FStaticParameterSet& NewParameters, FMaterialUpdateContext* MaterialUpdateContext)
 {
 	UpdateStaticPermutation(NewParameters, BasePropertyOverrides, false, MaterialUpdateContext);
@@ -5404,6 +5352,44 @@ int32 UMaterialInstance::CompilePropertyEx( class FMaterialCompiler* Compiler, c
 	return Parent ? Parent->CompilePropertyEx(Compiler, AttributeID) : INDEX_NONE;
 }
 
+void UMaterialInstance::GetUsedTexturesInCookCache(TArray<TObjectPtr<UObject>>& OutTextures)
+{
+	OutTextures.Empty();
+
+	const UMaterial* BaseMaterial = GetMaterial();
+	if (BaseMaterial && !BaseMaterial->IsDefaultMaterial())
+	{
+		const FMaterialCachedExpressionData& UsedCachedData = GetCachedExpressionData();
+		TArray<const FMaterialResource*, TInlineAllocator<4>> MatchedResources = GetCookedCacheMaterialResources();
+
+		for (const FMaterialResource* CurrentResource : MatchedResources)
+		{
+			const FUniformExpressionSet& CurrentSet = CurrentResource->GetUniformExpressions();
+
+			for (int32 TypeIndex = 0; TypeIndex < NumMaterialTextureParameterTypes; ++TypeIndex)
+			{
+				for (int32 TextureIndex = 0; TextureIndex < CurrentSet.GetNumTextures((EMaterialTextureParameterType)TypeIndex); ++TextureIndex)
+				{
+					const FMaterialTextureParameterInfo& Parameter = CurrentSet.GetTextureParameter((EMaterialTextureParameterType)TypeIndex, TextureIndex);
+					FMaterialParameterMetadata SourceValue;
+					bool bIsOverriden = false;
+
+					// Do not add overriden default textures
+					for (int32 ParamTypeIdx = (int32)EMaterialParameterType::Texture; (ParamTypeIdx <= (int32)EMaterialParameterType::SparseVolumeTexture) && !bIsOverriden; ++ParamTypeIdx)
+					{
+						bIsOverriden |= GetParameterValue((EMaterialParameterType)ParamTypeIdx, Parameter.ParameterInfo, SourceValue, EMaterialGetParameterValueFlags::CheckInstanceOverrides);
+					}
+
+					if (!bIsOverriden)
+					{
+						OutTextures.AddUnique(UsedCachedData.EditorOnlyData->EditorReferencedDefaultTextures[Parameter.TextureIndex]);
+					}
+				}
+			}
+		}
+	}
+}
+
 const FStaticParameterSetEditorOnlyData& UMaterialInstance::GetEditorOnlyStaticParameters() const
 {
 	const UMaterialInstanceEditorOnlyData* EditorOnly = GetEditorOnlyData();
@@ -5660,6 +5646,25 @@ void UMaterialInstance::DumpDebugInfo(FOutputDevice& OutputDevice) const
 	}
 }
 
+TObjectPtr<UObject> UMaterialInstance::GetDefaultTextureAtIdx(const int32 InIdx, TMicRecursionGuard RecursionGuard)
+{
+	// We also look through the Parents in case there's a texture that has been forced to 
+	// remain inside the parent because it is Project Private (PKG_NotExternallyReferenceable)
+
+	if (ReferencedDefaultTextures.IsValidIndex(InIdx) && ReferencedDefaultTextures[InIdx] != nullptr)
+	{
+		return ReferencedDefaultTextures[InIdx];
+	}
+
+	if (Parent && !RecursionGuard.Contains(this))
+	{
+		RecursionGuard.Set(this);
+		return Parent->GetDefaultTextureAtIdx(InIdx, RecursionGuard);
+	}
+
+	return nullptr;
+}
+
 void UMaterialInstance::SaveShaderStableKeys(const class ITargetPlatform* TP)
 {
 #if WITH_EDITOR
@@ -6004,4 +6009,4 @@ bool UMaterialInstance::GetTextureCollectionParameterValue(const FHashedMaterial
 {
 	bool bResult = Super::GetTextureCollectionParameterValue(ParameterInfo, OutValue, bOveriddenOnly);
 	return bResult;
-}
\ No newline at end of file
+}
diff --git a/Engine/Source/Runtime/Engine/Private/Materials/MaterialInstanceConstant.cpp b/Engine/Source/Runtime/Engine/Private/Materials/MaterialInstanceConstant.cpp
--- a/Engine/Source/Runtime/Engine/Private/Materials/MaterialInstanceConstant.cpp
+++ b/Engine/Source/Runtime/Engine/Private/Materials/MaterialInstanceConstant.cpp
@@ -213,23 +213,22 @@ void UMaterialInstanceConstant::UpdateCachedData()
 
 	if (!bLoadedCachedExpressionData)
 	{
-		const bool bUsingNewHLSLGenerator = IsUsingNewHLSLGenerator();
 		TUniquePtr<FMaterialCachedExpressionData> LocalCachedExpressionData;
 
 		// If we have overridden material layers, need to create a local cached expression data
 		// Otherwise we can leave it as null, and use cached data from our parent
 		const FStaticParameterSet& LocalStaticParameters = GetStaticParameters();
+
 		if (LocalStaticParameters.bHasMaterialLayers)
 		{
-			UMaterial* BaseMaterial = GetMaterial();
-
-			FMaterialLayersFunctions MaterialLayers;
-			LocalStaticParameters.GetMaterialLayers(MaterialLayers);
-			
-			FMaterialCachedExpressionContext Context;
-			Context.LayerOverrides = &MaterialLayers;
-			LocalCachedExpressionData.Reset(new FMaterialCachedExpressionData());
-			LocalCachedExpressionData->UpdateForExpressions(Context, BaseMaterial->GetExpressions(), GlobalParameter, INDEX_NONE);
+			LocalCachedExpressionData = CreateCachedExpressionData();
+
+			if (bHasStaticPermutationResource)
+			{
+				 //If material is already analyzed, grab textures from CachedExpressionData
+				ReferencedDefaultTextures = LocalCachedExpressionData->EditorOnlyData->EditorReferencedDefaultTextures;
+			}
+
 		}
 
 		CachedExpressionData = MoveTemp(LocalCachedExpressionData);
@@ -238,7 +237,16 @@ void UMaterialInstanceConstant::UpdateCachedData()
 			EditorOnlyData->CachedExpressionData = CachedExpressionData->EditorOnlyData;
 		}
 
+		if (!CachedExpressionData && bHasStaticPermutationResource)
+		{
+			// We need the textures. Create a throwaway CachedExpressionData
+			LocalCachedExpressionData = CreateCachedExpressionData();
+			ReferencedDefaultTextures = LocalCachedExpressionData->EditorOnlyData->EditorReferencedDefaultTextures;
+		}
+
+#if WITH_EDITOR
 		FObjectCacheEventSink::NotifyReferencedTextureChanged_Concurrent(this);
+#endif
 	}
 }
 
diff --git a/Engine/Source/Runtime/Engine/Private/Materials/MaterialInterface.cpp b/Engine/Source/Runtime/Engine/Private/Materials/MaterialInterface.cpp
--- a/Engine/Source/Runtime/Engine/Private/Materials/MaterialInterface.cpp
+++ b/Engine/Source/Runtime/Engine/Private/Materials/MaterialInterface.cpp
@@ -48,9 +48,11 @@
 #include "UnrealEngine.h"
 #include "UObject/ObjectEditorOptionalSupport.h"
 #include "UObject/ICookInfo.h"
+#include "Engine/TextureCollection.h"
 
 #if WITH_EDITOR
 #include "ObjectCacheEventSink.h"
+#include "Misc/PackageName.h"
 #endif
 
 #include UE_INLINE_GENERATED_CPP_BY_NAME(MaterialInterface)
@@ -295,6 +297,25 @@ const FMaterialCachedExpressionData& UMaterialInterface::GetCachedExpressionData
 	return LocalData ? *LocalData : FMaterialCachedExpressionData::EmptyData;
 }
 
+#if WITH_EDITOR
+int32 UMaterialInterface::GetDefaultTextureIdx(UObject* InTexture) const
+{
+	const FMaterialCachedExpressionData& UsedCachedData = GetCachedExpressionData();
+	return UsedCachedData.GetDefaultTextureIdx(InTexture);
+}
+
+bool UMaterialInterface::ContainsDefaultTexture(UObject* InTexture) const
+{
+	return GetDefaultTextureIdx(InTexture) != INDEX_NONE;
+}
+
+const TArray<TObjectPtr<UObject>>& UMaterialInterface::GetDefaultTextures() const
+{
+	const FMaterialCachedExpressionData& UsedCachedData = GetCachedExpressionData();
+	return UsedCachedData.EditorOnlyData->EditorReferencedDefaultTextures;
+}
+#endif // WITH_EDITOR
+
 bool UMaterialInterface::IsUsingNewHLSLGenerator() const
 {
 	static const auto CVar = IConsoleManager::Get().FindTConsoleVariableDataBool(TEXT("r.Material.Translator.EnableNew"));
@@ -371,7 +392,7 @@ void UMaterialInterface::GetQualityLevelUsage(TArray<bool, TInlineAllocator<EMat
 
 TArrayView<const TObjectPtr<UObject>> UMaterialInterface::GetReferencedTextures() const
 {
-	return GetCachedExpressionData().ReferencedTextures;
+	return ReferencedDefaultTextures;
 }
 
 TConstArrayView<TObjectPtr<UTextureCollection>> UMaterialInterface::GetReferencedTextureCollections() const
@@ -382,7 +403,8 @@ TConstArrayView<TObjectPtr<UTextureCollection>> UMaterialInterface::GetReference
 #if WITH_EDITOR
 void UMaterialInterface::GetReferencedTexturesAndOverrides(TSet<const UTexture*>& InOutTextures) const
 {
-	for (UObject* UsedObject : GetCachedExpressionData().ReferencedTextures)
+	const TArray<TObjectPtr<UObject>> Textures = GetDefaultTextures();
+	for (UObject* UsedObject : Textures)
 	{
 		if (const UTexture* UsedTexture = Cast<UTexture>(UsedObject))
 		{
@@ -404,6 +426,95 @@ void UMaterialInterface::GetUsedTexturesAndIndices(TArray<UTexture*>& OutTexture
 	OutIndices.AddDefaulted(OutTextures.Num());
 }
 
+void UMaterialInterface::GetTextureExpressionValues(const FMaterialResource* MaterialResource, TArray<UTexture*>& OutTextures, TArray<TArray<int32>>* OutIndices, bool bIncludeTextureCollections, bool bAllowOverrides) const
+{
+	check(MaterialResource);
+	const FUniformExpressionSet& UniformExpressions = MaterialResource->GetUniformExpressions();
+
+	if (OutIndices) // Try to prevent resizing since this would be expensive.
+	{
+		uint32 NumTextures = 0u;
+		for (uint32 TypeIndex = 0u; TypeIndex < NumMaterialTextureParameterTypes; ++TypeIndex)
+		{
+			NumTextures += UniformExpressions.GetNumTextures((EMaterialTextureParameterType)TypeIndex);
+		}
+		OutIndices->Empty(NumTextures);
+	}
+
+	for (int32 TypeIndex = 0; TypeIndex < NumMaterialTextureParameterTypes; ++TypeIndex)
+	{
+		// Iterate over each of the material's texture expressions.
+		for (int32 TextureIndex = 0; TextureIndex < UniformExpressions.GetNumTextures((EMaterialTextureParameterType)TypeIndex); TextureIndex++)
+		{
+			// Evaluate the expression in terms of this material instance.
+			UTexture* Texture = nullptr;
+			UniformExpressions.GetGameThreadTextureValue((EMaterialTextureParameterType)TypeIndex, TextureIndex, this, *MaterialResource, Texture, bAllowOverrides);
+
+			if (Texture)
+			{
+				const int32 InsertIndex = OutTextures.AddUnique(Texture);
+
+				if (OutIndices)
+				{
+					if (InsertIndex >= OutIndices->Num())
+					{
+						OutIndices->AddDefaulted(InsertIndex - OutIndices->Num() + 1);
+					}
+
+					const FMaterialTextureParameterInfo& Parameter = UniformExpressions.GetTextureParameter((EMaterialTextureParameterType)TypeIndex, TextureIndex);
+					(*OutIndices)[InsertIndex].Add(Parameter.TextureIndex);
+				}
+			}
+		}
+	}
+
+	if (bIncludeTextureCollections)
+	{
+		checkf(OutIndices == nullptr, TEXT("Texture Collections don't work with Texture Indices."));
+
+		for (int32 TextureCollectionIndex = 0; TextureCollectionIndex < UniformExpressions.GetNumTextureCollections(); TextureCollectionIndex++)
+		{
+			UTextureCollection* TextureCollection = nullptr;
+			UniformExpressions.GetGameThreadTextureCollectionValue(TextureCollectionIndex, this, *MaterialResource, TextureCollection);
+
+			if (TextureCollection)
+			{
+				for (UTexture* Texture : TextureCollection->Textures)
+				{
+					if (Texture)
+					{
+						OutTextures.AddUnique(Texture);
+					}
+				}
+			}
+		}
+	}
+}
+
+TArray<const FMaterialResource*, TInlineAllocator<4>> UMaterialInterface::GetMaterialResources(const UMaterialInterface* InMaterial, const FInt32Range InQualityLevelRange, const FInt32Range InShaderPlatformRange)
+{
+	TArray<const FMaterialResource*, TInlineAllocator<4>> MatchedResources;
+	// Parse all relevant quality and feature levels.
+	for (int32 QualityLevelIndex = InQualityLevelRange.GetLowerBoundValue(); QualityLevelIndex <= InQualityLevelRange.GetUpperBoundValue(); ++QualityLevelIndex)
+	{
+		for (int32 ShaderPlatformIndex = InShaderPlatformRange.GetLowerBoundValue(); ShaderPlatformIndex <= InShaderPlatformRange.GetUpperBoundValue(); ++ShaderPlatformIndex)
+		{
+			if (!FDataDrivenShaderPlatformInfo::IsValid((EShaderPlatform)ShaderPlatformIndex))
+			{
+				continue;
+			}
+
+			const FMaterialResource* CurrentResource = InMaterial->GetMaterialResource((EShaderPlatform)ShaderPlatformIndex, (EMaterialQualityLevel::Type)QualityLevelIndex);
+			if (CurrentResource)
+			{
+				MatchedResources.AddUnique(CurrentResource);
+			}
+		}
+	}
+
+	return MatchedResources;
+}
+
 #if WITH_EDITORONLY_DATA
 bool UMaterialInterface::GetStaticSwitchParameterValue(const FHashedMaterialParameterInfo& ParameterInfo, bool& OutValue, FGuid& OutExpressionGuid, bool bOveriddenOnly /*= false*/) const
 {
@@ -1897,6 +2008,90 @@ void UMaterialInterface::UpdateMaterialRenderProxy(FMaterialRenderProxy& Proxy)
 	}
 }
 
+#if WITH_EDITOR
+TArray<const FMaterialResource*, TInlineAllocator<4>> UMaterialInterface::GetCookedCacheMaterialResources() const
+{
+	TArray<const FMaterialResource*, TInlineAllocator<4>> CookedResources;
+	for (const auto& KVP : CachedMaterialResourcesForCooking)
+	{
+		const TArray<FMaterialResourceForCooking>& PlatformResources = KVP.Value;
+
+		for (const FMaterialResourceForCooking& ResForCooking : PlatformResources)
+		{
+			if (ResForCooking.Resource)
+			{
+				CookedResources.Add(ResForCooking.Resource);
+			}
+		}
+	}
+
+	return CookedResources;
+}
+
+void UMaterialInterface::RemoveUnusedDefaultTextures(const bool bInKeepProjectPrivateTextures)
+{
+	TArray<TObjectPtr<UObject>> UsedTextures;
+	GetUsedTexturesInCookCache(UsedTextures);
+
+	for (int32 Idx = 0; Idx < ReferencedDefaultTextures.Num(); ++Idx)
+	{
+		const TObjectPtr<UObject> CurrentTexture = ReferencedDefaultTextures[Idx];
+
+		if (!CurrentTexture)
+		{
+			continue;
+		}
+
+		// Check for project private textures (PKC_NotExternallyReferenceable)
+		const UPackage* TexturePackage = CurrentTexture->GetPackage();
+		const UPackage* OwnPackage = GetPackage();
+
+		if (!TexturePackage || !OwnPackage)
+		{
+			continue;
+		}
+
+		const bool bUsedTexture = UsedTextures.Contains(CurrentTexture);
+		bool bRemoveTexture = !bUsedTexture;
+		const bool bIsTexturePrivate = TexturePackage->GetAssetAccessSpecifier() == EAssetAccessSpecifier::Private;
+
+		if (!bInKeepProjectPrivateTextures)
+		{
+			// We are in a material instance, check for private package and mount points
+			if (bIsTexturePrivate)
+			{
+				const FName OwnMountPointName = FPackageName::GetPackageMountPoint(FNameBuilder(OwnPackage->GetFName()).ToView());
+				const FName TextureMountPointName = FPackageName::GetPackageMountPoint(FNameBuilder(TexturePackage->GetFName()).ToView());
+
+				if (!OwnMountPointName.IsNone() && !TextureMountPointName.IsNone())
+				{
+					if (OwnMountPointName != TextureMountPointName)
+					{
+						// Private texture package, different mount points -> remove texture
+						bRemoveTexture = true;
+					}
+				}
+				else
+				{
+					// Private texture package and None mount points -> remove texture
+					bRemoveTexture = true;
+				}
+			}
+		}
+		else
+		{
+			// We are in a material, we only remove if the texture is not private
+			bRemoveTexture &= !bIsTexturePrivate;
+		}
+
+		if (bRemoveTexture)
+		{
+			ReferencedDefaultTextures[Idx] = nullptr;
+		}
+	}
+}
+#endif // WITH_EDITOR
+
 bool FMaterialTextureInfo::IsValid(bool bCheckTextureIndex) const
 { 
 #if WITH_EDITORONLY_DATA
@@ -2193,6 +2388,48 @@ void UMaterialInterface::FilterOutPlatformShadingModels(const EShaderPlatform In
 	}
 }
 
+FInt32Range UMaterialInterface::BuildQualityLevelRange(TOptional<EMaterialQualityLevel::Type>  InQualityLevel)
+{
+	FInt32Range QualityLevelRange(0, EMaterialQualityLevel::Num - 1);
+
+	if (InQualityLevel.IsSet())
+	{
+		EMaterialQualityLevel::Type QualityLevel = InQualityLevel.GetValue();
+		ensure(QualityLevel != EMaterialQualityLevel::Num);
+
+		if (QualityLevel == EMaterialQualityLevel::Num)
+		{
+			QualityLevel = GetCurrentMaterialQualityLevelChecked();
+		}
+		QualityLevelRange = FInt32Range(QualityLevel, QualityLevel);
+	}
+
+	return QualityLevelRange;
+}
+
+FInt32Range UMaterialInterface::BuildShaderPlatformRange(TOptional<EShaderPlatform> InShaderPlatform)
+{
+#if WITH_EDITOR
+	FInt32Range ShaderPlatformRange(0, EShaderPlatform::SP_NumPlatforms - 1);
+
+	if (InShaderPlatform.IsSet())
+	{
+		EShaderPlatform ShaderPlatform = InShaderPlatform.GetValue();
+		ensure(ShaderPlatform != EShaderPlatform::SP_NumPlatforms);
+
+		if (ShaderPlatform == EShaderPlatform::SP_NumPlatforms)
+		{
+			ShaderPlatform = GetMaxShaderPlatformChecked();
+		}
+		ShaderPlatformRange = FInt32Range(ShaderPlatform, ShaderPlatform);
+	}
+#else
+	FInt32Range ShaderPlatformRange(GMaxRHIShaderPlatform, GMaxRHIShaderPlatform);
+#endif
+
+	return ShaderPlatformRange;
+}
+
 #if WITH_EDITORONLY_DATA
 
 
diff --git a/Engine/Source/Runtime/Engine/Private/Materials/MaterialShared.cpp b/Engine/Source/Runtime/Engine/Private/Materials/MaterialShared.cpp
--- a/Engine/Source/Runtime/Engine/Private/Materials/MaterialShared.cpp
+++ b/Engine/Source/Runtime/Engine/Private/Materials/MaterialShared.cpp
@@ -1567,6 +1567,65 @@ void FMaterialResource::LegacySerialize(FArchive& Ar)
 	}
 }
 
+TObjectPtr<UObject> FMaterialResource::GetMatDefaultTextureAtIdx(const int32 InIdx) const
+{
+	if (MaterialInstance)
+	{
+		TObjectPtr<UObject> Texture =  MaterialInstance->GetDefaultTextureAtIdx(InIdx);
+		if (Texture)
+		{
+			return Texture;
+		}
+	}
+
+	if (Material)
+	{
+		return Material->GetDefaultTextureAtIdx(InIdx);
+	}
+
+	return UMaterial::GetDefaultMaterial(MD_Surface)->GetDefaultTextureAtIdx(InIdx);
+}
+
+#if WITH_EDITOR
+int32 FMaterialResource::GetMatDefaultTextureIdx(TObjectPtr<UObject> InTexture) const
+{
+	if (MaterialInstance)
+	{
+		const int32 Idx = MaterialInstance->GetDefaultTextureIdx(InTexture);
+		if (Idx != INDEX_NONE)
+		{
+			return Idx;
+		}
+	}
+
+	if (Material)
+	{
+		return Material->GetDefaultTextureIdx(InTexture);
+	}
+
+	return UMaterial::GetDefaultMaterial(MD_Surface)->GetDefaultTextureIdx(InTexture);
+}
+
+const TArray<TObjectPtr<UObject>>& FMaterialResource::GetMatDefaultTextures() const
+{
+	if (MaterialInstance)
+	{
+		const TArray<TObjectPtr<UObject>>& Textures = MaterialInstance->GetDefaultTextures();
+		if (Textures.Num())
+		{
+			return Textures;
+		}
+	}
+
+	if (Material)
+	{
+		return Material->GetDefaultTextures();
+	}
+
+	return UMaterial::GetDefaultMaterial(MD_Surface)->GetDefaultTextures();
+}
+#endif // WITH_EDITOR
+
 TArrayView<const TObjectPtr<UObject>> FMaterialResource::GetReferencedTextures() const
 {
 	if (MaterialInstance)
@@ -4521,7 +4580,7 @@ void FMaterial::GetReferencedTexturesHash(FSHAHash& OutHash) const
 {
 	FSHA1 HashState;
 
-	const TArrayView<const TObjectPtr<UObject>> ReferencedTextures = GetReferencedTextures();
+	const TArray<TObjectPtr<UObject>>& ReferencedTextures = GetMatDefaultTextures();
 	// Hash the names of the uniform expression textures to capture changes in their order or values resulting from material compiler code changes
 	for (int32 TextureIndex = 0; TextureIndex < ReferencedTextures.Num(); TextureIndex++)
 	{
diff --git a/Engine/Source/Runtime/Engine/Private/Materials/MaterialUniformExpressions.h b/Engine/Source/Runtime/Engine/Private/Materials/MaterialUniformExpressions.h
--- a/Engine/Source/Runtime/Engine/Private/Materials/MaterialUniformExpressions.h
+++ b/Engine/Source/Runtime/Engine/Private/Materials/MaterialUniformExpressions.h
@@ -450,22 +450,7 @@ private:
 template<typename TextureType>
 static TextureType* GetIndexedTexture(const FMaterial& Material, int32 TextureIndex)
 {
-	UObject* IndexedTexture = nullptr;
-	const TArrayView<const TObjectPtr<UObject>> ReferencedTextures = Material.GetReferencedTextures();
-	if (ReferencedTextures.IsValidIndex(TextureIndex))
-	{
-		IndexedTexture = ReferencedTextures[TextureIndex];
-	}
-
-	if (IndexedTexture == nullptr)
-	{
-		static bool bWarnedOnce = false;
-		if (!bWarnedOnce)
-		{
-			UE_LOG(LogMaterial, Warning, TEXT("%s: Requesting an invalid TextureIndex! (%u / %u)"), *Material.GetFriendlyName(), TextureIndex, ReferencedTextures.Num());
-			bWarnedOnce = true;
-		}
-	}
+	UObject* IndexedTexture = Material.GetMatDefaultTextureAtIdx(TextureIndex);
 
 	// Can return nullptr if TextureType doesn't match type of indexed texture
 	return Cast<TextureType>(IndexedTexture);
diff --git a/Engine/Source/Runtime/Engine/Private/SparseVolumeTexture/SparseVolumeTexture.cpp b/Engine/Source/Runtime/Engine/Private/SparseVolumeTexture/SparseVolumeTexture.cpp
--- a/Engine/Source/Runtime/Engine/Private/SparseVolumeTexture/SparseVolumeTexture.cpp
+++ b/Engine/Source/Runtime/Engine/Private/SparseVolumeTexture/SparseVolumeTexture.cpp
@@ -1043,7 +1043,7 @@ void USparseVolumeTexture::NotifyMaterials(const ENotifyMaterialsEffectOnShaders
 		for (TObjectIterator<UMaterialInterface> It; It; ++It)
 		{
 			UMaterialInterface* MaterialInterface = *It;
-			if (!FPlatformProperties::IsServerOnly() && MaterialInterface->GetReferencedTextures().Contains(this))
+			if (!FPlatformProperties::IsServerOnly() && (MaterialInterface->ContainsDefaultTexture(this)))
 			{
 				UpdateContext.AddMaterialInterface(MaterialInterface);
 				// This is a bit tricky. We want to make sure all materials using this texture are
diff --git a/Engine/Source/Runtime/Engine/Public/MaterialCachedData.h b/Engine/Source/Runtime/Engine/Public/MaterialCachedData.h
--- a/Engine/Source/Runtime/Engine/Public/MaterialCachedData.h
+++ b/Engine/Source/Runtime/Engine/Public/MaterialCachedData.h
@@ -199,6 +199,10 @@ struct FMaterialCachedExpressionEditorOnlyData
 
 	UPROPERTY()
 	TArray<FName> ShaderTags;
+
+	/** All the default textures present in the nodes of the material this represents. */
+	UPROPERTY()
+	TArray<TObjectPtr<UObject>> EditorReferencedDefaultTextures;
 };
 
 USTRUCT()
@@ -268,6 +272,11 @@ struct FMaterialCachedExpressionData
 	void GetAllGlobalParametersOfType(EMaterialParameterType Type, TMap<FMaterialParameterInfo, FMaterialParameterMetadata>& OutParameters) const;
 	void GetAllGlobalParameterInfoOfType(EMaterialParameterType Type, TArray<FMaterialParameterInfo>& OutParameterInfo, TArray<FGuid>& OutParameterIds) const;
 
+#if WITH_EDITOR
+	/** @return The index at which the given texture sits in the default textures array or INDEX_NONE if not present. */
+	int32 GetDefaultTextureIdx(UObject* InTexture) const;
+#endif
+
 	ENGINE_API void AddReferencedObjects(FReferenceCollector& Collector);
 
 	/** Returns an array of the guids of functions used, with the call hierarchy flattened. */
@@ -341,10 +350,6 @@ struct FMaterialCachedExpressionData
 	UPROPERTY()
 	TArray<TSoftObjectPtr<USparseVolumeTexture>> SparseVolumeTextureValues;
 
-	/** Array of all texture referenced by this material */
-	UPROPERTY()
-	TArray<TObjectPtr<UObject>> ReferencedTextures;
-
 	UPROPERTY()
 	TArray<TObjectPtr<UTextureCollection>> ReferencedTextureCollections;
 
diff --git a/Engine/Source/Runtime/Engine/Public/MaterialShared.h b/Engine/Source/Runtime/Engine/Public/MaterialShared.h
--- a/Engine/Source/Runtime/Engine/Public/MaterialShared.h
+++ b/Engine/Source/Runtime/Engine/Public/MaterialShared.h
@@ -2673,6 +2673,13 @@ public:
 	ENGINE_API virtual void AddReferencedObjects(FReferenceCollector& Collector);
 
 	virtual TArrayView<const TObjectPtr<UObject>> GetReferencedTextures() const = 0;
+	virtual TObjectPtr<UObject> GetMatDefaultTextureAtIdx(const int32 inIdx) const = 0;
+
+#if WITH_EDITOR
+	virtual int32 GetMatDefaultTextureIdx(TObjectPtr<UObject> InTexture) const = 0;
+	virtual const TArray<TObjectPtr<UObject>>& GetMatDefaultTextures() const = 0;
+#endif
+
 	virtual TConstArrayView<TObjectPtr<UTextureCollection>> GetReferencedTextureCollections() const = 0;
 
 	/**
@@ -3261,6 +3268,12 @@ public:
 	ENGINE_API virtual FName GetAssetPath() const override;
 	ENGINE_API virtual bool IsUsingNewHLSLGenerator() const override;
 	ENGINE_API virtual bool CheckInValidStateForCompilation(FMaterialCompiler* Compiler) const override;
+
+	/** @return The index at which the given texture sits in the default textures array for this material. */
+	ENGINE_API virtual int32 GetMatDefaultTextureIdx(TObjectPtr<UObject> InTexture) const override;
+
+	/** @return All default textures present in this material. */
+	ENGINE_API virtual const TArray<TObjectPtr<UObject>>& GetMatDefaultTextures() const override;
 #endif // WITH_EDITOR
 
 	ENGINE_API virtual FString GetFullPath() const override;
@@ -3269,6 +3282,14 @@ public:
 
 	ENGINE_API virtual void LegacySerialize(FArchive& Ar) override;
 
+	/**
+	 * Returns the relevant texture for the given Idx.
+	 * At editor time, it looks through all default textures. At runtime, it has access to default textures based on compilation output.
+	 * @return The texture present at the given Idx or nullptr if not available.
+	 */
+	ENGINE_API virtual TObjectPtr<UObject> GetMatDefaultTextureAtIdx(const int32 InIdx) const override;
+
+	UE_DEPRECATED(5.8, "Please use the GetMatDefaultTextures for Editor or GetMatDefaultTextureAtIdx for runtime.")
 	ENGINE_API virtual TArrayView<const TObjectPtr<UObject>> GetReferencedTextures() const override;
 	ENGINE_API virtual TConstArrayView<TObjectPtr<UTextureCollection>> GetReferencedTextureCollections() const override;
 
diff --git a/Engine/Source/Runtime/Engine/Public/Materials/Material.h b/Engine/Source/Runtime/Engine/Public/Materials/Material.h
--- a/Engine/Source/Runtime/Engine/Public/Materials/Material.h
+++ b/Engine/Source/Runtime/Engine/Public/Materials/Material.h
@@ -1216,10 +1216,7 @@ private:
 	 * These are always valid and non-null, but only the entries affected by CacheResourceShadersForRendering are actually valid for rendering.
 	 */
 	TArray<FMaterialResource*> MaterialResources;
-#if WITH_EDITOR
-	/** Material resources being cached for cooking. */
-	TMap<const ITargetPlatform*, TArray<FMaterialResourceForCooking>> CachedMaterialResourcesForCooking;
-#endif
+
 	/** Flag used to guarantee that the RT is finished using various resources in this UMaterial before cleanup. */
 	FThreadSafeBool ReleasedByRT;
 
@@ -1351,6 +1348,9 @@ public:
 	ENGINE_API virtual int32 CompilePropertyEx( class FMaterialCompiler* Compiler, const FGuid& AttributeID ) override;
 	ENGINE_API virtual bool ShouldForcePlanePreview() override;
 	ENGINE_API virtual void ForceRecompileForRendering(EMaterialShaderPrecompileMode CompileMode = EMaterialShaderPrecompileMode::Default) override;
+
+	/* Returns all default textures used by the material, based on the compilation output. */
+	ENGINE_API virtual void GetUsedTexturesInCookCache(TArray<TObjectPtr<UObject>>& OutTextures) override;
 #endif // WITH_EDITOR
 	//~ End UMaterialInterface Interface.
 
@@ -1923,6 +1923,12 @@ protected:
 #endif
 
 public:
+	/**
+	 * Looks through the ReferencedDefaultTextures array for the texture present at the given Idx.
+	 * At editor time, it looks through all default textures. At runtime, it has access to default textures based on compilation output.
+	 * @return The default texture that is relevant to the given Idx or nullptr if not present.
+	 */
+	ENGINE_API virtual TObjectPtr<UObject> GetDefaultTextureAtIdx(const int32 InIdx, TMicRecursionGuard RecursionGuard = TMicRecursionGuard()) override;
 	ENGINE_API virtual void DumpDebugInfo(FOutputDevice& OutputDevice) const override;
 	void SaveShaderStableKeys(const class ITargetPlatform* TP);
 	ENGINE_API virtual void SaveShaderStableKeysInner(const class ITargetPlatform* TP, const struct FStableShaderKeyAndValue& SaveKeyVal) override;
diff --git a/Engine/Source/Runtime/Engine/Public/Materials/MaterialInstance.h b/Engine/Source/Runtime/Engine/Public/Materials/MaterialInstance.h
--- a/Engine/Source/Runtime/Engine/Public/Materials/MaterialInstance.h
+++ b/Engine/Source/Runtime/Engine/Public/Materials/MaterialInstance.h
@@ -867,10 +867,6 @@ private:
 	 * For example the material instance needs to support being rendered at different quality levels and feature levels within the same process.
 	 */
 	TArray<FMaterialResource*> StaticPermutationMaterialResources;
-#if WITH_EDITOR
-	/** Material resources being cached for cooking. */
-	TMap<const ITargetPlatform*, TArray<FMaterialResourceForCooking>> CachedMaterialResourcesForCooking;
-#endif
 
 #if WITH_EDITOR
 	struct FTextureParameterOverride
@@ -994,6 +990,10 @@ public:
 #if WITH_EDITOR
 	/** Allows material properties to be compiled with the option of being overridden by the material attributes input. */
 	ENGINE_API virtual int32 CompilePropertyEx(class FMaterialCompiler* Compiler, const FGuid& AttributeID) override;
+
+	/* Returns all default textures used by the material, based on the compilation output. */
+	ENGINE_API virtual void GetUsedTexturesInCookCache(TArray<TObjectPtr<UObject>>& OutTextures) override;
+
 #endif // WITH_EDITOR
 	//~ End UMaterialInterface Interface.
 
@@ -1104,6 +1104,9 @@ public:
 	ENGINE_API virtual void GetReferencedTexturesAndOverrides(TSet<const UTexture*>& InOutTextures) const;
 
 	ENGINE_API virtual void UpdateCachedData();
+
+	/** Creates and returns a CachedExpressionData that analyzes this material. */
+	TUniquePtr<FMaterialCachedExpressionData> CreateCachedExpressionData() const;
 #endif
 
 	ENGINE_API void GetBasePropertyOverridesHash(FSHAHash& OutHash)const;
@@ -1145,6 +1148,13 @@ public:
 	 */
 	ENGINE_API virtual void GetLightingGuidChain(bool bIncludeTextures, TArray<FGuid>& OutGuids) const override;
 
+	/**
+	 * Looks through this Instance and all its parents for the relevant texture.
+	 * At editor time, it looks through all default textures. At runtime, it has access to default textures based on compilation output.
+	 * @return The default texture that is relevant to the given Idx or nullptr if not present.
+	 */
+	ENGINE_API virtual TObjectPtr<UObject> GetDefaultTextureAtIdx(const int32 InIdx, TMicRecursionGuard RecursionGuard = TMicRecursionGuard()) override;
+
 	ENGINE_API virtual void DumpDebugInfo(FOutputDevice& OutputDevice) const override;
 	void SaveShaderStableKeys(const class ITargetPlatform* TP);
 	ENGINE_API virtual void SaveShaderStableKeysInner(const class ITargetPlatform* TP, const struct FStableShaderKeyAndValue& SaveKeyVal) override;
@@ -1192,7 +1202,8 @@ protected:
 
 	ENGINE_API bool SetParentInternal(class UMaterialInterface* NewParent, bool RecacheShaders);
 
-	void GetTextureExpressionValues(const FMaterialResource* MaterialResource, TArray<UTexture*>& OutTextures, TArray<TArray<int32>>* OutIndices = nullptr, bool bIncludeTextureCollections = false) const;
+	/** Returns the textures which are used by the Uniform Expressions of the given Material Resource, along with any existing transient overrides. */
+	ENGINE_API virtual void GetTextureExpressionValues(const FMaterialResource* MaterialResource, TArray<UTexture*>& OutTextures, TArray<TArray<int32>>* OutIndices = nullptr, bool bIncludeTextureCollections = false, bool bAllowOverrides = false) const override;
 	void GetTextureCollectionExpressionValues(const FMaterialResource* MaterialResource, TArray<UTextureCollection*>& OutTextureCollections) const;
 
 #if WITH_EDITOR
diff --git a/Engine/Source/Runtime/Engine/Public/Materials/MaterialInterface.h b/Engine/Source/Runtime/Engine/Public/Materials/MaterialInterface.h
--- a/Engine/Source/Runtime/Engine/Public/Materials/MaterialInterface.h
+++ b/Engine/Source/Runtime/Engine/Public/Materials/MaterialInterface.h
@@ -52,7 +52,6 @@ class UMaterialInstance;
 struct FDebugShaderTypeInfo;
 struct FMaterialParameterInfo;
 struct FMaterialResourceLocOnDisk;
-class FMaterialCachedData;
 struct FMaterialCachedExpressionData;
 struct FMaterialCachedExpressionEditorOnlyData;
 #if WITH_EDITOR
@@ -391,6 +390,10 @@ protected:
 	UPROPERTY()
 	TArray<FMaterialTextureInfo> TextureStreamingData;
 
+	/** Default Textures referenced at runtime based on compilation output. */
+	UPROPERTY()
+	TArray<TObjectPtr<UObject>> ReferencedDefaultTextures;
+
 	/** Array of user data stored with the asset */
 	UPROPERTY(EditAnywhere, AdvancedDisplay, Instanced, Category = Material)
 	TArray<TObjectPtr<UAssetUserData>> AssetUserData;
@@ -447,6 +450,12 @@ public:
 	UPROPERTY(EditAnywhere, Instanced, Category = ImportSettings)
 	TObjectPtr<class UAssetImportData> AssetImportData;
 	
+protected:
+#if WITH_EDITOR
+		/** Material resources being cached for cooking. */
+		TMap<const ITargetPlatform*, TArray<FMaterialResourceForCooking>> CachedMaterialResourcesForCooking;
+#endif
+
 private:
 	/** Unique ID for this material, used for caching during distributed lighting */
 	UPROPERTY()
@@ -523,6 +532,17 @@ public:
 
 	ENGINE_API virtual const FMaterialCachedExpressionData& GetCachedExpressionData(TMicRecursionGuard RecursionGuard = TMicRecursionGuard()) const;
 
+#if WITH_EDITOR
+	/** @return The Index at which the given texture is present in the editor only Default Textures array. */
+	ENGINE_API int32 GetDefaultTextureIdx(UObject* InTexture) const;
+
+	/** @return Whether the given texture is present in the editor only Default Textures. */
+	ENGINE_API bool ContainsDefaultTexture(UObject* InTexture) const;
+
+	/** @return The array of default textures present at editor time. */
+	ENGINE_API const TArray<TObjectPtr<UObject>>& GetDefaultTextures() const;
+#endif
+
 	ENGINE_API bool IsUsingNewHLSLGenerator() const;
 	ENGINE_API bool IsUsingNewTranslatorPrototype() const;
 
@@ -608,6 +628,9 @@ public:
 	ENGINE_API virtual void GetUsedTexturesAndIndices(TArray<UTexture*>& OutTextures, TArray< TArray<int32> >& OutIndices, EMaterialQualityLevel::Type QualityLevel, ERHIFeatureLevel::Type FeatureLevel) const final;
 	ENGINE_API virtual void GetUsedTexturesAndIndices(TArray<UTexture*>& OutTextures, TArray< TArray<int32> >& OutIndices, EMaterialQualityLevel::Type QualityLevel, EShaderPlatform InShaderPlatform) const;
 
+	/** Returns the textures which are used by the Uniform Expressions of the given Material Resource. */
+	ENGINE_API virtual void GetTextureExpressionValues(const FMaterialResource* MaterialResource, TArray<UTexture*>& OutTextures, TArray<TArray<int32>>* OutIndices = nullptr, bool bIncludeTextureCollections = false, bool bAllowOverrides = false) const;
+
 	/**
 	 * Override a specific texture (transient)
 	 *
@@ -676,6 +699,12 @@ public:
 	virtual const FMaterialResource* GetMaterialResource(ERHIFeatureLevel::Type InFeatureLevel, EMaterialQualityLevel::Type QualityLevel = EMaterialQualityLevel::Num) const final { return NULL; }
 	virtual const FMaterialResource* GetMaterialResource(EShaderPlatform InShaderPlatform, EMaterialQualityLevel::Type QualityLevel = EMaterialQualityLevel::Num) const { return NULL; }
 
+	/**
+	 * Get the Material Resources created by the given Material, restricted to the given Quality Range and Shader Platform Range
+	 * @return - The appropriate FMaterialResources, if any exist
+	 */
+	static TArray<const FMaterialResource*, TInlineAllocator<4>> GetMaterialResources(const UMaterialInterface* InMaterial, const FInt32Range InQualityLevelRange, const FInt32Range InShaderPlatformRange);
+
 	/**
 	 * Get the material layers stack
 	 * @return - material layers stack, or nullptr if material has no layers
@@ -809,9 +838,11 @@ public:
 	ENGINE_API virtual void GetReferencedTexturesAndOverrides(TSet<const UTexture*>& InOutTextures) const;
 #endif // WITH_EDITOR
 
-	/** Get textures referenced by expressions, including nested functions. */
 	ENGINE_API TArrayView<const TObjectPtr<UObject>> GetReferencedTextures() const;
 
+	/** @return The default texture relevant to the given Idx. These are present based the compilation output. */
+	virtual TObjectPtr<UObject> GetDefaultTextureAtIdx(const int32 InIdx, TMicRecursionGuard RecursionGuard = TMicRecursionGuard()) PURE_VIRTUAL(UMaterialInterface::GetDefaultTextureAtIdx, return nullptr;);
+
 	ENGINE_API TConstArrayView<TObjectPtr<UTextureCollection>> GetReferencedTextureCollections() const;
 
 	virtual void SaveShaderStableKeysInner(const class ITargetPlatform* TP, const struct FStableShaderKeyAndValue& SaveKeyVal)
@@ -1339,6 +1370,12 @@ public:
 
 	/** Filter out ShadingModels field to a shader platform settings */
 	ENGINE_API static void FilterOutPlatformShadingModels(EShaderPlatform Platform, FMaterialShadingModelField& ShadingModels);
+
+	/** Helper function for building a Range from a given Quality Level */
+	static FInt32Range BuildQualityLevelRange(TOptional<EMaterialQualityLevel::Type> InQualityLevel);
+
+	/** Helper function for building a Range from a given Shader Platform */
+	static FInt32Range BuildShaderPlatformRange(TOptional<EShaderPlatform> InShaderPlatform);
 	
 #if WITH_EDITOR
 	TUniquePtr<FMaterialInsights> MaterialInsight;
@@ -1363,6 +1400,17 @@ protected:
 	/** Bitfield to track which usage flags have triggered warnings so that we only warn once per usage. */
 	uint32 UsageFlagWarnings = 0;
 
+#if WITH_EDITOR
+	/* Returns all default textures used by the material, based on the compilation output. */
+	virtual void GetUsedTexturesInCookCache(TArray<TObjectPtr<UObject>>& OutTextures) PURE_VIRTUAL(UMaterialInterface::GetUsedTexturesInCookCache, return;);
+	
+	/** @return All FMaterialResources created during the cooking process. */
+	TArray<const FMaterialResource*, TInlineAllocator<4>> GetCookedCacheMaterialResources() const;
+
+	/** Removes all unused default textures based on the Cook Shader Compilation output. */
+	void RemoveUnusedDefaultTextures(const bool bInKeepProjectPrivateTextures);
+#endif
+
 private:
 	/** Bitfield to track game thread dispatch of usage flag checks. */
 	mutable std::atomic<uint32> UsageFlagCacheMiss { 0u };
