[NetPhysics] Simulation Decay, first implementation.
- Integrate less of accelerations onto velocities and integrate less of velocities onto position and rotation to slow down the simulation for simulated proxies.
- Making simulated proxies not fully forward predicted which improves their replication due to not mispredicting inputs while forward predicting.

diff --git a/Engine/Source/Runtime/Engine/Private/PhysicsEngine/NetworkPhysicsComponent.cpp b/Engine/Source/Runtime/Engine/Private/PhysicsEngine/NetworkPhysicsComponent.cpp
--- a/Engine/Source/Runtime/Engine/Private/PhysicsEngine/NetworkPhysicsComponent.cpp
+++ b/Engine/Source/Runtime/Engine/Private/PhysicsEngine/NetworkPhysicsComponent.cpp
@@ -46,7 +46,7 @@ namespace PhysicsReplicationCVars
 		static FAutoConsoleVariableRef CVarDynamicInputBufferScalingEnabled(TEXT("np2.Resim.DynamicInputBufferScaling.Enabled"), bDynamicInputBufferScalingEnabled, TEXT("Enable dynmic scaling of input buffer on the server."));
 		float DynamicInputBufferScalingMinBufferMs = 30.0f;
 		static FAutoConsoleVariableRef CVarDynamicInputBufferScalingMinBufferMs(TEXT("np2.Resim.DynamicInputBufferScaling.MinBufferTime"), DynamicInputBufferScalingMinBufferMs, TEXT("Time in milliseconds for the lowest allowed input buffer size, when below this the buffer will scale up over time."));
-		float DynamicInputBufferScalingMaxBufferMs = 150.0f;
+		float DynamicInputBufferScalingMaxBufferMs = 90.0f;
 		static FAutoConsoleVariableRef CVarDynamicInputBufferScalingMaxBufferMs(TEXT("np2.Resim.DynamicInputBufferScaling.MaxBufferTime"), DynamicInputBufferScalingMaxBufferMs, TEXT("Time in milliseconds to cap out the input buffer size."));
 		float DynamicInputBufferScalingAverageTime = 0.5f;
 		static FAutoConsoleVariableRef CVarDynamicInputBufferScalingAverageTime(TEXT("np2.Resim.DynamicInputBufferScaling.AverageTime"), DynamicInputBufferScalingAverageTime, TEXT("Time in seconds to keep a running average of the input buffer size over."));
@@ -86,6 +86,8 @@ namespace PhysicsReplicationCVars
 		bool bTriggerResimOnInputReceive = false;
 		static FAutoConsoleVariableRef CVarTriggerResimOnInputReceive(TEXT("np2.Resim.TriggerResimOnInputReceive"), bTriggerResimOnInputReceive, TEXT("When true, a resim will be requested to the frame of the latest frame of received inputs this frame"));
 
+		bool bEnableInputDecay = true;
+		static FAutoConsoleVariableRef CVarEnableInputDecay(TEXT("np2.Resim.EnableInputDecay"), bEnableInputDecay, TEXT("When true, apply the Input Decay on predicted inputs."));
 		bool bApplyInputDecayOverSetTime = false;
 		static FAutoConsoleVariableRef CVarApplyInputDecayOverSetTime(TEXT("np2.Resim.ApplyInputDecayOverSetTime"), bApplyInputDecayOverSetTime, TEXT("When true, apply the Input Decay Curve over a set amount of time instead of over the start of input prediction and end of resim which is variable each resimulation"));
 		float InputDecaySetTime = 0.15f;
@@ -95,6 +97,23 @@ namespace PhysicsReplicationCVars
 		float InputDecaySimProxyInputAtRuntime = 0.25f;
 		static FAutoConsoleVariableRef CVarInputDecaySimProxyInputAtRuntime(TEXT("np2.Resim.InputDecaySimProxyInputAtRuntime"), InputDecaySimProxyInputAtRuntime, TEXT("The amount of input decay to apply on simulated proxy inputs at runtime (outside of resim), if np2.Resim.ApplyInputDecaySimProxyInputAtRuntime is true"));
 
+		bool bSimulationDecayEnable = false;
+		static FAutoConsoleVariableRef CVarEnableSimulationDecay(TEXT("np2.Resim.SimulationDecay.Enable"), bSimulationDecayEnable, TEXT("When enabled, apply simulation decay on simulated proxies by essentially multiplying their physics step by a time scale value. Effectively forward predicting sim-proxies less, which can make desync less during forward prediction."));
+		bool bSimulationDecayApplyAtRuntime = true;
+		static FAutoConsoleVariableRef CVarApplySimulationDecayAtRuntime(TEXT("np2.Resim.SimulationDecay.ApplyAtRuntime"), bSimulationDecayApplyAtRuntime, TEXT("Apply simulation decay during non-resimulation frames."));
+		bool bSimulationDecayDynamicTimeScale = true;
+		static FAutoConsoleVariableRef CVarSimulationDecayDynamicTimeScale(TEXT("np2.Resim.SimulationDecay.DynamicTimeScale"), bSimulationDecayDynamicTimeScale, TEXT("Calculate the simulation decay timescale during resimulation to based on the average time between forward predicted timeline and the timeline the sim-proxy receive inputs in."));
+		float SimulationDecayTimeScaleDefault = 0.9f;
+		static FAutoConsoleVariableRef CVarSimulationDecayTimeScaleDefault(TEXT("np2.Resim.SimulationDecay.TimeScaleDefault"), SimulationDecayTimeScaleDefault, TEXT("Time scale multiplier applied at runtime and during resimulation if dynamic time scale is disabled."));
+		float SimulationDecayTimeScaleBase = 0.1f;
+		static FAutoConsoleVariableRef CVarSimulationDecayTimeScaleBase(TEXT("np2.Resim.SimulationDecay.TimeScaleBase"), SimulationDecayTimeScaleBase, TEXT("Base time scale value added onto dynamic calculated timescale."));
+		float SimulationDecayTimeScaleMin = 0.25f;
+		static FAutoConsoleVariableRef CVarSimulationDecayTimeScaleMin(TEXT("np2.Resim.SimulationDecay.TimeScaleMin"), SimulationDecayTimeScaleMin, TEXT("Minimum time scale allowed during dynamic scaling, clamped 0-1."));
+		float SimulationDecayTimeScaleMax = 1.0f;
+		static FAutoConsoleVariableRef CVarSimulationDecayTimeScaleMax(TEXT("np2.Resim.SimulationDecay.TimeScaleMax"), SimulationDecayTimeScaleMax, TEXT("Maximum time scale allowed during dynamic scaling, clamped 0-1."));
+		float SimulationDecayAverageOverTime = 2.0f;
+		static FAutoConsoleVariableRef CVarSimulationDecayAverageOverTime(TEXT("np2.Resim.SimulationDecay.AverageOverTime"), SimulationDecayAverageOverTime, TEXT("How long time in seconds to keep an average of the input prediction frames for, which is the base of dynamic time scaling. Value is the number of frames behind the forward predicted timeline that we receive inputs at."));
+
 		bool bEnableStatefulDeltaSerialization = true;
 		static FAutoConsoleVariableRef CVarResimEnableStatefulDeltaSerialization(TEXT("np2.Resim.StatefulDeltaSerialization.Enable"), bEnableStatefulDeltaSerialization, TEXT("Enables stateful delta serialization for FNetworkPhysicsData derived inputs and states. During FNetworkPhysicsData::NetSerialize there will be a valid pointer to a previous FNetworkPhysicsData which can be used for delta serialization, FNetworkPhysicsData::DeltaSourceData. NOTE: Switching this during gameplay might cause disconnections."));
 		bool bUseDefaultDeltaForDeltaSourceReplication = true;
@@ -2019,6 +2038,20 @@ void FAsyncNetworkPhysicsComponent::ConsumeAsyncInput(const int32 PhysicsStep)
 				}
 				ReceiveHelper(StateHistory.Get(), StateImportant.Get(), /*bImportant*/true, bCompareState);
 			}
+
+			// Running average of how far behind the forward predicted timeline we receive inputs
+			if (Chaos::FPBDRigidsSolver* RigidsSolver = GetRigidSolver())
+			{
+				if (!bIsServer && !IsLocallyControlled())
+				{
+					const float Dt = RigidsSolver->GetAsyncDeltaTime();
+					const int32 InputPredictedFrames = FMath::Max(RigidsSolver->GetCurrentFrame() - InputHistory->GetLatestFrame(), 0);
+
+					const float AverageOverSeconds = PhysicsReplicationCVars::ResimulationCVars::SimulationDecayAverageOverTime;
+					const float AverageAlpha = FMath::Clamp(Dt / (Dt + AverageOverSeconds), 0.0f, 1.0f);
+					InputPredictionFramesAverage += (static_cast<float>(InputPredictedFrames) - InputPredictionFramesAverage) * AverageAlpha;
+				}
+			}
 		}
 	}
 }
@@ -2047,11 +2080,6 @@ void FAsyncNetworkPhysicsComponent::OnPreProcessInputs_Internal(const int32 Phys
 {
 	QUICK_SCOPE_CYCLE_COUNTER(AsyncNetworkPhysicsComponent_OnPreProcessInputs_Internal);
 
-	ConsumeAsyncInput(PhysicsStep);
-
-	const FNetworkPhysicsSettingsNetworkPhysicsComponent& ComponentSettings = GetComponentSettings();
-	const bool bIsServer = IsServer();
-
 	bool bIsSolverReset = false;
 	bool bIsSolverResim = false;
 	if (Chaos::FPBDRigidsEvolution* Evolution = GetEvolution())
@@ -2060,6 +2088,14 @@ void FAsyncNetworkPhysicsComponent::OnPreProcessInputs_Internal(const int32 Phys
 		bIsSolverReset = Evolution->IsResetting();
 	}
 
+	if (!bIsSolverResim)
+	{
+		ConsumeAsyncInput(PhysicsStep);
+	}
+
+	const FNetworkPhysicsSettingsNetworkPhysicsComponent& ComponentSettings = GetComponentSettings();
+	const bool bIsServer = IsServer();
+
 #if DEBUG_NETWORK_PHYSICS
 	{
 		const int32 InputBufferSize = (bIsServer && InputHistory) ? (InputHistory->GetLatestFrame() - PhysicsStep) : 0;
@@ -2181,7 +2217,7 @@ void FAsyncNetworkPhysicsComponent::OnPreProcessInputs_Internal(const int32 Phys
 				// Calculate input decay if we are resimulating and we don't have up to date inputs
 				if (bIsSolverResim)
 				{
-					if (InputHistory->GetLatestFrame() < PhysicsStep)
+					if (PhysicsReplicationCVars::ResimulationCVars::bEnableInputDecay && InputHistory->GetLatestFrame() < PhysicsStep)
 					{
 						const float InputDecay = GetCurrentInputDecay();
 						PhysicsData->DecayData(InputDecay);
@@ -2247,7 +2283,7 @@ void FAsyncNetworkPhysicsComponent::OnPreProcessInputs_Internal(const int32 Phys
 				}
 
 				// Perform input decay on sim-proxy runtime (non-resim) input
-				if (!bIsSolverResim && bApplySimProxyInput && ComponentSettings.GetApplyInputDecaySimProxyInputAtRuntime())
+				if (PhysicsReplicationCVars::ResimulationCVars::bEnableInputDecay && !bIsSolverResim && bApplySimProxyInput && ComponentSettings.GetApplyInputDecaySimProxyInputAtRuntime())
 				{
 					PhysicsData->DecayData(ComponentSettings.GetInputDecaySimProxyInputAtRuntime());
 				}
@@ -2397,6 +2433,84 @@ void FAsyncNetworkPhysicsComponent::OnPostProcessInputs_Internal(const int32 Phy
 	}
 }
 
+void FAsyncNetworkPhysicsComponent::OnPostIntegrate_Internal()
+{
+	Chaos::FPBDRigidsSolver* RigidsSolver = GetRigidSolver();
+	if (!RigidsSolver || !RigidsSolver->GetEvolution())
+	{
+		return;
+	}
+
+	/** Simulation Decay */
+	{
+		if (!PhysicsReplicationCVars::ResimulationCVars::bSimulationDecayEnable || !PhysicsObject)
+		{
+			return;
+		}
+
+		if (IsServer() || IsLocallyControlled())
+		{
+			// Don't run on the server or autonomous proxy
+			return;
+		}
+
+		float TimeScale = PhysicsReplicationCVars::ResimulationCVars::SimulationDecayTimeScaleDefault;
+		Chaos::FRewindData* RewindData = RigidsSolver->GetRewindData();
+		if (RewindData && RewindData->IsResim())
+		{
+			if (PhysicsReplicationCVars::ResimulationCVars::bSimulationDecayDynamicTimeScale)
+			{
+				// Calculate the time scale to use during resimulation, based on how far behind the forward predicted timeline we receive inputs
+				const float ResimFrames = static_cast<float>(RewindData->GetLatestFrame() - RewindData->GetResimFrame());
+				const float TargetPredictionFrame = static_cast<float>(RigidsSolver->GetCurrentFrame()) - InputPredictionFramesAverage;
+				if (ResimFrames > 0)
+				{
+					const float BaseTimeScale = FMath::Clamp(PhysicsReplicationCVars::ResimulationCVars::SimulationDecayTimeScaleBase, 0.0f, 1.0f);
+					const float MinTimeScale = FMath::Clamp(PhysicsReplicationCVars::ResimulationCVars::SimulationDecayTimeScaleMin, 0.0f, 1.0f);
+					const float MaxTimeScale = FMath::Clamp(PhysicsReplicationCVars::ResimulationCVars::SimulationDecayTimeScaleMax, 0.0f, 1.0f);
+					TimeScale = FMath::Clamp(BaseTimeScale + ((ResimFrames - InputPredictionFramesAverage) / ResimFrames), MinTimeScale, MaxTimeScale);
+				}
+			}
+		}
+		else if (!PhysicsReplicationCVars::ResimulationCVars::bSimulationDecayApplyAtRuntime)
+		{
+			// Early out if we should not apply simulation decay on non-resim frames
+			return;
+		}
+
+
+		Chaos::FWritePhysicsObjectInterface_Internal Interface = Chaos::FPhysicsObjectInternalInterface::GetWrite();
+		if (Chaos::FPBDRigidParticleHandle* Particle = Interface.GetRigidParticle(PhysicsObject))
+		{
+			// Only apply on dynamic particles
+			if (Particle->IsDynamic())
+			{
+				// Pre Integration State
+				const Chaos::FVec3 PrePos = Particle->GetTransformXRCom().GetLocation();
+				const Chaos::FRotation3 PreRot = Particle->GetTransformXRCom().GetRotation();
+				const Chaos::FVec3 PreVel = Particle->GetPreV();
+				const Chaos::FVec3 PreAngVel = Particle->GetPreW();
+
+				// Post Integration State
+				const Chaos::FVec3 PostPos = Particle->GetTransformPQCom().GetLocation();
+				const Chaos::FRotation3 PostRot = Particle->GetTransformPQCom().GetRotation();
+				const Chaos::FVec3 PostVel = Particle->GetV();
+				const Chaos::FVec3 PostAngVel = Particle->GetW();
+
+				// Adjust integration result by time scaling the change between Pre and Post
+				const Chaos::FVec3 NewPos = PrePos + TimeScale * (PostPos - PrePos);
+				const Chaos::FRotation3 NewRot = Chaos::FRotation3::Slerp(PreRot, PostRot, TimeScale).GetNormalized();
+				const Chaos::FVec3 NewVel = PreVel + TimeScale * (PostVel - PreVel);
+				const Chaos::FVec3 NewAngVel = PreAngVel + TimeScale * (PostAngVel - PreAngVel);
+
+				Particle->SetTransformPQCom(NewPos, NewRot);
+				Particle->SetV(NewVel);
+				Particle->SetW(NewAngVel);
+			}
+		}
+	}
+}
+
 void FAsyncNetworkPhysicsComponent::SendInputData_Internal(FAsyncNetworkPhysicsComponentOutput& AsyncOutput, const int32 PhysicsStep)
 {
 	const bool bIsServer = IsServer();
diff --git a/Engine/Source/Runtime/Engine/Public/Physics/NetworkPhysicsComponent.h b/Engine/Source/Runtime/Engine/Public/Physics/NetworkPhysicsComponent.h
--- a/Engine/Source/Runtime/Engine/Public/Physics/NetworkPhysicsComponent.h
+++ b/Engine/Source/Runtime/Engine/Public/Physics/NetworkPhysicsComponent.h
@@ -1889,7 +1889,7 @@ struct FAsyncNetworkPhysicsComponentOutput : public Chaos::FSimCallbackOutput
 class FAsyncNetworkPhysicsComponent : public Chaos::TSimCallbackObject<
 	FAsyncNetworkPhysicsComponentInput,
 	FAsyncNetworkPhysicsComponentOutput,
-	Chaos::ESimCallbackOptions::PhysicsObjectUnregister>
+	Chaos::ESimCallbackOptions::Presimulate |Chaos::ESimCallbackOptions::PostIntegrate | Chaos::ESimCallbackOptions::PhysicsObjectUnregister>
 {
 	friend UNetworkPhysicsComponent;
 
@@ -1937,6 +1937,10 @@ private:
 	// Uninitialize, unbind delegates etc.
 	void OnUninitialize_Internal();
 
+	virtual void OnPreSimulate_Internal() override { }
+
+	virtual void OnPostIntegrate_Internal() override;
+
 	virtual void OnPhysicsObjectUnregistered_Internal(Chaos::FConstPhysicsObjectHandle PhysicsObject) override;
 
 	// Delegate ran at the start of FNetworkPhysicsCallback::ProcessInputs_Internal, used to receive and apply inputs and states if needed.
@@ -2047,6 +2051,9 @@ private:
 	float InputBufferTargetSizeMs = 90.0;
 	float InputBufferDynamicMinSizeMs = 0.0f;
 
+	// Simulation Decay properties
+	float InputPredictionFramesAverage = 0.0f;
+
 	// Cache predicted states and then compare incoming states via FNetworkPhysicsData::CompareData to trigger a resim if they desync
 	bool bCompareStateToTriggerRewind;
 
