Increase size with single vertex in case we have odd number of texcoords because the vertex declaration will declare the stream as half4/float4 while the last fetched stream will only contains a half2/float2 and the last element can in theory be read out of bounds then.

Vertex streams are all created with half4/float4 to reduce the number of PSO permutations coming from different texcoords in the vertex declaration

diff --git a/Engine/Source/Runtime/Engine/Private/Rendering/StaticMeshVertexBuffer.cpp b/Engine/Source/Runtime/Engine/Private/Rendering/StaticMeshVertexBuffer.cpp
--- a/Engine/Source/Runtime/Engine/Private/Rendering/StaticMeshVertexBuffer.cpp
+++ b/Engine/Source/Runtime/Engine/Private/Rendering/StaticMeshVertexBuffer.cpp
@@ -284,7 +284,13 @@ FBufferRHIRef FStaticMeshVertexBuffer::CreateTangentsRHIBuffer(FRHICommandListBa
 
 FBufferRHIRef FStaticMeshVertexBuffer::CreateTexCoordRHIBuffer(FRHICommandListBase& RHICmdList)
 {
-	return CreateRHIBuffer(RHICmdList, TexcoordData, GetNumVertices() * GetNumTexCoords(), BUF_Static | BUF_ShaderResource, TEXT("TexCoordRHIBuffer"));
+	// Increase size with single vertex in case we have odd number of texcoords because the vertex declaration
+	// will declare the stream as half4/float4 while it only contains a half2/float2 and the last element can in theory
+	// be out of bounds then
+	// Vertex streams are all created with half4/float4 to reduce the number of PSO permutations coming from different texcoords in the vertex declaration
+	uint32 Padding = (GetNumTexCoords() % 2 == 0) ? 0 : 1;
+	uint32 ResourceCount = (GetNumVertices() + Padding) * GetNumTexCoords();
+	return CreateRHIBuffer(RHICmdList, TexcoordData, ResourceCount, BUF_Static | BUF_ShaderResource, TEXT("TexCoordRHIBuffer"));
 }
 
 void FStaticMeshVertexBuffer::InitRHIForStreaming(
