Iris
* Minor api update to NetRootObjectAdapter
* Don't ensure if attaching to the same level as before

#rnx

diff --git a/Engine/Source/Runtime/Engine/Private/Net/Iris/ReplicationSystem/NetRootObjectAdapter.cpp b/Engine/Source/Runtime/Engine/Private/Net/Iris/ReplicationSystem/NetRootObjectAdapter.cpp
--- a/Engine/Source/Runtime/Engine/Private/Net/Iris/ReplicationSystem/NetRootObjectAdapter.cpp
+++ b/Engine/Source/Runtime/Engine/Private/Net/Iris/ReplicationSystem/NetRootObjectAdapter.cpp
@@ -67,6 +67,14 @@ void FNetRootObjectAdapter::Configure(const FRootObjectSettings& Settings)
 
 void FNetRootObjectAdapter::SetAttachedLevel(ULevel* Level)
 {
+	TWeakObjectPtr<ULevel> NewWeakLevel(Level);
+
+	// Nothing to do if the level didn't change.
+	if (NewWeakLevel.IsValid() && NewWeakLevel == WeakLevel)
+	{
+		return;
+	}
+
 	if (bIsReplicating)
 	{
 		ensureMsgf(!bIsReplicating, TEXT("Cannot set attached level after object: %s is replicating"), *GetNameSafe(WeakReplicatedObject.Get()));
@@ -269,7 +277,7 @@ void FNetRootObjectAdapter::StopReplication()
 	bIsReplicating = false;
 }
 
-void FNetRootObjectAdapter::RelevantWith(UObject* OtherRootObject)
+void FNetRootObjectAdapter::RelevantWith(const UObject* OtherRootObject)
 {
 	check(OtherRootObject);
 
@@ -350,7 +358,7 @@ void FNetRootObjectAdapter::RelevantWith(UObject* OtherRootObject)
 	});
 }
 
-void FNetRootObjectAdapter::RemoveRelevantWith(UObject* OtherRootObject)
+void FNetRootObjectAdapter::RemoveRelevantWith(const UObject* OtherRootObject)
 {
 	UObject* ReplicatedObject = WeakReplicatedObject.Get();
 	if (!ReplicatedObject)
diff --git a/Engine/Source/Runtime/Engine/Public/Net/Iris/ReplicationSystem/NetRootObjectAdapter.h b/Engine/Source/Runtime/Engine/Public/Net/Iris/ReplicationSystem/NetRootObjectAdapter.h
--- a/Engine/Source/Runtime/Engine/Public/Net/Iris/ReplicationSystem/NetRootObjectAdapter.h
+++ b/Engine/Source/Runtime/Engine/Public/Net/Iris/ReplicationSystem/NetRootObjectAdapter.h
@@ -78,13 +78,13 @@ public:
 	 * You can have relevancy links to multiple different objects. As long as of those is relevant to a client then this object will be too.
 	 * @OtherRootObject:  A replicated root object we want to be tie our relevancy status to.
 	 */
-	ENGINE_API void RelevantWith(UObject* OtherRootObject);
+	ENGINE_API void RelevantWith(const UObject* OtherRootObject);
 
 	/**
 	 * Remove the relevancy dependency between our object and another root object
 	 * @OtherRootObject: A replicated root object that was previously in a RelevantWith call.
 	 */
-	ENGINE_API void RemoveRelevantWith(UObject* OtherRootObject);
+	ENGINE_API void RemoveRelevantWith(const UObject* OtherRootObject);
 
 	/** The netobject factory to use with this rootobject. By default it will use the UNetRootObjectFactory. */
 	ENGINE_API FName GetNetFactoryName() const;
@@ -92,6 +92,8 @@ public:
 	/** Set to use a netobject factory different from the default. Must be set before replication starts. */
 	ENGINE_API void SetNetFactoryName(FName InFactoryName);
 
+	// Status / Config getters
+
 	bool IsInitialized() const
 	{
 		return bIsInitialized;
