Fixing a server crash during Subsystem cleanup. Changed nullptr check to IsValid check due to rare case where Subsystem deinitialization can point to a garbage collected object.

#tests Built and tested against Dedicated Server with Cooked Data
#rnx

diff --git a/Engine/Source/Runtime/Engine/Private/Subsystems/SubsystemCollection.cpp b/Engine/Source/Runtime/Engine/Private/Subsystems/SubsystemCollection.cpp
--- a/Engine/Source/Runtime/Engine/Private/Subsystems/SubsystemCollection.cpp
+++ b/Engine/Source/Runtime/Engine/Private/Subsystems/SubsystemCollection.cpp
@@ -272,7 +272,7 @@ void FSubsystemCollectionBase::Deinitialize()
 	{
 		UClass* KeyClass = Iter.Key();
 		USubsystem* Subsystem = Iter.Value();
-		if (Subsystem != nullptr && Subsystem->GetClass() == KeyClass)
+		if (IsValid(Subsystem) && Subsystem->GetClass() == KeyClass)
 		{
 			Subsystem->Deinitialize();
 			Subsystem->InternalOwningSubsystem = nullptr;
