UE: Fix crash at shutdown in the editor when closing a -log window or using ctrl-c

diff --git a/Engine/Source/Runtime/Engine/Private/Analytics/EditorAnalyticsSessionSummary.cpp b/Engine/Source/Runtime/Engine/Private/Analytics/EditorAnalyticsSessionSummary.cpp
--- a/Engine/Source/Runtime/Engine/Private/Analytics/EditorAnalyticsSessionSummary.cpp
+++ b/Engine/Source/Runtime/Engine/Private/Analytics/EditorAnalyticsSessionSummary.cpp
@@ -202,6 +202,18 @@ namespace EditorAnalyticsSessionSummaryHelpers
 
 void FEditorAnalyticsSessionSummary::UpdateUserFocusTimes(double CurrTimeSecs)
 {
+	// This can get called from a console control handler thread. Getting the SlateApplication is only valid on certain threads.
+	if (!IsInGameThread() && !IsInSlateThread() && !IsInAsyncLoadingThread() && !IsInParallelLoadingThread())
+	{
+		return;
+	}
+
+	// Make sure we have a slate application to query
+	if (!FSlateApplication::IsInitialized())
+	{
+		return;
+	}
+
 	FSlateApplication& SlateApplication = FSlateApplication::Get();
 
 	TSharedPtr<const FSlateUser> CurrentUser = SlateApplication.GetCursorUser();
