FNewSpline: Harden attribute channel parameterization when reparameterizing parent spline.

diff --git a/Engine/Source/Runtime/Engine/Private/Curves/Spline.cpp b/Engine/Source/Runtime/Engine/Private/Curves/Spline.cpp
--- a/Engine/Source/Runtime/Engine/Private/Curves/Spline.cpp
+++ b/Engine/Source/Runtime/Engine/Private/Curves/Spline.cpp
@@ -376,6 +376,10 @@ private:
 	int32 ConvertInternalParameterToNearestPointIndex(float Parameter) const;
 
 	void ValidateRotScale() const;
+
+	// Helper functions for keeping attribute channels stable when reparameterizing.
+	template <typename AttrType> TArray<AttrType> CacheChannelKnots(const FName& InName) const;
+	template <typename AttrType> void RestoreChannelKnots(const FName& InName, const TArray<AttrType>& InKnots);
 };
 
 template <typename AttrType>
@@ -596,11 +600,8 @@ int32 FNewSpline::SetAttributeParameter(const FName& Name, int32 Index, float Pa
 		}
 		else if (Index == 0)									// Case 3: Moving the first point (which has exactly 1 neighbor)
 		{
-			TArray<float> InternalParameters;
-			for (int InternalIdx = 1; InternalIdx < Channel->GetNumPoints() - 1; ++InternalIdx)
-			{
-				InternalParameters.Add(MapParameterFromChildSpace(Name, Channel->GetParameter(InternalIdx)));
-			}
+			// Save internal attribute point parameters for later re-parameterization after changing the child space mapping.
+			TArray<float> InternalParameters = CacheChannelKnots<float>(Name);
 
 			float ParentSpaceUpperBound = GetMappedChildSpace(Name).Max;
 			const float NeighborParentSpaceParameter = ToInternalSplineSpace(GetAttributeParameter<AttrType>(Name, Index + 1));
@@ -637,10 +638,7 @@ int32 FNewSpline::SetAttributeParameter(const FName& Name, int32 Index, float Pa
 			SetAttributeChannelRange(Name, FInterval1f(ParentSpaceLowerBound, ParentSpaceUpperBound), EMappingRangeSpace::Parent);
 
 			// Re-parameterize the internal points so that they do not move in parent space after the mapping range changed
-			for (int InternalIdx = 1; InternalIdx < Channel->GetNumPoints() - 1; ++InternalIdx)
-			{
-				Channel->SetParameter(InternalIdx, MapParameterToChildSpace(Name, InternalParameters[InternalIdx - 1]));
-			}
+			RestoreChannelKnots(Name, InternalParameters);
 			
 			if (bAttributesWillReorder)
 			{
@@ -652,11 +650,7 @@ int32 FNewSpline::SetAttributeParameter(const FName& Name, int32 Index, float Pa
 		else if (Index == Channel->GetNumPoints() - 1)			// Case 4: Moving the last point (which has exactly 1 neighbor)
 		{
 			// Save internal attribute point parameters for later re-parameterization after changing the child space mapping.
-			TArray<float> InternalParameters;
-			for (int InternalIdx = 1; InternalIdx < Channel->GetNumPoints() - 1; ++InternalIdx)
-			{
-				InternalParameters.Add(MapParameterFromChildSpace(Name, Channel->GetParameter(InternalIdx)));
-			}
+			TArray<float> InternalParameters = CacheChannelKnots<float>(Name);
 
 			float ParentSpaceLowerBound = GetMappedChildSpace(Name).Min;
 			const float NeighborParentSpaceParameter = ToInternalSplineSpace(GetAttributeParameter<AttrType>(Name, Index - 1));
@@ -693,10 +687,7 @@ int32 FNewSpline::SetAttributeParameter(const FName& Name, int32 Index, float Pa
 			SetAttributeChannelRange(Name, FInterval1f(ParentSpaceLowerBound, ParentSpaceUpperBound), EMappingRangeSpace::Parent);
 			
 			// Re-parameterize the internal points so that they do not move in parent space after the mapping range changed
-			for (int InternalIdx = 1; InternalIdx < Channel->GetNumPoints() - 1; ++InternalIdx)
-			{
-				Channel->SetParameter(InternalIdx, MapParameterToChildSpace(Name, InternalParameters[InternalIdx - 1]));
-			}
+			RestoreChannelKnots(Name, InternalParameters);
 
 			if (bAttributesWillReorder)
 			{
@@ -1329,7 +1320,7 @@ bool FSpline::CreateAttributeChannel(FName AttributeName) const
 		{
 			NewAttrChannel->SetPreInfinityMode(UE::Geometry::Spline::EOutOfBoundsHandlingMode::Constant);
 			NewAttrChannel->SetPostInfinityMode(UE::Geometry::Spline::EOutOfBoundsHandlingMode::Constant);
-		
+
 			return true;
 		}
 	}
@@ -2536,9 +2527,51 @@ void FNewSpline::UpdateSpline(const FSpline::FUpdateSplineParams& Params)
 
 void FNewSpline::Reparameterize()
 {
+	// Caching segment space attribute parameters so we can restore them (necessary because we store the mapping in parent space which are about to be invalidated).
+	TMap<FName, TPair<FInterval1f, TArray<float>>> CachedChannelParameterizations;
+	for (const FName& ChannelName : GetAttributeChannelNamesByValueType<float>())
+	{
+		if (ChannelName == RotationAttrName || ChannelName == ScaleAttrName || NumAttributeValues<float>(ChannelName) == 0)
+		{
+			continue;
+		}
+
+		ensureAlwaysMsgf(GetAttributeChannelRangeSpace(ChannelName) == EMappingRangeSpace::Parent, TEXT("Attribute channel mapping range space is not Parent!"));
+
+		FInterval1f ChannelRangeInParentSpace = GetAttributeChannelRange(ChannelName);
+		FInterval1f ChannelRangeInSegmentSpace = FInterval1f(FromInternalSplineSpace(ChannelRangeInParentSpace.Min), FromInternalSplineSpace(ChannelRangeInParentSpace.Max));
+
+		CachedChannelParameterizations.FindOrAdd(ChannelName) = TPair<FInterval1f, TArray<float>>(ChannelRangeInSegmentSpace, CacheChannelKnots<float>(ChannelName));
+
+		for (float& InternalKnot : CachedChannelParameterizations[ChannelName].Value)
+		{
+			InternalKnot = FromInternalSplineSpace(InternalKnot);
+		}
+	}
+
 	GetSpline().Reparameterize(GetParameterizationPolicy());
 	GetSpline().UpdateTangents();
 
+	for (const FName& ChannelName : GetAttributeChannelNamesByValueType<float>())
+	{
+		if (ChannelName == RotationAttrName || ChannelName == ScaleAttrName || NumAttributeValues<float>(ChannelName) == 0)
+		{
+			continue;
+		}
+
+		const FInterval1f& ChannelRangeInSegmentSpace = CachedChannelParameterizations[ChannelName].Key;
+		const FInterval1f ChannelRangeInParentSpace = FInterval1f(ToInternalSplineSpace(ChannelRangeInSegmentSpace.Min), ToInternalSplineSpace(ChannelRangeInSegmentSpace.Max));
+
+		SetAttributeChannelRange(ChannelName, ChannelRangeInParentSpace, EMappingRangeSpace::Parent);
+
+		for (float& InternalKnot : CachedChannelParameterizations[ChannelName].Value)
+		{
+			InternalKnot = ToInternalSplineSpace(InternalKnot);
+		}
+
+		RestoreChannelKnots(ChannelName, CachedChannelParameterizations[ChannelName].Value);
+	}
+
 	if (UE::Geometry::Spline::TTangentBezierSpline<FVector>* ScaleSpline = GetTypedAttributeChannel<UE::Geometry::Spline::TTangentBezierSpline<FVector>>(ScaleAttrName))
 	{
 		ScaleSpline->SetKnotVector(GetSpline().GetKnotVector());
@@ -2794,6 +2827,44 @@ void FNewSpline::ValidateRotScale() const
 	}
 }
 
+template <typename AttrType>
+TArray<AttrType> FNewSpline::CacheChannelKnots(const FName& InName) const
+{
+	using AttrSplineType = UE::Geometry::Spline::TTangentBezierSpline<AttrType>;
+
+	TArray<float> InternalParameters;
+
+	if (AttrSplineType* Channel = GetTypedAttributeChannel<AttrSplineType>(InName))
+	{
+		for (int InternalIdx = 1; InternalIdx < Channel->GetNumPoints() - 1; ++InternalIdx)
+		{
+			InternalParameters.Add(MapParameterFromChildSpace(InName, Channel->GetParameter(InternalIdx)));
+		}
+	}
+		
+	return InternalParameters;
+}
+
+template <typename AttrType>
+void FNewSpline::RestoreChannelKnots(const FName& InName, const TArray<AttrType>& InKnots)
+{
+	using AttrSplineType = UE::Geometry::Spline::TTangentBezierSpline<AttrType>;
+
+	if (InKnots.Num() == 0) return;
+
+	if (AttrSplineType* Channel = GetTypedAttributeChannel<AttrSplineType>(InName))
+	{
+		// Provided knots are for internal points, bail out if we don't have the correct number of knots.
+		if (ensureAlwaysMsgf(InKnots.Num() == Channel->GetNumPoints() - 2, TEXT("Invalid number of knots!")))
+		{
+			for (int InternalIdx = 1; InternalIdx < Channel->GetNumPoints() - 1; ++InternalIdx)
+			{
+				Channel->SetParameter(InternalIdx, MapParameterToChildSpace(InName, InKnots[InternalIdx - 1]));
+			}
+		}
+	}
+}
+
 /** FLegacySpline Implementation */
 
 FLegacySpline::FLegacySpline(const FNewSpline& Other)
