SplineComponent: Always make AFunctionalSplineTestBase a friend class

diff --git a/Engine/Source/Runtime/Engine/Classes/Components/SplineComponent.h b/Engine/Source/Runtime/Engine/Classes/Components/SplineComponent.h
--- a/Engine/Source/Runtime/Engine/Classes/Components/SplineComponent.h
+++ b/Engine/Source/Runtime/Engine/Classes/Components/SplineComponent.h
@@ -1092,10 +1092,8 @@ PRAGMA_ENABLE_DEPRECATION_WARNINGS
 	// FSplineComponentVisualizer will access some private members when attempting to call NotifyPropertiesModified
 	friend class FSplineComponentVisualizer;
 
-#if WITH_DEV_AUTOMATION_TESTS
 	// AFunctionalSplineTestBase accesses private methods for functional testing.
 	friend class AFunctionalSplineTestBase;
-#endif
 };
 
 /** Used to store spline data during RerunConstructionScripts */
