[HWRT] Fix selecting ray tracing visualization mode via in-editor UI not working when r.RayTracing.Visualize is set to an invalid option.

diff --git a/Engine/Source/Runtime/Engine/Private/RayTracingVisualizationData.cpp b/Engine/Source/Runtime/Engine/Private/RayTracingVisualizationData.cpp
--- a/Engine/Source/Runtime/Engine/Private/RayTracingVisualizationData.cpp
+++ b/Engine/Source/Runtime/Engine/Private/RayTracingVisualizationData.cpp
@@ -163,7 +163,7 @@ bool FRayTracingVisualizationData::HasOverrides() const
 FName FRayTracingVisualizationData::ApplyOverrides(const FName& InModeName) const
 {
 	check(IsInParallelRenderingThread());
-	return GRayTracingVisualizeMode != NAME_None ? GRayTracingVisualizeMode : InModeName;
+	return HasOverrides() ? GRayTracingVisualizeMode : InModeName;
 }
 
 void FRayTracingVisualizationData::ConfigureConsoleCommand()
