Null-check the USkinnedAsset when building the ISKM component desc.

diff --git a/Engine/Source/Runtime/Engine/Public/InstancedSkinnedMeshComponentHelper.h b/Engine/Source/Runtime/Engine/Public/InstancedSkinnedMeshComponentHelper.h
--- a/Engine/Source/Runtime/Engine/Public/InstancedSkinnedMeshComponentHelper.h
+++ b/Engine/Source/Runtime/Engine/Public/InstancedSkinnedMeshComponentHelper.h
@@ -59,16 +59,16 @@ FInstanceDataManagerSourceDataDesc FInstancedSkinnedMeshComponentHelper::GetComp
 
 	ComponentDesc.Flags = Flags;
 
-	// TODO: rename
-	ComponentDesc.MeshBounds = SkinnedAsset->GetBounds();
+	const FTransform& ComponentTransform = InComponent.GetComponentTransform();
+	const FVector3f Position(ComponentTransform.GetTranslation());
+
+	ComponentDesc.MeshBounds = SkinnedAsset ? FRenderBounds(SkinnedAsset->GetBounds()) : FRenderBounds(Position, Position);
 	ComponentDesc.NumCustomDataFloats = InComponent.NumCustomDataFloats;
 	ComponentDesc.NumInstances = InComponent.InstanceData.Num();
 
 	ComponentDesc.PrimitiveLocalToWorld = InComponent.GetRenderMatrix();
 	ComponentDesc.ComponentMobility = InComponent.GetMobility();
 
-	const FTransform& ComponentTransform = InComponent.GetComponentTransform();
-
 	ComponentDesc.BuildChangeSet =
 	[
 		&InComponent,
