SoundWaveProxyPlayer

Reuse SourceAudio buffer to cache Resampler samples

#rnx
#tests EngineTests.WaveProxyPlayer

diff --git a/Engine/Source/Runtime/Engine/Classes/Sound/SoundWaveProxyPlayer.h b/Engine/Source/Runtime/Engine/Classes/Sound/SoundWaveProxyPlayer.h
--- a/Engine/Source/Runtime/Engine/Classes/Sound/SoundWaveProxyPlayer.h
+++ b/Engine/Source/Runtime/Engine/Classes/Sound/SoundWaveProxyPlayer.h
@@ -334,9 +334,6 @@ private:
 	// "source audio" buffer to write "interleaved" audio to before de-interleaving
 	Audio::FAlignedFloatBuffer SourceAudio;
 	
-	// residual samples to pass back in to the resampler in the next render
-	Audio::FAlignedFloatBuffer LastSourceSamples;
-
 	// "resampled audio" buffer to render "Resampled" audio to before de-interleaving
 	Audio::FAlignedFloatBuffer ResampledAudio;
 
diff --git a/Engine/Source/Runtime/Engine/Private/SoundWaveProxyPlayer.cpp b/Engine/Source/Runtime/Engine/Private/SoundWaveProxyPlayer.cpp
--- a/Engine/Source/Runtime/Engine/Private/SoundWaveProxyPlayer.cpp
+++ b/Engine/Source/Runtime/Engine/Private/SoundWaveProxyPlayer.cpp
@@ -65,6 +65,7 @@ bool FSoundWaveProxyPlayer::SetSoundWave(FSoundWaveProxyPtr InWaveProxy)
 	}
 	
 	SoundWave = InWaveProxy;
+	SourceAudio.Reset();
 	SourceFrameIndex = 0;
 	AudioSyncFrameIndex = 0;
 	// update the loop with the cached loop settings to update the loop start/end frames
@@ -320,9 +321,8 @@ void FSoundWaveProxyPlayer::RenderMultiChannelAudio(Audio::FMultichannelBufferVi
 		int32 SourceNumFrames = RuntimeResampler.GetNumInputFramesNeededToProduceOutputFrames(OutputNumFrames);
 
 		// include the samples we used from the last render
-		int32 LastNumSamples = LastSourceSamples.Num();
+		int32 LastNumSamples = SourceAudio.Num();
 		SourceAudio.SetNum(SourceNumFrames * NumChannels);
-		FMemory::Memcpy(SourceAudio.GetData(), LastSourceSamples.GetData(), LastNumSamples * sizeof(float));
 		GenerateSourceAudio(MakeArrayView(SourceAudio.GetData() + LastNumSamples, SourceAudio.Num() - LastNumSamples));
 
 		int32 NumFramesConsumed = 0;
@@ -342,20 +342,14 @@ void FSoundWaveProxyPlayer::RenderMultiChannelAudio(Audio::FMultichannelBufferVi
 		}
 
 		int32 NumSamplesConsumed = NumFramesConsumed * NumChannels;
+		int32 NumSamplesRemaining = SourceAudio.Num() - NumSamplesConsumed;
 		if (NumSamplesConsumed < SourceAudio.Num())
 		{
 			// move the last frame of samples to the start of the SourceBuffer for future processing
-			int32 NumSamplesCopy = SourceAudio.Num() - NumSamplesConsumed;
-			check(NumSamplesCopy == NumChannels);
-			LastSourceSamples.SetNum(NumSamplesCopy);
-			FMemory::Memcpy(LastSourceSamples.GetData(), SourceAudio.GetData() + NumSamplesConsumed, NumSamplesCopy * sizeof(float));
+			check(NumSamplesRemaining == NumChannels);
+			FMemory::Memmove(SourceAudio.GetData(), SourceAudio.GetData() + NumSamplesConsumed, NumSamplesRemaining * sizeof(float));
 		}
-		else
-		{
-			LastSourceSamples.Reset();
-		}
-
-		SourceAudio.Reset();
+		SourceAudio.SetNum(NumSamplesRemaining, EAllowShrinking::No);
 	}
 }
 
@@ -788,5 +782,4 @@ void FSoundWaveProxyPlayer::InitializeResampler(int32 NumInputChannels)
 	bRequiresResampling = false;
 	RuntimeResampler.Reset(NumInputChannels);
 	ResampledAudio.Reset();
-	LastSourceSamples.Empty(NumInputChannels);
 }
