wrap use of TDataSnapshotProducer / TDataSnapshot in #if PLATFORM_DESKTOP to avoid excess system object (locks) on some platforms
#rnx

diff --git a/Engine/Plugins/Runtime/Metasound/Source/MetasoundEngine/Private/MetasoundWavePlayerNode.cpp b/Engine/Plugins/Runtime/Metasound/Source/MetasoundEngine/Private/MetasoundWavePlayerNode.cpp
--- a/Engine/Plugins/Runtime/Metasound/Source/MetasoundEngine/Private/MetasoundWavePlayerNode.cpp
+++ b/Engine/Plugins/Runtime/Metasound/Source/MetasoundEngine/Private/MetasoundWavePlayerNode.cpp
@@ -533,12 +533,12 @@ namespace Metasound
 			// Reset flag to track render cost.
 			bDidWaveRenderThisBlock = false;
 
-#if !UE_BUILD_SHIPPING
+#if PLATFORM_DESKTOP && !UE_BUILD_SHIPPING
 			if (bIsPlaying && WaveAsset->IsSoundWaveValid() && WaveAsset->GetSoundWaveProxy()->TryGetLatest())
 			{
 				StartPlaying();
 			}
-#endif // !UE_BUILD_SHIPPING
+#endif // #if PLATFORM_DESKTOP && !UE_BUILD_SHIPPING
 
 			// Update wave proxy reader with any new looping bounds. 
 			if (WaveProxyReader.IsValid())
@@ -828,7 +828,9 @@ namespace Metasound
 			if (WaveProxy.IsValid())
 			{
 				// Make sure we have up-to-date SoundWaveData
+#if PLATFORM_DESKTOP
 				WaveProxy->TryGetLatest();
+#endif // #if PLATFORM_DESKTOP
 
 				UE_LOG(LogMetaSound, Verbose, TEXT("Starting Sound: '%s'"), *CurrentWaveAsset->GetFName().ToString());
 
diff --git a/Engine/Source/Runtime/Engine/Classes/Sound/SoundWave.h b/Engine/Source/Runtime/Engine/Classes/Sound/SoundWave.h
--- a/Engine/Source/Runtime/Engine/Classes/Sound/SoundWave.h
+++ b/Engine/Source/Runtime/Engine/Classes/Sound/SoundWave.h
@@ -1563,8 +1563,12 @@ public:
 	ENGINE_API void OverrideWaveData(int32 InNumChannels, int32 InNumFrames, int32 InSampleRate);
 
 private:
+#if PLATFORM_DESKTOP
 	TSharedRef<TDataSnapshotProducer<FSoundWaveData>> SoundWaveData = MakeShared<TDataSnapshotProducer<FSoundWaveData>>();
 	TDataSnapshotProducer<FSoundWaveData>& SoundWaveDataPtr = *SoundWaveData;
+#else // #if PLATFORM_DESKTOP
+	TSharedPtr<FSoundWaveData, ESPMode::ThreadSafe> SoundWaveDataPtr{ MakeShared<FSoundWaveData>() };
+#endif // #if PLATFORM_DESKTOP
 
 	friend class FSoundWaveProxy;
 	friend class USoundFactory;
@@ -1818,14 +1822,20 @@ public:
 	{
 		return InProxy.TypeHash;
 	}
-	
+
+#if PLATFORM_DESKTOP
 	bool TryGetLatest()
 	{
 		return SoundWaveDataPtr.GetLatest();
 	}
+#endif // #if PLATFORM_DESKTOP
 	
 private:
+#if PLATFORM_DESKTOP
 	TDataSnapshot<FSoundWaveData> SoundWaveDataPtr;
+#else // #if PLATFORM_DESKTOP
+	TSharedPtr<FSoundWaveData, ESPMode::ThreadSafe> SoundWaveDataPtr;
+#endif // #if PLATFORM_DESKTOP
 	
 	uint32 TypeHash = uint32(INDEX_NONE);
 };
diff --git a/Engine/Source/Runtime/Engine/Private/SoundWave.cpp b/Engine/Source/Runtime/Engine/Private/SoundWave.cpp
--- a/Engine/Source/Runtime/Engine/Private/SoundWave.cpp
+++ b/Engine/Source/Runtime/Engine/Private/SoundWave.cpp
@@ -3381,8 +3381,12 @@ void USoundWave::CreateNewSoundWaveData()
 	// In the future, all clients should use the proxy system and USoundWave should 
 	// be extended to support notifying clients that it has been edited and replaced 
 	// an existing proxy with a new proxy. Clients can then choose whether to
-	// continue playing their original proxy, or retrieve the updated proxy. 
+	// continue playing their original proxy, or retrieve the updated proxy.
+#if PLATFORM_DESKTOP
 	SoundWaveDataPtr.Fork();
+#else // #if PLATFORM_DESKTOP
+	SoundWaveDataPtr = MakeShared<FSoundWaveData, ESPMode::ThreadSafe>();
+#endif // #if PLATFORM_DESKTOP
 	Proxy.Reset();
 
 }
@@ -4255,7 +4259,9 @@ FSoundWaveProxyPtr USoundWave::CreateSoundWaveProxy()
 #endif // #if WITH_EDITORONLY_DATA
 
 	check(SoundWaveDataPtr);
+#if PLATFORM_DESKTOP
 	SoundWaveDataPtr.Publish();
+#endif // #if PLATFORM_DESKTOP
 	return MakeShared<FSoundWaveProxy, ESPMode::ThreadSafe>(this);
 }
 
@@ -4435,6 +4441,7 @@ bool USoundWave::GetChunkData(int32 ChunkIndex, uint8** OutChunkData, bool bMake
 
 void USoundWave::OverrideWaveData(int32 InNumChannels, int32 InNumFrames, int32 InSampleRate)
 {
+#if PLATFORM_DESKTOP
 	check(InNumChannels > 0);
 	check(InSampleRate > 0);
 
@@ -4454,6 +4461,9 @@ void USoundWave::OverrideWaveData(int32 InNumChannels, int32 InNumFrames, int32
 		SoundWaveDataPtr->NameCached = GetFName();
 		Proxy = CreateSoundWaveProxy();
 	}
+#else // #if PLATFORM_DESKTOP
+	return;
+#endif // #if PLATFORM_DESKTOP
 }
 
 uint32 USoundWave::GetInterpolatedCookedFFTDataForTimeInternal(float InTime, uint32 StartingIndex, TArray<FSoundWaveSpectralData>& OutData, bool bLoop)
@@ -5245,7 +5255,11 @@ ESoundWaveLoadingBehavior USoundWave::GetLoadingBehavior(bool bCheckSoundClasses
 }
 
 FSoundWaveProxy::FSoundWaveProxy(USoundWave* InWave)
+#if PLATFORM_DESKTOP
 	: SoundWaveDataPtr(InWave->SoundWaveDataPtr.GetSnapshot())
+#else // #if PLATFORM_DESKTOP
+	: SoundWaveDataPtr(InWave->SoundWaveDataPtr)
+#endif // #if PLATFORM_DESKTOP
 {
 	LLM_SCOPE(ELLMTag::AudioSoundWaveProxies);
 
