[CIS] added missing PLATFORM_DESKTOP scope around SoundWaveProxyPlayer's use of TryGetLatest() on the sound proxy

diff --git a/Engine/Source/Runtime/Engine/Private/SoundWaveProxyPlayer.cpp b/Engine/Source/Runtime/Engine/Private/SoundWaveProxyPlayer.cpp
--- a/Engine/Source/Runtime/Engine/Private/SoundWaveProxyPlayer.cpp
+++ b/Engine/Source/Runtime/Engine/Private/SoundWaveProxyPlayer.cpp
@@ -47,8 +47,10 @@ bool FSoundWaveProxyPlayer::SetSoundWave(FSoundWaveProxyPtr InWaveProxy)
 		return false;
 	}
 
+#if PLATFORM_DESKTOP
 	// make sure we have the most up-to-date data for the Sound Wave
 	InWaveProxy->TryGetLatest();
+#endif // PLATFORM_DESKTOP
 	
 	if (!IsValidSoundWaveProxy(InWaveProxy.ToSharedRef()))
 	{
