Conditionally add mem mark if non exists in UAnimInstance::TriggerAnimNotifies

diff --git a/Engine/Source/Runtime/Engine/Private/Animation/AnimInstance.cpp b/Engine/Source/Runtime/Engine/Private/Animation/AnimInstance.cpp
--- a/Engine/Source/Runtime/Engine/Private/Animation/AnimInstance.cpp
+++ b/Engine/Source/Runtime/Engine/Private/Animation/AnimInstance.cpp
@@ -1716,6 +1716,15 @@ void UAnimInstance::TriggerAnimNotifies(float DeltaSeconds, ENotifyTriggerMode M
 	SCOPE_CYCLE_COUNTER(STAT_AnimTriggerAnimNotifies);
 	USkeletalMeshComponent* SkelMeshComp = GetSkelMeshComponent();
 
+	// This function can be called from outsode of anim evaluation (e.g. the ability system) and seems like it may call into 
+	// code that uses the mem stack allocator. Conditionally construct a mark if there is none.
+	TOptional<FMemMark> MaybeMark;
+	FMemStack& MemStack = FMemStack::Get();
+	if (MemStack.GetNumMarks() == 0)
+	{
+		MaybeMark.Emplace(MemStack);
+	}
+
 	// Array that will replace the 'ActiveAnimNotifyState' at the end of this function.
 	TArray<FAnimNotifyEvent> NewActiveAnimNotifyState;
 	NewActiveAnimNotifyState.Reserve(NotifyQueue.AnimNotifies.Num());
