fix UE_CVLOG_ARROW and UE_CVLOG_ARROW_MAG definition missing when UE_DEBUG_RECORDING_ENABLED is false

diff --git a/Engine/Source/Runtime/Engine/Public/VisualLogger/VisualLogger.h b/Engine/Source/Runtime/Engine/Public/VisualLogger/VisualLogger.h
--- a/Engine/Source/Runtime/Engine/Public/VisualLogger/VisualLogger.h
+++ b/Engine/Source/Runtime/Engine/Public/VisualLogger/VisualLogger.h
@@ -174,6 +174,8 @@ UE_API DECLARE_LOG_CATEGORY_EXTERN(LogVisual, Display, All);
 #define UE_CVLOG_CONVEXPOLY(Condition, LogOwner, CategoryName, Verbosity, Points, Color, Format, ...)
 #define UE_VLOG_ARROW(LogOwner, CategoryName, Verbosity, SegmentStart, SegmentEnd, Color, Format, ...) 
 #define UE_CVLOG_ARROW(Condition, LogOwner, CategoryName, Verbosity, SegmentStart, SegmentEnd, Color, Format, ...) 
+#define UE_VLOG_ARROW_MAG(LogOwner, CategoryName, Verbosity, SegmentStart, SegmentEnd, Color, Mag, Format, ...) 
+#define UE_CVLOG_ARROW_MAG(Condition, LogOwner, CategoryName, Verbosity, SegmentStart, SegmentEnd, Color, Mag, Format, ...) 
 #define UE_VLOG_CIRCLE(LogOwner, CategoryName, Verbosity, Center, UpAxis, Radius, Color, Format, ...)
 #define UE_CVLOG_CIRCLE(Condition, LogOwner, CategoryName, Verbosity, Center, UpAxis, Radius, Color, Format, ...)
 #define UE_VLOG_CIRCLE_THICK(LogOwner, CategoryName, Verbosity, Center, UpAxis, Radius, Color, Thickness, Format, ...)
