[Material Editor]Fix an issue of ParticleSubUV node when connect to a texture object parameter.

diff --git a/Engine/Source/Runtime/Engine/Private/Materials/MaterialExpressions.cpp b/Engine/Source/Runtime/Engine/Private/Materials/MaterialExpressions.cpp
--- a/Engine/Source/Runtime/Engine/Private/Materials/MaterialExpressions.cpp
+++ b/Engine/Source/Runtime/Engine/Private/Materials/MaterialExpressions.cpp
@@ -10221,6 +10221,10 @@ int32 UMaterialExpressionParticleSubUV::Compile(class FMaterialCompiler* Compile
 			return Compiler->Errorf(TEXT("%s"), *SamplerTypeError);
 		}
 		int32 TextureCodeIndex = Compiler->Texture(TextureToCompile, SamplerTypeToUse);
+        if (TextureObject.GetTracedInput().Expression)
+		{
+			TextureCodeIndex = TextureObject.Compile(Compiler);
+		}
 		return ParticleSubUV(Compiler, TextureCodeIndex, SamplerTypeToUse, CompileMipValue0(Compiler), CompileMipValue1(Compiler), MipValueMode, bBlend);
 	}
 	else
