Refactor all component specific PSO precaching code outside of the primitive scene proxy so it can be shared by other component types which are not derived from primitive scene proxy such as decal component

diff --git a/Engine/Plugins/FX/Niagara/Source/Niagara/Private/NiagaraComponent.cpp b/Engine/Plugins/FX/Niagara/Source/Niagara/Private/NiagaraComponent.cpp
--- a/Engine/Plugins/FX/Niagara/Source/Niagara/Private/NiagaraComponent.cpp
+++ b/Engine/Plugins/FX/Niagara/Source/Niagara/Private/NiagaraComponent.cpp
@@ -2441,7 +2441,7 @@ FPrimitiveSceneProxy* UNiagaraComponent::CreateSceneProxy()
 
 		// If PSOs not precached yet then rely on PSOs reuqested by the shared asset itself - this will make sure
 		// those pending requests are boosted if still pending
-		if (!bPSOPrecacheCalled)
+		if (!PSOPrecacheComponentData.IsPSOPrecacheCalled())
 		{
 			PrecacheAssetPSOs(Asset);
 		} 		
diff --git a/Engine/Source/Runtime/Engine/Classes/Components/DecalComponent.h b/Engine/Source/Runtime/Engine/Classes/Components/DecalComponent.h
--- a/Engine/Source/Runtime/Engine/Classes/Components/DecalComponent.h
+++ b/Engine/Source/Runtime/Engine/Classes/Components/DecalComponent.h
@@ -7,7 +7,7 @@
 #include "UObject/ObjectMacros.h"
 #include "Engine/EngineTypes.h"
 #include "Components/SceneComponent.h"
-#include "PSOPrecache.h"
+#include "PSOPrecacheComponent.h"
 
 #include "DecalComponent.generated.h"
 
@@ -125,11 +125,6 @@ public:
 	ENGINE_API virtual class UMaterialInstanceDynamic* CreateDynamicMaterialInstance();
 
 public:
-#if UE_WITH_PSO_PRECACHING
-	/** Atomic int used to track the last PSO precache events */
-	std::atomic<int> LatestPSOPrecacheJobSetCompleted = 0;
-	int32 LatestPSOPrecacheJobSet = 0;
-#endif
 
 	/** The decal proxy. */
 	FDeferredDecalProxy* SceneProxy;
@@ -146,9 +141,15 @@ protected:
 
 	/** Called when the life span of the decal has been exceeded */
 	ENGINE_API virtual void LifeSpanCallback();
+		
+	/** All PSO related data for the component */
+	FPSOPrecacheComponentData PSOPrecacheComponentData;
 
 public:
-	
+		
+	/** Get the PSO precaching data */
+	FPSOPrecacheComponentData& GetPSOPrecacheComponentData() { return PSOPrecacheComponentData; }
+
 	ENGINE_API void SetLifeSpan(const float LifeSpan);
 
 	/**
diff --git a/Engine/Source/Runtime/Engine/Classes/Components/PrimitiveComponent.h b/Engine/Source/Runtime/Engine/Classes/Components/PrimitiveComponent.h
--- a/Engine/Source/Runtime/Engine/Classes/Components/PrimitiveComponent.h
+++ b/Engine/Source/Runtime/Engine/Classes/Components/PrimitiveComponent.h
@@ -31,6 +31,7 @@
 #include "HLOD/HLODLevelExclusion.h"
 #include "Stats/Stats.h"
 #include "PSOPrecacheFwd.h"
+#include "PSOPrecacheComponent.h"
 #include "PipelineStateCache.h"
 #include "MeshDrawCommandStatsDefines.h"
 #include "PrimitiveSceneInfoData.h"
@@ -710,18 +711,8 @@ protected:
 	uint8 bStaticWhenNotMoveable:1;
 
 	uint8 bIgnoreBoundsForEditorFocus : 1;
-
-#if UE_WITH_PSO_PRECACHING
-	/** Whether PrecachedPSOs was called. */
-	uint8 bPrecachedPSOs : 1;
-#endif
-
+	
 public:
-#if UE_WITH_PSO_PRECACHING
-	/** PSOs requested priority */
-	EPSOPrecachePriority PSOPrecacheRequestPriority : 2;
-	static_assert((int)EPSOPrecachePriority::Highest < 1 << 2);
-#endif
 
 #if WITH_EDITOR
 	/** If true, then this component was added to the navigation octree when it was registered. */
@@ -801,16 +792,6 @@ public:
 	UPROPERTY(EditAnywhere, AdvancedDisplay, BlueprintReadOnly, Category=Lighting)
 	FLightingChannels LightingChannels;
 
-#if UE_WITH_PSO_PRECACHING 
-	/** Cached array of material PSO requests which can be used to boost the priority */
-	TArray<FMaterialPSOPrecacheRequestID> MaterialPSOPrecacheRequestIDs;
-	/** Atomic int used to track the last PSO precache events */
-	std::atomic<int> LatestPSOPrecacheJobSetCompleted = 0;
-	int32 LatestPSOPrecacheJobSet = 0;
-	/** Helper flag to check if PSOs have been precached already */
-	uint8 bPSOPrecacheCalled : 1;
-#endif
-
 	/**
 	* Defines how quickly it should be culled. For example buildings should have a low priority, but small dressing should have a high priority.
 	*/
@@ -933,6 +914,9 @@ protected:
 	/** Last time we checked AreAllCollideableDescendantsRelative(), so we can throttle those tests since it rarely changes once false. */
 	float LastCheckedAllCollideableDescendantsTime;
 
+	/** All PSO related data for the component */
+	FPSOPrecacheComponentData PSOPrecacheComponentData;
+
 private:
 	
 	float OcclusionBoundsSlack;
@@ -1003,21 +987,25 @@ public:
 	/** Precache all PSOs which can be used by the primitive component */
 	ENGINE_API virtual void PrecachePSOs() override;
 
-	/** Schedule task to mark render state dirty when the PSO precaching tasks are done */
-	ENGINE_API void RequestRecreateRenderStateWhenPSOPrecacheFinished(const FGraphEventArray& PSOPrecacheCompileEvents);
+	/** Store all the PSO precache data and schedule task to mark when precaching is done so component can be notified about it */
+	ENGINE_API void SetMaterialPSOPrecacheData(TArray<FMaterialPSOPrecacheRequestID>& MaterialPSOPrecacheRequestIDs, const FGraphEventArray& PSOPrecacheCompileEvents);
 
-	/** Check if PSOs are still precaching */
-	ENGINE_API bool IsPSOPrecaching() const;
+	/** Check if PSOs are still precaching */	
+	bool IsPSOPrecaching() const { return PSOPrecacheComponentData.IsPSOPrecaching(); }
 
-	/** Whether the render proxy should fallback to the default material because the PSOs are still precaching */
-	ENGINE_API bool ShouldRenderProxyFallbackToDefaultMaterial() const;
+	/** Get the PSO precaching data */
+	FPSOPrecacheComponentData& GetPSOPrecacheComponentData() { return PSOPrecacheComponentData; }
+	const FPSOPrecacheComponentData& GetPSOPrecacheComponentData() const { return PSOPrecacheComponentData; }
 
 	/**
 	 * Check if PSOs are still precaching and boost priority if not done yet.
 	 * PSOPrecachePriority specifies the new priority to boost to. calls that would reduce the current priority are ignored, highest pri is used for components that are actively rendering.
 	 * Returns true if the PSOs are still precaching.
 	 */
-	ENGINE_API bool CheckPSOPrecachingAndBoostPriority(EPSOPrecachePriority PSOPrecachePriority = EPSOPrecachePriority::High);
+	bool CheckPSOPrecachingAndBoostPriority(EPSOPrecachePriority PSOPrecachePriority = EPSOPrecachePriority::High)
+	{
+		return PSOPrecacheComponentData.CheckPSOPrecachingAndBoostPriority(PSOPrecachePriority);
+	}
 
 protected:
 
diff --git a/Engine/Source/Runtime/Engine/Private/Components/DecalComponent.cpp b/Engine/Source/Runtime/Engine/Private/Components/DecalComponent.cpp
--- a/Engine/Source/Runtime/Engine/Private/Components/DecalComponent.cpp
+++ b/Engine/Source/Runtime/Engine/Private/Components/DecalComponent.cpp
@@ -399,55 +399,22 @@ void UDecalComponent::PrecachePSOs()
 
 	if (DecalMaterial && !DecalMaterial->HasAnyFlags(RF_NeedPostLoad))
 	{
+		PSOPrecacheComponentData.Reset();
+
 		FPSOPrecacheParams PSOPrecacheParams;		
 		FPSOPrecacheVertexFactoryDataList VertexFactoryDataList;		
 		VertexFactoryDataList.Add(FPSOPrecacheVertexFactoryData(&FLocalVertexFactory::StaticType));
 
-
 		// Immediately create at high priority and thus doesn't need boosting anymore
 		TArray<FMaterialPSOPrecacheRequestID> MaterialPSOPrecacheRequestIDs;
-		FGraphEventArray GraphEvents = DecalMaterial->PrecachePSOs(VertexFactoryDataList, PSOPrecacheParams, EPSOPrecachePriority::High, MaterialPSOPrecacheRequestIDs);
-
+		FGraphEventArray PSOPrecacheCompileEvents = DecalMaterial->PrecachePSOs(VertexFactoryDataList, PSOPrecacheParams, EPSOPrecachePriority::High, MaterialPSOPrecacheRequestIDs);
+		
 		// Request recreate of the render state when the PSO compilation is ready (if we want to delay proxy creation)
-		if (GetPSOPrecacheProxyCreationStrategy() != EPSOPrecacheProxyCreationStrategy::AlwaysCreate)
+		auto Callback = [this]()
 		{
-			struct FPSODecalPrecacheFinishedTask
-			{
-				explicit FPSODecalPrecacheFinishedTask(UDecalComponent* InDecalComponent, int32 InJobSetThatJustCompleted)
-					: WeakDecalComponent(InDecalComponent),
-					JobSetThatJustCompleted(InJobSetThatJustCompleted)
-				{
-				}
-
-				static TStatId GetStatId() { return TStatId(); }
-				static ENamedThreads::Type GetDesiredThread() { return ENamedThreads::GameThread; }
-				static ESubsequentsMode::Type GetSubsequentsMode() { return ESubsequentsMode::TrackSubsequents; }
-
-				void DoTask(ENamedThreads::Type CurrentThread, const FGraphEventRef& MyCompletionGraphEvent)
-				{
-					if (UDecalComponent* DC = WeakDecalComponent.Get())
-					{
-						int32 CurrJobSetCompleted = DC->LatestPSOPrecacheJobSetCompleted.load();
-						while (CurrJobSetCompleted < JobSetThatJustCompleted && !DC->LatestPSOPrecacheJobSetCompleted.compare_exchange_weak(CurrJobSetCompleted, JobSetThatJustCompleted)) {}
-						DC->MarkRenderStateDirty();
-					}
-				}
-
-				TWeakObjectPtr<UDecalComponent> WeakDecalComponent;
-				int32 JobSetThatJustCompleted;
-			};
-
-			LatestPSOPrecacheJobSet++;
-			if (GraphEvents.Num() > 0)
-			{
-				TGraphTask<FPSODecalPrecacheFinishedTask>::CreateTask(&GraphEvents).ConstructAndDispatchWhenReady(this, LatestPSOPrecacheJobSet);
-			}
-			else
-			{ 
-				// No graph events to wait on, the job set can be considered complete.
-				LatestPSOPrecacheJobSetCompleted = LatestPSOPrecacheJobSet;
-			}
-		}
+			MarkRenderStateDirty();
+		};
+		PSOPrecacheComponentData.SetMaterialPSOPrecacheData(this, MaterialPSOPrecacheRequestIDs, PSOPrecacheCompileEvents, Callback);
 	}
 #endif
 }
@@ -493,7 +460,7 @@ FDeferredDecalProxy* UDecalComponent::CreateSceneProxy()
 	LLM_SCOPE(ELLMTag::SceneRender);
 
 #if UE_WITH_PSO_PRECACHING
-	if (LatestPSOPrecacheJobSetCompleted != LatestPSOPrecacheJobSet && GetPSOPrecacheProxyCreationStrategy() != EPSOPrecacheProxyCreationStrategy::AlwaysCreate)
+	if (PSOPrecacheComponentData.IsPSOPrecaching() && GetPSOPrecacheProxyCreationStrategy() != EPSOPrecacheProxyCreationStrategy::AlwaysCreate)
 	{
 		return nullptr;
 	}
diff --git a/Engine/Source/Runtime/Engine/Private/Components/PrimitiveComponent.cpp b/Engine/Source/Runtime/Engine/Private/Components/PrimitiveComponent.cpp
--- a/Engine/Source/Runtime/Engine/Private/Components/PrimitiveComponent.cpp
+++ b/Engine/Source/Runtime/Engine/Private/Components/PrimitiveComponent.cpp
@@ -406,11 +406,7 @@ UPrimitiveComponent::UPrimitiveComponent(const FObjectInitializer& ObjectInitial
 	bBulkReregister = false;
 	LastCheckedAllCollideableDescendantsTime = 0.f;
 
-#if UE_WITH_PSO_PRECACHING
-	bPSOPrecacheCalled = false;
-	PSOPrecacheRequestPriority = EPSOPrecachePriority::Medium;
-	bPrecachedPSOs = false;
-#endif // UE_WITH_PSO_PRECACHING
+	PSOPrecacheComponentData.Reset();
 	
 	bApplyImpulseOnDamage = true;
 	bReplicatePhysicsToAutonomousProxy = true;
@@ -755,9 +751,7 @@ void UPrimitiveComponent::OnUnregister()
 	FStaticLightingSystemInterface::OnPrimitiveComponentUnregistered.Broadcast(this);
 #endif
 
-#if UE_WITH_PSO_PRECACHING
-	bPrecachedPSOs = false;
-#endif
+	PSOPrecacheComponentData.Reset();
 }
 
 FPrimitiveComponentInstanceData::FPrimitiveComponentInstanceData(const UPrimitiveComponent* SourceComponent)
@@ -5232,8 +5226,7 @@ void UPrimitiveComponent::PrecachePSOs()
 	}
 
 	// clear the current request data
-	MaterialPSOPrecacheRequestIDs.Empty();
-	PSOPrecacheRequestPriority = EPSOPrecachePriority::Medium;
+	PSOPrecacheComponentData.Reset();
 
 	// Collect the data from the derived classes
 	FPSOPrecacheParams PSOPrecacheParams;
@@ -5241,78 +5234,25 @@ void UPrimitiveComponent::PrecachePSOs()
 	FMaterialInterfacePSOPrecacheParamsList PSOPrecacheDataArray;
 	CollectPSOPrecacheData(PSOPrecacheParams, PSOPrecacheDataArray);
 
-	FGraphEventArray GraphEvents;
-	PrecacheMaterialPSOs(PSOPrecacheDataArray, MaterialPSOPrecacheRequestIDs, GraphEvents);
-
-	RequestRecreateRenderStateWhenPSOPrecacheFinished(GraphEvents);
+	TArray<FMaterialPSOPrecacheRequestID> MaterialPSOPrecacheRequestIDs;
+	FGraphEventArray PSOPrecacheCompileEvents;
+	PrecacheMaterialPSOs(PSOPrecacheDataArray, MaterialPSOPrecacheRequestIDs, PSOPrecacheCompileEvents);
 
-	bPrecachedPSOs = true;
+	SetMaterialPSOPrecacheData(MaterialPSOPrecacheRequestIDs, PSOPrecacheCompileEvents);
 #endif
 }
 
-#if UE_WITH_PSO_PRECACHING
-struct FPSOPrecacheFinishedTask
-{
-	explicit FPSOPrecacheFinishedTask(UPrimitiveComponent* InPrimitiveComponent, int32 InJobSetThatJustCompleted)
-		: WeakPrimitiveComponent(InPrimitiveComponent),
-		JobSetThatJustCompleted(InJobSetThatJustCompleted)
-
-	{
-	}
-
-	static TStatId GetStatId() { return TStatId(); }
-	static ENamedThreads::Type GetDesiredThread() { return ENamedThreads::GameThread; }
-	static ESubsequentsMode::Type GetSubsequentsMode() { return ESubsequentsMode::TrackSubsequents; }
-
-	void DoTask(ENamedThreads::Type CurrentThread, const FGraphEventRef& MyCompletionGraphEvent)
-	{
-		if (UPrimitiveComponent* PC = WeakPrimitiveComponent.Get())
-		{
-			QUICK_SCOPE_CYCLE_COUNTER(STAT_PSOPrecacheFinishedTask);
-			int32 CurrJobSetCompleted = PC->LatestPSOPrecacheJobSetCompleted.load();
-			while (CurrJobSetCompleted < JobSetThatJustCompleted && !PC->LatestPSOPrecacheJobSetCompleted.compare_exchange_weak(CurrJobSetCompleted, JobSetThatJustCompleted)){}
-
-			if (PC->GetWorld())
-			{
-				PC->GetWorld()->AddPrimitiveComponentPSOPrecacheCompleted(PC);
-			}
-		}
-	}
-
-	TWeakObjectPtr<UPrimitiveComponent> WeakPrimitiveComponent;
-	int32 JobSetThatJustCompleted;
-};
-
-namespace PrimitiveComponentCVars
-{
-	static bool bDefaultMaterialDelayStratCanIgnorePSOPrioBoost = false;
-	static FAutoConsoleVariableRef CVarDefaultMaterialDelayStratCanIgnorePSOPrioBoost(
-		TEXT("r.psoprecache.DefaultMaterialDelayStratCanIgnorePSOPrioBoost"),
-		bDefaultMaterialDelayStratCanIgnorePSOPrioBoost,
-		TEXT("false: PSO precache requests are boosted during proxy create. (default) \n")
-		TEXT("true: When r.PSOPrecache.ProxyCreationDelayStrategy and r.PSOPrecache.DrawnComponentBoostStrategy are in use, priority boost during proxy create will be ignored."));
-}
-#endif
-
-void UPrimitiveComponent::RequestRecreateRenderStateWhenPSOPrecacheFinished(const FGraphEventArray& PSOPrecacheCompileEvents)
+void UPrimitiveComponent::SetMaterialPSOPrecacheData(TArray<FMaterialPSOPrecacheRequestID>& MaterialPSOPrecacheRequestIDs, const FGraphEventArray& PSOPrecacheCompileEvents)
 {
 #if UE_WITH_PSO_PRECACHING
-	// If the proxy creation strategy relies on knowing when the precached PSO has been compiled,
-	// schedule a task to mark the render state dirty when all PSOs are compiled so the proxy gets recreated.
-	if (UsePSOPrecacheRenderProxyDelay() && GetPSOPrecacheProxyCreationStrategy() != EPSOPrecacheProxyCreationStrategy::AlwaysCreate)
+	auto Callback = [this]()
 	{
-		LatestPSOPrecacheJobSet++;
-		if(!PSOPrecacheCompileEvents.IsEmpty())
-		{
-			TGraphTask<FPSOPrecacheFinishedTask>::CreateTask(&PSOPrecacheCompileEvents).ConstructAndDispatchWhenReady(this, LatestPSOPrecacheJobSet);
-		}
-		else
+		if (GetWorld())
 		{
-			// No graph events to wait on, the job set can be considered complete.
-			LatestPSOPrecacheJobSetCompleted = LatestPSOPrecacheJobSet;
+			GetWorld()->AddPrimitiveComponentPSOPrecacheCompleted(this);
 		}
-	}
-	bPSOPrecacheCalled = true;
+	};
+	PSOPrecacheComponentData.SetMaterialPSOPrecacheData(this, MaterialPSOPrecacheRequestIDs, PSOPrecacheCompileEvents, Callback);
 #endif //UE_WITH_PSO_PRECACHING
 }
 
@@ -5325,44 +5265,6 @@ bool UPrimitiveComponent::UsePSOPrecacheRenderProxyDelay() const
 #endif
 }
 
-bool UPrimitiveComponent::IsPSOPrecaching() const
-{
-#if UE_WITH_PSO_PRECACHING
-	return LatestPSOPrecacheJobSetCompleted != LatestPSOPrecacheJobSet;
-#else
-	return false;
-#endif // UE_WITH_PSO_PRECACHING
-}
-
-bool UPrimitiveComponent::ShouldRenderProxyFallbackToDefaultMaterial() const
-{
-#if UE_WITH_PSO_PRECACHING
-	return IsPSOPrecaching() && GetPSOPrecacheProxyCreationStrategy() == EPSOPrecacheProxyCreationStrategy::UseDefaultMaterialUntilPSOPrecached;
-#else
-	return false;
-#endif // UE_WITH_PSO_PRECACHING
-}
-bool UPrimitiveComponent::CheckPSOPrecachingAndBoostPriority(EPSOPrecachePriority NewPSOPrecachePriority)
-{
-#if UE_WITH_PSO_PRECACHING
-	bool bPrecacheStillRunning = IsPSOPrecaching();
-
-	ensure(!IsComponentPSOPrecachingEnabled() || bPSOPrecacheCalled);
-	check(NewPSOPrecachePriority == EPSOPrecachePriority::High || NewPSOPrecachePriority == EPSOPrecachePriority::Highest);
-
-	const bool bIgnoringBoostDueToDefaultMaterialStrat = PrimitiveComponentCVars::bDefaultMaterialDelayStratCanIgnorePSOPrioBoost && GetPSOPrecacheProxyCreationStrategy() == EPSOPrecacheProxyCreationStrategy::UseDefaultMaterialUntilPSOPrecached && ShouldBoostPSOPrecachePriorityOnDraw();
-
-	if (bPrecacheStillRunning && PSOPrecacheRequestPriority < NewPSOPrecachePriority && !bIgnoringBoostDueToDefaultMaterialStrat)
-	{
-		BoostPSOPriority(NewPSOPrecachePriority, MaterialPSOPrecacheRequestIDs);
-		PSOPrecacheRequestPriority = NewPSOPrecachePriority;
-	}
-	return bPrecacheStillRunning;
-#else
-	return false;
-#endif
-}
-
 // Deprecated in 5.7
 FPrimitiveMaterialPropertyDescriptor UPrimitiveComponent::GetUsedMaterialPropertyDesc(ERHIFeatureLevel::Type FeatureLevel) const
 {
diff --git a/Engine/Source/Runtime/Engine/Private/Components/StaticMeshComponent.cpp b/Engine/Source/Runtime/Engine/Private/Components/StaticMeshComponent.cpp
--- a/Engine/Source/Runtime/Engine/Private/Components/StaticMeshComponent.cpp
+++ b/Engine/Source/Runtime/Engine/Private/Components/StaticMeshComponent.cpp
@@ -872,7 +872,7 @@ void UStaticMeshComponent::OnRegister()
 	if (FApp::CanEverRender() && IsComponentPSOPrecachingEnabled())
 	{
 		bool bCanSkipPrecachePSOs = false;
-		if (bPrecachedPSOs && GetWorld() && GetWorld()->IsGameWorld())
+		if (PSOPrecacheComponentData.IsPSOPrecacheCalled() && GetWorld() && GetWorld()->IsGameWorld())
 		{
 			// World transform might have changes causing negative determinant which changes the culling mode
 			const bool bIsLocalToWorldDeterminantNegative = GetRenderMatrix().Determinant() < 0;
diff --git a/Engine/Source/Runtime/Engine/Private/PSOPrecache.cpp b/Engine/Source/Runtime/Engine/Private/PSOPrecache.cpp
--- a/Engine/Source/Runtime/Engine/Private/PSOPrecache.cpp
+++ b/Engine/Source/Runtime/Engine/Private/PSOPrecache.cpp
@@ -12,6 +12,7 @@
 #include "PrimitiveSceneProxy.h"
 #include "PrimitiveSceneInfo.h"
 #include "PSOPrecacheValidation.h"
+#include "PSOPrecacheComponent.h"
 
 static TAutoConsoleVariable<int32> CVarPrecacheGlobalComputeShaders(
 	TEXT("r.PSOPrecache.GlobalShaders"),
@@ -67,6 +68,13 @@ static FAutoConsoleVariableRef CVarPSOComponentBoostStrategy(
 	ECVF_ReadOnly
 );
 
+static bool bDefaultMaterialDelayStratCanIgnorePSOPrioBoost = false;
+static FAutoConsoleVariableRef CVarDefaultMaterialDelayStratCanIgnorePSOPrioBoost(
+	TEXT("r.PSOPrecache.DefaultMaterialDelayStratCanIgnorePSOPrioBoost"),
+	bDefaultMaterialDelayStratCanIgnorePSOPrioBoost,
+	TEXT("false: PSO precache requests are boosted during proxy create. (default) \n")
+	TEXT("true: When r.PSOPrecache.ProxyCreationDelayStrategy and r.PSOPrecache.DrawnComponentBoostStrategy are in use, priority boost during proxy create will be ignored."));
+
 int32 GPSOPrecacheMode = 0;
 static FAutoConsoleVariableRef CVarPSOPrecacheMode(
 	TEXT("r.PSOPrecache.Mode"),
@@ -238,4 +246,25 @@ FPSOPrecacheRequestResultArray RequestPrecachePSOs(EPSOPrecacheType PSOPrecacheT
 	}
 
 	return Results;
+}
+
+bool FPSOPrecacheComponentData::CheckPSOPrecachingAndBoostPriority(EPSOPrecachePriority NewPSOPrecachePriority)
+{
+#if UE_WITH_PSO_PRECACHING
+	bool bPrecacheStillRunning = IsPSOPrecaching();
+
+	ensure(!IsComponentPSOPrecachingEnabled() || bPSOPrecacheCalled);
+	check(NewPSOPrecachePriority == EPSOPrecachePriority::High || NewPSOPrecachePriority == EPSOPrecachePriority::Highest);
+
+	const bool bIgnoringBoostDueToDefaultMaterialStrat = bDefaultMaterialDelayStratCanIgnorePSOPrioBoost && GetPSOPrecacheProxyCreationStrategy() == EPSOPrecacheProxyCreationStrategy::UseDefaultMaterialUntilPSOPrecached && ShouldBoostPSOPrecachePriorityOnDraw();
+
+	if (bPrecacheStillRunning && PSOPrecacheRequestPriority < NewPSOPrecachePriority && !bIgnoringBoostDueToDefaultMaterialStrat)
+	{
+		BoostPSOPriority(NewPSOPrecachePriority, MaterialPSOPrecacheRequestIDs);
+		PSOPrecacheRequestPriority = NewPSOPrecachePriority;
+	}
+	return bPrecacheStillRunning;
+#else
+	return false;
+#endif
 }
\ No newline at end of file
diff --git a/Engine/Source/Runtime/Engine/Private/ParticleSystemSceneProxy.cpp b/Engine/Source/Runtime/Engine/Private/ParticleSystemSceneProxy.cpp
--- a/Engine/Source/Runtime/Engine/Private/ParticleSystemSceneProxy.cpp
+++ b/Engine/Source/Runtime/Engine/Private/ParticleSystemSceneProxy.cpp
@@ -534,7 +534,7 @@ FPrimitiveSceneProxy* UParticleSystemComponent::CreateSceneProxy()
 	if ((IsActive() == true)/** && (EmitterInstances.Num() > 0)*/ && Template)
 	{
 #if UE_WITH_PSO_PRECACHING
-		if (!bPSOPrecacheCalled)
+		if (!PSOPrecacheComponentData.IsPSOPrecacheCalled())
 		{
 			if (GCascadePSOPrecachingTime == 3)
 			{
diff --git a/Engine/Source/Runtime/Engine/Private/Particles/ParticleSystemComponent.cpp b/Engine/Source/Runtime/Engine/Private/Particles/ParticleSystemComponent.cpp
--- a/Engine/Source/Runtime/Engine/Private/Particles/ParticleSystemComponent.cpp
+++ b/Engine/Source/Runtime/Engine/Private/Particles/ParticleSystemComponent.cpp
@@ -324,24 +324,21 @@ void UFXSystemComponent::PrecacheAssetPSOs(UFXSystemAsset* FXSystemAsset)
 	FGraphEventRef GraphEvent = FXSystemAsset->GetPrecachePSOsEvent();
 
 	check(IsInGameThread() || IsInParallelGameThread());
-#if UE_WITH_PSO_PRECACHING
-	MaterialPSOPrecacheRequestIDs.Empty();
-	PSOPrecacheRequestPriority = EPSOPrecachePriority::Medium;
-#endif
+
+	PSOPrecacheComponentData.Reset();
+
 	// The asset will keep the Precache events alive, but these might be over. Avoid delaying scene proxy creation if everything is finished
 	bool bAllEventsDone = GraphEvent == nullptr || GraphEvent->IsComplete();
 
-	FGraphEventArray Events;
+	TArray<FMaterialPSOPrecacheRequestID> MaterialPSOPrecacheRequestIDs;
+	FGraphEventArray PSOPrecacheCompileEvents;
 	if (!bAllEventsDone)
 	{
-#if UE_WITH_PSO_PRECACHING
 		MaterialPSOPrecacheRequestIDs.Append(FXSystemAsset->GetMaterialPSOPrecacheRequestIDs());
-#endif
-		Events.Add(GraphEvent);
+		PSOPrecacheCompileEvents.Add(GraphEvent);
 	}
 
-	RequestRecreateRenderStateWhenPSOPrecacheFinished(Events);
-	bPSOPrecacheCalled = true;
+	SetMaterialPSOPrecacheData(MaterialPSOPrecacheRequestIDs, PSOPrecacheCompileEvents);
 #endif // UE_WITH_PSO_PRECACHING
 }
 
diff --git a/Engine/Source/Runtime/Engine/Private/PrimitiveSceneProxy.cpp b/Engine/Source/Runtime/Engine/Private/PrimitiveSceneProxy.cpp
--- a/Engine/Source/Runtime/Engine/Private/PrimitiveSceneProxy.cpp
+++ b/Engine/Source/Runtime/Engine/Private/PrimitiveSceneProxy.cpp
@@ -605,10 +605,11 @@ FPrimitiveSceneProxy::FPrimitiveSceneProxy(const FPrimitiveSceneProxyDesc& InPro
 #if UE_WITH_PSO_PRECACHING
 	if (UPrimitiveComponent* PrimitiveComponent = Cast<UPrimitiveComponent>(InProxyDesc.Component))
 	{
-		check(!PrimitiveComponent->ShouldRenderProxyFallbackToDefaultMaterial() || !PrimitiveComponent->MaterialPSOPrecacheRequestIDs.IsEmpty());
-		if (ShouldBoostPSOPrecachePriorityOnDraw() && PrimitiveComponent->IsPSOPrecaching())
+		FPSOPrecacheComponentData& PSOPrecacheComponentData = PrimitiveComponent->GetPSOPrecacheComponentData();
+		check(!PSOPrecacheComponentData.ShouldRenderProxyFallbackToDefaultMaterial() || !PSOPrecacheComponentData.GetMaterialPSOPrecacheRequestIDs().IsEmpty());
+		if (ShouldBoostPSOPrecachePriorityOnDraw() && PSOPrecacheComponentData.IsPSOPrecaching())
 		{
-			SetPSORequestsToBoostOnDraw(PrimitiveComponent->MaterialPSOPrecacheRequestIDs);
+			SetPSORequestsToBoostOnDraw(PSOPrecacheComponentData.GetMaterialPSOPrecacheRequestIDs());
 		}
 	}
 #endif
diff --git a/Engine/Source/Runtime/Engine/Private/PrimitiveSceneProxyDesc.cpp b/Engine/Source/Runtime/Engine/Private/PrimitiveSceneProxyDesc.cpp
--- a/Engine/Source/Runtime/Engine/Private/PrimitiveSceneProxyDesc.cpp
+++ b/Engine/Source/Runtime/Engine/Private/PrimitiveSceneProxyDesc.cpp
@@ -214,7 +214,7 @@ void FPrimitiveSceneProxyDesc::InitializeFromPrimitiveComponent(const UPrimitive
 	HiddenEditorViews = InComponent->GetHiddenEditorViews();
 	OverlayColor = InComponent->OverlayColor;
 #endif
-	bShouldRenderProxyFallbackToDefaultMaterial = InComponent->ShouldRenderProxyFallbackToDefaultMaterial();
+	bShouldRenderProxyFallbackToDefaultMaterial = InComponent->GetPSOPrecacheComponentData().ShouldRenderProxyFallbackToDefaultMaterial();
 
 	AdditionalStatObjectPtr = InComponent->AdditionalStatObject();
 	StatId = AdditionalStatObjectPtr ? AdditionalStatObjectPtr->GetStatID(true) : InComponent->GetStatID(true);
diff --git a/Engine/Source/Runtime/Engine/Public/PSOPrecache.h b/Engine/Source/Runtime/Engine/Public/PSOPrecache.h
--- a/Engine/Source/Runtime/Engine/Public/PSOPrecache.h
+++ b/Engine/Source/Runtime/Engine/Public/PSOPrecache.h
@@ -396,29 +396,6 @@ extern ENGINE_API bool ShouldBoostPSOPrecachePriorityOnDraw();
  */
 extern ENGINE_API bool IsPSOShaderPreloadingEnabled();
 
-enum class EPSOPrecacheProxyCreationStrategy : uint8
-{
-	// Always create the render proxy regardless of whether the PSO has finished precaching or not. 
-	// This will introduce a blocking wait when the proxy is rendered if the PSO is not ready.
-	AlwaysCreate = 0, 
-
-	// Delay the creation of the render proxy until the PSO has finished precaching. 
-	// This effectively skips drawing components until the PSO is ready, when the proxy will be created.
-	DelayUntilPSOPrecached = 1, 
-	
-	// Create a render proxy that uses the default material if the PSO has not finished precaching by creation time.
-	// The proxy will be re-created with the actual materials once the PSO is ready.
-	// Currently implemented only for static and skinned mesh components, while Niagara components will skip render proxy creation altogether.
-	UseDefaultMaterialUntilPSOPrecached = 2
-};
-
-extern ENGINE_API EPSOPrecacheProxyCreationStrategy GetPSOPrecacheProxyCreationStrategy();
-
-/**
- * Delay component proxy creation when it's requested PSOs are still precaching
- */
-extern ENGINE_API bool ProxyCreationWhenPSOReady();
-
 /**
  * Boost the primitive's precache PSO jobs to highest priority.
  */
diff --git a/Engine/Source/Runtime/Engine/Public/PSOPrecacheComponent.h b/Engine/Source/Runtime/Engine/Public/PSOPrecacheComponent.h
new file mode 100644
--- /dev/null
+++ b/Engine/Source/Runtime/Engine/Public/PSOPrecacheComponent.h
@@ -0,0 +1,177 @@
+// Copyright Epic Games, Inc. All Rights Reserved.
+#pragma once
+
+#include "RHIDefinitions.h"
+#include "RHIFeatureLevel.h"
+#include "RHIResources.h"
+#include "Engine/EngineTypes.h"
+#include "PipelineStateCache.h"
+#include "PSOPrecacheFwd.h"
+
+enum class EPSOPrecacheProxyCreationStrategy : uint8
+{
+	// Always create the render proxy regardless of whether the PSO has finished precaching or not. 
+	// This will introduce a blocking wait when the proxy is rendered if the PSO is not ready.
+	AlwaysCreate = 0, 
+
+	// Delay the creation of the render proxy until the PSO has finished precaching. 
+	// This effectively skips drawing components until the PSO is ready, when the proxy will be created.
+	DelayUntilPSOPrecached = 1, 
+	
+	// Create a render proxy that uses the default material if the PSO has not finished precaching by creation time.
+	// The proxy will be re-created with the actual materials once the PSO is ready.
+	// Currently implemented only for static and skinned mesh components, while Niagara components will skip render proxy creation altogether.
+	UseDefaultMaterialUntilPSOPrecached = 2
+};
+
+extern ENGINE_API EPSOPrecacheProxyCreationStrategy GetPSOPrecacheProxyCreationStrategy();
+
+/**
+ * Delay component proxy creation when it's requested PSOs are still precaching
+ */
+extern ENGINE_API bool ProxyCreationWhenPSOReady();
+
+struct FPSOPrecacheComponentData
+{			
+	void Reset()
+	{		
+#if UE_WITH_PSO_PRECACHING
+		bPSOPrecacheCalled = false;
+		PSOPrecacheRequestPriority = EPSOPrecachePriority::Medium;
+		MaterialPSOPrecacheRequestIDs.Empty();
+#endif // UE_WITH_PSO_PRECACHING
+	}
+
+	bool IsPSOPrecaching() const
+	{
+#if UE_WITH_PSO_PRECACHING
+		return LatestPSOPrecacheJobSetCompleted != LatestPSOPrecacheJobSet;
+#else
+		return false;
+#endif // UE_WITH_PSO_PRECACHING
+	}
+
+	bool IsPSOPrecacheCalled() const
+	{
+#if UE_WITH_PSO_PRECACHING
+		return bPSOPrecacheCalled;
+#else
+		return false;
+#endif // UE_WITH_PSO_PRECACHING
+	}
+
+#if UE_WITH_PSO_PRECACHING
+	const TArray<FMaterialPSOPrecacheRequestID>& GetMaterialPSOPrecacheRequestIDs() const
+	{
+		return MaterialPSOPrecacheRequestIDs;
+	}
+
+	void SetLatestJobSetCompleted(int32 JobSetThatJustCompleted)
+	{
+		int32 CurrJobSetCompleted = LatestPSOPrecacheJobSetCompleted.load();
+		while (CurrJobSetCompleted < JobSetThatJustCompleted && !LatestPSOPrecacheJobSetCompleted.compare_exchange_weak(CurrJobSetCompleted, JobSetThatJustCompleted)){}
+	}
+#endif // UE_WITH_PSO_PRECACHING
+	
+	bool ShouldRenderProxyFallbackToDefaultMaterial() const
+	{
+#if UE_WITH_PSO_PRECACHING
+		return IsPSOPrecaching() && GetPSOPrecacheProxyCreationStrategy() == EPSOPrecacheProxyCreationStrategy::UseDefaultMaterialUntilPSOPrecached;
+#else
+		return false;
+#endif // UE_WITH_PSO_PRECACHING
+	}
+	
+	// Set all the PSO precaching data and mark as precached - will also kick task to run when precaching is done and given callback will be executed then when the component is still valid
+	template<typename ComponentType, typename PrecachedFinishedCallback>
+	void SetMaterialPSOPrecacheData(ComponentType* Component, TArray<FMaterialPSOPrecacheRequestID>& InMaterialPSOPrecacheRequestIDs, const FGraphEventArray& PSOPrecacheCompileEvents, PrecachedFinishedCallback InCallback);
+	
+	// Check if PSOs are still compiling, and if so, boost them to requested priority
+	ENGINE_API bool CheckPSOPrecachingAndBoostPriority(EPSOPrecachePriority NewPSOPrecachePriority);
+
+#if UE_WITH_PSO_PRECACHING
+
+private:
+
+	/** Helper flag to check if PSOs have been precached already */
+	uint8 bPSOPrecacheCalled : 1;
+
+	/** PSOs requested priority */
+	EPSOPrecachePriority PSOPrecacheRequestPriority : 2;
+	static_assert((int)EPSOPrecachePriority::Highest < 1 << 2);
+
+	int32 LatestPSOPrecacheJobSet : 29 = 0;
+
+	/** Cached array of material PSO requests which can be used to boost the priority */
+	TArray<FMaterialPSOPrecacheRequestID> MaterialPSOPrecacheRequestIDs;
+
+	/** Atomic int used to track the last PSO precache events */
+	std::atomic<int> LatestPSOPrecacheJobSetCompleted = 0;		
+
+#endif // UE_WITH_PSO_PRECACHING
+};
+
+#if UE_WITH_PSO_PRECACHING
+
+template<typename ComponentType, typename PrecachedFinishedOnComponentCallback>
+struct FPSOPrecacheFinishedTask
+{
+	explicit FPSOPrecacheFinishedTask(ComponentType* Component, int32 InJobSetThatJustCompleted, PrecachedFinishedOnComponentCallback InCallback)
+		: WeakComponent(Component),
+		JobSetThatJustCompleted(InJobSetThatJustCompleted),
+		Callback(InCallback)
+	{
+	}
+
+	static TStatId GetStatId() { return TStatId(); }
+	static ENamedThreads::Type GetDesiredThread() { return ENamedThreads::GameThread; }
+	static ESubsequentsMode::Type GetSubsequentsMode() { return ESubsequentsMode::TrackSubsequents; }
+
+	void DoTask(ENamedThreads::Type CurrentThread, const FGraphEventRef& MyCompletionGraphEvent)
+	{
+		if (ComponentType* Component = WeakComponent.Get())
+		{
+			QUICK_SCOPE_CYCLE_COUNTER(STAT_PSOPrecacheFinishedTask);
+			FPSOPrecacheComponentData& PSOPrecacheComponentData = Component->GetPSOPrecacheComponentData();
+			PSOPrecacheComponentData.SetLatestJobSetCompleted(JobSetThatJustCompleted);
+
+			Callback();
+		}
+	}
+		
+	TWeakObjectPtr<ComponentType> WeakComponent;
+	int32 JobSetThatJustCompleted;
+	PrecachedFinishedOnComponentCallback Callback;
+};
+
+#endif
+
+template<typename ComponentType, typename PrecachedFinishedCallback>
+void FPSOPrecacheComponentData::SetMaterialPSOPrecacheData(
+	ComponentType* Component, 
+	TArray<FMaterialPSOPrecacheRequestID>& InMaterialPSOPrecacheRequestIDs, 
+	const FGraphEventArray& PSOPrecacheCompileEvents, 
+	PrecachedFinishedCallback InCallback)
+{
+#if UE_WITH_PSO_PRECACHING
+	MaterialPSOPrecacheRequestIDs = MoveTemp(InMaterialPSOPrecacheRequestIDs);
+	bPSOPrecacheCalled = true;
+
+	// If the proxy creation strategy relies on knowing when the precached PSO has been compiled,
+	// schedule a task to mark the render state dirty when all PSOs are compiled so the proxy gets recreated.
+	if (GetPSOPrecacheProxyCreationStrategy() != EPSOPrecacheProxyCreationStrategy::AlwaysCreate)
+	{
+		int32 NewJobID = LatestPSOPrecacheJobSet + 1;
+		LatestPSOPrecacheJobSet = NewJobID;
+		if(!PSOPrecacheCompileEvents.IsEmpty())
+		{
+			TGraphTask<FPSOPrecacheFinishedTask<ComponentType, PrecachedFinishedCallback>>::CreateTask(&PSOPrecacheCompileEvents).ConstructAndDispatchWhenReady(Component, NewJobID, InCallback);
+		}
+		else
+		{
+			// No graph events to wait on, the job set can be considered complete.
+			LatestPSOPrecacheJobSetCompleted = NewJobID;
+		}
+	}
+#endif //UE_WITH_PSO_PRECACHING
+}
\ No newline at end of file
