LFV - Updated tool tip.

#tests PC editor
#rnx
#Synced-CL 48139523

diff --git a/Engine/Source/Runtime/Engine/Classes/Components/LocalFogVolumeComponent.h b/Engine/Source/Runtime/Engine/Classes/Components/LocalFogVolumeComponent.h
--- a/Engine/Source/Runtime/Engine/Classes/Components/LocalFogVolumeComponent.h
+++ b/Engine/Source/Runtime/Engine/Classes/Components/LocalFogVolumeComponent.h
@@ -44,7 +44,7 @@ class ULocalFogVolumeComponent : public USceneComponent
 	UPROPERTY(EditAnywhere, BlueprintReadOnly, interp, Category = "Shading")
 	FLinearColor FogEmissive = FLinearColor::Black;
 
-	/** Controls the distance after which this volume starts to appear. (Unit = Centimeter) */
+	/** Controls the distance after which this volume starts to appear. (Unit = Centimeter) This distance is clamped to be greater than r.LocalFogVolume.GlobalStartDistance. */
 	UPROPERTY(EditAnywhere, BlueprintReadOnly, interp, Category = "Shading", meta = (DisplayName = "Start Distance", UIMin = "0.0", ClampMin = 0.0))
 	float FogStartDistance = 0.0f;
 
