[NetPhysics] Add a clamp to how far behind the requested resim frame we search for a valid frame before opting for an incomplete frame.

#Synced-CL 48076803

diff --git a/Engine/Source/Runtime/Engine/Private/PhysicsEngine/NetworkPhysicsComponent.cpp b/Engine/Source/Runtime/Engine/Private/PhysicsEngine/NetworkPhysicsComponent.cpp
--- a/Engine/Source/Runtime/Engine/Private/PhysicsEngine/NetworkPhysicsComponent.cpp
+++ b/Engine/Source/Runtime/Engine/Private/PhysicsEngine/NetworkPhysicsComponent.cpp
@@ -2080,12 +2080,12 @@ void FAsyncNetworkPhysicsComponent::OnPreProcessInputs_Internal(const int32 Phys
 {
 	QUICK_SCOPE_CYCLE_COUNTER(AsyncNetworkPhysicsComponent_OnPreProcessInputs_Internal);
 
-	bool bIsSolverReset = false;
+	bool bIsFirstResimFrame = false;
 	bool bIsSolverResim = false;
 	if (Chaos::FPBDRigidsEvolution* Evolution = GetEvolution())
 	{
 		bIsSolverResim = Evolution->IsResimming();
-		bIsSolverReset = Evolution->IsResetting();
+		bIsFirstResimFrame = Evolution->IsResetting();
 	}
 
 	if (!bIsSolverResim)
@@ -2100,7 +2100,7 @@ void FAsyncNetworkPhysicsComponent::OnPreProcessInputs_Internal(const int32 Phys
 	{
 		const int32 InputBufferSize = (bIsServer && InputHistory) ? (InputHistory->GetLatestFrame() - PhysicsStep) : 0;
 		const FString NetRoleString = bIsServer ? FString("SERVER") : (IsLocallyControlled() ? FString("AUTONO") : FString("PROXY "));
-		UE_LOG(LogChaos, Log, TEXT("%s | PT | OnPreProcessInputsInternal | At Frame %d | IsResim: %d | FirstResimFrame: %d | InputBuffer: %d | Name = %s"), *NetRoleString, PhysicsStep, bIsSolverResim, bIsSolverReset, InputBufferSize, *GetActorName());
+		UE_LOG(LogChaos, Log, TEXT("%s | PT | OnPreProcessInputsInternal | At Frame %d | IsResim: %d | FirstResimFrame: %d | InputBuffer: %d | Name = %s"), *NetRoleString, PhysicsStep, bIsSolverResim, bIsFirstResimFrame, InputBufferSize, *GetActorName());
 	}
 #endif
 
@@ -2111,8 +2111,11 @@ void FAsyncNetworkPhysicsComponent::OnPreProcessInputs_Internal(const int32 Phys
 		{
 			FNetworkPhysicsPayload* PhysicsData = StateData.Get();
 			PhysicsData->LocalFrame = PhysicsStep;
-			const bool bExactFrame = PhysicsReplicationCVars::ResimulationCVars::bAllowRewindToClosestState ? !bIsSolverReset : true;
-			if (StateHistory->ExtractData(PhysicsStep, bIsSolverReset, PhysicsData, (bExactFrame && bIsSolverResim)) && PhysicsData->bReceivedData)
+
+			/** Extract state to apply, only apply states that are received from the server (not predicted states)
+			* Require exact frame unless bAllowRewindToClosestState is true, then allow non-exact frame on first resim frame */ 
+			const bool bExactFrame = PhysicsReplicationCVars::ResimulationCVars::bAllowRewindToClosestState ? !bIsFirstResimFrame : true;
+			if (StateHistory->ExtractData(PhysicsStep, bIsFirstResimFrame, PhysicsData, (bExactFrame && bIsSolverResim)) && PhysicsData->bReceivedData)
 			{
 				if (bIsUsingLegacyData)
 				{
@@ -2193,7 +2196,7 @@ void FAsyncNetworkPhysicsComponent::OnPreProcessInputs_Internal(const int32 Phys
 			}
 
 			const bool RequireExactFrame = bIsServer && !ComponentSettings.GetAllowInputExtrapolation();
-			if (InputHistory->ExtractData(PhysicsStep, bIsSolverReset, PhysicsData, RequireExactFrame))
+			if (InputHistory->ExtractData(PhysicsStep, bIsFirstResimFrame, PhysicsData, RequireExactFrame))
 			{
 				const int32 CurrentInputFrame = bApplySimProxyInput ? InputHistory->GetLatestFrame() : PhysicsData->LocalFrame;
 
@@ -2237,7 +2240,7 @@ void FAsyncNetworkPhysicsComponent::OnPreProcessInputs_Internal(const int32 Phys
 							/** Note, extract the last frame but don't apply it here, it will be applied later in the flow
 							* Also note that we allow extracting non-exact data on the last frame, unless required by settings, to get the same data extracted as we did in the original ExtractData above */
 							bool bLastInput = NextExpectedInputFrame == CurrentInputFrame;
-							if (InputHistory->ExtractData(NextExpectedInputFrame, bIsSolverReset, PhysicsData, (RequireExactFrame || !bLastInput)) && !bLastInput)
+							if (InputHistory->ExtractData(NextExpectedInputFrame, bIsFirstResimFrame, PhysicsData, (RequireExactFrame || !bLastInput)) && !bLastInput)
 							{
 								// Perform input decay on sim-proxy runtime (non-resim) input
 								if (bApplySimProxyInput && ComponentSettings.GetApplyInputDecaySimProxyInputAtRuntime())
@@ -2332,18 +2335,18 @@ void FAsyncNetworkPhysicsComponent::OnPostProcessInputs_Internal(const int32 Phy
 	const FNetworkPhysicsSettingsNetworkPhysicsComponent& ComponentSettings = GetComponentSettings();
 	const bool bIsServer = IsServer();
 
-	bool bIsSolverReset = false;
+	bool bIsFirstResimFrame = false;
 	bool bIsSolverResim = false;
 	if (Chaos::FPBDRigidsEvolution* Evolution = GetEvolution())
 	{
 		bIsSolverResim = Evolution->IsResimming();
-		bIsSolverReset = Evolution->IsResetting();
+		bIsFirstResimFrame = Evolution->IsResetting();
 	}
 
 #if DEBUG_NETWORK_PHYSICS
 	{
 		FString NetRoleString = bIsServer ? FString("SERVER") : (IsLocallyControlled() ? FString("AUTONO") : FString("PROXY "));
-		UE_LOG(LogChaos, Log, TEXT("%s | PT | OnPostProcessInputsInternal | At Frame %d | IsResim: %d | FirstResimFrame: %d | Name = %s"), *NetRoleString, PhysicsStep, bIsSolverResim, bIsSolverReset, *GetActorName());
+		UE_LOG(LogChaos, Log, TEXT("%s | PT | OnPostProcessInputsInternal | At Frame %d | IsResim: %d | FirstResimFrame: %d | Name = %s"), *NetRoleString, PhysicsStep, bIsSolverResim, bIsFirstResimFrame, *GetActorName());
 	}
 #endif
 
diff --git a/Engine/Source/Runtime/Experimental/Chaos/Private/RewindData.cpp b/Engine/Source/Runtime/Experimental/Chaos/Private/RewindData.cpp
--- a/Engine/Source/Runtime/Experimental/Chaos/Private/RewindData.cpp
+++ b/Engine/Source/Runtime/Experimental/Chaos/Private/RewindData.cpp
@@ -1880,6 +1880,8 @@ enum class EResimFrameValidation : int32
 };
 CHAOS_API int32 ResimFrameValidation = static_cast<int32>(EResimFrameValidation::IslandValidation);
 FAutoConsoleVariableRef CVarResimFrameValidationLeniency(TEXT("p.Resim.ResimFrameValidation"), ResimFrameValidation, TEXT("0 = no leniency, all dirty particles need a valid target. 1 = Island leniency, all particles in resim islands need a valid target. 2 = Full leniency, only the particle triggering the resim need a valid target."));
+CHAOS_API int32 FindValidFrameNFramesBehind = 2;
+FAutoConsoleVariableRef CVarResimFindValidFrameNFramesBehind(TEXT("p.Resim.FindValidFrameNFramesBehind"), FindValidFrameNFramesBehind, TEXT("Number of frames to allow searching for a valid frame behind the requested frame."));
 CHAOS_API bool bResimIncompleteHistory = true;
 FAutoConsoleVariableRef CVarResimIncompleteHistory(TEXT("p.Resim.IncompleteHistory"), bResimIncompleteHistory, TEXT("If a valid resim frame can't be found, use the requested resim frame and perform a resimulation with incomplete data."));
 CHAOS_API bool bFindValidInputHistory = true;
@@ -1989,8 +1991,11 @@ int32 FRewindData::FindValidResimFrame(const int32 RequestedFrame)
 		IslandManager.FindParticlesInIslands(IslandValidationIslands, OUT IslandValidationIslandParticles);
 	}
 
-	// First frame of the history data
-	const int32 EarliestFrame = FMath::Max(GetEarliestFrame_Internal(), BlockResimFrame);
+	/** Get earliest frame we are allowed to rewind taking these three into account:
+	* - Earliest frame in history
+	* - Blocked resim frame 
+	* - Setting for how many frames behind the requested frame we allow to search (NOTE: -1 always added because the for-loop is excluding the EarliestFrame) */ 
+	const int32 EarliestFrame = FMath::Max3(GetEarliestFrame_Internal(), BlockResimFrame, RequestedFrame - 1 - FMath::Max(FindValidFrameNFramesBehind, 0));
 	bool bHasTargetHistory = false;
 	bool bHasCustomDataHistory = false;
 
