Add precaching support for light component and light functions

PR #13926

#Synced-CL 48139523

diff --git a/Engine/Source/Runtime/Engine/Classes/Components/LightComponent.h b/Engine/Source/Runtime/Engine/Classes/Components/LightComponent.h
--- a/Engine/Source/Runtime/Engine/Classes/Components/LightComponent.h
+++ b/Engine/Source/Runtime/Engine/Classes/Components/LightComponent.h
@@ -10,6 +10,8 @@
 #include "SceneTypes.h"
 #include "RenderResource.h"
 #include "Components/LightComponentBase.h"
+#include "PSOPrecacheFwd.h" 
+#include "PSOPrecacheComponent.h"
 #include "LightComponent.generated.h"
 
 class FLightSceneProxy;
@@ -500,6 +502,9 @@ public:
 	{
 		return NULL;
 	}
+	
+	/** Get the PSO precaching data */
+	FPSOPrecacheComponentData& GetPSOPrecacheComponentData() { return PSOPrecacheComponentData; }
 
 protected:
 	//~ Begin UActorComponent Interface
@@ -507,6 +512,7 @@ protected:
 	ENGINE_API virtual void CreateRenderState_Concurrent(FRegisterComponentContext* Context) override;
 	ENGINE_API virtual void SendRenderTransform_Concurrent() override;
 	ENGINE_API virtual void DestroyRenderState_Concurrent() override;
+	ENGINE_API virtual void PrecachePSOs() override;
 	//~ Begin UActorComponent Interface
 
 	//~ Begin USceneComponent Interface
@@ -516,6 +522,14 @@ protected:
 	//~ End USceneComponent Interface
 
 	void UpdateProxy(FUpdateLightProxyCallback&& Func);
+		
+	/** All PSO related data for the component */
+	FPSOPrecacheComponentData PSOPrecacheComponentData;
+
+	bool IsPSOPrecaching() const
+	{
+		return PSOPrecacheComponentData.IsPSOPrecaching();
+	}
 
 private:
 
diff --git a/Engine/Source/Runtime/Engine/Classes/Components/PointLightComponent.h b/Engine/Source/Runtime/Engine/Classes/Components/PointLightComponent.h
--- a/Engine/Source/Runtime/Engine/Classes/Components/PointLightComponent.h
+++ b/Engine/Source/Runtime/Engine/Classes/Components/PointLightComponent.h
@@ -86,6 +86,7 @@ public:
 	ENGINE_API virtual ELightComponentType GetLightType() const override;
 	ENGINE_API virtual float GetUniformPenumbraSize() const override;
 	ENGINE_API virtual FLightSceneProxy* CreateSceneProxy() const override;
+	ENGINE_API virtual void PrecachePSOs() override;
 
 	//~ Begin UObject Interface
 	ENGINE_API virtual void Serialize(FArchive& Ar) override;
diff --git a/Engine/Source/Runtime/Engine/Classes/Components/SpotLightComponent.h b/Engine/Source/Runtime/Engine/Classes/Components/SpotLightComponent.h
--- a/Engine/Source/Runtime/Engine/Classes/Components/SpotLightComponent.h
+++ b/Engine/Source/Runtime/Engine/Classes/Components/SpotLightComponent.h
@@ -36,6 +36,7 @@ class USpotLightComponent : public UPointLightComponent
 	ENGINE_API virtual bool AffectsBounds(const FBoxSphereBounds& InBounds) const override;
 	ENGINE_API virtual ELightComponentType GetLightType() const override;
 	ENGINE_API virtual FLightSceneProxy* CreateSceneProxy() const override;
+	ENGINE_API virtual void PrecachePSOs() override;
 
 	ENGINE_API virtual float ComputeLightBrightness() const override;
 #if WITH_EDITOR
diff --git a/Engine/Source/Runtime/Engine/Private/Components/DirectionalLightComponent.cpp b/Engine/Source/Runtime/Engine/Private/Components/DirectionalLightComponent.cpp
--- a/Engine/Source/Runtime/Engine/Private/Components/DirectionalLightComponent.cpp
+++ b/Engine/Source/Runtime/Engine/Private/Components/DirectionalLightComponent.cpp
@@ -1171,7 +1171,11 @@ bool UDirectionalLightComponent::CanEditChange(const FProperty* InProperty) cons
 
 FLightSceneProxy* UDirectionalLightComponent::CreateSceneProxy() const
 {
-	return new FDirectionalLightSceneProxy(this);
+	if (!IsPSOPrecaching())
+	{
+		return new FDirectionalLightSceneProxy(this);
+	}
+	return nullptr;
 }
 
 FVector4 UDirectionalLightComponent::GetLightPosition() const
diff --git a/Engine/Source/Runtime/Engine/Private/Components/LightComponent.cpp b/Engine/Source/Runtime/Engine/Private/Components/LightComponent.cpp
--- a/Engine/Source/Runtime/Engine/Private/Components/LightComponent.cpp
+++ b/Engine/Source/Runtime/Engine/Private/Components/LightComponent.cpp
@@ -30,6 +30,8 @@
 #include "Components/DirectionalLightComponent.h"
 #include "Components/BillboardComponent.h"
 #include "ComponentRecreateRenderStateContext.h"
+#include "LocalVertexFactory.h"
+#include "PSOPrecacheMaterial.h"
 #include "UObject/ICookInfo.h"
 #include "UObject/SoftObjectPath.h"
 #include "UObject/UnrealType.h"
@@ -329,6 +331,8 @@ void ULightComponentBase::OnRegister()
 		FStaticLightingSystemInterface::OnLightComponentRegistered.Broadcast(this);
 	}
 #endif
+
+	PrecachePSOs();
 }
 
 void ULightComponentBase::OnUnregister()
@@ -631,6 +635,8 @@ void ULightComponent::PostLoad()
 			IESBrightnessScale /= IESTexture->TextureMultiplier; // Previous version didn't apply IES texture multiplier, so cancel out
 		}
 	}
+
+	PrecachePSOs();
 }
 
 bool ULightComponent::CanTraceDistanceFieldShadows() const
@@ -975,6 +981,36 @@ void ULightComponent::DestroyRenderState_Concurrent()
 	bAddedToSceneVisible = false;
 }
 
+void ULightComponent::PrecachePSOs()
+{
+#if UE_WITH_PSO_PRECACHING
+	if (!FApp::CanEverRender() || !IsComponentPSOPrecachingEnabled())
+	{
+		return;
+	}
+	if (LightFunctionMaterial)
+	{
+		LightFunctionMaterial->ConditionalPostLoad();
+
+		FMaterialInterfacePSOPrecacheParams Params;
+		Params.Priority = EPSOPrecachePriority::High; //< No boosting right now - can be added
+		Params.MaterialInterface = LightFunctionMaterial;
+		Params.VertexFactoryDataList.Add(&FLocalVertexFactory::StaticType);
+				
+		TArray<FMaterialPSOPrecacheRequestID> MaterialPSOPrecacheRequestIDs;
+		FGraphEventArray PSOPrecacheCompileEvents;
+		PrecacheMaterialPSOs( { Params }, MaterialPSOPrecacheRequestIDs, PSOPrecacheCompileEvents);
+
+		// Request recreate of the render state when the PSO compilation is ready (if we want to delay proxy creation)
+		auto Callback = [this]()
+		{
+			MarkRenderStateDirty();
+		};
+		PSOPrecacheComponentData.SetMaterialPSOPrecacheData(this, MaterialPSOPrecacheRequestIDs, PSOPrecacheCompileEvents, Callback);
+	}
+#endif
+}
+
 void ULightComponent::UpdateProxy(FUpdateLightProxyCallback&& Func)
 {
 	if (SceneProxy)
@@ -1126,6 +1162,9 @@ void ULightComponent::SetLightFunctionMaterial(UMaterialInterface* NewLightFunct
 		StashedLightFunctionMaterial = nullptr;
 #endif
 		LightFunctionMaterial = NewLightFunctionMaterial;
+
+		PrecachePSOs();
+
 		if (SceneProxy)
 		{
 			SceneProxy->SetLightFunctionMaterial_GameThread(this);
diff --git a/Engine/Source/Runtime/Engine/Private/Components/PointLightComponent.cpp b/Engine/Source/Runtime/Engine/Private/Components/PointLightComponent.cpp
--- a/Engine/Source/Runtime/Engine/Private/Components/PointLightComponent.cpp
+++ b/Engine/Source/Runtime/Engine/Private/Components/PointLightComponent.cpp
@@ -129,13 +129,21 @@ static bool IsPointLightSupported(const UPointLightComponent* InLight)
 
 FLightSceneProxy* UPointLightComponent::CreateSceneProxy() const
 {
-	if (IsPointLightSupported(this))
+	if (IsPointLightSupported(this) && !IsPSOPrecaching())
 	{
 		return new FPointLightSceneProxy(this);
 	}
 	return nullptr;
 }
 
+void UPointLightComponent::PrecachePSOs()
+{
+	if (IsPointLightSupported(this))
+	{
+		Super::PrecachePSOs();
+	}
+}
+
 void UPointLightComponent::SetUseInverseSquaredFalloff(bool bNewValue)
 {
 	if (AreDynamicDataChangesAllowed()
diff --git a/Engine/Source/Runtime/Engine/Private/Components/RectLightComponent.cpp b/Engine/Source/Runtime/Engine/Private/Components/RectLightComponent.cpp
--- a/Engine/Source/Runtime/Engine/Private/Components/RectLightComponent.cpp
+++ b/Engine/Source/Runtime/Engine/Private/Components/RectLightComponent.cpp
@@ -116,7 +116,11 @@ URectLightComponent::URectLightComponent(const FObjectInitializer& ObjectInitial
 
 FLightSceneProxy* URectLightComponent::CreateSceneProxy() const
 {
-	return new FRectLightSceneProxy(this);
+	if (!IsPSOPrecaching())
+	{
+		return new FRectLightSceneProxy(this);
+	}
+	return nullptr;
 }
 
 void URectLightComponent::SetSourceTexture(UTexture* NewValue)
diff --git a/Engine/Source/Runtime/Engine/Private/Components/SpotLightComponent.cpp b/Engine/Source/Runtime/Engine/Private/Components/SpotLightComponent.cpp
--- a/Engine/Source/Runtime/Engine/Private/Components/SpotLightComponent.cpp
+++ b/Engine/Source/Runtime/Engine/Private/Components/SpotLightComponent.cpp
@@ -322,13 +322,21 @@ static bool IsSpotLightSupported(const USpotLightComponent* InLight)
 
 FLightSceneProxy* USpotLightComponent::CreateSceneProxy() const
 {
-	if (IsSpotLightSupported(this))
+	if (IsSpotLightSupported(this) && !IsPSOPrecaching())
 	{
 		return new FSpotLightSceneProxy(this);
 	}
 	return nullptr;
 }
 
+void USpotLightComponent::PrecachePSOs()
+{
+	if (IsSpotLightSupported(this))
+	{
+		Super::PrecachePSOs();
+	}
+}
+
 FSphere USpotLightComponent::GetBoundingSphere() const
 {
 	FSphere::FReal ConeAngle = GetHalfConeAngle();
diff --git a/Engine/Source/Runtime/Engine/Private/Materials/Material.cpp b/Engine/Source/Runtime/Engine/Private/Materials/Material.cpp
--- a/Engine/Source/Runtime/Engine/Private/Materials/Material.cpp
+++ b/Engine/Source/Runtime/Engine/Private/Materials/Material.cpp
@@ -7362,6 +7362,11 @@ bool UMaterial::IsPostProcessMaterial() const
 	return MaterialDomain == MD_PostProcess;
 }
 
+bool UMaterial::IsLightFunctionMaterial() const
+{
+	return MaterialDomain == MD_LightFunction;
+}
+
 bool UMaterial::IsPostProcessMaterialOutputingAlpha() const
 {
 	return UMaterial::IsPostProcessMaterial() && (BlendableOutputAlpha || BlendableLocation == BL_TranslucencyAfterDOF);
diff --git a/Engine/Source/Runtime/Engine/Private/Materials/MaterialInstance.cpp b/Engine/Source/Runtime/Engine/Private/Materials/MaterialInstance.cpp
--- a/Engine/Source/Runtime/Engine/Private/Materials/MaterialInstance.cpp
+++ b/Engine/Source/Runtime/Engine/Private/Materials/MaterialInstance.cpp
@@ -3669,13 +3669,13 @@ void UMaterialInstance::PostLoad()
 	{
 		// When dynamic preload shaders is enabled, we need to prelaod some material domains since there is no
 		// code logic within the PSO precaching system.
-		if (IsUIMaterial() || IsDeferredDecal() || IsPostProcessMaterial())
+		if (IsUIMaterial() || IsDeferredDecal() || IsPostProcessMaterial() || IsLightFunctionMaterial())
 		{
 			FGraphEventArray Unused;
 			PreloadMaterialShaderMap(GetMaterialResource(GMaxRHIShaderPlatform), Unused);
 		}
 	}
-	else if (IsDeferredDecal() || IsUIMaterial() || IsPostProcessMaterial())
+	else if (IsDeferredDecal() || IsUIMaterial() || IsPostProcessMaterial() || IsLightFunctionMaterial())
 	{
 		FPSOPrecacheParams PSOPrecacheParams;
 		UMaterialInterface::PrecachePSOs(&FLocalVertexFactory::StaticType, PSOPrecacheParams);
@@ -5203,6 +5203,11 @@ bool UMaterialInstance::IsPostProcessMaterial() const
 	return Parent ? Parent->IsPostProcessMaterial() : false;
 }
 
+bool UMaterialInstance::IsLightFunctionMaterial() const
+{
+	return Parent ? Parent->IsLightFunctionMaterial() : false;
+}
+
 bool UMaterialInstance::HasPixelAnimation() const
 {
 	return bHasPixelAnimation;
diff --git a/Engine/Source/Runtime/Engine/Private/Materials/MaterialInterface.cpp b/Engine/Source/Runtime/Engine/Private/Materials/MaterialInterface.cpp
--- a/Engine/Source/Runtime/Engine/Private/Materials/MaterialInterface.cpp
+++ b/Engine/Source/Runtime/Engine/Private/Materials/MaterialInterface.cpp
@@ -1678,6 +1678,11 @@ bool UMaterialInterface::IsPostProcessMaterial() const
 	return false;
 }
 
+bool UMaterialInterface::IsLightFunctionMaterial() const
+{
+	return false;
+}
+
 bool UMaterialInterface::GetCastDynamicShadowAsMasked() const
 {
 	return false;
diff --git a/Engine/Source/Runtime/Engine/Public/Materials/Material.h b/Engine/Source/Runtime/Engine/Public/Materials/Material.h
--- a/Engine/Source/Runtime/Engine/Public/Materials/Material.h
+++ b/Engine/Source/Runtime/Engine/Public/Materials/Material.h
@@ -1291,6 +1291,7 @@ public:
 	ENGINE_API virtual bool IsDeferredDecal() const override;
 	ENGINE_API virtual bool IsUIMaterial() const;
 	ENGINE_API virtual bool IsPostProcessMaterial() const;
+	ENGINE_API virtual bool IsLightFunctionMaterial() const;
 	ENGINE_API bool IsPostProcessMaterialOutputingAlpha() const;
 	ENGINE_API virtual USubsurfaceProfile* GetSubsurfaceProfile_Internal() const override;
 	ENGINE_API virtual USubsurfaceProfile* GetSubsurfaceProfileRoot_Internal() const override;
diff --git a/Engine/Source/Runtime/Engine/Public/Materials/MaterialInstance.h b/Engine/Source/Runtime/Engine/Public/Materials/MaterialInstance.h
--- a/Engine/Source/Runtime/Engine/Public/Materials/MaterialInstance.h
+++ b/Engine/Source/Runtime/Engine/Public/Materials/MaterialInstance.h
@@ -969,6 +969,7 @@ public:
 	ENGINE_API virtual bool IsDeferredDecal() const override;
 	ENGINE_API virtual bool IsUIMaterial() const override;
 	ENGINE_API virtual bool IsPostProcessMaterial() const override;
+	ENGINE_API virtual bool IsLightFunctionMaterial() const override;
 	ENGINE_API virtual bool HasVertexInterpolator() const override;
 	ENGINE_API virtual bool HasCustomizedUVs() const override;
 	ENGINE_API virtual bool HasPixelAnimation() const override;
diff --git a/Engine/Source/Runtime/Engine/Public/Materials/MaterialInterface.h b/Engine/Source/Runtime/Engine/Public/Materials/MaterialInterface.h
--- a/Engine/Source/Runtime/Engine/Public/Materials/MaterialInterface.h
+++ b/Engine/Source/Runtime/Engine/Public/Materials/MaterialInterface.h
@@ -1102,6 +1102,7 @@ public:
 	ENGINE_API virtual bool IsDeferredDecal() const;
 	ENGINE_API virtual bool IsUIMaterial() const;
 	ENGINE_API virtual bool IsPostProcessMaterial() const;
+	ENGINE_API virtual bool IsLightFunctionMaterial() const;
 	ENGINE_API virtual bool WritesToRuntimeVirtualTexture() const;
 	ENGINE_API virtual bool HasMeshPaintTexture() const;
 	ENGINE_API virtual bool HasCustomPrimitiveData() const;
