Statically cache the DeviceProfileSelectionModule value to avoid potentially loading ini from disk at runtime

#tests Editor, standalone Windows client
#rnx
#Synced-CL 48128650

diff --git a/Engine/Source/Runtime/Engine/Private/DeviceProfiles/DeviceProfileManager.cpp b/Engine/Source/Runtime/Engine/Private/DeviceProfiles/DeviceProfileManager.cpp
--- a/Engine/Source/Runtime/Engine/Private/DeviceProfiles/DeviceProfileManager.cpp
+++ b/Engine/Source/Runtime/Engine/Private/DeviceProfiles/DeviceProfileManager.cpp
@@ -1571,8 +1571,10 @@ IDeviceProfileSelectorModule* UDeviceProfileManager::GetPreviewDeviceProfileSele
 
 IDeviceProfileSelectorModule* UDeviceProfileManager::GetDeviceProfileSelectorModule()
 {
-	FString DeviceProfileSelectionModule;
-	if (GConfig->GetString(TEXT("DeviceProfileManager"), TEXT("DeviceProfileSelectionModule"), DeviceProfileSelectionModule, GEngineIni))
+	static FString DeviceProfileSelectionModule;
+	static bool bReadModuleNameFromConfig = GConfig->GetString(TEXT("DeviceProfileManager"), TEXT("DeviceProfileSelectionModule"), DeviceProfileSelectionModule, GEngineIni);
+
+	if (bReadModuleNameFromConfig)
 	{
 		if (IDeviceProfileSelectorModule* DPSelectorModule = FModuleManager::LoadModulePtr<IDeviceProfileSelectorModule>(*DeviceProfileSelectionModule))
 		{
