Update SetParameterValueOnMaterials functions to support MaterialSlotOverlayMaterials

#Synced-CL 48196008

diff --git a/Engine/Source/Runtime/Engine/Classes/Components/MeshComponent.h b/Engine/Source/Runtime/Engine/Classes/Components/MeshComponent.h
--- a/Engine/Source/Runtime/Engine/Classes/Components/MeshComponent.h
+++ b/Engine/Source/Runtime/Engine/Classes/Components/MeshComponent.h
@@ -233,6 +233,7 @@ public:
 	virtual void GetDefaultMaterialSlotsOverlayMaterial(TArray<TObjectPtr<UMaterialInterface>>& OutMaterialSlotOverlayMaterials) const { return; }
 
 protected:
+	ENGINE_API void SetParameterValueOnMaterials(const FName ParameterName, const TFunction<void(UMaterialInstanceDynamic* DynamicMaterial)>& SetParameterFunction);
 
 	//~ Begin UObject Interface.
 	ENGINE_API virtual void BeginDestroy() override;
@@ -275,4 +276,5 @@ protected:
 	uint8 bCachedMaterialParameterIndicesAreDirty : 1;
 
 	TSortedMap<FName, FMaterialParameterCache, FDefaultAllocator, FNameFastLess> MaterialParameterCache;
+	TSortedMap<FName, FMaterialParameterCache, FDefaultAllocator, FNameFastLess> MaterialSlotsOverlayMaterialParameterCache;
 };
diff --git a/Engine/Source/Runtime/Engine/Private/Components/MeshComponent.cpp b/Engine/Source/Runtime/Engine/Private/Components/MeshComponent.cpp
--- a/Engine/Source/Runtime/Engine/Private/Components/MeshComponent.cpp
+++ b/Engine/Source/Runtime/Engine/Private/Components/MeshComponent.cpp
@@ -445,24 +445,55 @@ TArray<class UMaterialInterface*> UMeshComponent::GetMaterials() const
 	return OutMaterials;
 }
 
-void UMeshComponent::SetScalarParameterValueOnMaterials(const FName ParameterName, const float ParameterValue)
+void UMeshComponent::SetParameterValueOnMaterials(const FName ParameterName, const TFunction<void(UMaterialInstanceDynamic* DynamicMaterial)>& SetParameterFunction)
 {
+	bool bDynamicSlotOverlayMaterialCreated = false;
+
+	TArray<TObjectPtr<UMaterialInterface>> RetrievedMaterialSlotsOverlayMaterials;
+	GetMaterialSlotsOverlayMaterial(RetrievedMaterialSlotsOverlayMaterials);
+	MaterialSlotsOverlayMaterial.SetNum(RetrievedMaterialSlotsOverlayMaterials.Num());
+
+	auto CreateMaterialInstanceDynamic = [this](UMaterialInterface* MaterialInterface, int32 MaterialIndex)
+		{
+			return CreateAndSetMaterialInstanceDynamic(MaterialIndex);
+		};
+
+	auto CreateMaterialSlotsOverlayMaterialInstanceDynamic = [this, &bDynamicSlotOverlayMaterialCreated](UMaterialInterface* MaterialInterface, int32 MaterialIndex)
+		{
+			bDynamicSlotOverlayMaterialCreated = true;
+
+			UMaterialInstanceDynamic* DynamicMaterial = UMaterialInstanceDynamic::Create(MaterialInterface, this);
+			DynamicMaterial->OnAssignedAsOverride(this);
+
+			MaterialSlotsOverlayMaterial[MaterialIndex] = DynamicMaterial;
+
+			return DynamicMaterial;
+		};
+
 	if (!bEnableMaterialParameterCaching)
 	{
-		const TArray<UMaterialInterface*> MaterialInterfaces = GetMaterials();
-		for (int32 MaterialIndex = 0; MaterialIndex < MaterialInterfaces.Num(); ++MaterialIndex)
-		{
-			UMaterialInterface* MaterialInterface = MaterialInterfaces[MaterialIndex];
-			if (MaterialInterface)
+		auto SetParameterValueOnMaterials_Internal = [&SetParameterFunction]
+			(const TArray<UMaterialInterface*>& MaterialInterfaces
+			, const TFunction<UMaterialInstanceDynamic*(UMaterialInterface* MaterialInterface, int32 MaterialIndex)>& CreateDynamicMaterialInstanceFunction
+			)
 			{
-				UMaterialInstanceDynamic* DynamicMaterial = Cast<UMaterialInstanceDynamic>(MaterialInterface);
-				if (!DynamicMaterial)
+				for (int32 MaterialIndex = 0; MaterialIndex < MaterialInterfaces.Num(); ++MaterialIndex)
 				{
-					DynamicMaterial = CreateAndSetMaterialInstanceDynamic(MaterialIndex);
+					UMaterialInterface* MaterialInterface = MaterialInterfaces[MaterialIndex];
+					if (MaterialInterface)
+					{
+						UMaterialInstanceDynamic* DynamicMaterial = Cast<UMaterialInstanceDynamic>(MaterialInterface);
+						if (!DynamicMaterial)
+						{
+							DynamicMaterial = CreateDynamicMaterialInstanceFunction(MaterialInterface, MaterialIndex);
+						}
+						SetParameterFunction(DynamicMaterial);
+					}
 				}
-				DynamicMaterial->SetScalarParameterValue(ParameterName, ParameterValue);
-			}
-		}
+			};
+
+		SetParameterValueOnMaterials_Internal(GetMaterials(), CreateMaterialInstanceDynamic);
+		SetParameterValueOnMaterials_Internal(RetrievedMaterialSlotsOverlayMaterials, CreateMaterialSlotsOverlayMaterialInstanceDynamic);
 	}
 	else
 	{
@@ -471,32 +502,55 @@ void UMeshComponent::SetScalarParameterValueOnMaterials(const FName ParameterNam
 			CacheMaterialParameterNameIndices();
 		}
 
-		// Look up material index array according to ParameterName
-		if (FMaterialParameterCache* ParameterCache = MaterialParameterCache.Find(ParameterName))
-		{
-			const TArray<int32>& MaterialIndices = ParameterCache->ScalarParameterMaterialIndices;
-			// Loop over all the material indices and update set the parameter value on the corresponding materials		
-			for (int32 MaterialIndex : MaterialIndices)
+		auto GetMaterialAtIndex = [this](int32 MaterialIndex) { return GetMaterial(MaterialIndex); };
+		auto GetMaterialSlotOverlayMaterialAtIndex = [RetrievedMaterialSlotsOverlayMaterials](int32 MaterialIndex) { return RetrievedMaterialSlotsOverlayMaterials[MaterialIndex]; };
+
+		auto SetParameterValueOnMaterials_Internal = [&SetParameterFunction, ParameterName, this]
+			(FMaterialParameterCache* ParameterCache
+			, const TFunction<UMaterialInterface*(int32 MaterialIndex)>& GetMaterialFunction
+			, const TFunction<UMaterialInstanceDynamic* (UMaterialInterface* MaterialInterface, int32 MaterialIndex)>& CreateDynamicMaterialInstanceFunction
+			)
 			{
-				UMaterialInterface* MaterialInterface = GetMaterial(MaterialIndex);
-				if (MaterialInterface)
+				if (ParameterCache != nullptr)
 				{
-					UMaterialInstanceDynamic* DynamicMaterial = Cast<UMaterialInstanceDynamic>(MaterialInterface);
-					if (!DynamicMaterial)
+					const TArray<int32>& MaterialIndices = ParameterCache->ScalarParameterMaterialIndices;
+					// Loop over all the material indices and update set the parameter value on the corresponding materials		
+					for (int32 MaterialIndex : MaterialIndices)
 					{
-						DynamicMaterial = CreateAndSetMaterialInstanceDynamic(MaterialIndex);
+						UMaterialInterface* MaterialInterface = GetMaterialFunction(MaterialIndex);
+						if (MaterialInterface)
+						{
+							UMaterialInstanceDynamic* DynamicMaterial = Cast<UMaterialInstanceDynamic>(MaterialInterface);
+							if (!DynamicMaterial)
+							{
+								DynamicMaterial = CreateDynamicMaterialInstanceFunction(MaterialInterface, MaterialIndex);
+							}
+							SetParameterFunction(DynamicMaterial);
+						}
 					}
-					DynamicMaterial->SetScalarParameterValue(ParameterName, ParameterValue);
 				}
-			}
-		}
-		else
-		{
-			UE_LOG(LogMaterialParameter, Log, TEXT("%s material parameter hasn't found on the component %s"), *ParameterName.ToString(), *GetPathName());
-		}
+				else
+				{
+					UE_LOG(LogMaterialParameter, Log, TEXT("%s material parameter hasn't found on the component %s"), *ParameterName.ToString(), *GetPathName());
+				}
+			};
+
+		SetParameterValueOnMaterials_Internal(MaterialParameterCache.Find(ParameterName), GetMaterialAtIndex, CreateMaterialInstanceDynamic);
+		SetParameterValueOnMaterials_Internal(MaterialSlotsOverlayMaterialParameterCache.Find(ParameterName), GetMaterialSlotOverlayMaterialAtIndex, CreateMaterialSlotsOverlayMaterialInstanceDynamic);
+	}
+
+	if (bDynamicSlotOverlayMaterialCreated)
+	{
+		PrecachePSOs();
+		MarkRenderStateDirty();
 	}
 }
 
+void UMeshComponent::SetScalarParameterValueOnMaterials(const FName ParameterName, const float ParameterValue)
+{
+	SetParameterValueOnMaterials(ParameterName, [ParameterName, ParameterValue](UMaterialInstanceDynamic* DynamicMaterial) { DynamicMaterial->SetScalarParameterValue(ParameterName, ParameterValue); });
+}
+
 void UMeshComponent::SetVectorParameterValueOnMaterials(const FName ParameterName, const FVector ParameterValue)
 {
 	SetColorParameterValueOnMaterials(ParameterName, FLinearColor(ParameterValue));
@@ -504,50 +558,7 @@ void UMeshComponent::SetVectorParameterValueOnMaterials(const FName ParameterNam
 
 void UMeshComponent::SetColorParameterValueOnMaterials(const FName ParameterName, const FLinearColor ParameterValue)
 {
-	if (!bEnableMaterialParameterCaching)
-	{
-		const TArray<UMaterialInterface*> MaterialInterfaces = GetMaterials();
-		for (int32 MaterialIndex = 0; MaterialIndex < MaterialInterfaces.Num(); ++MaterialIndex)
-		{
-			UMaterialInterface* MaterialInterface = MaterialInterfaces[MaterialIndex];
-			if (MaterialInterface)
-			{
-				UMaterialInstanceDynamic* DynamicMaterial = Cast<UMaterialInstanceDynamic>(MaterialInterface);
-				if (!DynamicMaterial)
-				{
-					DynamicMaterial = CreateAndSetMaterialInstanceDynamic(MaterialIndex);
-				}
-				DynamicMaterial->SetVectorParameterValue(ParameterName, ParameterValue);
-			}
-		}
-	}
-	else
-	{
-		if (bCachedMaterialParameterIndicesAreDirty)
-		{
-			CacheMaterialParameterNameIndices();
-		}
-
-		// Look up material index array according to ParameterName
-		if (FMaterialParameterCache* ParameterCache = MaterialParameterCache.Find(ParameterName))
-		{
-			const TArray<int32>& MaterialIndices = ParameterCache->VectorParameterMaterialIndices;
-			// Loop over all the material indices and update set the parameter value on the corresponding materials		
-			for (int32 MaterialIndex : MaterialIndices)
-			{
-				UMaterialInterface* MaterialInterface = GetMaterial(MaterialIndex);
-				if (MaterialInterface)
-				{
-					UMaterialInstanceDynamic* DynamicMaterial = Cast<UMaterialInstanceDynamic>(MaterialInterface);
-					if (!DynamicMaterial)
-					{
-						DynamicMaterial = CreateAndSetMaterialInstanceDynamic(MaterialIndex);
-					}
-					DynamicMaterial->SetVectorParameterValue(ParameterName, ParameterValue);
-				}
-			}
-		}
-	}
+	SetParameterValueOnMaterials(ParameterName, [ParameterName, ParameterValue](UMaterialInstanceDynamic* DynamicMaterial) { DynamicMaterial->SetVectorParameterValue(ParameterName, ParameterValue); });
 }
 
 void UMeshComponent::MarkCachedMaterialParameterNameIndicesDirty()
@@ -579,57 +590,64 @@ void UMeshComponent::CacheMaterialParameterNameIndices()
 		return;
 	}
 
-	// Clean up possible previous data
-	MaterialParameterCache.Reset();
-
 	// not sure if this is the best way to do this
 	const UWorld* World = GetWorld();
 	// to set the default value for scalar params, we use a FMaterialResource, which means the world has to be rendering
 	const bool bHasMaterialResource = (World && World->WorldType != EWorldType::Inactive);
-	
-	// Retrieve all used materials
-	TArray<UMaterialInterface*> MaterialInterfaces = GetMaterials();
-	int32 MaterialIndex = 0;
-	for (UMaterialInterface* MaterialInterface : MaterialInterfaces)
-	{
-		if (MaterialInterface)
-		{
-			TArray<FMaterialParameterInfo> OutParameterInfo;
-			TArray<FGuid> OutParameterIds;
 
-			// Retrieve all scalar parameter names from the material
-			MaterialInterface->GetAllScalarParameterInfo(OutParameterInfo, OutParameterIds);
-			for (FMaterialParameterInfo& ParameterInfo : OutParameterInfo)
+	auto CacheMaterialParameterNameIndices_Internal = [bHasMaterialResource](const TArray<UMaterialInterface*>& MaterialInterfacesToCache, TSortedMap<FName, FMaterialParameterCache, FDefaultAllocator, FNameFastLess>& MaterialParameterCacheToUpdate)
+		{
+			// Clean up possible previous data
+			MaterialParameterCacheToUpdate.Reset();
+			
+			int32 MaterialIndex = 0;
+			for (UMaterialInterface* MaterialInterface : MaterialInterfacesToCache)
 			{
-				// Add or retrieve entry for this parameter name
-				FMaterialParameterCache& ParameterCache = MaterialParameterCache.FindOrAdd(ParameterInfo.Name);
-				// Add the corresponding material index
-				ParameterCache.ScalarParameterMaterialIndices.Add(MaterialIndex);
-				
-				// GetScalarParameterDefault() expects to use a FMaterialResource, which means the world has to be rendering
-				if (bHasMaterialResource)
+				if (MaterialInterface)
 				{
-					// store the default value
-					 MaterialInterface->GetScalarParameterDefaultValue(ParameterInfo, ParameterCache.ScalarParameterDefaultValue);
+					TArray<FMaterialParameterInfo> OutParameterInfo;
+					TArray<FGuid> OutParameterIds;
+
+					// Retrieve all scalar parameter names from the material
+					MaterialInterface->GetAllScalarParameterInfo(OutParameterInfo, OutParameterIds);
+					for (FMaterialParameterInfo& ParameterInfo : OutParameterInfo)
+					{
+						// Add or retrieve entry for this parameter name
+						FMaterialParameterCache& ParameterCache = MaterialParameterCacheToUpdate.FindOrAdd(ParameterInfo.Name);
+						// Add the corresponding material index
+						ParameterCache.ScalarParameterMaterialIndices.Add(MaterialIndex);
+						
+						// GetScalarParameterDefault() expects to use a FMaterialResource, which means the world has to be rendering
+						if (bHasMaterialResource)
+						{
+							// store the default value
+							 MaterialInterface->GetScalarParameterDefaultValue(ParameterInfo, ParameterCache.ScalarParameterDefaultValue);
+						}
+					}
+
+					// Empty parameter names and ids
+					OutParameterInfo.Reset();
+					OutParameterIds.Reset();
+
+					// Retrieve all vector parameter names from the material
+					MaterialInterface->GetAllVectorParameterInfo(OutParameterInfo, OutParameterIds);
+					for (FMaterialParameterInfo& ParameterInfo : OutParameterInfo)
+					{
+						// Add or retrieve entry for this parameter name
+						FMaterialParameterCache& ParameterCache = MaterialParameterCacheToUpdate.FindOrAdd(ParameterInfo.Name);
+						// Add the corresponding material index
+						ParameterCache.VectorParameterMaterialIndices.Add(MaterialIndex);
+					}
 				}
+				++MaterialIndex;
 			}
+		};
 
-			// Empty parameter names and ids
-			OutParameterInfo.Reset();
-			OutParameterIds.Reset();
+	CacheMaterialParameterNameIndices_Internal(GetMaterials(), MaterialParameterCache);
 
-			// Retrieve all vector parameter names from the material
-			MaterialInterface->GetAllVectorParameterInfo(OutParameterInfo, OutParameterIds);
-			for (FMaterialParameterInfo& ParameterInfo : OutParameterInfo)
-			{
-				// Add or retrieve entry for this parameter name
-				FMaterialParameterCache& ParameterCache = MaterialParameterCache.FindOrAdd(ParameterInfo.Name);
-				// Add the corresponding material index
-				ParameterCache.VectorParameterMaterialIndices.Add(MaterialIndex);
-			}
-		}
-		++MaterialIndex;
-	}
+	TArray<TObjectPtr<UMaterialInterface>> RetrievedMaterialSlotsOverlayMaterials;
+	GetMaterialSlotsOverlayMaterial(RetrievedMaterialSlotsOverlayMaterials);
+	CacheMaterialParameterNameIndices_Internal(RetrievedMaterialSlotsOverlayMaterials, MaterialSlotsOverlayMaterialParameterCache);
 
 	bCachedMaterialParameterIndicesAreDirty = false;
 }
