[TRACE] Improved reporting of LevelStreaming, made it output trace regions instead of bookmarks to visualize the time interval it took to stream in levels much clearer

#rnx
#Synced-CL Unknown

diff --git a/Engine/Source/Runtime/Engine/Private/LevelStreaming.cpp b/Engine/Source/Runtime/Engine/Private/LevelStreaming.cpp
--- a/Engine/Source/Runtime/Engine/Private/LevelStreaming.cpp
+++ b/Engine/Source/Runtime/Engine/Private/LevelStreaming.cpp
@@ -1731,8 +1731,8 @@ bool ULevelStreaming::RequestLevel(UWorld* PersistentWorld, bool bAllowLevelLoad
 
 		// Kick off async load request.
 		STAT_ADD_CUSTOMMESSAGE_NAME( STAT_NamedMarker, *(FString( TEXT( "RequestLevel - " ) + DesiredPackageName.ToString() )) );
-		TRACE_BOOKMARK(TEXT("RequestLevel - %s"), *DesiredPackageName.ToString());
-			
+		TRACE_BEGIN_REGION(*WriteToString<256>(TEXT("Load: "), FNameBuilder(DesiredPackageName)), TEXT("LevelStreaming"));
+
 		FLinkerInstancingContext* InstancingContextPtr = nullptr;
 #if WITH_EDITOR
 		FLinkerInstancingContext InstancingContext;
@@ -1927,7 +1927,7 @@ void ULevelStreaming::AsyncLevelLoadComplete(const FName& InPackageName, UPackag
 	ULevel::StreamedLevelsOwningWorld.Remove(InPackageName);
 
 	STAT_ADD_CUSTOMMESSAGE_NAME( STAT_NamedMarker, *(FString( TEXT( "RequestLevelComplete - " ) + InPackageName.ToString() )) );
-	TRACE_BOOKMARK(TEXT("RequestLevelComplete - %s"), *InPackageName.ToString());
+	TRACE_END_REGION(*WriteToString<256>(TEXT("Load: "), FNameBuilder(InPackageName)));
 }
 
 bool ULevelStreaming::IsStreamingStatePending() const
