Add the option to set material overlay slots in blueprint function "SetOverlayMaterial" on any mesh component.

#Synced-CL 48230589

diff --git a/Engine/Source/Runtime/Engine/Classes/Components/MeshComponent.h b/Engine/Source/Runtime/Engine/Classes/Components/MeshComponent.h
--- a/Engine/Source/Runtime/Engine/Classes/Components/MeshComponent.h
+++ b/Engine/Source/Runtime/Engine/Classes/Components/MeshComponent.h
@@ -47,13 +47,13 @@ public:
 	UPROPERTY(EditAnywhere, BlueprintReadOnly, AdvancedDisplay, Category=Rendering)
 	float OverlayMaterialMaxDrawDistance;
 
-	/** Get the overlay material used by this instance */
+	/** Get the overlay material used by this instance. If GetMaterialSlot is enabled, the relative slot material is returned. */
 	UFUNCTION(BlueprintCallable, Category="Rendering|Material")
-	ENGINE_API class UMaterialInterface* GetOverlayMaterial() const;
+	ENGINE_API class UMaterialInterface* GetOverlayMaterial(bool bGetMaterialSlot = false, int32 SlotIndex = 0) const;
 
-	/** Change the overlay material used by this instance */
+	/** Change the overlay material used by this instance. If SetMaterialSlot is enabled, the relative slot material is changed. */
 	UFUNCTION(BlueprintCallable, Category="Rendering|Material")
-	ENGINE_API void SetOverlayMaterial(class UMaterialInterface* NewOverlayMaterial);
+	ENGINE_API void SetOverlayMaterial(class UMaterialInterface* NewOverlayMaterial, bool bSetMaterialSlot = false, int32 SlotIndex = 0);
 
 	/** Get the overlay material used by this instance */
 	UFUNCTION(BlueprintCallable, Category="Rendering|Material")
diff --git a/Engine/Source/Runtime/Engine/Private/Components/MeshComponent.cpp b/Engine/Source/Runtime/Engine/Private/Components/MeshComponent.cpp
--- a/Engine/Source/Runtime/Engine/Private/Components/MeshComponent.cpp
+++ b/Engine/Source/Runtime/Engine/Private/Components/MeshComponent.cpp
@@ -301,23 +301,42 @@ void UMeshComponent::GetUsedMaterials(TArray<UMaterialInterface*>& OutMaterials,
 	}
 }
 
-UMaterialInterface* UMeshComponent::GetOverlayMaterial() const
+UMaterialInterface* UMeshComponent::GetOverlayMaterial(bool bGetMaterialSlot, int32 SlotIndex) const
 {
-	if (OverlayMaterial)
+	if (bGetMaterialSlot)
+	{
+		if (MaterialSlotsOverlayMaterial.IsValidIndex(SlotIndex))
+		{
+			return MaterialSlotsOverlayMaterial[SlotIndex];
+		}
+	}
+	else if (OverlayMaterial)
 	{ 
 		return OverlayMaterial;
 	}
+
+	return GetDefaultOverlayMaterial();
+}
+
+void UMeshComponent::SetOverlayMaterial(UMaterialInterface* NewOverlayMaterial, bool bSetMaterialSlot, int32 SlotIndex)
+{
+	TObjectPtr<UMaterialInterface>* OverrideOverlaySlot = nullptr;
+
+	if (bSetMaterialSlot)
+	{
+		if (MaterialSlotsOverlayMaterial.IsValidIndex(SlotIndex))
+		{
+			OverrideOverlaySlot = &MaterialSlotsOverlayMaterial[SlotIndex];
+		}
+	}
 	else
 	{
-		return GetDefaultOverlayMaterial();
+		OverrideOverlaySlot = &OverlayMaterial;
 	}
-}
 
-void UMeshComponent::SetOverlayMaterial(UMaterialInterface* NewOverlayMaterial)
-{
-	if (OverlayMaterial != NewOverlayMaterial)
+	if (OverrideOverlaySlot && *OverrideOverlaySlot != NewOverlayMaterial)
 	{
-		OverlayMaterial = NewOverlayMaterial;
+		*OverrideOverlaySlot = NewOverlayMaterial;
 		// Precache PSOs again
 		PrecachePSOs();
 		MarkRenderStateDirty();
