[Level Instance] Fix a crash that could occur when trying to edit Level Instances

Change ULevelStreamingLevelInstance::ResetLevelInstanceLoaders to reset loaders for objects outered to LoadedLevel before detaching the OuterWorld linker, rather than after.
ResetLoaders might trigger async loads, which can in turn cause further loads that require the OuterWorld linker to complete.
This might happen, for example, when subobjects are replaced in PostLoad. If GC occurs between that replacement and the ResetLoaders call, the original object will be reloaded.

#Synced-CL 47987596

diff --git a/Engine/Source/Runtime/Engine/Private/LevelInstance/LevelInstanceLevelStreaming.cpp b/Engine/Source/Runtime/Engine/Private/LevelInstance/LevelInstanceLevelStreaming.cpp
--- a/Engine/Source/Runtime/Engine/Private/LevelInstance/LevelInstanceLevelStreaming.cpp
+++ b/Engine/Source/Runtime/Engine/Private/LevelInstance/LevelInstanceLevelStreaming.cpp
@@ -318,6 +318,20 @@ void ULevelStreamingLevelInstance::ResetLevelInstanceLoaders()
 	if (UWorld* OuterWorld = LoadedLevel ? LoadedLevel->GetTypedOuter<UWorld>() : nullptr)
 	{
 		UE_SCOPED_TIMER(*FString::Printf(TEXT("ULevelStreamingLevelInstance::ResetLevelInstanceLoaders(%s)"), *FPaths::GetBaseFilename(OuterWorld->GetPackage()->GetName())), LogLevelInstance, Log);
+		
+		// Reset loaders for objects outered to LoadedLevel before detaching the OuterWorld linker.
+		// ResetLoaders might trigger async loads, which can in turn cause further loads that require the OuterWorld linker to complete.
+		// This might happen, for example, when subobjects are replaced in PostLoad.
+		// If GC occurs between that replacement and the ResetLoaders call, the original object will be reloaded.
+		TArray<UObject*> PackagesToReset;
+		ForEachObjectWithOuter(LoadedLevel, [&PackagesToReset](UObject* InObject)
+		{
+			if (InObject && InObject->IsPackageExternal())
+			{
+				PackagesToReset.Add(InObject->GetExternalPackage());
+			}
+		}, /*bIncludeNestedObjects*/ true);
+		ResetLoaders(PackagesToReset);
 
 		FName PackageName = OuterWorld->GetPackage()->GetLoadedPath().GetPackageFName();
 		if (!ULevel::GetIsLevelPartitionedFromPackage(PackageName))
@@ -331,16 +345,6 @@ void ULevelStreamingLevelInstance::ResetLevelInstanceLoaders()
 			LinkerLoad->DetachLoader();
 		}
 
-		TArray<UObject*> PackagesToReset;
-		ForEachObjectWithOuter(LoadedLevel, [&PackagesToReset](UObject* InObject)
-		{
-			if (InObject && InObject->IsPackageExternal())
-			{
-				PackagesToReset.Add(InObject->GetExternalPackage());
-			}
-		}, /*bIncludeNestedObjects*/ true);
-		ResetLoaders(PackagesToReset);
-
 		bResetLoadersCalled = true;
 	}
 }
