Make SharedContext a unique ptr to force destruction when released in FAnimNode_SaveCachedPose::PostGraphUpdate
#Synced-CL 48139024

diff --git a/Engine/Source/Runtime/Engine/Classes/Animation/AnimNode_SaveCachedPose.h b/Engine/Source/Runtime/Engine/Classes/Animation/AnimNode_SaveCachedPose.h
--- a/Engine/Source/Runtime/Engine/Classes/Animation/AnimNode_SaveCachedPose.h
+++ b/Engine/Source/Runtime/Engine/Classes/Animation/AnimNode_SaveCachedPose.h
@@ -59,8 +59,16 @@ struct FAnimNode_SaveCachedPose : public FAnimNode_Base
 protected:
 	struct FCachedUpdateContext
 	{
+		FCachedUpdateContext() = default;
+		~FCachedUpdateContext() = default;
+		FCachedUpdateContext(FCachedUpdateContext&&) noexcept = default;
+		FCachedUpdateContext& operator=(FCachedUpdateContext&&) noexcept = default;
+
+		FCachedUpdateContext(FCachedUpdateContext&) = delete;
+		FCachedUpdateContext& operator=(const FCachedUpdateContext&) = delete;
+
 		FAnimationUpdateContext Context;
-		TSharedPtr<FAnimationUpdateSharedContext> SharedContext;
+		TUniquePtr<FAnimationUpdateSharedContext> SharedContext;
 	};
 
 	TArray<FCachedUpdateContext> CachedUpdateContexts;
@@ -83,3 +91,12 @@ public:
 
 	ENGINE_API void PostGraphUpdate();
 };
+
+template<>
+struct TStructOpsTypeTraits<FAnimNode_SaveCachedPose> : public TStructOpsTypeTraitsBase2<FAnimNode_SaveCachedPose>
+{
+	enum
+	{
+		WithCopy = false
+	};
+};
diff --git a/Engine/Source/Runtime/Engine/Private/Animation/AnimNode_SaveCachedPose.cpp b/Engine/Source/Runtime/Engine/Private/Animation/AnimNode_SaveCachedPose.cpp
--- a/Engine/Source/Runtime/Engine/Private/Animation/AnimNode_SaveCachedPose.cpp
+++ b/Engine/Source/Runtime/Engine/Private/Animation/AnimNode_SaveCachedPose.cpp
@@ -116,7 +116,7 @@ void FAnimNode_SaveCachedPose::Update_AnyThread(const FAnimationUpdateContext& C
 	// Make a minimal copy of the shared context for cached updates
 	if (FAnimationUpdateSharedContext* SharedContext = Context.GetSharedContext())
 	{
-		CachedUpdate.SharedContext = MakeShared<FAnimationUpdateSharedContext>();
+		CachedUpdate.SharedContext = MakeUnique<FAnimationUpdateSharedContext>();
 		CachedUpdate.SharedContext->CopyForCachedUpdate(*SharedContext);
 	}
 
@@ -202,28 +202,31 @@ void FAnimNode_SaveCachedPose::PostGraphUpdate()
 		}
 
 		// Update any branches that will be skipped
-		UE::Anim::FMessageStack& MessageStack = CachedUpdateContexts[MaxWeightIdx].SharedContext->MessageStack;
-		if(MessageStack.HasMessage<UE::Anim::FCachedPoseSkippedUpdateHandler>())
+		if (CachedUpdateContexts[MaxWeightIdx].SharedContext)
 		{
-			// Grab handles to all the execution paths that we are not proceeding with
-			TArray<UE::Anim::FMessageStack, TInlineAllocator<4>> SkippedMessageStacks;
-			for (int32 CurrIdx = 0; CurrIdx < NumContexts; ++CurrIdx)
+			UE::Anim::FMessageStack& MessageStack = CachedUpdateContexts[MaxWeightIdx].SharedContext->MessageStack;
+			if (MessageStack.HasMessage<UE::Anim::FCachedPoseSkippedUpdateHandler>())
 			{
-				if (CurrIdx != MaxWeightIdx)
+				// Grab handles to all the execution paths that we are not proceeding with
+				TArray<UE::Anim::FMessageStack, TInlineAllocator<4>> SkippedMessageStacks;
+				for (int32 CurrIdx = 0; CurrIdx < NumContexts; ++CurrIdx)
 				{
-					SkippedMessageStacks.Add(MoveTemp(CachedUpdateContexts[CurrIdx].SharedContext->MessageStack));
+					if (CurrIdx != MaxWeightIdx && CachedUpdateContexts[CurrIdx].SharedContext)
+					{
+						SkippedMessageStacks.Add(MoveTemp(CachedUpdateContexts[CurrIdx].SharedContext->MessageStack));
+					}
 				}
-			}
 
-			// Broadcast skipped update message to interested parties
-			MessageStack.ForEachMessage<UE::Anim::FCachedPoseSkippedUpdateHandler>([&SkippedMessageStacks](UE::Anim::FCachedPoseSkippedUpdateHandler& InMessage)
-			{
-				// Fire the event to inform listeners of 'skipped' paths
-				InMessage.OnUpdatesSkipped(SkippedMessageStacks);
+				// Broadcast skipped update message to interested parties
+				MessageStack.ForEachMessage<UE::Anim::FCachedPoseSkippedUpdateHandler>([&SkippedMessageStacks](UE::Anim::FCachedPoseSkippedUpdateHandler& InMessage)
+					{
+						// Fire the event to inform listeners of 'skipped' paths
+						InMessage.OnUpdatesSkipped(SkippedMessageStacks);
 
-				// We only want the topmost registered node here, so early out
-				return UE::Anim::FMessageStack::EEnumerate::Stop;
-			});
+						// We only want the topmost registered node here, so early out
+						return UE::Anim::FMessageStack::EEnumerate::Stop;
+					});
+			}
 		}
 	}
 
