Force property access pre-event graph update callsite to be invoked prior to InitializeRootNode being called.  This prevents property access errors coming from anim nodes that call GetEvaluateGraphExposedInputs from Initialize_AnyThread
#Synced-CL 48139024

diff --git a/Engine/Source/Runtime/Engine/Private/Animation/AnimInstance.cpp b/Engine/Source/Runtime/Engine/Private/Animation/AnimInstance.cpp
--- a/Engine/Source/Runtime/Engine/Private/Animation/AnimInstance.cpp
+++ b/Engine/Source/Runtime/Engine/Private/Animation/AnimInstance.cpp
@@ -28,6 +28,7 @@
 #include "Animation/ActiveMontageInstanceScope.h"
 #include "Animation/AnimNode_SaveCachedPose.h"
 #include "Animation/AnimSubsystem_SharedLinkedAnimLayers.h"
+#include "Animation/AnimSubsystem_PropertyAccess.h"
 #include "ObjectTrace.h"
 #if WITH_EDITOR
 #include "Engine/Blueprint.h"
@@ -368,6 +369,21 @@ void UAnimInstance::InitializeAnimation(bool bInDeferRootNodeInitialization)
 			}
 		}
 #endif
+
+		// Update any game thread pre-event graph property access calls so that these variables are available
+		// for any anim nodes that call GetEvaluateGraphExposedInputs from Initialize_AnyThread
+		AnimBlueprintClass->ForEachSubsystem(this, [this](const FAnimSubsystemInstanceContext& InContext)
+		{
+			if (InContext.SubsystemStruct == FAnimSubsystem_PropertyAccess::StaticStruct())
+			{
+				FAnimSubsystemUpdateContext UpdateContext(InContext, this, 0.f);
+				InContext.Subsystem.OnPreUpdate_GameThread(UpdateContext);
+
+				return EAnimSubsystemEnumeration::Stop;
+			}
+
+			return EAnimSubsystemEnumeration::Continue;
+		});
 	}
 
 	// before initialize, need to recalculate required bone list
