UAnimMontage::SetCompositeLength failing due to DataModelInterface being nullptr
#fix guard against cooked AnimMontages (cooked out, editor-only) data model usage

diff --git a/Engine/Source/Runtime/Engine/Private/Animation/AnimMontage.cpp b/Engine/Source/Runtime/Engine/Private/Animation/AnimMontage.cpp
--- a/Engine/Source/Runtime/Engine/Private/Animation/AnimMontage.cpp
+++ b/Engine/Source/Runtime/Engine/Private/Animation/AnimMontage.cpp
@@ -1417,15 +1417,20 @@ bool UAnimMontage::ContainRecursive(TArray<UAnimCompositeBase*>& CurrentAccumula
 void UAnimMontage::SetCompositeLength(float InLength)
 {
 #if WITH_EDITOR
-	const FFrameTime LengthInFrameTime = DataModelInterface->GetFrameRate().AsFrameTime(InLength);
-	PRAGMA_DISABLE_DEPRECATION_WARNINGS
-	Controller->SetNumberOfFrames(LengthInFrameTime.RoundToFrame());
-	PRAGMA_ENABLE_DEPRECATION_WARNINGS
-#else
-	PRAGMA_DISABLE_DEPRECATION_WARNINGS
-	SequenceLength = InLength;
-	PRAGMA_ENABLE_DEPRECATION_WARNINGS
-#endif	
+	if (!GetOutermost()->HasAnyPackageFlags(PKG_Cooked))
+	{
+		const FFrameTime LengthInFrameTime = DataModelInterface->GetFrameRate().AsFrameTime(InLength);
+		PRAGMA_DISABLE_DEPRECATION_WARNINGS
+		Controller->SetNumberOfFrames(LengthInFrameTime.RoundToFrame());
+		PRAGMA_ENABLE_DEPRECATION_WARNINGS
+	}
+	else
+#endif
+	{
+		PRAGMA_DISABLE_DEPRECATION_WARNINGS
+		SequenceLength = InLength;
+		PRAGMA_ENABLE_DEPRECATION_WARNINGS
+	}
 }
 
 //////////////////////////////////////////////////////////////////////////////////////////////
