Optimize FNewSpline to skip running UpdateSpline() twice when UpdateSpline(const FSpline::FUpdateSplineParams& Params) is called.

UpdateSpline is super expensive and it doesn't seem necessary to run this twice in a row.

#rnx

#Synced-CL 48039674

#Synced-CL 47965127

diff --git a/Engine/Source/Runtime/Engine/Private/Curves/Spline.cpp b/Engine/Source/Runtime/Engine/Private/Curves/Spline.cpp
--- a/Engine/Source/Runtime/Engine/Private/Curves/Spline.cpp
+++ b/Engine/Source/Runtime/Engine/Private/Curves/Spline.cpp
@@ -103,7 +103,14 @@ FAutoConsoleVariableRef CVarValidateRotScale(
 	GValidateRotScale,
 	TEXT("True if we should validate rotation and scale attributes when structurally modifying the spline.")
 );
-	
+
+bool GSkipDoubleUpdate = true;
+FAutoConsoleVariableRef CVarSkipDoubleUpdate(
+	TEXT("Spline.SkipDoubleUpdate"),
+	GSkipDoubleUpdate,
+	TEXT("When true, skip updating the spline twice if UpdateSpline with parameters is called.")
+);
+
 }
 
 UE::Geometry::Spline::EParameterizationPolicy GetParameterizationPolicy()
@@ -206,6 +213,7 @@ public:
 	float GetSplineLength();
 	
     virtual void SetClosedLoop(bool bInClosedLoop) override;
+	void SetClosedLoop(bool bInClosedLoop, bool bUpdateSpline);
 
 	void Reset();
 	
@@ -1688,6 +1696,8 @@ void FSpline::ResetScale()
 
 void FSpline::SetClosedLoop(bool bClosed)
 {
+	TRACE_CPUPROFILER_EVENT_SCOPE(FSpline::SetClosedLoop);
+
 	CachedUpdateSplineParams.bClosedLoop = bClosed;
 	UpdateSpline();
 }
@@ -1699,6 +1709,8 @@ bool FSpline::IsClosedLoop() const
 
 void FSpline::UpdateSpline(const FUpdateSplineParams& InParams)
 {
+	TRACE_CPUPROFILER_EVENT_SCOPE(FSpline::UpdateSpline);
+
 #if WITH_EDITOR
 	if (UE::Spline::GImplementation != CurrentImplementation)
 	{
@@ -2010,15 +2022,22 @@ float FNewSpline::GetSplineLength()
 }
 
 void FNewSpline::SetClosedLoop(bool bInClosedLoop)
+{
+	SetClosedLoop(bInClosedLoop, true);
+}
+
+void FNewSpline::SetClosedLoop(bool bInClosedLoop, bool bUpdateSpline)
 {
 	if (bInClosedLoop == GetSpline().IsClosedLoop())
 	{
 		return;
 	}
+	
+	TRACE_CPUPROFILER_EVENT_SCOPE(FNewSpline::SetClosedLoop);
 
 	MarkLegacyCurvesDirty();
-	
-    GetSpline().SetClosedLoop(bInClosedLoop);
+
+	GetSpline().SetClosedLoop(bInClosedLoop);
 	
 	UE::Geometry::Spline::TTangentBezierSpline<FVector>* ScaleChild = GetTypedAttributeChannel<UE::Geometry::Spline::TTangentBezierSpline<FVector>>(ScaleAttrName);
 	ScaleChild->SetClosedLoop(bInClosedLoop);
@@ -2028,7 +2047,10 @@ void FNewSpline::SetClosedLoop(bool bInClosedLoop)
 	RotChild->SetClosedLoop(bInClosedLoop);
 	RotChild->SetKnotVector(GetSpline().GetKnotVector());
 	
-    UpdateSpline();
+	if (bUpdateSpline)
+	{
+		UpdateSpline();
+	}
 }
 
 void FNewSpline::Reset()
@@ -2505,6 +2527,8 @@ FVector FNewSpline::EvaluateScale(float Parameter) const
 
 void FNewSpline::UpdateSpline()
 {
+	TRACE_CPUPROFILER_EVENT_SCOPE(FNewSpline::UpdateSpline);
+
 	// Update our parameterization, forward it to the rotation and scale channels, and then recompute tangents.
 	Reparameterize();
 
@@ -2516,9 +2540,21 @@ void FNewSpline::UpdateSpline()
 
 void FNewSpline::UpdateSpline(const FSpline::FUpdateSplineParams& Params)
 {
+	TRACE_CPUPROFILER_EVENT_SCOPE(FNewSpline::UpdateSpline);
+
 	// todo: handled unhandled params
-    
-    SetClosedLoop(Params.bClosedLoop);
+   
+	if (UE::Spline::GSkipDoubleUpdate)
+	{
+		// Avoid updating the spline twice.
+		constexpr bool bUpdateSpline = false;
+		SetClosedLoop(Params.bClosedLoop, bUpdateSpline);
+	}
+	else
+	{
+		SetClosedLoop(Params.bClosedLoop);
+	}
+
 	GetSpline().SetStationaryEndpoints(Params.bStationaryEndpoints);
     ReparamStepsPerSegment = Params.ReparamStepsPerSegment;
 
@@ -2527,6 +2563,8 @@ void FNewSpline::UpdateSpline(const FSpline::FUpdateSplineParams& Params)
 
 void FNewSpline::Reparameterize()
 {
+	TRACE_CPUPROFILER_EVENT_SCOPE(FNewSpline::Reparameterize);
+
 	// Caching segment space attribute parameters so we can restore them (necessary because we store the mapping in parent space which are about to be invalidated).
 	TMap<FName, TPair<FInterval1f, TArray<float>>> CachedChannelParameterizations;
 	for (const FName& ChannelName : GetAttributeChannelNamesByValueType<float>())
@@ -3297,6 +3335,8 @@ void FLegacySpline::ResetScale()
 
 void FLegacySpline::UpdateSpline(const FSpline::FUpdateSplineParams& InParams)
 {
+	TRACE_CPUPROFILER_EVENT_SCOPE(FLegacySpline::UpdateSpline);
+
 	const int32 NumPoints = PositionCurve.Points.Num();
 	check(RotationCurve.Points.Num() == NumPoints && ScaleCurve.Points.Num() == NumPoints);
 
