Iris Parallelization - Adding thread safety to NetworkMetricsDatabase and the related stats in NetDriver

[RN] Added thread safety to NetworkMetricsDatabase and the related stats in NetDriver
#tests Built and tested on Dedicated Server Cooked Data
#Synced-CL 47572861

diff --git a/Engine/Source/Runtime/Engine/Classes/Engine/NetDriver.h b/Engine/Source/Runtime/Engine/Classes/Engine/NetDriver.h
--- a/Engine/Source/Runtime/Engine/Classes/Engine/NetDriver.h
+++ b/Engine/Source/Runtime/Engine/Classes/Engine/NetDriver.h
@@ -11,7 +11,9 @@
 #include "Net/NetworkMetricsDatabase.h"
 #endif
 #include "HAL/IConsoleManager.h"
+#include "Iris/IrisConfig.h"
 #include "Math/RandomStream.h"
+#include "Misc/TransactionallySafeRWLock.h"
 #include "UObject/ObjectMacros.h"
 #include "UObject/UObjectGlobals.h"
 #include "UObject/Object.h"
@@ -1174,7 +1176,37 @@ public:
 		PRAGMA_ENABLE_DEPRECATION_WARNINGS
 	}
 
+#if UE_SUPPORT_PARALLEL_IRIS
+	void LockNetDriverStats(const bool bWriteLock)
+	{
+		if (bWriteLock)
+		{
+			NetDriverStatsMutex.WriteLock();
+		}
+		else
+		{
+			NetDriverStatsMutex.ReadLock();
+		}
+	}
+
+	void UnlockNetDriverStats(const bool bWriteLock)
+	{
+		if (bWriteLock)
+		{
+			NetDriverStatsMutex.WriteUnlock();
+		}
+		else
+		{
+			NetDriverStatsMutex.ReadUnlock();
+		}
+	}
+#endif // UE_SUPPORT_PARALLEL_IRIS
+
 private:
+#if UE_SUPPORT_PARALLEL_IRIS
+	/** Used by NetConnection to protect the public stats variables */
+	FTransactionallySafeRWLock NetDriverStatsMutex;
+#endif // UE_SUPPORT_PARALLEL_IRIS
 
 	/** Accumulated delta time since the creation of the NetDriver */
 	double ElapsedTime;
diff --git a/Engine/Source/Runtime/Engine/Private/Net/NetworkMetricsDatabase.cpp b/Engine/Source/Runtime/Engine/Private/Net/NetworkMetricsDatabase.cpp
--- a/Engine/Source/Runtime/Engine/Private/Net/NetworkMetricsDatabase.cpp
+++ b/Engine/Source/Runtime/Engine/Private/Net/NetworkMetricsDatabase.cpp
@@ -3,6 +3,7 @@
 #include "Net/NetworkMetricsDatabase.h"
 #include "Net/NetworkMetricsMutators.h"
 #include "EngineStats.h"
+#include "Iris/IrisConfig.h"
 #include "Logging/LogMacros.h"
 #include "ProfilingDebugging/CsvProfiler.h"
 #include "StructUtils/InstancedStruct.h"
@@ -17,6 +18,7 @@ DEFINE_LOG_CATEGORY_STATIC(LogNetworkMetrics, Error, All);
 
 void UNetworkMetricsDatabase::CreateFloat(const FName MetricName, float DefaultValue)
 {
+	UE_IRIS_PARALLEL_EXPR(UE::TWriteScopeLock LockGuard(NetworkMetricLock));
 	if (ensureMsgf(!MetricTypes.Contains(MetricName), TEXT("Metric %s already exists in the database."), *MetricName.ToString()))
 	{
 		UE::Net::FNetworkMetric<float> Metric;
@@ -29,6 +31,7 @@ void UNetworkMetricsDatabase::CreateFloat(const FName MetricName, float DefaultV
 
 void UNetworkMetricsDatabase::CreateInt(const FName MetricName, int64 DefaultValue)
 {
+	UE_IRIS_PARALLEL_EXPR(UE::TWriteScopeLock LockGuard(NetworkMetricLock));
 	if (ensureMsgf(!MetricTypes.Contains(MetricName), TEXT("Metric %s already exists in the database."), *MetricName.ToString()))
 	{
 		UE::Net::FNetworkMetric<int64> Metric;
@@ -41,6 +44,7 @@ void UNetworkMetricsDatabase::CreateInt(const FName MetricName, int64 DefaultVal
 
 bool UNetworkMetricsDatabase::SetFloat(const FName MetricName, float Value)
 {
+	UE_IRIS_PARALLEL_EXPR(UE::TWriteScopeLock LockGuard(NetworkMetricLock));
 	UE::Net::FNetworkMetric<float>* Metric = MetricFloats.Find(MetricName);
 	if (Metric)
 	{
@@ -56,6 +60,7 @@ bool UNetworkMetricsDatabase::SetFloat(const FName MetricName, float Value)
 
 bool UNetworkMetricsDatabase::SetMinFloat(const FName MetricName, float Value)
 {
+	UE_IRIS_PARALLEL_EXPR(UE::TWriteScopeLock LockGuard(NetworkMetricLock));
 	UE::Net::FNetworkMetric<float>* Metric = MetricFloats.Find(MetricName);
 	if (Metric)
 	{
@@ -71,6 +76,7 @@ bool UNetworkMetricsDatabase::SetMinFloat(const FName MetricName, float Value)
 
 bool UNetworkMetricsDatabase::SetMaxFloat(const FName MetricName, float Value)
 {
+	UE_IRIS_PARALLEL_EXPR(UE::TWriteScopeLock LockGuard(NetworkMetricLock));
 	UE::Net::FNetworkMetric<float>* Metric = MetricFloats.Find(MetricName);
 	if (Metric)
 	{
@@ -86,6 +92,7 @@ bool UNetworkMetricsDatabase::SetMaxFloat(const FName MetricName, float Value)
 
 float UNetworkMetricsDatabase::GetFloat(const FName MetricName)
 {
+	UE_IRIS_PARALLEL_EXPR(UE::TReadScopeLock LockGuard(NetworkMetricLock));
 	UE::Net::FNetworkMetric<float>* Metric = MetricFloats.Find(MetricName);
 	if (Metric)
 	{
@@ -100,6 +107,7 @@ float UNetworkMetricsDatabase::GetFloat(const FName MetricName)
 
 bool UNetworkMetricsDatabase::SetInt(const FName MetricName, int64 Value)
 {
+	UE_IRIS_PARALLEL_EXPR(UE::TWriteScopeLock LockGuard(NetworkMetricLock));
 	UE::Net::FNetworkMetric<int64>* Metric = MetricInts.Find(MetricName);
 	if (Metric)
 	{
@@ -115,6 +123,7 @@ bool UNetworkMetricsDatabase::SetInt(const FName MetricName, int64 Value)
 
 bool UNetworkMetricsDatabase::SetMinInt(const FName MetricName, int64 Value)
 {
+	UE_IRIS_PARALLEL_EXPR(UE::TWriteScopeLock LockGuard(NetworkMetricLock));
 	UE::Net::FNetworkMetric<int64>* Metric = MetricInts.Find(MetricName);
 	if (Metric)
 	{
@@ -130,6 +139,7 @@ bool UNetworkMetricsDatabase::SetMinInt(const FName MetricName, int64 Value)
 
 bool UNetworkMetricsDatabase::SetMaxInt(const FName MetricName, int64 Value)
 {
+	UE_IRIS_PARALLEL_EXPR(UE::TWriteScopeLock LockGuard(NetworkMetricLock));
 	UE::Net::FNetworkMetric<int64>* Metric = MetricInts.Find(MetricName);
 	if (Metric)
 	{
@@ -145,6 +155,7 @@ bool UNetworkMetricsDatabase::SetMaxInt(const FName MetricName, int64 Value)
 
 bool UNetworkMetricsDatabase::IncrementInt(const FName MetricName, int64 Value)
 {
+	UE_IRIS_PARALLEL_EXPR(UE::TWriteScopeLock LockGuard(NetworkMetricLock));
 	UE::Net::FNetworkMetric<int64>* Metric = MetricInts.Find(MetricName);
 	if (Metric)
 	{
@@ -160,6 +171,7 @@ bool UNetworkMetricsDatabase::IncrementInt(const FName MetricName, int64 Value)
 
 int64 UNetworkMetricsDatabase::GetInt(const FName MetricName)
 {
+	UE_IRIS_PARALLEL_EXPR(UE::TReadScopeLock LockGuard(NetworkMetricLock));
 	UE::Net::FNetworkMetric<int64>* Metric = MetricInts.Find(MetricName);
 	if (Metric)
 	{
@@ -174,11 +186,13 @@ int64 UNetworkMetricsDatabase::GetInt(const FName MetricName)
 
 bool UNetworkMetricsDatabase::Contains(const FName MetricName) const
 {
+	UE_IRIS_PARALLEL_EXPR(UE::TReadScopeLock LockGuard(NetworkMetricLock));
 	return MetricTypes.Contains(MetricName);
 }
 
 void UNetworkMetricsDatabase::AddMetricToSnapshot(UE::Net::FNetworkMetricSnapshot& Snapshot, const FName MetricName, const EMetricType MetricType)
 {
+	UE_IRIS_PARALLEL_EXPR(UE::TWriteScopeLock LockGuard(NetworkMetricLock));
 	if (MetricType == EMetricType::Integer)
 	{
 		const UE::Net::FNetworkMetric<int64>* Metric = MetricInts.Find(MetricName);
@@ -282,20 +296,26 @@ void UNetworkMetricsDatabase::ProcessListeners()
 
 void UNetworkMetricsDatabase::Reset()
 {
-	MetricInts.Reset();
-	MetricFloats.Reset();
-	MetricTypes.Reset();
+	{
+		UE_IRIS_PARALLEL_EXPR(UE::TWriteScopeLock LockGuard(NetworkMetricLock));
+		MetricInts.Reset();
+		MetricFloats.Reset();
+		MetricTypes.Reset();
+	}
+	
 	ResetListeners();
 }
 
 void UNetworkMetricsDatabase::ResetListeners()
 {
+	UE_IRIS_PARALLEL_EXPR(UE::TWriteScopeLock LockGuard(NetworkMetricLock));
 	LastReportListener.Reset();
 	ListenersToMetrics.Reset();
 }
 
 void UNetworkMetricsDatabase::Register(const FName MetricName, TWeakObjectPtr<UNetworkMetricsBaseListener> Reporter)
 {
+	UE_IRIS_PARALLEL_EXPR(UE::TWriteScopeLock LockGuard(NetworkMetricLock));
 	if (!ensure(Reporter.IsValid()))
 	{
 		UE_LOG(LogNetworkMetrics, Warning, TEXT("UNetworkMetricsDatabase::Register: Cannot register metric %s with invalid listener."), *MetricName.ToString());
@@ -357,6 +377,7 @@ void UNetworkMetricsDatabase::RegisterMutator(TWeakObjectPtr<UNetworkMetricsBase
 		}
 	}
 
+	UE_IRIS_PARALLEL_EXPR(UE::TWriteScopeLock LockGuard(NetworkMetricLock));
 	// Make sure the mutator name doesn't collide with a metric name that's already registered to the listener.
 	TSet<FNameAndType>* ListenerMetrics = ListenersToMetrics.Find(Listener);
 	if (ListenerMetrics)
diff --git a/Engine/Source/Runtime/Engine/Private/NetConnection.cpp b/Engine/Source/Runtime/Engine/Private/NetConnection.cpp
--- a/Engine/Source/Runtime/Engine/Private/NetConnection.cpp
+++ b/Engine/Source/Runtime/Engine/Private/NetConnection.cpp
@@ -2531,9 +2531,11 @@ void UNetConnection::FlushNet(bool bIgnoreSimulation)
 
 		OutBytes += PacketBytes;
 		OutTotalBytes += PacketBytes;
+		UE_IRIS_PARALLEL_EXPR(Driver->LockNetDriverStats(true));
 		Driver->OutBytes += PacketBytes;
 		Driver->OutTotalBytes += PacketBytes;
 		GNetOutBytes += PacketBytes;
+		UE_IRIS_PARALLEL_EXPR(Driver->UnlockNetDriverStats(true));
 
 		AnalyticsVars.OutAckOnlyCount += (NumAckBits > 0 && NumBunchBits == 0);
 
@@ -4434,9 +4436,11 @@ int32 UNetConnection::SendRawBunch(FOutBunch& Bunch, bool InAllowMerge, const FN
 	ValidateSendBuffer();
 	check(!Bunch.ReceivedAck);
 	check(!Bunch.IsError());
+	UE_IRIS_PARALLEL_EXPR(Driver->LockNetDriverStats(true));
 	Driver->OutBunches++;
 	Driver->OutTotalBunches++;
 	Driver->OutTotalReliableBunches += Bunch.bReliable ? 1 : 0;
+	UE_IRIS_PARALLEL_EXPR(Driver->UnlockNetDriverStats(true));
 
 	// Build header.
 	SendBunchHeader.Reset();
@@ -4544,7 +4548,9 @@ int32 UNetConnection::SendRawBunch(FOutBunch& Bunch, bool InAllowMerge, const FN
 
 	if (Bunch.bHasPackageMapExports)
 	{
+		UE_IRIS_PARALLEL_EXPR(Driver->LockNetDriverStats(true));
 		Driver->NetGUIDOutBytes += (SendBunchHeader.GetNumBits() + Bunch.GetNumBits()) >> 3;
+		UE_IRIS_PARALLEL_EXPR(Driver->UnlockNetDriverStats(true));
 	}
 
 	if (bAutoFlush)
diff --git a/Engine/Source/Runtime/Engine/Public/Net/NetworkMetricsDatabase.h b/Engine/Source/Runtime/Engine/Public/Net/NetworkMetricsDatabase.h
--- a/Engine/Source/Runtime/Engine/Public/Net/NetworkMetricsDatabase.h
+++ b/Engine/Source/Runtime/Engine/Public/Net/NetworkMetricsDatabase.h
@@ -3,6 +3,9 @@
 #pragma once
 
 #include "CoreMinimal.h"
+#include "HAL/CriticalSection.h"
+#include "Iris/IrisConfig.h"
+#include "Misc/TransactionallySafeRWLock.h"
 #include "NetworkMetricsDatabase.generated.h"
 
 #define UE_API ENGINE_API
@@ -100,6 +103,9 @@ private:
 
 	// Map of listeners to metrics registered with them. A listener may have an entry with an empty set if it has no metrics but some mutators registered. */
 	TMap<TWeakObjectPtr<UNetworkMetricsBaseListener>, TSet<FNameAndType>> ListenersToMetrics;
+#if UE_SUPPORT_PARALLEL_IRIS
+	mutable FTransactionallySafeRWLock NetworkMetricLock;
+#endif // UE_SUPPORT_PARALLEL_IRIS
 };
 
 /** 
