Fix FCompositeNavModifier::CreateAreaModifiers not taking into account FKSphylElem internal rotation
Fix FCompositeNavModifier::CreateAreaModifiers not using the proper length for a FKSphylElem
Fix FAreaNavModifier cylinder constructor not handling non vertical cylinders. This will now fallback to a box if the cylinder is not vertical
#rnx Fix an issue with FAreaNavModifiers that was wrongly assuming that cylinder physics shape we're always vertical
#Synced-CL 48151812

diff --git a/Engine/Source/Runtime/Engine/Private/AI/Navigation/NavigationModifier.cpp b/Engine/Source/Runtime/Engine/Private/AI/Navigation/NavigationModifier.cpp
--- a/Engine/Source/Runtime/Engine/Private/AI/Navigation/NavigationModifier.cpp
+++ b/Engine/Source/Runtime/Engine/Private/AI/Navigation/NavigationModifier.cpp
@@ -337,18 +337,33 @@ FAreaNavModifier::FAreaNavModifier()
 FAreaNavModifier::FAreaNavModifier(float Radius, float Height, const FTransform& LocalToWorld, const TSubclassOf<UNavAreaBase> InAreaClass)
 {
 	Init(InAreaClass);
-	
-	FVector Scale3D = LocalToWorld.GetScale3D().GetAbs();
-	const FVector::FReal RadiusScaled = Radius * FMath::Max(Scale3D.X, Scale3D.Y);
-	const FVector::FReal HeightScaled = Height * Scale3D.Z;
 
-	Points.SetNumUninitialized(2);
-	Points[0] = LocalToWorld.GetLocation();
-	Points[1].X = RadiusScaled;
-	Points[1].Z = HeightScaled;
-	ShapeType = ENavigationShapeType::Cylinder;
+	// check if the cylinder is vertical (Z axis pointing straight up or straight down)
+	const FVector TransformZAxis = LocalToWorld.GetUnitAxis(EAxis::Z);
+	const float VerticalDotProduct = TransformZAxis.Dot(FVector::UpVector);
+	// About 1 degree off vertical
+	if (FMath::IsNearlyEqual(FMath::Abs(VerticalDotProduct), 1.f, UE_KINDA_SMALL_NUMBER))
+	{
+		FVector Scale3D = LocalToWorld.GetScale3D().GetAbs();
+		const FVector::FReal RadiusScaled = Radius * FMath::Max(Scale3D.X, Scale3D.Y);
+		const FVector::FReal HeightScaled = Height * Scale3D.Z;
+
+		Points.SetNumUninitialized(2);
+		Points[0] = LocalToWorld.GetLocation();
+		Points[1].X = RadiusScaled;
+		Points[1].Z = HeightScaled;
+		ShapeType = ENavigationShapeType::Cylinder;
 
-	Bounds = FBox::BuildAABB(LocalToWorld.GetLocation(), FVector(RadiusScaled, RadiusScaled, HeightScaled));
+		Bounds = FBox::BuildAABB(LocalToWorld.GetLocation(), FVector(RadiusScaled, RadiusScaled, HeightScaled));
+	}
+	// Fallback to a box if the cylinder is not vertical
+	else
+	{
+		const float HalfHeight = Height * 0.5f;
+		FVector Center(0.f, 0.f, HalfHeight);
+		FVector Extents(Radius, Radius, HalfHeight);
+		SetBox(FBox::BuildAABB(Center, Extents), LocalToWorld);
+	}
 }
 
 FAreaNavModifier::FAreaNavModifier(const FVector& Extent, const FTransform& LocalToWorld, const TSubclassOf<UNavAreaBase> InAreaClass)
@@ -977,9 +992,10 @@ void FCompositeNavModifier::CreateAreaModifiers(const UPrimitiveComponent* PrimC
 	for (int32 Idx = 0; Idx < BodySetup->AggGeom.SphylElems.Num(); Idx++)
 	{
 		const FKSphylElem& SphylElem = BodySetup->AggGeom.SphylElems[Idx];
-		const FTransform AreaOffset(FVector(0, 0, -SphylElem.Length));
+		const float SphylHalfHeight = SphylElem.Length * 0.5f + SphylElem.Radius;
+		const FTransform AreaOffset(FVector(0, 0, -SphylHalfHeight));
 
-		FAreaNavModifier AreaMod(SphylElem.Radius, SphylElem.Length * 2.0f, AreaOffset * PrimComp->GetComponentTransform(), AreaClass);
+		FAreaNavModifier AreaMod(SphylElem.Radius, 2.f * SphylHalfHeight, AreaOffset * SphylElem.GetTransform() * PrimComp->GetComponentTransform(), AreaClass);
 		if (AreaClassToReplace)
 		{
 			AreaMod.SetAreaClassToReplace(AreaClassToReplace);
