[HWRT] Check whether skeletal mesh asset supports ray tracing.
- fixes some asserts and RAY_TRACING_REPRESENTATION stencil bit.

#Synced-CL 48227131

diff --git a/Engine/Source/Runtime/Engine/Private/SkeletalMeshSceneProxy.cpp b/Engine/Source/Runtime/Engine/Private/SkeletalMeshSceneProxy.cpp
--- a/Engine/Source/Runtime/Engine/Private/SkeletalMeshSceneProxy.cpp
+++ b/Engine/Source/Runtime/Engine/Private/SkeletalMeshSceneProxy.cpp
@@ -975,7 +975,8 @@ void FSkeletalMeshSceneProxy::DestroyRenderThreadResources()
 #if RHI_RAYTRACING
 bool FSkeletalMeshSceneProxy::HasRayTracingRepresentation() const
 {
-	return bRayTraceStatic || (CVarRayTracingSupportSkeletalMeshes.GetValueOnAnyThread() != 0 && CVarRayTracingSkeletalMeshes.GetValueOnAnyThread() != 0);
+	return RenderData->bSupportRayTracing
+		&& (bRayTraceStatic || (CVarRayTracingSupportSkeletalMeshes.GetValueOnAnyThread() != 0 && CVarRayTracingSkeletalMeshes.GetValueOnAnyThread() != 0));
 }
 
 RayTracing::FGeometryGroupHandle FSkeletalMeshSceneProxy::GetRayTracingGeometryGroupHandle() const
@@ -986,7 +987,7 @@ RayTracing::FGeometryGroupHandle FSkeletalMeshSceneProxy::GetRayTracingGeometryG
 
 TArray<FRayTracingGeometry*> FSkeletalMeshSceneProxy::GetStaticRayTracingGeometries() const
 {
-	if (IsRayTracingEnabled() && bRayTraceStatic)
+	if (IsRayTracingEnabled() && bRayTraceStatic && RenderData->bSupportRayTracing)
 	{
 		TArray<FRayTracingGeometry*> RayTracingGeometries;
 		RayTracingGeometries.AddDefaulted(RenderData->LODRenderData.Num());
