Spline Optimizations

+ Early out of SetClosed if it's already true.
+ Avoid re-constructing the TSet of normalized knot values in GetNearestAvailableKnotValue as this was being done repeatedly in GenerateDistanceKnotsForBezier. Reduces cost of GenerateDistanceKnotsForBezier by 90%.
+ Update FNewSpline::RebuildLegacyCurves to cache the attribute channels as these lookups were expensive. Also replace insertion sort of the rotation/scale points with a final sort at the end.
+ Add nullchecks when accessing scale/rotation attribute channels to prevent crashing when those channels are removed.

#Synced-CL 48255652

diff --git a/Engine/Source/Runtime/Engine/Private/Curves/Spline.cpp b/Engine/Source/Runtime/Engine/Private/Curves/Spline.cpp
--- a/Engine/Source/Runtime/Engine/Private/Curves/Spline.cpp
+++ b/Engine/Source/Runtime/Engine/Private/Curves/Spline.cpp
@@ -2,6 +2,7 @@
 
 #include "Curves/Spline.h"
 
+#include "Algo/Sort.h"
 #include "BoxTypes.h"
 #include "Async/TransactionallySafeMutex.h"
 #include "Components/SplineComponent.h"	// for FSplineCurves, ideally removed
@@ -2032,20 +2033,23 @@ void FNewSpline::SetClosedLoop(bool bInClosedLoop, bool bUpdateSpline)
 	{
 		return;
 	}
-	
 	TRACE_CPUPROFILER_EVENT_SCOPE(FNewSpline::SetClosedLoop);
 
 	MarkLegacyCurvesDirty();
 
 	GetSpline().SetClosedLoop(bInClosedLoop);
 	
-	UE::Geometry::Spline::TTangentBezierSpline<FVector>* ScaleChild = GetTypedAttributeChannel<UE::Geometry::Spline::TTangentBezierSpline<FVector>>(ScaleAttrName);
-	ScaleChild->SetClosedLoop(bInClosedLoop);
-	ScaleChild->SetKnotVector(GetSpline().GetKnotVector());
+	if (UE::Geometry::Spline::TTangentBezierSpline<FVector>* ScaleChild = GetTypedAttributeChannel<UE::Geometry::Spline::TTangentBezierSpline<FVector>>(ScaleAttrName))
+	{
+		ScaleChild->SetClosedLoop(bInClosedLoop);
+		ScaleChild->SetKnotVector(GetSpline().GetKnotVector());
+	}
 		
-	UE::Geometry::Spline::TTangentBezierSpline<FQuat>* RotChild = GetTypedAttributeChannel<UE::Geometry::Spline::TTangentBezierSpline<FQuat>>(RotationAttrName);
-	RotChild->SetClosedLoop(bInClosedLoop);
-	RotChild->SetKnotVector(GetSpline().GetKnotVector());
+	if (UE::Geometry::Spline::TTangentBezierSpline<FQuat>* RotChild = GetTypedAttributeChannel<UE::Geometry::Spline::TTangentBezierSpline<FQuat>>(RotationAttrName))
+	{
+		RotChild->SetClosedLoop(bInClosedLoop);
+		RotChild->SetKnotVector(GetSpline().GetKnotVector());
+	}
 	
 	if (bUpdateSpline)
 	{
@@ -2072,8 +2076,10 @@ void FNewSpline::ResetRotation()
     }
 
 	SetAttributeChannelRange(RotationAttrName, FInterval1f(0.f, 1.f), EMappingRangeSpace::Normalized);
-	UE::Geometry::Spline::TTangentBezierSpline<FQuat>* RotSpline = GetTypedAttributeChannel<UE::Geometry::Spline::TTangentBezierSpline<FQuat>>(RotationAttrName);
-	RotSpline->SetKnotVector(GetSpline().GetKnotVector());
+	if (UE::Geometry::Spline::TTangentBezierSpline<FQuat>* RotSpline = GetTypedAttributeChannel<UE::Geometry::Spline::TTangentBezierSpline<FQuat>>(RotationAttrName))
+	{
+		RotSpline->SetKnotVector(GetSpline().GetKnotVector());
+	}
 }
 
 void FNewSpline::ResetScale()
@@ -2086,8 +2092,10 @@ void FNewSpline::ResetScale()
     }
 
 	SetAttributeChannelRange(ScaleAttrName, FInterval1f(0.f, 1.f), EMappingRangeSpace::Normalized);
-	UE::Geometry::Spline::TTangentBezierSpline<FVector>* ScaleSpline = GetTypedAttributeChannel<UE::Geometry::Spline::TTangentBezierSpline<FVector>>(ScaleAttrName);
-	ScaleSpline->SetKnotVector(GetSpline().GetKnotVector());
+	if (UE::Geometry::Spline::TTangentBezierSpline<FVector>* ScaleSpline = GetTypedAttributeChannel<UE::Geometry::Spline::TTangentBezierSpline<FVector>>(ScaleAttrName))
+	{
+		ScaleSpline->SetKnotVector(GetSpline().GetKnotVector());
+	}
 }
 
 void FNewSpline::AddPoint(const FSplinePoint& Point)
@@ -2111,13 +2119,17 @@ void FNewSpline::AddPoint(const FSplinePoint& Point)
 	{
 		GetSpline().AppendPoint(ControlPoint);
 
-		UE::Geometry::Spline::TTangentBezierSpline<FVector>* ScaleChild = GetTypedAttributeChannel<UE::Geometry::Spline::TTangentBezierSpline<FVector>>(ScaleAttrName);
-		ScaleChild->AppendPoint(UE::Geometry::Spline::TTangentBezierSpline<FVector>::FTangentBezierControlPoint(FVector::OneVector));
-		ScaleChild->SetKnotVector(GetSpline().GetKnotVector());
+		if (UE::Geometry::Spline::TTangentBezierSpline<FVector>* ScaleChild = GetTypedAttributeChannel<UE::Geometry::Spline::TTangentBezierSpline<FVector>>(ScaleAttrName))
+		{
+			ScaleChild->AppendPoint(UE::Geometry::Spline::TTangentBezierSpline<FVector>::FTangentBezierControlPoint(FVector::OneVector));
+			ScaleChild->SetKnotVector(GetSpline().GetKnotVector());
+		}
 		
-		UE::Geometry::Spline::TTangentBezierSpline<FQuat>* RotChild = GetTypedAttributeChannel<UE::Geometry::Spline::TTangentBezierSpline<FQuat>>(RotationAttrName);
-		RotChild->AppendPoint(UE::Geometry::Spline::TTangentBezierSpline<FQuat>::FTangentBezierControlPoint(FQuat::Identity));
-		RotChild->SetKnotVector(GetSpline().GetKnotVector());
+		if (UE::Geometry::Spline::TTangentBezierSpline<FQuat>* RotChild = GetTypedAttributeChannel<UE::Geometry::Spline::TTangentBezierSpline<FQuat>>(RotationAttrName))
+		{
+			RotChild->AppendPoint(UE::Geometry::Spline::TTangentBezierSpline<FQuat>::FTangentBezierControlPoint(FQuat::Identity));
+			RotChild->SetKnotVector(GetSpline().GetKnotVector());
+		}
 		
 		MarkLegacyCurvesDirty();
 		return;
@@ -2126,13 +2138,17 @@ void FNewSpline::AddPoint(const FSplinePoint& Point)
     {
     	GetSpline().PrependPoint(ControlPoint);
     	
-    	UE::Geometry::Spline::TTangentBezierSpline<FVector>* ScaleChild = GetTypedAttributeChannel<UE::Geometry::Spline::TTangentBezierSpline<FVector>>(ScaleAttrName);
-    	ScaleChild->PrependPoint(UE::Geometry::Spline::TTangentBezierSpline<FVector>::FTangentBezierControlPoint(FVector::OneVector));
-    	ScaleChild->SetKnotVector(GetSpline().GetKnotVector());
+		if (UE::Geometry::Spline::TTangentBezierSpline<FVector>* ScaleChild = GetTypedAttributeChannel<UE::Geometry::Spline::TTangentBezierSpline<FVector>>(ScaleAttrName))
+		{
+			ScaleChild->PrependPoint(UE::Geometry::Spline::TTangentBezierSpline<FVector>::FTangentBezierControlPoint(FVector::OneVector));
+			ScaleChild->SetKnotVector(GetSpline().GetKnotVector());
+		}
     
-    	UE::Geometry::Spline::TTangentBezierSpline<FQuat>* RotChild = GetTypedAttributeChannel<UE::Geometry::Spline::TTangentBezierSpline<FQuat>>(RotationAttrName);
-    	RotChild->PrependPoint(UE::Geometry::Spline::TTangentBezierSpline<FQuat>::FTangentBezierControlPoint(FQuat::Identity));
-    	RotChild->SetKnotVector(GetSpline().GetKnotVector());
+		if (UE::Geometry::Spline::TTangentBezierSpline<FQuat>* RotChild = GetTypedAttributeChannel<UE::Geometry::Spline::TTangentBezierSpline<FQuat>>(RotationAttrName))
+		{
+			RotChild->PrependPoint(UE::Geometry::Spline::TTangentBezierSpline<FQuat>::FTangentBezierControlPoint(FQuat::Identity));
+			RotChild->SetKnotVector(GetSpline().GetKnotVector());
+		}
     	
     	MarkLegacyCurvesDirty();
     	return;
@@ -2143,14 +2159,18 @@ void FNewSpline::AddPoint(const FSplinePoint& Point)
 	float NewPointInternalParam = GetSpline().GetParameter(InsertedIndex);
 	
 	float ScaleParam = MapParameterToChildSpace(ScaleAttrName, NewPointInternalParam);
-	UE::Geometry::Spline::TTangentBezierSpline<FVector>* ScaleChild = GetTypedAttributeChannel<UE::Geometry::Spline::TTangentBezierSpline<FVector>>(ScaleAttrName);
-	ScaleChild->InsertPointAtGlobalParam(ScaleParam, UE::Geometry::Spline::TTangentBezierSpline<FVector>::FTangentBezierControlPoint(FVector::OneVector));
-	ScaleChild->SetKnotVector(GetSpline().GetKnotVector());
+	if (UE::Geometry::Spline::TTangentBezierSpline<FVector>* ScaleChild = GetTypedAttributeChannel<UE::Geometry::Spline::TTangentBezierSpline<FVector>>(ScaleAttrName))
+	{
+		ScaleChild->InsertPointAtGlobalParam(ScaleParam, UE::Geometry::Spline::TTangentBezierSpline<FVector>::FTangentBezierControlPoint(FVector::OneVector));
+		ScaleChild->SetKnotVector(GetSpline().GetKnotVector());
+	}
 	
 	float RotParam = MapParameterToChildSpace(RotationAttrName, NewPointInternalParam);
-	UE::Geometry::Spline::TTangentBezierSpline<FQuat>* RotChild = GetTypedAttributeChannel<UE::Geometry::Spline::TTangentBezierSpline<FQuat>>(RotationAttrName);
-	RotChild->InsertPointAtGlobalParam(RotParam, UE::Geometry::Spline::TTangentBezierSpline<FQuat>::FTangentBezierControlPoint(FQuat::Identity));
-	RotChild->SetKnotVector(GetSpline().GetKnotVector());
+	if (UE::Geometry::Spline::TTangentBezierSpline<FQuat>* RotChild = GetTypedAttributeChannel<UE::Geometry::Spline::TTangentBezierSpline<FQuat>>(RotationAttrName))
+	{
+		RotChild->InsertPointAtGlobalParam(RotParam, UE::Geometry::Spline::TTangentBezierSpline<FQuat>::FTangentBezierControlPoint(FQuat::Identity));
+		RotChild->SetKnotVector(GetSpline().GetKnotVector());
+	}
 	
 	MarkLegacyCurvesDirty();
 }
@@ -2173,13 +2193,17 @@ void FNewSpline::InsertPoint(const FSplinePoint& Point, int32 Index)
 	{
 		GetSpline().AppendPoint(ControlPoint);
 		
-		UE::Geometry::Spline::TTangentBezierSpline<FVector>* ScaleChild = GetTypedAttributeChannel<UE::Geometry::Spline::TTangentBezierSpline<FVector>>(ScaleAttrName);
-		ScaleChild->AppendPoint(UE::Geometry::Spline::TTangentBezierSpline<FVector>::FTangentBezierControlPoint(FVector::OneVector));
-		ScaleChild->SetKnotVector(GetSpline().GetKnotVector());
+		if (UE::Geometry::Spline::TTangentBezierSpline<FVector>* ScaleChild = GetTypedAttributeChannel<UE::Geometry::Spline::TTangentBezierSpline<FVector>>(ScaleAttrName))
+		{
+			ScaleChild->AppendPoint(UE::Geometry::Spline::TTangentBezierSpline<FVector>::FTangentBezierControlPoint(FVector::OneVector));
+			ScaleChild->SetKnotVector(GetSpline().GetKnotVector());
+		}
 		
-		UE::Geometry::Spline::TTangentBezierSpline<FQuat>* RotChild = GetTypedAttributeChannel<UE::Geometry::Spline::TTangentBezierSpline<FQuat>>(RotationAttrName);
-		RotChild->AppendPoint(UE::Geometry::Spline::TTangentBezierSpline<FQuat>::FTangentBezierControlPoint(FQuat::Identity));
-		RotChild->SetKnotVector(GetSpline().GetKnotVector());
+		if (UE::Geometry::Spline::TTangentBezierSpline<FQuat>* RotChild = GetTypedAttributeChannel<UE::Geometry::Spline::TTangentBezierSpline<FQuat>>(RotationAttrName))
+		{
+			RotChild->AppendPoint(UE::Geometry::Spline::TTangentBezierSpline<FQuat>::FTangentBezierControlPoint(FQuat::Identity));
+			RotChild->SetKnotVector(GetSpline().GetKnotVector());
+		}
 		
 		MarkLegacyCurvesDirty();
 		return;
@@ -2188,13 +2212,17 @@ void FNewSpline::InsertPoint(const FSplinePoint& Point, int32 Index)
 	{
 		GetSpline().PrependPoint(ControlPoint);
 		
-		UE::Geometry::Spline::TTangentBezierSpline<FVector>* ScaleChild = GetTypedAttributeChannel<UE::Geometry::Spline::TTangentBezierSpline<FVector>>(ScaleAttrName);
-		ScaleChild->PrependPoint(UE::Geometry::Spline::TTangentBezierSpline<FVector>::FTangentBezierControlPoint(FVector::OneVector));
-		ScaleChild->SetKnotVector(GetSpline().GetKnotVector());
+		if (UE::Geometry::Spline::TTangentBezierSpline<FVector>* ScaleChild = GetTypedAttributeChannel<UE::Geometry::Spline::TTangentBezierSpline<FVector>>(ScaleAttrName))
+		{
+			ScaleChild->PrependPoint(UE::Geometry::Spline::TTangentBezierSpline<FVector>::FTangentBezierControlPoint(FVector::OneVector));
+			ScaleChild->SetKnotVector(GetSpline().GetKnotVector());
+		}
 	
-		UE::Geometry::Spline::TTangentBezierSpline<FQuat>* RotChild = GetTypedAttributeChannel<UE::Geometry::Spline::TTangentBezierSpline<FQuat>>(RotationAttrName);
-		RotChild->PrependPoint(UE::Geometry::Spline::TTangentBezierSpline<FQuat>::FTangentBezierControlPoint(FQuat::Identity));
-		RotChild->SetKnotVector(GetSpline().GetKnotVector());
+		if (UE::Geometry::Spline::TTangentBezierSpline<FQuat>* RotChild = GetTypedAttributeChannel<UE::Geometry::Spline::TTangentBezierSpline<FQuat>>(RotationAttrName))
+		{
+			RotChild->PrependPoint(UE::Geometry::Spline::TTangentBezierSpline<FQuat>::FTangentBezierControlPoint(FQuat::Identity));
+			RotChild->SetKnotVector(GetSpline().GetKnotVector());
+		}
 		
 		MarkLegacyCurvesDirty();
 		return;
@@ -2206,14 +2234,18 @@ void FNewSpline::InsertPoint(const FSplinePoint& Point, int32 Index)
 	float NewPointInternalParam = GetSpline().GetParameter(InsertedIndex);
 	
 	float ScaleParam = MapParameterToChildSpace(ScaleAttrName, NewPointInternalParam);
-	UE::Geometry::Spline::TTangentBezierSpline<FVector>* ScaleChild = GetTypedAttributeChannel<UE::Geometry::Spline::TTangentBezierSpline<FVector>>(ScaleAttrName);
-	ScaleChild->InsertPointAtGlobalParam(ScaleParam, UE::Geometry::Spline::TTangentBezierSpline<FVector>::FTangentBezierControlPoint(FVector::OneVector));
-	ScaleChild->SetKnotVector(GetSpline().GetKnotVector());
+	if (UE::Geometry::Spline::TTangentBezierSpline<FVector>* ScaleChild = GetTypedAttributeChannel<UE::Geometry::Spline::TTangentBezierSpline<FVector>>(ScaleAttrName))
+	{
+		ScaleChild->InsertPointAtGlobalParam(ScaleParam, UE::Geometry::Spline::TTangentBezierSpline<FVector>::FTangentBezierControlPoint(FVector::OneVector));
+		ScaleChild->SetKnotVector(GetSpline().GetKnotVector());
+	}
 	
 	float RotParam = MapParameterToChildSpace(RotationAttrName, NewPointInternalParam);
-	UE::Geometry::Spline::TTangentBezierSpline<FQuat>* RotChild = GetTypedAttributeChannel<UE::Geometry::Spline::TTangentBezierSpline<FQuat>>(RotationAttrName);
-	RotChild->InsertPointAtGlobalParam(RotParam, UE::Geometry::Spline::TTangentBezierSpline<FQuat>::FTangentBezierControlPoint(FQuat::Identity));
-	RotChild->SetKnotVector(GetSpline().GetKnotVector());
+	if (UE::Geometry::Spline::TTangentBezierSpline<FQuat>* RotChild = GetTypedAttributeChannel<UE::Geometry::Spline::TTangentBezierSpline<FQuat>>(RotationAttrName))
+	{
+		RotChild->InsertPointAtGlobalParam(RotParam, UE::Geometry::Spline::TTangentBezierSpline<FQuat>::FTangentBezierControlPoint(FQuat::Identity));
+		RotChild->SetKnotVector(GetSpline().GetKnotVector());
+	}
 	
 	MarkLegacyCurvesDirty();
 }
@@ -2258,13 +2290,17 @@ void FNewSpline::RemovePoint(int32 Index)
 	
 	GetSpline().RemovePoint(Index);
 
-	UE::Geometry::Spline::TTangentBezierSpline<FVector>* ScaleChild = GetTypedAttributeChannel<UE::Geometry::Spline::TTangentBezierSpline<FVector>>(ScaleAttrName);
-	ScaleChild->RemovePoint(Index);
-	ScaleChild->SetKnotVector(GetSpline().GetKnotVector());
+	if (UE::Geometry::Spline::TTangentBezierSpline<FVector>* ScaleChild = GetTypedAttributeChannel<UE::Geometry::Spline::TTangentBezierSpline<FVector>>(ScaleAttrName))
+	{
+		ScaleChild->RemovePoint(Index);
+		ScaleChild->SetKnotVector(GetSpline().GetKnotVector());
+	}
 	
-	UE::Geometry::Spline::TTangentBezierSpline<FQuat>* RotChild = GetTypedAttributeChannel<UE::Geometry::Spline::TTangentBezierSpline<FQuat>>(RotationAttrName);
-	RotChild->RemovePoint(Index);
-	RotChild->SetKnotVector(GetSpline().GetKnotVector());
+	if (UE::Geometry::Spline::TTangentBezierSpline<FQuat>* RotChild = GetTypedAttributeChannel<UE::Geometry::Spline::TTangentBezierSpline<FQuat>>(RotationAttrName))
+	{
+		RotChild->RemovePoint(Index);
+		RotChild->SetKnotVector(GetSpline().GetKnotVector());
+	}
 	
     MarkLegacyCurvesDirty();
 }
@@ -2615,20 +2651,12 @@ void FNewSpline::Reparameterize()
 		ScaleSpline->SetKnotVector(GetSpline().GetKnotVector());
 		ScaleSpline->UpdateTangents();
 	}
-	else
-	{
-		UE_LOG(LogSpline, Warning, TEXT("Failed to get scale spline while reparameterizing!"));
-	}
 
 	if (UE::Geometry::Spline::TTangentBezierSpline<FQuat>* RotSpline = GetTypedAttributeChannel<UE::Geometry::Spline::TTangentBezierSpline<FQuat>>(RotationAttrName))
 	{
 		RotSpline->SetKnotVector(GetSpline().GetKnotVector());
 		RotSpline->UpdateTangents();
 	}
-	else
-	{
-		UE_LOG(LogSpline, Warning, TEXT("Failed to get rotation spline while reparameterizing!"));
-	}
 }
 
 float FNewSpline::FromNormalizedSpace(float Parameter) const
@@ -2717,6 +2745,8 @@ void FNewSpline::MarkLegacyCurvesDirty()
 
 void FNewSpline::RebuildLegacyCurves() const
 {
+	TRACE_CPUPROFILER_EVENT_SCOPE(FNewSpline::RebuildLegacyCurves);
+
 	// Curves are dirty if:
 	// 1) Current version is not the same as next version (if next version == version, it has not been dirtied)
 	// 2) UE::Spline::GImmediatelyUpdateLegacyCurves == true (legacy curves are updated any time they might be dirty)
@@ -2741,18 +2771,27 @@ void FNewSpline::RebuildLegacyCurves() const
     ScaleCurve.bIsLooped = IsClosedLoop();
 	ScaleCurve.LoopKeyOffset = 1.f;
 	
+	using FRotationChannel = UE::Geometry::Spline::TTangentBezierSpline<FQuat>;
+	using FScaleChannel = UE::Geometry::Spline::TTangentBezierSpline<FVector>;
+
+	const FRotationChannel* RotationChannel = GetTypedAttributeChannel<FRotationChannel>(RotationAttrName);
+	const FScaleChannel* ScaleChannel = GetTypedAttributeChannel<FScaleChannel>(ScaleAttrName);
+
     for (int32 i = 0; i < GetSpline().GetNumPoints(); ++i)
     {
     	const float AttributeParameter = static_cast<float>(i);
     	
     	// Additionally populate interp curves here so we don't double loop
     	PositionCurve.Points.Emplace(AttributeParameter, GetSpline().GetValue(i), GetSpline().GetTangentIn(i), GetSpline().GetTangentOut(i), ConvertTangentModeToInterpCurveMode(GetSpline().GetTangentModes()[i]));
-    	RotationCurve.AddPoint(AttributeParameter, GetAttributeValue<FQuat>(RotationAttrName, i));
+		RotationCurve.Points.Emplace(AttributeParameter, RotationChannel ? RotationChannel->GetValue(i) : FQuat::Identity);
     	RotationCurve.Points.Last().InterpMode = CIM_CurveAuto;
-    	ScaleCurve.AddPoint(AttributeParameter, GetAttributeValue<FVector>(ScaleAttrName, i));
+		ScaleCurve.Points.Emplace(AttributeParameter, ScaleChannel ? ScaleChannel->GetValue(i) : FVector::OneVector);
     	ScaleCurve.Points.Last().InterpMode = CIM_CurveAuto;
     }
 
+	Algo::SortBy(RotationCurve.Points, &FInterpCurvePoint<FQuat>::InVal);
+	Algo::SortBy(ScaleCurve.Points, &FInterpCurvePoint<FVector>::InVal);
+
 	PositionCurve.AutoSetTangents(0.f, false);
 	RotationCurve.AutoSetTangents(0.f, false);
 	ScaleCurve.AutoSetTangents(0.f, false);
diff --git a/Engine/Source/Runtime/GeometryCore/Public/Splines/BSpline.h b/Engine/Source/Runtime/GeometryCore/Public/Splines/BSpline.h
--- a/Engine/Source/Runtime/GeometryCore/Public/Splines/BSpline.h
+++ b/Engine/Source/Runtime/GeometryCore/Public/Splines/BSpline.h
@@ -671,14 +671,29 @@ protected:
 		bool bSearchLeft = true;
 	};
 
+	struct FKnotSearchCache
+	{
+		explicit FKnotSearchCache(const TArray<FKnot>& PairKnots)
+		{
+			// Build a set of existing knot values (normalize zeros to avoid -0/+0 duplicates)
+			KnotValuesSet.Reserve(PairKnots.Num());
+			for (const FKnot& Knot : PairKnots)
+			{
+				KnotValuesSet.Add(Param::NormalizeKey(Knot.Value));
+			}
+		}
+
+		TSet<float> KnotValuesSet;
+	};
+
 	/**
 	 * Finds the nearest available knot value that does not conflict with existing knots.
 	 * @param InSearchParams - Describes value to search for and how to search.
 	 * @return The nearest valid knot parameter that does not conflict with existing knots.
 	 */
-	float GetNearestAvailableKnotValue(const FValidKnotSearchParams& InSearchParams) const
+	float GetNearestAvailableKnotValue(const FValidKnotSearchParams& InSearchParams, const FKnotSearchCache& SearchCache) const
 	{
-		const float& DesiredParameter = InSearchParams.DesiredParameter;
+		const float DesiredParameter = InSearchParams.DesiredParameter;
 
 		if (!ensureAlwaysMsgf(!std::isnan(DesiredParameter), TEXT("DesiredParameter is NaN!")) ||
 			!ensureAlwaysMsgf(InSearchParams.bSearchLeft || InSearchParams.bSearchRight, TEXT("Invalid search params!")) ||
@@ -687,16 +702,8 @@ protected:
 			return DesiredParameter;
 		}
 
-		// Build a set of existing knot values (normalize zeros to avoid -0/+0 duplicates)
-		TSet<float> KnotValuesSet;
-		KnotValuesSet.Reserve(PairKnots.Num());
-		for (const FKnot& Knot : PairKnots)
-		{
-			KnotValuesSet.Add(Param::NormalizeKey(Knot.Value));
-		}
-
 		// Return the caller's exact value; NormalizeKey is only for set membership (folds -0.0f/+0.0f).
-		if (!KnotValuesSet.Contains(Param::NormalizeKey(DesiredParameter)))
+		if (!SearchCache.KnotValuesSet.Contains(Param::NormalizeKey(DesiredParameter)))
 		{
 			return DesiredParameter;
 		}
@@ -726,7 +733,7 @@ protected:
 				{
 					Left = NewLeft;
 					AnyProgress = true;
-					if (!KnotValuesSet.Contains(Param::NormalizeKey(Left)))
+					if (!SearchCache.KnotValuesSet.Contains(Param::NormalizeKey(Left)))
 					{
 						return Left;
 					}
@@ -748,7 +755,7 @@ protected:
 				{
 					Right = NewRight;
 					AnyProgress = true;
-					if (!KnotValuesSet.Contains(Param::NormalizeKey(Right)))
+					if (!SearchCache.KnotValuesSet.Contains(Param::NormalizeKey(Right)))
 					{
 						return Right;
 					}
@@ -786,7 +793,7 @@ protected:
 				: Param::EDir::Left);
 
 			float Probe = Param::NormalizeKey(Fallback);
-			if (KnotValuesSet.Contains(Probe))
+			if (SearchCache.KnotValuesSet.Contains(Probe))
 			{
 				// Try the opposite direction once
 				Fallback = Param::Step(
@@ -796,7 +803,7 @@ protected:
 				Probe = Param::NormalizeKey(Fallback);
 			}
 
-			ensureMsgf(!KnotValuesSet.Contains(Probe), TEXT("Fallback still conflicts; check knot packing logic."));
+			ensureMsgf(!SearchCache.KnotValuesSet.Contains(Probe), TEXT("Fallback still conflicts; check knot packing logic."));
 			return Fallback;
 		}
 	}
diff --git a/Engine/Source/Runtime/GeometryCore/Public/Splines/PolyBezierSpline.h b/Engine/Source/Runtime/GeometryCore/Public/Splines/PolyBezierSpline.h
--- a/Engine/Source/Runtime/GeometryCore/Public/Splines/PolyBezierSpline.h
+++ b/Engine/Source/Runtime/GeometryCore/Public/Splines/PolyBezierSpline.h
@@ -827,8 +827,10 @@ public:
 			return NewPointIndex;
 		}
 
+		typename Base::FKnotSearchCache SearchCache{ this->PairKnots };
+
 		// Make sure we have a valid parameter
-		Parameter = Base::GetNearestAvailableKnotValue(Parameter);
+		Parameter = Base::GetNearestAvailableKnotValue(Parameter, SearchCache);
 		
         FInterval1f ParameterRange = GetParameterSpace();
 		int32 SegmentIndex;
@@ -1198,7 +1200,8 @@ public:
 		
 		// Prevent reuse of an existing knot.
 		const float DesiredNewParameter = NewParameter;
-		const float ActualNewParameter = this->GetNearestAvailableKnotValue(DesiredNewParameter);
+		typename Base::FKnotSearchCache SearchCache{ this->PairKnots };
+		const float ActualNewParameter = this->GetNearestAvailableKnotValue(DesiredNewParameter, SearchCache);
 
 		const float DesiredToOldDistance = FMath::Abs(DesiredNewParameter - OldParameter);
 		const float DesiredToActualNewDistance = FMath::Abs(DesiredNewParameter - ActualNewParameter);
@@ -1572,7 +1575,9 @@ private:
 	 */
 	void GenerateDistanceKnotsForBezier(EParameterizationPolicy Mode)
 	{
-		TArray<ValueType> Points = this->Values;
+		TRACE_CPUPROFILER_EVENT_SCOPE(GenerateDistanceKnotsForBezier);
+
+		const TArray<ValueType>& Points = this->Values;
 
 		// For a Bezier spline, the control points don't all lie on the curve
 		// So we need to compute chord lengths between segment endpoints only
@@ -1617,9 +1622,18 @@ private:
 		typename Base::FValidKnotSearchParams SearchParams;
 		SearchParams.bSearchLeft = false;	// we are appending repeatedly, we want insertion order to be predictable (always go after conflicting knots).
 
+		typename Base::FKnotSearchCache SearchCache{ this->PairKnots };
+
+		const auto InsertKnotAndUpdateCache = [&](uint32 InMultiplicity)
+		{
+			const float KnotValue = this->GetNearestAvailableKnotValue(SearchParams, SearchCache);
+			SearchCache.KnotValuesSet.Add(Param::NormalizeKey(KnotValue));
+			this->InsertKnot(FKnot(KnotValue, InMultiplicity));
+		};
+
 		// Add start knot with multiplicity 4
 		SearchParams.DesiredParameter = 0.f;
-		this->InsertKnot(FKnot(this->GetNearestAvailableKnotValue(SearchParams), 4));
+		InsertKnotAndUpdateCache(4);
         
 		// Calculate accumulated distances for internal knots
 		float AccumLength = 0.0f;
@@ -1633,7 +1647,7 @@ private:
 				SearchParams.DesiredParameter = AccumLength;
 				
 				// Interior knots at segment boundaries with multiplicity 3
-				this->InsertKnot(FKnot(this->GetNearestAvailableKnotValue(SearchParams), 3));
+				InsertKnotAndUpdateCache(3);
 			}
 		}
 
@@ -1641,7 +1655,7 @@ private:
 			? SegmentLengths.Last()
 			: 0.0f);
 
-		this->InsertKnot(FKnot(this->GetNearestAvailableKnotValue(SearchParams), 4));
+		InsertKnotAndUpdateCache(4);
 
 		// Mark flat knots cache as dirty
 		this->MarkFlatKnotsCacheDirty();
