Fix issue where AsyncSaveGameToSlotForLocalPlayer would ensure if called twice in a row.
It now sets the LastSuccessfulSaveRequest to the request that actually completed instead of the last requested.
#Synced-CL 48252322

diff --git a/Engine/Source/Runtime/Engine/Private/GameFramework/SaveGame.cpp b/Engine/Source/Runtime/Engine/Private/GameFramework/SaveGame.cpp
--- a/Engine/Source/Runtime/Engine/Private/GameFramework/SaveGame.cpp
+++ b/Engine/Source/Runtime/Engine/Private/GameFramework/SaveGame.cpp
@@ -194,17 +194,17 @@ bool ULocalPlayerSaveGame::AsyncSaveGameToSlotForLocalPlayer()
 
 void ULocalPlayerSaveGame::ProcessSaveComplete(const FString& SlotName, const int32 UserIndex, bool bSuccess, int32 SaveRequest)
 {
-	// Now that a save completed, update the success/failure
+	// Now that a save completed, update the success/failure and set based on the request that just completed
 	// Every time CurrentSaveRequest is incremented it will result in setting either success or error, but there could be multiple in flight that are processed in order
 	if (bSuccess)
 	{
 		ensure(CurrentSaveRequest > LastSuccessfulSaveRequest);
-		LastSuccessfulSaveRequest = CurrentSaveRequest;
+		LastSuccessfulSaveRequest = SaveRequest;
 	}
 	else
 	{
 		ensure(CurrentSaveRequest > LastErrorSaveRequest);
-		LastErrorSaveRequest = CurrentSaveRequest;
+		LastErrorSaveRequest = SaveRequest;
 	}
 
 	HandlePostSave(bSuccess);
