PR #14033: Update comment regarding VSync capability

#rnx
#Synced-CL 48222653

diff --git a/Engine/Source/Runtime/Engine/Private/GameUserSettings.cpp b/Engine/Source/Runtime/Engine/Private/GameUserSettings.cpp
--- a/Engine/Source/Runtime/Engine/Private/GameUserSettings.cpp
+++ b/Engine/Source/Runtime/Engine/Private/GameUserSettings.cpp
@@ -453,7 +453,7 @@ void UGameUserSettings::ApplyNonResolutionSettings()
 		int32 VSyncValue = 0;
 		if (GConfig->GetInt(*ConfigSection, TEXT("r.Vsync"), VSyncValue, GEngineIni))
 		{
-			// VSync was already set by system settings. We are capable of setting it here.
+			// VSync was already set by system settings. We are not capable of setting it here.
 		}
 		else
 		{
