[Engine] Animation assets validate their UAnimMetaData.

#tests Editor asset validation
#Synced-CL 48286836

diff --git a/Engine/Source/Runtime/Engine/Private/AnimationAsset.cpp b/Engine/Source/Runtime/Engine/Private/AnimationAsset.cpp
--- a/Engine/Source/Runtime/Engine/Private/AnimationAsset.cpp
+++ b/Engine/Source/Runtime/Engine/Private/AnimationAsset.cpp
@@ -784,6 +784,13 @@ EDataValidationResult UAnimationAsset::IsDataValid(FDataValidationContext& Conte
 			Result = EDataValidationResult::Invalid;
 		}
 	}
+	for (const UAnimMetaData* Datum : MetaData)
+	{
+		if (Datum != nullptr)
+		{
+			Result = CombineDataValidationResults(Result, Datum->IsDataValid(Context));
+		}
+	}
 	return Result;
 }
 #endif
