Decrease max spotlight culling angle in order fix float precision issues in RayInfiniteConeVolumeOverlap causing ray tracing lighting grid to be broken for 90 degree spotlights.

#Synced-CL 48295796

diff --git a/Engine/Plugins/Experimental/GPULightmass/Source/GPULightmass/Private/Scene/Lights.cpp b/Engine/Plugins/Experimental/GPULightmass/Source/GPULightmass/Private/Scene/Lights.cpp
--- a/Engine/Plugins/Experimental/GPULightmass/Source/GPULightmass/Private/Scene/Lights.cpp
+++ b/Engine/Plugins/Experimental/GPULightmass/Source/GPULightmass/Private/Scene/Lights.cpp
@@ -161,19 +161,18 @@ FSpotLightBuildInfo::FSpotLightBuildInfo(USpotLightComponent* SpotLightComponent
 	Position = SpotLightComponent->GetLightPosition();
 	Direction = SpotLightComponent->GetDirection();
 	AttenuationRadius = SpotLightComponent->AttenuationRadius;
-	InnerConeAngle = SpotLightComponent->InnerConeAngle;
-	OuterConeAngle = SpotLightComponent->OuterConeAngle;
+
+	const FVector2f ClampedConeAngles = SpotLightComponent->GetClampedConeAngles();
+	InnerConeAngle = ClampedConeAngles.X;
+	OuterConeAngle = ClampedConeAngles.Y;
 }
 
 bool FSpotLightBuildInfo::AffectsBounds(const FBoxSphereBounds& InBounds) const
 {
 	if ((InBounds.Origin - Position).SizeSquared() <= FMath::Square(AttenuationRadius + InBounds.SphereRadius))
 	{
-		float ClampedInnerConeAngle = FMath::Clamp(InnerConeAngle, 0.0f, 89.0f) * (float)PI / 180.0f;
-		float ClampedOuterConeAngle = FMath::Clamp(OuterConeAngle * (float)PI / 180.0f, ClampedInnerConeAngle + 0.001f, 89.0f * (float)PI / 180.0f + 0.001f);
-
-		float Sin = FMath::Sin(ClampedOuterConeAngle);
-		float Cos = FMath::Cos(ClampedOuterConeAngle);
+		float Sin = FMath::Sin(OuterConeAngle);
+		float Cos = FMath::Cos(OuterConeAngle);
 
 		FVector	U = Position - (InBounds.SphereRadius / Sin) * Direction;
 		FVector	D = InBounds.Origin - U;
@@ -201,10 +200,9 @@ FSpotLightRenderState::FSpotLightRenderState(USpotLightComponent* SpotLightCompo
 	Position = SpotLightComponent->GetLightPosition();
 	Direction = SpotLightComponent->GetDirection();
 	{
-		const float ClampedInnerConeAngle = FMath::Clamp(SpotLightComponent->InnerConeAngle, 0.0f, 89.0f) * (float)PI / 180.0f;
-		const float ClampedOuterConeAngle = FMath::Clamp(SpotLightComponent->OuterConeAngle * (float)PI / 180.0f, ClampedInnerConeAngle + 0.001f, 89.0f * (float)PI / 180.0f + 0.001f);
-		const float CosOuterCone = FMath::Cos(ClampedOuterConeAngle);
-		const float CosInnerCone = FMath::Cos(ClampedInnerConeAngle);
+		const FVector2f ClampedConeAngles = SpotLightComponent->GetClampedConeAngles();
+		const float CosOuterCone = FMath::Cos(ClampedConeAngles.Y);
+		const float CosInnerCone = FMath::Cos(ClampedConeAngles.X);
 		const float InvCosConeDifference = 1.0f / (CosInnerCone - CosOuterCone);
 		SpotAngles = FVector2D(CosOuterCone, InvCosConeDifference);
 	}
diff --git a/Engine/Source/Runtime/Engine/Classes/Components/SpotLightComponent.h b/Engine/Source/Runtime/Engine/Classes/Components/SpotLightComponent.h
--- a/Engine/Source/Runtime/Engine/Classes/Components/SpotLightComponent.h
+++ b/Engine/Source/Runtime/Engine/Classes/Components/SpotLightComponent.h
@@ -44,6 +44,7 @@ class USpotLightComponent : public UPointLightComponent
 	ENGINE_API virtual FBox GetStreamingBounds() const override;
 #endif
 
+	ENGINE_API FVector2f GetClampedConeAngles() const;
 	ENGINE_API float GetHalfConeAngle() const;
 	ENGINE_API float GetCosHalfConeAngle() const;
 
diff --git a/Engine/Source/Runtime/Engine/Private/Components/SpotLightComponent.cpp b/Engine/Source/Runtime/Engine/Private/Components/SpotLightComponent.cpp
--- a/Engine/Source/Runtime/Engine/Private/Components/SpotLightComponent.cpp
+++ b/Engine/Source/Runtime/Engine/Private/Components/SpotLightComponent.cpp
@@ -15,7 +15,6 @@
 
 #include UE_INLINE_GENERATED_CPP_BY_NAME(SpotLightComponent)
 
-
 /**
  * The scene info for a spot light.
  */
@@ -45,14 +44,14 @@ public:
 	FSpotLightSceneProxy(const USpotLightComponent* Component)
 	:	FPointLightSceneProxy(Component)
 	{
-		const float ClampedInnerConeAngle = FMath::Clamp(Component->InnerConeAngle,0.0f,89.0f) * (float)UE_PI / 180.0f;
-		const float ClampedOuterConeAngle = FMath::Clamp(Component->OuterConeAngle * (float)UE_PI / 180.0f,ClampedInnerConeAngle + 0.001f,89.0f * (float)UE_PI / 180.0f + 0.001f);
-		OuterConeAngle = ClampedOuterConeAngle;
-		CosOuterCone = FMath::Cos(ClampedOuterConeAngle);
-		SinOuterCone = FMath::Sin(ClampedOuterConeAngle);
-		CosInnerCone = FMath::Cos(ClampedInnerConeAngle);
+		const FVector2f ClampedConeAngles = Component->GetClampedConeAngles();
+
+		OuterConeAngle = ClampedConeAngles.Y;
+		CosOuterCone = FMath::Cos(ClampedConeAngles.Y);
+		SinOuterCone = FMath::Sin(ClampedConeAngles.Y);
+		CosInnerCone = FMath::Cos(ClampedConeAngles.X);
 		InvCosConeDifference = 1.0f / (CosInnerCone - CosOuterCone);
-		InvTanOuterCone = 1.0f / FMath::Tan(ClampedOuterConeAngle);
+		InvTanOuterCone = 1.0f / FMath::Tan(ClampedConeAngles.Y);
 	}
 
 	/** Accesses parameters needed for rendering the light. */
@@ -202,11 +201,18 @@ USpotLightComponent::USpotLightComponent(const FObjectInitializer& ObjectInitial
 	OuterConeAngle = 44.0f;
 }
 
+FVector2f USpotLightComponent::GetClampedConeAngles() const
+{
+	FVector2f ConeAngles;
+	ConeAngles.X = FMath::Clamp(InnerConeAngle, 0.0f, 88.9f) * (float)UE_PI / 180.0f;
+	ConeAngles.Y = FMath::Clamp(OuterConeAngle * (float)UE_PI / 180.0f, ConeAngles.X + 0.001f, 88.9f * (float)UE_PI / 180.0f + 0.001f);
+	return ConeAngles;
+}
+
 float USpotLightComponent::GetHalfConeAngle() const
 {
-	const float ClampedInnerConeAngle = FMath::Clamp(InnerConeAngle, 0.0f, 89.0f) * (float)UE_PI / 180.0f;
-	const float ClampedOuterConeAngle = FMath::Clamp(OuterConeAngle * (float)UE_PI / 180.0f, ClampedInnerConeAngle + 0.001f, 89.0f * (float)UE_PI / 180.0f + 0.001f);
-	return ClampedOuterConeAngle;
+	const FVector2f ClampedConeAngles = GetClampedConeAngles();
+	return ClampedConeAngles.Y;
 }
 
 float USpotLightComponent::GetCosHalfConeAngle() const
@@ -352,11 +358,10 @@ bool USpotLightComponent::AffectsBounds(const FBoxSphereBounds& InBounds) const
 		return false;
 	}
 
-	float	ClampedInnerConeAngle = FMath::Clamp(InnerConeAngle,0.0f,89.0f) * (float)UE_PI / 180.0f,
-			ClampedOuterConeAngle = FMath::Clamp(OuterConeAngle * (float)UE_PI / 180.0f,ClampedInnerConeAngle + 0.001f,89.0f * (float)UE_PI / 180.0f + 0.001f);
+	const FVector2f ClampedConeAngles = GetClampedConeAngles();
 
-	float	Sin = FMath::Sin(ClampedOuterConeAngle),
-			Cos = FMath::Cos(ClampedOuterConeAngle);
+	float	Sin = FMath::Sin(ClampedConeAngles.Y),
+			Cos = FMath::Cos(ClampedConeAngles.Y);
 
 	FVector	U = GetComponentLocation() - (InBounds.SphereRadius / Sin) * GetDirection(),
 			D = InBounds.Origin - U;
