Fixed blendspace sample mirroring with upstream mirror nodes.

#Synced-CL 48298300

diff --git a/Engine/Source/Runtime/Engine/Classes/Animation/BlendSpace.h b/Engine/Source/Runtime/Engine/Classes/Animation/BlendSpace.h
--- a/Engine/Source/Runtime/Engine/Classes/Animation/BlendSpace.h
+++ b/Engine/Source/Runtime/Engine/Classes/Animation/BlendSpace.h
@@ -709,12 +709,16 @@ protected:
 
 	/** Initialize Per Bone Blend **/
 	void InitializePerBoneBlend();
-
-	/** Ticks the samples in SampleDataList apart from the HighestWeightIndex one. */
+	
+	UE_DEPRECATED(5.8, "Please use TickFollowerSamples() that takes in the blendspace's tick record.")
 	void TickFollowerSamples(
 		TArray<FBlendSampleData> &SampleDataList, const int32 HighestWeightIndex, FAnimAssetTickContext &Context, 
 		bool bResetMarkerDataOnFollowers, bool bLooping, const UMirrorDataTable* MirrorDataTable = nullptr) const;
 
+	/** Ticks the samples in SampleDataList apart from the HighestWeightIndex one. */
+	void TickFollowerSamples(TArray<FBlendSampleData> &SampleDataList, const int32 HighestWeightIndex, FAnimAssetTickContext &Context, 
+		bool bResetMarkerDataOnFollowers, bool bLooping, const FAnimTickRecord& BlendSpaceInstance) const;
+	
 	/** Returns the blend input clamped to the valid range, unless that axis has been set to wrap in which case no clamping is done **/
 	FVector GetClampedBlendInput(const FVector& BlendInput) const;
 	
@@ -751,6 +755,9 @@ private:
 	// Internal helper function for UpdateBlendSamples and TickAssetPlayer
 	bool UpdateBlendSamples_Internal(const FVector& InBlendSpacePosition, float InDeltaTime, TArray<FBlendSampleData>& InOutOldSampleDataList, TArray<FBlendSampleData>& InOutSampleDataCache, int32& InOutCachedTriangulationIndex) const;	
 
+	/** Gets the appropriate mirror data table for tick record's blend sample. */
+	const UMirrorDataTable* GetEffectiveMirrorDataTableForSample(const FAnimTickRecord& InInstance, const FBlendSample& InBlendSample) const;
+	
 public:
 #if WITH_EDITORONLY_DATA
 	UE_DEPRECATED(5.1, "This property is deprecated. Please use/see bContainsRotationOffsetMeshSpaceSamples instead")
diff --git a/Engine/Source/Runtime/Engine/Private/Animation/BlendSpace.cpp b/Engine/Source/Runtime/Engine/Private/Animation/BlendSpace.cpp
--- a/Engine/Source/Runtime/Engine/Private/Animation/BlendSpace.cpp
+++ b/Engine/Source/Runtime/Engine/Private/Animation/BlendSpace.cpp
@@ -293,6 +293,27 @@ bool UBlendSpace::UpdateBlendSamples_Internal(
 	return bSuccessfullySampled;
 }
 
+const UMirrorDataTable* UBlendSpace::GetEffectiveMirrorDataTableForSample(const FAnimTickRecord& InInstance, const FBlendSample& InBlendSample) const
+{
+	const bool bInstanceIsMirrored = InInstance.MirrorDataTable != nullptr;
+	const bool bSampleIsMirrored = InBlendSample.bMirror;
+
+	// We have an upstream mirror request (via graph message), and our sample is unmirrored, use graph message's mirror data table as usual.
+	if (bInstanceIsMirrored && !bSampleIsMirrored)
+	{
+		return InInstance.MirrorDataTable;
+	}
+
+	// No upstream mirroring request (via graph message) therefore we are good to use the blendspace's mirror data table since our sample is mirrored.
+	if (!bInstanceIsMirrored && bSampleIsMirrored)
+	{
+		return MirrorDataTable;
+	}
+
+	// We have an upstream mirror request (via graph message) AND we have our sample mirrored. We don't mirror since we need to do the inverse of what the rest of the samples mirror state is (driven by graph message). 
+	return nullptr;
+}
+
 bool UBlendSpace::UpdateBlendSamples(const FVector& InBlendSpacePosition, float InDeltaTime, TArray<FBlendSampleData>& InOutSampleDataCache, int32& InOutCachedTriangulationIndex) const
 {
 	TArray<FBlendSampleData>& OldSampleDataList = FBlendSpaceScratchData::Get().OldSampleDataList;
@@ -616,7 +637,10 @@ void UBlendSpace::TickAssetPlayer(FAnimTickRecord& Instance, struct FAnimNotifyQ
 						// Update leader delta.
 						float NewDelta = Instance.DeltaTimeRecord->Delta * LeaderSample.RateScale * LeaderSample.Animation->RateScale;
 						Context.SetLeaderDelta(NewDelta);
-						
+
+						// Get mirror table.
+						const UMirrorDataTable* LeaderMirrorDataTable = GetEffectiveMirrorDataTableForSample(Instance, LeaderSample);
+
 						// Advance time as usual.
 						{
 							float CurrentTime = NormalizedCurrentTime * NewAnimLength;
@@ -630,7 +654,7 @@ void UBlendSpace::TickAssetPlayer(FAnimTickRecord& Instance, struct FAnimNotifyQ
 
 						// Sync position BEFORE advancing time.
 						LeaderSample.Animation->GetMarkerIndicesForTime(LeaderSampleData.PreviousTime, Instance.bLooping, Context.MarkerTickContext.GetValidMarkerNames(), LeaderSampleData.MarkerTickRecord.PreviousMarker, LeaderSampleData.MarkerTickRecord.NextMarker);
-						const FMarkerSyncAnimPosition StartSyncPosition = LeaderSample.Animation->GetMarkerSyncPositionFromMarkerIndicies(LeaderSampleData.MarkerTickRecord.PreviousMarker.MarkerIndex, LeaderSampleData.MarkerTickRecord.NextMarker.MarkerIndex, LeaderSampleData.PreviousTime, Instance.MirrorDataTable);
+						const FMarkerSyncAnimPosition StartSyncPosition = LeaderSample.Animation->GetMarkerSyncPositionFromMarkerIndicies(LeaderSampleData.MarkerTickRecord.PreviousMarker.MarkerIndex, LeaderSampleData.MarkerTickRecord.NextMarker.MarkerIndex, LeaderSampleData.PreviousTime, LeaderMirrorDataTable);
 						Context.MarkerTickContext.SetMarkerSyncStartPosition(StartSyncPosition);
 
 						// @todo: See if we can have reset blend samples directly update the context and/or tick record (no need to compute info twice, if possible). 
@@ -640,7 +664,7 @@ void UBlendSpace::TickAssetPlayer(FAnimTickRecord& Instance, struct FAnimNotifyQ
 						// Sync position AFTER advancing time.
 						LeaderSample.Animation->GetMarkerIndicesForTime(LeaderSampleData.Time, Instance.bLooping, Context.MarkerTickContext.GetValidMarkerNames(), LeaderSampleData.MarkerTickRecord.PreviousMarker, LeaderSampleData.MarkerTickRecord.NextMarker);
 						*Instance.MarkerTickRecord = LeaderSampleData.MarkerTickRecord;
-						const FMarkerSyncAnimPosition EndSyncPosition = LeaderSample.Animation->GetMarkerSyncPositionFromMarkerIndicies(LeaderSampleData.MarkerTickRecord.PreviousMarker.MarkerIndex, LeaderSampleData.MarkerTickRecord.NextMarker.MarkerIndex, LeaderSampleData.Time, Instance.MirrorDataTable);
+						const FMarkerSyncAnimPosition EndSyncPosition = LeaderSample.Animation->GetMarkerSyncPositionFromMarkerIndicies(LeaderSampleData.MarkerTickRecord.PreviousMarker.MarkerIndex, LeaderSampleData.MarkerTickRecord.NextMarker.MarkerIndex, LeaderSampleData.Time, LeaderMirrorDataTable);
 						Context.MarkerTickContext.SetMarkerSyncEndPosition(EndSyncPosition);
 
 						// Update context's position ratios
@@ -661,7 +685,8 @@ void UBlendSpace::TickAssetPlayer(FAnimTickRecord& Instance, struct FAnimNotifyQ
 					if (LeaderSample.Animation)
 					{
 						bool bResetMarkerDataOnFollowers = false;
-
+						const UMirrorDataTable* LeaderMirrorDataTable = GetEffectiveMirrorDataTableForSample(Instance, LeaderSample);
+						
 						// Invalidate sample followers' tick records if instance doesn't have any valid sync marker data. 
 						if (!Instance.MarkerTickRecord->IsValid(Instance.bLooping))
 						{
@@ -673,7 +698,7 @@ void UBlendSpace::TickAssetPlayer(FAnimTickRecord& Instance, struct FAnimNotifyQ
 						else if (!LeaderSampleData.MarkerTickRecord.IsValid(Instance.bLooping) && Context.MarkerTickContext.GetMarkerSyncStartPosition().IsValid())
 						{
 							// TODO: Look into the reason for not passing bLooping variable and just forcing the vale to be true. 
-							LeaderSample.Animation->GetMarkerIndicesForPosition(Context.MarkerTickContext.GetMarkerSyncStartPosition(), true, LeaderSampleData.MarkerTickRecord.PreviousMarker, LeaderSampleData.MarkerTickRecord.NextMarker, LeaderSampleData.Time, Instance.MirrorDataTable);
+							LeaderSample.Animation->GetMarkerIndicesForPosition(Context.MarkerTickContext.GetMarkerSyncStartPosition(), true, LeaderSampleData.MarkerTickRecord.PreviousMarker, LeaderSampleData.MarkerTickRecord.NextMarker, LeaderSampleData.Time, LeaderMirrorDataTable);
 						}
 
 						// Only tick samples if leader sample has any delta time to consume.
@@ -683,12 +708,12 @@ void UBlendSpace::TickAssetPlayer(FAnimTickRecord& Instance, struct FAnimNotifyQ
 						if (!FMath::IsNearlyZero(NewDeltaTime))
 						{
 							// Tick leader sample
-							LeaderSample.Animation->TickByMarkerAsLeader(LeaderSampleData.MarkerTickRecord, Context.MarkerTickContext, LeaderSampleData.Time, LeaderSampleData.PreviousTime, NewDeltaTime, Instance.bLooping, Instance.MirrorDataTable);
+							LeaderSample.Animation->TickByMarkerAsLeader(LeaderSampleData.MarkerTickRecord, Context.MarkerTickContext, LeaderSampleData.Time, LeaderSampleData.PreviousTime, NewDeltaTime, Instance.bLooping, LeaderMirrorDataTable);
 
 							check(!Instance.bLooping || Context.MarkerTickContext.IsMarkerSyncStartValid());
 							
 							// Tick all the follower samples
-							TickFollowerSamples(SampleDataList, HighestMarkerSyncWeightIndex, Context, bResetMarkerDataOnFollowers, Instance.bLooping, Instance.MirrorDataTable);
+							TickFollowerSamples(SampleDataList, HighestMarkerSyncWeightIndex, Context, bResetMarkerDataOnFollowers, Instance.bLooping, Instance);
 						}
 						else if (!Instance.MarkerTickRecord->IsValid(Instance.bLooping))
 						{
@@ -696,12 +721,12 @@ void UBlendSpace::TickAssetPlayer(FAnimTickRecord& Instance, struct FAnimNotifyQ
 							LeaderSample.Animation->GetMarkerIndicesForTime(LeaderSampleData.Time, Instance.bLooping, Context.MarkerTickContext.GetValidMarkerNames(), LeaderSampleData.MarkerTickRecord.PreviousMarker, LeaderSampleData.MarkerTickRecord.NextMarker);
 
 							// Get sync position for followers to sync up to.
-							const FMarkerSyncAnimPosition SyncPosition = LeaderSample.Animation->GetMarkerSyncPositionFromMarkerIndicies(LeaderSampleData.MarkerTickRecord.PreviousMarker.MarkerIndex, LeaderSampleData.MarkerTickRecord.NextMarker.MarkerIndex, LeaderSampleData.Time, Instance.MirrorDataTable);
+							const FMarkerSyncAnimPosition SyncPosition = LeaderSample.Animation->GetMarkerSyncPositionFromMarkerIndicies(LeaderSampleData.MarkerTickRecord.PreviousMarker.MarkerIndex, LeaderSampleData.MarkerTickRecord.NextMarker.MarkerIndex, LeaderSampleData.Time, LeaderMirrorDataTable);
 							Context.MarkerTickContext.SetMarkerSyncStartPosition(SyncPosition);
 							Context.MarkerTickContext.SetMarkerSyncEndPosition(SyncPosition);
 							
 							// Make all follower samples match next sync position to equal that of the leader.
-							TickFollowerSamples(SampleDataList, HighestMarkerSyncWeightIndex, Context, true, Instance.bLooping, Instance.MirrorDataTable);
+							TickFollowerSamples(SampleDataList, HighestMarkerSyncWeightIndex, Context, true, Instance.bLooping, Instance);
 						}
 						
 						NormalizedCurrentTime = LeaderSampleData.Time / LeaderSample.GetSamplePlayLength();
@@ -749,7 +774,7 @@ void UBlendSpace::TickAssetPlayer(FAnimTickRecord& Instance, struct FAnimNotifyQ
 						}
 
 						// Tick all samples as followers
-						TickFollowerSamples(SampleDataList, INDEX_NONE, Context, false, Instance.bLooping, Instance.MirrorDataTable);
+						TickFollowerSamples(SampleDataList, INDEX_NONE, Context, false, Instance.bLooping, Instance);
 					}
 						
 					*Instance.MarkerTickRecord = SampleDataItem.MarkerTickRecord;
@@ -1979,6 +2004,14 @@ void UBlendSpace::InitializePerBoneBlend()
 void UBlendSpace::TickFollowerSamples(
 	TArray<FBlendSampleData>& SampleDataList, const int32 HighestWeightIndex, FAnimAssetTickContext& Context, 
 	bool bResetMarkerDataOnFollowers, bool bLooping, const UMirrorDataTable* InMirrorDataTable) const
+{
+	FAnimTickRecord TempTickRecord;
+	TempTickRecord.MirrorDataTable = InMirrorDataTable;
+	TickFollowerSamples(SampleDataList, HighestWeightIndex, Context, bResetMarkerDataOnFollowers, bLooping, TempTickRecord);
+}
+
+void UBlendSpace::TickFollowerSamples(TArray<FBlendSampleData>& SampleDataList, const int32 HighestWeightIndex, FAnimAssetTickContext& Context,
+	bool bResetMarkerDataOnFollowers, bool bLooping, const FAnimTickRecord& BlendSpaceInstance) const
 {
 	for (int32 SampleIndex = 0; SampleIndex < SampleDataList.Num(); ++SampleIndex)
 	{
@@ -1996,7 +2029,7 @@ void UBlendSpace::TickFollowerSamples(
 			{
 				Sample.Animation->TickByMarkerAsFollower(
 					SampleDataItem.MarkerTickRecord, Context.MarkerTickContext, SampleDataItem.Time, 
-					SampleDataItem.PreviousTime, Context.GetLeaderDelta(), bLooping, InMirrorDataTable);
+					SampleDataItem.PreviousTime, Context.GetLeaderDelta(), bLooping, GetEffectiveMirrorDataTableForSample(BlendSpaceInstance, Sample));
 			}
 		}
 	}
