Expose CPUTemp and ThermalStatus CSV metrics to other platforms than Android
-New Cross Platform Thermal stats: DeviceTemperature ThermalState.
-DeviceTemperature is in degrees C, and only implemented for android now.
-ThermalState mirrors the android thermal status enum and ios thermal state values are mapped onto the first four values of it.  See the comments in ThermalStats.h and the platform implementations for more detail.
-The pre-existing TemperatureLevel stat is deprecated.  It's reports temperature on android and thermal state on ios.

#Synced-CL 47939539

diff --git a/Engine/Plugins/MovieScene/MovieRenderPipeline/Source/MovieRenderPipelineCore/Private/MoviePipelineUtils.cpp b/Engine/Plugins/MovieScene/MovieRenderPipeline/Source/MovieRenderPipelineCore/Private/MoviePipelineUtils.cpp
--- a/Engine/Plugins/MovieScene/MovieRenderPipeline/Source/MovieRenderPipelineCore/Private/MoviePipelineUtils.cpp
+++ b/Engine/Plugins/MovieScene/MovieRenderPipeline/Source/MovieRenderPipelineCore/Private/MoviePipelineUtils.cpp
@@ -1147,7 +1147,11 @@ namespace UE
 			FPlatformMisc::GetOSVersions(OSVersionLabel, OSSubVersionLabel);
 			InFileMetadata.Add(TEXT("unreal/system/osVersion"), OSVersionLabel);
 			InFileMetadata.Add(TEXT("unreal/system/osSubVersion"), OSSubVersionLabel);
-			InFileMetadata.Add(TEXT("unreal/system/deviceTemp"), LexToString(FPlatformMisc::GetDeviceTemperatureLevel()));
+			PRAGMA_DISABLE_DEPRECATION_WARNINGS
+			InFileMetadata.Add(TEXT("unreal/system/deviceTemp"), LexToString(FPlatformMisc::GetDeviceTemperatureLevel())); // Deprecated in 5.8 in favor of the two separate stats below.
+			PRAGMA_ENABLE_DEPRECATION_WARNINGS
+			InFileMetadata.Add(TEXT("unreal/system/deviceTemperature"), LexToString(FPlatformMisc::GetDeviceTemperature()));
+			InFileMetadata.Add(TEXT("unreal/system/deviceThermalState"), LexToString(static_cast<int>(FPlatformMisc::GetDeviceThermalState())));
 			InFileMetadata.Add(TEXT("unreal/system/deviceMakeAndModel"), FPlatformMisc::GetDeviceMakeAndModel());
 			
 			if (InOutputDir.Len() > 0)
diff --git a/Engine/Source/Runtime/Core/Private/Android/AndroidPlatformMisc.cpp b/Engine/Source/Runtime/Core/Private/Android/AndroidPlatformMisc.cpp
--- a/Engine/Source/Runtime/Core/Private/Android/AndroidPlatformMisc.cpp
+++ b/Engine/Source/Runtime/Core/Private/Android/AndroidPlatformMisc.cpp
@@ -2760,10 +2760,20 @@ bool FAndroidMisc::IsInLowPowerMode()
 }
 
 float FAndroidMisc::GetDeviceTemperatureLevel()
+{
+	return FAndroidMisc::GetDeviceTemperature();
+}
+
+float FAndroidMisc::GetDeviceTemperature()
 {
 	return GetBatteryState().Temperature;
 }
 
+EDeviceThermalState FAndroidMisc::GetDeviceThermalState()
+{
+	return FAndroidStats::GetThermalStatus();
+}
+
 bool FAndroidMisc::AreHeadPhonesPluggedIn()
 {
 	return HeadPhonesArePluggedIn;
diff --git a/Engine/Source/Runtime/Core/Private/Android/AndroidStats.cpp b/Engine/Source/Runtime/Core/Private/Android/AndroidStats.cpp
--- a/Engine/Source/Runtime/Core/Private/Android/AndroidStats.cpp
+++ b/Engine/Source/Runtime/Core/Private/Android/AndroidStats.cpp
@@ -5,8 +5,10 @@
 #include "CoreMinimal.h"
 #include "Android/AndroidPlatformMisc.h"
 #include "Stats/Stats.h"
+#include "Stats/ThermalStats.h"
 #include "HAL/IConsoleManager.h"
 #include "ProfilingDebugging/CsvProfiler.h"
+#include <android/thermal.h>
 
 DEFINE_LOG_CATEGORY_STATIC(LogAndroidStats, Log, Log);
 
@@ -319,6 +321,15 @@ void FAndroidStats::LogGPUStats()
 	::LogGPUStats();
 }
 
+// We want to mirror the android enum values, so check some of them.
+static_assert((int)EDeviceThermalState::None == (int)ATHERMAL_STATUS_NONE);
+static_assert((int)EDeviceThermalState::Severe == (int)ATHERMAL_STATUS_SEVERE);
+EDeviceThermalState FAndroidStats::GetThermalStatus()
+{
+	return static_cast<EDeviceThermalState>(GThermalStatus);
+}
+
+
 void FAndroidStats::UpdateAndroidStats()
 {
 #if CSV_PROFILER || STATS
@@ -416,19 +427,21 @@ void FAndroidStats::UpdateAndroidStats()
 		};
 
 		check(FrequencyGroups.Num() <= UE::Android::MaxFrequencyGroupStats);
-#if CSV_PROFILER
+#if CSV_PROFILER || STATS
 		{
 			UE::TUniqueLock Lock(CSVStatsLock);
 			UE::Android::FCPUStatValues& CSVStatValues = UE::Android::FCPUStatValues::Get();
 
-			CSVStatValues.CPUTemp = FAndroidMisc::GetCPUTemperature();;
-			CSVStatValues.Thermals5S = FAndroidPlatformThermal::GetThermalStress(FAndroidPlatformThermal::EForecastPeriod::FIVE_SEC);;
+			CSVStatValues.CPUTemp = FAndroidMisc::GetCPUTemperature();
+			CSVStatValues.Thermals5S = FAndroidPlatformThermal::GetThermalStress(FAndroidPlatformThermal::EForecastPeriod::FIVE_SEC);
 			CSVStatValues.ThermalStatus = GThermalStatus;
+#if CSV_PROFILER 
 			CSVStatValues.TrimMemoryBackgroundLevel = GTrimMemoryBackgroundLevel;
 			CSVStatValues.TrimMemoryForegroundLevel = GTrimMemoryForegroundLevel;
 			CSVStatValues.MemStats = FAndroidPlatformMemory::GetStats();
+#endif //#if CSV_PROFILER
 		}
-#endif //CSV_PROFILER
+#endif //#if CSV_PROFILER || STATS
 
 		float MaxSingleCoreUtilization[UE::Android::MaxFrequencyGroupStats] = { 0.0f };
 		{
@@ -464,6 +477,8 @@ void FAndroidStats::UpdateAndroidStats()
 		GIsStatTaskActive = false;
 	}, UE::Tasks::ETaskPriority::BackgroundLow);
 #endif
+
+	UE::ThermalStats::GetAndPushStats();
 }
 
 
diff --git a/Engine/Source/Runtime/Core/Private/IOS/IOSPlatformMisc.cpp b/Engine/Source/Runtime/Core/Private/IOS/IOSPlatformMisc.cpp
--- a/Engine/Source/Runtime/Core/Private/IOS/IOSPlatformMisc.cpp
+++ b/Engine/Source/Runtime/Core/Private/IOS/IOSPlatformMisc.cpp
@@ -463,6 +463,26 @@ float FIOSPlatformMisc::GetDeviceTemperatureLevel()
 	return -1.0f;
 }
 
+EDeviceThermalState FIOSPlatformMisc::GetDeviceThermalState()
+{
+#if !PLATFORM_TVOS
+	// We are mapping the ios thermal state values onto the first four of our thermal state enum values, which are based on 
+	// the android thermal status values.
+	// Nominal and None both refer to states where the device is entirely fine on heat and not throttling.
+	// Fair/Serious Light/Moderate are less exactly comparable.  Android, per docs, does some throttling.  IOS does not.  But
+	// in both cases the device is physically getting warmer.
+	// Critical/Severe indicate that serious throttling is kicking in and the device is probably noticably hot to touch.
+	switch ([[IOSAppDelegate GetDelegate]GetThermalState] )
+	{
+	case NSProcessInfoThermalStateNominal:	return EDeviceThermalState::None; break;
+	case NSProcessInfoThermalStateFair:		return EDeviceThermalState::Light; break;
+	case NSProcessInfoThermalStateSerious:	return EDeviceThermalState::Moderate; break;
+	case NSProcessInfoThermalStateCritical:	return EDeviceThermalState::Severe; break;
+	}
+#endif
+	return EDeviceThermalState::Unsupported;
+}
+
 bool FIOSPlatformMisc::IsInLowPowerMode()
 {
 #if !PLATFORM_TVOS
diff --git a/Engine/Source/Runtime/Core/Private/Stats/ThermalStats.cpp b/Engine/Source/Runtime/Core/Private/Stats/ThermalStats.cpp
new file mode 100644
--- /dev/null
+++ b/Engine/Source/Runtime/Core/Private/Stats/ThermalStats.cpp
@@ -0,0 +1,36 @@
+// Copyright Epic Games, Inc. All Rights Reserved.
+
+#include "Stats/ThermalStats.h"
+#include "HAL/PlatformMisc.h"
+
+DEFINE_STAT(STAT_DeviceTemperature);
+DEFINE_STAT(STAT_ThermalState);
+
+CSV_DEFINE_CATEGORY(Thermal, true);
+CSV_DEFINE_STAT(Thermal, DeviceTemperature);
+CSV_DEFINE_STAT(Thermal, ThermalState);
+
+namespace UE
+{
+	namespace ThermalStats
+	{
+		void GetAndPushStats()
+		{
+			const float DeviceTemperature = FPlatformMisc::GetDeviceTemperature();
+			const EDeviceThermalState ThermalState = FPlatformMisc::GetDeviceThermalState();
+
+			//UE_LOG(LogStats, Warning, TEXT("TemperatureLevel(legacy)=%3.1f DeviceTemperature=%3.1f ThermalState=%i"), FPlatformMisc::GetDeviceTemperatureLevel(), DeviceTemperature, ThermalState);
+
+#if CSV_PROFILER
+			CSV_CUSTOM_STAT_DEFINED(DeviceTemperature, DeviceTemperature, ECsvCustomStatOp::Set);
+			CSV_CUSTOM_STAT_DEFINED(ThermalState, static_cast<int32>(ThermalState), ECsvCustomStatOp::Set);
+#endif
+#if STATS
+			static const FName TemperatureStatName = GET_STATFNAME(STAT_DeviceTemperature);
+			static const FName ThermalStateStatName = GET_STATFNAME(STAT_ThermalState);
+			SET_FLOAT_STAT_FName(TemperatureStatName, DeviceTemperature);
+			SET_DWORD_STAT_FName(ThermalStateStatName, ThermalState);
+#endif
+		}
+	}
+}
\ No newline at end of file
diff --git a/Engine/Source/Runtime/Core/Public/Android/AndroidPlatformMisc.h b/Engine/Source/Runtime/Core/Public/Android/AndroidPlatformMisc.h
--- a/Engine/Source/Runtime/Core/Public/Android/AndroidPlatformMisc.h
+++ b/Engine/Source/Runtime/Core/Public/Android/AndroidPlatformMisc.h
@@ -149,6 +149,8 @@ public:
 	static CORE_API bool IsRunningOnBattery();
 	static CORE_API bool IsInLowPowerMode();
 	static CORE_API float GetDeviceTemperatureLevel();
+	static CORE_API float GetDeviceTemperature();
+	static CORE_API EDeviceThermalState GetDeviceThermalState();
 	static CORE_API bool AreHeadPhonesPluggedIn();
 	static CORE_API ENetworkConnectionType GetNetworkConnectionType();
 #if USE_ANDROID_JNI
diff --git a/Engine/Source/Runtime/Core/Public/Android/AndroidStats.h b/Engine/Source/Runtime/Core/Public/Android/AndroidStats.h
--- a/Engine/Source/Runtime/Core/Public/Android/AndroidStats.h
+++ b/Engine/Source/Runtime/Core/Public/Android/AndroidStats.h
@@ -1,6 +1,8 @@
 // Copyright Epic Games, Inc. All Rights Reserved.
 #pragma once
 
+enum class EDeviceThermalState : int32;
+
 class FAndroidStats
 {
 public:
@@ -9,4 +11,7 @@ public:
 	static void OnThermalStatusChanged(int Status);
 	static void OnTrimMemory(int TrimLevel);
 	static void LogGPUStats();
+
+	// Stat accessors
+	static EDeviceThermalState GetThermalStatus();
 };
diff --git a/Engine/Source/Runtime/Core/Public/GenericPlatform/GenericPlatformMisc.h b/Engine/Source/Runtime/Core/Public/GenericPlatform/GenericPlatformMisc.h
--- a/Engine/Source/Runtime/Core/Public/GenericPlatform/GenericPlatformMisc.h
+++ b/Engine/Source/Runtime/Core/Public/GenericPlatform/GenericPlatformMisc.h
@@ -554,6 +554,26 @@ enum class EContextSwitchFlags : uint8
 
 ENUM_CLASS_FLAGS(EContextSwitchFlags);
 
+/**
+ * Generic Thermal State Enum returned by platform implementations of FGenericPlatformMisc::GetDeviceThermalState.
+ * This is a direct mapping to Android getThermalStatus.  This also means that historical android status data matches.
+ * IOS NSProcessInfoThermalState, which provides less detail, will be mapped onto the first 4 values.
+ * Unless a platform provides enough to supply at least the None and Severe states GetDeviceThermalState should be left Unsupported.
+ * In all cases platform specific stats may have better information.
+ */
+enum class EDeviceThermalState : int32
+{
+	Unsupported = -1,	// Feature is unsupported or an error is returned from the platform (which may also means unsupported on this device).
+
+	None = 0,		// No thermal perf impacts.								If this function is supported on a platform this value should be meaningful.
+	Light = 1,		// Device likely warm, but little or no throttling.
+	Moderate = 2,	// Device likely warmer, but little or no throttling.
+	Severe = 3,		// Serious throttling engaged, Device likely hot.		If this function is supported on a platform this value should be meaningful.
+	Critical = 4,	// Fully throttling.
+	Emergency = 5,	// Device functionality limited.
+	Shutdown = 6,	// Device should shutdown.
+};
+
 /**
  * Generic implementation for most platforms
  */
@@ -1507,8 +1527,30 @@ public:
 	 * Level is a relative value that is platform dependent. Lower is cooler, higher is warmer.
 	 * Level of -1 means Unimplemented.
 	 */
+	UE_DEPRECATED(5.8, "GetDeviceTemperatureLevel() has been superseded by GetDeviceTemperature() and/or GetDeviceThermalState() because it's behavior was poorly defined (on android it returned temperature, on ios it returned a thermal status enum value).")
 	static CORE_API float GetDeviceTemperatureLevel();
 
+	/**
+	 * Returns the current device temperature in degrees Celcius.
+	 * Keep in mind that no specific temperature probe position can be assumed.  It will vary across platforms and devices.
+	 * Inspecting the platform implementation might yeild more information about this stat or more detailed platform specific stats.
+	 * -1.0 is returned for a platform that does not support this measure at all.
+	 * A platform that does support it, but is not currently providing data may return -1 or 0.  
+	 * Theoretically -1 and 0 could be real temperatures... 
+	 */
+	static CORE_API float GetDeviceTemperature() { return -1.0f; }
+
+	/**
+	 * Returns the current device temperature in degrees Celcius.
+	 * Keep in mind that no specific temperature probe position can be assumed.  It will vary across platforms and devices.
+	 * Inspecting the platform implementation might yeild more information about this stat or more detailed platform specific stats.
+	 * -1.0 is returned for a platform that does not support this measure at all.
+	 * A platform that does support it, but is not currently providing data may return -1 or 0.
+	 * Theoretically -1 and 0 could be real temperatures, but we will assume they are not.
+	 */
+	static CORE_API EDeviceThermalState GetDeviceThermalState() { return EDeviceThermalState::Unsupported; }
+
+
 	/**
 	 * Allows a game/program/etc to control the game directory in a special place (for instance, monolithic programs that don't have .uprojects)
 	 */
diff --git a/Engine/Source/Runtime/Core/Public/IOS/IOSPlatformMisc.h b/Engine/Source/Runtime/Core/Public/IOS/IOSPlatformMisc.h
--- a/Engine/Source/Runtime/Core/Public/IOS/IOSPlatformMisc.h
+++ b/Engine/Source/Runtime/Core/Public/IOS/IOSPlatformMisc.h
@@ -59,6 +59,7 @@ struct CORE_API FIOSPlatformMisc : public FApplePlatformMisc
 	static int GetBatteryLevel();
 	static bool IsRunningOnBattery();
 	static float GetDeviceTemperatureLevel();
+	static EDeviceThermalState GetDeviceThermalState();
 	static EDeviceScreenOrientation GetDeviceOrientation();
 	UE_DEPRECATED(5.1, "SetDeviceOrientation is deprecated. Use SetAllowedDeviceOrientation instead.")
 	static void SetDeviceOrientation(EDeviceScreenOrientation NewDeviceOrientation);
diff --git a/Engine/Source/Runtime/Core/Public/Stats/ThermalStats.h b/Engine/Source/Runtime/Core/Public/Stats/ThermalStats.h
new file mode 100644
--- /dev/null
+++ b/Engine/Source/Runtime/Core/Public/Stats/ThermalStats.h
@@ -0,0 +1,44 @@
+// Copyright Epic Games, Inc. All Rights Reserved.
+
+#pragma once
+
+#include "CoreTypes.h"
+#include "Stats/Stats.h"
+#include "ProfilingDebugging/CsvProfiler.h"
+
+DECLARE_STATS_GROUP(TEXT("Thermal Stats"), STATGROUP_Thermal, STATCAT_Advanced);
+
+// Generic thermal stats
+// Note that readings represented by generic thermal stats are similar but not exactly the same across platforms or even across device models.
+// If no similar reading exists the generic stat should be left unsupported.
+
+// DeviceTemperature stat for the device.  
+// This is meant to represent the 'main' temperature reading related to throttling or shutdown on a device.
+// However it might be a cpu measurement, it might be a battery measurement, etc.  The temperatures at which throttling occurs might vary between devices.
+// Be cautious about comparing this value across different devices models even within the same platform.
+// This may be duplicated in a platform specific stat.
+// Look at the platform specific code that fills it in and refer to platform documentation for more detail.
+// It is possible that this value will only be available on certain non-shipping builds on some platforms.
+// -1 or 0 will indicate that no value is being collected (though theoretically they could also be actual temperatures).
+DECLARE_FLOAT_COUNTER_STAT_EXTERN(TEXT("Device Temperature"), STAT_DeviceTemperature, STATGROUP_Thermal, CORE_API);
+
+// Qualitative thermal state of the device.
+// See EDeviceThermalState for more information.
+DECLARE_DWORD_COUNTER_STAT_EXTERN(TEXT("Thermal State"), STAT_ThermalState, STATGROUP_Thermal, CORE_API);
+
+CSV_DECLARE_CATEGORY_EXTERN(Thermal);
+CSV_DECLARE_STAT_EXTERN(Thermal, DeviceTemperature);
+CSV_DECLARE_STAT_EXTERN(Thermal, ThermalState);
+
+namespace UE
+{
+	namespace ThermalStats
+	{
+		// One option is to call the following function to pull the stats using FPlatformMisc accessors, 
+		// then push them out to stats.
+		void GetAndPushStats();
+
+		// Another option is to push the stats defined above directly in platform code.
+		// If you do that you should still implement FPlatformMisc::GetDeviceTemperature(), etc. There are other ways they are output.
+	}
+}
\ No newline at end of file
diff --git a/Engine/Source/Runtime/Engine/Private/ChartCreation.cpp b/Engine/Source/Runtime/Engine/Private/ChartCreation.cpp
--- a/Engine/Source/Runtime/Engine/Private/ChartCreation.cpp
+++ b/Engine/Source/Runtime/Engine/Private/ChartCreation.cpp
@@ -649,8 +649,14 @@ void FPerformanceTrackingChart::Reset(const FDateTime& InStartTime)
 	// ignore zero since we can infer that from the frame counter
 	AsyncPackageQueueHistogram.InitFromArray({ 1.0, 250.0, 500.0 });
 
+	PRAGMA_DISABLE_DEPRECATION_WARNINGS
 	StartTemperatureLevel = -1.0f;
 	StopTemperatureLevel = -1.0f;
+	PRAGMA_ENABLE_DEPRECATION_WARNINGS
+	StartTemperature = -1.0f;
+	StopTemperature = -1.0f;
+	StartThermalState = EDeviceThermalState::Unsupported;
+	StopThermalState = EDeviceThermalState::Unsupported;
 
 	StartBatteryLevel = -1;
 	StopBatteryLevel = -1;
@@ -697,7 +703,9 @@ void FPerformanceTrackingChart::AccumulateWith(const FPerformanceTrackingChart&
 	MinDrawnPrimitives = FMath::Min(MinDrawnPrimitives, Chart.MinDrawnPrimitives);
 	TotalDrawnPrimitives += Chart.TotalDrawnPrimitives;
 	// Assume these charts are appended in order. We won't change the start values, but will update the end values.
+	PRAGMA_DISABLE_DEPRECATION_WARNINGS
 	StopTemperatureLevel = Chart.StopTemperatureLevel;
+	PRAGMA_ENABLE_DEPRECATION_WARNINGS
 	StopBatteryLevel = Chart.StopBatteryLevel;
 	AccumulatedChartTime += Chart.AccumulatedChartTime;
 	TimeDisregarded += Chart.TimeDisregarded;
@@ -728,7 +736,11 @@ double FPerformanceTrackingChart::GetTotalHitchFrameTime(bool bSubtractHitchThre
 
 void FPerformanceTrackingChart::StartCharting()
 {
-	StartTemperatureLevel = FPlatformMisc::GetDeviceTemperatureLevel();
+	PRAGMA_DISABLE_DEPRECATION_WARNINGS
+	StartTemperatureLevel = FPlatformMisc::GetDeviceTemperatureLevel(); // Deprecated in 5.8 in favor of the two stats below.
+	PRAGMA_ENABLE_DEPRECATION_WARNINGS
+	StartTemperature = FPlatformMisc::GetDeviceTemperature();
+	StartThermalState = FPlatformMisc::GetDeviceThermalState();
 	StartBatteryLevel = FPlatformMisc::GetBatteryLevel();
 	DeviceProfileName = UDeviceProfileManager::Get().GetActiveDeviceProfileName();
 	bIsChartingPaused = false;
@@ -736,7 +748,11 @@ void FPerformanceTrackingChart::StartCharting()
 
 void FPerformanceTrackingChart::StopCharting()
 {
-	StopTemperatureLevel = FPlatformMisc::GetDeviceTemperatureLevel();
+	PRAGMA_DISABLE_DEPRECATION_WARNINGS
+	StopTemperatureLevel = FPlatformMisc::GetDeviceTemperatureLevel(); // Deprecated in 5.8 in favor of the two stats below.
+	PRAGMA_ENABLE_DEPRECATION_WARNINGS
+	StopTemperature = FPlatformMisc::GetDeviceTemperature();
+	StopThermalState = FPlatformMisc::GetDeviceThermalState();
 	StopBatteryLevel = FPlatformMisc::GetBatteryLevel();
 }
 
@@ -1040,8 +1056,15 @@ void FPerformanceTrackingChart::DumpChartToAnalyticsParams(const FString& InMapN
 			const float FrameRateLimit = UserSettingsObj->GetFrameRateLimit();
 			InParamArray.Add(FAnalyticsEventAttribute(TEXT("FrameRateLimit"), FrameRateLimit));
 
+			PRAGMA_DISABLE_DEPRECATION_WARNINGS
 			InParamArray.Add(FAnalyticsEventAttribute(TEXT("StartTemp"), StartTemperatureLevel));
 			InParamArray.Add(FAnalyticsEventAttribute(TEXT("StopTemp"), StopTemperatureLevel));
+			PRAGMA_ENABLE_DEPRECATION_WARNINGS
+
+			InParamArray.Add(FAnalyticsEventAttribute(TEXT("StartTemperature"), StartTemperature));
+			InParamArray.Add(FAnalyticsEventAttribute(TEXT("StopTemperature"), StopTemperature));
+			InParamArray.Add(FAnalyticsEventAttribute(TEXT("StartThermalState"), (int)StartThermalState));
+			InParamArray.Add(FAnalyticsEventAttribute(TEXT("StopThermalState"), (int)StopThermalState));
 			InParamArray.Add(FAnalyticsEventAttribute(TEXT("StartBatteryLevel"), StartBatteryLevel));
 			InParamArray.Add(FAnalyticsEventAttribute(TEXT("StopBatteryLevel"), StopBatteryLevel));
 		}
diff --git a/Engine/Source/Runtime/Engine/Public/ChartCreation.h b/Engine/Source/Runtime/Engine/Public/ChartCreation.h
--- a/Engine/Source/Runtime/Engine/Public/ChartCreation.h
+++ b/Engine/Source/Runtime/Engine/Public/ChartCreation.h
@@ -12,6 +12,8 @@
 #include "Scalability.h"
 #include "Delegates/IDelegateInstance.h"
 
+enum class EDeviceThermalState : int32;
+
 //////////////////////////////////////////////////////////////////////
 
 // What kind of hitch was detected (if any)
@@ -223,8 +225,14 @@ public:
 	/** Number of frames that were disregarded because the frame was too long. This is probably a clock glitch or an app returning from background, suspend, etc, and can really skew the data. */
 	int FramesDisregarded;
 
+	UE_DEPRECATED(5.8, "StartTemperatureLevel has been superseded by StartTemperature and/or StartThermalState because it's behavior was poorly defined (on android it was temperature, on ios it was a thermal status enum value).")
 	float StartTemperatureLevel;
+	UE_DEPRECATED(5.8, "StopTemperatureLevel has been superseded by StopTemperature and/or StopThermalState because it's behavior was poorly defined (on android it was temperature, on ios it was a thermal status enum value).")
 	float StopTemperatureLevel;
+	float StartTemperature;
+	float StopTemperature;
+	EDeviceThermalState StartThermalState;
+	EDeviceThermalState StopThermalState;
 
 	int StartBatteryLevel;
 	int StopBatteryLevel;
diff --git a/Engine/Source/Runtime/Launch/Private/IOS/LaunchIOS.cpp b/Engine/Source/Runtime/Launch/Private/IOS/LaunchIOS.cpp
--- a/Engine/Source/Runtime/Launch/Private/IOS/LaunchIOS.cpp
+++ b/Engine/Source/Runtime/Launch/Private/IOS/LaunchIOS.cpp
@@ -30,6 +30,7 @@
 #include "Misc/EmbeddedCommunication.h"
 #include "Misc/CoreDelegates.h"
 #include "Misc/FeedbackContext.h"
+#include "Stats/ThermalStats.h"
 #include "IOS/IOSPlatformCrashContext.h"
 
 #if PLATFORM_VISIONOS && UE_USE_SWIFT_UI_MAIN
@@ -562,6 +563,7 @@ void FAppEntry::Tick()
 #if BUILD_EMBEDDED_APP
 	GWasTickSuspended = false;
 #endif
+	UE::ThermalStats::GetAndPushStats();
     
 	if (AudioContextResumeTime != 0)
 	{
