Change USplineComponent::UpdateSpline to only update the current spline rather than both SplineCurves and Spline.

Shows a 5x improvement in UpdateSpline.

This can be disabled via SplineComponent.OnlyUpdateActiveSpline 0.

#rnx

diff --git a/Engine/Config/BaseEditor.ini b/Engine/Config/BaseEditor.ini
--- a/Engine/Config/BaseEditor.ini
+++ b/Engine/Config/BaseEditor.ini
@@ -764,4 +764,7 @@ PersistentAllocatorReserveSizeMB=2048
 +VariantTags=(Name="ControlRigModules_Quadruped",Label="Quadruped",ToolTip=NSLOCTEXT("[/Script/RigVM]", "1A58BF794BCA84FB89CD74B4C75152B7", "Quadruped modules"),Color=(R=0.514918,G=0.046665,B=1.000000,A=1.000000),bShowInUserInterface=True,bMarksSubjectAsInvalid=False)
 
 [AIAssistant]
-bIsEnabled=true
\ No newline at end of file
+bIsEnabled=true
+
+[ConsoleVariables]
+SplineComponent.OnlyUpdateActiveSpline=0
\ No newline at end of file
diff --git a/Engine/Source/Runtime/Engine/Private/Components/SplineComponent.cpp b/Engine/Source/Runtime/Engine/Private/Components/SplineComponent.cpp
--- a/Engine/Source/Runtime/Engine/Private/Components/SplineComponent.cpp
+++ b/Engine/Source/Runtime/Engine/Private/Components/SplineComponent.cpp
@@ -114,6 +114,13 @@ FAutoConsoleVariableRef CVarForceLastAuthority(
 	TEXT("If 1, synchronization on load will treat FSplineCurves as the authority. If 2, synchronization on load will tread FSpline as the authority.")
 );
 
+bool GOnlyUpdateActiveSpline = true;
+FAutoConsoleVariableRef CVarOnlyUpdateActiveSpline(
+	TEXT("SplineComponent.OnlyUpdateActiveSpline"),
+	GOnlyUpdateActiveSpline,
+	TEXT("When true, UpdateSpline will only call update on the current spline in use, based on ShouldUseSplineCurves(). Off by default in editor.")
+);
+
 static ELastAuthority GetEffectiveLastAuthority(const ELastAuthority InLastAuthority)
 {
 	ELastAuthority UseLastAuthority = InLastAuthority;
@@ -582,6 +589,8 @@ void USplineComponent::OnRegister()
 
 void FSplineCurves::UpdateSpline(bool bClosedLoop, bool bStationaryEndpoints, int32 ReparamStepsPerSegment, bool bLoopPositionOverride, float LoopPosition, const FVector& Scale3D)
 {
+	TRACE_CPUPROFILER_EVENT_SCOPE(FSplineCurves::UpdateSpline);
+
 	const int32 NumPoints = Position.Points.Num();
 	check(Rotation.Points.Num() == NumPoints && Scale.Points.Num() == NumPoints);
 
@@ -637,9 +646,25 @@ void FSplineCurves::UpdateSpline(bool bClosedLoop, bool bStationaryEndpoints, in
 
 void USplineComponent::UpdateSpline()
 {
+	TRACE_CPUPROFILER_EVENT_SCOPE(USplineComponent::UpdateSpline);
+
 	FSpline::FUpdateSplineParams Params{bClosedLoop, bStationaryEndpoints, ReparamStepsPerSegment, bLoopPositionOverride, LoopPosition, GetComponentTransform().GetScale3D()};
-	WarninglessSplineCurves().UpdateSpline(Params.bClosedLoop, Params.bStationaryEndpoints, Params.ReparamStepsPerSegment, Params.bLoopPositionOverride, Params.LoopPosition, Params.Scale3D);
-	Spline.UpdateSpline(Params);
+	if (UE::SplineComponent::GOnlyUpdateActiveSpline)
+	{
+		if (UE::SplineComponent::ShouldUseSplineCurves())
+		{
+			WarninglessSplineCurves().UpdateSpline(Params.bClosedLoop, Params.bStationaryEndpoints, Params.ReparamStepsPerSegment, Params.bLoopPositionOverride, Params.LoopPosition, Params.Scale3D);
+		}
+		else
+		{
+			Spline.UpdateSpline(Params);
+		}
+	}
+	else
+	{
+		WarninglessSplineCurves().UpdateSpline(Params.bClosedLoop, Params.bStationaryEndpoints, Params.ReparamStepsPerSegment, Params.bLoopPositionOverride, Params.LoopPosition, Params.Scale3D);
+		Spline.UpdateSpline(Params);
+	}
 
 	// todo: make Spline replicate
 	//MARK_PROPERTY_DIRTY_FROM_NAME(USplineComponent, Spline, this);
