Change subsystem collection initialization to make re-entrant access to subsystems during subsystem initialization more robust without needing to use InitializeDependency.
Some subsystem accesses are significantly removed from the subsystem initialization e.g. subsystem loads content -> trigger BP or UMG compilation -> runs code in those modules that accesses engine subsystems

#Synced-CL Unknown

diff --git a/Engine/Source/Runtime/Engine/Private/Subsystems/SubsystemCollection.cpp b/Engine/Source/Runtime/Engine/Private/Subsystems/SubsystemCollection.cpp
--- a/Engine/Source/Runtime/Engine/Private/Subsystems/SubsystemCollection.cpp
+++ b/Engine/Source/Runtime/Engine/Private/Subsystems/SubsystemCollection.cpp
@@ -2,10 +2,13 @@
 
 #include "Subsystems/SubsystemCollection.h"
 
+#include "Algo/RemoveIf.h"
 #include "Subsystems/Subsystem.h"
 #include "UObject/Interface.h"
+#include "UObject/UObjectGlobals.h"
 #include "UObject/UObjectHash.h"
 #include "UObject/Package.h"
+#include "Logging/StructuredLog.h"
 #include "Modules/ModuleManager.h"
 #include "Misc/PackageName.h"
 
@@ -31,18 +34,25 @@ private:
 // globals without thread protection must be accessed only from GameThread
 FDelegateHandle FSubsystemModuleWatcher::ModulesChangedHandle;
 static TArray<FSubsystemCollectionBase*> GlobalSubsystemCollections;
-static TMap<FName, TArray<TSubclassOf<UDynamicSubsystem>>> GlobalDynamicSystemModuleMap;
+static TMap<FName, TArray<UClass*>> GlobalDynamicSystemModuleMap;
+
+struct FSubsystemCollectionInitialization
+{
+	// Classes of subsystem we intend to initialize, for handling calls to GetSubsystem during initialization
+	// Key is a base/interface/concrete class which may be used in a call to GetSubsystem
+	// Value is the concrete class that will be returned
+	// So when key == value, this is a concrete class we intend to initialize
+	TMap<UClass*, UClass*> ClassMap;
+};
 
 FSubsystemCollectionBase::FSubsystemCollectionBase()
 	: Outer(nullptr)
-	, bPopulating(false)
 {
 }
 
 FSubsystemCollectionBase::FSubsystemCollectionBase(UClass* InBaseType)
 	: BaseType(InBaseType)
 	, Outer(nullptr)
-	, bPopulating(false)
 {
 	check(BaseType);
 }
@@ -78,6 +88,23 @@ USubsystem* FSubsystemCollectionBase::GetSubsystemInternal(UClass* SubsystemClas
 		}
 	}
 
+	if (Initialization != nullptr)
+	{
+		// When attempting to access one subsystem from the initialization of another, check the set of subsystems we intend to create
+		// for a subsystem of the exact class or of a derived type
+		if (UClass* Class = Initialization->ClassMap.FindRef(SubsystemClass))
+		{
+			// Cast away const explicitly to initialize a subsystem re-entrantly...
+			FSubsystemCollectionBase* MutableThis = const_cast<FSubsystemCollectionBase*>(this);
+			return MutableThis->AddAndInitializeValidatedSubsystem(Class);
+		}
+		else 
+		{
+			ensureMsgf(false, TEXT("Failed a call to GetSubsystem from within Initialize of another subsystem as the requested type %s was not present in the set of classes we expected to initialize"),
+				*SubsystemClass->GetPathName());
+		}
+	}
+
 	return nullptr;
 }
 
@@ -177,7 +204,7 @@ void FSubsystemCollectionBase::Initialize(UObject* NewOuter)
 	check(Outer);
 	if (ensure(BaseType) && ensureMsgf(SubsystemMap.Num() == 0, TEXT("Currently don't support repopulation of Subsystem Collections.")))
 	{
-		check(!bPopulating); //Populating collections on multiple threads?
+		check(!IsPopulating()); //Populating collections on multiple threads?
 		
 		//non-thread-safe use of Global lists, must be from GameThread:
 		check(IsInGameThread());
@@ -189,32 +216,27 @@ void FSubsystemCollectionBase::Initialize(UObject* NewOuter)
 
 		UE_LOG(LogSubsystemCollection, Verbose, TEXT("Initializing subsystem collection for %s with type %s"), *GetNameSafe(NewOuter), *GetNameSafe(BaseType));
 		
-		TGuardValue<bool> PopulatingGuard(bPopulating, true);
-
+		TArray<UClass*> SubsystemClasses;
 		if (BaseType->IsChildOf(UDynamicSubsystem::StaticClass()))
 		{
-			for (const TPair<FName, TArray<TSubclassOf<UDynamicSubsystem>>>& SubsystemClasses : GlobalDynamicSystemModuleMap)
+			for (const TPair<FName, TArray<UClass*>>& ModuleClasses : GlobalDynamicSystemModuleMap)
 			{
-				for (const TSubclassOf<UDynamicSubsystem>& SubsystemClass : SubsystemClasses.Value)
+				for (UClass* SubsystemClass : ModuleClasses.Value)
 				{
 					if (SubsystemClass->IsChildOf(BaseType))
 					{
-						AddAndInitializeSubsystem(SubsystemClass);
+						SubsystemClasses.Add(SubsystemClass);
 					}
 				}
 			}
 		}
 		else
 		{
-			TArray<UClass*> SubsystemClasses;
 			GetDerivedClasses(BaseType, SubsystemClasses, true);
-
-			for (UClass* SubsystemClass : SubsystemClasses)
-			{
-				AddAndInitializeSubsystem(SubsystemClass);
-			}
 		}
 
+		AddAndInitializeSubsystems(SubsystemClasses);
+
 		ModulesUnloadedHandle = FCoreUObjectDelegates::CompiledInUObjectsRemovedDelegate.AddRaw(this, &FSubsystemCollectionBase::RemoveSubsystemsInPackages);
 
 		// Statically track collections
@@ -236,6 +258,11 @@ FSubsystemCollectionBase::~FSubsystemCollectionBase()
 	Deinitialize();
 }
 
+bool FSubsystemCollectionBase::IsPopulating() const
+{
+	return Initialization != nullptr;
+}
+
 void FSubsystemCollectionBase::Deinitialize()
 {
 	//non-thread-safe use of Global lists, must be from GameThread:
@@ -287,14 +314,27 @@ void FSubsystemCollectionBase::Deinitialize()
 USubsystem* FSubsystemCollectionBase::InitializeDependency(TSubclassOf<USubsystem> SubsystemClass)
 {
 	USubsystem* Subsystem = nullptr;
-	if (ensureMsgf(SubsystemClass, TEXT("Attempting to add invalid subsystem as dependancy.")))
+	if (ensureMsgf(SubsystemClass, TEXT("Attempting to add invalid subsystem as dependency."))
+		&& ensureMsgf(SubsystemClass->IsChildOf(BaseType), TEXT("ClassType (%s) must be a subclass of BaseType(%s)."), *SubsystemClass->GetPathName(), *BaseType->GetPathName()))
 	{
 		UE_LOG(LogSubsystemCollection, VeryVerbose, TEXT("Attempting to initialize subsystem dependency (%s)"), *SubsystemClass->GetName());
 
-		if (ensureMsgf(bPopulating, TEXT("InitializeDependancy() should only be called from System USubsystem::Initialization() implementations."))
-			&& ensureMsgf(SubsystemClass->IsChildOf(BaseType), TEXT("ClassType (%s) must be a subclass of BaseType(%s)."), *SubsystemClass->GetName(), *BaseType->GetName()))
+		if (ensureMsgf(Initialization != nullptr, TEXT("InitializeDependency() should only be called from System USubsystem::Initialization() implementations.")))
 		{
-			Subsystem = AddAndInitializeSubsystem(SubsystemClass);
+			UClass* ConcreteClass = Initialization->ClassMap.FindRef(SubsystemClass);
+			if (!ConcreteClass)
+			{
+				UE_LOGFMT(LogSubsystemCollection, Verbose, "InitializeDependency for subsystem class {RequestedClass} did not find an existing class which the subsytem intended to initialize, trying to initialize with the class given anyway.",
+					FTopLevelAssetPath(SubsystemClass));
+				ConcreteClass = SubsystemClass;
+			}
+			else if (ConcreteClass != SubsystemClass)
+			{
+				UE_LOGFMT(LogSubsystemCollection, Verbose, "InitializeDependency for subsystem class {RequestedClass} redirected this call to create class {ConcreteClass} based on inheritance hierarchy.", 
+					FTopLevelAssetPath(SubsystemClass), FTopLevelAssetPath(ConcreteClass));
+			}
+
+			Subsystem = AddAndInitializeSubsystem(ConcreteClass); 
 		}
 
 		UE_CLOG(!Subsystem, LogSubsystemCollection, Log, TEXT("Failed to initialize subsystem dependency (%s)"), *SubsystemClass->GetName());
@@ -308,9 +348,68 @@ void FSubsystemCollectionBase::AddReferencedObjects(UObject* Referencer, FRefere
 	Collector.AddStableReferenceMap(SubsystemMap);
 }
 
+void FSubsystemCollectionBase::AddAndInitializeSubsystems(TConstArrayView<UClass*> SubsystemClasses)
+{
+	FSubsystemCollectionInitialization LocalInitialization;
+
+	for (UClass* Class : SubsystemClasses)
+	{
+		if (Class->HasAllClassFlags(CLASS_Abstract) || Class->GetAuthoritativeClass() != Class)
+		{
+			continue;
+		}
+		if (!Class->GetDefaultObject<USubsystem>()->ShouldCreateSubsystem(Outer))
+		{
+			UE_LOGFMT(LogSubsystemCollection, Verbose, "Not creating subsystem of class {Class} as it returned false from ShouldCreateSubsystem",
+				FTopLevelAssetPath(Class));
+			continue;
+		}
+		LocalInitialization.ClassMap.Add(Class, Class);
+	}
+
+	// Store a map of classes we intend to initialize for re-entrant calls to GetSubsystem 
+	for (UClass* ConcreteClass : SubsystemClasses)
+	{
+		// Only process classes which were added mapping to themselves above, not more classes which
+		// were added during this loop
+		if (LocalInitialization.ClassMap.FindRef(ConcreteClass) != ConcreteClass)
+		{
+			continue;
+		}
+		for (UClass* Parent = ConcreteClass->GetSuperClass();
+			Parent != nullptr && Parent != BaseType && !LocalInitialization.ClassMap.Contains(Parent);
+			Parent = Parent->GetSuperClass())
+		{
+			LocalInitialization.ClassMap.Add(Parent, ConcreteClass);
+		}
+
+		for (const FImplementedInterface& Interface : ConcreteClass->Interfaces)
+		{
+			if (!LocalInitialization.ClassMap.Contains(Interface.Class))
+			{
+				LocalInitialization.ClassMap.Add(Interface.Class, ConcreteClass);
+			}
+		}
+	}
+
+	// Only mark ourselves as populating now - some subsystems call GetSubsystem in their ShouldCreateSubsystem 
+	TGuardValue PopulatingGuard(Initialization, &LocalInitialization);
+	for (TPair<UClass*, UClass*> Pair : LocalInitialization.ClassMap)
+	{
+		// Only initialize the requested classes, other keys are for finding via the inheritence hierarchy 
+		// with re-entrant calls to GetSubsystem
+		// Skip creation if we created the subsystem re-entrantly
+		if (Pair.Key == Pair.Value && !SubsystemMap.Contains(Pair.Key))
+		{
+			AddAndInitializeValidatedSubsystem(Pair.Key);
+		}
+	}
+}
+
 USubsystem* FSubsystemCollectionBase::AddAndInitializeSubsystem(UClass* SubsystemClass)
 {
-	TGuardValue<bool> PopulatingGuard(bPopulating, true);
+	FSubsystemCollectionInitialization TempInitialization;
+	TGuardValue PopulatingGuard(Initialization, Initialization ? Initialization : &TempInitialization);
 
 	if (!SubsystemMap.Contains(SubsystemClass))
 	{
@@ -331,24 +430,7 @@ USubsystem* FSubsystemCollectionBase::AddAndInitializeSubsystem(UClass* Subsyste
 			const USubsystem* CDO = SubsystemClass->GetDefaultObject<USubsystem>();
 			if (CDO->ShouldCreateSubsystem(Outer))
 			{
-				USubsystem* Subsystem = NewObject<USubsystem>(Outer, SubsystemClass);
-				SubsystemMap.Add(SubsystemClass,Subsystem);
-				Subsystem->InternalOwningSubsystem = this;
-				Subsystem->Initialize(*this);
-				
-				// Add this new subsystem to any existing maps of base classes to lists of subsystems
-				// Not calling FatalErrorIfIteratingSubsystems because adding to the end of the array is safe for index-based iteration
-				for (TPair<UClass*, TUniquePtr<FSubsystemArray>>& Pair : SubsystemArrayMap)
-				{
-					const bool bIsInterface = Pair.Key->IsChildOf<UInterface>();
-					if ((!bIsInterface && SubsystemClass->IsChildOf(Pair.Key)) || 
-						(bIsInterface && SubsystemClass->ImplementsInterface(Pair.Key)))
-					{
-						Pair.Value->Subsystems.Add(Subsystem);
-					}
-				}
-
-				return Subsystem;
+				return AddAndInitializeValidatedSubsystem(SubsystemClass);
 			}
 
 			UE_LOG(LogSubsystemCollection, VeryVerbose, TEXT("Subsystem does not exist, but CDO choose to not create (%s)"), *SubsystemClass->GetName());
@@ -360,6 +442,34 @@ USubsystem* FSubsystemCollectionBase::AddAndInitializeSubsystem(UClass* Subsyste
 	return SubsystemMap.FindRef(SubsystemClass);
 }
 
+USubsystem* FSubsystemCollectionBase::AddAndInitializeValidatedSubsystem(UClass* SubsystemClass)
+{
+	checkf(Initialization, TEXT("AddAndInitializeValidatedSubsystem should only be called from an initialization context"));
+	checkf(!SubsystemMap.Contains(SubsystemClass), TEXT("AddAndInitializeValidatedSubsystem called for a class which has already been added"));
+
+	UE_SCOPED_ENGINE_ACTIVITY(TEXT("Initializing Subsystem %s"), *SubsystemClass->GetName());
+	UE_LOG_CONTEXT("SubsystemBeingInitialized", FTopLevelAssetPath(SubsystemClass));
+
+	USubsystem* Subsystem = NewObject<USubsystem>(Outer, SubsystemClass);
+	SubsystemMap.Add(SubsystemClass,Subsystem);
+	Subsystem->InternalOwningSubsystem = this;
+	Subsystem->Initialize(*this);
+	
+	// Add this new subsystem to any existing maps of base classes to lists of subsystems
+	// Not calling FatalErrorIfIteratingSubsystems because adding to the end of the array is safe for index-based iteration
+	for (TPair<UClass*, TUniquePtr<FSubsystemArray>>& Pair : SubsystemArrayMap)
+	{
+		const bool bIsInterface = Pair.Key->IsChildOf<UInterface>();
+		if ((!bIsInterface && SubsystemClass->IsChildOf(Pair.Key)) || 
+			(bIsInterface && SubsystemClass->ImplementsInterface(Pair.Key)))
+		{
+			Pair.Value->Subsystems.Add(Subsystem);
+		}
+	}
+
+	return Subsystem;
+}
+
 void FSubsystemCollectionBase::RemoveAndDeinitializeSubsystem(USubsystem* Subsystem)
 {
 	check(Subsystem);
@@ -400,40 +510,51 @@ void FSubsystemCollectionBase::RemoveAndDeinitializeSubsystem(USubsystem* Subsys
 
 void FSubsystemCollectionBase::ActivateExternalSubsystem(UClass* SubsystemClass)
 {
-	AddAllInstances(SubsystemClass);
+	AddAllInstances(MakeConstArrayView(&SubsystemClass, 1));
 }
 
 void FSubsystemCollectionBase::DeactivateExternalSubsystem(UClass* SubsystemClass)
 {
-	RemoveAllInstances(SubsystemClass);
+	RemoveAllInstances(MakeConstArrayView(&SubsystemClass, 1));
 }
 
-void FSubsystemCollectionBase::AddAllInstances(UClass* SubsystemClass)
+void FSubsystemCollectionBase::AddAllInstances(TConstArrayView<UClass*> SubsystemClasses)
 {
 	//non-thread-safe use of Global lists, must be from GameThread:
 	check(IsInGameThread());
 
+	TArray<UClass*> RelevantClasses;
 	for (FSubsystemCollectionBase* SubsystemCollection : GlobalSubsystemCollections)
 	{
-		if (SubsystemClass->IsChildOf(SubsystemCollection->BaseType))
+		RelevantClasses.Reset();
+		for (UClass* Class : SubsystemClasses)
 		{
-			SubsystemCollection->AddAndInitializeSubsystem(SubsystemClass);
+			if (Class->IsChildOf(SubsystemCollection->BaseType))
+			{
+				RelevantClasses.Add(Class);
+			}
 		}
+
+		SubsystemCollection->AddAndInitializeSubsystems(RelevantClasses);
 	}
 }
 
-void FSubsystemCollectionBase::RemoveAllInstances(UClass* SubsystemClass)
+void FSubsystemCollectionBase::RemoveAllInstances(TConstArrayView<UClass*> SubsystemClasses)
 {
 	TArray<UObject*> SubsystemsToRemove;
-	GetObjectsOfClass(SubsystemClass, SubsystemsToRemove);
-
-	for(UObject* SubsystemObj : SubsystemsToRemove)
+	for (UClass* SubsystemClass : SubsystemClasses)
 	{
-		USubsystem* Subsystem = CastChecked<USubsystem>(SubsystemObj);
+		SubsystemsToRemove.Reset();
+		GetObjectsOfClass(SubsystemClass, SubsystemsToRemove);
 
-		if (Subsystem->InternalOwningSubsystem)
+		for(UObject* SubsystemObj : SubsystemsToRemove)
 		{
-			Subsystem->InternalOwningSubsystem->RemoveAndDeinitializeSubsystem(Subsystem);
+			USubsystem* Subsystem = CastChecked<USubsystem>(SubsystemObj);
+
+			if (Subsystem->InternalOwningSubsystem)
+			{
+				Subsystem->InternalOwningSubsystem->RemoveAndDeinitializeSubsystem(Subsystem);
+			}
 		}
 	}
 }
@@ -479,7 +600,7 @@ void FSubsystemModuleWatcher::InitializeModuleWatcher()
 				const FName ModuleName = FPackageName::GetShortFName(ClassPackage->GetFName());
 				if (FModuleManager::Get().IsModuleLoaded(ModuleName))
 				{
-					TArray<TSubclassOf<UDynamicSubsystem>>& ModuleSubsystemClasses = GlobalDynamicSystemModuleMap.FindOrAdd(ModuleName);
+					TArray<UClass*>& ModuleSubsystemClasses = GlobalDynamicSystemModuleMap.FindOrAdd(ModuleName);
 					ModuleSubsystemClasses.Add(SubsystemClass);
 				}
 			}
@@ -512,7 +633,7 @@ void FSubsystemModuleWatcher::AddClassesForModule(const FName& InModuleName)
 		return;
 	}
 
-	TArray<TSubclassOf<UDynamicSubsystem>> SubsystemClasses;
+	TArray<UClass*> SubsystemClasses;
 	TArray<UObject*> PackageObjects;
 	GetObjectsWithPackage(ClassPackage, PackageObjects, false);
 	for (UObject* Object : PackageObjects)
@@ -521,9 +642,10 @@ void FSubsystemModuleWatcher::AddClassesForModule(const FName& InModuleName)
 		if (CurrentClass && !CurrentClass->HasAllClassFlags(CLASS_Abstract) && CurrentClass->IsChildOf(UDynamicSubsystem::StaticClass()))
 		{
 			SubsystemClasses.Add(CurrentClass);
-			FSubsystemCollectionBase::AddAllInstances(CurrentClass);
 		}
 	}
+	FSubsystemCollectionBase::AddAllInstances(SubsystemClasses);
+
 	if (SubsystemClasses.Num() > 0)
 	{
 		GlobalDynamicSystemModuleMap.Add(InModuleName, MoveTemp(SubsystemClasses));
@@ -534,13 +656,10 @@ void FSubsystemModuleWatcher::RemoveClassesForModule(const FName& InModuleName)
 	//non-thread-safe use of Global lists, must be from GameThread:
 	check(IsInGameThread());
 
-	TArray<TSubclassOf<UDynamicSubsystem>>* SubsystemClasses = GlobalDynamicSystemModuleMap.Find(InModuleName);
+	TArray<UClass*>* SubsystemClasses = GlobalDynamicSystemModuleMap.Find(InModuleName);
 	if (SubsystemClasses)
 	{
-		for (TSubclassOf<UDynamicSubsystem>& SubsystemClass : *SubsystemClasses)
-		{
-			FSubsystemCollectionBase::RemoveAllInstances(SubsystemClass);
-		}
+		FSubsystemCollectionBase::RemoveAllInstances(*SubsystemClasses);
 		GlobalDynamicSystemModuleMap.Remove(InModuleName);
 	}
 }
diff --git a/Engine/Source/Runtime/Engine/Public/Subsystems/SubsystemCollection.h b/Engine/Source/Runtime/Engine/Public/Subsystems/SubsystemCollection.h
--- a/Engine/Source/Runtime/Engine/Public/Subsystems/SubsystemCollection.h
+++ b/Engine/Source/Runtime/Engine/Public/Subsystems/SubsystemCollection.h
@@ -9,6 +9,8 @@
 class USubsystem;
 class UDynamicSubsystem;
 
+struct FSubsystemCollectionInitialization;
+
 DECLARE_LOG_CATEGORY_EXTERN(LogSubsystemCollection, Log, All);
 
 class FSubsystemCollectionBase
@@ -98,10 +100,17 @@ protected:
 	/** Adds a subsystem to the collection and initializes it */
 	ENGINE_API virtual USubsystem* AddAndInitializeSubsystem(UClass* SubsystemClass);
 
+	/** Add and initialize many subsystems at once, handling dependencies between them as with Initialize */
+	ENGINE_API void AddAndInitializeSubsystems(TConstArrayView<UClass*> SubsystemClasses);
+
 	/** Removes a subsystem to the collection and deinitializes it */
 	ENGINE_API void RemoveAndDeinitializeSubsystem(USubsystem* Subsystem);
 
+	// True if we're adding subsystems either from Initialize or piecemeal
+	[[nodiscard]] bool IsPopulating() const;
+
 private:
+
 	TMap<TObjectPtr<UClass>, TObjectPtr<USubsystem>> SubsystemMap;
 
 	mutable TMap<UClass*, TUniquePtr<FSubsystemArray>> SubsystemArrayMap;
@@ -110,7 +119,9 @@ private:
 
 	UObject* Outer;
 
-	bool bPopulating;
+	// Non-null inside Initialize and other places that we add subsystems
+	FSubsystemCollectionInitialization* Initialization = nullptr;
+
 	mutable bool bIterating = false; // True if iterating over SubsystemMap
 
 	FDelegateHandle ModulesUnloadedHandle;
@@ -118,11 +129,14 @@ private:
 private:
 	friend class FSubsystemModuleWatcher;
 
-	/** Add Instances of the specified Subsystem class to all existing SubsystemCollections of the correct type */
-	static ENGINE_API void AddAllInstances(UClass* SubsystemClass);
+	/** For a class which has already been validated as something we should create, initialize the subsystem and add it to the collection */
+	USubsystem* AddAndInitializeValidatedSubsystem(UClass* SubsystemClass);
+
+	/** Add Instances of the specified Subsystem classes to all existing SubsystemCollections of the correct type */
+	static ENGINE_API void AddAllInstances(TConstArrayView<UClass*> SubsystemClasses);
 
-	/** Remove Instances of the specified Subsystem class from all existing SubsystemCollections of the correct type */
-	static ENGINE_API void RemoveAllInstances(UClass* SubsystemClass);
+	/** Remove Instances of the specified Subsystem classes from all existing SubsystemCollections of the correct type */
+	static ENGINE_API void RemoveAllInstances(TConstArrayView<UClass*> SubsystemClasses);
 };
 
 template<typename TBaseType>
