Fixed RendereAssetUpdate to unlock mips when IO request has been canceled

#rnx

diff --git a/Engine/Source/Runtime/Engine/Private/Streaming/Texture2DStreamIn.cpp b/Engine/Source/Runtime/Engine/Private/Streaming/Texture2DStreamIn.cpp
--- a/Engine/Source/Runtime/Engine/Private/Streaming/Texture2DStreamIn.cpp
+++ b/Engine/Source/Runtime/Engine/Private/Streaming/Texture2DStreamIn.cpp
@@ -10,6 +10,13 @@ Texture2DStreamIn.cpp: Stream in helper for 2D textures.
 #include "Rendering/Texture2DResource.h"
 #include "Streaming/Texture2DUpdate.h"
 
+static int32 GStreamingForceUnlockMipsOnIOCancel = 1;
+static FAutoConsoleVariableRef CVarStreamingForceUnlockMipsOnIOCancel(
+	TEXT("r.Streaming.ForceUnlockMipsOnIOCancel"),
+	GStreamingForceUnlockMipsOnIOCancel,
+	TEXT("Whether to force unlock mips on IO Cancel."),
+	ECVF_Default);
+
 FTexture2DStreamIn::FTexture2DStreamIn(UTexture2D* InTexture)
 	: FTexture2DUpdate(InTexture)
 {
@@ -87,7 +94,7 @@ void FTexture2DStreamIn::DoUnlockNewMips(FRHICommandListImmediate& RHICmdList, c
 {
 	check(Context.CurrentThread == TT_Render);
 
-	if (IntermediateTextureRHI && Context.Resource)
+	if (IntermediateTextureRHI && (Context.Resource || (IsCancelled() && GStreamingForceUnlockMipsOnIOCancel)))
 	{
 		// With virtual textures, all mips exist although they might not be allocated.
 		const int32 MipOffset = !!(IntermediateTextureRHI->GetFlags() & TexCreate_Virtual) ? 0 : PendingFirstLODIdx;
