Renamed ShowFlag "Mode Widget" to "Mode Gizmos"

#Synced-CL 48021612

diff --git a/Engine/Source/Runtime/Engine/Public/ShowFlagsValues.inl b/Engine/Source/Runtime/Engine/Public/ShowFlagsValues.inl
--- a/Engine/Source/Runtime/Engine/Public/ShowFlagsValues.inl
+++ b/Engine/Source/Runtime/Engine/Public/ShowFlagsValues.inl
@@ -267,8 +267,8 @@ SHOWFLAG_FIXED_IN_SHIPPING(0, Selection, SFG_Advanced, NSLOCTEXT("UnrealEd", "Se
 SHOWFLAG_FIXED_IN_SHIPPING(0, VisualizeLevelInstanceEditing, SFG_Visualize, NSLOCTEXT("UnrealEd", "VisualizeLevelInstanceEditingSF", "Visualize Level Instance Editing"))
 /** Used to determine whether there are any editing LevelInstances to enable the post processing pass */
 SHOWFLAG_FIXED_IN_SHIPPING(0, EditingLevelInstance, SFG_Transient, NSLOCTEXT("UnrealEd", "EditingLevelInstanceSF", "Editing LevelInstance"))
-/** Draws mode specific widgets and controls in the viewports (should only be set on viewport clients that are editing the level itself) */
-SHOWFLAG_FIXED_IN_SHIPPING(0, ModeWidgets, SFG_Advanced, NSLOCTEXT("UnrealEd", "ModeWidgetsSF", "Mode Widgets"))
+/** Draws mode specific gizmos and controls in the viewports (should only be set on viewport clients that are editing the level itself) */
+SHOWFLAG_FIXED_IN_SHIPPING(0, ModeWidgets, SFG_Advanced, NSLOCTEXT("UnrealEd", "ModeWidgetsSF", "Mode Gizmos"))
 /**  */
 SHOWFLAG_FIXED_IN_SHIPPING(0, Bounds,  SFG_Advanced, NSLOCTEXT("UnrealEd", "BoundsSF", "Bounds"))
 /** Draws each hit proxy in the scene with a different color, for now only available in the editor */
