Fixed BaseColor default value in FMaterialAttributeDefinitionMap match the constant default on the pin.

[RN] Fixed BaseColor default value inconsistencies to be black.
#Synced-CL 48350544

diff --git a/Engine/Source/Editor/MaterialEditor/Private/MaterialEditor.cpp b/Engine/Source/Editor/MaterialEditor/Private/MaterialEditor.cpp
--- a/Engine/Source/Editor/MaterialEditor/Private/MaterialEditor.cpp
+++ b/Engine/Source/Editor/MaterialEditor/Private/MaterialEditor.cpp
@@ -5152,7 +5152,7 @@ void FMaterialEditor::OnResetToDefault(const FToolMenuContext& InMenuContext) co
 		switch (PropertyId)
 		{
 			case MP_BaseColor:
-				EditorOnlyData->BaseColor.Constant = FLinearColor::Gray;
+				EditorOnlyData->BaseColor.Constant = FMaterialAttributeDefinitionMap::GetDefaultValue(MP_BaseColor);
 				TargetPin->GetSchema()->TrySetDefaultValue(*TargetPin, EditorOnlyData->BaseColor.GetDefaultValue());
 				break;
 
diff --git a/Engine/Source/Runtime/Core/Private/UObject/FortniteMainBranchObjectVersion.cpp b/Engine/Source/Runtime/Core/Private/UObject/FortniteMainBranchObjectVersion.cpp
--- a/Engine/Source/Runtime/Core/Private/UObject/FortniteMainBranchObjectVersion.cpp
+++ b/Engine/Source/Runtime/Core/Private/UObject/FortniteMainBranchObjectVersion.cpp
@@ -13,7 +13,7 @@ TMap<FGuid, FGuid> FFortniteMainBranchObjectVersion::GetSystemGuids()
 	SystemGuids.Add(DevGuids.MATERIALSHADERMAP_DERIVEDDATA_VER, FGuid("4DCE94C9E9994B3DA256F6A015B87CA9"));
 	SystemGuids.Add(DevGuids.NIAGARASHADERMAP_DERIVEDDATA_VER, FGuid("CC229D919A214E59B79D01AC2E5BF981"));
 	SystemGuids.Add(DevGuids.Niagara_LatestScriptCompileVersion, FGuid("FB15C2870A4E4AE8B5F58566F0E3825E"));
-	SystemGuids.Add(DevGuids.MaterialTranslationDDCVersion, FGuid("FD43494F1D8443D79A06E6E28F503A42"));
+	SystemGuids.Add(DevGuids.MaterialTranslationDDCVersion, FGuid("E3ED11E905F64505BD5A6CB7E3599C10"));
 
 	// These GUIDs are in AppendRestrictedFortniteMainStreamObjectVersionGuids as they should rarely be changed.
 	//SystemGuids.Add(DevGuids.NANITE_DERIVEDDATA_VER, FGuid("ACF98184978742A592537DA223A1E6BE"));
diff --git a/Engine/Source/Runtime/Engine/Private/Materials/Material.cpp b/Engine/Source/Runtime/Engine/Private/Materials/Material.cpp
--- a/Engine/Source/Runtime/Engine/Private/Materials/Material.cpp
+++ b/Engine/Source/Runtime/Engine/Private/Materials/Material.cpp
@@ -8143,7 +8143,7 @@ bool UMaterial::IsRefractionPinPluggedIn(const UMaterialEditorOnlyData* EditorOn
 
 UMaterialEditorOnlyData::UMaterialEditorOnlyData()
 {
-	BaseColor.Constant = FLinearColor::Gray;
+	BaseColor.Constant = FMaterialAttributeDefinitionMap::GetDefaultValue(MP_BaseColor);
 	Metallic.Constant = FMaterialAttributeDefinitionMap::GetDefaultValue(MP_Metallic).X;
 	Specular.Constant = FMaterialAttributeDefinitionMap::GetDefaultValue(MP_Specular).X;
 	Roughness.Constant = FMaterialAttributeDefinitionMap::GetDefaultValue(MP_Roughness).X;
diff --git a/Engine/Source/Runtime/Engine/Private/Materials/MaterialAttributeDefinitionMap.cpp b/Engine/Source/Runtime/Engine/Private/Materials/MaterialAttributeDefinitionMap.cpp
--- a/Engine/Source/Runtime/Engine/Private/Materials/MaterialAttributeDefinitionMap.cpp
+++ b/Engine/Source/Runtime/Engine/Private/Materials/MaterialAttributeDefinitionMap.cpp
@@ -331,7 +331,7 @@ void FMaterialAttributeDefinitionMap::InitializeAttributeMap()
 	AttributeMap.Empty(MP_MAX + 1);
 
 	// Basic attributes
-	Add(FGuid(0x69B8D336, 0x16ED4D49, 0x9AA49729, 0x2F050F7A), TEXT("BaseColor"),		MP_BaseColor,		MCT_Float3,	FVector4(0,0,0,0),	SF_Pixel);
+	Add(FGuid(0x69B8D336, 0x16ED4D49, 0x9AA49729, 0x2F050F7A), TEXT("BaseColor"),		MP_BaseColor,		MCT_Float3,	FVector4(0,0,0,1),	SF_Pixel);
 	Add(FGuid(0x57C3A161, 0x7F064296, 0xB00B24A5, 0xA496F34C), TEXT("Metallic"),		MP_Metallic,		MCT_Float,	FVector4(0,0,0,0),	SF_Pixel);
 	Add(FGuid(0x9FDAB399, 0x25564CC9, 0x8CD2D572, 0xC12C8FED), TEXT("Specular"),		MP_Specular,		MCT_Float,	FVector4(.5,0,0,0), SF_Pixel);
 	Add(FGuid(0xD1DD967C, 0x4CAD47D3, 0x9E6346FB, 0x08ECF210), TEXT("Roughness"),		MP_Roughness,		MCT_Float,	FVector4(.5,0,0,0), SF_Pixel);
