Iris
* Clarified FRootObjectReplicationParams::bNeedsPreUpdate to mention it's actor only at the moment

#rnx
#Synced-CL 48203440

diff --git a/Engine/Source/Runtime/Engine/Private/Net/Iris/ReplicationSystem/EngineReplicationBridge.cpp b/Engine/Source/Runtime/Engine/Private/Net/Iris/ReplicationSystem/EngineReplicationBridge.cpp
--- a/Engine/Source/Runtime/Engine/Private/Net/Iris/ReplicationSystem/EngineReplicationBridge.cpp
+++ b/Engine/Source/Runtime/Engine/Private/Net/Iris/ReplicationSystem/EngineReplicationBridge.cpp
@@ -95,18 +95,16 @@ bool IsActorValidForIrisReplication(const AActor* Actor)
 
 void ActorReplicationBridgePreUpdateFunction(TArrayView<UObject*> Instances, const UObjectReplicationBridge* Bridge)
 {
-	if (const UEngineReplicationBridge* EngineReplicationBridge = Cast<const UEngineReplicationBridge>(Bridge))
+	const UEngineReplicationBridge* EngineReplicationBridge = CastChecked<const UEngineReplicationBridge>(Bridge);
+	if (!EngineReplicationBridge->AllowExecutePreReplication())
 	{
-		if (!EngineReplicationBridge->AllowExecutePreReplication())
-		{
-			return;
-		}
+		return;
 	}
 	
-	UNetDriver* NetDriver = CastChecked<const UEngineReplicationBridge>(Bridge)->GetNetDriver();
+	UNetDriver* NetDriver = EngineReplicationBridge->GetNetDriver();
 	for (UObject* Instance : Instances)
 	{
-		AActor* Actor = Cast<AActor>(Instance);
+		AActor* Actor = CastChecked<AActor>(Instance);
 		if (IsActorValidForIrisReplication(Actor))
 		{
 			Actor->CallPreReplication(NetDriver);
diff --git a/Engine/Source/Runtime/Net/Iris/Public/Iris/ReplicationSystem/ObjectReplicationBridge.h b/Engine/Source/Runtime/Net/Iris/Public/Iris/ReplicationSystem/ObjectReplicationBridge.h
--- a/Engine/Source/Runtime/Net/Iris/Public/Iris/ReplicationSystem/ObjectReplicationBridge.h
+++ b/Engine/Source/Runtime/Net/Iris/Public/Iris/ReplicationSystem/ObjectReplicationBridge.h
@@ -253,10 +253,10 @@ public:
 
 	struct FRootObjectReplicationParams
 	{
-		/** When true it means the object wants to receive a PreUpdate callback just before it gets polled. */
+		/** *Actor only* When true the actor will receive a PreReplication callback just before it gets polled. */
 		bool bNeedsPreUpdate = false;
 
-		/** When true the object has a dynamic world location and we should ask the bridge to update its current location everytime it is polled. */
+		/** When true the object has a dynamic world location and we should ask the bridge to update its current location every time it is polled. */
 		bool bNeedsWorldLocationUpdate = false;
 
 		/** Whether the object is dormant or not */
