PoseAsset: Introduce intermediary interfaces for decoupling GetAnimationPose from extractcontext / bonecontainer

#rnx

#Synced-CL 48347766

diff --git a/Engine/Source/Runtime/Engine/Classes/Animation/PoseAsset.h b/Engine/Source/Runtime/Engine/Classes/Animation/PoseAsset.h
--- a/Engine/Source/Runtime/Engine/Classes/Animation/PoseAsset.h
+++ b/Engine/Source/Runtime/Engine/Classes/Animation/PoseAsset.h
@@ -20,6 +20,7 @@ class USkeletalMeshComponent;
 struct FLiveLinkCurveElement;
 struct FReferenceSkeleton;
 class FPoseAssetDetails;
+class FAnimationRuntimeRetargetBoneTransformContext;
 
 /** 
  * Pose data 
@@ -90,6 +91,121 @@ struct FPoseAssetInfluenceForPose
 	int32 BoneTransformIndex = INDEX_NONE;
 };
 
+/**
+ * Intermediary Interface class to be implemented to extract information from a pose asset.
+ * This is used for UPoseAsset::GetAnimationPose, for example.
+ */
+class FPoseAssetExtractionContext
+{
+public:
+	FPoseAssetExtractionContext() {}
+	virtual ~FPoseAssetExtractionContext() {}
+
+	/**
+	 * The number of pose curves available in this context
+	 * @return The number of pose curves available in this context
+	 */
+	virtual int32 NumPoses() const = 0;
+
+	/**
+	 * Retrieves the name of the pose given an index.
+	 * @param InPoseIndexInContext The index of the pose to return the name of
+	 * @return The name of the pose given the index
+	 */
+	virtual const FName& GetPoseName(const int32& InPoseIndexInContext) const = 0;
+
+	/**
+	 * Retrieves the index of the pose in the pose asset given the index of the pose within the context.
+	 * @param InPoseIndexInContext The index of the pose within the context
+	 * @return The index of the pose in the pose asset
+	 */
+	virtual int32 GetPoseIndexInPoseAsset(const int32& InPoseIndexInContext) const = 0;
+
+	/**
+	 * Retrieves the weight / value of a pose in the context
+	 * @param InPoseIndexInContext The index of the pose in the context
+	 * @return The weight / value of a pose in the context
+	 */
+	virtual float GetPoseWeight(const int32& InPoseIndexInContext) const = 0;
+
+	/**
+	 * Retrieves the pose stored in the context as a FPoseCurve
+	 * @param InPoseIndexInContext The index of the pose in the context
+	 * @return The FPoseCurve representing the stored pose in the context
+	 */
+	virtual FPoseCurve GetPoseCurve(const int32& InPoseIndexInContext) const
+	{
+		return FPoseCurve(
+			GetPoseIndexInPoseAsset(InPoseIndexInContext),
+			GetPoseName(InPoseIndexInContext),
+			GetPoseWeight(InPoseIndexInContext)
+		);
+	}
+
+	/**
+	 * Returns the number of curves stored in the context
+	 * @return The number of curves stored in the context
+	 */
+	virtual int32 NumCurves() const = 0;
+
+	/**
+	 * Blends an existing blended curve onto this context given a weight
+	 * @param InCurve The curve to blend onto this context
+	 * @param InWeight The weight to use when blending
+	 */
+	virtual void BlendPoseCurve(const FBlendedCurve& InCurve, float InWeight) = 0;
+
+	/**
+	 * Returns the target skeleton this context is using
+	 * @param bEvenIfUnreachable If true the skeleton will be returned using a most expensive lookup if the cached weakptr is invalid
+	 * @return The target skeleton this context is using
+	 */
+	virtual const USkeleton* GetTargetSkeleton(bool bEvenIfUnreachable = false) const = 0;
+
+	/**
+	 * Returns true if a given bone is required to be computed.
+	 * @param InOutputBoneIndex The bone index within the context
+	 * @return True if a given bone is required to be computed.
+	 */
+	virtual bool IsOutputBoneRequired(const int32& InOutputBoneIndex) const = 0;
+
+	/**
+	 * Returns the bone index within the output pose of this context given the skeleton pose bone index
+	 * @param InSkeletonBoneIndex The index of the bone in the skeleton's pose
+	 * @return The bone index within the output pose
+	 */
+	virtual int32 GetOutputPoseBoneIndexFromSkeletonPoseBoneIndex(const int32& InSkeletonBoneIndex) const = 0;
+
+	/**
+	 * Returns the curve filter to apply when blending curves (or nullptr)
+	 * @return The curve filter to apply when blending curves (or nullptr)
+	 */
+	virtual const UE::Anim::FCurveFilter* GetCurveFilter() const = 0;
+
+	/**
+	 * Resets the output pose to additive identity (with zero scale)
+	 */
+	virtual void ResetPoseToAdditiveIdentity() = 0;
+
+	/**
+	 * Resets the output pose to the reference pose
+	 */
+	virtual void ResetPoseToRefPose() = 0;
+
+	/**
+	 * Returns the transform within the output pose given a bone index
+	 * @param InOutputBoneIndex The index of the bone within this context
+	 * @return The transform within the output pose given a bone index
+	 */
+	virtual FTransform& GetOutputTransform(const int32& InOutputBoneIndex) = 0;
+
+	/**
+	 * Returns the context to use for FAnimationRuntime::RetargetBoneTransform
+	 * @return The context to use for FAnimationRuntime::RetargetBoneTransform
+	 */
+	virtual const FAnimationRuntimeRetargetBoneTransformContext* GetRetargetBoneTransformContext() const = 0;
+};
+
 /**
 * Pose data container
 * 
@@ -150,6 +266,7 @@ private:
 	bool IsValid() const { return PoseFNames.Num() == Poses.Num() && Tracks.Num() == TrackBoneIndices.Num(); }
 	ENGINE_API void GetPoseCurve(const FPoseData* PoseData, FBlendedCurve& OutCurve) const;
 	ENGINE_API void BlendPoseCurve(const FPoseData* PoseData, FBlendedCurve& OutCurve, float Weight) const;
+	ENGINE_API void BlendPoseCurve(const FPoseData* PoseData, FPoseAssetExtractionContext* InOutContext, float Weight) const;
 
 	// we have to delete tracks if skeleton has modified
 	// usually this may not be issue since once cooked, it should match
@@ -189,7 +306,6 @@ struct TStructOpsTypeTraits<FPoseDataContainer> : public TStructOpsTypeTraitsBas
 	};
 };
 
-
 /**
  * Pose Asset that can be blended by weight of curves 
  */
@@ -248,6 +364,7 @@ public:
 	UE_DEPRECATED(4.26, "Use GetAnimationPose with other signature")
 	ENGINE_API bool GetAnimationPose(struct FCompactPose& OutPose, FBlendedCurve& OutCurve, const FAnimExtractContext& ExtractionContext) const;
 	ENGINE_API bool GetAnimationPose(struct FAnimationPoseData& OutAnimationPoseData, const FAnimExtractContext& ExtractionContext, const bool bIgnoreSourceAnimationCurves = false, const bool bAllowNegativeWeights = false) const;
+	ENGINE_API bool GetAnimationPose(FPoseAssetExtractionContext* InOutContext, const bool bIgnoreSourceAnimationCurves = false, const bool bAllowNegativeWeights = false) const;
 
 	UE_DEPRECATED(4.26, "Use GetBaseAnimationPose with other signature")
 	ENGINE_API void GetBaseAnimationPose(struct FCompactPose& OutPose, FBlendedCurve& OutCurve) const;
diff --git a/Engine/Source/Runtime/Engine/Private/Animation/AnimationRuntime.cpp b/Engine/Source/Runtime/Engine/Private/Animation/AnimationRuntime.cpp
--- a/Engine/Source/Runtime/Engine/Private/Animation/AnimationRuntime.cpp
+++ b/Engine/Source/Runtime/Engine/Private/Animation/AnimationRuntime.cpp
@@ -2892,15 +2892,22 @@ EBoneTranslationRetargetingMode::Type FAnimationRuntime::GetBoneTranslationRetar
 
 void FAnimationRuntime::RetargetBoneTransform(const USkeleton* SourceSkeleton, const FName& SourceName, const TArray<FTransform>& RetargetTransforms, FTransform& BoneTransform, const int32 SkeletonBoneIndex, const FCompactPoseBoneIndex& BoneIndex, const FBoneContainer& RequiredBones, const bool bIsBakedAdditive)
 {
+	FBoneContainerRetargetBoneTransformContext Context(RequiredBones);
+	RetargetBoneTransform(SourceSkeleton, SourceName, RetargetTransforms, BoneTransform, SkeletonBoneIndex, BoneIndex, &Context, bIsBakedAdditive);
+}
+
+void FAnimationRuntime::RetargetBoneTransform(const USkeleton* SourceSkeleton, const FName& SourceName, const TArray<FTransform>& RetargetTransforms, FTransform& BoneTransform, const int32 SkeletonBoneIndex, const FCompactPoseBoneIndex& BoneIndex, const FAnimationRuntimeRetargetBoneTransformContext* InContext, const bool bIsBakedAdditive)
+{
+	check(InContext);
 	check(!RetargetTransforms.IsEmpty());
 	if (SourceSkeleton)
 	{
 		// Retrieve skeleton, even if it is unreachable (but not GC-ed yet)
 		constexpr bool bEvenIfUnreachable = true; 
-		const USkeleton* TargetSkeleton = RequiredBones.GetSkeletonAsset(bEvenIfUnreachable);
+		const USkeleton* TargetSkeleton = InContext->GetSkeleton(bEvenIfUnreachable);
 		const FSkeletonRemapping& SkeletonRemapping = UE::Anim::FSkeletonRemappingRegistry::Get().GetRemapping(SourceSkeleton, TargetSkeleton);
 
-		const int32 TargetSkeletonBoneIndex = RequiredBones.GetSkeletonIndex(BoneIndex);
+		const int32 TargetSkeletonBoneIndex = InContext->GetSkeletonBoneIndex(BoneIndex.GetInt());
 		int32 SourceSkeletonBoneIndex = SkeletonBoneIndex;
 
 		// Apply compatible skeleton remapping if required
@@ -2922,7 +2929,7 @@ void FAnimationRuntime::RetargetBoneTransform(const USkeleton* SourceSkeleton, c
 			TargetSkeletonBoneIndex,
 			SourceSkeleton,
 			TargetSkeleton,
-			RequiredBones.GetDisableRetargeting());
+			InContext->IsRetargetingDisabled());
 			
 		switch (RetargetMode)
 		{
@@ -2935,7 +2942,7 @@ void FAnimationRuntime::RetargetBoneTransform(const USkeleton* SourceSkeleton, c
 					const float SourceTranslationLength = SkeletonRefPoseArray[SourceSkeletonBoneIndex].GetTranslation().Size();
 					if (SourceTranslationLength > UE_KINDA_SMALL_NUMBER)
 					{
-						const float TargetTranslationLength = RequiredBones.GetRefPoseTransform(BoneIndex).GetTranslation().Size();
+						const float TargetTranslationLength = InContext->GetRefPoseTransform(BoneIndex.GetInt()).GetTranslation().Size();
 						BoneTransform.ScaleTranslation(TargetTranslationLength / SourceTranslationLength);
 					}
 				}
@@ -2944,7 +2951,7 @@ void FAnimationRuntime::RetargetBoneTransform(const USkeleton* SourceSkeleton, c
 
 			case EBoneTranslationRetargetingMode::Skeleton:
 			{
-				BoneTransform.SetTranslation(bIsBakedAdditive ? FVector::ZeroVector : RequiredBones.GetRefPoseTransform(BoneIndex).GetTranslation());
+				BoneTransform.SetTranslation(bIsBakedAdditive ? FVector::ZeroVector : InContext->GetRefPoseTransform(BoneIndex.GetInt()).GetTranslation());
 				break;
 			}
 
@@ -2955,7 +2962,7 @@ void FAnimationRuntime::RetargetBoneTransform(const USkeleton* SourceSkeleton, c
 				if (!bIsBakedAdditive)
 				{
 					const TArray<FTransform>& AuthoredOnRefSkeleton = RetargetTransforms;
-					const FTransform& RefPoseTransform = RequiredBones.GetRefPoseTransform(BoneIndex);
+					const FTransform& RefPoseTransform = InContext->GetRefPoseTransform(BoneIndex.GetInt());
 
 					// Remap the base pose onto the target skeleton so that we are working entirely in target space
 					if(AuthoredOnRefSkeleton.IsValidIndex(SourceSkeletonBoneIndex))
@@ -2994,12 +3001,12 @@ void FAnimationRuntime::RetargetBoneTransform(const USkeleton* SourceSkeleton, c
 			{
 				if (!bIsBakedAdditive)
 				{
-					const FRetargetSourceCachedData& RetargetSourceCachedData = RequiredBones.GetRetargetSourceCachedData(SourceName, SkeletonRemapping, RetargetTransforms);
+					const FRetargetSourceCachedData& RetargetSourceCachedData = InContext->GetRetargetSourceCachedData(SourceName, SkeletonRemapping, RetargetTransforms);
 					const TArray<FOrientAndScaleRetargetingCachedData>& OrientAndScaleDataArray = RetargetSourceCachedData.OrientAndScaleData;
 					const TArray<int32>& CompactPoseIndexToOrientAndScaleIndex = RetargetSourceCachedData.CompactPoseIndexToOrientAndScaleIndex;
 
 					// If we have any cached retargeting data.
-					if ((OrientAndScaleDataArray.Num() > 0) && (CompactPoseIndexToOrientAndScaleIndex.Num() == RequiredBones.GetCompactPoseNumBones()))
+					if ((OrientAndScaleDataArray.Num() > 0) && (CompactPoseIndexToOrientAndScaleIndex.Num() == InContext->GetNumBonesInContext()))
 					{
 						const int32 OrientAndScaleIndex = CompactPoseIndexToOrientAndScaleIndex[BoneIndex.GetInt()];
 						if (OrientAndScaleIndex != INDEX_NONE)
diff --git a/Engine/Source/Runtime/Engine/Private/Animation/PoseAsset.cpp b/Engine/Source/Runtime/Engine/Private/Animation/PoseAsset.cpp
--- a/Engine/Source/Runtime/Engine/Private/Animation/PoseAsset.cpp
+++ b/Engine/Source/Runtime/Engine/Private/Animation/PoseAsset.cpp
@@ -118,6 +118,22 @@ void FPoseDataContainer::BlendPoseCurve(const FPoseData* PoseData, FBlendedCurve
 	}
 }
 
+void FPoseDataContainer::BlendPoseCurve(const FPoseData* PoseData, FPoseAssetExtractionContext* InOutContext, float Weight) const
+{
+	if (PoseData)
+	{
+		const TArray<float>& CurveValues = PoseData->CurveData;
+		checkSlow(CurveValues.Num() == Curves.Num());
+
+		FBlendedCurve Curve;
+		Curve.SetFilter(InOutContext->GetCurveFilter());
+
+		GetPoseCurve(PoseData, Curve);
+
+		InOutContext->BlendPoseCurve(Curve, Weight);
+	}
+}
+
 FPoseData* FPoseDataContainer::FindPoseData(FName PoseName)
 {
 	int32 PoseIndex = PoseFNames.Find(PoseName);
@@ -604,37 +620,134 @@ bool UPoseAsset::GetAnimationPose(struct FCompactPose& OutPose, FBlendedCurve& O
 	return GetAnimationPose(OutPoseData, ExtractionContext);
 }
 
+class FPoseAssetAnimBPExtractionContext : public FPoseAssetExtractionContext
+{
+public:
+
+	FPoseAssetAnimBPExtractionContext(const FAnimExtractContext& InExtractionContext, struct FAnimationPoseData& OutAnimationPoseData)
+		: FPoseAssetExtractionContext()
+		, ExtractionContext(InExtractionContext)
+		, AnimationPoseData(OutAnimationPoseData)
+		, RetargetBoneTransformContext(OutAnimationPoseData.GetPose().GetBoneContainer())
+	{
+	}
+
+	virtual int32 NumPoses() const override
+	{
+		return ExtractionContext.PoseCurves.Num();
+	}
+	
+	virtual const FName& GetPoseName(const int32& InPoseIndex) const override
+	{
+		return ExtractionContext.PoseCurves[InPoseIndex].Name;
+	}
+
+	virtual int32 GetPoseIndexInPoseAsset(const int32& InPoseIndex) const override
+	{
+		return ExtractionContext.PoseCurves[InPoseIndex].PoseIndex;
+	}
+
+	virtual float GetPoseWeight(const int32& InPoseIndex) const override
+	{
+		return ExtractionContext.PoseCurves[InPoseIndex].Value;
+	}
+
+	virtual const USkeleton* GetTargetSkeleton(bool bEvenIfUnreachable) const override
+	{
+		return AnimationPoseData.GetPose().GetBoneContainer().GetSkeletonAsset(bEvenIfUnreachable);
+	}
+
+	virtual bool IsOutputBoneRequired(const int32& InOutputBoneIndex) const override
+	{
+		return ExtractionContext.IsBoneRequired(InOutputBoneIndex);
+	}
+
+	virtual int32 GetOutputPoseBoneIndexFromSkeletonPoseBoneIndex(const int32& InSkeletonBoneIndex) const override
+	{
+		return AnimationPoseData.GetPose().GetBoneContainer().GetCompactPoseIndexFromSkeletonPoseIndex(FSkeletonPoseBoneIndex(InSkeletonBoneIndex)).GetInt();
+	}
+
+	virtual const UE::Anim::FCurveFilter* GetCurveFilter() const override
+	{
+		return AnimationPoseData.GetCurve().GetFilter();
+	}
+
+	virtual void ResetPoseToAdditiveIdentity() override
+	{
+		AnimationPoseData.GetPose().ResetToAdditiveIdentity();
+	}
+
+	virtual void ResetPoseToRefPose() override
+	{
+		AnimationPoseData.GetPose().ResetToRefPose();
+	}
+
+	virtual FTransform& GetOutputTransform(const int32& InOutputBoneIndex) override
+	{
+		return AnimationPoseData.GetPose()[FCompactPoseBoneIndex(InOutputBoneIndex)];
+	}
+
+	virtual const FAnimationRuntimeRetargetBoneTransformContext* GetRetargetBoneTransformContext() const override
+	{
+		return &RetargetBoneTransformContext;
+	}
+
+	virtual int32 NumCurves() const override
+	{
+		return AnimationPoseData.GetCurve().Num();
+	}
+
+	virtual void BlendPoseCurve(const FBlendedCurve& InCurve, float InWeight) override
+	{
+		UE::Anim::FNamedValueArrayUtils::Union(AnimationPoseData.GetCurve(), InCurve,
+			[InWeight](UE::Anim::FCurveElement& InOutElement, const UE::Anim::FCurveElement& InElement, UE::Anim::ENamedValueUnionFlags InFlags)
+			{
+				InOutElement.Value = InElement.Value * InWeight + InOutElement.Value;
+				InOutElement.Flags |= InElement.Flags;
+			});
+	}
+
+private:
+
+	const FAnimExtractContext& ExtractionContext;
+	struct FAnimationPoseData& AnimationPoseData;
+	FBoneContainerRetargetBoneTransformContext RetargetBoneTransformContext;
+};
+
 bool UPoseAsset::GetAnimationPose(struct FAnimationPoseData& OutAnimationPoseData, const FAnimExtractContext& ExtractionContext, const bool bIgnoreSourceAnimationCurves, const bool bAllowNegativeWeights) const
+{
+	FPoseAssetAnimBPExtractionContext PoseAssetExtractionContext(ExtractionContext, OutAnimationPoseData);
+	return GetAnimationPose(&PoseAssetExtractionContext, bIgnoreSourceAnimationCurves, bAllowNegativeWeights);
+}
+
+bool UPoseAsset::GetAnimationPose(FPoseAssetExtractionContext* InOutContext, const bool bIgnoreSourceAnimationCurves, const bool bAllowNegativeWeights) const
 {
 	ANIM_MT_SCOPE_CYCLE_COUNTER(PoseAssetGetAnimationPose, !IsInGameThread());
+	check(InOutContext);
 
 	// if we have any pose curve
-	if (ExtractionContext.PoseCurves.Num() > 0)
+	if (InOutContext->NumPoses())
 	{
-		FCompactPose& OutPose = OutAnimationPoseData.GetPose();
-		FBlendedCurve& OutCurve = OutAnimationPoseData.GetCurve();
-
-		const FBoneContainer& RequiredBones = OutPose.GetBoneContainer();
 		USkeleton* MySkeleton = GetSkeleton();
 		
 		check(PoseContainer.IsValid());
 
 		if (bAdditivePose)
 		{
-			OutPose.ResetToAdditiveIdentity();
+			InOutContext->ResetPoseToAdditiveIdentity();
 		}
 		else
 		{
-			OutPose.ResetToRefPose();
+			InOutContext->ResetPoseToRefPose();
 		}
 
 		const int32 TrackNum = PoseContainer.Tracks.Num();
 
-		const FSkeletonRemapping& SkeletonRemapping = UE::Anim::FSkeletonRemappingRegistry::Get().GetRemapping(GetSkeleton(), RequiredBones.GetSkeletonAsset());
+		const FSkeletonRemapping& SkeletonRemapping = UE::Anim::FSkeletonRemappingRegistry::Get().GetRemapping(GetSkeleton(), InOutContext->GetTargetSkeleton());
 		// Single pose optimized evaluation path - explicitly used by PoseByName Animation Node
-		if (ExtractionContext.PoseCurves.Num() == 1)
+		if (InOutContext->NumPoses() == 1)
 		{
-			const FPoseCurve& Curve = ExtractionContext.PoseCurves[0];
+			FPoseCurve Curve = InOutContext->GetPoseCurve(0);
 			const int32 PoseIndex = Curve.PoseIndex; 
 			if (ensure(PoseIndex != INDEX_NONE))
 			{
@@ -649,17 +762,17 @@ bool UPoseAsset::GetAnimationPose(struct FAnimationPoseData& OutAnimationPoseDat
 					if (!bIgnoreSourceAnimationCurves)
 					{
 						// Blend curve
-						PoseContainer.BlendPoseCurve(&Pose, OutCurve, Weight);
+						PoseContainer.BlendPoseCurve(&Pose, InOutContext, Weight);
 					}
 
 					// Per-track (bone) transform
 					for (int32 TrackIndex = 0; TrackIndex < TrackNum; ++TrackIndex)
 					{
 						const FSkeletonPoseBoneIndex SkeletonBoneIndex = FSkeletonPoseBoneIndex(SkeletonRemapping.IsValid() ? SkeletonRemapping.GetTargetSkeletonBoneIndex(PoseContainer.TrackBoneIndices[TrackIndex]) : PoseContainer.TrackBoneIndices[TrackIndex]);
-						const FCompactPoseBoneIndex CompactIndex = RequiredBones.GetCompactPoseIndexFromSkeletonPoseIndex(SkeletonBoneIndex);
-						
+						const FCompactPoseBoneIndex CompactIndex(InOutContext->GetOutputPoseBoneIndexFromSkeletonPoseBoneIndex(SkeletonBoneIndex.GetInt()));
+
 						// If bone index is invalid, or not required for the pose - skip
-						if (!CompactIndex.IsValid() || !ExtractionContext.IsBoneRequired(CompactIndex.GetInt()))
+						if (!CompactIndex.IsValid() || !InOutContext->IsOutputBoneRequired(CompactIndex.GetInt()))
 						{
 							continue;
 						}
@@ -673,7 +786,7 @@ bool UPoseAsset::GetAnimationPose(struct FAnimationPoseData& OutAnimationPoseDat
 
 						if (InfluenceIndex != INDEX_NONE)
 						{
-							FTransform& OutBoneTransform =  OutPose[CompactIndex];
+							FTransform& OutBoneTransform = InOutContext->GetOutputTransform(CompactIndex.GetInt());
 							const FPoseAssetInfluence& Influence = PoseInfluences[InfluenceIndex];
 							const int32& BonePoseIndex = Influence.BoneTransformIndex;
 
@@ -696,7 +809,7 @@ bool UPoseAsset::GetAnimationPose(struct FAnimationPoseData& OutAnimationPoseDat
 							}
 
 							// Retarget the bone transform
-							FAnimationRuntime::RetargetBoneTransform(MySkeleton, GetRetargetTransformsSourceName(), GetRetargetTransforms(), OutBoneTransform, SkeletonBoneIndex.GetInt(), CompactIndex, RequiredBones, bAdditivePose);
+							FAnimationRuntime::RetargetBoneTransform(MySkeleton, GetRetargetTransformsSourceName(), GetRetargetTransforms(), OutBoneTransform, SkeletonBoneIndex.GetInt(), CompactIndex, InOutContext->GetRetargetBoneTransformContext(), bAdditivePose);
 							OutBoneTransform.NormalizeRotation();
 						}
 					}					
@@ -717,7 +830,7 @@ bool UPoseAsset::GetAnimationPose(struct FAnimationPoseData& OutAnimationPoseDat
 			for(int32 TrackIndex = 0; TrackIndex < TrackNum; ++TrackIndex)
 			{
 				const FSkeletonPoseBoneIndex SkeletonBoneIndex = FSkeletonPoseBoneIndex(SkeletonRemapping.IsValid() ? SkeletonRemapping.GetTargetSkeletonBoneIndex(PoseContainer.TrackBoneIndices[TrackIndex]) : PoseContainer.TrackBoneIndices[TrackIndex]);
-				const FCompactPoseBoneIndex CompactIndex = RequiredBones.GetCompactPoseIndexFromSkeletonPoseIndex(SkeletonBoneIndex);
+				const FCompactPoseBoneIndex CompactIndex(InOutContext->GetOutputPoseBoneIndexFromSkeletonPoseBoneIndex(SkeletonBoneIndex.GetInt()));
 
 				// we add even if it's invalid because we want it to match with track index
 				BoneIndices[TrackIndex].SkeletonBoneIndex = SkeletonBoneIndex.GetInt();
@@ -742,10 +855,10 @@ bool UPoseAsset::GetAnimationPose(struct FAnimationPoseData& OutAnimationPoseDat
 			float TotalWeight = 0.f;
 			// we iterate through to see if we have that corresponding pose
 
-			const int32 NumPoseCurves = ExtractionContext.PoseCurves.Num();
+			const int32 NumPoseCurves = InOutContext->NumPoses();
 			for (int32 CurveIndex = 0; CurveIndex < NumPoseCurves; ++CurveIndex)
 			{
-				const FPoseCurve& Curve = ExtractionContext.PoseCurves[CurveIndex];
+				const FPoseCurve Curve = InOutContext->GetPoseCurve(CurveIndex);
 				const int32 PoseIndex = Curve.PoseIndex; 
 				if (ensure(PoseIndex != INDEX_NONE))
 				{
@@ -757,7 +870,7 @@ bool UPoseAsset::GetAnimationPose(struct FAnimationPoseData& OutAnimationPoseDat
 					const float WeightToCheck = bAllowNegativeWeights ? FMath::Abs(Value) : Value;
 					if ((!bAdditivePose || PoseIndex != BasePoseIndex) && FAnimationRuntime::HasWeight(WeightToCheck))
 					{
-						TotalWeight += Value;
+						TotalWeight += WeightToCheck;
 
 						// Set pose weight and bit, and add weighted pose index
 						PoseWeights[PoseIndex] = Value;
@@ -782,7 +895,7 @@ bool UPoseAsset::GetAnimationPose(struct FAnimationPoseData& OutAnimationPoseDat
 						{
 							// Do curve blending inline as we are looping over weights anyway
 							const FPoseData& Pose = PoseContainer.Poses[WeightedPoseIndex];
-							PoseContainer.BlendPoseCurve(&Pose, OutCurve, PoseWeight);
+							PoseContainer.BlendPoseCurve(&Pose, InOutContext, PoseWeight);
 						}
 					}
 				}
@@ -798,7 +911,7 @@ bool UPoseAsset::GetAnimationPose(struct FAnimationPoseData& OutAnimationPoseDat
 						{
 							const FPoseData& Pose = PoseContainer.Poses[WeightedPoseIndex];
 							const float& Weight = PoseWeights[WeightedPoseIndex];
-							PoseContainer.BlendPoseCurve(&Pose, OutCurve, Weight);
+							PoseContainer.BlendPoseCurve(&Pose, InOutContext, Weight);
 						}
 					}
 				}
@@ -810,7 +923,7 @@ bool UPoseAsset::GetAnimationPose(struct FAnimationPoseData& OutAnimationPoseDat
 					const FCompactPoseBoneIndex& CompactIndex = BoneIndices[TrackIndex].CompactBoneIndex;
 
 					// If bone index is invalid, or not required for the pose - skip
-					if (CompactIndex == INDEX_NONE || !ExtractionContext.IsBoneRequired(CompactIndex.GetInt()))
+					if (CompactIndex == INDEX_NONE || !InOutContext->IsOutputBoneRequired(CompactIndex.GetInt()))
 					{
 						continue;
 					}
@@ -820,7 +933,7 @@ bool UPoseAsset::GetAnimationPose(struct FAnimationPoseData& OutAnimationPoseDat
 					// When additive, or for any bone that has no pose influences. Set transform to input.
 					if (bAdditivePose || PoseInfluences.Num() == 0)
 					{
-						OutBoneTransform = OutPose[CompactIndex];
+						OutBoneTransform = InOutContext->GetOutputTransform(CompactIndex.GetInt());
 					}
 
 					const int32 NumInfluences = PoseInfluences.Num();
@@ -867,19 +980,19 @@ bool UPoseAsset::GetAnimationPose(struct FAnimationPoseData& OutAnimationPoseDat
 						// In case no influencing poses had any weight, set transform to input
 						if(!FAnimWeight::IsRelevant(TotalLocalWeight))
 						{
-							OutBoneTransform = OutPose[CompactIndex];
+							OutBoneTransform = InOutContext->GetOutputTransform(CompactIndex.GetInt());
 						}
 						else if (!FAnimWeight::IsFullWeight(TotalLocalWeight) && !bAdditivePose)
 						{
-							OutBoneTransform.AccumulateWithShortestRotation(OutPose[CompactIndex], ScalarRegister(1.f - TotalLocalWeight));
+							OutBoneTransform.AccumulateWithShortestRotation(InOutContext->GetOutputTransform(CompactIndex.GetInt()), ScalarRegister(1.f - TotalLocalWeight));
 						}
 					}
 
 					// Retarget the blended transform, and copy to output pose
-					FAnimationRuntime::RetargetBoneTransform(MySkeleton, GetRetargetTransformsSourceName(), GetRetargetTransforms(), OutBoneTransform,  BoneIndices[TrackIndex].SkeletonBoneIndex, CompactIndex, RequiredBones, bAdditivePose);
+					FAnimationRuntime::RetargetBoneTransform(MySkeleton, GetRetargetTransformsSourceName(), GetRetargetTransforms(), OutBoneTransform,  BoneIndices[TrackIndex].SkeletonBoneIndex, CompactIndex, InOutContext->GetRetargetBoneTransformContext(), bAdditivePose);
 
-					OutPose[CompactIndex] = OutBoneTransform;
-					OutPose[CompactIndex].NormalizeRotation();
+					InOutContext->GetOutputTransform(CompactIndex.GetInt()) = OutBoneTransform;
+					InOutContext->GetOutputTransform(CompactIndex.GetInt()).NormalizeRotation();
 				}
 
 				return true;
@@ -2201,21 +2314,37 @@ bool UPoseAsset::ConvertToFullPose(const int32 PoseIndex, TArray<FTransform>& Ou
 		return false;
 	}
 
-	if (PoseIndex == BasePoseIndex || BasePoseIndex == INDEX_NONE || !bAdditivePose)
+	if (PoseIndex == BasePoseIndex || !bAdditivePose)
 	{
 		OutFullPose = PoseContainer.Poses[PoseIndex].LocalSpacePose;
 	}
+	else if (BasePoseIndex == INDEX_NONE)
+	{
+		const TArray<FTransform>& AdditivePose = PoseContainer.Poses[PoseIndex].LocalSpacePose;
+		const TArray<FTransform>& RetargetTransforms = GetRetargetTransforms();
+
+		OutFullPose.Reset();
+		OutFullPose.Reserve(OutInfluences.Num());
+
+		for (const FPoseAssetInfluenceForPose& Influence : OutInfluences)
+		{
+			const FTransform& BaseTransform = PoseContainer.GetDefaultTransform(Influence.TrackIndex, RetargetTransforms);
+			FTransform& BoneTransform = OutFullPose.Add_GetRef(BaseTransform);
+			BlendFromIdentityAndAccumulateAdditively_Custom(BoneTransform, AdditivePose[Influence.BoneTransformIndex], 1);
+		}
+	}
 	else
 	{
-		const TArray<FTransform> BasePose = PoseContainer.Poses[BasePoseIndex].LocalSpacePose;
-		const TArray<FTransform> AdditivePose = PoseContainer.Poses[PoseIndex].LocalSpacePose;
+		const TArray<FTransform>& AdditivePose = PoseContainer.Poses[PoseIndex].LocalSpacePose;
+		const TArray<FTransform>& BasePose = PoseContainer.Poses[BasePoseIndex].LocalSpacePose;
+
 		OutFullPose.Reset();
 		OutFullPose.Reserve(OutInfluences.Num());
 
 		for (const FPoseAssetInfluenceForPose& Influence : OutInfluences)
 		{
 			FTransform& BoneTransform = OutFullPose.Add_GetRef(BasePose[Influence.TrackIndex]);
-			BoneTransform.AccumulateWithAdditiveScale(AdditivePose[Influence.BoneTransformIndex], ScalarRegister(1));
+			BlendFromIdentityAndAccumulateAdditively_Custom(BoneTransform, AdditivePose[Influence.BoneTransformIndex], 1);
 		}
 	}
 
@@ -2230,8 +2359,8 @@ bool UPoseAsset::ConvertToFullPose(const int32 PoseIndex, TArray<float>& OutFull
 	}
 	else
 	{
-		const TArray<float> BaseCurves = PoseContainer.Poses[BasePoseIndex].CurveData;
-		const TArray<float> AdditiveCurves = PoseContainer.Poses[PoseIndex].CurveData;
+		const TArray<float>& BaseCurves = PoseContainer.Poses[BasePoseIndex].CurveData;
+		const TArray<float>& AdditiveCurves = PoseContainer.Poses[PoseIndex].CurveData;
 		check(BaseCurves.Num() == AdditiveCurves.Num());
 		OutFullCurves.Reset();
 		OutFullCurves.Reserve(BaseCurves.Num());
diff --git a/Engine/Source/Runtime/Engine/Public/AnimationRuntime.h b/Engine/Source/Runtime/Engine/Public/AnimationRuntime.h
--- a/Engine/Source/Runtime/Engine/Public/AnimationRuntime.h
+++ b/Engine/Source/Runtime/Engine/Public/AnimationRuntime.h
@@ -99,6 +99,97 @@ public:
 	virtual int32 GetPerBoneInterpolationIndex(const FSkeletonPoseBoneIndex InSkeletonBoneIndex, const USkeleton* TargetSkeleton, const IInterpolationIndexProvider::FPerBoneInterpolationData* Data) const = 0;
 };
 
+/**
+ * Intermediary Interface class to be implemented to access FAnimationRuntime::RetargetBoneTransform.
+ * This is specialized in FBoneContainerRetargetBoneTransformContext, for example.
+ */
+class FAnimationRuntimeRetargetBoneTransformContext
+{
+public:
+
+	FAnimationRuntimeRetargetBoneTransformContext() {}
+	virtual ~FAnimationRuntimeRetargetBoneTransformContext() {}
+
+	/** Returns the skeleton of this context. Typically maps to the desired target skeleton for the retarget */
+	virtual const USkeleton* GetSkeleton(bool bEvenIfUnreachable = false) const = 0;
+
+	/**
+	 * Returns the skeleton bone index for a context bone index.
+	 * Context bone index in this case is index within the wrapped context - for example the pose.
+	 */
+	virtual int32 GetSkeletonBoneIndex(const int32& InContextBoneIndex) const = 0;
+
+	/**
+	 * Returns the number of bones within this context - for example the number of bones within a pose
+	 */
+	virtual int32 GetNumBonesInContext() const = 0;
+	
+	/**
+	 * Returns true if retargeting for this context is disabled.
+	 */
+	virtual bool IsRetargetingDisabled() const = 0;
+
+	/**
+	 * Returns the reference pose transform for a given context bone index. The context bone index refers 
+	 * to the index of the bone within the context - for example a wrapped pose.
+	 */
+	virtual const FTransform& GetRefPoseTransform(const int32& InContextBoneIndex) const = 0;
+
+	/**
+	 * Returns the cached source data for retargeting within this context.
+	 * This is current implemented only against the FBoneContainer.
+	 */
+	virtual const FRetargetSourceCachedData& GetRetargetSourceCachedData(const FName& InSourceName, const FSkeletonRemapping& InRemapping, const TArray<FTransform>& InRetargetTransforms) const = 0;
+};
+
+/**
+ * Specialization of the FAnimationRuntimeRetargetBoneTransformContext to map against a FBoneContainer.
+ * This allows to Retarget Bone Transforms for all elements in an FBoneContainer.
+ */ 
+class FBoneContainerRetargetBoneTransformContext : public FAnimationRuntimeRetargetBoneTransformContext
+{
+public:
+
+	FBoneContainerRetargetBoneTransformContext(const FBoneContainer& InBoneContainer)
+		: BoneContainer(InBoneContainer)
+	{
+	}
+
+	virtual const USkeleton* GetSkeleton(bool bEvenIfUnreachable = false) const override
+	{
+		return BoneContainer.GetSkeletonAsset(bEvenIfUnreachable);
+	}
+	
+	virtual int32 GetSkeletonBoneIndex(const int32& InContextBoneIndex) const override
+	{
+		return BoneContainer.GetSkeletonIndex(FCompactPoseBoneIndex(InContextBoneIndex));
+	}
+	
+	virtual int32 GetNumBonesInContext() const override
+	{
+		return BoneContainer.GetCompactPoseNumBones();
+	}
+	
+	virtual bool IsRetargetingDisabled() const override
+	{
+		return BoneContainer.GetDisableRetargeting();
+	}
+	
+	virtual const FTransform& GetRefPoseTransform(const int32& InContextBoneIndex) const override
+	{
+		return BoneContainer.GetRefPoseTransform(FCompactPoseBoneIndex(InContextBoneIndex));
+	}
+	
+	virtual const FRetargetSourceCachedData& GetRetargetSourceCachedData(const FName& InSourceName, const FSkeletonRemapping& InRemapping, const TArray<FTransform>& InRetargetTransforms) const override
+	{
+		return BoneContainer.GetRetargetSourceCachedData(InSourceName, InRemapping, InRetargetTransforms);
+	}
+
+private:
+
+	const FBoneContainer& BoneContainer;
+};
+
 /** In AnimationRunTime Library, we extract animation data based on Skeleton hierarchy, not ref pose hierarchy. 
 	Ref pose will need to be re-mapped later **/
 class FAnimationRuntime
@@ -657,7 +748,9 @@ public:
 	* @param	RequiredBones		BoneContainer to which this is retargeting
 	*/
 	static ENGINE_API void RetargetBoneTransform(const USkeleton* SourceSkeleton, const FName& RetargetSource, FTransform& BoneTransform, const int32 SkeletonBoneIndex, const FCompactPoseBoneIndex& BoneIndex, const FBoneContainer& RequiredBones, const bool bIsBakedAdditive);
-	static ENGINE_API void RetargetBoneTransform(const USkeleton* MySkeleton, const FName& SourceName, const TArray<FTransform>& RetargetTransforms, FTransform& BoneTransform, const int32 SkeletonBoneIndex, const FCompactPoseBoneIndex& BoneIndex, const FBoneContainer& RequiredBones, const bool bIsBakedAdditive);
+	static ENGINE_API void RetargetBoneTransform(const USkeleton* SourceSkeleton, const FName& SourceName, const TArray<FTransform>& RetargetTransforms, FTransform& BoneTransform, const int32 SkeletonBoneIndex, const FCompactPoseBoneIndex& BoneIndex, const FBoneContainer& RequiredBones, const bool bIsBakedAdditive);
+	static ENGINE_API void RetargetBoneTransform(const USkeleton* SourceSkeleton, const FName& SourceName, const TArray<FTransform>& RetargetTransforms, FTransform& BoneTransform, const int32 SkeletonBoneIndex, const FCompactPoseBoneIndex& BoneIndex, const FAnimationRuntimeRetargetBoneTransformContext* InContext, const bool bIsBakedAdditive);
+	
 	/** 
 	 * Calculate distance how close two strings are. 
 	 * By close, it calculates how many operations to transform First to Second 
