Minor crash-fix in FMaterial::Translate().

#rnx
#Synced-CL 48350544

diff --git a/Engine/Source/Runtime/Engine/Private/Materials/MaterialShared.cpp b/Engine/Source/Runtime/Engine/Private/Materials/MaterialShared.cpp
--- a/Engine/Source/Runtime/Engine/Private/Materials/MaterialShared.cpp
+++ b/Engine/Source/Runtime/Engine/Private/Materials/MaterialShared.cpp
@@ -3531,8 +3531,11 @@ bool FMaterial::Translate(const FMaterialShaderMapId& InShaderMapId,
 	FMaterialCompilationOutput& OutCompilationOutput,
 	TRefCountPtr<FSharedShaderCompilerEnvironment>& OutMaterialEnvironment)
 {
-	// Not all Insight data will be filled out by both translators.
-	GetMaterialInterface()->MaterialInsight.Reset(new FMaterialInsights);
+	// Reset the material insights instance if this is a valid material interface.
+	if (UMaterialInterface* Interface = GetMaterialInterface())
+	{
+		Interface->MaterialInsight.Reset(new FMaterialInsights);
+	}
 
 	// Use the new translator if the shader map uses it. Note: the new translator does not support Substrate yet, so if Substrate is on, fallback to the old translator.
 	if (InShaderMapId.bUsingNewHLSLGenerator && !Substrate::IsSubstrateEnabled())
