Revert part of previous change so that FSubsystemModuleWatcher::RemoveClassesForModule still loops over TArray directly and not TArrayView within RemoveAllInstances

#Synced-CL Unknown

diff --git a/Engine/Source/Runtime/Engine/Private/Subsystems/SubsystemCollection.cpp b/Engine/Source/Runtime/Engine/Private/Subsystems/SubsystemCollection.cpp
--- a/Engine/Source/Runtime/Engine/Private/Subsystems/SubsystemCollection.cpp
+++ b/Engine/Source/Runtime/Engine/Private/Subsystems/SubsystemCollection.cpp
@@ -659,7 +659,10 @@ void FSubsystemModuleWatcher::RemoveClassesForModule(const FName& InModuleName)
 	TArray<UClass*>* SubsystemClasses = GlobalDynamicSystemModuleMap.Find(InModuleName);
 	if (SubsystemClasses)
 	{
-		FSubsystemCollectionBase::RemoveAllInstances(*SubsystemClasses);
+		for (UClass* SubsystemClass : *SubsystemClasses)
+		{
+			FSubsystemCollectionBase::RemoveAllInstances(MakeArrayView(&SubsystemClass, 1));
+		}
 		GlobalDynamicSystemModuleMap.Remove(InModuleName);
 	}
 }
