Skip rebuilding legacy curves in FNewSpline::UpdateSpline and instead just mark it dirty.

Can be disabled with Spline.SkipRebuildingLegacyCurvesInUpdate 0.

#rnx
#tests PIE, geometry processing unit tests
#Synced-CL 48247553

diff --git a/Engine/Source/Runtime/Engine/Private/Curves/Spline.cpp b/Engine/Source/Runtime/Engine/Private/Curves/Spline.cpp
--- a/Engine/Source/Runtime/Engine/Private/Curves/Spline.cpp
+++ b/Engine/Source/Runtime/Engine/Private/Curves/Spline.cpp
@@ -112,6 +112,13 @@ FAutoConsoleVariableRef CVarSkipDoubleUpdate(
 	TEXT("When true, skip updating the spline twice if UpdateSpline with parameters is called.")
 );
 
+bool GSkipRebuildingLegacyCurvesInUpdate = true;
+FAutoConsoleVariableRef CVarSkipRebuildingLegacyCurvesInUpdate(
+	TEXT("Spline.SkipRebuildingLegacyCurvesInUpdate"),
+	GSkipRebuildingLegacyCurvesInUpdate,
+	TEXT("When true, only mark legacy curves as dirty in UdpateSpline rather than rebuilding immediately.")
+);
+
 }
 
 UE::Geometry::Spline::EParameterizationPolicy GetParameterizationPolicy()
@@ -2574,8 +2581,15 @@ void FNewSpline::UpdateSpline()
 
     MarkReparamTableDirty();
 
-	// Updates legacy curves based on points, never evaluates internal spline.
-    RebuildLegacyCurves();					
+	if (UE::Spline::GSkipRebuildingLegacyCurvesInUpdate)
+	{
+		MarkLegacyCurvesDirty();
+	}
+	else
+	{
+		// Updates legacy curves based on points, never evaluates internal spline.
+		RebuildLegacyCurves();
+	}
 }
 
 void FNewSpline::UpdateSpline(const FSpline::FUpdateSplineParams& Params)
