Add Insights tracking for active SCW count
#rnx

#Synced-CL 48334554

diff --git a/Engine/Source/Runtime/Engine/Private/ShaderCompiler/ShaderCompilerThreadRunnable.cpp b/Engine/Source/Runtime/Engine/Private/ShaderCompiler/ShaderCompilerThreadRunnable.cpp
--- a/Engine/Source/Runtime/Engine/Private/ShaderCompiler/ShaderCompilerThreadRunnable.cpp
+++ b/Engine/Source/Runtime/Engine/Private/ShaderCompiler/ShaderCompilerThreadRunnable.cpp
@@ -59,6 +59,7 @@ static const TCHAR* GWorkerInputFilename = TEXT("WorkerInputOnly.in");
 static const TCHAR* GWorkerOutputFilename = TEXT("WorkerOutputOnly.out");
 
 TRACE_DECLARE_MEMORY_COUNTER(AsyncCompilationSCWMemory, TEXT("AsyncCompilation/SCWMemory"));
+TRACE_DECLARE_INT_COUNTER(AsyncCompilationSCWsActive, TEXT("AsyncCompilation/SCWsActive"));
 
 FShaderCompilerMemoryLimits GetShaderCompilerCurrentMemoryLimits()
 {
@@ -150,6 +151,8 @@ struct FShaderCompileWorkerInfo
 			}
 			FPlatformProcess::CloseProc(WorkerProcess);
 			WorkerProcess = FProcHandle();
+
+			TRACE_COUNTER_DECREMENT(AsyncCompilationSCWsActive);
 		}
 	}
 
@@ -161,6 +164,8 @@ struct FShaderCompileWorkerInfo
 			FPlatformProcess::GetProcReturnCode(WorkerProcess, &ReturnCode);
 			FPlatformProcess::CloseProc(WorkerProcess);
 			WorkerProcess = FProcHandle();
+
+			TRACE_COUNTER_DECREMENT(AsyncCompilationSCWsActive);
 		}
 		return ReturnCode;
 	}
@@ -742,6 +747,8 @@ bool FShaderCompileThreadRunnable::LaunchWorkersIfNeeded()
 				CurrentWorkerInfo.bLaunchedWorker = true;
 
 				NumberLaunched++;
+
+				TRACE_COUNTER_INCREMENT(AsyncCompilationSCWsActive);
 			}
 		}
 	}
