[Backout] - CL26778915
Original CL Desc
-----------------------------------------------------------------
World Partition: always include editor-only references in the instancing context when cooking instanced packages.
This fixes the case where a cooked actor is referencing another cooked actors through an editor-only reference.

[Backout] - CL42392054
Original CL Desc
-----------------------------------------------------------------
World Partition: make sure to mark non-editoronly objects referenced by editoronly objects or through editoronly properties from editoronly  or non-editoronly objects as transient so they don't get included in cooked builds.
This works differently in PIE, as we use the instancing context to skip loading these objects.

#rnx
#Synced-CL 48350871

diff --git a/Engine/Source/Runtime/Engine/Private/WorldPartition/WorldPartitionLevelHelper.cpp b/Engine/Source/Runtime/Engine/Private/WorldPartition/WorldPartitionLevelHelper.cpp
--- a/Engine/Source/Runtime/Engine/Private/WorldPartition/WorldPartitionLevelHelper.cpp
+++ b/Engine/Source/Runtime/Engine/Private/WorldPartition/WorldPartitionLevelHelper.cpp
@@ -278,12 +278,6 @@ void FWorldPartitionLevelHelper::MoveExternalActorsToLevel(const TArray<FWorldPa
 			}
 			check(Actor->GetPackage() == LevelPackage);
 
-			if (PackageObjectMapping.bIsEditorOnly)
-			{
-				Actor->SetFlags(RF_Transient);
-				UE_LOG(LogWorldPartition, Log, TEXT("Marked actor %s transient as it was referenced by an editor-only context"), *Actor->GetPathName());
-			}
-
 			// Process objects found in the source actor package
 			if (bIsActoPackageExternal)
 			{
@@ -329,12 +323,6 @@ void FWorldPartitionLevelHelper::MoveExternalActorsToLevel(const TArray<FWorldPa
 
 							Object->Rename(*NewName.ToString(), LevelPackage, REN_NonTransactional | REN_DoNotDirty | REN_DontCreateRedirectors);
 						}
-
-						if (PackageObjectMapping.bIsEditorOnly)
-						{
-							Object->SetFlags(RF_Transient);
-							UE_LOG(LogWorldPartition, Log, TEXT("Marked actor object %s transient as it was referenced by an editor-only context"), *Object->GetPathName());
-						}
 					}
 				}
 			}
diff --git a/Engine/Source/Runtime/Engine/Private/WorldPartition/WorldPartitionRuntimeLevelStreamingCell.cpp b/Engine/Source/Runtime/Engine/Private/WorldPartition/WorldPartitionRuntimeLevelStreamingCell.cpp
--- a/Engine/Source/Runtime/Engine/Private/WorldPartition/WorldPartitionRuntimeLevelStreamingCell.cpp
+++ b/Engine/Source/Runtime/Engine/Private/WorldPartition/WorldPartitionRuntimeLevelStreamingCell.cpp
@@ -328,43 +328,46 @@ void UWorldPartitionRuntimeLevelStreamingCell::AddActorToCell(const FStreamingGe
 		ParentContainerInstance = CurrentContainerInstance->GetParentContainerInstance();
 	}
 
-	for (const FGuid& EditorReferenceGuid : ActorDescView.GetEditorReferences())
+	if (!IsRunningCookCommandlet())
 	{
-		FName ReferencePackage;
-		FName ReferencePath;
-		FTopLevelAssetPath ReferenceBaseClass;
-		FTopLevelAssetPath ReferenceNativeClass;
-
-		// Special case where the actor descriptor view has an invalid references, use invalid reference information as the actor guid isn't necessarily in the container instance.
-		if (const FStreamingGenerationActorDescView::FInvalidReference* InvalidReference = ActorDescView.GetInvalidReference(EditorReferenceGuid))
-		{
-			ReferencePackage = InvalidReference->ActorPackage;
-			ReferencePath = *InvalidReference->ActorSoftPath.ToString();
-			ReferenceBaseClass = InvalidReference->BaseClass;
-			ReferenceNativeClass = InvalidReference->NativeClass;
-		}
-		else
+		for (const FGuid& EditorReferenceGuid : ActorDescView.GetEditorReferences())
 		{
-			const FWorldPartitionActorDescInstance& ReferenceActorDesc = ContainerInstance->GetActorDescInstanceChecked(EditorReferenceGuid);
-			ReferencePackage = ReferenceActorDesc.GetActorPackage();
-			ReferencePath = *ReferenceActorDesc.GetActorSoftPath().ToString();
-			ReferenceBaseClass = ReferenceActorDesc.GetBaseClass();
-			ReferenceNativeClass = ReferenceActorDesc.GetNativeClass();
+			FName ReferencePackage;
+			FName ReferencePath;
+			FTopLevelAssetPath ReferenceBaseClass;
+			FTopLevelAssetPath ReferenceNativeClass;
+
+			// Special case where the actor descriptor view has an invalid references, use invalid reference information as the actor guid isn't necessarily in the container instance.
+			if (const FStreamingGenerationActorDescView::FInvalidReference* InvalidReference = ActorDescView.GetInvalidReference(EditorReferenceGuid))
+			{
+				ReferencePackage = InvalidReference->ActorPackage;
+				ReferencePath = *InvalidReference->ActorSoftPath.ToString();
+				ReferenceBaseClass = InvalidReference->BaseClass;
+				ReferenceNativeClass = InvalidReference->NativeClass;
+			}
+			else
+			{
+				const FWorldPartitionActorDescInstance& ReferenceActorDesc = ContainerInstance->GetActorDescInstanceChecked(EditorReferenceGuid);
+				ReferencePackage = ReferenceActorDesc.GetActorPackage();
+				ReferencePath = *ReferenceActorDesc.GetActorSoftPath().ToString();
+				ReferenceBaseClass = ReferenceActorDesc.GetBaseClass();
+				ReferenceNativeClass = ReferenceActorDesc.GetNativeClass();
+			}
+
+			Packages.Emplace(
+				ReferencePackage,
+				ReferencePath,
+				ReferenceBaseClass,
+				ReferenceNativeClass,
+				ContainerID,
+				ContainerTransform,
+				FTransform::Identity,
+				ContainerPackage,
+				GetWorld()->GetPackage()->GetFName(),
+				ContainerID.GetActorGuid(EditorReferenceGuid),
+				true
+			);
 		}
-		
-		Packages.Emplace(
-			ReferencePackage,
-			ReferencePath,
-			ReferenceBaseClass,
-			ReferenceNativeClass,
-			ContainerID,
-			ContainerTransform,
-			FTransform::Identity,
-			ContainerPackage,
-			GetWorld()->GetPackage()->GetFName(),
-			ContainerID.GetActorGuid(EditorReferenceGuid),
-			true
-		);
 	}
 
 	FWorldPartitionRuntimeCellObjectMapping ActorMapping(
