[InitDefaultMaterials]
- Fix InitDefaultMaterials recursivity to avoid all the warnings from the loader at boot
- Remove unused GPowerToRoughnessMaterialFunction
- Load GConvertFromDiffSpecMaterialFunction on demand

#Synced-CL Unknown

diff --git a/Engine/Source/Runtime/Engine/Private/Materials/Material.cpp b/Engine/Source/Runtime/Engine/Private/Materials/Material.cpp
--- a/Engine/Source/Runtime/Engine/Private/Materials/Material.cpp
+++ b/Engine/Source/Runtime/Engine/Private/Materials/Material.cpp
@@ -549,10 +549,10 @@ void UMaterial::ForceNoCompilationInPostLoad(bool bForceNoCompilation)
 #endif
 }
 
-static UMaterialFunction* GPowerToRoughnessMaterialFunction = NULL;
 static UMaterialFunction* GConvertFromDiffSpecMaterialFunction = NULL;
 
 static UMaterial* GDefaultMaterials[MD_MAX] = {0};
+static bool       GDefaultMaterialsInitialized[MD_MAX] = {false};
 
 static const TCHAR* GDefaultMaterialNames[MD_MAX] =
 {
@@ -576,23 +576,83 @@ static const TCHAR* GDefaultMaterialNames[MD_MAX] =
 // Need to know if the default materials have been initialized in case dynamic shader preloading is enabled,
 // we don't want to kick preload shader jobs since they can call GetDefaultMaterial and causing a crash in InitDefaultMaterials.
 static bool bDefaultMaterialInitialized = false;
+static bool bDefaultMaterialInitializationStarted = false;
 
 bool UMaterialInterface::IsDefaultMaterialInitialized()
 {
 	return bDefaultMaterialInitialized;
 }
 
+UMaterial* TryResolveDefaultMaterial(int32 Domain)
+{
+	FString ResolvedPath = ResolveIniObjectsReference(GDefaultMaterialNames[Domain]);
+	return FindObject<UMaterial>(nullptr, *ResolvedPath);
+}
+
+// We can save time if instead of blocking after compilation of each synchronous material we block after scheduling all of them
+bool GPoolSpecialMaterialsCompileJobs = true;
+bool PoolSpecialMaterialsCompileJobs()
+{
+	return GPoolSpecialMaterialsCompileJobs;
+}
+
+namespace MaterialImpl
+{
+	void ConditionalFinishInitialization(int32 Domain)
+	{
+		if (!bDefaultMaterialInitialized)
+		{
+			SCOPED_BOOT_TIMING("UMaterialInterface::InitDefaultMaterials");
+
+			if (!GDefaultMaterialsInitialized[Domain])
+			{
+				FString ResolvedPath = ResolveIniObjectsReference(GDefaultMaterialNames[Domain]);
+				GDefaultMaterials[Domain] = FindObject<UMaterial>(nullptr, *ResolvedPath);
+				checkf(GDefaultMaterials[Domain] != nullptr, TEXT("Cannot load default material '%s' [Domain=%s] from path '%s'"), GDefaultMaterialNames[Domain], *MaterialDomainString((EMaterialDomain)Domain), *ResolvedPath);
+				if (GDefaultMaterials[Domain])
+				{
+					GDefaultMaterials[Domain]->AddToRoot();
+				}
+				GDefaultMaterialsInitialized[Domain] = true;
+			}
+
+			bool bAllDefaultMaterialInitialized = true;
+			for (int32 Index = 0; Index < MD_MAX; ++Index)
+			{
+				bAllDefaultMaterialInitialized &= GDefaultMaterialsInitialized[Index];
+			}
+	
+			if (bAllDefaultMaterialInitialized)
+			{
+				bDefaultMaterialInitialized = true;
+
+				// Block after scheduling for compilation all (hopefully) default materials.
+				// Even if not all of them ended up being post-loaded, block here just out of extra caution
+				if (GPoolSpecialMaterialsCompileJobs == true)
+				{
+					GPoolSpecialMaterialsCompileJobs = false;
+					if (GShaderCompilingManager)
+					{
+						GShaderCompilingManager->FinishAllCompilation();
+					}
+				}
+
+				// Now precache PSOs for all the default materials after the default materials are marked initialize
+				// PSO precaching can request default materials so they have to marked as initialized to avoid endless recursion
+				if (PipelineStateCache::IsPSOPrecachingEnabled() || IsPSOShaderPreloadingEnabled())
+				{
+					UMaterialInterface::PrecacheDefaultMaterialPSOs();
+				}
+			}
+		}
+	}
+}
+
 void UMaterialInterface::InitDefaultMaterials()
 {
-	// Note that this function will (in fact must!) be called recursively. This
-	// guarantees that the default materials will have been loaded and pointers
-	// set before any other material interface has been instantiated -- even
-	// one of the default materials! It is actually possible to assert that
-	// these materials exist in the UMaterial or UMaterialInstance constructor.
-	// 
 	// The check for initialization is purely an optimization as initializing
 	// the default materials is only done very early in the boot process.
-	if (!bDefaultMaterialInitialized)
+	if (!bDefaultMaterialInitializationStarted)
 	{
 		SCOPED_BOOT_TIMING("UMaterialInterface::InitDefaultMaterials");
 		check(IsInGameThread());
@@ -600,8 +660,10 @@ void UMaterialInterface::InitDefaultMaterials()
 		{
 			return;
 		}
-		static int32 RecursionLevel = 0;
-		RecursionLevel++;
+
+		// We only do this step once, launch async loading tasks with top priority to fullfill these requests as soon as possible.
+		// When done loading, these will call ConditionalFinishInitialization which handles the remaining of the initialization.
+		bDefaultMaterialInitializationStarted = true;
 
 		// Explicitly initialize UMaterialInterface and UMaterial CDOs in this order to avoid issues with recursion and
 		// async loading depending on initialization order of compiled-in UObjects.
@@ -609,47 +671,19 @@ void UMaterialInterface::InitDefaultMaterials()
 		UMaterialInterface::StaticClass()->GetDefaultObject();
 		UMaterial::StaticClass()->GetDefaultObject();
 		
-#if WITH_EDITOR
-		GPowerToRoughnessMaterialFunction = LoadObject< UMaterialFunction >(nullptr, TEXT("/Engine/Functions/Engine_MaterialFunctions01/Shading/PowerToRoughness.PowerToRoughness"), nullptr, LOAD_None, nullptr);
-		checkf( GPowerToRoughnessMaterialFunction, TEXT("Cannot load PowerToRoughness") );
-		GPowerToRoughnessMaterialFunction->AddToRoot();
-
-		GConvertFromDiffSpecMaterialFunction = LoadObject< UMaterialFunction >(nullptr, TEXT("/Engine/Functions/Engine_MaterialFunctions01/Shading/ConvertFromDiffSpec.ConvertFromDiffSpec"), nullptr, LOAD_None, nullptr);
-		checkf( GConvertFromDiffSpecMaterialFunction, TEXT("Cannot load ConvertFromDiffSpec") );
-		GConvertFromDiffSpecMaterialFunction->AddToRoot();
-#endif
-
 		for (int32 Domain = 0; Domain < MD_MAX; ++Domain)
 		{
-			if (GDefaultMaterials[Domain] == nullptr)
-			{
-				FString ResolvedPath = ResolveIniObjectsReference(GDefaultMaterialNames[Domain]);
+			FString ResolvedPath = ResolveIniObjectsReference(GDefaultMaterialNames[Domain]);
 
-				GDefaultMaterials[Domain] = FindObject<UMaterial>(nullptr, *ResolvedPath);
-				if (GDefaultMaterials[Domain] == nullptr
-#if USE_EVENT_DRIVEN_ASYNC_LOAD_AT_BOOT_TIME
-					&& (RecursionLevel == 1 || !GEventDrivenLoaderEnabled)
-#endif
-					)
-				{
-					GDefaultMaterials[Domain] = LoadObject<UMaterial>(nullptr, *ResolvedPath, nullptr, LOAD_DisableDependencyPreloading, nullptr);
-					checkf(GDefaultMaterials[Domain] != nullptr, TEXT("Cannot load default material '%s' [Domain=%s] from path '%s'"), GDefaultMaterialNames[Domain], *MaterialDomainString((EMaterialDomain)Domain), *ResolvedPath);
-				}
-				if (GDefaultMaterials[Domain])
+			FLoadPackageAsyncOptionalParams Params;
+			Params.PackagePriority = MAX_int32;
+			Params.CompletionDelegate = MakeUnique<FLoadPackageAsyncDelegate>(FLoadPackageAsyncDelegate::CreateLambda(
+				[Domain](const FName& PackageName, UPackage* Package, EAsyncLoadingResult::Type Result)
 				{
-					GDefaultMaterials[Domain]->AddToRoot();
-				}
-			}
-		}
-
-		RecursionLevel--;
-		bDefaultMaterialInitialized = RecursionLevel == 0;
+					MaterialImpl::ConditionalFinishInitialization(Domain);
+				}));
 
-		// Now precache PSOs for all the default materials after the default materials are marked initialize
-		// PSO precaching can request default materials so they have to marked as initialized to avoid endless recursion
-		if (bDefaultMaterialInitialized && (PipelineStateCache::IsPSOPrecachingEnabled() || IsPSOShaderPreloadingEnabled()))
-		{
-			PrecacheDefaultMaterialPSOs();
+			LoadPackageAsync(FPackagePath::FromPackageNameChecked(FTopLevelAssetPath(ResolvedPath).GetPackageName()), MoveTemp(Params));
 		}
 	}
 }
@@ -700,13 +734,6 @@ void UMaterialInterface::PostCDOContruct()
 	}
 }
 
-// We can save time if instead of blocking after compilation of each synchronous material we block after scheduling all of them
-bool GPoolSpecialMaterialsCompileJobs = true;
-bool PoolSpecialMaterialsCompileJobs()
-{
-	return GPoolSpecialMaterialsCompileJobs;
-}
-
 void UMaterialInterface::PostLoadDefaultMaterials()
 {
 	LLM_SCOPE(ELLMTag::Materials);
@@ -715,49 +742,39 @@ void UMaterialInterface::PostLoadDefaultMaterials()
 	// is an optimization and guarantees that default materials are post loaded
 	// in the order material domains are defined. Surface -> deferred decal -> etc.
 	static bool bPostLoaded = false;
+	static TSet<UMaterialInterface*> OnStack;
 	if (!bPostLoaded)
 	{
 		check(IsInGameThread());
 		bPostLoaded = true;
 
-#if WITH_EDITOR
-		GPowerToRoughnessMaterialFunction->ConditionalPostLoad();
-		GConvertFromDiffSpecMaterialFunction->ConditionalPostLoad();
-#endif
-
-		for (int32 Domain = 0; Domain < MD_MAX; ++Domain)
+		if (!bDefaultMaterialInitialized)
 		{
-			UMaterial* Material = GDefaultMaterials[Domain];
-#if USE_EVENT_DRIVEN_ASYNC_LOAD_AT_BOOT_TIME
-			check(Material || (GIsInitialLoad && GEventDrivenLoaderEnabled));
-			if (Material && !Material->HasAnyFlags(RF_NeedLoad))
-#else
-			check(Material);
-			if (Material)
-#endif
+			TRACE_CPUPROFILER_EVENT_SCOPE(UMaterialInterface::PostLoadDefaultMaterials);
+
+			check(IsInGameThread());
+			OnStack.Add(this);
+			ON_SCOPE_EXIT { OnStack.Remove(this); };
+
+			for (int32 Domain = 0; Domain < MD_MAX; ++Domain)
 			{
-				Material->ConditionalPostLoad();
-				// Sometimes the above will get called before the material has been fully serialized
-				// in this case its NeedPostLoad flag will not be cleared.
-				if (Material->HasAnyFlags(RF_NeedPostLoad))
+				UMaterial* Material = TryResolveDefaultMaterial(Domain);
+				if (OnStack.Contains(Material))
 				{
-					bPostLoaded = false;
+					// We obviously can't sync load ourselves while we're in our own postload function
+					continue;
 				}
-			}
-			else
-			{
-				bPostLoaded = false;
-			}
-		}
 
-		// Block after scheduling for compilation all (hopefully) default materials.
-		// Even if not all of them ended up being post-loaded, block here just out of extra caution
-		if (GPoolSpecialMaterialsCompileJobs == true)
-		{
-			GPoolSpecialMaterialsCompileJobs = false;
-			if (GShaderCompilingManager)
-			{
-				GShaderCompilingManager->FinishAllCompilation();
+				// This will forceload the object if it hasn't already been loaded.
+				// Sync loading from inside PostLoad is supported by the loader.
+				Material = UMaterial::GetDefaultMaterial((EMaterialDomain)Domain);
+				check(Material);
+				if (Material)
+				{
+					// The only reason these would fail is if there is a recursive loads.. (i.e. a material trying to load and postload itself indirectly since we already handle the direct case)
+					check(!Material->HasAnyFlags(RF_NeedLoad));
+					check(!Material->HasAnyFlags(RF_NeedPostLoad));
+				}
 			}
 		}
 	}
@@ -1058,7 +1075,15 @@ UMaterial* UMaterial::GetDefaultMaterial(EMaterialDomain Domain)
 {
 	InitDefaultMaterials();
 	check(Domain >= MD_Surface && Domain < MD_MAX);
-	check(GDefaultMaterials[Domain] != NULL);
+
+	if (!GDefaultMaterialsInitialized[Domain])
+	{
+		FString ResolvedPath = ResolveIniObjectsReference(GDefaultMaterialNames[Domain]);
+		LoadObject<UMaterial>(nullptr, *ResolvedPath, nullptr, LOAD_DisableDependencyPreloading, nullptr);
+		MaterialImpl::ConditionalFinishInitialization(Domain);
+	}
+
+	check(GDefaultMaterials[Domain] != nullptr);
 	UMaterial* Default = GDefaultMaterials[Domain];
 	return Default;
 }
@@ -1070,6 +1095,15 @@ bool UMaterial::IsDefaultMaterial() const
 	{
 		bDefault = (this == GDefaultMaterials[Domain]);
 	}
+
+	if (!bDefaultMaterialInitialized)
+	{
+		for (int32 Domain = 0; !bDefault && Domain < MD_MAX; ++Domain)
+		{
+			bDefault = (this == TryResolveDefaultMaterial(Domain));
+		}
+	}
+
 	return bDefault;
 }
 
@@ -3221,6 +3255,13 @@ void UMaterial::BackwardsCompatibilityInputConversion()
 		if( bIsDS && !bIsBMS )
 		{
 			// ConvertFromDiffSpec
+			if (!GConvertFromDiffSpecMaterialFunction)
+			{
+				GConvertFromDiffSpecMaterialFunction = LoadObject< UMaterialFunction >(nullptr, TEXT("/Engine/Functions/Engine_MaterialFunctions01/Shading/ConvertFromDiffSpec.ConvertFromDiffSpec"), nullptr, LOAD_None, nullptr);
+				checkf(GConvertFromDiffSpecMaterialFunction, TEXT("Cannot load ConvertFromDiffSpec"));
+				GConvertFromDiffSpecMaterialFunction->AddToRoot();
+				check(!GConvertFromDiffSpecMaterialFunction->HasAnyFlags(RF_NeedLoad | RF_NeedPostLoad));
+			}
 
 			check( GConvertFromDiffSpecMaterialFunction );
 
diff --git a/Engine/Source/Runtime/Engine/Private/Materials/MaterialInterface.cpp b/Engine/Source/Runtime/Engine/Private/Materials/MaterialInterface.cpp
--- a/Engine/Source/Runtime/Engine/Private/Materials/MaterialInterface.cpp
+++ b/Engine/Source/Runtime/Engine/Private/Materials/MaterialInterface.cpp
@@ -185,13 +185,8 @@ UMaterialInterface::UMaterialInterface(const FObjectInitializer& ObjectInitializ
 	if (!HasAnyFlags(RF_ClassDefaultObject))
 	{
 		MaterialDomainString(MD_Surface); // find the enum for this now before we start saving
-#if USE_EVENT_DRIVEN_ASYNC_LOAD_AT_BOOT_TIME
-		if (!GIsInitialLoad || !GEventDrivenLoaderEnabled)
-#endif
-		{
-			InitDefaultMaterials();
-			AssertDefaultMaterialsExist();
-		}
+
+		InitDefaultMaterials();
 
 		if (SamplerTypeEnum == nullptr)
 		{
diff --git a/Engine/Source/Runtime/Engine/Public/Materials/MaterialInterface.h b/Engine/Source/Runtime/Engine/Public/Materials/MaterialInterface.h
--- a/Engine/Source/Runtime/Engine/Public/Materials/MaterialInterface.h
+++ b/Engine/Source/Runtime/Engine/Public/Materials/MaterialInterface.h
@@ -1371,7 +1371,7 @@ private:
 	/**
 	 * Post loads all default materials.
 	 */
-	static void PostLoadDefaultMaterials();
+	void PostLoadDefaultMaterials();
 
 	/**
 	* Cached type information for the sampler type enumeration. 
