NetTest
* Added unit test for StopReplication behavior with the NetRootObjectAdapter on dynamic and stable objects

#rnx
#Synced-CL 48203440

diff --git a/Engine/Source/Runtime/Engine/Public/Net/Iris/ReplicationSystem/NetRootObjectAdapter.h b/Engine/Source/Runtime/Engine/Public/Net/Iris/ReplicationSystem/NetRootObjectAdapter.h
--- a/Engine/Source/Runtime/Engine/Public/Net/Iris/ReplicationSystem/NetRootObjectAdapter.h
+++ b/Engine/Source/Runtime/Engine/Public/Net/Iris/ReplicationSystem/NetRootObjectAdapter.h
@@ -99,11 +99,13 @@ public:
 		return bIsInitialized;
 	}
 
+	/** Server only: tells if the object was added to the replication system or not */
 	bool IsReplicating() const
 	{
 		return bIsReplicating;
 	}
 
+	/** Server only: tells if the object is already attached to a level or not */
 	bool HasAttachedLevel() const
 	{
 		return WeakLevel.IsValid();
