[AsyncCompilation/Texture]
- Restore async texture compilation since the landscape code has changed and doesn't seem to be causing issues anymore

#Synced-CL Unknown

diff --git a/Engine/Source/Runtime/Engine/Private/Texture.cpp b/Engine/Source/Runtime/Engine/Private/Texture.cpp
--- a/Engine/Source/Runtime/Engine/Private/Texture.cpp
+++ b/Engine/Source/Runtime/Engine/Private/Texture.cpp
@@ -554,7 +554,7 @@ bool UTexture::Modify(bool bAlwaysMarkDirty)
 	if (!IsAsyncCacheComplete())
 	{
 		// try to throw away the data since we're going to invalidate it anyway
-		//FAssetCompilingManager::Get().MarkCompilationAsCanceled({ this });
+		FAssetCompilingManager::Get().MarkCompilationAsCanceled({ this });
 	}
 
 	// Before applying any modification to the texture
