Attribute channels tangents fix

#Synced-CL 48389179

diff --git a/Engine/Source/Runtime/Engine/Private/Curves/Spline.cpp b/Engine/Source/Runtime/Engine/Private/Curves/Spline.cpp
--- a/Engine/Source/Runtime/Engine/Private/Curves/Spline.cpp
+++ b/Engine/Source/Runtime/Engine/Private/Curves/Spline.cpp
@@ -2648,6 +2648,10 @@ void FNewSpline::Reparameterize()
 		}
 
 		RestoreChannelKnots(ChannelName, CachedChannelParameterizations[ChannelName].Value);
+		if (UE::Geometry::Spline::TTangentBezierSpline<float>* Channel = GetTypedAttributeChannel<UE::Geometry::Spline::TTangentBezierSpline<float>>(ChannelName))
+		{
+			Channel->UpdateTangents();
+		}
 	}
 
 	if (UE::Geometry::Spline::TTangentBezierSpline<FVector>* ScaleSpline = GetTypedAttributeChannel<UE::Geometry::Spline::TTangentBezierSpline<FVector>>(ScaleAttrName))
