- Fix invalid initial state / incorrect transitions in physics fields

#tests Enable rhi validation and run physics fields

#Synced-CL 48229894

diff --git a/Engine/Source/Runtime/Engine/Private/PhysicsField/PhysicsFieldComponent.cpp b/Engine/Source/Runtime/Engine/Private/PhysicsField/PhysicsFieldComponent.cpp
--- a/Engine/Source/Runtime/Engine/Private/PhysicsField/PhysicsFieldComponent.cpp
+++ b/Engine/Source/Runtime/Engine/Private/PhysicsField/PhysicsFieldComponent.cpp
@@ -144,7 +144,7 @@ void InitInternalBuffer(FRHICommandListBase& RHICmdList, const TCHAR* DebugName,
 }
 
 template<typename BufferType, int ElementSize, EPixelFormat PixelFormat>
-void UpdateInternalBuffer(FRHICommandListBase& RHICmdList, const TCHAR* DebugName, const uint32 ElementCount, const BufferType* InputData, FRWBuffer& OutputBuffer, const bool bInitField = false)
+void UpdateInternalBuffer(FRHICommandListBase& RHICmdList, const TCHAR* DebugName, const uint32 ElementCount, const BufferType* InputData, FRWBuffer& OutputBuffer, const bool bInitField = false, ERHIAccess InitAccess = ERHIAccess::SRVMask)
 {
 	if (ElementCount > 0 && InputData)
 	{
@@ -154,7 +154,7 @@ void UpdateInternalBuffer(FRHICommandListBase& RHICmdList, const TCHAR* DebugNam
 		if (bInitField)
 		{
 			LLM_SCOPE(ELLMTag::Physics);
-			OutputBuffer.Initialize(RHICmdList, DebugName, sizeof(BufferType), BufferCount, PixelFormat, BUF_Static);
+			OutputBuffer.Initialize(RHICmdList, DebugName, sizeof(BufferType), BufferCount, PixelFormat, InitAccess, BUF_Static);
 		}
 
 		void* OutputData = RHICmdList.LockBuffer(OutputBuffer.Buffer, 0, BufferBytes, RLM_WriteOnly);
@@ -604,7 +604,7 @@ void FPhysicsFieldResource::UpdateResource(FRHICommandList& RHICmdList,
 			DispatchComputeShader(RHICmdList, ComputeShader.GetShader(), NumGroups, 1, 1);
 			UnsetShaderParametersLegacyCS(RHICmdList, ComputeShader);
 
-			RHICmdList.Transition(FRHITransitionInfo(this->ClipmapBuffer.UAV, ERHIAccess::Unknown, ERHIAccess::SRVCompute));
+			RHICmdList.Transition(FRHITransitionInfo(this->ClipmapBuffer.UAV, ERHIAccess::UAVCompute, ERHIAccess::SRVMask));
 		}
 	}
 }
