PR #14045: FFXSystem: avoid excessive resource transitions, if legacy-particles are disabled

#Synced-CL 48229894

diff --git a/Engine/Source/Runtime/Engine/Private/Particles/FXSystem.cpp b/Engine/Source/Runtime/Engine/Private/Particles/FXSystem.cpp
--- a/Engine/Source/Runtime/Engine/Private/Particles/FXSystem.cpp
+++ b/Engine/Source/Runtime/Engine/Private/Particles/FXSystem.cpp
@@ -530,16 +530,11 @@ void FFXSystem::PreRender(FRDGBuilder& GraphBuilder, TConstStridedView<FSceneVie
 				SCOPED_DRAW_EVENT(RHICmdList, GPUParticles_PreRender);
 				UpdateMultiGPUResources(RHICmdList);
 
-				PrepareGPUSimulation(RHICmdList);
 				SimulateGPUParticles(RHICmdList, EParticleSimulatePhase::Main, {}, nullptr);
-				FinalizeGPUSimulation(RHICmdList);
 
 				if (IsParticleCollisionModeSupported(GetShaderPlatform(), PCM_DistanceField))
 				{
-					PrepareGPUSimulation(RHICmdList);
 					SimulateGPUParticles(RHICmdList, EParticleSimulatePhase::CollisionDistanceField, ViewUniformBuffer, GlobalDistanceFieldParameterData);
-					//particles rendered during basepass may need to read pos/velocity buffers; must finalize unless we know for sure that nothing in base pass will read it.
-					FinalizeGPUSimulation(RHICmdList);
 				}
 			}
 		);
@@ -562,9 +557,7 @@ void FFXSystem::PostRenderOpaque(FRDGBuilder& GraphBuilder, TConstStridedView<FS
 			[this, ViewUniformBuffer](FRHICommandListImmediate& RHICmdList)
 			{
 				SCOPED_DRAW_EVENT(RHICmdList, GPUParticles_PostRenderOpaque);
-				PrepareGPUSimulation(RHICmdList);
 				SimulateGPUParticles(RHICmdList, EParticleSimulatePhase::CollisionDepthBuffer, ViewUniformBuffer, nullptr);
-				FinalizeGPUSimulation(RHICmdList);
 			}
 		);
 	}
diff --git a/Engine/Source/Runtime/Engine/Private/Particles/FXSystemPrivate.h b/Engine/Source/Runtime/Engine/Private/Particles/FXSystemPrivate.h
--- a/Engine/Source/Runtime/Engine/Private/Particles/FXSystemPrivate.h
+++ b/Engine/Source/Runtime/Engine/Private/Particles/FXSystemPrivate.h
@@ -270,9 +270,6 @@ public:
 	 */
 	bool AddSortedGPUSimulation(FRHICommandListBase& RHICmdList, FParticleSimulationGPU* Simulation, const FVector& ViewOrigin, bool bIsTranslucent, FGPUSortManager::FAllocationInfo& OutInfo);
 
-	void PrepareGPUSimulation(FRHICommandListImmediate& RHICmdList);
-	void FinalizeGPUSimulation(FRHICommandListImmediate& RHICmdList);
-
 	/** Get the shared SortManager, used in the rendering loop to call FGPUSortManager::OnPreRender() and FGPUSortManager::OnPostRenderOpaque() */
 	virtual FGPUSortManager* GetGPUSortManager() const override;
 
diff --git a/Engine/Source/Runtime/Engine/Private/Particles/ParticleGpuSimulation.cpp b/Engine/Source/Runtime/Engine/Private/Particles/ParticleGpuSimulation.cpp
--- a/Engine/Source/Runtime/Engine/Private/Particles/ParticleGpuSimulation.cpp
+++ b/Engine/Source/Runtime/Engine/Private/Particles/ParticleGpuSimulation.cpp
@@ -4760,34 +4760,6 @@ bool FFXSystem::RequiresRayTracingSceneInternal() const
 	return false;
 }
 
-void FFXSystem::PrepareGPUSimulation(FRHICommandListImmediate& RHICmdList)
-{
-	// Grab resources.
-	FParticleStateTextures& CurrentStateTextures = ParticleSimulationResources->GetCurrentStateTextures();
-
-	if (CurrentStateTextures.bTexturesCleared)
-	{
-		// Setup render states.
-		RHICmdList.Transition({
-			FRHITransitionInfo(CurrentStateTextures.PositionTextureRHI, ERHIAccess::SRVMask, ERHIAccess::RTV),
-			FRHITransitionInfo(CurrentStateTextures.VelocityTextureRHI, ERHIAccess::SRVMask, ERHIAccess::RTV) });
-	}
-	else
-	{
-		SCOPED_GPU_MASK(RHICmdList, FRHIGPUMask::All());
-
-		FRHITexture* CurrentStateRenderTargets[] = { CurrentStateTextures.PositionTextureRHI, CurrentStateTextures.VelocityTextureRHI };
-		FRHIRenderPassInfo RenderPassInfo(UE_ARRAY_COUNT(CurrentStateRenderTargets), CurrentStateRenderTargets, ERenderTargetActions::Clear_Store);
-		RHICmdList.BeginRenderPass(RenderPassInfo, TEXT("PrepareGPUSimulation_GPUParticlesClearStateTextures"));
-		RHICmdList.EndRenderPass();
-
-		CurrentStateTextures.bTexturesCleared = true;
-	}
-}
-
-void FFXSystem::FinalizeGPUSimulation(FRHICommandListImmediate& RHICmdList)
-{
-}
 
 void FFXSystem::SimulateGPUParticles(
 	FRHICommandListImmediate& RHICmdList,
@@ -4810,7 +4782,6 @@ void FFXSystem::SimulateGPUParticles(
 	FRHITexture* CurrentStateRenderTargets[2] = { CurrentStateTextures.PositionTextureRHI, CurrentStateTextures.VelocityTextureRHI };
 	FRHITexture* PreviousStateRenderTargets[2] = { PrevStateTextures.PositionTextureRHI, PrevStateTextures.VelocityTextureRHI };
 
-
 #if WITH_MGPU
 	TArray<FTransferResourceParams, TInlineAllocator<4>> CrossGPUTransferResources;
 	const bool bCrossTransferEnabled = GNumExplicitGPUsForRendering > 1;
@@ -4836,20 +4807,49 @@ void FFXSystem::SimulateGPUParticles(
 	}
 #endif
 
+	// Helper to handle transitions the current state render targets only when required
+	bool bIsCurrentTextureAccessRTV = false;
+
+	const auto PrepareCurrentRenderTargetsAccess =
+		[&]()
+		{
+			if (bIsCurrentTextureAccessRTV == false)
+			{
+				bIsCurrentTextureAccessRTV = true;
+				RHICmdList.Transition({
+					FRHITransitionInfo(CurrentStateTextures.PositionTextureRHI, ERHIAccess::SRVMask, ERHIAccess::RTV),
+					FRHITransitionInfo(CurrentStateTextures.VelocityTextureRHI, ERHIAccess::SRVMask, ERHIAccess::RTV)
+				});
+			}
+		};
+
+	const auto FinalizeCurrentRenderTargetsAccess =
+		[&]()
+		{
+			if (bIsCurrentTextureAccessRTV)
+			{
+				RHICmdList.Transition({
+					FRHITransitionInfo(CurrentStateTextures.PositionTextureRHI, ERHIAccess::RTV, ERHIAccess::SRVMask),
+					FRHITransitionInfo(CurrentStateTextures.VelocityTextureRHI, ERHIAccess::RTV, ERHIAccess::SRVMask)
+				});
+			}
+		};
+
 	{
 		// On some platforms, the textures are filled with garbage after creation, so we need to clear them to black the first time we use them
-		if ( !CurrentStateTextures.bTexturesCleared )
+		if (!CurrentStateTextures.bTexturesCleared)
 		{
 			SCOPED_GPU_MASK(RHICmdList, FRHIGPUMask::All());
 
-			FRHIRenderPassInfo RenderPassInfo(2, CurrentStateRenderTargets, ERenderTargetActions::Clear_Store);
-			RHICmdList.BeginRenderPass(RenderPassInfo, TEXT("GPUParticlesClearStateTextures"));
+			FRHIRenderPassInfo RenderPassInfo(UE_ARRAY_COUNT(CurrentStateRenderTargets), CurrentStateRenderTargets, ERenderTargetActions::Clear_Store);
+			RHICmdList.BeginRenderPass(RenderPassInfo, TEXT("PrepareGPUSimulation_GPUParticlesClearStateTextures"));
 			RHICmdList.EndRenderPass();
 
 			CurrentStateTextures.bTexturesCleared = true;
+			bIsCurrentTextureAccessRTV = true;
 		}
-		
-		if ( !PrevStateTextures.bTexturesCleared )
+
+		if (!PrevStateTextures.bTexturesCleared)
 		{
 			SCOPED_GPU_MASK(RHICmdList, FRHIGPUMask::All());
 
@@ -4864,7 +4864,8 @@ void FFXSystem::SimulateGPUParticles(
 			PrevStateTextures.bTexturesCleared = true;
 		}
 	}
-	
+
+
 	// Simulations that don't use vector fields can share some state.
 	FVectorFieldUniformBufferRef EmptyVectorFieldUniformBuffer;
 	{
@@ -5018,6 +5019,8 @@ void FFXSystem::SimulateGPUParticles(
 
 	if ( SimulationCommands.Num() || TilesToClear.Num())
 	{
+		PrepareCurrentRenderTargetsAccess();
+
 		FRHIRenderPassInfo RPInfo(2, CurrentStateRenderTargets, ERenderTargetActions::Load_Store);
 		RHICmdList.BeginRenderPass(RPInfo, TEXT("GPUParticles_SimulateAndClear"));
 		{
@@ -5115,11 +5118,8 @@ void FFXSystem::SimulateGPUParticles(
 #endif
 	}
 	
-	// finish current state render
-	RHICmdList.Transition({
-		FRHITransitionInfo(CurrentStateRenderTargets[0], ERHIAccess::RTV, ERHIAccess::SRVMask),
-		FRHITransitionInfo(CurrentStateRenderTargets[1], ERHIAccess::RTV, ERHIAccess::SRVMask)});
-
+	FinalizeCurrentRenderTargetsAccess();
+	
 	if (SimulationCommands.Num() && FixDeltaSeconds > 0)
 	{
 		//the fixed timestep works in two stages.  A first stage which simulates the fixed timestep and this second stage which simulates any remaining time from the actual delta time.  e.g.  fixed timestep of 16ms and actual dt of 23ms
