Add an ensure when attaching a garbage component to a live parent and vice versa

#Synced-CL 47063821

diff --git a/Engine/Source/Runtime/Engine/Private/Components/SceneComponent.cpp b/Engine/Source/Runtime/Engine/Private/Components/SceneComponent.cpp
--- a/Engine/Source/Runtime/Engine/Private/Components/SceneComponent.cpp
+++ b/Engine/Source/Runtime/Engine/Private/Components/SceneComponent.cpp
@@ -821,6 +821,7 @@ void USceneComponent::AttachToParent()
 	// so this call will just add this component to an AttachChildren array of a the Parent component
 	if (GetAttachParent())
 	{
+		ensureMsgf(IsValid(GetAttachParent()), TEXT("Trying to attach a scene component to a parent which is garbage"));
 		if (AttachToComponent(GetAttachParent(), FAttachmentTransformRules::KeepRelativeTransform, GetAttachSocketName()) == false)
 		{
 			// Failed to attach, we need to clear AttachParent so we don't think we're actually attached when we're not.
@@ -2327,6 +2328,14 @@ bool USceneComponent::AttachToComponent(USceneComponent* Parent, const FAttachme
 			return false;
 		}
 
+		{
+			bool ThisValid = IsValid(this);
+			bool ParentValid = IsValid(Parent);
+			// Allow attaching two garbage things to each other in case that's spammy with e.g. undo buffer stuff?
+			ensureMsgf(ThisValid == ParentValid, TEXT("Attempting to create attachment between garbage and non-garbage objects. Parent: %s Attachment: %s ParentValid: %d, AttachmentValid: %d"),
+				*Parent->GetPathName(), *GetPathName(), ParentValid, ThisValid);
+		}
+
 		AActor* MyActor = GetOwner();
 		AActor* TheirActor = Parent->GetOwner();
 
