Cloth - Added Cvar p.ClothPhysics.OnUpdateTransformTeleport to disable OnUpdateTransform teleports received from non clothing systems.

#Synced-CL 48079863

diff --git a/Engine/Source/Runtime/Engine/Private/SkeletalMeshComponentPhysics.cpp b/Engine/Source/Runtime/Engine/Private/SkeletalMeshComponentPhysics.cpp
--- a/Engine/Source/Runtime/Engine/Private/SkeletalMeshComponentPhysics.cpp
+++ b/Engine/Source/Runtime/Engine/Private/SkeletalMeshComponentPhysics.cpp
@@ -60,6 +60,7 @@ TAutoConsoleVariable<int32> CVarEnableClothPhysics(TEXT("p.ClothPhysics"), 1, TE
 static TAutoConsoleVariable<int32> CVarEnableClothPhysicsUseTaskThread(TEXT("p.ClothPhysics.UseTaskThread"), 1, TEXT("If 1, run cloth on the task thread. If 0, run on game thread."));
 static TAutoConsoleVariable<int32> CVarClothPhysicsTickWaitForParallelClothTask(TEXT("p.ClothPhysics.WaitForParallelClothTask"), 0, TEXT("If 1, always wait for cloth task completion in the Cloth Tick function. If 0, wait at end-of-frame updates instead if allowed by component settings"));
 static TAutoConsoleVariable<int32> CVarClothPhysicsTickWithMontages(TEXT("p.ClothPhysics.TickWithMontages"), 0, TEXT("Set to 1 to keep legacy behavior and tick clothing with montage related VisibilityBasedAnimTickOption."));
+static TAutoConsoleVariable<bool> CVarClothPhysicsOnUpdateTransformTeleport(TEXT("p.ClothPhysics.OnUpdateTransformTeleport"), true, TEXT("Allow transform changes to trigger cloth teleports."));
 
 static TAutoConsoleVariable<int32> CVarEnableKinematicDeferralPrePhysicsCondition(TEXT("p.EnableKinematicDeferralPrePhysicsCondition"), 1, TEXT("If is 1, and deferral would've been disallowed due to EUpdateTransformFlags, allow if in PrePhysics tick. If 0, condition is unchanged."));
 
@@ -1746,7 +1747,7 @@ void USkeletalMeshComponent::OnUpdateTransform(EUpdateTransformFlags UpdateTrans
 	}
 
 	// Mark the cloth simulation transform update as pending (the actual update will happen in the cloth thread if the simulation is running)
-	UpdateClothTransform(Teleport);
+	UpdateClothTransform(CVarClothPhysicsOnUpdateTransformTeleport.GetValueOnGameThread() ? Teleport : ETeleportType::None);
 
 	if (CVarUpdateRigidBodyScaling.GetValueOnGameThread())
 	{
