Fixing bugs from subsystem refactor.
Overrides of functions on FSubsystemCollection were not handled by the new function AddAndInitializeValidatedSubsystem.
Module loads during subsystem initialization should add any new classes to the pending initialization.

#Synced-CL Unknown

diff --git a/Engine/Source/Runtime/Engine/Classes/Engine/World.h b/Engine/Source/Runtime/Engine/Classes/Engine/World.h
--- a/Engine/Source/Runtime/Engine/Classes/Engine/World.h
+++ b/Engine/Source/Runtime/Engine/Classes/Engine/World.h
@@ -904,10 +904,12 @@ struct FDeferredComponentMoveData
 class FWorldSubsystemCollection : public FObjectSubsystemCollection<UWorldSubsystem>
 {
 public:
+	using Super = FObjectSubsystemCollection<UWorldSubsystem>;
+
 	FWorldSubsystemCollection();
 	virtual ~FWorldSubsystemCollection();
 private:
-	virtual USubsystem* AddAndInitializeSubsystem(UClass* SubsystemClass) override;
+	virtual void OnSubsystemAdded(TNonNullPtr<USubsystem> Subsystem) override;
 
 };
 
diff --git a/Engine/Source/Runtime/Engine/Private/Subsystems/SubsystemCollection.cpp b/Engine/Source/Runtime/Engine/Private/Subsystems/SubsystemCollection.cpp
--- a/Engine/Source/Runtime/Engine/Private/Subsystems/SubsystemCollection.cpp
+++ b/Engine/Source/Runtime/Engine/Private/Subsystems/SubsystemCollection.cpp
@@ -43,6 +43,11 @@ struct FSubsystemCollectionInitialization
 	// Value is the concrete class that will be returned
 	// So when key == value, this is a concrete class we intend to initialize
 	TMap<UClass*, UClass*> ClassMap;
+
+	// List of classes to be initialized. Classes may be initialized before we reach them in the queue via
+	// re-entrancy.
+	// We also may add more classes to the queue from re-entrancy, e.g. loading modules.
+	TArray<UClass*> Queue;
 };
 
 FSubsystemCollectionBase::FSubsystemCollectionBase()
@@ -365,6 +370,7 @@ void FSubsystemCollectionBase::AddAndInitializeSubsystems(TConstArrayView<UClass
 			continue;
 		}
 		LocalInitialization.ClassMap.Add(Class, Class);
+		LocalInitialization.Queue.Add(Class);
 	}
 
 	// Store a map of classes we intend to initialize for re-entrant calls to GetSubsystem 
@@ -392,16 +398,32 @@ void FSubsystemCollectionBase::AddAndInitializeSubsystems(TConstArrayView<UClass
 		}
 	}
 
-	// Only mark ourselves as populating now - some subsystems call GetSubsystem in their ShouldCreateSubsystem 
-	TGuardValue PopulatingGuard(Initialization, &LocalInitialization);
-	for (TPair<UClass*, UClass*> Pair : LocalInitialization.ClassMap)
+	// If we're re-entering initialization (e.g. loading a module during subsystem initialization) then add
+	// our classes to the existing set/queue and allow the outer loop to do the creation
+	if (Initialization)
 	{
-		// Only initialize the requested classes, other keys are for finding via the inheritence hierarchy 
-		// with re-entrant calls to GetSubsystem
-		// Skip creation if we created the subsystem re-entrantly
-		if (Pair.Key == Pair.Value && !SubsystemMap.Contains(Pair.Key))
+		for (const TPair<UClass*, UClass*>& Pair : LocalInitialization.ClassMap)
 		{
-			AddAndInitializeValidatedSubsystem(Pair.Key);
+			if (!Initialization->ClassMap.Contains(Pair.Key))
+			{
+				Initialization->ClassMap.Add(Pair.Key, Pair.Value);	
+			}
+		}
+		Initialization->Queue.Append(LocalInitialization.Queue);
+	}
+	else 
+	{
+		// Only mark ourselves as populating now - some subsystems call GetSubsystem in their ShouldCreateSubsystem 
+		TGuardValue PopulatingGuard(Initialization, &LocalInitialization);
+		for (TPair<UClass*, UClass*> Pair : LocalInitialization.ClassMap)
+		{
+			// Only initialize the requested classes, other keys are for finding via the inheritence hierarchy 
+			// with re-entrant calls to GetSubsystem
+			// Skip creation if we created the subsystem re-entrantly
+			if (Pair.Key == Pair.Value && !SubsystemMap.Contains(Pair.Key))
+			{
+				AddAndInitializeValidatedSubsystem(Pair.Key);
+			}
 		}
 	}
 }
@@ -467,6 +489,8 @@ USubsystem* FSubsystemCollectionBase::AddAndInitializeValidatedSubsystem(UClass*
 		}
 	}
 
+	OnSubsystemAdded(Subsystem);
+
 	return Subsystem;
 }
 
@@ -479,6 +503,8 @@ void FSubsystemCollectionBase::RemoveAndDeinitializeSubsystem(USubsystem* Subsys
 
 	const UClass* SubsystemClass = Subsystem->GetClass();
 
+	OnPreRemoveSubsystem(Subsystem);
+
 	TArray<UClass*> ClassesToRemove;
 
 	for (TPair<UClass*, TUniquePtr<FSubsystemArray>>& Pair : SubsystemArrayMap)
diff --git a/Engine/Source/Runtime/Engine/Private/World.cpp b/Engine/Source/Runtime/Engine/Private/World.cpp
--- a/Engine/Source/Runtime/Engine/Private/World.cpp
+++ b/Engine/Source/Runtime/Engine/Private/World.cpp
@@ -739,15 +739,11 @@ void FWorldPartitionEvents::BroadcastWorldPartitionUninitialized(UWorld* InWorld
 FWorldSubsystemCollection::FWorldSubsystemCollection() = default;
 FWorldSubsystemCollection::~FWorldSubsystemCollection() = default;
 
-USubsystem* FWorldSubsystemCollection::AddAndInitializeSubsystem(UClass* SubsystemClass)
+void FWorldSubsystemCollection::OnSubsystemAdded(TNonNullPtr<USubsystem> Subsystem)
 {
-	USubsystem* Subsystem = FObjectSubsystemCollection<UWorldSubsystem>::AddAndInitializeSubsystem(SubsystemClass);
-	if (!Subsystem)
-	{
-		return nullptr;
-	}
+	Super::OnSubsystemAdded(Subsystem);
 
-	UWorldSubsystem* WorldSubsystem = CastChecked<UWorldSubsystem>(Subsystem);
+	UWorldSubsystem* WorldSubsystem = CastChecked<UWorldSubsystem>(Subsystem.Get());
 	UWorld* World = CastChecked<UWorld>(GetOuter());
 
 	//run relevant functions we may have missed if we are initialized after the world has set up
@@ -762,7 +758,6 @@ USubsystem* FWorldSubsystemCollection::AddAndInitializeSubsystem(UClass* Subsyst
 		WorldSubsystem->OnWorldBeginPlay(*World);
 		WorldSubsystem->EnsureHasCalledBeginPlay();
 	}
-	return Subsystem;
 }
 
 
diff --git a/Engine/Source/Runtime/Engine/Public/Subsystems/SubsystemCollection.h b/Engine/Source/Runtime/Engine/Public/Subsystems/SubsystemCollection.h
--- a/Engine/Source/Runtime/Engine/Public/Subsystems/SubsystemCollection.h
+++ b/Engine/Source/Runtime/Engine/Public/Subsystems/SubsystemCollection.h
@@ -98,7 +98,7 @@ protected:
 	UObject* GetOuter() const { return Outer; }
 
 	/** Adds a subsystem to the collection and initializes it */
-	ENGINE_API virtual USubsystem* AddAndInitializeSubsystem(UClass* SubsystemClass);
+	ENGINE_API USubsystem* AddAndInitializeSubsystem(UClass* SubsystemClass);
 
 	/** Add and initialize many subsystems at once, handling dependencies between them as with Initialize */
 	ENGINE_API void AddAndInitializeSubsystems(TConstArrayView<UClass*> SubsystemClasses);
@@ -109,6 +109,13 @@ protected:
 	// True if we're adding subsystems either from Initialize or piecemeal
 	[[nodiscard]] bool IsPopulating() const;
 
+	/* Hooks for subclasses such as FWorldSubsystemCollection */
+	// Called after a subsystem has been initialized and added to the collection
+	virtual void OnSubsystemAdded(TNonNullPtr<USubsystem> Subsystem) { }
+
+	// Called before a subsystem is de-initialized and remove from the collection 
+	virtual void OnPreRemoveSubsystem(TNonNullPtr<USubsystem> Subsystem) { }
+
 private:
 
 	TMap<TObjectPtr<UClass>, TObjectPtr<USubsystem>> SubsystemMap;
