UE-160309 - Undoing parent deletion creates duplicate child actors

Not at all actors have a default root component. Without one, we can't correctly attach the child actor to the child actor component. This can lead to unattached child actors in the outliner. This can be fixed by simply spawning a default root component to use on the child actor.

Additionally, we need to move the Modify call to the end of CreateChildActor to ensure that we correctly take a snaphshot of the actor with its new root component.

#Synced-CL 47364356

diff --git a/Engine/Source/Runtime/Engine/Private/Components/ChildActorComponent.cpp b/Engine/Source/Runtime/Engine/Private/Components/ChildActorComponent.cpp
--- a/Engine/Source/Runtime/Engine/Private/Components/ChildActorComponent.cpp
+++ b/Engine/Source/Runtime/Engine/Private/Components/ChildActorComponent.cpp
@@ -170,7 +170,7 @@ void UChildActorComponent::Serialize(FArchive& Ar)
 			}
 		}
 	}
-	// Since we sometimes serialize properties in instead of using duplication and we can end up pointing at the wrong template
+	// Since we sometimes serialize properties instead of using duplication and we can end up pointing at the wrong template
 	else if (!Ar.IsPersistent() && ChildActorTemplate)
 	{
 		if (IsTemplate())
@@ -760,10 +760,10 @@ void UChildActorComponent::CreateChildActor(TFunction<void(AActor*)> CustomizerF
 	DestroyChildActor();
 
 	// If we have a class to spawn.
-	if(ChildActorClass != nullptr)
+	if (ChildActorClass != nullptr)
 	{
 		UWorld* World = GetWorld();
-		if(World != nullptr)
+		if (World != nullptr)
 		{
 			// If we're in a replay scrub and this is a startup actor and the class exists, assume the reference will be restored later
 			if (World->IsPlayingReplay() && World->GetDemoNetDriver()->IsRestoringStartupActors() && MyOwner && MyOwner->IsNetStartupActor() && ChildActorClass)
@@ -824,11 +824,10 @@ void UChildActorComponent::CreateChildActor(TFunction<void(AActor*)> CustomizerF
 					Params.Template = ChildActorTemplate;
 				}
 				Params.ObjectFlags |= (RF_TextExportTransient | RF_DuplicateTransient);
-				if (!HasAllFlags(RF_Transactional))
-				{
-					Params.ObjectFlags &= ~RF_Transactional;
-				}
-
+				
+				// If necessary, the actor will eventually be saved to the transaction buffer,
+				// but only after we've finished initializing it manually.
+				Params.ObjectFlags &= ~RF_Transactional;
 
 				// Spawn actor of desired class
 				ConditionalUpdateComponentToWorld();
@@ -837,7 +836,7 @@ void UChildActorComponent::CreateChildActor(TFunction<void(AActor*)> CustomizerF
 				ChildActor = World->SpawnActor(ChildActorClass, &Location, &Rotation, Params);
 
 				// If spawn was successful, 
-				if(ChildActor != nullptr)
+				if (ChildActor != nullptr)
 				{
 					if (IsEditorOnly() || (MyOwner && MyOwner->IsEditorOnly()))
 					{
@@ -860,10 +859,19 @@ void UChildActorComponent::CreateChildActor(TFunction<void(AActor*)> CustomizerF
 					const FComponentInstanceDataCache* ComponentInstanceData = (CachedInstanceData ? CachedInstanceData->ComponentInstanceData.Get() : nullptr);
 					ChildActor->FinishSpawning(GetComponentTransform(), false, ComponentInstanceData);
 
-					if (USceneComponent* ChildRoot = ChildActor->GetRootComponent())
+					USceneComponent* ChildRootComponent = ChildActor->GetRootComponent();
+					
+					// Since child actors are expected to be attached to parent actors, they must have a root component.
+					if (!ChildRootComponent)
 					{
-						TGuardValue<TEnumAsByte<EComponentMobility::Type>> MobilityGuard(ChildRoot->Mobility, Mobility);
-						ChildRoot->AttachToComponent(this, FAttachmentTransformRules::SnapToTargetNotIncludingScale);
+						ChildRootComponent = NewObject<USceneComponent>(ChildActor, USceneComponent::GetDefaultSceneRootVariableName());
+						ChildActor->SetRootComponent(ChildRootComponent);
+						ChildRootComponent->RegisterComponent();
+					}
+
+					{
+						TGuardValue<TEnumAsByte<EComponentMobility::Type>> MobilityGuard(ChildRootComponent->Mobility, Mobility);
+						ChildRootComponent->AttachToComponent(this, FAttachmentTransformRules::SnapToTargetNotIncludingScale);
 					}
 
 					SetIsReplicated(ChildActor->GetIsReplicated());
@@ -882,6 +890,15 @@ void UChildActorComponent::CreateChildActor(TFunction<void(AActor*)> CustomizerF
 						}
 					}
 
+					// Must be called last: ensures all relevant state is saved to the transaction bufffer
+					const bool bIsTransactional = HasAllFlags(RF_Transactional);
+					if (bIsTransactional)
+					{
+						ChildActor->SetFlags(RF_Transactional);
+						ChildRootComponent->SetFlags(RF_Transactional);
+						ChildActor->Modify();
+						ChildRootComponent->Modify();
+					}
 				}
 			}
 		}
