Disable new ensures about attaching garbage SceneComponents to non-garbage ones and vice-versa.

#Synced-CL 47063821

diff --git a/Engine/Source/Runtime/Engine/Private/Components/SceneComponent.cpp b/Engine/Source/Runtime/Engine/Private/Components/SceneComponent.cpp
--- a/Engine/Source/Runtime/Engine/Private/Components/SceneComponent.cpp
+++ b/Engine/Source/Runtime/Engine/Private/Components/SceneComponent.cpp
@@ -76,6 +76,16 @@ namespace SceneComponentCVars
 		TEXT("false: fix up any children that are missing the parent, true: fixes up all children except non-replicated root components (default)"),
 		ECVF_Default);
 
+#if DO_ENSURE
+	bool bCheckGarbageAttachments = false;
+	static FAutoConsoleVariableRef CVarCheckGarbageAttachments(
+		TEXT("s.CheckGarbageAttachments"),
+		bCheckGarbageAttachments,
+		TEXT("Whether or not to have ensures/assertions regarding attaching child components that are marked garbage to parents that aren't and vice versa"),
+		ECVF_Default
+	);
+#endif
+
 #if UE_WITH_REMOTE_OBJECT_HANDLE
 	bool bAllowRemoteAttachedChildren = true;
 	static FAutoConsoleVariableRef AllowRemoteAttachedChildren(TEXT("DSTM.AllowRemoteAttachedChildren"), bAllowRemoteAttachedChildren, TEXT("Do we allow AttachChildren of a SceneComponent to exist only Remotely?"));
@@ -821,7 +831,7 @@ void USceneComponent::AttachToParent()
 	// so this call will just add this component to an AttachChildren array of a the Parent component
 	if (GetAttachParent())
 	{
-		ensureMsgf(IsValid(GetAttachParent()), TEXT("Trying to attach a scene component to a parent which is garbage"));
+		ensureMsgf(!SceneComponentCVars::bCheckGarbageAttachments || IsValid(GetAttachParent()), TEXT("Trying to attach a scene component to a parent which is garbage"));
 		if (AttachToComponent(GetAttachParent(), FAttachmentTransformRules::KeepRelativeTransform, GetAttachSocketName()) == false)
 		{
 			// Failed to attach, we need to clear AttachParent so we don't think we're actually attached when we're not.
@@ -2328,6 +2338,7 @@ bool USceneComponent::AttachToComponent(USceneComponent* Parent, const FAttachme
 			return false;
 		}
 
+		if (SceneComponentCVars::bCheckGarbageAttachments)
 		{
 			bool ThisValid = IsValid(this);
 			bool ParentValid = IsValid(Parent);
