Add missing DO_ENSURE wrapping around use of cvar

#Synced-CL 47063821

diff --git a/Engine/Source/Runtime/Engine/Private/Components/SceneComponent.cpp b/Engine/Source/Runtime/Engine/Private/Components/SceneComponent.cpp
--- a/Engine/Source/Runtime/Engine/Private/Components/SceneComponent.cpp
+++ b/Engine/Source/Runtime/Engine/Private/Components/SceneComponent.cpp
@@ -2338,6 +2338,7 @@ bool USceneComponent::AttachToComponent(USceneComponent* Parent, const FAttachme
 			return false;
 		}
 
+#if DO_ENSURE
 		if (SceneComponentCVars::bCheckGarbageAttachments)
 		{
 			bool ThisValid = IsValid(this);
@@ -2346,6 +2347,7 @@ bool USceneComponent::AttachToComponent(USceneComponent* Parent, const FAttachme
 			ensureMsgf(ThisValid == ParentValid, TEXT("Attempting to create attachment between garbage and non-garbage objects. Parent: %s Attachment: %s ParentValid: %d, AttachmentValid: %d"),
 				*Parent->GetPathName(), *GetPathName(), ParentValid, ThisValid);
 		}
+#endif // DO_ENSURE
 
 		AActor* MyActor = GetOwner();
 		AActor* TheirActor = Parent->GetOwner();
