Fix crash in FSkeletalMeshSceneProxy::HasRayTracingRepresentation().

#Synced-CL 48396324

diff --git a/Engine/Source/Runtime/Engine/Private/SkeletalMeshSceneProxy.cpp b/Engine/Source/Runtime/Engine/Private/SkeletalMeshSceneProxy.cpp
--- a/Engine/Source/Runtime/Engine/Private/SkeletalMeshSceneProxy.cpp
+++ b/Engine/Source/Runtime/Engine/Private/SkeletalMeshSceneProxy.cpp
@@ -141,6 +141,7 @@ FSkeletalMeshSceneProxy::FSkeletalMeshSceneProxy(const FSkinnedMeshSceneProxyDes
 	NumLODs = RenderData->LODRenderData.Num();
 
 #if RHI_RAYTRACING
+	bHasRayTracingRepresentation = RenderData->bSupportRayTracing;
 	bRayTraceStatic = bRenderStatic || bInstancedSkinnedMesh;
 #endif
 
@@ -975,7 +976,7 @@ void FSkeletalMeshSceneProxy::DestroyRenderThreadResources()
 #if RHI_RAYTRACING
 bool FSkeletalMeshSceneProxy::HasRayTracingRepresentation() const
 {
-	return RenderData->bSupportRayTracing
+	return bHasRayTracingRepresentation
 		&& (bRayTraceStatic || (CVarRayTracingSupportSkeletalMeshes.GetValueOnAnyThread() != 0 && CVarRayTracingSkeletalMeshes.GetValueOnAnyThread() != 0));
 }
 
@@ -987,7 +988,7 @@ RayTracing::FGeometryGroupHandle FSkeletalMeshSceneProxy::GetRayTracingGeometryG
 
 TArray<FRayTracingGeometry*> FSkeletalMeshSceneProxy::GetStaticRayTracingGeometries() const
 {
-	if (IsRayTracingEnabled() && bRayTraceStatic && RenderData->bSupportRayTracing)
+	if (IsRayTracingEnabled() && bHasRayTracingRepresentation && bRayTraceStatic)
 	{
 		TArray<FRayTracingGeometry*> RayTracingGeometries;
 		RayTracingGeometries.AddDefaulted(RenderData->LODRenderData.Num());
diff --git a/Engine/Source/Runtime/Engine/Public/SkeletalMeshSceneProxy.h b/Engine/Source/Runtime/Engine/Public/SkeletalMeshSceneProxy.h
--- a/Engine/Source/Runtime/Engine/Public/SkeletalMeshSceneProxy.h
+++ b/Engine/Source/Runtime/Engine/Public/SkeletalMeshSceneProxy.h
@@ -180,6 +180,7 @@ protected:
 	uint8 bCanHighlightSelectedSections : 1;
 	uint8 bRenderStatic:1;
 #if RHI_RAYTRACING
+	uint8 bHasRayTracingRepresentation : 1;
 	uint8 bRayTraceStatic : 1;
 #endif
 #if !(UE_BUILD_SHIPPING || UE_BUILD_TEST)
