Remove more uses of FRHICommandListExecutor::GetImmediateCommandList(), replacing them with either FRHICommandListImmediate::Get(), or handing down an existing RHICmdList.

#Synced-CL 48409310

diff --git a/Engine/Plugins/FX/Niagara/Source/Niagara/Classes/NiagaraDataInterfaceArrayImpl.h b/Engine/Plugins/FX/Niagara/Source/Niagara/Classes/NiagaraDataInterfaceArrayImpl.h
--- a/Engine/Plugins/FX/Niagara/Source/Niagara/Classes/NiagaraDataInterfaceArrayImpl.h
+++ b/Engine/Plugins/FX/Niagara/Source/Niagara/Classes/NiagaraDataInterfaceArrayImpl.h
@@ -408,7 +408,7 @@ struct FNDIArrayProxyImpl : public INDIArrayProxyBase
 		{
 			if ( FNDIArrayInstanceData_RenderThread<TArrayType>* InstanceData_RT = PerInstanceData_RenderThread.Find(InstanceID) )
 			{
-				FRHICommandListImmediate& RHICmdList = FRHICommandListExecutor::GetImmediateCommandList();
+				FRHICommandListImmediate& RHICmdList = FRHICommandListImmediate::Get();
 				InstanceData_RT->UpdateData(RHICmdList, GameToRenderInstanceData->ArrayData);
 			}
 		}
diff --git a/Engine/Plugins/Media/ElectraUtil/Source/ElectraSamples/Private/Android/ElectraTextureSample_Android.cpp b/Engine/Plugins/Media/ElectraUtil/Source/ElectraSamples/Private/Android/ElectraTextureSample_Android.cpp
--- a/Engine/Plugins/Media/ElectraUtil/Source/ElectraSamples/Private/Android/ElectraTextureSample_Android.cpp
+++ b/Engine/Plugins/Media/ElectraUtil/Source/ElectraSamples/Private/Android/ElectraTextureSample_Android.cpp
@@ -813,7 +813,7 @@ void FElectraTextureSample::CleanupImageResources()
 		}
 		else if (IsInRenderingThread())
 		{
-			FRHICommandListExecutor::GetImmediateCommandList().EnqueueLambda([Resources, ReleaseFN](FRHICommandListImmediate&)
+			FRHICommandListImmediate::Get().EnqueueLambda([Resources, ReleaseFN](FRHICommandListImmediate&)
 			{
 				CleanupImageResourcesJNI(Resources, ReleaseFN);
 			});
diff --git a/Engine/Plugins/Runtime/HairStrands/Source/HairStrandsCore/Private/GroomTextureBuilder.cpp b/Engine/Plugins/Runtime/HairStrands/Source/HairStrandsCore/Private/GroomTextureBuilder.cpp
--- a/Engine/Plugins/Runtime/HairStrands/Source/HairStrandsCore/Private/GroomTextureBuilder.cpp
+++ b/Engine/Plugins/Runtime/HairStrands/Source/HairStrandsCore/Private/GroomTextureBuilder.cpp
@@ -1276,7 +1276,7 @@ static void CopyTextures(FRDGBuilder& GraphBuilder, FGlobalShaderMap* ShaderMap,
 	const uint32 TextureCount = GetHairTextureLayoutTextureCount(InLayout);
 	for (uint32 TextureIt=0; TextureIt<TextureCount;++TextureIt)
 	{
-		// Split texture array into seperate texture
+		// Split texture array into separate texture
 		FRDGTextureRef Texture = AddTextureCopyPass(GraphBuilder, ShaderMap, TextureIt, In.Texture);
 
 		// Enqueue copy to readback texture
@@ -1287,8 +1287,6 @@ static void CopyTextures(FRDGBuilder& GraphBuilder, FGlobalShaderMap* ShaderMap,
 static void CopyReadbackToTexture(TUniquePtr<FRHIGPUTextureReadback>& In, UTexture2D* Out)
 {
 #if WITH_EDITORONLY_DATA
-	FRHICommandListImmediate& RHICmdList = FRHICommandListExecutor::GetImmediateCommandList();
-
 	uint32 BytePerPixels = 0;
 	uint32 Width = Out->Source.GetSizeX();
 	uint32 Height = Out->Source.GetSizeY();
diff --git a/Engine/Plugins/VirtualProduction/Rivermax/RivermaxMedia/Source/RivermaxMedia/Private/RivermaxMediaPlayer.cpp b/Engine/Plugins/VirtualProduction/Rivermax/RivermaxMedia/Source/RivermaxMedia/Private/RivermaxMediaPlayer.cpp
--- a/Engine/Plugins/VirtualProduction/Rivermax/RivermaxMedia/Source/RivermaxMedia/Private/RivermaxMediaPlayer.cpp
+++ b/Engine/Plugins/VirtualProduction/Rivermax/RivermaxMedia/Source/RivermaxMedia/Private/RivermaxMediaPlayer.cpp
@@ -606,7 +606,7 @@ namespace UE::RivermaxMedia
 			return false;
 		}
 
-		FRDGBuilder GraphBuilder((FRHICommandListExecutor::GetImmediateCommandList()));
+		FRDGBuilder GraphBuilder(RHICmdList);
 		if (ConverterSetup.PreConvertFunc)
 		{
 			ConverterSetup.PreConvertFunc(GraphBuilder);
diff --git a/Engine/Plugins/VirtualProduction/TextureShare/Source/TextureShare/Private/Game/WorldSubsystem/TextureShareWorldSubsystemContext.cpp b/Engine/Plugins/VirtualProduction/TextureShare/Source/TextureShare/Private/Game/WorldSubsystem/TextureShareWorldSubsystemContext.cpp
--- a/Engine/Plugins/VirtualProduction/TextureShare/Source/TextureShare/Private/Game/WorldSubsystem/TextureShareWorldSubsystemContext.cpp
+++ b/Engine/Plugins/VirtualProduction/TextureShare/Source/TextureShare/Private/Game/WorldSubsystem/TextureShareWorldSubsystemContext.cpp
@@ -319,7 +319,7 @@ void FTextureShareWorldSubsystemContext::OnTextureShareBackBufferReadyToPresent_
 	using namespace UE::TextureShareWorldSubsystemContext;
 	UE_TS_LOG(LogTextureShareWorldSubsystem, Log, TEXT("%s:OnTextureShareBackBufferReadyToPresent_RenderThread()"), *ObjectProxy.GetName_RenderThread());
 
-	FRHICommandListImmediate& RHICmdList = FRHICommandListExecutor::GetImmediateCommandList();
+	FRHICommandListImmediate& RHICmdList = FRHICommandListImmediate::Get();
 
 	if (ObjectProxy.IsFrameSyncActive_RenderThread())
 	{
diff --git a/Engine/Source/Runtime/Engine/Private/PrecomputedVolumetricLightmap.cpp b/Engine/Source/Runtime/Engine/Private/PrecomputedVolumetricLightmap.cpp
--- a/Engine/Source/Runtime/Engine/Private/PrecomputedVolumetricLightmap.cpp
+++ b/Engine/Source/Runtime/Engine/Private/PrecomputedVolumetricLightmap.cpp
@@ -258,7 +258,7 @@ ENGINE_API void FPrecomputedVolumetricLightmapData::InitRHI(FRHICommandListBase&
 			BrickData.DirectionalLightShadowing.CreateTexture(RHICmdList, BrickDataDimensions);
 		}
 
-		GVolumetricLightmapBrickAtlas.Insert(INT_MAX, this);
+		GVolumetricLightmapBrickAtlas.Insert(RHICmdList.GetAsImmediate(), INT_MAX, this);
 
 		// It is now safe to release the brick data used for upload. They will stay in GPU memory until UMapBuildDataRegistry::BeginDestroy().
 		BrickData.ReleaseRHI();
@@ -313,7 +313,7 @@ ENGINE_API void FPrecomputedVolumetricLightmapData::ReleaseRHI()
 {
 	if (GetFeatureLevel() >= ERHIFeatureLevel::SM5)
 	{
-		GVolumetricLightmapBrickAtlas.Remove(this);
+		GVolumetricLightmapBrickAtlas.Remove(FRHICommandListImmediate::Get(), this);
 
 		{
 			IndirectionTexture.Texture.SafeRelease();
@@ -417,15 +417,13 @@ inline FIntVector ComputeBrickLayoutPosition(int32 BrickLayoutAllocation, FIntVe
 	return BrickPosition;
 }
 
-ENGINE_API void FPrecomputedVolumetricLightmapData::AddToSceneData(FPrecomputedVolumetricLightmapData* SceneData)
+ENGINE_API void FPrecomputedVolumetricLightmapData::AddToSceneData(FRHICommandListImmediate& RHICmdList, FPrecomputedVolumetricLightmapData* SceneData)
 {
 	if (SceneDataAdded.Find(SceneData) != INDEX_NONE)
 	{
 		return;
 	}
 
-	FRHICommandListImmediate& RHICmdList = FRHICommandListExecutor::GetImmediateCommandList();
-
 	if (IndirectionTextureDimensions.GetMax() > 0)
 	{
 		SceneDataAdded.Add(SceneData);
@@ -584,7 +582,7 @@ ENGINE_API void FPrecomputedVolumetricLightmapData::AddToSceneData(FPrecomputedV
 	}
 }
 
-ENGINE_API void FPrecomputedVolumetricLightmapData::RemoveFromSceneData(FPrecomputedVolumetricLightmapData* SceneData, int32 PersistentLevelBrickDataBaseOffset)
+ENGINE_API void FPrecomputedVolumetricLightmapData::RemoveFromSceneData(FRHICommandListImmediate& RHICmdList, FPrecomputedVolumetricLightmapData* SceneData, int32 PersistentLevelBrickDataBaseOffset)
 {
 	if (SceneDataAdded.Find(SceneData) == INDEX_NONE)
 	{
@@ -593,8 +591,6 @@ ENGINE_API void FPrecomputedVolumetricLightmapData::RemoveFromSceneData(FPrecomp
 
 	SceneDataAdded.Remove(SceneData);
 
-	FRHICommandListImmediate& RHICmdList = FRHICommandListExecutor::GetImmediateCommandList();
-
 	if (IndirectionTextureDimensions.GetMax() > 0)
 	{
 		// Do nothing, as when a VLM data with indirection texture is being destroyed, the persistent level is going away
@@ -969,12 +965,10 @@ void CopyDataIntoAtlas(FRHICommandList& RHICmdList, ERHIFeatureLevel::Type Featu
 	RHICmdList.Transition(Infos);
 }
 
-void FVolumetricLightmapBrickAtlas::Insert(int32 Index, FPrecomputedVolumetricLightmapData* Data)
+void FVolumetricLightmapBrickAtlas::Insert(FRHICommandListImmediate& RHICmdList, int32 Index, FPrecomputedVolumetricLightmapData* Data)
 {
 	check(!Allocations.FindByPredicate([Data](const Allocation& Other) { return Other.Data == Data; }));
 
-	FRHICommandListImmediate& RHICmdList = FRHICommandListExecutor::GetImmediateCommandList();
-
 	bool bReadAfterCreate = false;
 
 	if (!bInitialized)
@@ -1118,7 +1112,7 @@ void FVolumetricLightmapBrickAtlas::Insert(int32 Index, FPrecomputedVolumetricLi
 	TextureSet = NewTextureSet; // <-- Old texture references are released here
 }
 
-void FVolumetricLightmapBrickAtlas::Remove(FPrecomputedVolumetricLightmapData* Data)
+void FVolumetricLightmapBrickAtlas::Remove(FRHICommandListImmediate& RHICmdList, FPrecomputedVolumetricLightmapData* Data)
 {
 	Allocation* AllocationEntry = Allocations.FindByPredicate([Data](const Allocation& Other) { return Other.Data == Data; });
 	if (!AllocationEntry)
@@ -1157,8 +1151,6 @@ void FVolumetricLightmapBrickAtlas::Remove(FPrecomputedVolumetricLightmapData* D
 			NewTextureSet.Initialize(BrickDataDimensions, TextureSet);
 		}
 
-		FRHICommandListImmediate& RHICmdList = FRHICommandListExecutor::GetImmediateCommandList();
-
 		{
 			int32 BrickStartAllocation = 0;
 
@@ -1199,7 +1191,7 @@ void FVolumetricLightmapBrickAtlas::ReleaseRHI()
 	TArray<Allocation> AllocationList = Allocations;
 	for (Allocation& Alloc : AllocationList)
 	{
-		Remove(Alloc.Data);
+		Remove(FRHICommandListImmediate::Get(), Alloc.Data);
 	}
 
 	TextureSet.Release();
diff --git a/Engine/Source/Runtime/Engine/Public/PrecomputedVolumetricLightmap.h b/Engine/Source/Runtime/Engine/Public/PrecomputedVolumetricLightmap.h
--- a/Engine/Source/Runtime/Engine/Public/PrecomputedVolumetricLightmap.h
+++ b/Engine/Source/Runtime/Engine/Public/PrecomputedVolumetricLightmap.h
@@ -172,8 +172,22 @@ public:
 	}
 
 	ENGINE_API void HandleDataMovementInAtlas(FRHIComputeCommandList& RHICmdList, int32 OldOffset, int32 NewOffset);
-	ENGINE_API void AddToSceneData(FPrecomputedVolumetricLightmapData* SceneData);
-	ENGINE_API void RemoveFromSceneData(FPrecomputedVolumetricLightmapData* SceneData, int32 PersistentLevelBrickDataBaseOffset);
+	
+	ENGINE_API void AddToSceneData(FRHICommandListImmediate& RHICmdList, FPrecomputedVolumetricLightmapData* SceneData);
+	
+	UE_DEPRECATED(5.8, "AddToSceneData now requires a command list.")
+	void AddToSceneData(FPrecomputedVolumetricLightmapData* SceneData)
+	{
+		AddToSceneData(FRHICommandListImmediate::Get(), SceneData);
+	}
+
+	ENGINE_API void RemoveFromSceneData(FRHICommandListImmediate& RHICmdList, FPrecomputedVolumetricLightmapData* SceneData, int32 PersistentLevelBrickDataBaseOffset);
+
+	UE_DEPRECATED(5.8, "RemoveFromSceneData now requires a command list.")
+	void RemoveFromSceneData(FPrecomputedVolumetricLightmapData* SceneData, int32 PersistentLevelBrickDataBaseOffset)
+	{
+		RemoveFromSceneData(FRHICommandListImmediate::Get(), SceneData, PersistentLevelBrickDataBaseOffset);
+	}
 
 	SIZE_T GetAllocatedBytes() const;
 
@@ -583,8 +597,20 @@ public:
 
 	TArray<Allocation> Allocations;
 
-	ENGINE_API void Insert(int32 Index, FPrecomputedVolumetricLightmapData* Data);
-	ENGINE_API void Remove(FPrecomputedVolumetricLightmapData* Data);
+	ENGINE_API void Insert(FRHICommandListImmediate& RHICmdList, int32 Index, FPrecomputedVolumetricLightmapData* Data);
+	ENGINE_API void Remove(FRHICommandListImmediate& RHICmdList, FPrecomputedVolumetricLightmapData* Data);
+
+	UE_DEPRECATED(5.8, "Insert now requires a command list.")
+	void Insert(int32 Index, FPrecomputedVolumetricLightmapData* Data)
+	{
+		Insert(FRHICommandListImmediate::Get(), Index, Data);
+	}
+
+	UE_DEPRECATED(5.8, "Remove now requires a command list.")
+	void Remove(FPrecomputedVolumetricLightmapData* Data)
+	{
+		Remove(FRHICommandListImmediate::Get(), Data);
+	}
 
 private:
 	bool bInitialized;
diff --git a/Engine/Source/Runtime/RHI/Private/DynamicRHI.cpp b/Engine/Source/Runtime/RHI/Private/DynamicRHI.cpp
--- a/Engine/Source/Runtime/RHI/Private/DynamicRHI.cpp
+++ b/Engine/Source/Runtime/RHI/Private/DynamicRHI.cpp
@@ -437,7 +437,7 @@ void RHIExit()
 		PipelineStateCache::Shutdown();
 
 		// Flush any potential commands queued before we shut things down.
-		FRHICommandListExecutor::GetImmediateCommandList().ImmediateFlush(EImmediateFlushType::FlushRHIThread);
+		FRHICommandListImmediate::Get().ImmediateFlush(EImmediateFlushType::FlushRHIThread);
 
 		// Destruct the dynamic RHI.
 		GDynamicRHI->Shutdown();
@@ -451,8 +451,8 @@ void RHIExit()
 	else if (GUsingNullRHI)
 	{
 		// If we are using NullRHI flush the command list here in case somethings has been added to the command list during exit calls
-		FRHICommandListExecutor::GetImmediateCommandList().ImmediateFlush(EImmediateFlushType::FlushRHIThreadFlushResources);
-		FRHICommandListExecutor::GetImmediateCommandList().ImmediateFlush(EImmediateFlushType::FlushRHIThread);
+		FRHICommandListImmediate::Get().ImmediateFlush(EImmediateFlushType::FlushRHIThreadFlushResources);
+		FRHICommandListImmediate::Get().ImmediateFlush(EImmediateFlushType::FlushRHIThread);
 	}
 
 	GRHICommandList.CleanupGraphEvents();
diff --git a/Engine/Source/Runtime/RHI/Public/DynamicRHI.h b/Engine/Source/Runtime/RHI/Public/DynamicRHI.h
--- a/Engine/Source/Runtime/RHI/Public/DynamicRHI.h
+++ b/Engine/Source/Runtime/RHI/Public/DynamicRHI.h
@@ -380,7 +380,7 @@ public:
 	/**
 	* Creates a bound shader state instance which encapsulates a decl, vertex shader and pixel shader
 	* CAUTION: Even though this is marked as threadsafe, it is only valid to call from the render thread or the RHI thread. It need not be threadsafe unless the RHI support parallel translation.
-	* CAUTION: Platforms that support RHIThread but don't actually have a threadsafe implementation must flush internally with FScopedRHIThreadStaller StallRHIThread(FRHICommandListExecutor::GetImmediateCommandList()); when the call is from the render thread
+	* CAUTION: Platforms that support RHIThread but don't actually have a threadsafe implementation must flush internally with FScopedRHIThreadStaller StallRHIThread(FRHICommandListImmediate::Get()); when the call is from the render thread
 	* @param VertexDeclaration - existing vertex decl
 	* @param VertexShader - existing vertex shader
 	* @param GeometryShader - existing geometry shader
@@ -393,7 +393,7 @@ public:
 	/**
 	* Creates a bound shader state instance which encapsulates an amplification shader, a mesh shader, and pixel shader
 	* CAUTION: Even though this is marked as threadsafe, it is only valid to call from the render thread or the RHI thread. It need not be threadsafe unless the RHI support parallel translation.
-	* CAUTION: Platforms that support RHIThread but don't actually have a threadsafe implementation must flush internally with FScopedRHIThreadStaller StallRHIThread(FRHICommandListExecutor::GetImmediateCommandList()); when the call is from the render thread
+	* CAUTION: Platforms that support RHIThread but don't actually have a threadsafe implementation must flush internally with FScopedRHIThreadStaller StallRHIThread(FRHICommandListImmediate::Get()); when the call is from the render thread
 	* @param AmplificationShader - existing amplification shader
 	* @param MeshShader - existing mesh shader
 	* @param PixelShader - existing pixel shader
@@ -409,7 +409,7 @@ public:
 	* @return FGraphicsPipelineStateRHIRef that can be bound for rendering; nullptr if the compilation fails.
 	* CAUTION: On certain RHI implementations (eg, ones that do not support runtime compilation) a compilation failure is a Fatal error and this function will not return.
 	* CAUTION: Even though this is marked as threadsafe, it is only valid to call from the render thread or the RHI thread. It need not be threadsafe unless the RHI support parallel translation.
-	* CAUTION: Platforms that support RHIThread but don't actually have a threadsafe implementation must flush internally with FScopedRHIThreadStaller StallRHIThread(FRHICommandListExecutor::GetImmediateCommandList()); when the call is from the render thread
+	* CAUTION: Platforms that support RHIThread but don't actually have a threadsafe implementation must flush internally with FScopedRHIThreadStaller StallRHIThread(FRHICommandListImmediate::Get()); when the call is from the render thread
 	*/
 	// FlushType: Thread safe
 	virtual FGraphicsPipelineStateRHIRef RHICreateGraphicsPipelineState(const FGraphicsPipelineStateInitializer& Initializer) = 0;
diff --git a/Engine/Source/Runtime/RHI/Public/RHIValidation.h b/Engine/Source/Runtime/RHI/Public/RHIValidation.h
--- a/Engine/Source/Runtime/RHI/Public/RHIValidation.h
+++ b/Engine/Source/Runtime/RHI/Public/RHIValidation.h
@@ -337,48 +337,18 @@ public:
 		RHI->RHIUnmapStagingSurface_RenderThread(RHICmdList, Texture, GPUIndex);
 	}
 
-	/**
-	* Creates a bound shader state instance which encapsulates a decl, vertex shader and pixel shader
-	* CAUTION: Even though this is marked as threadsafe, it is only valid to call from the render thread or the RHI thread. It need not be threadsafe unless the RHI support parallel translation.
-	* CAUTION: Platforms that support RHIThread but don't actually have a threadsafe implementation must flush internally with FScopedRHIThreadStaller StallRHIThread(FRHICommandListExecutor::GetImmediateCommandList()); when the call is from the render thread
-	* @param VertexDeclaration - existing vertex decl
-	* @param VertexShader - existing vertex shader
-	* @param GeometryShader - existing geometry shader
-	* @param PixelShader - existing pixel shader
-	*/
-	// FlushType: Thread safe, but varies depending on the RHI
 	virtual FBoundShaderStateRHIRef RHICreateBoundShaderState(FRHIVertexDeclaration* VertexDeclaration, FRHIVertexShader* VertexShader, FRHIPixelShader* PixelShader, FRHIGeometryShader* GeometryShader) override final
 	{
 		return RHI->RHICreateBoundShaderState(VertexDeclaration, VertexShader, PixelShader, GeometryShader);
 	}
 
 #if PLATFORM_SUPPORTS_MESH_SHADERS && PLATFORM_USE_FALLBACK_PSO
-	/**
-	* Creates a bound shader state instance which encapsulates an amplification shader, a mesh shader, and pixel shader
-	* CAUTION: Even though this is marked as threadsafe, it is only valid to call from the render thread or the RHI thread. It need not be threadsafe unless the RHI support parallel translation.
-	* CAUTION: Platforms that support RHIThread but don't actually have a threadsafe implementation must flush internally with FScopedRHIThreadStaller StallRHIThread(FRHICommandListExecutor::GetImmediateCommandList()); when the call is from the render thread
-	* @param AmplificationShader - existing amplification shader
-	* @param MeshShader - existing mesh shader
-	* @param PixelShader - existing pixel shader
-	*/
-	// FlushType: Thread safe, but varies depending on the RHI
 	virtual FBoundShaderStateRHIRef RHICreateBoundShaderState(FRHIAmplificationShader* AmplificationShader, FRHIMeshShader* MeshShader, FRHIPixelShader* PixelShader) override final
 	{
 		return RHI->RHICreateBoundShaderState(AmplificationShader, MeshShader, PixelShader);
 	}
 #endif
 
-	/**
-	* Creates a graphics pipeline state object (PSO) that represents a complete gpu pipeline for rendering.
-	* This function should be considered expensive to call at runtime and may cause hitches as pipelines are compiled.
-	* @param Initializer - Descriptor object defining all the information needed to create the PSO, as well as behavior hints to the RHI.
-	* @return FGraphicsPipelineStateRHIRef that can be bound for rendering; nullptr if the compilation fails.
-	* CAUTION: On certain RHI implementations (eg, ones that do not support runtime compilation) a compilation failure is a Fatal error and this function will not return.
-	* CAUTION: Even though this is marked as threadsafe, it is only valid to call from the render thread or the RHI thread. It need not be threadsafe unless the RHI support parallel translation.
-	* CAUTION: Platforms that support RHIThread but don't actually have a threadsafe implementation must flush internally with FScopedRHIThreadStaller StallRHIThread(FRHICommandListExecutor::GetImmediateCommandList()); when the call is from the render thread
-	*/
-	// FlushType: Thread safe
-	// TODO: [PSO API] Make pure virtual
 	virtual FGraphicsPipelineStateRHIRef RHICreateGraphicsPipelineState(const FGraphicsPipelineStateInitializer& Initializer) override final
 	{
 		ValidatePipeline(Initializer);
diff --git a/Engine/Source/Runtime/RenderCore/Private/RenderGraphResourcePool.cpp b/Engine/Source/Runtime/RenderCore/Private/RenderGraphResourcePool.cpp
--- a/Engine/Source/Runtime/RenderCore/Private/RenderGraphResourcePool.cpp
+++ b/Engine/Source/Runtime/RenderCore/Private/RenderGraphResourcePool.cpp
@@ -286,7 +286,7 @@ void FRDGTransientResourceAllocator::ReleaseRHI()
 {
 	if (Allocator)
 	{
-		FRHICommandListImmediate& RHICmdList = FRHICommandListExecutor::GetImmediateCommandList();
+		FRHICommandListImmediate& RHICmdList = FRHICommandListImmediate::Get();
 
 		ReleasePendingDeallocations();
 		PendingDeallocationList.Empty();
@@ -395,7 +395,7 @@ void FRDGTransientResourceAllocator::ReleasePendingDeallocations()
 		{
 			const FRHITransition* Transition = RHICreateTransition(FRHITransitionCreateInfo(ERHIPipeline::Graphics, ERHIPipeline::Graphics, ERHITransitionCreateFlags::None, Transitions, Aliases));
 
-			FRHICommandListImmediate& RHICmdList = FRHICommandListExecutor::GetImmediateCommandList();
+			FRHICommandListImmediate& RHICmdList = FRHICommandListImmediate::Get();
 			RHICmdList.BeginTransition(Transition);
 			RHICmdList.EndTransition(Transition);
 			RHICmdList.SetTrackedAccess(EpilogueResourceAccesses);
diff --git a/Engine/Source/Runtime/RenderCore/Private/UnifiedBuffer.cpp b/Engine/Source/Runtime/RenderCore/Private/UnifiedBuffer.cpp
--- a/Engine/Source/Runtime/RenderCore/Private/UnifiedBuffer.cpp
+++ b/Engine/Source/Runtime/RenderCore/Private/UnifiedBuffer.cpp
@@ -1580,7 +1580,7 @@ void FScatterUploadBuffer::InitPreSized(uint32 NumElements, uint32 InNumBytesPer
 }
 
 template<typename ResourceType>
-void FScatterUploadBuffer::ResourceUploadTo(FRHICommandList& RHICmdList, const ResourceType& DstBuffer, bool bFlush)
+void FScatterUploadBuffer::ResourceUploadTo(FRHICommandList& RHICmdList, const ResourceType& DstBuffer)
 {
 	TRACE_CPUPROFILER_EVENT_SCOPE(FScatterUploadBuffer::ResourceUploadTo);
 
@@ -1677,11 +1677,6 @@ void FScatterUploadBuffer::ResourceUploadTo(FRHICommandList& RHICmdList, const R
 	// We need to unbind shader SRVs in this case, because scatter upload buffers are sometimes used more than once in a
 	// frame, and this can cause rendering bugs on D3D12, where the driver fails to update the bound SRV with new data.
 	UnsetShaderSRVs(RHICmdList, ComputeShader, ComputeShader.GetComputeShader());
-
-	if (bFlush)
-	{
-		FRHICommandListExecutor::GetImmediateCommandList().ImmediateFlush(EImmediateFlushType::DispatchToRHIThread);
-	}
 }
 
 template RENDERCORE_API void MemsetResource<FRWBufferStructured>(FRHICommandList& RHICmdList, const FRWBufferStructured& DstBuffer, const FMemsetResourceParams& Params);
diff --git a/Engine/Source/Runtime/RenderCore/Public/UnifiedBuffer.h b/Engine/Source/Runtime/RenderCore/Public/UnifiedBuffer.h
--- a/Engine/Source/Runtime/RenderCore/Public/UnifiedBuffer.h
+++ b/Engine/Source/Runtime/RenderCore/Public/UnifiedBuffer.h
@@ -92,8 +92,15 @@ public:
 	RENDERCORE_API void Init( uint32 NumElements, uint32 InNumBytesPerElement, bool bInFloat4Buffer, const TCHAR* DebugName );
 
 	template<typename ResourceType>
-	void ResourceUploadTo(FRHICommandList& RHICmdList, const ResourceType& DstBuffer, bool bFlush = false);
-	
+	void ResourceUploadTo(FRHICommandList& RHICmdList, const ResourceType& DstBuffer);
+
+	template<typename ResourceType>
+	UE_DEPRECATED(5.8, "ResourceUploadTo with the optional bFlush argument is deprecated. If submission of GPU work is required immediately afterwards, the caller should do this via the immediate command list.")
+	void ResourceUploadTo(FRHICommandList& RHICmdList, const ResourceType& DstBuffer, bool /*unused bFlush*/)
+	{
+		ResourceUploadTo(RHICmdList, DstBuffer);
+	}
+
 	void Add( uint32 Index, const void* Data, uint32 Num = 1 )
 	{
 		void* Dst = Add_GetRef( Index, Num );
diff --git a/Engine/Source/Runtime/Renderer/Private/RendererScene.cpp b/Engine/Source/Runtime/Renderer/Private/RendererScene.cpp
--- a/Engine/Source/Runtime/Renderer/Private/RendererScene.cpp
+++ b/Engine/Source/Runtime/Renderer/Private/RendererScene.cpp
@@ -3092,7 +3092,7 @@ void FScene::RemovePrecomputedLightVolume(const FPrecomputedLightVolume* Volume)
 		});
 }
 
-void FVolumetricLightmapSceneData::AddLevelVolume(const FPrecomputedVolumetricLightmap* InVolume, EShadingPath ShadingPath, bool bIsPersistentLevel)
+void FVolumetricLightmapSceneData::AddLevelVolume(FRHICommandListImmediate& RHICmdList, const FPrecomputedVolumetricLightmap* InVolume, EShadingPath ShadingPath, bool bIsPersistentLevel)
 {
 	LevelVolumetricLightmaps.Add(InVolume);
 
@@ -3101,17 +3101,17 @@ void FVolumetricLightmapSceneData::AddLevelVolume(const FPrecomputedVolumetricLi
 		PersistentLevelVolumetricLightmap = InVolume;
 	}
 
-	InVolume->Data->AddToSceneData(&GlobalVolumetricLightmapData);
+	InVolume->Data->AddToSceneData(RHICmdList, &GlobalVolumetricLightmapData);
 	
 	// Invalidate CPU lightmap lookup cache
 	CPUInterpolationCache.Empty();
 }
 
-void FVolumetricLightmapSceneData::RemoveLevelVolume(const FPrecomputedVolumetricLightmap* InVolume)
+void FVolumetricLightmapSceneData::RemoveLevelVolume(FRHICommandListImmediate& RHICmdList, const FPrecomputedVolumetricLightmap* InVolume)
 {
 	LevelVolumetricLightmaps.Remove(InVolume);
 
-	InVolume->Data->RemoveFromSceneData(&GlobalVolumetricLightmapData, PersistentLevelVolumetricLightmap ? PersistentLevelVolumetricLightmap->Data->BrickDataBaseOffsetInAtlas : 0);
+	InVolume->Data->RemoveFromSceneData(RHICmdList, &GlobalVolumetricLightmapData, PersistentLevelVolumetricLightmap ? PersistentLevelVolumetricLightmap->Data->BrickDataBaseOffsetInAtlas : 0);
 
 	if (PersistentLevelVolumetricLightmap == InVolume)
 	{
@@ -3172,9 +3172,9 @@ void FScene::AddPrecomputedVolumetricLightmap(const FPrecomputedVolumetricLightm
 	FScene* Scene = this;
 
 	ENQUEUE_RENDER_COMMAND(AddVolumeCommand)
-		([Scene, Volume, bIsPersistentLevel] (FRHICommandListBase&)
+		([Scene, Volume, bIsPersistentLevel] (FRHICommandListImmediate& RHICmdList)
 	{
-		Scene->VolumetricLightmapSceneData.AddLevelVolume(Volume, Scene->GetShadingPath(), bIsPersistentLevel);
+		Scene->VolumetricLightmapSceneData.AddLevelVolume(RHICmdList, Volume, Scene->GetShadingPath(), bIsPersistentLevel);
 	});
 }
 
@@ -3183,9 +3183,9 @@ void FScene::RemovePrecomputedVolumetricLightmap(const FPrecomputedVolumetricLig
 	FScene* Scene = this; 
 
 	ENQUEUE_RENDER_COMMAND(RemoveVolumeCommand)
-		([Scene, Volume] (FRHICommandListBase&)
+		([Scene, Volume] (FRHICommandListImmediate& RHICmdList)
 	{
-		Scene->VolumetricLightmapSceneData.RemoveLevelVolume(Volume);
+		Scene->VolumetricLightmapSceneData.RemoveLevelVolume(RHICmdList, Volume);
 	});
 }
 
diff --git a/Engine/Source/Runtime/Renderer/Private/ScenePrivate.h b/Engine/Source/Runtime/Renderer/Private/ScenePrivate.h
--- a/Engine/Source/Runtime/Renderer/Private/ScenePrivate.h
+++ b/Engine/Source/Runtime/Renderer/Private/ScenePrivate.h
@@ -1851,8 +1851,10 @@ public:
 	}
 
 	bool HasData() const;
-	void AddLevelVolume(const class FPrecomputedVolumetricLightmap* InVolume, EShadingPath ShadingPath, bool bIsPersistentLevel);
-	void RemoveLevelVolume(const class FPrecomputedVolumetricLightmap* InVolume);
+
+	void AddLevelVolume   (FRHICommandListImmediate& RHICmdList, const class FPrecomputedVolumetricLightmap* InVolume, EShadingPath ShadingPath, bool bIsPersistentLevel);
+	void RemoveLevelVolume(FRHICommandListImmediate& RHICmdList, const class FPrecomputedVolumetricLightmap* InVolume);
+
 	const FPrecomputedVolumetricLightmap* GetLevelVolumetricLightmap() const;
 
 	TMap<FVector, FVolumetricLightmapInterpolation> CPUInterpolationCache;
diff --git a/Engine/Source/Runtime/Renderer/Private/SceneVisibility.cpp b/Engine/Source/Runtime/Renderer/Private/SceneVisibility.cpp
--- a/Engine/Source/Runtime/Renderer/Private/SceneVisibility.cpp
+++ b/Engine/Source/Runtime/Renderer/Private/SceneVisibility.cpp
@@ -3420,7 +3420,7 @@ void FGPUOcclusionParallel::Finalize()
 
 	for (FGPUOcclusionParallelPacket* Packet : Packets)
 	{
-		FGPUOcclusionPacket::FProcessVisitor ProcessVisitor(*Packet, *RHICmdList, DynamicVertexBuffer);
+		FGPUOcclusionPacket::FProcessVisitor ProcessVisitor(*Packet, DynamicVertexBuffer);
 		ProcessVisitor.Replay(Packet->RecordVisitor);
 	}
 
diff --git a/Engine/Source/Runtime/Renderer/Private/SceneVisibilityPrivate.h b/Engine/Source/Runtime/Renderer/Private/SceneVisibilityPrivate.h
--- a/Engine/Source/Runtime/Renderer/Private/SceneVisibilityPrivate.h
+++ b/Engine/Source/Runtime/Renderer/Private/SceneVisibilityPrivate.h
@@ -1002,9 +1002,8 @@ public:
 	// Visitor variant that will process the requests immediately or replay recorded commands.
 	struct FProcessVisitor
 	{
-		FProcessVisitor(FGPUOcclusionPacket& InPacket, FRHICommandList& InRHICmdList, FGlobalDynamicVertexBuffer& InDynamicVertexBuffer)
+		FProcessVisitor(FGPUOcclusionPacket& InPacket, FGlobalDynamicVertexBuffer& InDynamicVertexBuffer)
 			: Packet(InPacket)
-			, RHICmdList(InRHICmdList)
 			, DynamicVertexBuffer(InDynamicVertexBuffer)
 			, PrimitiveOcclusionHistorySet(Packet.ViewState.Occlusion.PrimitiveOcclusionHistorySet)
 		{}
@@ -1067,7 +1066,6 @@ public:
 		void SubmitThrottledOcclusionQueries();
 
 		FGPUOcclusionPacket& Packet;
-		FRHICommandList& RHICmdList;
 		FGlobalDynamicVertexBuffer& DynamicVertexBuffer;
 		TArray<FThrottledOcclusionQuery, SceneRenderingAllocator> ThrottledOcclusionQueries;
 		TSet<FPrimitiveOcclusionHistory, FPrimitiveOcclusionHistoryKeyFuncs>& PrimitiveOcclusionHistorySet;
@@ -1202,7 +1200,7 @@ public:
 	FGPUOcclusionSerial(FVisibilityViewPacket& InViewPacket)
 		: FGPUOcclusion(InViewPacket)
 		, Packet(InViewPacket, State)
-		, ProcessVisitor(Packet, FRHICommandListExecutor::GetImmediateCommandList(), DynamicVertexBuffer)
+		, ProcessVisitor(Packet, DynamicVertexBuffer)
 	{}
 
 	void AddPrimitives(FPrimitiveRange PrimitiveRange) override;
diff --git a/Engine/Source/Runtime/SlateRHIRenderer/Private/SlateRHIRenderer.cpp b/Engine/Source/Runtime/SlateRHIRenderer/Private/SlateRHIRenderer.cpp
--- a/Engine/Source/Runtime/SlateRHIRenderer/Private/SlateRHIRenderer.cpp
+++ b/Engine/Source/Runtime/SlateRHIRenderer/Private/SlateRHIRenderer.cpp
@@ -176,7 +176,7 @@ struct FSlateViewportInfo : public FRenderResource
 	void ReleaseRHI() override
 	{
 		// Full GPU sync here to simplify memory lifetime of the underlying resource.
-		FRHICommandListExecutor::GetImmediateCommandList().BlockUntilGPUIdle();
+		FRHICommandListImmediate::Get().BlockUntilGPUIdle();
 		ViewportRHI.SafeRelease();
 	}
 };
