Substrate - Added toon profile asset and sampling exemple for direct diffuse only.
When using blendable gbuffer, SpecularColor is now 565 while anisotropy and ToonProfileID are 8 bits.

#rnx
#tests PC editor

#Synced-CL 48448616

diff --git a/Engine/Config/BaseEditor.ini b/Engine/Config/BaseEditor.ini
--- a/Engine/Config/BaseEditor.ini
+++ b/Engine/Config/BaseEditor.ini
@@ -71,6 +71,7 @@ SdfFontAtlasSize=2048
 +RenderableThumbnailTypes=(ClassNeedingThumbnailName="/Script/Engine.CurveLinearColor",RendererClassName="/Script/UnrealEd.CurveLinearColorThumbnailRenderer")
 +RenderableThumbnailTypes=(ClassNeedingThumbnailName="/Script/Engine.CurveFloat",RendererClassName="/Script/UnrealEd.CurveFloatThumbnailRenderer")
 +RenderableThumbnailTypes=(ClassNeedingThumbnailName="/Script/Engine.CurveVector",RendererClassName="/Script/UnrealEd.CurveVector3ThumbnailRenderer")
++RenderableThumbnailTypes=(ClassNeedingThumbnailName="/Script/Engine.ToonProfile",RendererClassName="/Script/UnrealEd.ToonProfileRenderer")
 
 [/Script/UnrealEd.DefaultSizedThumbnailRenderer]
 DefaultSizeX=512
diff --git a/Engine/Plugins/Editor/EngineAssetDefinitions/Source/Private/AssetDefinition_ToonProfile.h b/Engine/Plugins/Editor/EngineAssetDefinitions/Source/Private/AssetDefinition_ToonProfile.h
new file mode 100644
--- /dev/null
+++ b/Engine/Plugins/Editor/EngineAssetDefinitions/Source/Private/AssetDefinition_ToonProfile.h
@@ -0,0 +1,26 @@
+// Copyright Epic Games, Inc. All Rights Reserved.
+
+#pragma once
+
+#include "Engine/ToonProfile.h"
+#include "AssetDefinitionDefault.h"
+
+#include "AssetDefinition_ToonProfile.generated.h"
+
+UCLASS()
+class UAssetDefinition_ToonProfile : public UAssetDefinitionDefault
+{
+	GENERATED_BODY()
+
+public:
+	// UAssetDefinition Begin
+	virtual FText GetAssetDisplayName() const override { return NSLOCTEXT("AssetTypeActions", "AssetTypeActions_ToonProfile", "Toon Profile"); }
+	virtual FLinearColor GetAssetColor() const override { return FLinearColor(FColor(128, 128, 255)); }
+	virtual TSoftClassPtr<UObject> GetAssetClass() const override { return UToonProfile::StaticClass(); }
+	virtual TConstArrayView<FAssetCategoryPath> GetAssetCategories() const override
+    {
+    	static const auto Categories = { EAssetCategoryPaths::Material / NSLOCTEXT("Material", "MaterialAssetSubMenu_Profiles", "Profiles") };
+    	return Categories;
+    }
+	// UAssetDefinition End
+};
diff --git a/Engine/Shaders/Private/BasePassPixelShader.usf b/Engine/Shaders/Private/BasePassPixelShader.usf
--- a/Engine/Shaders/Private/BasePassPixelShader.usf
+++ b/Engine/Shaders/Private/BasePassPixelShader.usf
@@ -1904,7 +1904,8 @@ void FPixelShaderInOut_MainPS(
 			V,
 			SurfaceWorldNormal,
 			MaterialParameters.WorldPosition_CamRelative,
-			0.0f // Curvature
+			0.0f, // Curvature
+			Dither
 			);
 		
 		// Propagate some more data
diff --git a/Engine/Shaders/Private/Substrate/Substrate.ush b/Engine/Shaders/Private/Substrate/Substrate.ush
--- a/Engine/Shaders/Private/Substrate/Substrate.ush
+++ b/Engine/Shaders/Private/Substrate/Substrate.ush
@@ -713,10 +713,12 @@ struct FSubstrateBSDF
 
 ///////////////////////////////////////////////////////////////////////////////
 // Toon BSDF members
+// TOON_PROFILEID is in [0, 255]
 #define TOON_DIFFUSEALBEDO(X)			X.VGPRs[0].xyz
 #define TOON_ROUGHNESS(X)				X.VGPRs[0].w
 #define TOON_SPECULARCOLOR(X)			X.VGPRs[1].xyz
 #define TOON_ANISOTROPY(X)				X.VGPRs[1].w
+#define TOON_PROFILEID(X)				X.VGPRs[2].x
 
 ///////////////////////////////////////////////////////////////////////////////
 // Forward declarations
@@ -1968,6 +1970,7 @@ FSubstrateData GetSubstrateSingleLayerWaterBSDF(
 }
 
 FSubstrateData GetSubstrateToonBSDF(
+	float ToonProfileId, // in [0, 255]
 	float3 DiffuseColor, float3 SpecularColor, float Roughness, float Anisotropy, 
 	float3 Emissive,
 	uint SharedLocalBasisIndex, inout uint SharedLocalBasisTypes)
@@ -1993,6 +1996,7 @@ FSubstrateData GetSubstrateToonBSDF(
 	BSDF_SETTYPE					(SubstrateData.InlinedBSDF, SUBSTRATE_BSDF_TYPE_TOON);
 	BSDF_SETSHAREDLOCALBASISID		(SubstrateData.InlinedBSDF, SharedLocalBasisIndex);
 	BSDF_SETEMISSIVE				(SubstrateData.InlinedBSDF, Emissive);
+	TOON_PROFILEID					(SubstrateData.InlinedBSDF) = ToonProfileId;
 	TOON_DIFFUSEALBEDO				(SubstrateData.InlinedBSDF) = DiffuseColor;
 	TOON_SPECULARCOLOR				(SubstrateData.InlinedBSDF) = SpecularColor;
 	TOON_ROUGHNESS					(SubstrateData.InlinedBSDF) = Roughness;
@@ -4105,6 +4109,34 @@ float3 UnpackR7G7B6Linear(uint rgb)
 	return Out;
 }
 
+uint PackToonCustomData(in float3 SpecCol, in float Aniso, in float ProfileId, in float Dither)
+{
+	// We execute a dither that can go in positive/negative direction to avoid surfaces to looks overall brighter.
+	const float CenteredDither = Dither - 0.5f;
+	const float Dither8bits = CenteredDither * 1.0f / 255.0f;
+	const float Dither7bits = CenteredDither * 1.0f / 127.0f;
+	const float Dither6bits = CenteredDither * 1.0f / 63.0f;
+	const float Dither5bits = CenteredDither * 1.0f / 31.0f;
+
+	SpecCol = DitherIfNonBlack(SubstrateLinearToSrgb(saturate(SpecCol)), float3(Dither5bits, Dither6bits, Dither5bits));
+	Aniso = DitherIfNonBlack(saturate(Aniso), Dither8bits);
+
+	uint Packed =	(uint(SpecCol.r *  31.0f) & 0x1F) <<  0;	// Spec color R 5 bits
+	Packed |=		(uint(SpecCol.g *  63.0f) & 0x3F) <<  5;	// Spec color G 6 bits
+	Packed |=		(uint(SpecCol.b *  31.0f) & 0x1F) << 11;	// Spec Color B 5 bits
+	Packed |=		(uint(Aniso     * 255.0f) & 0xFF) << 16;	// Aniso        8 bits
+	Packed |=		(uint(ProfileId         ) & 0xFF) << 24;	// ProfileId    8 bits
+
+	return Packed;
+}
+
+void UnpackToonCustomData(in uint PackedToonCustomData, inout float3 SpecCol, inout float Aniso, inout float ProfileId)
+{
+	SpecCol		= SubstrateSrgbToLinear(float3(BitFieldExtractU32(PackedToonCustomData, 5, 0), BitFieldExtractU32(PackedToonCustomData, 6, 5), BitFieldExtractU32(PackedToonCustomData, 5, 11)) * float3(1.0f / 31.0f, 1.0f / 63.0f, 1.0f / 31.0f));
+	Aniso		= float(UnpackB8(PackedToonCustomData));
+	ProfileId	= float(BitFieldExtractU32(PackedToonCustomData, 8, 24)); // Stays in [0, 255]
+}
+
 uint PackR7(float Value, float Dither)
 {
 	return uint(Dither7bits(Value, Dither) * 127.0f);
@@ -4803,12 +4835,11 @@ FSubstrateBSDF UnpackSubstrateBSDFIn(FSubstrateMaterialContainer SubstrateBuffer
 			if (SubstrateHeader.GetComplexPathMode() == COMPLEXPATH_MODE_TOON)
 			{
 				float4 DiffuseRoughness		= UnpackColorGamma2ToLinearAlphaLinear(DataX);
-				float4 SpecularAnisotropy	= UnpackColorGamma2ToLinearAlphaLinear(DataY);
-	
 				TOON_DIFFUSEALBEDO(OutBSDF)	= DiffuseRoughness.xyz;
 				TOON_ROUGHNESS(OutBSDF)		= DiffuseRoughness.w;
-				TOON_SPECULARCOLOR(OutBSDF)	= SpecularAnisotropy.xyz;
-				TOON_ANISOTROPY(OutBSDF)	= SpecularAnisotropy.w;
+
+				UnpackToonCustomData(DataY, TOON_SPECULARCOLOR(OutBSDF), TOON_ANISOTROPY(OutBSDF), TOON_PROFILEID(OutBSDF));
+
 				return OutBSDF;
 			}
 			else
diff --git a/Engine/Shaders/Private/Substrate/SubstrateEvaluation.ush b/Engine/Shaders/Private/Substrate/SubstrateEvaluation.ush
--- a/Engine/Shaders/Private/Substrate/SubstrateEvaluation.ush
+++ b/Engine/Shaders/Private/Substrate/SubstrateEvaluation.ush
@@ -34,6 +34,7 @@
 #include "../ShadingModels.ush"
 #include "../ShadingModelsSampling.ush"
 #include "../AreaLightCommon.ush"
+#include "../ToonProfileCommon.ush"
 
 #include "../HairStrands/HairStrandsCommon.ush"
 #include "../HairStrands/HairStrandsDeepTransmittanceCommon.ush"
@@ -1309,10 +1310,12 @@ FSubstrateEvaluateResult SubstrateEvaluateBSDFCommon(FSubstrateBSDFContext BSDFC
 			const float3 DiffuseColor		= TOON_DIFFUSEALBEDO(BSDFContext.BSDF);
 			const float3 SpecularColor		= TOON_SPECULARCOLOR(BSDFContext.BSDF);
 
-			Sample.DiffuseColor				= DiffuseColor;
+			const float ToonProfileId		= TOON_PROFILEID(BSDFContext.BSDF);
+
+			Sample.DiffuseColor				= DiffuseColor * EvaluateToonProfile(ToonProfileId, BSDFContext.Context.NoL); // AreaLight.NoL to handle capsule shape?
 
 			// TOON_TODO Hardcoded Toon Lambert diffuse before we implement the profile.
-			Sample.DiffusePathValue			= Diffuse_Lambert(Sample.DiffuseColor) * floor(saturate(AreaLight.NoL) * 10)*0.1;
+			Sample.DiffusePathValue			= Diffuse_Lambert(Sample.DiffuseColor);
 			Sample.IntegratedDiffuseValue	= ShadowTerms.SurfaceShadow * AreaLight.Falloff * AreaLight.FalloffColor * Sample.DiffusePathValue;
 
 			Sample.DiffusePDF				= BSDFContext.SatNoL / PI;
diff --git a/Engine/Shaders/Private/Substrate/SubstrateExport.ush b/Engine/Shaders/Private/Substrate/SubstrateExport.ush
--- a/Engine/Shaders/Private/Substrate/SubstrateExport.ush
+++ b/Engine/Shaders/Private/Substrate/SubstrateExport.ush
@@ -70,7 +70,7 @@ float3 SubstrateGetBSDFPrecomputedLighting(
 	float DiffuseEnvLightingNoL = 1.0f;
 	float3 LuminanceWeightFinal = LuminanceWeight(DiffuseEnvLightingNoL, CurrentBSDF);
 
-	// Evaluate material environment ligthing parameters
+	// Evaluate material environment lighting parameters
 	const bool bEnableSpecular = false;
 	FSubstrateEnvLightResult SubstrateEnvLight = SubstrateEvaluateForEnvLight(EnvBSDFContext, bEnableSpecular, Settings);
 
@@ -878,8 +878,8 @@ void PackSubstrateOut(
 						case SUBSTRATE_BSDF_TYPE_TOON:
 						{
 							SUBSTRATE_STORE_UINT1(PackColorLinearToGamma2AlphaLinear(float4(TOON_DIFFUSEALBEDO(BSDF), TOON_ROUGHNESS(BSDF))));
-							SUBSTRATE_STORE_UINT1(PackColorLinearToGamma2AlphaLinear(float4(TOON_SPECULARCOLOR(BSDF), TOON_ANISOTROPY(BSDF))));
-							// 8 bytes
+							SUBSTRATE_STORE_UINT1(PackToonCustomData(TOON_SPECULARCOLOR(BSDF), TOON_ANISOTROPY(BSDF), TOON_PROFILEID(BSDF), Dither));
+							// 12 bytes
 						}
 						break;
 						//case SUBSTRATE_BSDF_TYPE_SINGLELAYERWATER: 
@@ -1030,7 +1030,8 @@ FExportResult SubstrateMaterialExportOut(
 	in float3 V,
 	in float3 SurfaceWorldNormal,
 	in float3 WorldPosition_CamRelative,
-	in float MobileShadingPathCurvature)
+	in float MobileShadingPathCurvature,
+	in float Dither = 0.5f)
 {
 	FExportResult Out = (FExportResult)0;
 	Out.IndirectIrradianceScale = 1.0;
@@ -1247,9 +1248,11 @@ FExportResult SubstrateMaterialExportOut(
 
 				// DiffuseColor in BaseColor
 				// Roughness in Roughness
-				Out.Metallic = 0.0;
-				Out.Specular = 1.0;
-				Out.CustomData = float4(TOON_SPECULARCOLOR(CurrentBSDF), TOON_ANISOTROPY(CurrentBSDF));
+				Out.Metallic = 0.0; // Cannot reuse metallic as it maybe written by decal
+				Out.Specular = 1.0;	// Cannot reuse specular as it can be used to apply decal on specular color
+
+				uint PackedToonCustomData = PackToonCustomData(TOON_SPECULARCOLOR(CurrentBSDF), TOON_ANISOTROPY(CurrentBSDF), TOON_PROFILEID(CurrentBSDF), Dither);
+				Out.CustomData = UnpackRGBA8(PackedToonCustomData);
 			}
 			else
 #endif
diff --git a/Engine/Shaders/Private/Substrate/SubstrateRead.ush b/Engine/Shaders/Private/Substrate/SubstrateRead.ush
--- a/Engine/Shaders/Private/Substrate/SubstrateRead.ush
+++ b/Engine/Shaders/Private/Substrate/SubstrateRead.ush
@@ -115,10 +115,14 @@ FSubstrateGBufferBSDF SubstrateReadGBufferBSDF(in FScreenSpaceData ScreenSpaceDa
 		}
 		else if (GBuffer.ShadingModelID == SHADINGMODELID_SUBSTRATE_TOON)
 		{
-			float3 SpecularColor= GBuffer.CustomData.rgb;
-			float Anisotropy	= GBuffer.CustomData.a;
+			// special read from custom data
+			float3 SpecularColor= 0.0f;
+			float Anisotropy    = 0.0f;
+			float ToonProfileId	= 0.0f;
+			UnpackToonCustomData(PackRGBA8(GBuffer.CustomData), SpecularColor, Anisotropy, ToonProfileId);
 
 			BSDF = GetSubstrateToonBSDF(
+				ToonProfileId,
 				GBuffer.BaseColor, SpecularColor, GBuffer.Roughness, Anisotropy,
 				Emissive,
 				SharedLocalBasisIndex, SharedLocalBasisTypes).InlinedBSDF;
diff --git a/Engine/Shaders/Private/ToonProfile.usf b/Engine/Shaders/Private/ToonProfile.usf
new file mode 100644
--- /dev/null
+++ b/Engine/Shaders/Private/ToonProfile.usf
@@ -0,0 +1,35 @@
+// Copyright Epic Games, Inc. All Rights Reserved.
+
+#include "Common.ush" 
+#include "ToonProfileCommon.ush"
+
+
+#ifdef ToonProfileUpdateAtlasTileCS
+
+int2 ToonProfileTextureResolution;
+uint ToonProfileIndex;
+
+RWTexture2DArray<float4> ToonProfileTexture;
+
+Buffer<float4> DiffuseRampBufferSRV;
+
+[numthreads(8, 8, 1)]
+void ToonProfileUpdateAtlasTileCS(uint3 DispatchThreadId : SV_DispatchThreadID)
+{
+	const uint2 PixelPos = DispatchThreadId.xy;
+	if (all(PixelPos < uint2(ToonProfileTextureResolution)))
+	{
+		const float2 TexCoord = float2(PixelPos + 0.5f) / float2(ToonProfileTextureResolution);
+
+		if (PixelPos.y == 0)
+		{
+			ToonProfileTexture[uint3(PixelPos, ToonProfileIndex)] = DiffuseRampBufferSRV[PixelPos.x];
+		}
+		else
+		{
+			ToonProfileTexture[uint3(PixelPos, ToonProfileIndex)] = float4(1, 0, 1, 0);
+		}
+	}
+}
+
+#endif
diff --git a/Engine/Shaders/Private/ToonProfileCommon.ush b/Engine/Shaders/Private/ToonProfileCommon.ush
new file mode 100644
--- /dev/null
+++ b/Engine/Shaders/Private/ToonProfileCommon.ush
@@ -0,0 +1,14 @@
+// Copyright Epic Games, Inc. All Rights Reserved.
+
+#pragma once
+
+
+float3 EvaluateToonProfile(in float ToonProfileId, float NoL)
+{
+	const float3 ToonProfileCoord = float3(
+		NoL * 0.5f + 0.5f,									// The toon ramp is applied over NoL in [-1, 1]
+		0.5f * View.ToonProfileTextureSizeAndInvSize.w,		// Center of the first row of texel
+		ToonProfileId										// Slice index
+	);
+	return View.ToonProfileTexture.SampleLevel(View.ToonProfileSampler, ToonProfileCoord, 0).xyz;
+}
\ No newline at end of file
diff --git a/Engine/Shaders/Shared/SubstrateVisualizeDefinitions.h b/Engine/Shaders/Shared/SubstrateVisualizeDefinitions.h
--- a/Engine/Shaders/Shared/SubstrateVisualizeDefinitions.h
+++ b/Engine/Shaders/Shared/SubstrateVisualizeDefinitions.h
@@ -680,14 +680,15 @@ FSubstratePixelDebugData ConvertToSerializableSubstratePixelDebugData(uint2 InCo
 
 		case SUBSTRATE_BSDF_TYPE_TOON:
 		{
-			Data.Closures[i].DiffuseR	= TOON_DIFFUSEALBEDO(BSDF).r;
-			Data.Closures[i].DiffuseG	= TOON_DIFFUSEALBEDO(BSDF).g;
-			Data.Closures[i].DiffuseB	= TOON_DIFFUSEALBEDO(BSDF).b;
-			Data.Closures[i].F0R		= TOON_SPECULARCOLOR(BSDF).r;
-			Data.Closures[i].F0G		= TOON_SPECULARCOLOR(BSDF).g;
-			Data.Closures[i].F0B		= TOON_SPECULARCOLOR(BSDF).b;
-			Data.Closures[i].Roughness	= TOON_ROUGHNESS(BSDF);
-			Data.Closures[i].Anisotropy	= TOON_ANISOTROPY(BSDF);
+			Data.Closures[i].DiffuseR		= TOON_DIFFUSEALBEDO(BSDF).r;
+			Data.Closures[i].DiffuseG		= TOON_DIFFUSEALBEDO(BSDF).g;
+			Data.Closures[i].DiffuseB		= TOON_DIFFUSEALBEDO(BSDF).b;
+			Data.Closures[i].F0R			= TOON_SPECULARCOLOR(BSDF).r;
+			Data.Closures[i].F0G			= TOON_SPECULARCOLOR(BSDF).g;
+			Data.Closures[i].F0B			= TOON_SPECULARCOLOR(BSDF).b;
+			Data.Closures[i].Roughness		= TOON_ROUGHNESS(BSDF);
+			Data.Closures[i].Anisotropy		= TOON_ANISOTROPY(BSDF);
+			Data.Closures[i].SSSPRofileID	= TOON_PROFILEID(BSDF);
 		}
 		break;
 
diff --git a/Engine/Source/Editor/DetailCustomizations/Private/MaterialInstanceDynamicDetails.cpp b/Engine/Source/Editor/DetailCustomizations/Private/MaterialInstanceDynamicDetails.cpp
--- a/Engine/Source/Editor/DetailCustomizations/Private/MaterialInstanceDynamicDetails.cpp
+++ b/Engine/Source/Editor/DetailCustomizations/Private/MaterialInstanceDynamicDetails.cpp
@@ -39,6 +39,8 @@ void FMaterialInstanceDynamicDetails::CustomizeDetails(IDetailLayoutBuilder& Det
 	OverrideSubSurfaceHandle->MarkHiddenByCustomization();
 	TSharedRef<IPropertyHandle> OverrideSpecularProfileHandle = DetailLayout.GetProperty(GET_MEMBER_NAME_CHECKED(UMaterialInstance, bOverrideSpecularProfile), UMaterialInstance::StaticClass());
 	OverrideSpecularProfileHandle->MarkHiddenByCustomization();
+	TSharedRef<IPropertyHandle> OverrideToonProfileHandle = DetailLayout.GetProperty(GET_MEMBER_NAME_CHECKED(UMaterialInstance, bOverrideToonProfile), UMaterialInstance::StaticClass());
+	OverrideToonProfileHandle->MarkHiddenByCustomization();
 	TSharedRef<IPropertyHandle> BasePropertiesHandle = DetailLayout.GetProperty(GET_MEMBER_NAME_CHECKED(UMaterialInstance, BasePropertyOverrides), UMaterialInstance::StaticClass());
 	BasePropertiesHandle->MarkHiddenByCustomization();
 
diff --git a/Engine/Source/Editor/EditorStyle/Private/StarshipStyle.cpp b/Engine/Source/Editor/EditorStyle/Private/StarshipStyle.cpp
--- a/Engine/Source/Editor/EditorStyle/Private/StarshipStyle.cpp
+++ b/Engine/Source/Editor/EditorStyle/Private/StarshipStyle.cpp
@@ -6836,6 +6836,7 @@ void FStarshipEditorStyle::FStyle::SetupClassIconsAndThumbnails()
 			{TEXT("TextureRenderTargetCube")},
 			{TEXT("TimeCodeSynchronizer")},
 			{TEXT("TouchInterface")},
+			{TEXT("ToonProfile")},
 			{TEXT("TriggerBase")},
 			{TEXT("TriggerBox")},
 			{TEXT("TriggerCapsule")},
diff --git a/Engine/Source/Editor/MaterialEditor/Private/MaterialEditorInstanceDetailCustomization.cpp b/Engine/Source/Editor/MaterialEditor/Private/MaterialEditorInstanceDetailCustomization.cpp
--- a/Engine/Source/Editor/MaterialEditor/Private/MaterialEditorInstanceDetailCustomization.cpp
+++ b/Engine/Source/Editor/MaterialEditor/Private/MaterialEditorInstanceDetailCustomization.cpp
@@ -175,6 +175,8 @@ void FMaterialInstanceParameterDetails::CustomizeDetails(IDetailLayoutBuilder& D
 		DetailLayout.HideProperty("SubsurfaceProfile");
 		DetailLayout.HideProperty("bOverrideSpecularProfile");
 		DetailLayout.HideProperty("SpecularProfile");
+		DetailLayout.HideProperty("bOverrideToonProfile");
+		DetailLayout.HideProperty("ToonProfile");
 		DetailLayout.HideProperty("BasePropertyOverrides");
 		DetailLayout.HideProperty("MaterialLayersParameterValues");
 	}
@@ -250,6 +252,8 @@ void FMaterialInstanceParameterDetails::CustomizeDetails(IDetailLayoutBuilder& D
 			DetailLayout.HideProperty("SubsurfaceProfile");
 			DetailLayout.HideProperty("bOverrideSpecularProfile");
 			DetailLayout.HideProperty("SpecularProfile");
+			DetailLayout.HideProperty("bOverrideToonProfile");
+			DetailLayout.HideProperty("ToonProfile");
 			DetailLayout.HideProperty("BasePropertyOverrides");
 		}
 		else
@@ -320,6 +324,8 @@ void FMaterialInstanceParameterDetails::CustomizeDetails(IDetailLayoutBuilder& D
 				DetailLayout.HideProperty("SubsurfaceProfile");
 				DetailLayout.HideProperty("bOverrideSpecularProfile");
 				DetailLayout.HideProperty("SpecularProfile");
+				DetailLayout.HideProperty("bOverrideToonProfile");
+				DetailLayout.HideProperty("ToonProfile");
 				DetailLayout.HideProperty("BasePropertyOverrides");
 
 				// Set up the override logic for the subsurface profile
@@ -337,20 +343,33 @@ void FMaterialInstanceParameterDetails::CustomizeDetails(IDetailLayoutBuilder& D
 						.Visibility(TAttribute<EVisibility>::Create(TAttribute<EVisibility>::FGetter::CreateSP(this, &FMaterialInstanceParameterDetails::ShouldShowSubsurfaceProfile)));
 				}
 
-				// Set up the override logic for the specular profile
+				// Set up the override logic for the specular and toon profiles
 				if (Substrate::IsSubstrateEnabled())
 				{
-					TAttribute<bool> IsParamEnabled = TAttribute<bool>::Create(TAttribute<bool>::FGetter::CreateLambda([this](){ return (bool)MaterialEditorInstance->bOverrideSpecularProfile; }));
+					TAttribute<bool> IsSpecProfileParamEnabled = TAttribute<bool>::Create(TAttribute<bool>::FGetter::CreateLambda([this](){ return (bool)MaterialEditorInstance->bOverrideSpecularProfile; }));
 
-					IDetailPropertyRow& PropertyRow = MaterialPropertyOverrideGroup.AddPropertyRow(DetailLayout.GetProperty("SpecularProfile"));
-					PropertyRow
-					.EditCondition(IsParamEnabled, 
+					IDetailPropertyRow& SpecProfilePropertyRow = MaterialPropertyOverrideGroup.AddPropertyRow(DetailLayout.GetProperty("SpecularProfile"));
+					SpecProfilePropertyRow
+					.EditCondition(IsSpecProfileParamEnabled,
 								   FOnBooleanValueChanged::CreateLambda([this](bool NewValue) {
 									   MaterialEditorInstance->bOverrideSpecularProfile = (uint32)NewValue;
 									   MaterialEditorInstance->PostEditChange();
 									   FEditorSupportDelegates::RedrawAllViewports.Broadcast();
 					}))
 					.Visibility(TAttribute<EVisibility>::Create(TAttribute<EVisibility>::FGetter::CreateSP(this, &FMaterialInstanceParameterDetails::ShouldShowSpecularProfile)));
+
+
+					TAttribute<bool> IsToonProfileParamEnabled = TAttribute<bool>::Create(TAttribute<bool>::FGetter::CreateLambda([this](){ return (bool)MaterialEditorInstance->bOverrideToonProfile; }));
+
+					IDetailPropertyRow& ToonProfilePropertyRow = MaterialPropertyOverrideGroup.AddPropertyRow(DetailLayout.GetProperty("ToonProfile"));
+					ToonProfilePropertyRow
+					.EditCondition(IsToonProfileParamEnabled,
+								   FOnBooleanValueChanged::CreateLambda([this](bool NewValue) {
+									   MaterialEditorInstance->bOverrideToonProfile = (uint32)NewValue;
+									   MaterialEditorInstance->PostEditChange();
+									   FEditorSupportDelegates::RedrawAllViewports.Broadcast();
+					}))
+					.Visibility(TAttribute<EVisibility>::Create(TAttribute<EVisibility>::FGetter::CreateSP(this, &FMaterialInstanceParameterDetails::ShouldShowToonProfile)));
 				}
 			
 				// Append the base property overrides to the Material Property Override Group
@@ -712,6 +731,12 @@ EVisibility FMaterialInstanceParameterDetails::ShouldShowSpecularProfile() const
 	return bHasProfiles ? EVisibility::Visible : EVisibility::Collapsed;
 }
 
+EVisibility FMaterialInstanceParameterDetails::ShouldShowToonProfile() const
+{
+	const bool bHasProfiles = Substrate::IsSubstrateEnabled() && (MaterialEditorInstance->SourceInstance->NumToonProfile_Internal() > 0 || MaterialEditorInstance->SourceInstance->GetToonProfileOverride_Internal() != nullptr);
+	return bHasProfiles ? EVisibility::Visible : EVisibility::Collapsed;
+}
+
 void FMaterialInstanceParameterDetails::OnCopyParameterValues(int32 ParameterGroupIndex)
 {
 	if (!MaterialEditorInstance || !MaterialEditorInstance->ParameterGroups.IsValidIndex(ParameterGroupIndex))
diff --git a/Engine/Source/Editor/MaterialEditor/Private/MaterialEditorInstanceDetailCustomization.h b/Engine/Source/Editor/MaterialEditor/Private/MaterialEditorInstanceDetailCustomization.h
--- a/Engine/Source/Editor/MaterialEditor/Private/MaterialEditorInstanceDetailCustomization.h
+++ b/Engine/Source/Editor/MaterialEditor/Private/MaterialEditorInstanceDetailCustomization.h
@@ -65,6 +65,9 @@ private:
 	/** Returns true if the Specular Profile options should be displayed */
 	EVisibility ShouldShowSpecularProfile() const;
 
+	/** Returns true if the Toon Profile options should be displayed */
+	EVisibility ShouldShowToonProfile() const;
+
 	//Functions supporting copy/paste of entire parameter groups.
 
 	/**
diff --git a/Engine/Source/Editor/UnrealEd/Classes/Factories/ToonProfileFactory.h b/Engine/Source/Editor/UnrealEd/Classes/Factories/ToonProfileFactory.h
new file mode 100644
--- /dev/null
+++ b/Engine/Source/Editor/UnrealEd/Classes/Factories/ToonProfileFactory.h
@@ -0,0 +1,21 @@
+// Copyright Epic Games, Inc. All Rights Reserved.
+
+#pragma once
+
+#include "CoreMinimal.h"
+#include "UObject/ObjectMacros.h"
+#include "Factories/Factory.h"
+#include "ToonProfileFactory.generated.h"
+
+UCLASS(hidecategories=Object)
+class UToonProfileFactory : public UFactory
+{
+	GENERATED_UCLASS_BODY()
+
+	//~ Begin UFactory Interface
+	virtual UObject* FactoryCreateNew(UClass* Class,UObject* InParent,FName Name,EObjectFlags Flags,UObject* Context,FFeedbackContext* Warn) override;
+	//~ End UFactory Interface	
+};
+
+
+
diff --git a/Engine/Source/Editor/UnrealEd/Classes/MaterialEditor/MaterialEditorInstanceConstant.h b/Engine/Source/Editor/UnrealEd/Classes/MaterialEditor/MaterialEditorInstanceConstant.h
--- a/Engine/Source/Editor/UnrealEd/Classes/MaterialEditor/MaterialEditorInstanceConstant.h
+++ b/Engine/Source/Editor/UnrealEd/Classes/MaterialEditor/MaterialEditorInstanceConstant.h
@@ -336,6 +336,14 @@ private:
 	UPROPERTY(EditAnywhere, Category = MaterialEditorInstanceConstant)
 	uint32 bOverrideSpecularProfile : 1;
 
+	/** Toon profile */
+	UPROPERTY(EditAnywhere, BlueprintReadOnly, Category = Material, meta = (DisplayName = "Toon Profile"))
+	TObjectPtr<class UToonProfile> ToonProfile;
+
+	/** Defines if ToonProfile from this instance is used or it uses the parent one. */
+	UPROPERTY(EditAnywhere, Category = MaterialEditorInstanceConstant)
+	uint32 bOverrideToonProfile : 1;
+
 	UPROPERTY()
 	uint32 bOverrideBaseProperties_DEPRECATED : 1;
 
diff --git a/Engine/Source/Editor/UnrealEd/Classes/ThumbnailRendering/ToonProfileRenderer.h b/Engine/Source/Editor/UnrealEd/Classes/ThumbnailRendering/ToonProfileRenderer.h
new file mode 100644
--- /dev/null
+++ b/Engine/Source/Editor/UnrealEd/Classes/ThumbnailRendering/ToonProfileRenderer.h
@@ -0,0 +1,28 @@
+// Copyright Epic Games, Inc. All Rights Reserved.
+
+/**
+ *
+ * This thumbnail renderer displays the texture for the object in question
+ */
+
+#pragma once
+
+#include "CoreMinimal.h"
+#include "UObject/ObjectMacros.h"
+#include "ThumbnailRendering/TextureThumbnailRenderer.h"
+#include "ToonProfileRenderer.generated.h"
+
+class FCanvas;
+class FRenderTarget;
+
+UCLASS()
+class UToonProfileRenderer : public UTextureThumbnailRenderer
+{
+	GENERATED_UCLASS_BODY()
+
+	// Begin UThumbnailRenderer Object
+	virtual void GetThumbnailSize(UObject* Object, float Zoom, uint32& OutWidth, uint32& OutHeight) const override;
+	virtual void Draw(UObject* Object, int32 X, int32 Y, uint32 Width, uint32 Height, FRenderTarget*, FCanvas* Canvas, bool bAdditionalViewFamily) override;
+	// End UThumbnailRenderer Object
+};
+
diff --git a/Engine/Source/Editor/UnrealEd/Private/Factories/EditorFactories.cpp b/Engine/Source/Editor/UnrealEd/Private/Factories/EditorFactories.cpp
--- a/Engine/Source/Editor/UnrealEd/Private/Factories/EditorFactories.cpp
+++ b/Engine/Source/Editor/UnrealEd/Private/Factories/EditorFactories.cpp
@@ -112,6 +112,7 @@
 #include "Factories/SubsurfaceProfileFactory.h"
 #include "Factories/SpecularProfileFactory.h"
 #include "Factories/NeuralProfileFactory.h"
+#include "Factories/ToonProfileFactory.h"
 #include "Factories/Texture2dFactoryNew.h"
 #include "Engine/Texture.h"
 #include "Factories/TextureFactory.h"
@@ -148,6 +149,7 @@
 #include "Engine/SubsurfaceProfile.h"
 #include "Engine/SpecularProfile.h"
 #include "Engine/NeuralProfile.h"
+#include "Engine/ToonProfile.h"
 #include "Misc/ConfigCacheIni.h"
 #include "Misc/FeedbackContext.h"
 #include "GameFramework/WorldSettings.h"
@@ -7747,6 +7749,26 @@ UObject* UNeuralProfileFactory::FactoryCreateNew(UClass* Class, UObject* InParen
 	return Object;
 }
 
+
+/*-----------------------------------------------------------------------------
+	UToonProfileFactory implementation.
+-----------------------------------------------------------------------------*/
+UToonProfileFactory::UToonProfileFactory(const FObjectInitializer& ObjectInitializer)
+: Super(ObjectInitializer)
+{
+	SupportedClass = UToonProfile::StaticClass();
+	bCreateNew = true;
+	bEditorImport = false;
+	bEditAfterNew = true;
+}
+
+UObject* UToonProfileFactory::FactoryCreateNew(UClass* InClass, UObject* InParent, FName InName, EObjectFlags Flags, UObject* Context, FFeedbackContext* Warn)
+{
+	UToonProfile* Object = NewObject<UToonProfile>(InParent, InName, Flags);	
+	Object->Guid = FGuid::NewGuid();
+	return Object;
+}
+
 /*-----------------------------------------------------------------------------
 	UTouchInterfaceFactory implementation.
 -----------------------------------------------------------------------------*/
diff --git a/Engine/Source/Editor/UnrealEd/Private/PreviewMaterial.cpp b/Engine/Source/Editor/UnrealEd/Private/PreviewMaterial.cpp
--- a/Engine/Source/Editor/UnrealEd/Private/PreviewMaterial.cpp
+++ b/Engine/Source/Editor/UnrealEd/Private/PreviewMaterial.cpp
@@ -988,6 +988,9 @@ void UMaterialEditorInstanceConstant::CopyToSourceInstance(const bool bForceStat
 		SourceInstance->bOverrideSpecularProfile = bOverrideSpecularProfile;
 		SourceInstance->SpecularProfileOverride = SpecularProfile;
 
+		SourceInstance->bOverrideToonProfile = bOverrideToonProfile;
+		SourceInstance->ToonProfileOverride = ToonProfile;
+
 		// Update object references and parameter names.
 		SourceInstance->UpdateParameterNames();
 		VisibleExpressions.Empty();
@@ -1189,6 +1192,12 @@ void UMaterialEditorInstanceConstant::CopyBasePropertiesFromParent()
 		{
 			SpecularProfile = SourceInstance->GetSpecularProfileOverride_Internal();
 		}
+
+		bOverrideToonProfile = SourceInstance->bOverrideToonProfile;
+		if (SourceInstance->NumToonProfile_Internal() > 0)
+		{
+			ToonProfile = SourceInstance->GetToonProfileOverride_Internal();
+		}
 	}
 
 	// Post process blendable location and priority overrides
diff --git a/Engine/Source/Editor/UnrealEd/Private/ThumbnailRendering/ToonProfileRenderer.cpp b/Engine/Source/Editor/UnrealEd/Private/ThumbnailRendering/ToonProfileRenderer.cpp
new file mode 100644
--- /dev/null
+++ b/Engine/Source/Editor/UnrealEd/Private/ThumbnailRendering/ToonProfileRenderer.cpp
@@ -0,0 +1,43 @@
+// Copyright Epic Games, Inc. All Rights Reserved.
+
+#include "ThumbnailRendering/ToonProfileRenderer.h"
+#include "CanvasItem.h"
+#include "Engine/Engine.h"
+#include "EngineGlobals.h"
+#include "Engine/ToonProfile.h"
+#include "CanvasTypes.h"
+
+#include UE_INLINE_GENERATED_CPP_BY_NAME(ToonProfileRenderer)
+
+UToonProfileRenderer::UToonProfileRenderer(const FObjectInitializer& ObjectInitializer)
+	: Super(ObjectInitializer)
+{
+}
+
+void UToonProfileRenderer::GetThumbnailSize(UObject* Object, float Zoom, uint32& OutWidth, uint32& OutHeight) const
+{
+	OutWidth = 128;
+	OutHeight = 128;
+}
+
+void UToonProfileRenderer::Draw(UObject* Object, int32 X, int32 Y, uint32 Width, uint32 Height, FRenderTarget*, FCanvas* Canvas, bool bAdditionalViewFamily)
+{
+	UToonProfile* LocalToonProfile = Cast<UToonProfile>(Object);
+	if (LocalToonProfile)
+	{
+		FLinearColor Col; 
+
+		const float TimeStep = 0.1f;
+		for (float Time = 0.0f; Time < 1.0f; Time += TimeStep)
+		{
+			Col = LocalToonProfile->Settings.DiffuseRamp.GetLinearColorValue(Time);
+			Col.A = 1;
+			Canvas->DrawTile(float(Width) * Time, 0, float(Width) * (Time + TimeStep), Height, 0, 0, 1, 1, Col);
+		}
+
+		FCanvasTextItem TextItem(FVector2D(5.0f, 5.0f), FText::FromString(TEXT("TOON")), GEngine->GetLargeFont(), FLinearColor::White);
+		TextItem.EnableShadow(FLinearColor::Black);
+		TextItem.Scale = FVector2D(Width / 128.0f, Height / 128.0f);
+		TextItem.Draw(Canvas);
+	}
+}
diff --git a/Engine/Source/Editor/UnrealEd/Public/UnrealEd.h b/Engine/Source/Editor/UnrealEd/Public/UnrealEd.h
--- a/Engine/Source/Editor/UnrealEd/Public/UnrealEd.h
+++ b/Engine/Source/Editor/UnrealEd/Public/UnrealEd.h
@@ -155,7 +155,6 @@ MONOLITHIC_HEADER_BOILERPLATE()
 #include "Factories/SkeletonFactory.h"
 #include "Factories/SlateBrushAssetFactory.h"
 #include "Factories/SlateWidgetStyleAssetFactory.h"
-#include "Factories/SpecularProfileFactory.h"
 #include "Factories/StructureFactory.h"
 #include "Factories/SubsurfaceProfileFactory.h"
 #include "Factories/Texture2dFactoryNew.h"
@@ -164,6 +163,7 @@ MONOLITHIC_HEADER_BOILERPLATE()
 #include "Factories/TrueTypeFontFactory.h"
 #include "Factories/TextureRenderTargetCubeFactoryNew.h"
 #include "Factories/TextureRenderTargetFactoryNew.h"
+#include "Factories/ToonProfileFactory.h"
 #include "Factories/TouchInterfaceFactory.h"
 #include "Factories/VectorFieldStaticFactory.h"
 #include "Factories/ReimportVectorFieldStaticFactory.h"
@@ -269,6 +269,7 @@ MONOLITHIC_HEADER_BOILERPLATE()
 #include "ThumbnailRendering/SubsurfaceProfileRenderer.h"
 #include "ThumbnailRendering/SpecularProfileRenderer.h"
 #include "ThumbnailRendering/NeuralProfileRenderer.h"
+#include "ThumbnailRendering/ToonProfileRenderer.h"
 #include "Factories/TextureCubeThumbnailRenderer.h"
 #include "Factories/Texture2dArrayThumbnailRenderer.h"
 #include "Editor/TransBuffer.h"
diff --git a/Engine/Source/Runtime/Engine/Classes/Engine/EngineTypes.h b/Engine/Source/Runtime/Engine/Classes/Engine/EngineTypes.h
--- a/Engine/Source/Runtime/Engine/Classes/Engine/EngineTypes.h
+++ b/Engine/Source/Runtime/Engine/Classes/Engine/EngineTypes.h
@@ -25,6 +25,7 @@ class UPrimitiveComponent;
 class USceneComponent;
 class USubsurfaceProfile;
 class USpecularProfile;
+class UToonProfile;
 
 /**
  * Default number of components to expect in TInlineAllocators used with AActor component arrays.
@@ -861,6 +862,11 @@ public:
 	int32 CountSpecularProfiles() const { return SpecularProfiles.Num(); }
 	USpecularProfile* GetSpecularProfile(int32 Index) const { return SpecularProfiles.IsValidIndex(Index) ? SpecularProfiles[Index] : nullptr; }
 
+	// Toon profiles
+	void AddToonProfile(UToonProfile* InProfile) { if (InProfile) ToonProfiles.Add(InProfile); }
+	int32 CountToonProfiles() const { return ToonProfiles.Num(); }
+	UToonProfile* GetToonProfile(int32 Index) const { return ToonProfiles.IsValidIndex(Index) ? ToonProfiles[Index] : nullptr; }
+
 	// Shading model from expression
 	void SetShadingModelFromExpression(bool bIn) { bHasShadingModelFromExpression = bIn ? 1u : 0u; }
 	bool HasShadingModelFromExpression() const { return bHasShadingModelFromExpression > 0u; }
@@ -909,6 +915,8 @@ private:
 	TArray<TObjectPtr<USubsurfaceProfile>> SubsurfaceProfiles;
 
 	TArray<TObjectPtr<USpecularProfile>> SpecularProfiles;
+
+	TArray<TObjectPtr<UToonProfile>> ToonProfiles;
 	
 	/* Material expression GUIDs for all the traversed Substrate nodes */
 	bool bGatherMaterialExpressionGuids = false;
diff --git a/Engine/Source/Runtime/Engine/Classes/Engine/ToonProfile.h b/Engine/Source/Runtime/Engine/Classes/Engine/ToonProfile.h
new file mode 100644
--- /dev/null
+++ b/Engine/Source/Runtime/Engine/Classes/Engine/ToonProfile.h
@@ -0,0 +1,84 @@
+// Copyright Epic Games, Inc. All Rights Reserved.
+#pragma once
+
+#include "CoreMinimal.h"
+#include "UObject/ObjectMacros.h"
+#include "UObject/Object.h"
+#include "Misc/Guid.h"
+#include "RenderResource.h"
+#include "Curves/CurveLinearColor.h"
+#include "ToonProfile.generated.h"
+
+#define MAX_TOON_PROFILE_COUNT 256
+
+class UTexture2D;
+class FRDGBuilder;
+class FTextureReference;
+
+
+
+// Struct with all the settings we want in UToonProfile, separate to make it easer to pass this data around in the engine.
+USTRUCT(BlueprintType)
+struct FToonProfileStruct
+{
+	GENERATED_USTRUCT_BODY()
+
+	/**
+	* Define the diffuse ramp
+	*/
+	UPROPERTY(Category = "Procedural", EditAnywhere, meta = (AllowZoomOutput = "false", ShowZoomButtons = "false", ViewMinInput = "0", ViewMaxInput = "1", ViewMinOutput = "0", ViewMaxOutput = "1", TimelineLength = "1", ShowInputGridNumbers="false", ShowOutputGridNumbers="false"))
+	FRuntimeCurveLinearColor DiffuseRamp;
+
+	FToonProfileStruct();
+
+	void Invalidate()
+	{
+		*this = FToonProfileStruct();
+	}
+};
+
+/**
+ * Toon profile asset, can be specified at a material. 
+ * Don't change at runtime. All properties in here are per material.
+ */
+UCLASS(autoexpandcategories = ToonProfile, MinimalAPI)
+class UToonProfile : public UObject
+{
+	GENERATED_UCLASS_BODY()
+
+	UPROPERTY(Category = UToonProfile, EditAnywhere, meta = (ShowOnlyInnerProperties))
+	struct FToonProfileStruct Settings;
+
+	UPROPERTY()
+	FGuid Guid;
+
+	//~ Begin UObject Interface
+	virtual void BeginDestroy();
+	virtual void PostEditChangeProperty(struct FPropertyChangedEvent& PropertyChangedEvent);
+	virtual void PostDuplicate(EDuplicateMode::Type DuplicateMode);
+	//~ End UObject Interface
+};
+
+namespace ToonProfile
+{
+// Atlas - Initializes or updates the contents of the toon profile texture.
+ENGINE_API void UpdateToonProfileTextureAtlas(FRDGBuilder& GraphBuilder, EShaderPlatform ShaderPlatform);
+
+// Atlas - Returns the toon profile texture if it exists, or null.
+ENGINE_API FRHITexture* GetToonProfileTextureAtlas();
+
+// Atlas - Returns the toon profile texture if it exists, or black.
+ENGINE_API FRHITexture* GetToonProfileTextureAtlasWithFallback();
+
+// Profile - Initializes or updates the contents of the toon profile texture.
+ENGINE_API int32 AddOrUpdateProfile(const UToonProfile* InProfile, const FGuid& InGuid, const FToonProfileStruct InSettings);
+
+// Profile - Returns the toon profile ID shader parameter name
+ENGINE_API FName GetToonProfileParameterName(const UToonProfile* InProfile);
+
+// Profile - Returns the toon profile ID for a given toon Profile object in [0, 255]
+ENGINE_API float GetToonProfileId(const UToonProfile* In);
+
+// Profile - Returns the shader parameter name for a toon profile.
+ENGINE_API FName CreateToonProfileParameterName(UToonProfile* InProfile);
+}
\ No newline at end of file
diff --git a/Engine/Source/Runtime/Engine/Private/Materials/HLSLMaterialTranslator.cpp b/Engine/Source/Runtime/Engine/Private/Materials/HLSLMaterialTranslator.cpp
--- a/Engine/Source/Runtime/Engine/Private/Materials/HLSLMaterialTranslator.cpp
+++ b/Engine/Source/Runtime/Engine/Private/Materials/HLSLMaterialTranslator.cpp
@@ -14136,6 +14136,8 @@ bool FHLSLMaterialTranslator::FSubstrateCompilationContext::SubstrateGenerateDer
 					{
 						SubstrateMaterialRequestedSizeByte += UintByteSize;
 						SubstrateMaterialRequestedSizeByte += UintByteSize;
+						// We cannot add more uint here to stay compatible with legacy gbuffer when Substrate::IsSubstrateBlendableGBufferEnabled(ShaderPlatform).
+						// Possible work around: disable some toon shading feature that cannot be stored in the GBuffer, or increase Custom GBufferD size to 64 bits.
 						break;
 					}
 					case SUBSTRATE_BSDF_TYPE_UNLIT:
@@ -14923,13 +14925,20 @@ int32 FHLSLMaterialTranslator::SubstrateSingleLayerWaterBSDF(
 }
 
 int32 FHLSLMaterialTranslator::SubstrateToonBSDF(
+	int32 ToonProfileId,
 	int32 DiffuseColor, int32 SpecularColor, int32 Roughness, int32 Anisotropy,
 	int32 EmissiveColor, int32 Normal,
 	const FString& SharedLocalBasisIndexMacro, FSubstrateOperator* PromoteToOperator)
 {
+	if (ToonProfileId <= INDEX_NONE)
+	{
+		Errorf(TEXT("Substrate error: invalid Toon profile specified (asset: %s [%s]: ).\r\n"), *Material->GetDebugName(), *Material->GetAssetPath().ToString());
+		return INDEX_NONE;
+	}
 	return AddCodeChunk(
-		MCT_Substrate, TEXT("Parameters.%s.PromoteParameterBlendedBSDFToOperator(GetSubstrateToonBSDF(%s, %s, %s, %s, %s, %s, Parameters.%s.Types), %u, %u, %u, %u) /* Normal:%s */"),
+		MCT_Substrate, TEXT("Parameters.%s.PromoteParameterBlendedBSDFToOperator(GetSubstrateToonBSDF(%s, %s, %s, %s, %s, %s, %s, Parameters.%s.Types), %u, %u, %u, %u) /* Normal:%s */"),
 		*GetParametersSubstrateTreeName(CurrentSubstrateCompilationContext),
+		*SubstrateGetCastParameterCode(ToonProfileId, MCT_Float),
 		*SubstrateGetCastParameterCode(DiffuseColor, MCT_Float3),
 		*SubstrateGetCastParameterCode(SpecularColor, MCT_Float3),
 		*SubstrateGetCastParameterCode(Roughness, MCT_Float),
diff --git a/Engine/Source/Runtime/Engine/Private/Materials/HLSLMaterialTranslator.h b/Engine/Source/Runtime/Engine/Private/Materials/HLSLMaterialTranslator.h
--- a/Engine/Source/Runtime/Engine/Private/Materials/HLSLMaterialTranslator.h
+++ b/Engine/Source/Runtime/Engine/Private/Materials/HLSLMaterialTranslator.h
@@ -1415,6 +1415,7 @@ protected:
 		int32 EmissiveColor, int32 TopMaterialOpacity, int32 WaterAlbedo, int32 WaterExtinction, int32 WaterPhaseG, 
 		int32 ColorScaleBehindWater, int32 Normal, const FString& SharedLocalBasisIndexMacro, FSubstrateOperator* PromoteToOperator) override;
 	virtual int32 SubstrateToonBSDF(
+		int32 ToonProfileId,
 		int32 DiffuseColor, int32 SpecularColor, int32 Roughness, int32 Anisotropy,
 		int32 EmissiveColor, int32 Normal, 
 		const FString& SharedLocalBasisIndexMacro, FSubstrateOperator* PromoteToOperator) override;
diff --git a/Engine/Source/Runtime/Engine/Private/Materials/Material.cpp b/Engine/Source/Runtime/Engine/Private/Materials/Material.cpp
--- a/Engine/Source/Runtime/Engine/Private/Materials/Material.cpp
+++ b/Engine/Source/Runtime/Engine/Private/Materials/Material.cpp
@@ -53,6 +53,7 @@
 #include "Materials/MaterialUniformExpressions.h"
 #include "Engine/SubsurfaceProfile.h"
 #include "Engine/SpecularProfile.h"
+#include "Engine/ToonProfile.h"
 #include "Engine/TextureCollection.h"
 #include "EditorSupportDelegates.h"
 #include "ComponentRecreateRenderStateContext.h"
@@ -507,6 +508,19 @@ public:
 					}
 				}
 			}
+			else if (Type == EMaterialParameterType::Scalar && NumToonProfileRT() > 0)
+			{
+				for (uint32 It=0,Count=NumToonProfileRT();It<Count;++It)
+				{
+					if (ParameterInfo.Name == ToonProfile::GetToonProfileParameterName(GetToonProfileRT(It)))
+					{
+						const UToonProfile* ToonProfile = GetToonProfileRT(It);
+						const UToonProfile* ToonProfileOverride = GetToonProfileOverrideRT();
+						OutValue = ToonProfile::GetToonProfileId(ToonProfileOverride ? ToonProfileOverride : ToonProfile);
+						return true;
+					}
+				}
+			}
 
 			return false;
 		}
@@ -5931,6 +5945,13 @@ void UMaterial::RebuildShadingModelField()
 		{
 			SpecularProfiles[It] = SubstrateMaterialInfo.GetSpecularProfile(It);
 		}
+
+		// Set Toon profile if any
+		ToonProfiles.SetNum(SubstrateMaterialInfo.CountToonProfiles());
+		for (int32 It = 0, Count = SubstrateMaterialInfo.CountToonProfiles(); It<Count; ++It)
+		{
+			ToonProfiles[It] = SubstrateMaterialInfo.GetToonProfile(It);
+		}
 	}
 	// If using shading model from material expression, go through the expressions and look for the ShadingModel expression to figure out what shading models need to be supported in this material.
 	// This might not be the same as what is actually compiled in to the shader, since there might be feature switches, static switches etc. that skip certain shading models.
@@ -7563,6 +7584,18 @@ USpecularProfile* UMaterial::GetSpecularProfile_Internal(uint32 Index) const
 	return SpecularProfiles[Index];
 }
 
+uint32 UMaterial::NumToonProfile_Internal() const
+{
+	return ToonProfiles.Num();
+}
+
+UToonProfile* UMaterial::GetToonProfile_Internal(uint32 Index) const
+{
+	checkSlow(IsInGameThread());
+	check(Index<uint32(ToonProfiles.Num()));
+	return ToonProfiles[Index];
+}
+
 UNeuralProfile* UMaterial::GetNeuralProfile_Internal() const
 {
 	checkSlow(IsInGameThread());
diff --git a/Engine/Source/Runtime/Engine/Private/Materials/MaterialExpressionSubstrate.cpp b/Engine/Source/Runtime/Engine/Private/Materials/MaterialExpressionSubstrate.cpp
--- a/Engine/Source/Runtime/Engine/Private/Materials/MaterialExpressionSubstrate.cpp
+++ b/Engine/Source/Runtime/Engine/Private/Materials/MaterialExpressionSubstrate.cpp
@@ -11,6 +11,7 @@
 #include "Materials/MaterialExpressionMakeMaterialAttributes.h"
 #include "Engine/SubsurfaceProfile.h"
 #include "Engine/SpecularProfile.h"
+#include "Engine/ToonProfile.h"
 #include "MaterialExpressionIO.h"
 
 #if WITH_EDITOR
@@ -2460,9 +2461,21 @@ int32 UMaterialExpressionSubstrateToonBSDF::Compile(class FMaterialCompiler* Com
 		return Compiler->Errorf(TEXT("Substrate Toon BSDF node cannot be the root of a parameter blending sub tree."));
 	}
 
+	int32 ToonProfileIdCodeChunk = INDEX_NONE;
+	if (ToonProfile != nullptr)
+	{
+		const FName ToonProfileParameterName = ToonProfile::CreateToonProfileParameterName(ToonProfile);
+		ToonProfileIdCodeChunk = Compiler->ForceCast(Compiler->ScalarParameter(ToonProfileParameterName, 1.0f), MCT_Float1);
+	}
+	if (ToonProfileIdCodeChunk == INDEX_NONE)
+	{
+		ToonProfileIdCodeChunk = Compiler->Constant(0.0f);
+	}
+
 	int32 OutputCodeChunk = Compiler->SubstrateToonBSDF(
+		ToonProfileIdCodeChunk,
 		CompileWithDefaultFloat3(Compiler, DiffuseColor, 0.0f, 0.0f, 0.0f),
-		CompileWithDefaultFloat1(Compiler, SpecularColor, 0.0f),
+		CompileWithDefaultFloat3(Compiler, SpecularColor, 0.0f, 0.0f, 0.0f),
 		CompileWithDefaultFloat1(Compiler, Roughness, 0.5f),
 		CompileWithDefaultFloat1(Compiler, Anisotropy, 0.0f),
 		CompileWithDefaultFloat3(Compiler, EmissiveColor, 0.0f, 0.0f, 0.0f),
@@ -2514,6 +2527,11 @@ void UMaterialExpressionSubstrateToonBSDF::GatherSubstrateMaterialInfo(FSubstrat
 {
 	SubstrateMaterialInfo.AddShadingModel(SSM_Toon);
 	SubstrateMaterialInfo.AddGuid(MaterialExpressionGuid);
+
+	if (ToonProfile != nullptr)
+	{
+		SubstrateMaterialInfo.AddToonProfile(ToonProfile);
+	}
 }
 
 FSubstrateOperator* UMaterialExpressionSubstrateToonBSDF::SubstrateGenerateMaterialTopologyTree(class FMaterialCompiler* Compiler, class UMaterialExpression* Parent, int32 OutputIndex)
diff --git a/Engine/Source/Runtime/Engine/Private/Materials/MaterialInstance.cpp b/Engine/Source/Runtime/Engine/Private/Materials/MaterialInstance.cpp
--- a/Engine/Source/Runtime/Engine/Private/Materials/MaterialInstance.cpp
+++ b/Engine/Source/Runtime/Engine/Private/Materials/MaterialInstance.cpp
@@ -39,6 +39,7 @@
 #include "Materials/MaterialSharedPrivate.h"
 #include "Engine/SubsurfaceProfile.h"
 #include "Engine/SpecularProfile.h"
+#include "Engine/ToonProfile.h"
 #include "ProfilingDebugging/CookStats.h"
 #include "ProfilingDebugging/LoadTimeTracker.h"
 #include "ObjectCacheEventSink.h"
@@ -351,6 +352,20 @@ bool FMaterialInstanceResource::GetParameterValue(EMaterialParameterType Type, c
 			}
 		}
 	}
+	else if (Type == EMaterialParameterType::Scalar && NumToonProfileRT() > 0)
+	{
+		const UToonProfile* SPOverrideRT = GetToonProfileOverrideRT();
+		for (uint32 It=0,Count=NumToonProfileRT();It<Count;++It)
+		{
+			if (ParameterInfo.Name == ToonProfile::GetToonProfileParameterName(GetToonProfileRT(It)))
+			{
+				check(ParameterInfo.Association == EMaterialParameterAssociation::GlobalParameter);
+				OutValue = ToonProfile::GetToonProfileId(SPOverrideRT ? SPOverrideRT : GetToonProfileRT(It));
+				bResult = true;
+				break;
+			}
+		}
+	}
 	if (!bResult)
 	{
 		// Check for instances overrides
@@ -5319,7 +5334,7 @@ USubsurfaceProfile* UMaterialInstance::GetSubsurfaceProfileOverride_Internal() c
 
 uint32 UMaterialInstance::NumSpecularProfile_Internal() const
 {
-	// Return the subsurface profile count form the root material.
+	// Return the specular profile count form the root material.
 	checkSlow(IsInGameThread());
 	return Parent ? Parent->NumSpecularProfile_Internal() : 0;
 }
@@ -5347,6 +5362,36 @@ USpecularProfile* UMaterialInstance::GetSpecularProfileOverride_Internal() const
 	return Out;
 }
 
+uint32 UMaterialInstance::NumToonProfile_Internal() const
+{
+	// Return the toon profile count form the root material.
+	checkSlow(IsInGameThread());
+	return Parent ? Parent->NumToonProfile_Internal() : 0;
+}
+
+UToonProfile* UMaterialInstance::GetToonProfile_Internal(uint32 Index) const
+{
+	// Return the Toon profile from the root material.
+	checkSlow(IsInGameThread());
+
+	return Parent ? Parent->GetToonProfile_Internal(Index) : 0;
+}
+
+UToonProfile* UMaterialInstance::GetToonProfileOverride_Internal() const
+{
+	// Recursively look for override among the ancestors
+	UToonProfile* Out = nullptr;
+	if (bOverrideToonProfile && ToonProfileOverride)
+	{
+		Out = ToonProfileOverride;
+	}
+	else if (Parent)
+	{
+		Out = Parent->GetToonProfileOverride_Internal();
+	}
+	return Out;
+}
+
 bool UMaterialInstance::CastsRayTracedShadows() const
 {
 	//#dxr_todo: do per material instance override?
@@ -5574,6 +5619,10 @@ bool UMaterialInstance::IsRedundant() const
 	{
 		return false;
 	}
+	if (Substrate::IsSubstrateEnabled() && NumToonProfile_Internal() > 0 && GetToonProfile_Internal(0) != Parent->GetToonProfile_Internal(0))
+	{
+		return false;
+	}
 	// Assume that if any properties are overridden they are different to their parent
 	if (
 	   BasePropertyOverrides.bOverride_OpacityMaskClipValue 
diff --git a/Engine/Source/Runtime/Engine/Private/Materials/MaterialInterface.cpp b/Engine/Source/Runtime/Engine/Private/Materials/MaterialInterface.cpp
--- a/Engine/Source/Runtime/Engine/Private/Materials/MaterialInterface.cpp
+++ b/Engine/Source/Runtime/Engine/Private/Materials/MaterialInterface.cpp
@@ -26,6 +26,7 @@
 #include "Engine/SubsurfaceProfile.h"
 #include "Engine/SpecularProfile.h"
 #include "Engine/NeuralProfile.h"
+#include "Engine/ToonProfile.h"
 #include "Interfaces/ITargetPlatform.h"
 #include "Components/PrimitiveComponent.h"
 #include "ContentStreaming.h"
@@ -1637,6 +1638,16 @@ USpecularProfile* UMaterialInterface::GetSpecularProfileOverride_Internal() cons
 	return nullptr;
 }
 
+UToonProfile* UMaterialInterface::GetToonProfile_Internal(uint32 Index) const
+{
+	return nullptr;
+}
+
+UToonProfile* UMaterialInterface::GetToonProfileOverride_Internal() const
+{
+	return nullptr;
+}
+
 UNeuralProfile* UMaterialInterface::GetNeuralProfile_Internal() const
 {
 	return nullptr;
@@ -1652,6 +1663,11 @@ uint32 UMaterialInterface::NumSpecularProfile_Internal() const
 	return 0u;
 }
 
+uint32 UMaterialInterface::NumToonProfile_Internal() const
+{
+	return 0u;
+}
+
 bool UMaterialInterface::CastsRayTracedShadows() const
 {
 	return true;
@@ -1894,6 +1910,54 @@ void UMaterialInterface::UpdateMaterialRenderProxy(FMaterialRenderProxy& Proxy)
 				InProxy->SetSubsurfaceProfileRT(SubsurfaceProfileOverride);
 			});
 		}
+
+		// Toon profiles
+		{
+			struct FEntry
+			{
+				UToonProfile* Profile = nullptr;
+				FToonProfileStruct Settings;
+				const FTextureReference* Texture = nullptr;
+				FGuid Guid;
+			};
+			TArray<FEntry> Entries;
+			for (int32 It = 0, Count = NumToonProfile_Internal(); It < Count; ++It)
+			{
+				FEntry& Entry = Entries.AddDefaulted_GetRef();
+				Entry.Profile = GetToonProfile_Internal(It);
+				if (Entry.Profile)
+				{
+					Entry.Settings = Entry.Profile->Settings;
+					Entry.Guid = Entry.Profile->Guid;
+				}
+			}
+
+			UToonProfile* ToonProfileOverride = GetToonProfileOverride_Internal();
+			if (Entries.Num() > 0 || ToonProfileOverride)
+			{
+				ENQUEUE_RENDER_COMMAND(UpdateMaterialRenderProxyToon)(
+					[InProxy, Entries, ToonProfileOverride](FRHICommandListImmediate& RHICmdList)
+					{
+						for (const FEntry& Entry : Entries)
+						{
+							if (Entry.Profile)
+							{
+								const uint32 AllocationId = ToonProfile::AddOrUpdateProfile(Entry.Profile, Entry.Guid, Entry.Settings);
+								check(AllocationId >= 0 && AllocationId < MAX_TOON_PROFILE_COUNT);
+							}
+							InProxy->AddToonProfileRT(Entry.Profile);
+						}
+
+						if (ToonProfileOverride)
+						{
+							ToonProfile::AddOrUpdateProfile(ToonProfileOverride, ToonProfileOverride->Guid, ToonProfileOverride->Settings);
+						}
+
+						// Set the Toon profile override using the default ToonProfile pointer.
+						InProxy->SetToonProfileOverrideRT(ToonProfileOverride);
+					});
+			}
+		}
 	}
 }
 
diff --git a/Engine/Source/Runtime/Engine/Private/Rendering/ToonProfile.cpp b/Engine/Source/Runtime/Engine/Private/Rendering/ToonProfile.cpp
new file mode 100644
--- /dev/null
+++ b/Engine/Source/Runtime/Engine/Private/Rendering/ToonProfile.cpp
@@ -0,0 +1,511 @@
+// Copyright Epic Games, Inc. All Rights Reserved.
+
+#include "Engine/ToonProfile.h"
+#include "Engine/Texture2D.h"
+#include "DataDrivenShaderPlatformInfo.h"
+#include "Math/Float16.h"
+#include "Rendering/BurleyNormalizedSSS.h"
+#include "EngineModule.h"
+#include "RenderTargetPool.h"
+#include "PixelShaderUtils.h"
+#include "RenderingThread.h"
+#include "Rendering/Texture2DResource.h"
+#include "RenderGraphResources.h"
+#include "RenderGraphBuilder.h"
+#include "SubstrateDefinitions.h"
+#include "ShaderCompilerCore.h"
+
+#include UE_INLINE_GENERATED_CPP_BY_NAME(ToonProfile)
+
+DEFINE_LOG_CATEGORY_STATIC(LogToonProfile, Log, All);
+
+///////////////////////////////////////////////////////////////////////////////////////////////////
+
+
+static TAutoConsoleVariable<int32> CVarToonProfileResolution(
+	TEXT("r.Substrate.ToonProfile.Resolution"),
+	64,
+	TEXT("The resolution of the toon profile texture.\n"),
+	ECVF_RenderThreadSafe
+);
+
+#if !(UE_BUILD_SHIPPING || UE_BUILD_TEST)
+static TAutoConsoleVariable<int32> CVarToonProfileForceUpdate(
+	TEXT("r.Substrate.ToonProfile.ForceUpdate"),
+	0,
+	TEXT("0: Only update the toon profile as needed.\n")
+	TEXT("1: Force to update the toon profile every frame for debugging.\n"),
+	ECVF_Cheat | ECVF_RenderThreadSafe);
+#endif
+
+static bool ForceUpdateToonProfile()
+{
+#if !(UE_BUILD_SHIPPING || UE_BUILD_TEST)
+	return (CVarToonProfileForceUpdate.GetValueOnAnyThread() == 1);
+#else
+	return false;
+#endif
+}
+
+///////////////////////////////////////////////////////////////////////////////////////////////////
+
+class FToonProfileUpdateAtlasTileCS : public FGlobalShader
+{
+public:
+	DECLARE_GLOBAL_SHADER(FToonProfileUpdateAtlasTileCS);
+	SHADER_USE_PARAMETER_STRUCT(FToonProfileUpdateAtlasTileCS, FGlobalShader);
+
+	using FPermutationDomain = TShaderPermutationDomain<>;
+
+	BEGIN_SHADER_PARAMETER_STRUCT(FParameters, )
+		SHADER_PARAMETER(FIntPoint, ToonProfileTextureResolution)
+		SHADER_PARAMETER(uint32, ToonProfileIndex)
+		SHADER_PARAMETER_RDG_TEXTURE_UAV(RWTexture2DArray, ToonProfileTexture)
+		SHADER_PARAMETER_RDG_BUFFER_SRV(Buffer, DiffuseRampBufferSRV)
+	END_SHADER_PARAMETER_STRUCT()
+
+	static int32 GetThreadGroupSize()
+	{
+		return 8;
+	}
+
+	static bool ShouldCompilePermutation(const FGlobalShaderPermutationParameters& Parameters)
+	{
+		return Substrate::IsToonProfileEnabled();
+	}
+	static void ModifyCompilationEnvironment(const FGlobalShaderPermutationParameters& Parameters, FShaderCompilerEnvironment& OutEnvironment)
+	{
+		FGlobalShader::ModifyCompilationEnvironment(Parameters, OutEnvironment);
+		OutEnvironment.SetDefine(TEXT("THREAD_GROUP_COUNT"), GetThreadGroupSize());
+	}
+};
+
+IMPLEMENT_GLOBAL_SHADER(FToonProfileUpdateAtlasTileCS, "/Engine/Private/ToonProfile.usf", "ToonProfileUpdateAtlasTileCS", SF_Compute);
+
+///////////////////////////////////////////////////////////////////////////////////////////////////
+// FToonProfileTextureManager
+
+struct FToonProfileIdQuery
+{
+	bool bFound = false;		// default to not found
+	uint32 AllocationId = 0;	// default entry
+};
+
+// render thread
+class FToonProfileTextureManager : public FRenderResource
+{
+public:
+	// constructor
+	FToonProfileTextureManager();
+
+	// destructor
+	~FToonProfileTextureManager();
+
+
+	// convenience, can be optimized 
+	// @param Profile must not be 0, game thread pointer, do not dereference, only for comparison
+	int32 AddOrUpdateProfile(const UToonProfile* InProfile, const FGuid& InGuid, const FToonProfileStruct InSettings);
+
+	// O(n) n is a small number
+	// @param InProfile must not be 0, game thread pointer, do not dereference, only for comparison
+	// @return AllocationId INDEX_NONE: no allocation, should be deallocated with DeallocateToonProfile()
+	int32 AddProfile(const UToonProfile* InProfile, const FGuid& InGuid, const FToonProfileStruct Settings);
+
+	// O(n) to find the element, n is the SSProfile count and usually quite small
+	void RemoveProfile(const UToonProfile* InProfile);
+
+	// @param InProfile must not be 0, game thread pointer, do not dereference, only for comparison
+	void UpdateProfile(const UToonProfile* InProfile, const FToonProfileStruct InSettings);
+
+	// @param InProfile must not be 0, game thread pointer, do not dereference, only for comparison
+	void UpdateProfile(FToonProfileIdQuery ToonProfileIdQuery, const FToonProfileStruct Settings);
+
+	// return the parameter name for a given profile
+	FName GetParameterName(const UToonProfile* InProfile) const;
+
+	// @return can be nullptr if there is no ToonProfile
+	IPooledRenderTarget* GetAtlasTexture(FRDGBuilder& GraphBuilder, EShaderPlatform ShaderPlatform);
+	IPooledRenderTarget* GetAtlasTexture();
+
+	//~ Begin FRenderResource Interface.
+	/**
+		* Release textures when device is lost/destroyed.
+		*/
+	virtual void ReleaseRHI() override;
+
+	// @param InProfile must not be 0, game thread pointer, do not dereference, only for comparison
+	// @return INDEX_NONE if not found
+	FToonProfileIdQuery FindAllocationId(const UToonProfile* InProfile) const;
+
+private:
+
+	struct FProfileEntry
+	{
+		FToonProfileStruct Settings;
+		const UToonProfile* Profile = nullptr; // Game thread pointer! Do not dereference, only for comparison.
+		FName ParameterName;
+	};
+	TArray<FProfileEntry> ToonProfileEntries;
+};
+
+// Global resources - lives on the render thread
+TGlobalResource<FToonProfileTextureManager> GToonProfileTextureManager;
+
+// ToonProfile atlas storing several texture profiles or 0 if there is no user
+static TRefCountPtr<IPooledRenderTarget> GToonProfileTextureAtlas;
+
+FToonProfileTextureManager::FToonProfileTextureManager()
+{
+	check(IsInGameThread());
+
+	// add element 0, it is used as default profile
+	FProfileEntry& Entry = ToonProfileEntries.AddDefaulted_GetRef();
+	Entry.Settings = FToonProfileStruct();
+	Entry.Profile = nullptr;
+}
+
+FToonProfileTextureManager::~FToonProfileTextureManager()
+{
+}
+
+int32 FToonProfileTextureManager::AddOrUpdateProfile(const UToonProfile* InProfile, const FGuid& Guid, const FToonProfileStruct InSettings)
+{
+	check(InProfile);
+
+	FToonProfileIdQuery ToonProfileIdQuery = FindAllocationId(InProfile);
+		
+	if (ToonProfileIdQuery.bFound)
+	{
+		UpdateProfile(ToonProfileIdQuery, InSettings);
+		return ToonProfileIdQuery.AllocationId;
+	}
+
+	return AddProfile(InProfile, Guid, InSettings);
+}
+
+int32 FToonProfileTextureManager::AddProfile(const UToonProfile* InProfile, const FGuid& Guid, const FToonProfileStruct InSettings)
+{
+	check(InProfile);
+	check(!FindAllocationId(InProfile).bFound);
+
+	FToonProfileIdQuery ToonProfileIdQuery;
+	{
+		for (int32 i = 1; i < ToonProfileEntries.Num(); ++i)
+		{
+			if (ToonProfileEntries[i].Profile == nullptr)
+			{
+				ToonProfileIdQuery.bFound = true;
+				ToonProfileIdQuery.AllocationId = i;
+
+				FProfileEntry& Entry= ToonProfileEntries[ToonProfileIdQuery.AllocationId];
+				Entry.Profile 		= InProfile;
+				Entry.ParameterName = FName(TEXT("__ToonProfile") + Guid.ToString());
+				break;
+			}
+		}
+
+		if (!ToonProfileIdQuery.bFound)
+		{
+			ToonProfileIdQuery.bFound = true;
+			ToonProfileIdQuery.AllocationId = ToonProfileEntries.Num();
+
+			FProfileEntry& Entry= ToonProfileEntries.AddDefaulted_GetRef();
+			Entry.Profile 		= InProfile;
+			Entry.ParameterName = FName(TEXT("__ToonProfile") + Guid.ToString());
+		}
+	}
+
+	UpdateProfile(ToonProfileIdQuery, InSettings);
+
+	return ToonProfileIdQuery.AllocationId;
+}
+
+void FToonProfileTextureManager::RemoveProfile(const UToonProfile* InProfile)
+{
+	FToonProfileIdQuery ToonProfileIdQuery = FindAllocationId(InProfile);
+
+	if (ToonProfileIdQuery.bFound)
+	{
+		// >0 as 0 is used as default profile which should never be removed
+		check(ToonProfileIdQuery.AllocationId > 0);
+		check(ToonProfileEntries[ToonProfileIdQuery.AllocationId].Profile == InProfile);
+
+		// make it available for reuse
+		ToonProfileEntries[ToonProfileIdQuery.AllocationId].Profile = nullptr;
+		ToonProfileEntries[ToonProfileIdQuery.AllocationId].Settings.Invalidate();
+		ToonProfileEntries[ToonProfileIdQuery.AllocationId].ParameterName = FName();
+	}
+}
+
+void FToonProfileTextureManager::UpdateProfile(const UToonProfile* InProfile, const FToonProfileStruct InSettings) 
+{ 
+	UpdateProfile(FindAllocationId(InProfile), InSettings); 
+}
+
+void FToonProfileTextureManager::UpdateProfile(FToonProfileIdQuery ToonProfileIdQuery, const FToonProfileStruct Settings)
+{
+	check(IsInRenderingThread());
+
+	if (ToonProfileIdQuery.bFound)
+	{
+		check(ToonProfileIdQuery.AllocationId < uint32(ToonProfileEntries.Num()));
+		ToonProfileEntries[ToonProfileIdQuery.AllocationId].Settings = Settings;
+		GToonProfileTextureAtlas.SafeRelease();
+	}
+}
+
+FName FToonProfileTextureManager::GetParameterName(const UToonProfile* InProfile) const
+{
+	FToonProfileIdQuery ToonProfileIdQuery = FindAllocationId(InProfile);
+	if (ToonProfileIdQuery.bFound)
+	{
+		// make it available for reuse
+		return ToonProfileEntries[ToonProfileIdQuery.AllocationId].ParameterName;
+	}
+	return FName();
+}
+
+IPooledRenderTarget* FToonProfileTextureManager::GetAtlasTexture()
+{
+	return GToonProfileTextureAtlas;
+}
+
+IPooledRenderTarget* FToonProfileTextureManager::GetAtlasTexture(FRDGBuilder& GraphBuilder, EShaderPlatform ShaderPlatform)
+{
+	if (!Substrate::IsToonProfileEnabled())
+	{
+		return nullptr;
+	}
+
+	const uint32 ToonProfileCount = ToonProfileEntries.Num();
+	check(ToonProfileCount);
+
+	// Since reference texture can be streamed/loaded progressively, we track if textures have been updated to update the LUT
+	bool bForceUpdate = ForceUpdateToonProfile();
+
+	if (bForceUpdate)
+	{
+		GToonProfileTextureAtlas.SafeRelease();
+		GToonProfileTextureAtlas = nullptr;
+	}
+
+	if (GToonProfileTextureAtlas == nullptr)
+	{
+		FRHICommandListImmediate& RHICmdList = GraphBuilder.RHICmdList;
+
+		const uint32 ToonProfileResolution = CVarToonProfileResolution.GetValueOnRenderThread();
+
+		RDG_EVENT_SCOPE(GraphBuilder, "ToonProfile");
+
+		// 1. Create atlas texture
+		const FIntPoint ToonProfileResolution2D = FIntPoint(ToonProfileResolution, 8);
+		const EPixelFormat Format = PF_FloatRGBA;
+		FPooledRenderTargetDesc Desc(FPooledRenderTargetDesc::Create2DArrayDesc(ToonProfileResolution2D, Format, FClearValueBinding::None, TexCreate_None, TexCreate_UAV | TexCreate_ShaderResource, false, ToonProfileCount));
+		GRenderTargetPool.FindFreeElement(RHICmdList, Desc, GToonProfileTextureAtlas, TEXT("ToonProfileTexture"));
+		const auto GlobalShaderMap = GetGlobalShaderMap(ShaderPlatform);
+		FRDGTextureRef ToonProfileTexture = GraphBuilder.RegisterExternalTexture(GToonProfileTextureAtlas, TEXT("ToonProfileTexture"));
+		
+		// 2. Optionally clear profiles texture
+		if (bForceUpdate)
+		{
+			AddClearUAVPass(GraphBuilder, GraphBuilder.CreateUAV(ToonProfileTexture), 0.f);
+		}
+
+		// 3. Fill in profiles
+		FRDGTextureUAVRef ToonProfileUAVSkipBarrier = GraphBuilder.CreateUAV(ToonProfileTexture, ERDGUnorderedAccessViewFlags::SkipBarrier);
+		for (uint32 LayerIt = 0; LayerIt < ToonProfileCount; ++LayerIt)
+		{
+			const FToonProfileStruct Data = ToonProfileEntries[LayerIt].Settings;
+
+			FToonProfileUpdateAtlasTileCS::FParameters* PassParameters = GraphBuilder.AllocParameters<FToonProfileUpdateAtlasTileCS::FParameters>();
+			PassParameters->ToonProfileTexture = ToonProfileUAVSkipBarrier;
+			PassParameters->ToonProfileTextureResolution = ToonProfileTexture->Desc.Extent;
+			PassParameters->ToonProfileIndex = LayerIt;
+
+			auto CreateBufferFRomCurve = [&](FRuntimeCurveLinearColor ColorCurve)
+			{
+					const uint32 ComponentsPerElement = 4;
+				const uint32 BytesPerComponent = sizeof(FFloat16);
+				const uint32 BytesPerElement = ComponentsPerElement * BytesPerComponent;
+
+				TArray<FFloat16> ColorCurveData;
+				ColorCurveData.SetNum(ToonProfileResolution * ComponentsPerElement);
+				for (uint32 It = 0; It < ToonProfileResolution; ++It)
+				{
+					FLinearColor CurveColorSample = ColorCurve.GetLinearColorValue(float(It + 0.5f) / ToonProfileResolution);
+					ColorCurveData[It * 4 + 0] = CurveColorSample.R;
+					ColorCurveData[It * 4 + 1] = CurveColorSample.G;
+					ColorCurveData[It * 4 + 2] = CurveColorSample.B;
+					ColorCurveData[It * 4 + 3] = CurveColorSample.A;
+				}
+
+				FRDGBufferRef ColorCurveBuffer = GraphBuilder.CreateBuffer(FRDGBufferDesc::CreateBufferDesc(BytesPerElement, ToonProfileResolution), TEXT("ColorCurveBuffer"), ERDGBufferFlags::MultiFrame);
+				GraphBuilder.QueueBufferUpload(ColorCurveBuffer, ColorCurveData.GetData(), BytesPerComponent * ColorCurveData.Num(), ERDGInitialDataFlags::None);
+
+				return ColorCurveBuffer;
+			};
+
+			FRDGBufferRef DiffuseRampBuffer = CreateBufferFRomCurve(Data.DiffuseRamp);
+			PassParameters->DiffuseRampBufferSRV = GraphBuilder.CreateSRV(DiffuseRampBuffer, Format);
+
+			FToonProfileUpdateAtlasTileCS::FPermutationDomain PermutationVector;
+			TShaderMapRef<FToonProfileUpdateAtlasTileCS> Shader(GlobalShaderMap, PermutationVector);
+			FComputeShaderUtils::AddPass(GraphBuilder, 
+				RDG_EVENT_NAME("ToonProfile::UpdateAtlasTileCS"), 
+				Shader, 
+				PassParameters, 
+				FIntVector(
+					FMath::DivideAndRoundUp(PassParameters->ToonProfileTextureResolution.X, FToonProfileUpdateAtlasTileCS::GetThreadGroupSize()), 
+					FMath::DivideAndRoundUp(PassParameters->ToonProfileTextureResolution.Y, FToonProfileUpdateAtlasTileCS::GetThreadGroupSize()), 
+					1)
+			);
+		}
+
+		// Transit texture to SRV for letting the non-RDG resource to be bound later by the various passes
+		FRDGExternalAccessQueue ExternalAccessQueue;
+		ExternalAccessQueue.Add(ToonProfileTexture);
+		ExternalAccessQueue.Submit(GraphBuilder);
+	}
+	return GToonProfileTextureAtlas;
+}
+
+void FToonProfileTextureManager::ReleaseRHI()
+{
+	GToonProfileTextureAtlas.SafeRelease();
+}
+
+FToonProfileIdQuery FToonProfileTextureManager::FindAllocationId(const UToonProfile* InProfile) const
+{
+	FToonProfileIdQuery ToonProfileIdQuery;
+
+	// we start at 1 because [0] is the default profile and always [0].Profile = 0 so we don't need to iterate that one
+	for (int32 i = 1; i < ToonProfileEntries.Num(); ++i)
+	{
+		if (ToonProfileEntries[i].Profile == InProfile)
+		{
+			ToonProfileIdQuery.bFound = true;
+			ToonProfileIdQuery.AllocationId = i;
+			return ToonProfileIdQuery;
+		}
+	}
+	return ToonProfileIdQuery;
+}
+
+///////////////////////////////////////////////////////////////////////////////////////////////////
+// FToonProfileStruct
+
+FToonProfileStruct::FToonProfileStruct()
+{
+	auto AddPointToCurve = [](FRuntimeCurveLinearColor& LinearColorCurve, float Time, float Value)
+	{
+		LinearColorCurve.ColorCurves[0].AddKey(Time, Value);
+		LinearColorCurve.ColorCurves[1].AddKey(Time, Value);
+		LinearColorCurve.ColorCurves[2].AddKey(Time, Value);
+	};
+
+	AddPointToCurve(DiffuseRamp, 0.00, 0.00);
+	AddPointToCurve(DiffuseRamp, 0.25, 0.00);
+	AddPointToCurve(DiffuseRamp, 0.26, 0.25);
+	AddPointToCurve(DiffuseRamp, 0.50, 0.25);
+	AddPointToCurve(DiffuseRamp, 0.51, 0.50);
+	AddPointToCurve(DiffuseRamp, 0.75, 0.50);
+	AddPointToCurve(DiffuseRamp, 0.76, 0.75);
+	AddPointToCurve(DiffuseRamp, 0.99, 0.75);
+	AddPointToCurve(DiffuseRamp, 1.00, 1.00);
+}
+
+///////////////////////////////////////////////////////////////////////////////////////////////////
+// UToonProfile
+
+UToonProfile::UToonProfile(const FObjectInitializer& ObjectInitializer)
+	: Super(ObjectInitializer)
+{
+}
+
+void UToonProfile::BeginDestroy()
+{
+	UToonProfile* Ref = this;
+	ENQUEUE_RENDER_COMMAND(RemoveToonProfile)(
+	[Ref](FRHICommandList& RHICmdList)
+	{
+		GToonProfileTextureManager.RemoveProfile(Ref);
+	});
+
+	Super::BeginDestroy();
+}
+
+void UToonProfile::PostEditChangeProperty(struct FPropertyChangedEvent& PropertyChangedEvent)
+{
+	const FToonProfileStruct LocalSettings = this->Settings;
+	UToonProfile* LocalProfile = this;
+	GetRendererModule().InvalidatePathTracedOutput();
+
+	ENQUEUE_RENDER_COMMAND(UpdateToonProfile)(
+	[LocalSettings, LocalProfile](FRHICommandListImmediate& RHICmdList)
+	{
+		// any changes to the setting require an update of the texture
+		GToonProfileTextureManager.UpdateProfile(LocalProfile, LocalSettings);
+	});
+}
+
+void UToonProfile::PostDuplicate(EDuplicateMode::Type DuplicateMode)
+{
+	Super::PostDuplicate(DuplicateMode);
+
+	// When a Toon Profile asset is duplicated/copied pasted (e.g. from the asset browser), we want the guid to be regenerated.
+	Guid = FGuid::NewGuid();
+}
+
+///////////////////////////////////////////////////////////////////////////////////////////////////
+// Public API
+
+namespace ToonProfile
+{
+	FName GetToonProfileParameterName(const UToonProfile* In)
+	{
+		return GToonProfileTextureManager.GetParameterName(In);
+	}
+
+	float GetToonProfileId(const UToonProfile* In)
+	{
+		// No profile specified means we use the default one (constant one)
+		int32 AllocationId = 0;
+		if (In)
+		{
+			// can be optimized (cached)
+			AllocationId = GToonProfileTextureManager.FindAllocationId(In).AllocationId;
+		}
+		return AllocationId;
+	}
+	
+	int32 AddOrUpdateProfile(const UToonProfile* InProfile, const FGuid& InGuid, const FToonProfileStruct InSettings)
+	{
+		return GToonProfileTextureManager.AddOrUpdateProfile(InProfile, InGuid, InSettings);
+	}
+
+	FRHITexture* GetToonProfileTextureAtlas()
+	{
+		return GToonProfileTextureAtlas ? GToonProfileTextureAtlas->GetRHI() : nullptr;
+	}
+
+	FRHITexture* GetToonProfileTextureAtlasWithFallback()
+	{
+		return GToonProfileTextureAtlas ? GToonProfileTextureAtlas->GetRHI() : static_cast<FRHITexture*>(GWhiteTexture->TextureRHI);
+	}
+
+	void UpdateToonProfileTextureAtlas(FRDGBuilder& GraphBuilder, EShaderPlatform ShaderPlatform)
+	{
+		GToonProfileTextureManager.GetAtlasTexture(GraphBuilder, ShaderPlatform);
+	}
+
+	
+	static FName CreateToonProfileParameterName(const FGuid& InGuid)
+	{
+		return FName(TEXT("__ToonProfile") + InGuid.ToString());
+	}
+
+	FName CreateToonProfileParameterName(UToonProfile* InProfile)
+	{
+		return InProfile ? CreateToonProfileParameterName(InProfile->Guid) : FName();
+	}
+} // namespace ToonProfile
diff --git a/Engine/Source/Runtime/Engine/Private/SceneManagement.cpp b/Engine/Source/Runtime/Engine/Private/SceneManagement.cpp
--- a/Engine/Source/Runtime/Engine/Private/SceneManagement.cpp
+++ b/Engine/Source/Runtime/Engine/Private/SceneManagement.cpp
@@ -1306,6 +1306,11 @@ FViewUniformShaderParameters::FViewUniformShaderParameters()
 	SpecularProfileTexture = GBlackArrayTexture->TextureRHI;
 	SpecularProfileSampler = TStaticSamplerState<SF_Bilinear, AM_Clamp, AM_Clamp, AM_Clamp>::GetRHI();
 
+	// Toon profiles
+	ToonProfileTextureSizeAndInvSize = FVector4f(1.f, 1.f, 1.f, 1.f);
+	ToonProfileTexture = GBlackArrayTexture->TextureRHI;
+	ToonProfileSampler = TStaticSamplerState<SF_Bilinear, AM_Clamp, AM_Clamp, AM_Clamp>::GetRHI();
+
 	VTFeedbackBuffer = GEmptyStructuredBufferWithUAV->UnorderedAccessViewRHI;
 
 	BlueNoiseScalarTexture = GBlackTextureWithSRV->TextureRHI;
diff --git a/Engine/Source/Runtime/Engine/Public/MaterialCompiler.h b/Engine/Source/Runtime/Engine/Public/MaterialCompiler.h
--- a/Engine/Source/Runtime/Engine/Public/MaterialCompiler.h
+++ b/Engine/Source/Runtime/Engine/Public/MaterialCompiler.h
@@ -713,6 +713,7 @@ public:
 		int32 BaseColor, int32 Metallic, int32 Specular, int32 Roughness, int32 EmissiveColor, int32 TopMaterialOpacity, 
 		int32 WaterAlbedo, int32 WaterExtinction, int32 WaterPhaseG, int32 ColorScaleBehindWater, int32 Normal, const FString& SharedLocalBasisIndexMacro, FSubstrateOperator* PromoteToOperator) = 0;
 	virtual int32 SubstrateToonBSDF(
+		int32 ToonProfileId,
 		int32 DiffuseColor, int32 SpecularColor, int32 Roughness, int32 Anisotropy,
 		int32 EmissiveColor, int32 Normal,
 		const FString& SharedLocalBasisIndexMacro, FSubstrateOperator* PromoteToOperator) = 0;
@@ -1521,11 +1522,13 @@ public:
 	}
 
 	virtual int32 SubstrateToonBSDF(
+		int32 ToonProfileId,
 		int32 DiffuseColor, int32 SpecularColor, int32 Roughness, int32 Anisotropy,
 		int32 EmissiveColor, int32 Normal,
 		const FString& SharedLocalBasisIndexMacro, FSubstrateOperator* PromoteToOperator) override
 	{
 		return Compiler->SubstrateToonBSDF(
+			ToonProfileId,
 			DiffuseColor, SpecularColor, Roughness, Anisotropy, 
 			EmissiveColor, Normal,
 			SharedLocalBasisIndexMacro, PromoteToOperator);
diff --git a/Engine/Source/Runtime/Engine/Public/Materials/Material.h b/Engine/Source/Runtime/Engine/Public/Materials/Material.h
--- a/Engine/Source/Runtime/Engine/Public/Materials/Material.h
+++ b/Engine/Source/Runtime/Engine/Public/Materials/Material.h
@@ -32,6 +32,7 @@ class UPhysicalMaterial;
 class UPhysicalMaterialMask;
 class USubsurfaceProfile;
 class USpecularProfile;
+class UToonProfile;
 class UTexture;
 struct FODSCRequestPayload;
 
@@ -1304,6 +1305,8 @@ public:
 	ENGINE_API virtual uint32 NumSpecularProfile_Internal() const override;
 	ENGINE_API virtual USpecularProfile* GetSpecularProfile_Internal(uint32 Index) const override;
 	ENGINE_API virtual UNeuralProfile* GetNeuralProfile_Internal() const override;
+	ENGINE_API virtual uint32 NumToonProfile_Internal() const override;
+	ENGINE_API virtual UToonProfile* GetToonProfile_Internal(uint32 Index) const override;
 	ENGINE_API virtual bool CastsRayTracedShadows() const override;
 	ENGINE_API virtual bool IsTessellationEnabled() const override;
 	ENGINE_API virtual FDisplacementScaling GetDisplacementScaling() const override;
diff --git a/Engine/Source/Runtime/Engine/Public/Materials/MaterialExpressionSubstrate.h b/Engine/Source/Runtime/Engine/Public/Materials/MaterialExpressionSubstrate.h
--- a/Engine/Source/Runtime/Engine/Public/Materials/MaterialExpressionSubstrate.h
+++ b/Engine/Source/Runtime/Engine/Public/Materials/MaterialExpressionSubstrate.h
@@ -861,8 +861,8 @@ class UMaterialExpressionSubstrateToonBSDF : public UMaterialExpressionSubstrate
 	FExpressionInput EmissiveColor;
 
 	/** Toon profile used to control the appearance of this BSDF. */
-//	UPROPERTY(EditAnywhere, BlueprintReadOnly, Category = Material, meta = (DisplayName = "Toon Profile"))
-//	TObjectPtr<class UToonProfile> ToonProfile;
+	UPROPERTY(EditAnywhere, BlueprintReadOnly, Category = Material, meta = (DisplayName = "Toon Profile"))
+	TObjectPtr<class UToonProfile> ToonProfile;
 
 	//~ Begin UMaterialExpression Interface
 #if WITH_EDITOR
diff --git a/Engine/Source/Runtime/Engine/Public/Materials/MaterialInstance.h b/Engine/Source/Runtime/Engine/Public/Materials/MaterialInstance.h
--- a/Engine/Source/Runtime/Engine/Public/Materials/MaterialInstance.h
+++ b/Engine/Source/Runtime/Engine/Public/Materials/MaterialInstance.h
@@ -668,6 +668,10 @@ class UMaterialInstance : public UMaterialInterface
 	UPROPERTY()
 	TObjectPtr<class USpecularProfile> SpecularProfileOverride;
 
+	/** Defines ToonProfile override used by this instance. If not set, it uses the parent ones. */
+	UPROPERTY()
+	TObjectPtr<class UToonProfile> ToonProfileOverride;
+
 	/**
 	 * Indicates whether the instance has static permutation resources (which are required when static parameters are present) 
 	 * Read directly from the rendering thread, can only be modified with the use of a FMaterialUpdateContext.
@@ -684,6 +688,10 @@ class UMaterialInstance : public UMaterialInterface
 	UPROPERTY(EditAnywhere, BlueprintReadOnly, Category = MaterialInstance)
 	uint8 bOverrideSpecularProfile:1;
 
+	/** Defines if ToonProfile from this instance is used or it uses the parent one. */
+	UPROPERTY(EditAnywhere, BlueprintReadOnly, Category = MaterialInstance)
+	uint8 bOverrideToonProfile:1;
+
 	/** For post process materials, use BlendableLocationOverride. */
 	UPROPERTY(EditAnywhere, BlueprintReadWrite, Category = MaterialInstance)
 	uint8 bOverrideBlendableLocation:1;
@@ -985,6 +993,9 @@ public:
 	ENGINE_API virtual uint32 NumSpecularProfile_Internal() const override;
 	ENGINE_API virtual USpecularProfile* GetSpecularProfile_Internal(uint32 Index) const override;
 	ENGINE_API virtual USpecularProfile* GetSpecularProfileOverride_Internal() const override;
+	ENGINE_API virtual uint32 NumToonProfile_Internal() const override;
+	ENGINE_API virtual UToonProfile* GetToonProfile_Internal(uint32 Index) const override;
+	ENGINE_API virtual UToonProfile* GetToonProfileOverride_Internal() const override;
 	ENGINE_API virtual bool CastsRayTracedShadows() const override;
 	ENGINE_API virtual bool IsTessellationEnabled() const override;
 	ENGINE_API virtual bool HasSubstrateRoughnessTracking() const override;
diff --git a/Engine/Source/Runtime/Engine/Public/Materials/MaterialInterface.h b/Engine/Source/Runtime/Engine/Public/Materials/MaterialInterface.h
--- a/Engine/Source/Runtime/Engine/Public/Materials/MaterialInterface.h
+++ b/Engine/Source/Runtime/Engine/Public/Materials/MaterialInterface.h
@@ -47,6 +47,7 @@ class UPhysicalMaterialMask;
 class USubsurfaceProfile;
 class USpecularProfile;
 class UNeuralProfile;
+class UToonProfile;
 class UTexture;
 class UMaterialInstance;
 struct FDebugShaderTypeInfo;
@@ -362,6 +363,10 @@ public:
 	UPROPERTY()
 	TArray<TObjectPtr<class USpecularProfile>> SpecularProfiles;
 
+	/** Toon Profile. For internal usage, not editable/visible */
+	UPROPERTY()
+	TArray<TObjectPtr<class UToonProfile>> ToonProfiles;
+
 	/** Neural network profile. For internal usage, not editable/visible */
 	UPROPERTY(EditAnywhere, BlueprintReadOnly, Category = PostProcessMaterial, meta = (DisplayName = "Neural Profile"))
 	TObjectPtr<class UNeuralProfile> NeuralProfile;
@@ -1091,6 +1096,9 @@ public:
 	ENGINE_API virtual uint32 NumSpecularProfile_Internal() const;
 	ENGINE_API virtual USpecularProfile* GetSpecularProfile_Internal(uint32 Index) const;
 	ENGINE_API virtual USpecularProfile* GetSpecularProfileOverride_Internal() const;
+	ENGINE_API virtual uint32 NumToonProfile_Internal() const;
+	ENGINE_API virtual UToonProfile* GetToonProfile_Internal(uint32 Index) const;
+	ENGINE_API virtual UToonProfile* GetToonProfileOverride_Internal() const;
 	ENGINE_API virtual UNeuralProfile* GetNeuralProfile_Internal() const;
 	ENGINE_API virtual bool CastsRayTracedShadows() const;
 	ENGINE_API virtual bool IsTessellationEnabled() const;
diff --git a/Engine/Source/Runtime/Engine/Public/Materials/MaterialRenderProxy.h b/Engine/Source/Runtime/Engine/Public/Materials/MaterialRenderProxy.h
--- a/Engine/Source/Runtime/Engine/Public/Materials/MaterialRenderProxy.h
+++ b/Engine/Source/Runtime/Engine/Public/Materials/MaterialRenderProxy.h
@@ -26,6 +26,7 @@ class URuntimeVirtualTexture;
 class USparseVolumeTexture;
 class USubsurfaceProfile;
 class USpecularProfile;
+class UToonProfile;
 class UNeuralProfile;
 class UTexture;
 class UTextureCollection;
@@ -238,6 +239,14 @@ public:
 	void SetSpecularProfileOverrideRT(const USpecularProfile* Ptr) { SpecularProfileOverrideRT = Ptr; }
 	const USpecularProfile* GetSpecularProfileOverrideRT() const { return SpecularProfileOverrideRT; }
 
+	// Toon profiles
+	void AddToonProfileRT(const UToonProfile* Ptr) { ToonProfilesRT.Add(Ptr); }
+	const UToonProfile* GetToonProfileRT(uint32 Index) const { check(Index<uint32(ToonProfilesRT.Num())); return ToonProfilesRT[Index]; }
+	const uint32 NumToonProfileRT() const { return ToonProfilesRT.Num(); }
+	// Toon profiles override
+	void SetToonProfileOverrideRT(const UToonProfile* Ptr) { ToonProfileOverrideRT = Ptr; }
+	const UToonProfile* GetToonProfileOverrideRT() const { return ToonProfileOverrideRT; }
+
 	// Neural profiles
 	void SetNeuralProfileRT(const UNeuralProfile* Ptr) { NeuralProfileRT = Ptr; }
 	const UNeuralProfile* GetNeuralProfileRT() const { return NeuralProfileRT; }
@@ -268,6 +277,8 @@ private:
 	const USpecularProfile* SpecularProfileOverrideRT = nullptr;
 	const UNeuralProfile* NeuralProfileRT;
 	FString MaterialName;
+	TArray<const UToonProfile*> ToonProfilesRT;
+	const UToonProfile* ToonProfileOverrideRT = nullptr;
 
 	/** Incremented each time UniformExpressionCache is modified */
 	mutable int32 UniformExpressionCacheSerialNumber = 0;
diff --git a/Engine/Source/Runtime/Engine/Public/SceneView.h b/Engine/Source/Runtime/Engine/Public/SceneView.h
--- a/Engine/Source/Runtime/Engine/Public/SceneView.h
+++ b/Engine/Source/Runtime/Engine/Public/SceneView.h
@@ -1070,6 +1070,7 @@ enum ETranslucencyVolumeCascade
 	VIEW_UNIFORM_BUFFER_MEMBER(FVector4f, SSProfilesTextureSizeAndInvSize) \
 	VIEW_UNIFORM_BUFFER_MEMBER(FVector4f, SSProfilesPreIntegratedTextureSizeAndInvSize) \
 	VIEW_UNIFORM_BUFFER_MEMBER(FVector4f, SpecularProfileTextureSizeAndInvSize) \
+	VIEW_UNIFORM_BUFFER_MEMBER(FVector4f, ToonProfileTextureSizeAndInvSize) \
 	VIEW_UNIFORM_BUFFER_MEMBER(FVector3f, PhysicsFieldClipmapCenter) \
 	VIEW_UNIFORM_BUFFER_MEMBER(float, PhysicsFieldClipmapDistance) \
 	VIEW_UNIFORM_BUFFER_MEMBER(int, PhysicsFieldClipmapResolution) \
@@ -1220,6 +1221,9 @@ BEGIN_GLOBAL_SHADER_PARAMETER_STRUCT_WITH_CONSTRUCTOR(FViewUniformShaderParamete
 	// Specular Profile
 	SHADER_PARAMETER_TEXTURE(Texture2DArray, SpecularProfileTexture)
 	SHADER_PARAMETER_SAMPLER(SamplerState, SpecularProfileSampler)
+	// Toon Profile texture
+	SHADER_PARAMETER_TEXTURE(Texture2DArray, ToonProfileTexture)
+	SHADER_PARAMETER_SAMPLER(SamplerState, ToonProfileSampler)
 	// Water
 	SHADER_PARAMETER_SRV(Buffer<float4>, WaterIndirection)
 	SHADER_PARAMETER_SRV(Buffer<float4>, WaterData)
diff --git a/Engine/Source/Runtime/RenderCore/Private/RenderUtils.cpp b/Engine/Source/Runtime/RenderCore/Private/RenderUtils.cpp
--- a/Engine/Source/Runtime/RenderCore/Private/RenderUtils.cpp
+++ b/Engine/Source/Runtime/RenderCore/Private/RenderUtils.cpp
@@ -2321,6 +2321,11 @@ namespace Substrate
 		return IsSubstrateEnabled() && CVar.Get(InPlatform) > 0;
 	}
 
+	bool IsToonProfileEnabled()
+	{
+		return IsSubstrateEnabled();
+	}
+
 	uint32 GetSheenQuality()
 	{
 		return CVarSubstrateSheenQuality.GetValueOnAnyThread();
diff --git a/Engine/Source/Runtime/RenderCore/Public/RenderUtils.h b/Engine/Source/Runtime/RenderCore/Public/RenderUtils.h
--- a/Engine/Source/Runtime/RenderCore/Public/RenderUtils.h
+++ b/Engine/Source/Runtime/RenderCore/Public/RenderUtils.h
@@ -639,6 +639,7 @@ namespace Substrate
 	RENDERCORE_API uint32 GlintLUTIndex();
 	RENDERCORE_API float GlintLevelBias();
 	RENDERCORE_API float GlintLevelMin();
+	RENDERCORE_API bool IsToonProfileEnabled();
 	RENDERCORE_API bool IsSpecularProfileEnabled();
 	RENDERCORE_API bool IsSpecularProfileEnabled(EShaderPlatform InPlatform);
 	RENDERCORE_API bool IsDBufferPassEnabled(EShaderPlatform InPlatform);
diff --git a/Engine/Source/Runtime/Renderer/Private/DeferredShadingRenderer.cpp b/Engine/Source/Runtime/Renderer/Private/DeferredShadingRenderer.cpp
--- a/Engine/Source/Runtime/Renderer/Private/DeferredShadingRenderer.cpp
+++ b/Engine/Source/Runtime/Renderer/Private/DeferredShadingRenderer.cpp
@@ -75,6 +75,7 @@
 #include "ProfilingDebugging/CpuProfilerTrace.h"
 #include "Engine/SubsurfaceProfile.h"
 #include "Engine/SpecularProfile.h"
+#include "Engine/ToonProfile.h"
 #include "SceneCaptureRendering.h"
 #include "NaniteSceneProxy.h"
 #include "Nanite/NaniteRayTracing.h"
@@ -1985,6 +1986,7 @@ void FDeferredShadingSceneRenderer::Render(FRDGBuilder& GraphBuilder, const FSce
 		// Force the subsurface profiles and specular profiles textures to be updated.
 		SubsurfaceProfile::UpdateSubsurfaceProfileTexture(GraphBuilder, ShaderPlatform);
 		SpecularProfile::UpdateSpecularProfileTextureAtlas(GraphBuilder, ShaderPlatform);
+		ToonProfile::UpdateToonProfileTextureAtlas(GraphBuilder, ShaderPlatform);
 
 		// Force the rect light texture & IES texture to be updated.
 		RectLightAtlas::UpdateAtlasTexture(GraphBuilder, FeatureLevel);
diff --git a/Engine/Source/Runtime/Renderer/Private/MobileShadingRenderer.cpp b/Engine/Source/Runtime/Renderer/Private/MobileShadingRenderer.cpp
--- a/Engine/Source/Runtime/Renderer/Private/MobileShadingRenderer.cpp
+++ b/Engine/Source/Runtime/Renderer/Private/MobileShadingRenderer.cpp
@@ -71,6 +71,7 @@
 #include "IESTextureManager.h"
 #include "SceneUniformBuffer.h"
 #include "Engine/SpecularProfile.h"
+#include "Engine/ToonProfile.h"
 #include "LocalFogVolumeRendering.h"
 #include "SceneCaptureRendering.h"
 #include "WaterInfoTextureRendering.h"
@@ -1138,6 +1139,7 @@ void FMobileSceneRenderer::Render(FRDGBuilder& GraphBuilder, const FSceneRenderU
 		// Force the subsurface profile & specular profile textures to be updated.
 		SubsurfaceProfile::UpdateSubsurfaceProfileTexture(GraphBuilder, ShaderPlatform);
 		SpecularProfile::UpdateSpecularProfileTextureAtlas(GraphBuilder, ShaderPlatform);
+		ToonProfile::UpdateToonProfileTextureAtlas(GraphBuilder, ShaderPlatform);
 
 		if (bDeferredShading)
 		{
diff --git a/Engine/Source/Runtime/Renderer/Private/SceneRendering.cpp b/Engine/Source/Runtime/Renderer/Private/SceneRendering.cpp
--- a/Engine/Source/Runtime/Renderer/Private/SceneRendering.cpp
+++ b/Engine/Source/Runtime/Renderer/Private/SceneRendering.cpp
@@ -97,6 +97,7 @@
 #include "VolumetricFog.h"
 #include "PrimitiveSceneShaderData.h"
 #include "Engine/SpecularProfile.h"
+#include "Engine/ToonProfile.h"
 #include "Engine/VolumeTexture.h"
 #include "GPUDebugCrashUtils.h"
 #include "MeshDrawCommandStats.h"
@@ -1839,6 +1840,15 @@ void FViewInfo::SetupUniformBufferParameters(
 		ViewUniformShaderParameters.SpecularProfileSampler = TStaticSamplerState<SF_Bilinear, AM_Clamp, AM_Clamp, AM_Clamp>::GetRHI();
 	}
 
+	// Toon Profiles
+	{
+		FRHITexture* Texture = ToonProfile::GetToonProfileTextureAtlasWithFallback();
+		FIntVector TextureSize = Texture->GetSizeXYZ();
+		ViewUniformShaderParameters.ToonProfileTextureSizeAndInvSize = FVector4f(TextureSize.X, TextureSize.Y, 1.0f / TextureSize.X, 1.0f / TextureSize.Y);
+		ViewUniformShaderParameters.ToonProfileTexture = Texture;
+		ViewUniformShaderParameters.ToonProfileSampler = TStaticSamplerState<SF_Bilinear, AM_Clamp, AM_Clamp, AM_Clamp>::GetRHI();
+	}
+
 	{
 		// This is the CVar default
 		float Value = 1.0f;
diff --git a/Engine/Source/Runtime/Renderer/Private/Substrate/SubstrateVisualize.cpp b/Engine/Source/Runtime/Renderer/Private/Substrate/SubstrateVisualize.cpp
--- a/Engine/Source/Runtime/Renderer/Private/Substrate/SubstrateVisualize.cpp
+++ b/Engine/Source/Runtime/Renderer/Private/Substrate/SubstrateVisualize.cpp
@@ -840,7 +840,7 @@ void AddProcessAndPrintSubstrateMaterialPropertiesPasses(FRDGBuilder& GraphBuild
 						PrintFColorSmall(DrawPosX + DataOffset, DrawPosY, Closure.SSSRescaledMFPR, Closure.SSSRescaledMFPG, Closure.SSSRescaledMFPB);// Aliased variable for Eye BSDF
 						NewLine();
 
-						PrintS(DrawPosX + 0.0f, DrawPosY, TEXT("SSSPRofileID"), Orange);
+						PrintS(DrawPosX + 0.0f, DrawPosY, TEXT("SSSProfileID"), Orange);
 						PrintI(DrawPosX + DataOffset, DrawPosY, Closure.SSSPRofileID);
 						NewLine();
 					}
@@ -868,6 +868,10 @@ void AddProcessAndPrintSubstrateMaterialPropertiesPasses(FRDGBuilder& GraphBuild
 					}
 					else if (Closure.Type == SUBSTRATE_BSDF_TYPE_TOON)
 					{
+						PrintS(DrawPosX + 0.0f, DrawPosY, TEXT("ToonProfileID"), Orange);
+						PrintI(DrawPosX + DataOffset, DrawPosY, Closure.SSSPRofileID);
+						NewLine();
+
 						PrintS(DrawPosX + 0.0f, DrawPosY, TEXT("Diffuse Color"), Orange);
 						PrintFColorSmall(DrawPosX + DataOffset, DrawPosY, Closure.DiffuseR, Closure.DiffuseG, Closure.DiffuseB);
 						NewLine();
