Fix use of UE_EXPERIMENTAL on UInstancedSkinnedMeshComponent so that it compiles in non-engine code

#Synced-CL 48390580

diff --git a/Engine/Source/Runtime/Engine/Classes/Components/InstancedSkinnedMeshComponent.h b/Engine/Source/Runtime/Engine/Classes/Components/InstancedSkinnedMeshComponent.h
--- a/Engine/Source/Runtime/Engine/Classes/Components/InstancedSkinnedMeshComponent.h
+++ b/Engine/Source/Runtime/Engine/Classes/Components/InstancedSkinnedMeshComponent.h
@@ -51,8 +51,7 @@ struct FSkinnedMeshInstanceData
 };
 
 UCLASS(ClassGroup=Rendering, hidecategories=(Object,Activation,Collision,"Components|Activation",Physics), editinlinenew, meta=(BlueprintSpawnableComponent), Blueprintable, MinimalAPI)
-UE_EXPERIMENTAL(5.6, "This class is currently extremely experimental and should not be used at this time.")
-class UInstancedSkinnedMeshComponent : public USkinnedMeshComponent
+class UE_EXPERIMENTAL(5.6, "This class is currently extremely experimental and should not be used at this time.") UInstancedSkinnedMeshComponent : public USkinnedMeshComponent
 {
 	GENERATED_UCLASS_BODY()
 	
