[shaders] make shadermap refcount member atomic to resolve TSAN issue

#Synced-CL 48456217

diff --git a/Engine/Source/Runtime/Engine/Private/Materials/MaterialShader.cpp b/Engine/Source/Runtime/Engine/Private/Materials/MaterialShader.cpp
--- a/Engine/Source/Runtime/Engine/Private/Materials/MaterialShader.cpp
+++ b/Engine/Source/Runtime/Engine/Private/Materials/MaterialShader.cpp
@@ -3737,8 +3737,6 @@ void FMaterialShaderMap::RegisterForODSC(EShaderPlatform InShaderPlatform)
 
 void FMaterialShaderMap::AddRef()
 {
-	//#todo-mw: re-enable to try to find potential corruption of the global shader map ID array
-	//check(IsInGameThread());
 	FScopeLock ScopeLock(&GIdToMaterialShaderMapCS);
 	check(!bDeletedThroughDeferredCleanup);
 	++NumRefs;
@@ -3746,13 +3744,10 @@ void FMaterialShaderMap::AddRef()
 
 void FMaterialShaderMap::Release()
 {
-	//#todo-mw: re-enable to try to find potential corruption of the global shader map ID array
-	//check(IsInGameThread());
-
 	{
 		FScopeLock ScopeLock(&GIdToMaterialShaderMapCS);
 
-		check(NumRefs > 0);
+		check(NumRefs.load() > 0);
 		if (--NumRefs == 0)
 		{
 			if (bRegistered)
diff --git a/Engine/Source/Runtime/Engine/Public/MaterialShared.h b/Engine/Source/Runtime/Engine/Public/MaterialShared.h
--- a/Engine/Source/Runtime/Engine/Public/MaterialShared.h
+++ b/Engine/Source/Runtime/Engine/Public/MaterialShared.h
@@ -1850,7 +1850,7 @@ private:
 	uint32 CompilingId = 0;
 #endif // WITH_EDITOR
 
-	mutable int32 NumRefs;
+	mutable std::atomic<int32> NumRefs;
 
 	/** Used to catch errors where the shader map is deleted directly. */
 	bool bDeletedThroughDeferredCleanup;
