Dstm: Chaos vehicle support 1.2: Add migration support for NetworkPhysicsComponent

#Synced-CL 48437998

diff --git a/Engine/Source/Runtime/Engine/Private/PhysicsEngine/NetworkPhysicsComponent.cpp b/Engine/Source/Runtime/Engine/Private/PhysicsEngine/NetworkPhysicsComponent.cpp
--- a/Engine/Source/Runtime/Engine/Private/PhysicsEngine/NetworkPhysicsComponent.cpp
+++ b/Engine/Source/Runtime/Engine/Private/PhysicsEngine/NetworkPhysicsComponent.cpp
@@ -14,6 +14,10 @@
 
 #include "Iris/ReplicationState/PropertyNetSerializerInfoRegistry.h"
 
+#if UE_WITH_REMOTE_OBJECT_HANDLE
+#include "UObject/UObjectMigrationContext.h"
+#endif
+
 #include UE_INLINE_GENERATED_CPP_BY_NAME(NetworkPhysicsComponent)
 
 namespace PhysicsReplicationCVars
@@ -538,6 +542,11 @@ void UNetworkPhysicsComponent::InitializeComponent()
 {
 	Super::InitializeComponent();
 
+	InitializePhysicsReplication();
+}
+
+void UNetworkPhysicsComponent::InitializePhysicsReplication()
+{
 	// Cache CVar values
 	bEnableUnreliableFlow = PhysicsReplicationCVars::ResimulationCVars::bEnableUnreliableFlow;
 	bEnableReliableFlow = PhysicsReplicationCVars::ResimulationCVars::bEnableReliableFlow;
@@ -616,6 +625,28 @@ void UNetworkPhysicsComponent::InitializeComponent()
 	}
 }
 
+#if UE_WITH_REMOTE_OBJECT_HANDLE
+void UNetworkPhysicsComponent::PostMigrate(const struct FUObjectMigrationContext& MigrationContext)
+{
+	Super::PostMigrate(MigrationContext);
+
+	if (MigrationContext.MigrationSide == EObjectMigrationSide::Send)
+	{
+		// Need to uninitialize the component during migration, because the component is not GC'ed until next frame
+		// Then the call back on AsyncPhysicsNetworkComponent would cause a crash
+		UninitializeComponent();
+	}
+
+	if (MigrationContext.MigrationSide == EObjectMigrationSide::Receive)
+	{
+		// Need to reinitialize the component (excluding initialization of the base)
+		// The reason is ActorComponents do not get reinitialized after migrated to another server since all data on ActorComponent
+		// is serialized and migrated over. 
+		InitializePhysicsReplication();
+	}
+}
+#endif
+
 void UNetworkPhysicsComponent::UninitializeComponent()
 {
 	Super::UninitializeComponent();
diff --git a/Engine/Source/Runtime/Engine/Public/Physics/NetworkPhysicsComponent.h b/Engine/Source/Runtime/Engine/Public/Physics/NetworkPhysicsComponent.h
--- a/Engine/Source/Runtime/Engine/Public/Physics/NetworkPhysicsComponent.h
+++ b/Engine/Source/Runtime/Engine/Public/Physics/NetworkPhysicsComponent.h
@@ -1519,6 +1519,10 @@ protected:
 	UPROPERTY(Transient, ReplicatedUsing = OnRep_SetReplicatedStateCollection)
 	FNetworkPhysicsDataCollection ReplicatedStateCollection;
 
+#if UE_WITH_REMOTE_OBJECT_HANDLE
+	ENGINE_API virtual void PostMigrate(const struct FUObjectMigrationContext& MigrationContext) override;
+#endif
+
 
 private:
 	/** Send network data from marshaled data */
@@ -1570,6 +1574,9 @@ private:
 	/** Returns if this frame is valid to store at the next index in the delta source state array, frame value needs to match with the modulo of the array size */
 	const bool IsValidNextDeltaSourceState(const int32 Frame) const { return GetDeltaSourceIndexForFrame(Frame) == GetNextDeltaSourceStateIndex(); }
 
+	/** Initializes necessary data for NetworkPhysicsComponent_Internal and callbacks */
+	void InitializePhysicsReplication();
+
 public:
 	/** Get the delta source stored at index or with frame, @param bValueIsIndex switches between grabbing frame vs index
 	* Note, passing in @param Value as -1 returns the latest data
