Prevent subsystem collection from populating the array map with classes which are not held for GC by a subsystem actually in the collection.
Also deprecates a protected method which will make it harder to make the collection threadsafe in future.

#Synced-CL Unknown

diff --git a/Engine/Source/Runtime/Engine/Private/Subsystems/SubsystemCollection.cpp b/Engine/Source/Runtime/Engine/Private/Subsystems/SubsystemCollection.cpp
--- a/Engine/Source/Runtime/Engine/Private/Subsystems/SubsystemCollection.cpp
+++ b/Engine/Source/Runtime/Engine/Private/Subsystems/SubsystemCollection.cpp
@@ -11,6 +11,7 @@
 #include "Logging/StructuredLog.h"
 #include "Modules/ModuleManager.h"
 #include "Misc/PackageName.h"
+#include "Clang/ClangPlatformCompilerPreSetup.h"
 
 DEFINE_LOG_CATEGORY(LogSubsystemCollection);
 
@@ -84,13 +85,9 @@ USubsystem* FSubsystemCollectionBase::GetSubsystemInternal(UClass* SubsystemClas
 	{
 		return SystemPtr;
 	}
-	else
+	else if (const FSubsystemCollectionBase::FSubsystemArray* SystemPtrs = FindAndPopulateSubsystemArray(SubsystemClass))
 	{
-		const FSubsystemCollectionBase::FSubsystemArray& SystemPtrs = FindAndPopulateSubsystemArrayInternal(SubsystemClass);
-		if (SystemPtrs.Subsystems.Num() > 0)
-		{
-			return SystemPtrs.Subsystems[0];
-		}
+		return SystemPtrs->Subsystems[0];
 	}
 
 	if (Initialization != nullptr)
@@ -113,12 +110,6 @@ USubsystem* FSubsystemCollectionBase::GetSubsystemInternal(UClass* SubsystemClas
 	return nullptr;
 }
 
-TArray<USubsystem*> FSubsystemCollectionBase::GetSubsystemArrayCopy(UClass* SubsystemClass) const
-{
-	FSubsystemArray& List = FindAndPopulateSubsystemArrayInternal(SubsystemClass);
-	return List.Subsystems;
-}
-
 void FSubsystemCollectionBase::ForEachSubsystem(TFunctionRef<void(USubsystem*)> Operation) const
 {
 	TGuardValue<bool> Guard{bIterating, true};
@@ -130,24 +121,44 @@ void FSubsystemCollectionBase::ForEachSubsystem(TFunctionRef<void(USubsystem*)>
 
 FSubsystemCollectionBase::FSubsystemArray& FSubsystemCollectionBase::FindAndPopulateSubsystemArrayInternal(UClass* SubsystemClass) const
 {
+PRAGMA_DISABLE_DEPRECATION_WARNINGS
+	if (const FSubsystemCollectionBase::FSubsystemArray* Array = FindAndPopulateSubsystemArray(SubsystemClass))
+	{
+		return *const_cast<FSubsystemArray*>(Array);
+	}
+	static FSubsystemArray Empty;
+	return Empty;
+PRAGMA_ENABLE_DEPRECATION_WARNINGS
+}
+
+const FSubsystemCollectionBase::FSubsystemArray* FSubsystemCollectionBase::FindAndPopulateSubsystemArray(UClass* SubsystemClass) const
+{
+	// NOTE: There is no thread safety here for multiple threads trying to access a subsystem by a base class/interface.
+	// Ideally there should be.
 	const bool bIsInterface = SubsystemClass->IsChildOf<UInterface>();
 	if (!SubsystemArrayMap.Contains(SubsystemClass))
 	{
-		TUniquePtr<FSubsystemArray>& NewList = SubsystemArrayMap.Emplace(SubsystemClass, MakeUnique<FSubsystemArray>());
+		FSubsystemArray NewList;
 		for (auto Iter = SubsystemMap.CreateConstIterator(); Iter; ++Iter)
 		{
 			UClass* KeyClass = Iter.Key();
 			if ((!bIsInterface && KeyClass->IsChildOf(SubsystemClass)) || 
 				(bIsInterface && KeyClass->ImplementsInterface(SubsystemClass)))
 			{
-				NewList->Subsystems.Add(Iter.Value());
+				NewList.Subsystems.Add(Iter.Value());
 			}
 		}
-		return *NewList;
+		if (NewList.Subsystems.Num())
+		{
+			return SubsystemArrayMap.Add(SubsystemClass, MakeUnique<FSubsystemArray>(MoveTemp(NewList))).Get();
+		}
+		// If nothing was found, we don't want to store this key in the map as we remove lists when they become empty.
+		// We also don't pass the keys in SubsystemArrayMap to GC.
+		return nullptr;
 	}
 
 	const TUniquePtr<FSubsystemArray>& List = SubsystemArrayMap.FindChecked(SubsystemClass);
-	return *List;
+	return List.Get();
 }
 
 void FSubsystemCollectionBase::ForEachSubsystemOfClass(UClass* SubsystemClass, TFunctionRef<void(USubsystem*)> Operation) const
@@ -157,12 +168,13 @@ void FSubsystemCollectionBase::ForEachSubsystemOfClass(UClass* SubsystemClass, T
 		SubsystemClass = USubsystem::StaticClass();
 	}
 
-	const FSubsystemArray& List = FindAndPopulateSubsystemArrayInternal(SubsystemClass);
-
-	TGuardValue<bool> IterationGuard{List.bIsIterating, true};
-	for (int32 i=0; i < List.Subsystems.Num(); ++i)
+	if (const FSubsystemArray* List = FindAndPopulateSubsystemArray(SubsystemClass))
 	{
-		Operation(List.Subsystems[i]);
+		TGuardValue<bool> IterationGuard{List->bIsIterating, true};
+		for (int32 i=0; i < List->Subsystems.Num(); ++i)
+		{
+			Operation(List->Subsystems[i]);
+		}
 	}
 }
 
diff --git a/Engine/Source/Runtime/Engine/Public/Subsystems/SubsystemCollection.h b/Engine/Source/Runtime/Engine/Public/Subsystems/SubsystemCollection.h
--- a/Engine/Source/Runtime/Engine/Public/Subsystems/SubsystemCollection.h
+++ b/Engine/Source/Runtime/Engine/Public/Subsystems/SubsystemCollection.h
@@ -76,10 +76,25 @@ protected:
 	ENGINE_API USubsystem* GetSubsystemInternal(UClass* SubsystemClass) const;
 
 	// Fetch a list of subsystems that derive from the given class, populating the cache if necessary
+	UE_DEPRECATED("5.8", "FindAndPopulateSubsystemArrayInternal has been deprecated. Use GetSubsystemArrayCopy instead.")
 	ENGINE_API FSubsystemArray& FindAndPopulateSubsystemArrayInternal(UClass* SubsystemClass) const;
 
 	/** Get a list of Subsystems by type */
-	ENGINE_API TArray<USubsystem*> GetSubsystemArrayCopy(UClass* SubsystemClass) const;
+	TArray<USubsystem*> GetSubsystemArrayCopy(UClass* SubsystemClass) const
+	{
+		return GetSubsystemArrayCopy<USubsystem>(SubsystemClass);
+	}
+
+	/** Get a list of Subsystems by type */
+	template<typename SubsystemType>
+	TArray<SubsystemType*> GetSubsystemArrayCopy(UClass* SubsystemClass) const
+	{
+		if (const FSubsystemArray* Referenced = FindAndPopulateSubsystemArray(SubsystemClass))
+		{
+			return TArray<SubsystemType*>(reinterpret_cast<SubsystemType* const*>(Referenced->Subsystems.GetData()), Referenced->Subsystems.Num());
+		}
+		return {};
+	}
 
 	/** 
 	 *  Run the given operation on each registered subsystem.
@@ -120,6 +135,7 @@ private:
 
 	TMap<TObjectPtr<UClass>, TObjectPtr<USubsystem>> SubsystemMap;
 
+	// Keys for this map should be a base class or interface implemented by a class in SubsystemMap
 	mutable TMap<UClass*, TUniquePtr<FSubsystemArray>> SubsystemArrayMap;
 
 	UClass* BaseType;
@@ -139,6 +155,12 @@ private:
 	/** For a class which has already been validated as something we should create, initialize the subsystem and add it to the collection */
 	USubsystem* AddAndInitializeValidatedSubsystem(UClass* SubsystemClass);
 
+	/**
+	 * Private to keep space for making this threadsafe in future as it returns a reference
+	 * If this returns a non-null pointer, there will be at least one element in the array
+	 */
+	ENGINE_API const FSubsystemArray* FindAndPopulateSubsystemArray(UClass* SubsystemClass) const;
+
 	/** Add Instances of the specified Subsystem classes to all existing SubsystemCollections of the correct type */
 	static ENGINE_API void AddAllInstances(TConstArrayView<UClass*> SubsystemClasses);
 
@@ -166,9 +188,7 @@ public:
 	{
 		// Force a compile time check that TSubsystemClass derives from TBaseType, the internal code only enforces it's a USubsystem
 		TSubclassOf<TBaseType> SubsystemBaseClass = SubsystemClass;
-
-		FSubsystemArray& Array = FindAndPopulateSubsystemArrayInternal(SubsystemBaseClass);
-		return TArray<TSubsystemClass*>(reinterpret_cast<TSubsystemClass**>(Array.Subsystems.GetData()), Array.Subsystems.Num());
+		return FSubsystemCollectionBase::GetSubsystemArrayCopy<TSubsystemClass>(SubsystemBaseClass);
 	}
 
 	/** Perform an operation on all subsystems of a given type in the collection */
@@ -221,9 +241,7 @@ public:
 	{
 		// Force a compile time check that TSubsystemClass derives from TBaseType, the internal code only enforces it's a USubsystem
 		TSubclassOf<TBaseType> SubsystemBaseClass = SubsystemClass;
-
-		FSubsystemArray& Array = FindAndPopulateSubsystemArrayInternal(SubsystemBaseClass);
-		return TArray<TSubsystemClass*>(reinterpret_cast<TSubsystemClass**>(Array.Subsystems.GetData()), Array.Subsystems.Num());
+		return FSubsystemCollectionBase::GetSubsystemArrayCopy<TSubsystemClass>(SubsystemBaseClass);
 	}
 
 	/** Perform an operation on all subsystems in the collection */
@@ -251,5 +269,4 @@ public:
 		: FSubsystemCollectionBase(TBaseType::StaticClass())
 	{
 	}
-};
-
+};
\ No newline at end of file
