Add some UE_LOG_CONTEXT for static mesh DDC/compile operations

#Synced-CL Unknown

diff --git a/Engine/Source/Runtime/Engine/Private/StaticMesh.cpp b/Engine/Source/Runtime/Engine/Private/StaticMesh.cpp
--- a/Engine/Source/Runtime/Engine/Private/StaticMesh.cpp
+++ b/Engine/Source/Runtime/Engine/Private/StaticMesh.cpp
@@ -4113,6 +4113,7 @@ UE_TRACE_EVENT_END()
 
 void FStaticMeshRenderData::Cache(const ITargetPlatform* TargetPlatform, UStaticMesh* Owner, const FStaticMeshLODSettings& LODSettings)
 {
+	UE_LOG_CONTEXT("StaticMesh", FSoftObjectPath(Owner));
 	if (Owner->GetOutermost()->bIsCookedForEditor)
 	{
 		// Don't cache for cooked packages
diff --git a/Engine/Source/Runtime/Engine/Private/StaticMeshCompiler.cpp b/Engine/Source/Runtime/Engine/Private/StaticMeshCompiler.cpp
--- a/Engine/Source/Runtime/Engine/Private/StaticMeshCompiler.cpp
+++ b/Engine/Source/Runtime/Engine/Private/StaticMeshCompiler.cpp
@@ -14,6 +14,7 @@
 #include "EngineLogs.h"
 #include "Settings/EditorExperimentalSettings.h"
 #include "GameFramework/Pawn.h"
+#include "Logging/StructuredLog.h"
 #include "Misc/QueuedThreadPoolWrapper.h"
 #include "Templates/GuardValueAccessors.h"
 #include "ProfilingDebugging/CountersTrace.h"
@@ -319,6 +320,8 @@ void FStaticMeshCompilingManager::PostCompilation(UStaticMesh* StaticMesh)
 {
 	using namespace StaticMeshCompilingManagerImpl;
 	
+	UE_LOG_CONTEXT("StaticMesh", FSoftObjectPath(StaticMesh));
+	
 	// If AsyncTask is null here, the task got canceled so we don't need to do anything
 	if (StaticMesh->AsyncTask && !IsEngineExitRequested())
 	{
