Support robust option to remove inline skinning shaders when using skin cache.
Removed old setting to do this:  r.SkinCache.SkipCompilingGPUSkinVF=1.
Moved setting to r.SkinCache.CompileShaders which now supports inline only/both/skin cache only options.
Previously ShouldWeCompileGPUSkinVFShaders() looked at the non readonly option r.SkinCache.Mode which may not be uniformly set in device profiles.
It also looked at r.SkinCache.DefaultBehavior which is a setting that can be overridden per mesh.
In addition we now force skin cache in GetGPUSkinTechnique() when inline skinning shaders are removed.
Removed r.SkinCache.SceneMemoryLimitInMB which is no longer used.
#Synced-CL 48453925

diff --git a/Engine/Source/Runtime/Engine/Classes/Engine/RendererSettings.h b/Engine/Source/Runtime/Engine/Classes/Engine/RendererSettings.h
--- a/Engine/Source/Runtime/Engine/Classes/Engine/RendererSettings.h
+++ b/Engine/Source/Runtime/Engine/Classes/Engine/RendererSettings.h
@@ -312,6 +312,17 @@ enum class ESlateMaterialShaderMode : uint8
 	Dynamic = 2 	UMETA(DisplayName = "Dynamic"),
 };
 
+UENUM()
+enum class ESkinningShaderSupportMode : uint8
+{
+	// Only compile vertex factories for inline skinning. 
+	InlineShadersOnly = 0,
+	// Compile shaders for both inline and skin cache skinning.
+	InlineAndSkinCacheShaders = 1,
+	// Only compile shaders for skin cache skinning.
+	SkinCacheShadersOnly = 2,
+};
+
 /**
  * Rendering settings.
  */
@@ -1349,15 +1360,6 @@ class URendererSettings : public UDeveloperSettings
 		ConfigRestartRequired = true))
 	uint32 bUseHairStrandsAutoLODMode  : 1;
 
-	/**
-	"Skin cache allows a compute shader to skin once each vertex, save those results into a new buffer and reuse those calculations when later running the depth, base and velocity passes. This also allows opting into the 'recompute tangents' for skinned mesh instance feature. Disabling will reduce the number of shader permutations required per material. Changing this setting requires restarting the editor."
-	*/
-	UPROPERTY(config, EditAnywhere, Category = Optimizations, meta = (
-		ConsoleVariable = "r.SkinCache.CompileShaders", DisplayName = "Support Compute Skin Cache",
-		ToolTip = "Cannot be disabled while Ray Tracing is enabled as it is then required.",
-		ConfigRestartRequired = true))
-	uint32 bSupportSkinCacheShaders : 1;
-
 	/**
 	"Support hardware variable rate shading.
 	*/
@@ -1367,25 +1369,6 @@ class URendererSettings : public UDeveloperSettings
 		ConfigRestartRequired = true))
 	uint32 bSupportHardwareVariableRateShading : 1;
 
-	/**
-	"When enabled this will skip compiling GPU skin vertex factory shader variants with the assumption that all skinning work will be done via the skin cache."
-	*/
-	UPROPERTY(config, EditAnywhere, Category = Optimizations, meta = (
-		ConsoleVariable = "r.SkinCache.SkipCompilingGPUSkinVF", DisplayName = "Reduce GPU Skin Vertex Factory shader permutations",
-		ToolTip = "Cannot be enabled while the skin cache is turned off.",
-		ConfigRestartRequired = true))
-	uint32 bSkipCompilingGPUSkinVF : 1;
-
-	UPROPERTY(config, EditAnywhere, Category = Optimizations, meta = (
-		ConsoleVariable = "r.SkinCache.DefaultBehavior", DisplayName = "Default Skin Cache Behavior",
-		ToolTip = "Default behavior if all skeletal meshes are included/excluded from the skin cache. If Support Ray Tracing is enabled on a mesh, the skin cache will be used for Ray Tracing updates on that mesh regardless of this setting."))
-	ESkinCacheDefaultBehavior DefaultSkinCacheBehavior;
-
-	UPROPERTY(config, EditAnywhere, Category = Optimizations, meta = (
-		ConsoleVariable = "r.SkinCache.SceneMemoryLimitInMB", DisplayName = "Maximum memory for Compute Skin Cache per world (MB)",
-		ToolTip = "Maximum amount of memory (in MB) per world/scene allowed for the Compute Skin Cache to generate output vertex data and recompute tangents."))
-	float SkinCacheSceneMemoryLimitInMB;
-
 	UPROPERTY(config, EditAnywhere, Category = MobileShaderPermutationReduction, meta = (
 		ConsoleVariable = "r.Mobile.EnableStaticAndCSMShadowReceivers", DisplayName = "Support Combined Static and CSM Shadowing",
 		ToolTip = "Allow primitives to receive both static and CSM shadows from a stationary light. Disabling will free a mobile texture sampler and reduce shader permutations. Changing this setting requires restarting the editor.",
@@ -1513,7 +1496,23 @@ class URendererSettings : public UDeveloperSettings
 		ConfigRestartRequired = true,
 		ClampMin = "12", UIMin = "12"))
 		FPerPlatformInt MaxSkinBones;
-	
+
+	/**
+	 * Skin cache allows a compute shader to skin once each vertex, save those results into a new buffer and reuse those calculations when later running the depth, base and velocity passes. 
+	 * This also allows opting into the 'recompute tangents' for skinned mesh instance feature. 
+	 * Selecting only one shader path will reduce the number of shader permutations required per material. Changing this setting requires restarting the editor.
+	 */
+	UPROPERTY(config, EditAnywhere, Category = Skinning, meta = (
+		ConsoleVariable = "r.SkinCache.CompileShaders", DisplayName = "Skin Cache Shader Support",
+		ToolTip = "Select the skinning shader types to build. Selecting only one shader type will reduce the number of shader permutations required per material. Skin cache shaders are required if Ray Tracing is enabled.",
+		ConfigRestartRequired = true))
+	ESkinningShaderSupportMode SkinningShaderSupport;
+
+	UPROPERTY(config, EditAnywhere, Category = Skinning, meta = (
+		ConsoleVariable = "r.SkinCache.DefaultBehavior", DisplayName = "Skin Cache Default Behavior",
+		ToolTip = "Default behavior for whether skeletal meshes use the skin cache. If Support Ray Tracing is enabled on a mesh, the skin cache will be used for Ray Tracing updates on that mesh regardless of this setting."))
+	ESkinCacheDefaultBehavior DefaultSkinCacheBehavior;
+
 	UE_DEPRECATED(5.6, "MobilePlanarReflectionMode is deprecated. The screen space reflections on mobile is supported and the pixel projected reflection is removed.")
 	UPROPERTY()
 	uint8 MobilePlanarReflectionMode;
diff --git a/Engine/Source/Runtime/Engine/Private/GPUSkinCache.cpp b/Engine/Source/Runtime/Engine/Private/GPUSkinCache.cpp
--- a/Engine/Source/Runtime/Engine/Private/GPUSkinCache.cpp
+++ b/Engine/Source/Runtime/Engine/Private/GPUSkinCache.cpp
@@ -75,8 +75,6 @@ class FSkinCacheExecHelper : public FSelfRegisteringExec
 };
 static FSkinCacheExecHelper GSkelMeshExecHelper;
 
-static int32 GEnableGPUSkinCacheShaders = 0;
-
 static TAutoConsoleVariable<bool> CVarAllowGPUSkinCache(
 	TEXT("r.SkinCache.Allow"),
 	true,
@@ -84,32 +82,25 @@ static TAutoConsoleVariable<bool> CVarAllowGPUSkinCache(
 	ECVF_RenderThreadSafe | ECVF_ReadOnly
 );
 
+static int32 GEnableGPUSkinCacheShaders = 0;
 static FAutoConsoleVariableRef CVarEnableGPUSkinCacheShaders(
 	TEXT("r.SkinCache.CompileShaders"),
 	GEnableGPUSkinCacheShaders,
-	TEXT("Whether or not to compile the GPU compute skinning cache shaders.\n")
-	TEXT("This will compile the shaders for skinning on a compute job and not skin on the vertex shader.\n")
-	TEXT("GPUSkinVertexFactory.usf needs to be touched to cause a recompile if this changes.\n")
-	TEXT("0 is off(default), 1 is on"),
+	TEXT("Select which skinning paths to compile shaders for.\n")
+	TEXT("Inline shaders apply skinning in the vertex shader.\n")
+	TEXT("Skin cache shaders are used for skinning in compute jobs.\n")
+	TEXT("0 (default) compiles inline only, 1 compiles both paths, 2 compiles only skin cache shaders."),
 	ECVF_RenderThreadSafe | ECVF_ReadOnly
 );
 
-static TAutoConsoleVariable<bool> CVarSkipCompilingGPUSkinVF(
-	TEXT("r.SkinCache.SkipCompilingGPUSkinVF"),
-	false,
-	TEXT("Reduce GPU Skin Vertex Factory shader permutations. Cannot be disabled while the skin cache is turned off.\n")
-	TEXT(" False ( 0): Compile all GPU Skin Vertex factory variants.\n")
-	TEXT(" True  ( 1): Don't compile all GPU Skin Vertex factory variants."),
-    ECVF_RenderThreadSafe | ECVF_ReadOnly
-);
-
 int32 GEnableGPUSkinCache = 1;
 static TAutoConsoleVariable<int32> CVarEnableGPUSkinCache(
 	TEXT("r.SkinCache.Mode"),
 	1,
 	TEXT("Whether or not to use the GPU compute skinning cache.\n")
 	TEXT("This will perform skinning on a compute job and not skin on the vertex shader.\n")
-	TEXT("Requires r.SkinCache.CompileShaders=1 and r.SkinCache.Allow=1\n")
+	TEXT("Enabling requires r.SkinCache.CompileShaders!=0 and r.SkinCache.Allow=1\n")
+	TEXT("Disabling requires r.SkinCache.CompileShaders!=2\n")
 	TEXT(" 0: off\n")
 	TEXT(" 1: on(default)\n"),
 	ECVF_RenderThreadSafe
@@ -120,7 +111,8 @@ static TAutoConsoleVariable<int32> CVarDefaultGPUSkinCacheBehavior(
 	(int32)ESkinCacheDefaultBehavior::Inclusive,
 	TEXT("Default behavior if all skeletal meshes are included/excluded from the skin cache. If Support Ray Tracing is enabled on a mesh, will force inclusive behavior on that mesh.\n")
 	TEXT(" Exclusive ( 0): All skeletal meshes are excluded from the skin cache. Each must opt in individually.\n")
-	TEXT(" Inclusive ( 1): All skeletal meshes are included into the skin cache. Each must opt out individually. (default)")
+	TEXT(" Inclusive ( 1): All skeletal meshes are included into the skin cache. Each must opt out individually. (default)"),
+	FConsoleVariableDelegate::CreateLambda([](IConsoleVariable* InVariable)	{ FGlobalComponentRecreateRenderStateContext Context; })
 	);
 
 int32 GSkinCacheRecomputeTangents = 2;
@@ -128,7 +120,7 @@ TAutoConsoleVariable<int32> CVarGPUSkinCacheRecomputeTangents(
 	TEXT("r.SkinCache.RecomputeTangents"),
 	2,
 	TEXT("This option enables recomputing the vertex tangents on the GPU.\n")
-	TEXT("Can be changed at runtime, requires both r.SkinCache.CompileShaders=1, r.SkinCache.Mode=1, r.SkinCache.Allow=1 and r.SkinCache.StoreDuplicatedVerticesForRecomputeTangents\n")
+	TEXT("Can be changed at runtime, requires both r.SkinCache.CompileShaders!=0, r.SkinCache.Mode=1, r.SkinCache.Allow=1 and r.SkinCache.StoreDuplicatedVerticesForRecomputeTangents\n")
 	TEXT(" 0: off\n")
 	TEXT(" 1: on, forces all skinned object to Recompute Tangents\n")
 	TEXT(" 2: on, only recompute tangents on skinned objects who ticked the Recompute Tangents checkbox(default)\n"),
@@ -151,14 +143,6 @@ static TAutoConsoleVariable<float> CVarGPUSkinCacheDebug(
 	ECVF_RenderThreadSafe
 );
 
-static float GSkinCacheSceneMemoryLimitInMB = 128.0f;
-static TAutoConsoleVariable<float> CVarGPUSkinCacheSceneMemoryLimitInMB(
-	TEXT("r.SkinCache.SceneMemoryLimitInMB"),
-	128.0f,
-	TEXT("Maximum memory allowed to be allocated per World/Scene in Megs"),
-	ECVF_RenderThreadSafe
-);
-
 static int32 GStoreDuplicatedVerticesForRecomputeTangents = 0;
 FAutoConsoleVariableRef CVarGPUSkinCacheStoreDuplicatedVerticesForRecomputeTangents(
 	TEXT("r.SkinCache.StoreDuplicatedVerticesForRecomputeTangents"),
@@ -227,52 +211,6 @@ static int32 GGPUSkinCacheFlushCounter = 0;
 
 const float MBSize = 1048576.f; // 1024 x 1024 bytes
 
-static inline bool IsGPUSkinCacheEnable(EShaderPlatform Platform)
-{
-	static FShaderPlatformCachedIniValue<int32> PerPlatformCVar(TEXT("r.SkinCache.Mode"));
-	return (PerPlatformCVar.Get(Platform) != 0);
-}
-
-static inline bool IsGPUSkinCacheInclusive(EShaderPlatform Platform)
-{
-	static FShaderPlatformCachedIniValue<int32> PerPlatformCVar(TEXT("r.SkinCache.DefaultBehavior"));
-	return (PerPlatformCVar.Get(Platform) != 0);
-}
-
-bool ShouldWeCompileGPUSkinVFShaders(EShaderPlatform Platform, ERHIFeatureLevel::Type FeatureLevel)
-{
-	// If the skin cache is not available on this platform we need to compile GPU Skin VF shaders.
-	if (IsGPUSkinCacheAvailable(Platform) == false)
-	{
-		return true;
-	}
-
-	// If the skin cache is not available on this platform we need to compile GPU Skin VF Shaders.
-	if (IsGPUSkinCacheEnable(Platform) == false)
-	{
-		return true;
-	}
-
-	// If the skin cache has been globally disabled for all skeletal meshes we need to compile GPU Skin VF Shaders.
-	if (IsGPUSkinCacheInclusive(Platform) == false)
-	{
-		return true;
-	}
-
-	// Some mobile GPUs (MALI) has a 64K elements limitation on texel buffers
-	// This results in meshes with more than 64k vertices having their skin cache entries disabled at runtime.
-	// We don't have a reliable way of checking this at cook time, so for mobile we must always cache skin cache
-	// shaders so we have something to fall back to.
-	if (FeatureLevel == ERHIFeatureLevel::ES3_1)
-	{
-		return true;
-	}
-
-	// If the skin cache is enabled and we've been asked to skip GPU Skin VF shaders.
-	static FShaderPlatformCachedIniValue<bool> PerPlatformCVar(TEXT("r.SkinCache.SkipCompilingGPUSkinVF"));
-	return (PerPlatformCVar.Get(Platform) == false);
-}
-
 ESkinCacheDefaultBehavior GetSkinCacheDefaultBehavior()
 {
 	return ESkinCacheDefaultBehavior(CVarDefaultGPUSkinCacheBehavior->GetInt()) == ESkinCacheDefaultBehavior::Inclusive
@@ -2524,37 +2462,29 @@ void FGPUSkinCache::CVarSinkFunction()
 {
 	int32 NewGPUSkinCacheValue = CVarEnableGPUSkinCache.GetValueOnAnyThread() != 0;
 	int32 NewRecomputeTangentsValue = CVarGPUSkinCacheRecomputeTangents.GetValueOnAnyThread();
-	const float NewSceneMaxSizeInMb = CVarGPUSkinCacheSceneMemoryLimitInMB.GetValueOnAnyThread();
 	const int32 NewNumTangentIntermediateBuffers = CVarGPUSkinNumTangentIntermediateBuffers.GetValueOnAnyThread();
-	const bool NewSkipCompilingGPUSkinVF = CVarSkipCompilingGPUSkinVF.GetValueOnAnyThread();
 
-	if (GEnableGPUSkinCacheShaders)
+	if (GEnableGPUSkinCacheShaders == 0)
+	{
+		NewGPUSkinCacheValue = 0;
+		NewRecomputeTangentsValue = 0;
+	}
+	else
 	{
 		if (GIsRHIInitialized && IsGPUSkinCacheRayTracingSupported() && IsRayTracingEnabled())
 		{
 			// Skin cache is *required* for ray tracing.
 			NewGPUSkinCacheValue = 1;
 		}
-	}
-	else
-	{
-		NewGPUSkinCacheValue = 0;
-		NewRecomputeTangentsValue = 0;
-	}
 
-	// We don't have GPU Skin VF shaders at all so we can't fallback to using GPU Skinning.
-	if (NewSkipCompilingGPUSkinVF)
-	{
-		// If we had the skin cache enabled and we are turning it off.
-		if (GEnableGPUSkinCache && (NewGPUSkinCacheValue == 0))
+		if (GEnableGPUSkinCacheShaders == 2)
 		{
+			// Inline skinning isn't supported.
 			NewGPUSkinCacheValue = 1;
-			UE_LOG(LogSkinCache, Warning, TEXT("Attemping to turn off the GPU Skin Cache, but we don't have GPU Skin VF shaders to fallback to (r.SkinCache.SkipCompilingGPUSkinVF=1).  Leaving skin cache turned on."));
 		}
 	}
 
-	if (NewGPUSkinCacheValue != GEnableGPUSkinCache || NewRecomputeTangentsValue != GSkinCacheRecomputeTangents
-		|| NewSceneMaxSizeInMb != GSkinCacheSceneMemoryLimitInMB || NewNumTangentIntermediateBuffers != GNumTangentIntermediateBuffers)
+	if (NewGPUSkinCacheValue != GEnableGPUSkinCache || NewRecomputeTangentsValue != GSkinCacheRecomputeTangents || NewNumTangentIntermediateBuffers != GNumTangentIntermediateBuffers)
 	{		
 		if (NewRecomputeTangentsValue == 1 && GStoreDuplicatedVerticesForRecomputeTangents == 0)
 		{
@@ -2562,12 +2492,11 @@ void FGPUSkinCache::CVarSinkFunction()
 		}
 
 		ENQUEUE_RENDER_COMMAND(DoEnableSkinCaching)(UE::RenderCommandPipe::SkeletalMesh,
-			[NewRecomputeTangentsValue, NewGPUSkinCacheValue, NewSceneMaxSizeInMb, NewNumTangentIntermediateBuffers](FRHICommandList& RHICmdList)
+			[NewRecomputeTangentsValue, NewGPUSkinCacheValue, NewNumTangentIntermediateBuffers](FRHICommandList& RHICmdList)
 		{
 			GNumTangentIntermediateBuffers = FMath::Max(NewNumTangentIntermediateBuffers, 1);
 			GEnableGPUSkinCache = NewGPUSkinCacheValue;
 			GSkinCacheRecomputeTangents = NewRecomputeTangentsValue;
-			GSkinCacheSceneMemoryLimitInMB = NewSceneMaxSizeInMb;
 			++GGPUSkinCacheFlushCounter;
 		});
 
diff --git a/Engine/Source/Runtime/Engine/Private/GPUSkinVertexFactory.cpp b/Engine/Source/Runtime/Engine/Private/GPUSkinVertexFactory.cpp
--- a/Engine/Source/Runtime/Engine/Private/GPUSkinVertexFactory.cpp
+++ b/Engine/Source/Runtime/Engine/Private/GPUSkinVertexFactory.cpp
@@ -811,7 +811,7 @@ bool TGPUSkinVertexFactory<BoneInfluenceType>::ShouldCompilePermutation(const FV
 	// Compile the shader for UBI if UBI is enabled and we're not forcing the use of a deformer for all UBI meshes
 	const bool bUnlimitedBoneInfluences = BoneInfluenceType == UnlimitedBoneInfluence && bUseUBI && !bUseDeformerForUBI;
 
-	return ShouldWeCompileGPUSkinVFShaders(Parameters.Platform, Parameters.MaterialParameters.FeatureLevel) &&
+	return AreInlineSkinningShadersEnabled(Parameters.Platform) &&
 		  (((Parameters.MaterialParameters.bIsUsedWithSkeletalMesh || Parameters.MaterialParameters.bIsUsedWithMorphTargets) && (BoneInfluenceType != UnlimitedBoneInfluence || bUnlimitedBoneInfluences)) 
 			  || Parameters.MaterialParameters.bIsSpecialEngineMaterial);
 }
diff --git a/Engine/Source/Runtime/Engine/Private/RendererSettings.cpp b/Engine/Source/Runtime/Engine/Private/RendererSettings.cpp
--- a/Engine/Source/Runtime/Engine/Private/RendererSettings.cpp
+++ b/Engine/Source/Runtime/Engine/Private/RendererSettings.cpp
@@ -85,8 +85,7 @@ URendererSettings::URendererSettings(const FObjectInitializer& ObjectInitializer
 	bSupportPointLightWholeSceneShadows = true;
 	MorphTargetMaxBlendWeight = 5.f;
 	bSupportSkyAtmosphere = true;
-	bSupportSkinCacheShaders = false;
-	bSkipCompilingGPUSkinVF = false;
+	SkinningShaderSupport = ESkinningShaderSupportMode::InlineShadersOnly;
 	GPUSimulationTextureSizeX = 1024;
 	GPUSimulationTextureSizeY = 1024;
 	bEnableRayTracing = 0;
@@ -268,17 +267,17 @@ void URendererSettings::PostEditChangeProperty(FPropertyChangedEvent& PropertyCh
 			GPUSimulationTextureSizeY = FMath::RoundUpToPowerOfTwo( FMath::Clamp(GPUSimulationTextureSizeY, MinGPUSimTextureSize, MaxGPUSimTextureSize) );
 		}
 
-		if (PropertyChangedEvent.Property->GetFName() == GET_MEMBER_NAME_CHECKED(URendererSettings, bEnableRayTracing) 
+		if (PropertyChangedEvent.Property->GetFName() == GET_MEMBER_NAME_CHECKED(URendererSettings, bEnableRayTracing)
 			&& bEnableRayTracing 
-			&& !bSupportSkinCacheShaders)
+			&& SkinningShaderSupport==ESkinningShaderSupportMode::InlineShadersOnly)
 		{
 			FString FullPath = FPaths::ConvertRelativePathToFull(GetDefaultConfigFilename());
 			FPlatformFileManager::Get().GetPlatformFile().SetReadOnly(*FullPath, false);
 
 			if (FMessageDialog::Open(EAppMsgType::YesNo, LOCTEXT("Skin Cache Disabled", "Ray Tracing requires enabling skin cache. Do you want to automatically enable skin cache now?")) == EAppReturnType::Yes)
 			{
-				bSupportSkinCacheShaders = 1;
-				UpdateDependentPropertyInConfigFile(this, GET_MEMBER_NAME_CHECKED(URendererSettings, bSupportSkinCacheShaders));
+				SkinningShaderSupport = ESkinningShaderSupportMode::InlineAndSkinCacheShaders;
+				UpdateDependentPropertyInConfigFile(this, GET_MEMBER_NAME_CHECKED(URendererSettings, SkinningShaderSupport));
 			}
 			else
 			{
@@ -287,15 +286,15 @@ void URendererSettings::PostEditChangeProperty(FPropertyChangedEvent& PropertyCh
 			}
 		}
 
-		if (PropertyChangedEvent.Property->GetFName() == GET_MEMBER_NAME_CHECKED(URendererSettings, bSupportSkinCacheShaders)
-			&& !bSupportSkinCacheShaders
-			&& bSkipCompilingGPUSkinVF)
+		if (PropertyChangedEvent.Property->GetFName() == GET_MEMBER_NAME_CHECKED(URendererSettings, SkinningShaderSupport)
+			&& bEnableRayTracing
+			&& SkinningShaderSupport == ESkinningShaderSupportMode::InlineShadersOnly)
 		{
-			FString FullPath = FPaths::ConvertRelativePathToFull(GetDefaultConfigFilename());
-			FPlatformFileManager::Get().GetPlatformFile().SetReadOnly(*FullPath, false);
+			FMessageDialog::Open(EAppMsgType::Ok, LOCTEXT("Ray Tracing Requires Skin Cache", "Ray Tracing is enabled and requires the use of skin cache shaders. The Skin Cache Shader Support has been reenabled."));
 
-			bSkipCompilingGPUSkinVF = 0;
-			UpdateDependentPropertyInConfigFile(this, GET_MEMBER_NAME_CHECKED(URendererSettings, bSkipCompilingGPUSkinVF));
+			SkinningShaderSupport = ESkinningShaderSupportMode::InlineAndSkinCacheShaders;
+			
+			UpdateDependentPropertyInConfigFile(this, GET_MEMBER_NAME_CHECKED(URendererSettings, SkinningShaderSupport));
 		}
 
 		if (PropertyChangedEvent.Property->GetFName() == GET_MEMBER_NAME_CHECKED(URendererSettings, DynamicGlobalIllumination) 
@@ -447,18 +446,6 @@ bool URendererSettings::CanEditChange(const FProperty* InProperty) const
 		return false;
 	}
 
-	if ((InProperty->GetFName() == GET_MEMBER_NAME_CHECKED(URendererSettings, bSupportSkinCacheShaders)))
-	{
-		//only allow DISABLE of skincache shaders if raytracing is also disabled as skincache is a dependency of raytracing.
-		return !bSupportSkinCacheShaders || !bEnableRayTracing;
-	}
-
-	// the bSkipCompilingGPUSkinVF setting can only be edited if the skin cache is on.
-	if ((InProperty->GetFName() == GET_MEMBER_NAME_CHECKED(URendererSettings, bSkipCompilingGPUSkinVF)))
-	{
-		return bSupportSkinCacheShaders;
-	}
-
 	// the following settings can only be edited if ray tracing is enabled
 	if ((InProperty->GetFName() == GET_MEMBER_NAME_CHECKED(URendererSettings, bEnablePathTracing)) ||
 		(InProperty->GetFName() == GET_MEMBER_NAME_CHECKED(URendererSettings, bEnableMegaLights)) ||
diff --git a/Engine/Source/Runtime/Engine/Private/SkeletalRenderGPUSkin.cpp b/Engine/Source/Runtime/Engine/Private/SkeletalRenderGPUSkin.cpp
--- a/Engine/Source/Runtime/Engine/Private/SkeletalRenderGPUSkin.cpp
+++ b/Engine/Source/Runtime/Engine/Private/SkeletalRenderGPUSkin.cpp
@@ -182,7 +182,7 @@ inline ESkeletalMeshGPUSkinTechnique GetGPUSkinTechnique(USkinnedMeshComponent*
 		{
 			GPUSkinTechnique = ESkeletalMeshGPUSkinTechnique::MeshDeformer;
 		}
-		else if (GEnableGPUSkinCache && MeshComponent->IsSkinCacheAllowed(LODIndex))
+		else if (!AreInlineSkinningShadersEnabled(GetFeatureLevelShaderPlatform(FeatureLevel)) || (GEnableGPUSkinCache && MeshComponent->IsSkinCacheAllowed(LODIndex)))
 		{
 			GPUSkinTechnique = ESkeletalMeshGPUSkinTechnique::GPUSkinCache;
 
@@ -210,7 +210,7 @@ inline ESkeletalMeshGPUSkinTechnique GetGPUSkinTechnique(const FSkinnedMeshScene
 	{
 		GPUSkinTechnique = ESkeletalMeshGPUSkinTechnique::MeshDeformer;
 	}
-	else if (GEnableGPUSkinCache && InMeshDesc.IsSkinCacheAllowed(LODIndex))
+	else if (!AreInlineSkinningShadersEnabled(GetFeatureLevelShaderPlatform(FeatureLevel)) || (GEnableGPUSkinCache && InMeshDesc.IsSkinCacheAllowed(LODIndex)))
 	{
 		GPUSkinTechnique = ESkeletalMeshGPUSkinTechnique::GPUSkinCache;
 
@@ -1954,19 +1954,16 @@ const FVertexFactory* FSkeletalMeshObjectGPUSkin::GetSkinVertexFactory(const FSc
 	}
 #endif
 
-	// Skin cache can fall back to the base vertex factory if it gets too full.
-	if (DynamicData->GPUSkinTechnique == ESkeletalMeshGPUSkinTechnique::GPUSkinCache && FGPUSkinCache::IsEntryValid(SkinCacheEntry, ChunkIdx))
+	// Return the passthrough vertes factory for Skin Cache.
+	if (DynamicData->GPUSkinTechnique == ESkeletalMeshGPUSkinTechnique::GPUSkinCache)
 	{
+		check(FGPUSkinCache::IsEntryValid(SkinCacheEntry, ChunkIdx));
+		
 		return LOD.GPUSkinVertexFactories.PassthroughVertexFactory.Get();
 	}
 
-	// If we have not compiled GPU Skin vertex factory variants
-	static const auto CVar = IConsoleManager::Get().FindConsoleVariable(TEXT("r.SkinCache.SkipCompilingGPUSkinVF"));
-	if (FeatureLevel != ERHIFeatureLevel::ES3_1 && CVar && CVar->GetBool() == true)
-	{
-		UE_LOG(LogSkeletalMesh, Display, TEXT("We are attempting to render with a GPU Skin Vertex Factory, but r.SkinCache.SkipCompilingGPUSkinVF=1 so we don't have shaders.  Skeletal meshes will draw in ref pose.  Either disable r.SkinCache.SkipCompilingGPUSkinVF or increase the r.SkinCache.SceneMemoryLimitInMB size."));
-		return LOD.GPUSkinVertexFactories.PassthroughVertexFactory.Get();
-	}
+	// Shouldn't get here if inline shaders were disabled.
+	check(AreInlineSkinningShadersEnabled(GetFeatureLevelShaderPlatform(FeatureLevel)))
 
 	// No passthrough usage so return the base skin vertex factory.
 	return GetBaseSkinVertexFactory(LODIndex, ChunkIdx);
diff --git a/Engine/Source/Runtime/Engine/Public/GPUSkinCache.h b/Engine/Source/Runtime/Engine/Public/GPUSkinCache.h
--- a/Engine/Source/Runtime/Engine/Public/GPUSkinCache.h
+++ b/Engine/Source/Runtime/Engine/Public/GPUSkinCache.h
@@ -16,12 +16,10 @@
 // * Saves vertex shader computations when we render an object multiple times (EarlyZ, velocity, shadow, BasePass, CustomDepth, Shadow masking)
 // * Fixes velocity rendering (needed for MotionBlur and TemporalAA) for WorldPosOffset animation and morph target animation
 // * RecomputeTangents results in improved tangent space for WorldPosOffset animation and morph target animation
-// * fixed amount of memory per Scene (r.SkinCache.SceneMemoryLimitInMB)
 // * Velocity Rendering for MotionBlur and TemporalAA (test Velocity in BasePass)
 // * r.SkinCache.Mode and r.SkinCache.RecomputeTangents can be toggled at runtime
 
 // TODO:
-// * Test: Tessellation
 // * Quality/Optimization: increase TANGENT_RANGE for better quality or accumulate two components in one 32bit value
 // * Bug: UpdateMorphVertexBuffer needs to handle SkinCacheObjects that have been rejected by the SkinCache (e.g. because it was running out of memory)
 // * Refactor: Unify the 3 compute shaders to use the same C++ setup code for the variables
@@ -65,8 +63,6 @@ struct FVertexBufferAndSRV;
 struct FRayTracingGeometrySegment;
 struct FGlobalShaderPermutationParameters;
 
-extern bool ShouldWeCompileGPUSkinVFShaders(EShaderPlatform Platform, ERHIFeatureLevel::Type FeatureLevel);
-
 UE_DEPRECATED(5.6, "Use GPUSkinCacheStoreDuplicatedVertices instead to check if duplicate vertices need to be stored.")
 extern ENGINE_API bool GPUSkinCacheNeedsDuplicatedVertices();
 
diff --git a/Engine/Source/Runtime/RenderCore/Private/RenderUtils.cpp b/Engine/Source/Runtime/RenderCore/Private/RenderUtils.cpp
--- a/Engine/Source/Runtime/RenderCore/Private/RenderUtils.cpp
+++ b/Engine/Source/Runtime/RenderCore/Private/RenderUtils.cpp
@@ -1047,7 +1047,7 @@ RENDERCORE_API void RenderUtilsInit()
 			{
 				GRayTracingMode = ERayTracingMode::Disabled;
 
-				UE_LOG(LogRendererCore, Fatal, TEXT("Ray tracing requires skin cache to be enabled. Set r.SkinCache.CompileShaders=1."));
+				UE_LOG(LogRendererCore, Fatal, TEXT("Ray tracing requires skin cache to be enabled. Set r.SkinCache.CompileShaders>=1."));
 			}
 
 			const bool bRayTracingShadersAllowedOnCurrentPlatform = GRayTracingShadersPlatformMask[(int)GMaxRHIShaderPlatform];
@@ -1562,9 +1562,24 @@ RENDERCORE_API bool IsUsingDBuffers(const FStaticShaderPlatform Platform)
 	return GDBufferPlatformMask[(int)Platform] || (Substrate::IsSubstrateEnabled() && !Substrate::IsSubstrateBlendableGBufferEnabled(Platform));
 }
 
+RENDERCORE_API bool AreInlineSkinningShadersEnabled(EShaderPlatform Platform)
+{
+	// Some mobile GPUs (MALI) has a 64K elements limitation on texel buffers
+	// This results in meshes with more than 64k vertices having skin cache disabled at runtime.
+	// So for mobile we must always support inline skinning shaders so that we have something to fall back to.
+	if (IsMobilePlatform(Platform))
+	{
+		return true;
+	}
+
+	static FShaderPlatformCachedIniValue<bool> PerPlatformAllowCVar(TEXT("r.SkinCache.Allow"));
+	static FShaderPlatformCachedIniValue<int32> PerPlatformCompileCVar(TEXT("r.SkinCache.CompileShaders"));
+	return PerPlatformAllowCVar.Get(Platform) == 0 || PerPlatformCompileCVar.Get(Platform) != 2;
+}
+
 RENDERCORE_API bool AreSkinCacheShadersEnabled(EShaderPlatform Platform)
 {
-	static FShaderPlatformCachedIniValue<bool> PerPlatformCVar(TEXT("r.SkinCache.CompileShaders"));
+	static FShaderPlatformCachedIniValue<int32> PerPlatformCVar(TEXT("r.SkinCache.CompileShaders"));
 	return (PerPlatformCVar.Get(Platform) != 0);
 }
 
diff --git a/Engine/Source/Runtime/RenderCore/Public/RenderUtils.h b/Engine/Source/Runtime/RenderCore/Public/RenderUtils.h
--- a/Engine/Source/Runtime/RenderCore/Public/RenderUtils.h
+++ b/Engine/Source/Runtime/RenderCore/Public/RenderUtils.h
@@ -525,6 +525,11 @@ RENDERCORE_API bool UseShaderPipelines(EShaderPlatform ShaderPlatform);
  */
 RENDERCORE_API bool UseRemoveUnsedInterpolators(EShaderPlatform ShaderPlatform);
 
+/**
+ *   Checks if inline skinning shaders are enabled for the platform (via r.SkinCache.CompileShaders)
+ */
+RENDERCORE_API bool AreInlineSkinningShadersEnabled(EShaderPlatform Platform);
+
 /**
  *   Checks if skin cache shaders are enabled for the platform (via r.SkinCache.CompileShaders)
  */
