Fix bug where the tick manager could get deadlocked in EndPhysics because it failed to properly flush the manual dispatches from DuringPhysics and instead flushed EndPhysics twice

#Synced-CL 48468383

diff --git a/Engine/Source/Runtime/Engine/Private/TickTaskManager.cpp b/Engine/Source/Runtime/Engine/Private/TickTaskManager.cpp
--- a/Engine/Source/Runtime/Engine/Private/TickTaskManager.cpp
+++ b/Engine/Source/Runtime/Engine/Private/TickTaskManager.cpp
@@ -977,12 +977,12 @@ public:
 					if (GIdleTaskWorkMS >= 0.0)
 					{
 						bool bDoIdleWork = (GIdleTaskWorkMS > 0.0);
-						bool bDoDeadlockCheck = ManualDispatchTicks[WorldTickGroup].Num() > 0;
+						bool bDoDeadlockCheck = ManualDispatchTicks[Block].Num() > 0;
 						uint64 EndIdleWork = 0;
 						int32 PreviousTasksRemaining = TickCompletionEvents[Block].Num();
 						int32 IdleCount = 0;
 
-						auto IdleWorkUpdate = [this, &bDoIdleWork, &bDoDeadlockCheck, &EndIdleWork, &PreviousTasksRemaining, &IdleCount, WorldTickGroup]
+						auto IdleWorkUpdate = [this, &bDoIdleWork, &bDoDeadlockCheck, &EndIdleWork, &PreviousTasksRemaining, &IdleCount, Block]
 							(int32 TasksRemaining) -> FTaskGraphInterface::EProcessTasksOperation
 						{
 							if (bDoIdleWork)
@@ -1008,7 +1008,7 @@ public:
 									if (++IdleCount > IdleCountToBeDeadlocked)
 									{
 										// Nothing is changing so this could be deadlocked, make sure our manual dispatches have happened
-										this->VerifyManualDispatch(WorldTickGroup);
+										this->VerifyManualDispatch(Block);
 
 										bDoDeadlockCheck = false;
 									}
