Check size of subsystems array before returning entry 0 - invariant that array should be empty can be violated by re-entrancy during FSubsystemCollectionBase::RemoveAllInstances for example

#Synced-CL Unknown

diff --git a/Engine/Source/Runtime/Engine/Private/Subsystems/SubsystemCollection.cpp b/Engine/Source/Runtime/Engine/Private/Subsystems/SubsystemCollection.cpp
--- a/Engine/Source/Runtime/Engine/Private/Subsystems/SubsystemCollection.cpp
+++ b/Engine/Source/Runtime/Engine/Private/Subsystems/SubsystemCollection.cpp
@@ -87,7 +87,11 @@ USubsystem* FSubsystemCollectionBase::GetSubsystemInternal(UClass* SubsystemClas
 	}
 	else if (const FSubsystemCollectionBase::FSubsystemArray* SystemPtrs = FindAndPopulateSubsystemArray(SubsystemClass))
 	{
-		return SystemPtrs->Subsystems[0];
+		// There should be an invariant that this array is not empty, but re-entrancy currently causes that in some scenarios
+		if (SystemPtrs->Subsystems.Num() > 0)
+		{
+			return SystemPtrs->Subsystems[0];
+		}
 	}
 
 	if (Initialization != nullptr)
