Updated the description of r.PSOPrecache.GlobalShader to reflect the recent changes that were made for the dynamic shader preloading refactor in CL 36481445.

#Synced-CL 48500228

diff --git a/Engine/Source/Runtime/Engine/Private/PSOPrecache.cpp b/Engine/Source/Runtime/Engine/Private/PSOPrecache.cpp
--- a/Engine/Source/Runtime/Engine/Private/PSOPrecache.cpp
+++ b/Engine/Source/Runtime/Engine/Private/PSOPrecache.cpp
@@ -18,7 +18,7 @@ static TAutoConsoleVariable<int32> CVarPrecacheGlobalComputeShaders(
 	TEXT("r.PSOPrecache.GlobalShaders"),
 	0,
 	TEXT("Precache global shaders during startup (disable(0) - only compute shaders(1) - all global shaders(2).\n") 
-	TEXT("Note: r.PSOPrecache.GlobalShaders == 2 is only supported when IsDynamicShaderPreloadingEnabled is enabled."),
+	TEXT("Note: r.PSOPrecache.GlobalShaders == 2 is only supported when r.PSOPrecache.Mode == 1 (Preload shaders)."),
 	ECVF_ReadOnly
 );
 
