Substrate - fixed material shading model rebuild with toon shading.
Using the lit model since it is not available for setup outside of substrate.

#rnx
#tests PC editor
#Synced-CL 48498703

diff --git a/Engine/Source/Runtime/Engine/Private/Materials/Material.cpp b/Engine/Source/Runtime/Engine/Private/Materials/Material.cpp
--- a/Engine/Source/Runtime/Engine/Private/Materials/Material.cpp
+++ b/Engine/Source/Runtime/Engine/Private/Materials/Material.cpp
@@ -5887,6 +5887,11 @@ void UMaterial::RebuildShadingModelField()
 					BlendMode = EBlendMode::BLEND_Opaque;
 				}
 			}
+			else if (SubstrateMaterialInfo.HasOnlyShadingModel(SSM_Toon))
+			{
+				MaterialDomain = EMaterialDomain::MD_Surface;
+				ShadingModel = MSM_DefaultLit;
+			}
 			else if (SubstrateMaterialInfo.HasOnlyShadingModel(SSM_LightFunction))
 			{
 				MaterialDomain = EMaterialDomain::MD_LightFunction;
