Fix build with merged modules

#rnx

#Synced-CL 48438982

#Synced-CL 48460260

diff --git a/Engine/Source/Runtime/Engine/Public/Net/PerfCountersHelpers.h b/Engine/Source/Runtime/Engine/Public/Net/PerfCountersHelpers.h
--- a/Engine/Source/Runtime/Engine/Public/Net/PerfCountersHelpers.h
+++ b/Engine/Source/Runtime/Engine/Public/Net/PerfCountersHelpers.h
@@ -55,9 +55,9 @@
 	inline void PerfCountersSet(const FString& Name, int32 Val, uint32 Flags = 0) { /* no-op */ };
 	inline void PerfCountersSet(const FString& Name, const FString& Val, uint32 Flags = 0) { /* no-op */ };
 	inline int32 PerfCountersIncrement(const FString & Name, int32 Add = 1, int32 DefaultValue = 0, uint32 Flags = 0) { return DefaultValue + Add; };
-	inline float ENGINE_API PerfCountersGet(const FString& Name, float DefaultVal) { return DefaultVal; }
-	inline double ENGINE_API PerfCountersGet(const FString& Name, double DefaultVal) { return DefaultVal; }
-	inline int32 ENGINE_API PerfCountersGet(const FString& Name, int32 DefaultVal) { return DefaultVal; }
-	inline uint32 ENGINE_API PerfCountersGet(const FString& Name, uint32 DefaultVal) { return DefaultVal; }
+	inline float PerfCountersGet(const FString& Name, float DefaultVal) { return DefaultVal; }
+	inline double PerfCountersGet(const FString& Name, double DefaultVal) { return DefaultVal; }
+	inline int32 PerfCountersGet(const FString& Name, int32 DefaultVal) { return DefaultVal; }
+	inline uint32 PerfCountersGet(const FString& Name, uint32 DefaultVal) { return DefaultVal; }
 
 #endif //USE_SERVER_PERF_COUNTERS
