[shaders] default r.DumpShaderDebugInfo to OnError in code instead of in ConsoleVariables.ini; this allows it to be set via UserEngine.ini (since the former overrides the latter), and follows the pattern of other shader-related cvars in ConsoleVariables.ini (where info about the variables is provided but no variable values are changed by default)

#Synced-CL 48460480

diff --git a/Engine/Config/ConsoleVariables.ini b/Engine/Config/ConsoleVariables.ini
--- a/Engine/Config/ConsoleVariables.ini
+++ b/Engine/Config/ConsoleVariables.ini
@@ -39,7 +39,7 @@
 ;  Mode 1: dump all. WARNING: leaving this on for a while will fill your hard drive with many small files and folders.
 ;  Mode 2: dump on compilation failure only (default).
 ;  Mode 3: dump on compilation failure or warnings.
-r.DumpShaderDebugInfo=2
+;r.DumpShaderDebugInfo=2
 
 ; When this is enabled, shader compilation warnings will emit as errors
 ;r.Shaders.WarningsAsErrors=0 ;; 0: off, 1: global shaders, 2: all
diff --git a/Engine/Source/Runtime/Engine/Private/ShaderCompiler/ShaderCompiler.cpp b/Engine/Source/Runtime/Engine/Private/ShaderCompiler/ShaderCompiler.cpp
--- a/Engine/Source/Runtime/Engine/Private/ShaderCompiler/ShaderCompiler.cpp
+++ b/Engine/Source/Runtime/Engine/Private/ShaderCompiler/ShaderCompiler.cpp
@@ -303,7 +303,7 @@ constexpr auto& CastEnumToUnderlyingTypeReference(EnumType& Type)
 // For example if there are a lot of content shader compile errors you want to skip over without relaunching
 bool GRetryShaderCompilation = true;
 
-static FShaderCompilingManager::EDumpShaderDebugInfo GDumpShaderDebugInfo = FShaderCompilingManager::EDumpShaderDebugInfo::Never;
+static FShaderCompilingManager::EDumpShaderDebugInfo GDumpShaderDebugInfo = FShaderCompilingManager::EDumpShaderDebugInfo::OnError;
 static FAutoConsoleVariableRef CVarDumpShaderDebugInfo(
 	TEXT("r.DumpShaderDebugInfo"),
 	CastEnumToUnderlyingTypeReference(GDumpShaderDebugInfo),
