Made some const functions in FWorldPartitionClassDescRegistry public for a licensee.

#Synced-CL 48214939

diff --git a/Engine/Source/Runtime/Engine/Public/WorldPartition/WorldPartitionClassDescRegistry.h b/Engine/Source/Runtime/Engine/Public/WorldPartition/WorldPartitionClassDescRegistry.h
--- a/Engine/Source/Runtime/Engine/Public/WorldPartition/WorldPartitionClassDescRegistry.h
+++ b/Engine/Source/Runtime/Engine/Public/WorldPartition/WorldPartitionClassDescRegistry.h
@@ -45,10 +45,13 @@ private:
 
 	friend class FActorDescArchive;
 	friend class UActorDescContainer;
+
+public:
 	ENGINE_API const FWorldPartitionActorDesc* GetClassDescDefault(const FTopLevelAssetPath& InClassPath) const;
 	ENGINE_API const FWorldPartitionActorDesc* GetClassDescDefaultForActor(const FTopLevelAssetPath& InClassPath) const;
 	ENGINE_API const FWorldPartitionActorDesc* GetClassDescDefaultForClass(const FTopLevelAssetPath& InClassPath) const;
 
+private:
 	friend class UWorldPartitionResaveActorsBuilder;
 	const FParentClassMap& GetParentClassMap() const { check(IsInitialized()); return ParentClassMap; }
 
