[AsyncCompilation/Texture]
- Remove assert that is not true anymore in the cancellation case

#Synced-CL Unknown

diff --git a/Engine/Source/Runtime/Engine/Private/Texture.cpp b/Engine/Source/Runtime/Engine/Private/Texture.cpp
--- a/Engine/Source/Runtime/Engine/Private/Texture.cpp
+++ b/Engine/Source/Runtime/Engine/Private/Texture.cpp
@@ -545,8 +545,6 @@ void UTexture::BlockOnAnyAsyncBuild()
 	{
 		FTextureCompilingManager::Get().FinishCompilation({this});
 	}
-
-	check( ! IsDefaultTexture() ); // this is always true even in failure/error cases
 }
 
 bool UTexture::Modify(bool bAlwaysMarkDirty)
