[Network] Fix StaticMeshActor instances not replicating when setting is done on the instance instead of blueprint asset.
- Also only show the instance replication settings on instances

#Synced-CL 48279775

diff --git a/Engine/Source/Runtime/Engine/Classes/Engine/StaticMeshActor.h b/Engine/Source/Runtime/Engine/Classes/Engine/StaticMeshActor.h
--- a/Engine/Source/Runtime/Engine/Classes/Engine/StaticMeshActor.h
+++ b/Engine/Source/Runtime/Engine/Classes/Engine/StaticMeshActor.h
@@ -28,10 +28,9 @@ private:
 	UPROPERTY(Category = StaticMeshActor, VisibleAnywhere, BlueprintReadOnly, meta = (ExposeFunctionCategories = "Mesh,Rendering,Physics,Components|StaticMesh", AllowPrivateAccess = "true"))
 	TObjectPtr<class UStaticMeshComponent> StaticMeshComponent;
 
-protected:
-	ENGINE_API virtual void BeginPlay() override;
-
 public:
+	ENGINE_API virtual void PostInitializeComponents() override;
+
 	/** This static mesh should replicate movement. Automatically sets the RemoteRole and bReplicateMovement flags. Meant to be edited on placed actors (those other two properties are not) */
 	UPROPERTY(Category=Actor, EditAnywhere, AdvancedDisplay)
 	bool bStaticMeshReplicateMovement;
diff --git a/Engine/Source/Runtime/Engine/Private/StaticMeshActor.cpp b/Engine/Source/Runtime/Engine/Private/StaticMeshActor.cpp
--- a/Engine/Source/Runtime/Engine/Private/StaticMeshActor.cpp
+++ b/Engine/Source/Runtime/Engine/Private/StaticMeshActor.cpp
@@ -41,8 +41,10 @@ AStaticMeshActor::AStaticMeshActor(const FObjectInitializer& ObjectInitializer)
 	bCanBeInCluster = (GetClass() == AStaticMeshActor::StaticClass());
 }
 
-void AStaticMeshActor::BeginPlay()
+void AStaticMeshActor::PostInitializeComponents()
 {
+	Super::PostInitializeComponents();
+
 	// Since we allow AStaticMeshActor to specify whether it replicates via bStaticMeshReplicateMovement - per placed instance
 	// and we do the normal SetReplicates call in PostInitProperties, before instanced properties are serialized in, 
 	// we need to do this here. 
@@ -53,11 +55,9 @@ void AStaticMeshActor::BeginPlay()
 		bReplicates = false;
 		SetRemoteRoleForBackwardsCompat(ROLE_SimulatedProxy);
 		SetReplicates(true);
-	}	
+	}
 
 	SetPhysicsReplicationMode(StaticMeshPhysicsReplicationMode);
-
-	Super::BeginPlay();
 }
 
 FString AStaticMeshActor::GetDetailedInfoInternal() const
