Removed unnecessary FLightSceneProxyComponentAccessor as FScene is already friend of FLightSceneProxy.

#Synced-CL 48520266

diff --git a/Engine/Source/Runtime/Engine/Public/LightSceneProxy.h b/Engine/Source/Runtime/Engine/Public/LightSceneProxy.h
--- a/Engine/Source/Runtime/Engine/Public/LightSceneProxy.h
+++ b/Engine/Source/Runtime/Engine/Public/LightSceneProxy.h
@@ -576,18 +576,4 @@ private:
 	// May store the label or name of the actor containing the component or if there is no actor the name of the component itself
 	FString OwnerNameOrLabel;
 #endif
-
-	friend struct FLightSceneProxyComponentAccessor;
-};
-
-/** Accessor used by FScene::GetRelevantLights. */
-struct FLightSceneProxyComponentAccessor
-{
-private:
-	static const ULightComponent* Get(const FLightSceneProxy& LightSceneProxy)
-	{
-		return LightSceneProxy.LightComponent;
-	}
-	friend class FScene;
-};
-
+};
\ No newline at end of file
diff --git a/Engine/Source/Runtime/Renderer/Private/RendererScene.cpp b/Engine/Source/Runtime/Renderer/Private/RendererScene.cpp
--- a/Engine/Source/Runtime/Renderer/Private/RendererScene.cpp
+++ b/Engine/Source/Runtime/Renderer/Private/RendererScene.cpp
@@ -4172,7 +4172,7 @@ void FScene::GetRelevantLights( UPrimitiveComponent* Primitive, TArray<const ULi
 				for (const FLightSceneProxy* LightSceneProxy : RelevantLightProxies)
 				{
 					check(LightSceneProxy);
-					if (const ULightComponent* LightComponent = FLightSceneProxyComponentAccessor::Get(*LightSceneProxy))
+					if (const ULightComponent* LightComponent = LightSceneProxy->LightComponent)
 					{
 						RelevantLights->Add(LightComponent);
 					}
