[Audio] Play Sound anim notify UX improved wrt bFollow & AttachName interplay
+ AttachName value now governs bFollow value & editability rather than the other way around
+ bFollow now automatically true and readonly when AttachName is set
+ Make bFollow property displayed next to AttachName in detail panel (AttachName not moved, still at the top)

#Synced-CL 48461631

diff --git a/Engine/Source/Runtime/Engine/Classes/Animation/AnimNotifies/AnimNotify_PlaySound.h b/Engine/Source/Runtime/Engine/Classes/Animation/AnimNotifies/AnimNotify_PlaySound.h
--- a/Engine/Source/Runtime/Engine/Classes/Animation/AnimNotifies/AnimNotify_PlaySound.h
+++ b/Engine/Source/Runtime/Engine/Classes/Animation/AnimNotifies/AnimNotify_PlaySound.h
@@ -28,8 +28,13 @@ public:
 	ENGINE_API virtual void Notify(USkeletalMeshComponent* MeshComp, UAnimSequenceBase* Animation, const FAnimNotifyEventReference& EventReference) override;
 #if WITH_EDITOR
 	ENGINE_API virtual void ValidateAssociatedAssets() override;
-#endif
 	// End UAnimNotify interface
+	
+	//~ Begin UObject interface
+	ENGINE_API virtual void PostLoad() override;
+	ENGINE_API virtual void PostEditChangeProperty(FPropertyChangedEvent& PropertyChangedEvent) override;
+	//~ End UObject interface
+#endif
 
 	// Sound to Play
 	UPROPERTY(EditAnywhere, BlueprintReadWrite, Category="AnimNotify", meta=(ExposeOnSpawn = true))
@@ -43,19 +48,19 @@ public:
 	UPROPERTY(EditAnywhere, BlueprintReadWrite, Category="AnimNotify", meta=(ExposeOnSpawn = true))
 	float PitchMultiplier;
 
-	// If this sound should follow its owner
-	UPROPERTY(EditAnywhere, BlueprintReadWrite, Category = "AnimNotify")
+	// If this sound should follow its owner. Automatically true if Attach Name is set. 
+	UPROPERTY(EditAnywhere, BlueprintReadWrite, Category = "AnimNotify", meta=(EditCondition="bFollowEditable", HideEditConditionToggle, DisplayPriority=1))
 	uint32 bFollow:1;
 
 #if WITH_EDITORONLY_DATA
+	UPROPERTY(Transient)
+	uint32 bFollowEditable:1;
+	
 	UPROPERTY(Config, EditAnywhere, Category = "AnimNotify")
 	uint32 bPreviewIgnoreAttenuation:1;
 #endif
 
 	// Socket or bone name to attach sound to
-	UPROPERTY(EditAnywhere, BlueprintReadWrite, Category="AnimNotify", meta=(EditCondition="bFollow", ExposeOnSpawn = true))
+	UPROPERTY(EditAnywhere, BlueprintReadWrite, Category="AnimNotify", meta=(ExposeOnSpawn = true, DisplayPriority=0))
 	FName AttachName;
 };
-
-
-
diff --git a/Engine/Source/Runtime/Engine/Private/Animation/AnimNotify_PlaySound.cpp b/Engine/Source/Runtime/Engine/Private/Animation/AnimNotify_PlaySound.cpp
--- a/Engine/Source/Runtime/Engine/Private/Animation/AnimNotify_PlaySound.cpp
+++ b/Engine/Source/Runtime/Engine/Private/Animation/AnimNotify_PlaySound.cpp
@@ -26,6 +26,7 @@ UAnimNotify_PlaySound::UAnimNotify_PlaySound()
 
 #if WITH_EDITORONLY_DATA
 	NotifyColor = FColor(196, 142, 255, 255);
+	bFollowEditable = true;
 	bPreviewIgnoreAttenuation = false;
 #endif // WITH_EDITORONLY_DATA
 }
@@ -109,5 +110,41 @@ void UAnimNotify_PlaySound::ValidateAssociatedAssets()
 		}
 	}
 }
-#endif
 
+void UAnimNotify_PlaySound::PostLoad()
+{
+	Super::PostLoad();
+	
+	if (!AttachName.IsNone())
+	{
+		// Sound MUST follow its owner if it follows one of its bones (edit condition prevents bFollow to become false in this case).
+		bFollowEditable = false;
+		bFollow = true;
+	}
+	else
+	{
+		// If no bone name is given, the user can choose whether it follows the root bone or not by setting bFollow themselves.
+		bFollowEditable = true;
+	}
+}
+
+void UAnimNotify_PlaySound::PostEditChangeProperty(FPropertyChangedEvent& PropertyChangedEvent)
+{
+	Super::PostEditChangeProperty(PropertyChangedEvent);
+	
+	if (PropertyChangedEvent.GetMemberPropertyName() == GET_MEMBER_NAME_CHECKED(ThisClass, AttachName))
+	{
+		if (!AttachName.IsNone())
+		{
+			// Sound MUST follow its owner if it follows one of its bones (edit condition prevents bFollow to become false in this case).
+			bFollowEditable = false;
+			bFollow = true;
+		}
+		else
+		{
+			// If no bone name is given, the user can choose whether it follows the root bone or not by setting bFollow themselves.
+			bFollowEditable = true;
+		}
+	}
+}
+#endif
