Immediate physics - improve support for applying forces

#Synced-CL 48547444

diff --git a/Engine/Source/Runtime/Engine/Private/PhysicsEngine/ImmediatePhysics/ImmediatePhysicsChaos/ImmediatePhysicsActorHandle_Chaos.cpp b/Engine/Source/Runtime/Engine/Private/PhysicsEngine/ImmediatePhysics/ImmediatePhysicsChaos/ImmediatePhysicsActorHandle_Chaos.cpp
--- a/Engine/Source/Runtime/Engine/Private/PhysicsEngine/ImmediatePhysics/ImmediatePhysicsChaos/ImmediatePhysicsActorHandle_Chaos.cpp
+++ b/Engine/Source/Runtime/Engine/Private/PhysicsEngine/ImmediatePhysics/ImmediatePhysicsChaos/ImmediatePhysicsActorHandle_Chaos.cpp
@@ -411,12 +411,132 @@ namespace ImmediatePhysics_Chaos
 	{
 		using namespace Chaos;
 
-		if (FPBDRigidParticleHandle * Rigid = Handle()->CastToRigidParticle())
+		if (FPBDRigidParticleHandle* Rigid = Handle()->CastToRigidParticle())
 		{
 			Rigid->AddTorque(Torque);
 		}
 	}
 
+	void FActorHandle::AddForce(const FVector& Linear, EForceType ForceType)
+	{
+		using namespace Chaos;
+
+		if (FPBDRigidParticleHandle* Rigid = Handle()->CastToRigidParticle())
+		{
+			switch (ForceType)
+			{
+			case EForceType::AddForce:
+				Rigid->AddForce(Linear);
+				break;
+			case EForceType::AddAcceleration:
+				Rigid->Acceleration() += Linear;
+				break;
+			case EForceType::AddImpulse:
+				Rigid->LinearImpulseVelocity() += Linear * Rigid->InvM();
+				break;
+			case EForceType::AddVelocity:
+				Rigid->LinearImpulseVelocity() += Linear;
+				break;
+			default:
+				break;
+			}
+		}
+	}
+
+	// Calculates the force which, when applied at R (relative to the CoM), results in an acceleration Accel
+	// at that location.
+	static FVector CalculateForceForAccelAtPoint(
+		Chaos::FPBDRigidParticleHandle* Rigid, const FVector& R, const FVector& Accel)
+	{
+		const double A = Accel.Length();
+		if (A > UE_SMALL_NUMBER)
+		{
+			const FVector Dir = Accel / A;
+			// This is the resulting component of linear acceleration in the direction dir that
+			// would be achieved by unit force
+			const double Response = FVector::DotProduct(
+				Dir * Rigid->InvM() + FVector::CrossProduct(Chaos::FParticleUtilities::GetWorldInvInertia(Rigid) * 
+					FVector::CrossProduct(R, Dir), R), 
+				Dir);
+			if (Response > UE_SMALL_NUMBER)
+			{
+				return Dir * (A / Response);
+			}
+		}
+		return FVector::ZeroVector;
+	}
+
+	void FActorHandle::AddForceAtLocation(const FVector& Linear, const FVector& Location, EForceType ForceType)
+	{
+		using namespace Chaos;
+
+		if (FPBDRigidParticleHandle* Rigid = Handle()->CastToRigidParticle())
+		{
+			const FVector CoM = FParticleUtilities::GetCoMWorldPosition(Rigid);
+			const FVector LocationRelCoM = Location - CoM;
+			switch (ForceType)
+			{
+			case EForceType::AddForce:
+			{
+				Rigid->AddForce(Linear);
+				Rigid->AddTorque(FVector::CrossProduct(LocationRelCoM, Linear));
+			}
+			break;
+			case EForceType::AddAcceleration:
+			{
+				const FVector DesiredForce = CalculateForceForAccelAtPoint(Rigid, LocationRelCoM, Linear);
+				Rigid->AddForce(DesiredForce);
+				Rigid->AddTorque(FVector::CrossProduct(LocationRelCoM, DesiredForce));
+			}
+			break;
+			case EForceType::AddImpulse:
+			{
+				Rigid->LinearImpulseVelocity() += Linear * Rigid->InvM();
+				Rigid->AngularImpulseVelocity() += FParticleUtilities::GetWorldInvInertia(Rigid) *
+					FVector::CrossProduct(LocationRelCoM, Linear);
+			}
+			break;
+			case EForceType::AddVelocity:
+			{
+				// The accel->force calculation is the same as velChange->impulse
+				const FVector DesiredImpulse = CalculateForceForAccelAtPoint(Rigid, LocationRelCoM, Linear);
+				Rigid->LinearImpulseVelocity() += DesiredImpulse * Rigid->InvM();
+				Rigid->AngularImpulseVelocity() += FParticleUtilities::GetWorldInvInertia(Rigid) *
+					FVector::CrossProduct(LocationRelCoM, DesiredImpulse);
+			}
+			break;
+			default:
+				break;
+			}
+		}
+	}
+
+	void FActorHandle::AddTorque(const FVector& Angular, EForceType ForceType)
+	{
+		using namespace Chaos;
+
+		if (FPBDRigidParticleHandle* Rigid = Handle()->CastToRigidParticle())
+		{
+			switch (ForceType)
+			{
+			case EForceType::AddForce:
+				Rigid->AddTorque(Angular);
+				break;
+			case EForceType::AddAcceleration:
+				Rigid->AngularAcceleration() += Angular;
+				break;
+			case EForceType::AddImpulse:
+				Rigid->AngularImpulseVelocity() += FParticleUtilities::GetWorldInvInertia(Rigid) * Angular;
+				break;
+			case EForceType::AddVelocity:
+				Rigid->AngularImpulseVelocity() += Angular;
+				break;
+			default:
+				break;
+			}
+		}
+	}
+
 	void FActorHandle::AddRadialForce(const FVector& Origin, FReal Strength, FReal Radius, ERadialImpulseFalloff Falloff, EForceType ForceType)
 	{
 		using namespace Chaos;
diff --git a/Engine/Source/Runtime/Engine/Public/Physics/ImmediatePhysics/ImmediatePhysicsChaos/ImmediatePhysicsActorHandle_Chaos.h b/Engine/Source/Runtime/Engine/Public/Physics/ImmediatePhysics/ImmediatePhysicsChaos/ImmediatePhysicsActorHandle_Chaos.h
--- a/Engine/Source/Runtime/Engine/Public/Physics/ImmediatePhysics/ImmediatePhysicsChaos/ImmediatePhysicsActorHandle_Chaos.h
+++ b/Engine/Source/Runtime/Engine/Public/Physics/ImmediatePhysics/ImmediatePhysicsChaos/ImmediatePhysicsActorHandle_Chaos.h
@@ -128,14 +128,27 @@ namespace ImmediatePhysics_Chaos
 		/** Get the angular velocity */
 		ENGINE_API FVector GetAngularVelocity() const;
 
+		/** Adds a force at the centre of mass */
 		ENGINE_API void AddForce(const FVector& Force);
 
-		ENGINE_API void AddTorque(const FVector& Torque);
+		/** Adds a "linear push" at the centre of mass, which can be a force, acceleration, impulse or velocity change */
+		ENGINE_API void AddForce(const FVector& Linear, EForceType ForceType);
 
-		ENGINE_API void AddRadialForce(const FVector& Origin, FReal Strength, FReal Radius, ERadialImpulseFalloff Falloff, EForceType ForceType);
+		/** Adds a "linear push" at the specified world location, which can be a force, acceleration, impulse or velocity change */
+		ENGINE_API void AddForceAtLocation(const FVector& Linear, const FVector& Location, EForceType ForceType);
 
+		/** Adds a linear impulse at the specified world location */
 		ENGINE_API void AddImpulseAtLocation(FVector Impulse, FVector Location);
 
+		/** Adds a torque to the body */
+		ENGINE_API void AddTorque(const FVector& Torque);
+
+		/** Adds a "rotation push" to the body, which can be an angular force, acceleration, impulse or velocity change */
+		ENGINE_API void AddTorque(const FVector& Angular, EForceType ForceType);
+
+		/** Adds a "linear push" (depending on the ForceType) to the centre of mass of the body that is directed away from, and decays with distance from (up to Radius), Origin */
+		ENGINE_API void AddRadialForce(const FVector& Origin, FReal Strength, FReal Radius, ERadialImpulseFalloff Falloff, EForceType ForceType);
+
 		/** Set the linear damping*/
 		ENGINE_API void SetLinearDamping(FReal NewLinearDamping);
 
