Alpha: Fixes to primitive alpha holdout for heterogenous volume.

* Expose new IsDeferredPrimitiveAlphaHoldoutSupported() that matches SUPPORT_PRIMITIVE_ALPHA_HOLDOUT, to be used as a condition for shader resources & bindings.
* Automatically disable holdout when alpha output is disabled in renderer settings to prevent added shader cost for unused feature.
* Clean up non-reflected alpha definitions deprecated in 5.5.

#Synced-CL 48564344

diff --git a/Engine/Source/Runtime/Engine/Classes/Engine/RendererSettings.h b/Engine/Source/Runtime/Engine/Classes/Engine/RendererSettings.h
--- a/Engine/Source/Runtime/Engine/Classes/Engine/RendererSettings.h
+++ b/Engine/Source/Runtime/Engine/Classes/Engine/RendererSettings.h
@@ -107,10 +107,6 @@ namespace EVertexDeformationOutputsVelocity
 	};
 }
 
-/**
- * Enumerates available options for alpha channel through post processing. The renderer will always generate premultiplied RGBA
- * with alpha as translucency (0 = fully opaque; 1 = fully translucent).
- */
 UENUM()
 namespace EAlphaChannelMode
 {
@@ -128,13 +124,6 @@ namespace EAlphaChannelMode
 	};
 }
 
-PRAGMA_DISABLE_DEPRECATION_WARNINGS
-namespace EAlphaChannelMode
-{
-	ENGINE_API EAlphaChannelMode::Type FromInt(int32 InAlphaChannelMode);
-}
-PRAGMA_ENABLE_DEPRECATION_WARNINGS
-
 /** used by FPostProcessSettings AutoExposure*/
 UENUM()
 namespace EAutoExposureMethodUI
diff --git a/Engine/Source/Runtime/Engine/Private/RendererSettings.cpp b/Engine/Source/Runtime/Engine/Private/RendererSettings.cpp
--- a/Engine/Source/Runtime/Engine/Private/RendererSettings.cpp
+++ b/Engine/Source/Runtime/Engine/Private/RendererSettings.cpp
@@ -33,16 +33,6 @@ extern UNREALED_API class UEditorEngine* GEditor;
 
 #define LOCTEXT_NAMESPACE "RendererSettings"
 
-PRAGMA_DISABLE_DEPRECATION_WARNINGS
-namespace EAlphaChannelMode
-{
-	EAlphaChannelMode::Type FromInt(int32 InAlphaChannelMode)
-	{
-		return static_cast<EAlphaChannelMode::Type>(FMath::Clamp(InAlphaChannelMode, (int32)Disabled, (int32)AllowThroughTonemapper));
-	}
-}
-PRAGMA_ENABLE_DEPRECATION_WARNINGS
-
 namespace EDefaultBackBufferPixelFormat
 {
 	EPixelFormat Convert2PixelFormat(EDefaultBackBufferPixelFormat::Type InDefaultBackBufferPixelFormat)
@@ -379,6 +369,12 @@ void URendererSettings::PostEditChangeProperty(FPropertyChangedEvent& PropertyCh
 			}
 		}
 
+		if (PropertyChangedEvent.GetPropertyName() == GET_MEMBER_NAME_CHECKED(URendererSettings, bEnableAlphaChannelInPostProcessing))
+		{
+			// Auto-disable primitive alpha holdout if alpha output is disabled.
+			bDeferredSupportPrimitiveAlphaHoldout &= bEnableAlphaChannelInPostProcessing;
+		}
+
 		ExportValuesToConsoleVariables(PropertyChangedEvent.Property);
 
 		if (PropertyChangedEvent.Property->GetFName() == GET_MEMBER_NAME_CHECKED(URendererSettings, ReflectionCaptureResolution) &&
diff --git a/Engine/Source/Runtime/Renderer/Private/HeterogeneousVolumes/HeterogeneousVolumes.cpp b/Engine/Source/Runtime/Renderer/Private/HeterogeneousVolumes/HeterogeneousVolumes.cpp
--- a/Engine/Source/Runtime/Renderer/Private/HeterogeneousVolumes/HeterogeneousVolumes.cpp
+++ b/Engine/Source/Runtime/Renderer/Private/HeterogeneousVolumes/HeterogeneousVolumes.cpp
@@ -1098,7 +1098,7 @@ void FDeferredShadingSceneRenderer::RenderHeterogeneousVolumes(
 		HeterogeneousVolumeRadiance = GraphBuilder.CreateTexture(Desc, TEXT("HeterogeneousVolumes"));
 		AddClearUAVPass(GraphBuilder, GraphBuilder.CreateUAV(HeterogeneousVolumeRadiance), FLinearColor::Black);
 
-		if (IsPrimitiveAlphaHoldoutEnabledForAnyView(Views))
+		if (IsDeferredPrimitiveAlphaHoldoutSupported())
 		{
 			Desc.Format = PF_R8;
 			HeterogeneousVolumeHoldout = GraphBuilder.CreateTexture(Desc, TEXT("HeterogeneousVolume.Holdout"));
diff --git a/Engine/Source/Runtime/Renderer/Private/HeterogeneousVolumes/HeterogeneousVolumesLiveShadingPipeline.cpp b/Engine/Source/Runtime/Renderer/Private/HeterogeneousVolumes/HeterogeneousVolumesLiveShadingPipeline.cpp
--- a/Engine/Source/Runtime/Renderer/Private/HeterogeneousVolumes/HeterogeneousVolumesLiveShadingPipeline.cpp
+++ b/Engine/Source/Runtime/Renderer/Private/HeterogeneousVolumes/HeterogeneousVolumesLiveShadingPipeline.cpp
@@ -1380,7 +1380,8 @@ void RenderSingleScatteringWithLiveShading(
 
 		// Output
 		PassParameters->RWLightingTexture = GraphBuilder.CreateUAV(HeterogeneousVolumeRadianceTexture);
-		if (IsPrimitiveAlphaHoldoutEnabled(View))
+
+		if (IsDeferredPrimitiveAlphaHoldoutSupported())
 		{
 			PassParameters->RWHoldoutTexture = GraphBuilder.CreateUAV(HeterogeneousVolumeHoldoutTexture);
 		}
diff --git a/Engine/Source/Runtime/Renderer/Private/PostProcess/PostProcessing.cpp b/Engine/Source/Runtime/Renderer/Private/PostProcess/PostProcessing.cpp
--- a/Engine/Source/Runtime/Renderer/Private/PostProcess/PostProcessing.cpp
+++ b/Engine/Source/Runtime/Renderer/Private/PostProcess/PostProcessing.cpp
@@ -122,7 +122,7 @@ static bool GPostProcessingPropagateAlpha = false;
 * with IConsoleManager::Get().FindTConsoleVariableDataInt returning 0 with a boolean cvar set to True, we now use
 * a FAutoConsoleVariableRef which will warn licensees with typed-access at runtime, see IConsoleObject::AsVariableBool()
 * or IConsoleObject::AsVariableInt(). However both CVar->GetBool() & CVar->GetInt() will continue to work, assuming > 0
-* or EAlphaChannelMode::Type comparisons were used.
+* comparison was used.
 */
 FAutoConsoleVariableRef CVarPostProcessingPropagateAlpha(
 	TEXT("r.PostProcessing.PropagateAlpha"),
diff --git a/Engine/Source/Runtime/Renderer/Private/SceneRendering.cpp b/Engine/Source/Runtime/Renderer/Private/SceneRendering.cpp
--- a/Engine/Source/Runtime/Renderer/Private/SceneRendering.cpp
+++ b/Engine/Source/Runtime/Renderer/Private/SceneRendering.cpp
@@ -6317,6 +6317,12 @@ void CreateQuarterResolutionDepthMinAndMaxFromDepthTexture(FRDGBuilder& GraphBui
 	}
 }
 
+bool IsDeferredPrimitiveAlphaHoldoutSupported()
+{
+	// Read-only cvar is expected to match SUPPORT_PRIMITIVE_ALPHA_HOLDOUT shader definition
+	return CVarPrimitiveAlphaHoldoutSupport.GetValueOnRenderThread();
+}
+
 bool IsPrimitiveAlphaHoldoutEnabled(const FViewInfo& View)
 {
 	// Note: r.Deferred.SupportPrimitiveAlphaHoldout excludes the path tracer
diff --git a/Engine/Source/Runtime/Renderer/Private/SceneRendering.h b/Engine/Source/Runtime/Renderer/Private/SceneRendering.h
--- a/Engine/Source/Runtime/Renderer/Private/SceneRendering.h
+++ b/Engine/Source/Runtime/Renderer/Private/SceneRendering.h
@@ -3134,10 +3134,13 @@ inline const FViewFamilyInfo& GetViewFamilyInfo(const TArrayView<const FViewInfo
 	return *(const FViewFamilyInfo*)Views[0].Family;
 }
 
-/** Checks whether primitive alpha holdout is active.*/
+/** Checks whether deferred primitive alpha holdout is supported in shaders.*/
+bool IsDeferredPrimitiveAlphaHoldoutSupported();
+
+/** Checks whether primitive alpha holdout is enabled.*/
 bool IsPrimitiveAlphaHoldoutEnabled(const FViewInfo& View);
 
-/** Checks whether primitive alpha holdout is active for any view.*/
+/** Checks whether primitive alpha holdout is enabled for any view.*/
 bool IsPrimitiveAlphaHoldoutEnabledForAnyView(TArrayView<const FViewInfo> Views);
 
 bool SceneCaptureRequiresAlphaChannel(const FSceneView& View);
