- Added FrameNum and UTC time to stat unit so we have a way to match up video captures with logs, etc. Requires setting "stats.UnitTimestamp" cvar to enable

#Synced-CL 48456107

diff --git a/Engine/Source/Runtime/Engine/Private/UnrealClient.cpp b/Engine/Source/Runtime/Engine/Private/UnrealClient.cpp
--- a/Engine/Source/Runtime/Engine/Private/UnrealClient.cpp
+++ b/Engine/Source/Runtime/Engine/Private/UnrealClient.cpp
@@ -46,6 +46,12 @@ static FAutoConsoleVariableRef CVarHitProxyCaptureEnable(
 	GHitProxyCaptureNextUpdate,
 	TEXT("Enables GPU capture of hit proxy rendering on the next update."));
 
+static int32 GShowTimeStampsInStatUnit = 0;
+static FAutoConsoleVariableRef CVarShowTimeStampsInStatUnit(
+	TEXT("stats.UnitTimestamp"),
+	GShowTimeStampsInStatUnit,
+	TEXT("Enables UTC time and framecounter in stat unit."));
+
 bool FViewport::bIsGameRenderingEnabled = true;
 int32 FViewport::PresentAndStopMovieDelay = 0;
 
@@ -923,6 +929,18 @@ int32 FStatUnitData::DrawStat(FViewport* InViewport, FCanvas* InCanvas, int32 In
 			PushRows(/* RowsCount = */ 1);
 		}
 
+		// Timestamp / Frame count if enabled
+		if (GShowTimeStampsInStatUnit)
+		{
+			DrawTitleString(TEXT("UTC Time"), NoUnitGraphColor);
+			DrawDefaultAvgCell(FDateTime::UtcNow().ToString(TEXT("[%H.%M.%S:%s]")), StatGreen);
+			PushRows(/* RowsCount = */ 1);
+
+			DrawTitleString(TEXT("FrameNum"), NoUnitGraphColor);
+			DrawDefaultAvgCell(FString::Printf(TEXT("%llu"), GFrameCounter), StatGreen);
+			PushRows(/* RowsCount = */ 1);
+		}
+
 		if (bShowTSRStatistics && GPixelRenderCounters.GetPixelDisplayCount())
 		{
 			const uint32 PixelRenderCount = GPixelRenderCounters.GetPixelRenderCount();
@@ -2546,7 +2564,7 @@ bool FViewportClient::InputAxis(
 		bool bGamepad)
 {
 	// Note: Set the timestamp to 0u here to represent an "invalid" timestamp, since this older input axis
-	// function did not account for timestamps. If we initalized it to the current cycle count here,
+	// function did not account for timestamps. If we initialized it to the current cycle count here,
 	// then deprecated code paths would always be considered "newer" then anything calling the newer API, where the timestamp
 	// is calculated when the actual slate input event is constructed. 
 	FInputKeyEventArgs Args(Viewport, InputDevice, Key, Delta, DeltaTime, NumSamples, 0u);
