Add r.mobile.skincache.allow to support disabling skin cache on PC ES31 renderer when it is on for PC desktop renderer.
Moved global state for runtime toggleable skin cache settings into function accessors so that feature level can be taken into account.
No longer "fix" toggleable settings based on support in a cvar sink function, but instead apply the fixup in the accessors.
#Synced-CL 48548904

diff --git a/Engine/Source/Editor/Persona/Private/SMorphTargetViewer.cpp b/Engine/Source/Editor/Persona/Private/SMorphTargetViewer.cpp
--- a/Engine/Source/Editor/Persona/Private/SMorphTargetViewer.cpp
+++ b/Engine/Source/Editor/Persona/Private/SMorphTargetViewer.cpp
@@ -377,7 +377,7 @@ void SMorphTargetListRow::OnMorphTargetWeightValueCommitted( float NewWeight, ET
 bool SMorphTargetListRow::IsMorphTargetWeightSliderEnabled() const
 {
 	const uint32 CVarMorphTargetModeValue = GetDefault<URendererSettings>()->bUseGPUMorphTargets;
-	return GEnableGPUSkinCache > 0 ? CVarMorphTargetModeValue > 0 : true;
+	return FGPUSkinCache::IsGPUSkinCacheEnabled(GMaxRHIFeatureLevel) ? CVarMorphTargetModeValue > 0 : true;
 }
 
 FText SMorphTargetListRow::GetMorphTargetWeightSliderToolTip() const
diff --git a/Engine/Source/Editor/UnrealEd/Private/SEditorViewport.cpp b/Engine/Source/Editor/UnrealEd/Private/SEditorViewport.cpp
--- a/Engine/Source/Editor/UnrealEd/Private/SEditorViewport.cpp
+++ b/Engine/Source/Editor/UnrealEd/Private/SEditorViewport.cpp
@@ -710,7 +710,7 @@ void SEditorViewport::BindCommands()
 
 	MAP_VIEWMODE_ACTION( Commands.VisualizeLWCComplexity, VMI_LWCComplexity);
 
-	if (GEnableGPUSkinCache)
+	if (FGPUSkinCache::IsGPUSkinCacheEnabled(GMaxRHIFeatureLevel))
 	{
 		MAP_VIEWMODE_ACTION(Commands.VisualizeGPUSkinCacheMode, VMI_VisualizeGPUSkinCache);
 		FGPUSkinCacheVisualizationMenuCommands::Get().BindCommands(CommandListRef, Client);
diff --git a/Engine/Source/Editor/UnrealEd/Private/ViewportToolbar/UnrealEdViewportToolbar.cpp b/Engine/Source/Editor/UnrealEd/Private/ViewportToolbar/UnrealEdViewportToolbar.cpp
--- a/Engine/Source/Editor/UnrealEd/Private/ViewportToolbar/UnrealEdViewportToolbar.cpp
+++ b/Engine/Source/Editor/UnrealEd/Private/ViewportToolbar/UnrealEdViewportToolbar.cpp
@@ -2294,7 +2294,7 @@ void PopulateViewModesMenu(UToolMenu* InMenu)
 			}
 		}
 
-		if (GEnableGPUSkinCache
+		if (FGPUSkinCache::IsGPUSkinCacheEnabled(GMaxRHIFeatureLevel)
 			&& UE::UnrealEd::Private::ViewModesSubmenu::IsMenuSectionAvailable(
 				Context, EHidableViewModeMenuSections::GPUSkinCache
 			))
diff --git a/Engine/Source/Runtime/Engine/Private/Components/SkinnedMeshComponent.cpp b/Engine/Source/Runtime/Engine/Private/Components/SkinnedMeshComponent.cpp
--- a/Engine/Source/Runtime/Engine/Private/Components/SkinnedMeshComponent.cpp
+++ b/Engine/Source/Runtime/Engine/Private/Components/SkinnedMeshComponent.cpp
@@ -2421,7 +2421,7 @@ bool USkinnedMeshComponent::GetTwistAndSwingAngleOfDeltaRotationFromRefPose(FNam
 
 bool USkinnedMeshComponent::IsSkinCacheAllowed(int32 LodIdx) const
 {
-	if (!GEnableGPUSkinCache || !GetScene() || !GetScene()->GetGPUSkinCache())
+	if (!GetScene() || !GetScene()->GetGPUSkinCache() || !FGPUSkinCache::IsGPUSkinCacheEnabled(GetScene()->GetFeatureLevel()))
 	{
 		return false;
 	}
@@ -2442,7 +2442,7 @@ bool USkinnedMeshComponent::IsSkinCacheAllowed(int32 LodIdx) const
 
 	if (FSkeletalMeshRenderData* SkelMeshRenderData = GetSkinnedAsset()->GetResourceForRendering())
 	{
-		if (SkelMeshRenderData->bSupportRayTracing && FGPUSkinCache::IsGPUSkinCacheRayTracingSupported() && IsRayTracingEnabled())
+		if (SkelMeshRenderData->bSupportRayTracing && FGPUSkinCache::IsGPUSkinCacheRayTracingSupported(GetScene()->GetFeatureLevel()) && IsRayTracingEnabled())
 		{
 			return true;
 		}
diff --git a/Engine/Source/Runtime/Engine/Private/GPUSkinCache.cpp b/Engine/Source/Runtime/Engine/Private/GPUSkinCache.cpp
--- a/Engine/Source/Runtime/Engine/Private/GPUSkinCache.cpp
+++ b/Engine/Source/Runtime/Engine/Private/GPUSkinCache.cpp
@@ -82,10 +82,16 @@ static TAutoConsoleVariable<bool> CVarAllowGPUSkinCache(
 	ECVF_RenderThreadSafe | ECVF_ReadOnly
 );
 
-static int32 GEnableGPUSkinCacheShaders = 0;
-static FAutoConsoleVariableRef CVarEnableGPUSkinCacheShaders(
+static TAutoConsoleVariable<bool> CVarMobileAllowGPUSkinCache(
+	TEXT("r.Mobile.SkinCache.Allow"),
+	true,
+	TEXT("Whether or not to allow the GPU skin Cache system to be enabled on the mobile renderer.\n"),
+	ECVF_RenderThreadSafe | ECVF_ReadOnly
+);
+
+static TAutoConsoleVariable<int32> CVarEnableGPUSkinCacheShaders(
 	TEXT("r.SkinCache.CompileShaders"),
-	GEnableGPUSkinCacheShaders,
+	0,
 	TEXT("Select which skinning paths to compile shaders for.\n")
 	TEXT("Inline shaders apply skinning in the vertex shader.\n")
 	TEXT("Skin cache shaders are used for skinning in compute jobs.\n")
@@ -93,18 +99,21 @@ static FAutoConsoleVariableRef CVarEnableGPUSkinCacheShaders(
 	ECVF_RenderThreadSafe | ECVF_ReadOnly
 );
 
+PRAGMA_DISABLE_DEPRECATION_WARNINGS
 int32 GEnableGPUSkinCache = 1;
-static TAutoConsoleVariable<int32> CVarEnableGPUSkinCache(
+static FAutoConsoleVariableRef CVarEnableGPUSkinCache(
 	TEXT("r.SkinCache.Mode"),
-	1,
+	GEnableGPUSkinCache,
 	TEXT("Whether or not to use the GPU compute skinning cache.\n")
 	TEXT("This will perform skinning on a compute job and not skin on the vertex shader.\n")
 	TEXT("Enabling requires r.SkinCache.CompileShaders!=0 and r.SkinCache.Allow=1\n")
 	TEXT("Disabling requires r.SkinCache.CompileShaders!=2\n")
 	TEXT(" 0: off\n")
 	TEXT(" 1: on(default)\n"),
+	FConsoleVariableDelegate::CreateLambda([](IConsoleVariable*) { FGPUSkinCache::OnCVarUpdate(); }),
 	ECVF_RenderThreadSafe
 );
+PRAGMA_ENABLE_DEPRECATION_WARNINGS
 
 static TAutoConsoleVariable<int32> CVarDefaultGPUSkinCacheBehavior(
 	TEXT("r.SkinCache.DefaultBehavior"),
@@ -112,27 +121,30 @@ static TAutoConsoleVariable<int32> CVarDefaultGPUSkinCacheBehavior(
 	TEXT("Default behavior if all skeletal meshes are included/excluded from the skin cache. If Support Ray Tracing is enabled on a mesh, will force inclusive behavior on that mesh.\n")
 	TEXT(" Exclusive ( 0): All skeletal meshes are excluded from the skin cache. Each must opt in individually.\n")
 	TEXT(" Inclusive ( 1): All skeletal meshes are included into the skin cache. Each must opt out individually. (default)"),
-	FConsoleVariableDelegate::CreateLambda([](IConsoleVariable* InVariable)	{ FGlobalComponentRecreateRenderStateContext Context; })
-	);
+	FConsoleVariableDelegate::CreateLambda([](IConsoleVariable*) { FGPUSkinCache::OnCVarUpdate(); }),
+	ECVF_RenderThreadSafe
+);
 
 int32 GSkinCacheRecomputeTangents = 2;
-TAutoConsoleVariable<int32> CVarGPUSkinCacheRecomputeTangents(
+static FAutoConsoleVariableRef CVarGPUSkinCacheRecomputeTangents(
 	TEXT("r.SkinCache.RecomputeTangents"),
-	2,
+	GSkinCacheRecomputeTangents,
 	TEXT("This option enables recomputing the vertex tangents on the GPU.\n")
 	TEXT("Can be changed at runtime, requires both r.SkinCache.CompileShaders!=0, r.SkinCache.Mode=1, r.SkinCache.Allow=1 and r.SkinCache.StoreDuplicatedVerticesForRecomputeTangents\n")
 	TEXT(" 0: off\n")
 	TEXT(" 1: on, forces all skinned object to Recompute Tangents\n")
 	TEXT(" 2: on, only recompute tangents on skinned objects who ticked the Recompute Tangents checkbox(default)\n"),
+	FConsoleVariableDelegate::CreateLambda([](IConsoleVariable*) { FGPUSkinCache::OnCVarUpdate(); }),
 	ECVF_RenderThreadSafe
 );
 
 static int32 GNumTangentIntermediateBuffers = 1;
-static TAutoConsoleVariable<float> CVarGPUSkinNumTangentIntermediateBuffers(
+static FAutoConsoleVariableRef CVarGPUSkinNumTangentIntermediateBuffers(
 	TEXT("r.SkinCache.NumTangentIntermediateBuffers"),
-	1,
+	GNumTangentIntermediateBuffers,
 	TEXT("How many intermediate buffers to use for intermediate results while\n")
 	TEXT("doing Recompute Tangents; more may allow the GPU to overlap compute jobs."),
+	FConsoleVariableDelegate::CreateLambda([](IConsoleVariable*) { FGPUSkinCache::OnCVarUpdate(); }),
 	ECVF_RenderThreadSafe
 );
 
@@ -144,7 +156,7 @@ static TAutoConsoleVariable<float> CVarGPUSkinCacheDebug(
 );
 
 static int32 GStoreDuplicatedVerticesForRecomputeTangents = 0;
-FAutoConsoleVariableRef CVarGPUSkinCacheStoreDuplicatedVerticesForRecomputeTangents(
+static FAutoConsoleVariableRef CVarGPUSkinCacheStoreDuplicatedVerticesForRecomputeTangents(
 	TEXT("r.SkinCache.StoreDuplicatedVerticesForRecomputeTangents"),
 	GStoreDuplicatedVerticesForRecomputeTangents,
 	TEXT("0: Don't store duplicated vertices for all skeletal mesh render sections. It will still be stored if the render section has bRecomputeTangent set. (default)\n")
@@ -153,7 +165,7 @@ FAutoConsoleVariableRef CVarGPUSkinCacheStoreDuplicatedVerticesForRecomputeTange
 );
 
 static int32 GUseDuplicatedVerticesForRecomputeTangents = 1;
-FAutoConsoleVariableRef CVarGPUSkinCacheAllowDupedVertesForRecomputeTangents(
+static FAutoConsoleVariableRef CVarGPUSkinCacheAllowDupedVertesForRecomputeTangents(
 	TEXT("r.SkinCache.UseDuplicatedVerticesForRecomputeTangents"),
 	GUseDuplicatedVerticesForRecomputeTangents,
 	TEXT("0: Disable usage of duplicated vertices for runtime tangent recomputation/\n")
@@ -162,7 +174,7 @@ FAutoConsoleVariableRef CVarGPUSkinCacheAllowDupedVertesForRecomputeTangents(
 );
 
 int32 GRecomputeTangentsParallelDispatch = 0;
-FAutoConsoleVariableRef CVarRecomputeTangentsParallelDispatch(
+static FAutoConsoleVariableRef CVarRecomputeTangentsParallelDispatch(
 	TEXT("r.SkinCache.RecomputeTangentsParallelDispatch"),
 	GRecomputeTangentsParallelDispatch,
 	TEXT("This option enables parallel dispatches for recompute tangents.\n")
@@ -172,7 +184,7 @@ FAutoConsoleVariableRef CVarRecomputeTangentsParallelDispatch(
 );
 
 static int32 GSkinCachePrintMemorySummary = 0;
-FAutoConsoleVariableRef CVarGPUSkinCachePrintMemorySummary(
+static FAutoConsoleVariableRef CVarGPUSkinCachePrintMemorySummary(
 	TEXT("r.SkinCache.PrintMemorySummary"),
 	GSkinCachePrintMemorySummary,
 	TEXT("Print break down of memory usage.")
@@ -181,8 +193,8 @@ FAutoConsoleVariableRef CVarGPUSkinCachePrintMemorySummary(
 	ECVF_RenderThreadSafe
 );
 
-bool GSkinCacheAsyncCompute = false;
-FAutoConsoleVariableRef CVarSkinCacheAsyncCompute(
+static bool GSkinCacheAsyncCompute = false;
+static FAutoConsoleVariableRef CVarSkinCacheAsyncCompute(
 	TEXT("r.SkinCache.AsyncCompute"),
 	GSkinCacheAsyncCompute,
 	TEXT(" 0: off\n")
@@ -201,7 +213,7 @@ static FAutoConsoleVariable CVarSkinCacheDynamicMeshBounds(
 	ECVF_ReadOnly
 );
 
-int32 GNumDispatchesToCapture = 0;
+static int32 GNumDispatchesToCapture = 0;
 static FAutoConsoleVariableRef CVarGPUSkinCacheNumDispatchesToCapture(
 	TEXT("r.SkinCache.Capture"),
 	GNumDispatchesToCapture,
@@ -1085,7 +1097,7 @@ FGPUSkinCache::FGPUSkinCache(ERHIFeatureLevel::Type InFeatureLevel, UWorld* InWo
 {
 	check(World);
 
-	if (GSkinCacheRecomputeTangents == 1 && GStoreDuplicatedVerticesForRecomputeTangents == 0)
+	if (GetGPUSkinCacheRecomputeTangentsMode(FeatureLevel) == 1 && GStoreDuplicatedVerticesForRecomputeTangents == 0)
 	{
 		UE_LOG(LogSkinCache, Warning, TEXT("r.SkinCache.RecomputeTangents is set to 1 to update all skinned objects but duplicated vertices are not are not always stored. Set r.SkinCache.RecomputeTangents to 2 or r.SkinCache.StoreDuplicatedVerticesForRecomputeTangents to 1."));
 	}
@@ -1125,7 +1137,6 @@ public:
 
 	static bool ShouldCompilePermutation(const FGlobalShaderPermutationParameters& Parameters)
 	{
-		// currently only implemented and tested on Window SM5 (needs Compute, Atomics, SRV for index buffers, UAV for VertexBuffers)
 		return IsGPUSkinCacheAvailable(Parameters.Platform);
 	}
 	
@@ -1168,7 +1179,6 @@ public:
 
 	static bool ShouldCompilePermutation(const FGlobalShaderPermutationParameters& Parameters)
 	{
-		// currently only implemented and tested on Window SM5 (needs Compute, Atomics, SRV for index buffers, UAV for VertexBuffers)
 		return IsGPUSkinCacheAvailable(Parameters.Platform);
 	}
 
@@ -1255,14 +1265,15 @@ void FGPUSkinCache::DispatchUpdateSkinTangentsTrianglePass(FRHICommandList& RHIC
 
 	if (!GRecomputeTangentsParallelDispatch)
 	{
-		if (StagingBuffers.Num() != GNumTangentIntermediateBuffers)
+		const int32 NumTangentIntermediateBuffers = FMath::Max(GNumTangentIntermediateBuffers, 1);
+		if (StagingBuffers.Num() != NumTangentIntermediateBuffers)
 		{
 			// Release extra buffers if shrinking
-			for (int32 Index = GNumTangentIntermediateBuffers; Index < StagingBuffers.Num(); ++Index)
+			for (int32 Index = NumTangentIntermediateBuffers; Index < StagingBuffers.Num(); ++Index)
 			{
 				StagingBuffers[Index].Release();
 			}
-			StagingBuffers.SetNum(GNumTangentIntermediateBuffers, EAllowShrinking::No);
+			StagingBuffers.SetNum(NumTangentIntermediateBuffers, EAllowShrinking::No);
 		}
 
 		// no need to clear the staging buffer because we create it cleared and clear it after each usage in the per vertex pass
@@ -1719,7 +1730,7 @@ void FGPUSkinCache::DispatchPassSetup(FTaskData& TaskData)
 	}
 
 #if RHI_RAYTRACING
-	if (IsGPUSkinCacheRayTracingSupported() && IsRayTracingEnabled())
+	if (IsGPUSkinCacheRayTracingSupported(FeatureLevel) && IsRayTracingEnabled())
 	{
 		for (FGPUSkinCacheEntry* SkinCacheEntry : PendingProcessRTGeometryEntries)
 		{
@@ -2002,7 +2013,7 @@ void FGPUSkinCache::DispatchPassExecute(FTaskData& TaskData)
 
 		if (TaskData.Pipeline == ERHIPipeline::Graphics)
 		{
-			RHICmdList.Transition(TransitionInfos, ERHIPipeline::Graphics, IsGPUSkinCacheRayTracingSupported() && !FRDGBuilder::IsImmediateMode() ? ERHIPipeline::All : ERHIPipeline::Graphics);
+			RHICmdList.Transition(TransitionInfos, ERHIPipeline::Graphics, IsGPUSkinCacheRayTracingSupported(FeatureLevel) && !FRDGBuilder::IsImmediateMode() ? ERHIPipeline::All : ERHIPipeline::Graphics);
 		}
 		else
 		{
@@ -2052,7 +2063,8 @@ void FGPUSkinCache::ProcessEntry(FRHICommandList& RHICmdList, const FProcessEntr
 	{
 		if (Inputs.Mode == EGPUSkinCacheEntryMode::Raster)
 		{
-			if (GSkinCacheRecomputeTangents > 0)
+			const int32 SkinCacheRecomputeTangents = GetGPUSkinCacheRecomputeTangentsMode(FeatureLevel);
+			if (SkinCacheRecomputeTangents > 0)
 			{
 				RecomputeTangentSections.Reserve(LodData.RenderSections.Num());
 
@@ -2061,7 +2073,7 @@ void FGPUSkinCache::ProcessEntry(FRHICommandList& RHICmdList, const FProcessEntr
 					FGPUSkinCacheEntry::FRecomputeTangentSection& RecomputeTangentSection = RecomputeTangentSections.Emplace_GetRef();
 
 					const FSkelMeshRenderSection& RenderSection = LodData.RenderSections[Index];
-					if (LodData.MultiSizeIndexContainer.GetIndexBuffer() && (GSkinCacheRecomputeTangents == 1 || RenderSection.bRecomputeTangent))
+					if (LodData.MultiSizeIndexContainer.GetIndexBuffer() && (SkinCacheRecomputeTangents == 1 || RenderSection.bRecomputeTangent))
 					{
 						RecomputeTangentSection.bEnable = true;
 
@@ -2272,12 +2284,26 @@ void FGPUSkinCache::ProcessEntry(FRHICommandList& RHICmdList, const FProcessEntr
 #endif
 }
 
-bool FGPUSkinCache::IsGPUSkinCacheRayTracingSupported()
+bool FGPUSkinCache::IsGPUSkinCacheEnabled(ERHIFeatureLevel::Type FeatureLevel)
+{
+	// Override runtime mode for cases where either inline or skin cache are not available.
+	const EShaderPlatform ShaderPlatform = GetFeatureLevelShaderPlatform(FeatureLevel);
+	PRAGMA_DISABLE_DEPRECATION_WARNINGS
+	return (GEnableGPUSkinCache > 0 || !AreInlineSkinningShadersEnabled(ShaderPlatform)) && IsGPUSkinCacheAvailable(ShaderPlatform);
+	PRAGMA_ENABLE_DEPRECATION_WARNINGS
+}
+
+int32 FGPUSkinCache::GetGPUSkinCacheRecomputeTangentsMode(ERHIFeatureLevel::Type FeatureLevel)
+{
+	return IsGPUSkinCacheEnabled(FeatureLevel) ? GSkinCacheRecomputeTangents : 0;
+}
+
+bool FGPUSkinCache::IsGPUSkinCacheRayTracingSupported(ERHIFeatureLevel::Type FeatureLevel)
 {
 #if RHI_RAYTRACING
 	static const auto CVar = IConsoleManager::Get().FindConsoleVariable(TEXT("r.RayTracing.Geometry.SupportSkeletalMeshes"));
 	static const bool SupportSkeletalMeshes = CVar->GetInt() != 0;
-	return IsRayTracingAllowed() && SupportSkeletalMeshes && GEnableGPUSkinCache;
+	return IsRayTracingAllowed() && SupportSkeletalMeshes && IsGPUSkinCacheEnabled(FeatureLevel);
 #else
 	return false;
 #endif
@@ -2287,7 +2313,7 @@ bool FGPUSkinCache::IsGPUSkinCacheRayTracingSupported()
 
 void FGPUSkinCache::ProcessRayTracingGeometryToUpdate(FRHICommandList& RHICmdList, FGPUSkinCacheEntry* SkinCacheEntry)
 {
-	check(IsGPUSkinCacheRayTracingSupported() && IsRayTracingEnabled());
+	check(IsGPUSkinCacheRayTracingSupported(FeatureLevel) && IsRayTracingEnabled());
 
 	if (SkinCacheEntry && SkinCacheEntry->GPUSkin && SkinCacheEntry->GPUSkin->bSupportRayTracing)
 	{
@@ -2458,64 +2484,30 @@ void FGPUSkinCache::SetEntryGPUSkin(FGPUSkinCacheEntry* Entry, FSkeletalMeshObje
 	}
 }
 
-void FGPUSkinCache::CVarSinkFunction()
+void FGPUSkinCache::OnCVarUpdate()
 {
-	int32 NewGPUSkinCacheValue = CVarEnableGPUSkinCache.GetValueOnAnyThread() != 0;
-	int32 NewRecomputeTangentsValue = CVarGPUSkinCacheRecomputeTangents.GetValueOnAnyThread();
-	const int32 NewNumTangentIntermediateBuffers = CVarGPUSkinNumTangentIntermediateBuffers.GetValueOnAnyThread();
-
-	if (GEnableGPUSkinCacheShaders == 0)
+	if (GSkinCacheRecomputeTangents == 1 && GStoreDuplicatedVerticesForRecomputeTangents == 0)
 	{
-		NewGPUSkinCacheValue = 0;
-		NewRecomputeTangentsValue = 0;
+		UE_LOG(LogSkinCache, Warning, TEXT("r.SkinCache.RecomputeTangents is set to 1 to update all skinned objects but duplicated vertices are not are not always stored. Set r.SkinCache.RecomputeTangents to 2 or r.SkinCache.StoreDuplicatedVerticesForRecomputeTangents to 1."));
 	}
-	else
+
+	ENQUEUE_RENDER_COMMAND(DoEnableSkinCaching)(UE::RenderCommandPipe::SkeletalMesh, [](FRHICommandList& RHICmdList)
 	{
-		if (GIsRHIInitialized && IsGPUSkinCacheRayTracingSupported() && IsRayTracingEnabled())
-		{
-			// Skin cache is *required* for ray tracing.
-			NewGPUSkinCacheValue = 1;
-		}
+		++GGPUSkinCacheFlushCounter;
+	});
 
-		if (GEnableGPUSkinCacheShaders == 2)
+	TArray<UActorComponent*> Components;
+	for (USkinnedMeshComponent* Component : TObjectRange<USkinnedMeshComponent>())
+	{
+		if (Component->IsRegistered() && Component->IsRenderStateCreated())
 		{
-			// Inline skinning isn't supported.
-			NewGPUSkinCacheValue = 1;
+			Components.Emplace(Component);
 		}
 	}
 
-	if (NewGPUSkinCacheValue != GEnableGPUSkinCache || NewRecomputeTangentsValue != GSkinCacheRecomputeTangents || NewNumTangentIntermediateBuffers != GNumTangentIntermediateBuffers)
-	{		
-		if (NewRecomputeTangentsValue == 1 && GStoreDuplicatedVerticesForRecomputeTangents == 0)
-		{
-			UE_LOG(LogSkinCache, Warning, TEXT("r.SkinCache.RecomputeTangents is set to 1 to update all skinned objects but duplicated vertices are not are not always stored. Set r.SkinCache.RecomputeTangents to 2 or r.SkinCache.StoreDuplicatedVerticesForRecomputeTangents to 1."));
-		}
-
-		ENQUEUE_RENDER_COMMAND(DoEnableSkinCaching)(UE::RenderCommandPipe::SkeletalMesh,
-			[NewRecomputeTangentsValue, NewGPUSkinCacheValue, NewNumTangentIntermediateBuffers](FRHICommandList& RHICmdList)
-		{
-			GNumTangentIntermediateBuffers = FMath::Max(NewNumTangentIntermediateBuffers, 1);
-			GEnableGPUSkinCache = NewGPUSkinCacheValue;
-			GSkinCacheRecomputeTangents = NewRecomputeTangentsValue;
-			++GGPUSkinCacheFlushCounter;
-		});
-
-		TArray<UActorComponent*> Components;
-
-		for (USkinnedMeshComponent* Component : TObjectRange<USkinnedMeshComponent>())
-		{
-			if (Component->IsRegistered() && Component->IsRenderStateCreated())
-			{
-				Components.Emplace(Component);
-			}
-		}
-
-		FGlobalComponentRecreateRenderStateContext Context(Components);
-	}
+	FGlobalComponentRecreateRenderStateContext Context(Components);
 }
 
-FAutoConsoleVariableSink FGPUSkinCache::CVarSink(FConsoleCommandDelegate::CreateStatic(&CVarSinkFunction));
-
 void FGPUSkinCache::PrintMemorySummary() const
 {
 	UE_LOG(LogSkinCache, Display, TEXT("======= Skin Cache Memory Usage Summary ======="));
diff --git a/Engine/Source/Runtime/Engine/Private/SkeletalMeshUpdater.cpp b/Engine/Source/Runtime/Engine/Private/SkeletalMeshUpdater.cpp
--- a/Engine/Source/Runtime/Engine/Private/SkeletalMeshUpdater.cpp
+++ b/Engine/Source/Runtime/Engine/Private/SkeletalMeshUpdater.cpp
@@ -93,7 +93,7 @@ void FSkeletalMeshUpdatePacket::Init(FSceneInterface* InScene, FGPUSkinCache* In
 	Scene = InScene;
 
 #if RHI_RAYTRACING
-	bSkinCacheForRayTracingSupported = GPUSkinCache && GEnableGPUSkinCache && FGPUSkinCache::IsGPUSkinCacheRayTracingSupported();
+	bSkinCacheForRayTracingSupported = GPUSkinCache && FGPUSkinCache::IsGPUSkinCacheRayTracingSupported(GPUSkinCache->GetFeatureLevel());
 #endif
 
 	Init(Initializer);
diff --git a/Engine/Source/Runtime/Engine/Private/SkeletalRenderGPUSkin.cpp b/Engine/Source/Runtime/Engine/Private/SkeletalRenderGPUSkin.cpp
--- a/Engine/Source/Runtime/Engine/Private/SkeletalRenderGPUSkin.cpp
+++ b/Engine/Source/Runtime/Engine/Private/SkeletalRenderGPUSkin.cpp
@@ -182,7 +182,7 @@ inline ESkeletalMeshGPUSkinTechnique GetGPUSkinTechnique(USkinnedMeshComponent*
 		{
 			GPUSkinTechnique = ESkeletalMeshGPUSkinTechnique::MeshDeformer;
 		}
-		else if (!AreInlineSkinningShadersEnabled(GetFeatureLevelShaderPlatform(FeatureLevel)) || (GEnableGPUSkinCache && MeshComponent->IsSkinCacheAllowed(LODIndex)))
+		else if (!AreInlineSkinningShadersEnabled(GetFeatureLevelShaderPlatform(FeatureLevel)) || MeshComponent->IsSkinCacheAllowed(LODIndex))
 		{
 			GPUSkinTechnique = ESkeletalMeshGPUSkinTechnique::GPUSkinCache;
 
@@ -210,7 +210,7 @@ inline ESkeletalMeshGPUSkinTechnique GetGPUSkinTechnique(const FSkinnedMeshScene
 	{
 		GPUSkinTechnique = ESkeletalMeshGPUSkinTechnique::MeshDeformer;
 	}
-	else if (!AreInlineSkinningShadersEnabled(GetFeatureLevelShaderPlatform(FeatureLevel)) || (GEnableGPUSkinCache && InMeshDesc.IsSkinCacheAllowed(LODIndex)))
+	else if (!AreInlineSkinningShadersEnabled(GetFeatureLevelShaderPlatform(FeatureLevel)) || InMeshDesc.IsSkinCacheAllowed(LODIndex))
 	{
 		GPUSkinTechnique = ESkeletalMeshGPUSkinTechnique::GPUSkinCache;
 
@@ -716,7 +716,7 @@ void FSkeletalMeshObjectGPUSkin::InitResources(const FSkinnedMeshSceneProxyDesc&
 					VertexAttributeMask |= FGPUSkinPassthroughVertexFactory::EVertexAttributeFlags::Tangent;
 				}
 			}
-			else if (SkelLOD.GPUSkinTechnique == ESkeletalMeshGPUSkinTechnique::GPUSkinCache || (FGPUSkinCache::IsGPUSkinCacheRayTracingSupported() && SkelLOD.SkelMeshRenderData->bSupportRayTracing))
+			else if (SkelLOD.GPUSkinTechnique == ESkeletalMeshGPUSkinTechnique::GPUSkinCache || (FGPUSkinCache::IsGPUSkinCacheRayTracingSupported(FeatureLevel) && SkelLOD.SkelMeshRenderData->bSupportRayTracing))
 			{
 				VertexAttributeMask = FGPUSkinPassthroughVertexFactory::EVertexAttributeFlags::Position | FGPUSkinPassthroughVertexFactory::EVertexAttributeFlags::Tangent;
 			}
@@ -904,7 +904,7 @@ int32 FDynamicSkelMeshObjectDataGPUSkin::GetLODIndex(EGPUSkinCacheEntryMode Mode
 
 bool FSkeletalMeshObjectGPUSkin::IsSkinCacheEnabled(EGPUSkinCacheEntryMode Mode) const
 {
-	return GEnableGPUSkinCache &&
+	return FGPUSkinCache::IsGPUSkinCacheEnabled(FeatureLevel) &&
 		// Force skin cache enabled for ray tracing if the inline skinning technique was requested.
 		(DynamicData->GPUSkinTechnique == ESkeletalMeshGPUSkinTechnique::GPUSkinCache || (DynamicData->GPUSkinTechnique == ESkeletalMeshGPUSkinTechnique::Inline && Mode == EGPUSkinCacheEntryMode::RayTracing));
 }
@@ -1021,7 +1021,7 @@ void FSkeletalMeshObjectGPUSkin::ProcessUpdatedDynamicData(FRHICommandList& RHIC
 		}
 	}
 
-	const bool bGPUSkinCacheEnabled = GPUSkinCache && GEnableGPUSkinCache && DynamicData->GPUSkinTechnique == ESkeletalMeshGPUSkinTechnique::GPUSkinCache;
+	const bool bGPUSkinCacheEnabled = GPUSkinCache && DynamicData->GPUSkinTechnique == ESkeletalMeshGPUSkinTechnique::GPUSkinCache && FGPUSkinCache::IsGPUSkinCacheEnabled(FeatureLevel);
 
 	// Immediately release any stale entry if we've recently switched to a LOD level that disallows skin cache
 	if (!bGPUSkinCacheEnabled)
@@ -1815,8 +1815,9 @@ void FSkeletalMeshObjectGPUSkin::FSkeletalMeshObjectLOD::UpdateMorphVertexBuffer
 	FSkeletalMeshLODRenderData& LodData = SkelMeshRenderData->LODRenderData[LODIndex];
 
 	// Whether all sections of the LOD perform GPU recompute tangent
-	bool bAllSectionsDoGPURecomputeTangent = bGPUSkinCacheEnabled && GSkinCacheRecomputeTangents > 0;
-	if (bAllSectionsDoGPURecomputeTangent && GSkinCacheRecomputeTangents == 2)
+	const int32 RecomputeTangentsMode = FGPUSkinCache::GetGPUSkinCacheRecomputeTangentsMode(FeatureLevel);
+	bool bAllSectionsDoGPURecomputeTangent = bGPUSkinCacheEnabled && RecomputeTangentsMode > 0;
+	if (bAllSectionsDoGPURecomputeTangent && RecomputeTangentsMode == 2)
 	{
 		for (int32 i = 0; i < LodData.RenderSections.Num(); ++i)
 		{
@@ -2403,7 +2404,7 @@ void FSkeletalMeshObjectGPUSkin::GetUsedVertexFactoryData(
 	FGPUSkinDataType GPUSkinDataType;
 	InitGPUSkinVertexFactoryComponents(&GPUSkinDataType, VertexBuffers, nullptr /*FGPUBaseSkinVertexFactory*/);
 
-	if (GPUSkinTechnique != ESkeletalMeshGPUSkinTechnique::Inline || FGPUSkinCache::IsGPUSkinCacheRayTracingSupported())
+	if (GPUSkinTechnique != ESkeletalMeshGPUSkinTechnique::Inline || FGPUSkinCache::IsGPUSkinCacheRayTracingSupported(InFeatureLevel))
 	{
 		const FVertexFactoryType* GPUSkinVFType = &FGPUSkinPassthroughVertexFactory::StaticType;
 		bool bSupportsManualVertexFetch = GPUSkinVFType->SupportsManualVertexFetch(GMaxRHIFeatureLevel);
diff --git a/Engine/Source/Runtime/Engine/Private/SkinnedMeshSceneProxyDesc.cpp b/Engine/Source/Runtime/Engine/Private/SkinnedMeshSceneProxyDesc.cpp
--- a/Engine/Source/Runtime/Engine/Private/SkinnedMeshSceneProxyDesc.cpp
+++ b/Engine/Source/Runtime/Engine/Private/SkinnedMeshSceneProxyDesc.cpp
@@ -23,8 +23,6 @@
 #include "Animation/DebugSkelMeshComponent.h"
 #endif
 
-extern ENGINE_API int32 GEnableGPUSkinCache;
-
 FSkinnedMeshSceneProxyDesc::FSkinnedMeshSceneProxyDesc(const USkinnedMeshComponent* Component)
 { 
 	InitializeFromSkinnedMeshComponent(Component);
@@ -307,7 +305,7 @@ bool FSkinnedMeshSceneProxyDesc::IsSkinCacheAllowed(int32 LodIdx) const
 {
 	const bool bGlobalDefault = GetSkinCacheDefaultBehavior() == ESkinCacheDefaultBehavior::Inclusive;
 
-	if (!GEnableGPUSkinCache || !Scene || !Scene->GetGPUSkinCache())
+	if (!Scene || !Scene->GetGPUSkinCache() || !FGPUSkinCache::IsGPUSkinCacheEnabled(Scene->GetFeatureLevel()))
 	{
 		return false;
 	}
@@ -326,7 +324,7 @@ bool FSkinnedMeshSceneProxyDesc::IsSkinCacheAllowed(int32 LodIdx) const
 
 	if (FSkeletalMeshRenderData* SkelMeshRenderData = GetSkinnedAsset()->GetResourceForRendering())
 	{
-		if (SkelMeshRenderData->bSupportRayTracing && FGPUSkinCache::IsGPUSkinCacheRayTracingSupported() && IsRayTracingEnabled())
+		if (SkelMeshRenderData->bSupportRayTracing && FGPUSkinCache::IsGPUSkinCacheRayTracingSupported(Scene->GetFeatureLevel()) && IsRayTracingEnabled())
 		{
 			return true;
 		}
@@ -363,7 +361,7 @@ bool FSkinnedMeshSceneProxyDynamicData::IsSkinCacheAllowed(int32 LodIdx, const U
 {
 	const bool bGlobalDefault = GetSkinCacheDefaultBehavior() == ESkinCacheDefaultBehavior::Inclusive;
 
-	if (!GEnableGPUSkinCache)
+	if (!FGPUSkinCache::IsGPUSkinCacheEnabled(GMaxRHIFeatureLevel))
 	{
 		return false;
 	}
@@ -382,7 +380,7 @@ bool FSkinnedMeshSceneProxyDynamicData::IsSkinCacheAllowed(int32 LodIdx, const U
 
 	if (FSkeletalMeshRenderData* SkelMeshRenderData = InSkinnedAsset->GetResourceForRendering())
 	{
-		if (SkelMeshRenderData->bSupportRayTracing && FGPUSkinCache::IsGPUSkinCacheRayTracingSupported() && IsRayTracingEnabled())
+		if (SkelMeshRenderData->bSupportRayTracing && FGPUSkinCache::IsGPUSkinCacheRayTracingSupported(GMaxRHIFeatureLevel) && IsRayTracingEnabled())
 		{
 			return true;
 		}
diff --git a/Engine/Source/Runtime/Engine/Private/UnrealEngine.cpp b/Engine/Source/Runtime/Engine/Private/UnrealEngine.cpp
--- a/Engine/Source/Runtime/Engine/Private/UnrealEngine.cpp
+++ b/Engine/Source/Runtime/Engine/Private/UnrealEngine.cpp
@@ -8094,8 +8094,7 @@ bool UEngine::HandleSkeletalMeshReportCommand(const TCHAR* Cmd, FOutputDevice& A
 		return SizeInKB;
 	};
 
-	static const IConsoleVariable* DefaultGPUSkinCacheBehavior = IConsoleManager::Get().FindConsoleVariable(TEXT("r.SkinCache.DefaultBehavior"));
-	bool bSkinCacheGlobalDefault = DefaultGPUSkinCacheBehavior && ESkinCacheDefaultBehavior(DefaultGPUSkinCacheBehavior->GetInt()) == ESkinCacheDefaultBehavior::Inclusive;
+	bool bSkinCacheGlobalDefault = GetSkinCacheDefaultBehavior() == ESkinCacheDefaultBehavior::Inclusive;
 
 	int32 NumMeshes = 0;
 	FRHIResourceInfo ResourceInfo;
@@ -8115,7 +8114,8 @@ bool UEngine::HandleSkeletalMeshReportCommand(const TCHAR* Cmd, FOutputDevice& A
 	
 				FSkeletalMeshLODInfo* LodInfo = Mesh->GetLODInfo(LODIdx);
 				check(LodInfo);
-				bool bSkinCachedLODEnabled = GEnableGPUSkinCache && (LodInfo->SkinCacheUsage == ESkinCacheUsage::Auto ? bSkinCacheGlobalDefault : LodInfo->SkinCacheUsage == ESkinCacheUsage::Enabled);
+				bool bSkinCachedLODEnabled = FGPUSkinCache::IsGPUSkinCacheEnabled(GMaxRHIFeatureLevel) &&
+					(LodInfo->SkinCacheUsage == ESkinCacheUsage::Auto ? bSkinCacheGlobalDefault : LodInfo->SkinCacheUsage == ESkinCacheUsage::Enabled);
 
 				int32 MorphDataKB = GetResourceSizeKB(RenderData.MorphTargetVertexInfoBuffers.MorphDataBuffer);
 				int32 ClothVertexBufferKB = GetResourceSizeKB(RenderData.ClothVertexBuffer.VertexBufferRHI);
diff --git a/Engine/Source/Runtime/Engine/Public/GPUSkinCache.h b/Engine/Source/Runtime/Engine/Public/GPUSkinCache.h
--- a/Engine/Source/Runtime/Engine/Public/GPUSkinCache.h
+++ b/Engine/Source/Runtime/Engine/Public/GPUSkinCache.h
@@ -70,9 +70,8 @@ extern ENGINE_API bool GPUSkinCacheStoreDuplicatedVertices();
 
 extern ENGINE_API ESkinCacheDefaultBehavior GetSkinCacheDefaultBehavior();
 
-// Is it actually enabled?
+UE_DEPRECATED(5.8, "Use FGPUSkinCache::IsGPUSkinCacheEnabled()" )
 extern ENGINE_API int32 GEnableGPUSkinCache;
-extern int32 GSkinCacheRecomputeTangents;
 
 class FGPUSkinCacheEntry;
 
@@ -178,7 +177,9 @@ public:
 	static FColor GetVisualizationDebugColor(const FName& GPUSkinCacheVisualizationMode, FGPUSkinCacheEntry* Entry, FGPUSkinCacheEntry* RayTracingEntry, uint32 SectionIndex);
 	ENGINE_API void DrawVisualizationInfoText(const FName& GPUSkinCacheVisualizationMode, FScreenMessageWriter& ScreenMessageWriter) const;
 
-	static bool IsGPUSkinCacheRayTracingSupported();
+	ENGINE_API static bool IsGPUSkinCacheEnabled(ERHIFeatureLevel::Type FeatureLevel);
+	static int32 GetGPUSkinCacheRecomputeTangentsMode(ERHIFeatureLevel::Type FeatureLevel);
+	static bool IsGPUSkinCacheRayTracingSupported(ERHIFeatureLevel::Type FeatureLevel);
 
 	static FRWBuffer* GetPositionBuffer(FRDGBuilder& GraphBuilder, FGPUSkinCacheEntry const* Entry, uint32 SectionIndex);
 	static FRWBuffer* GetPreviousPositionBuffer(FRDGBuilder& GraphBuilder, FGPUSkinCacheEntry const* Entry, uint32 SectionIndex);
@@ -243,6 +244,8 @@ public:
 	// Releases the dynamic mesh bounds slot associated with a mesh.
 	ENGINE_API void ReleaseDynamicMeshBoundsSlot(int32 Offset, int32 NumSlots);
 
+	static void OnCVarUpdate();
+
 private:
 	friend FGPUSkinCacheEntry;
 
@@ -300,9 +303,6 @@ private:
 	ERHIFeatureLevel::Type FeatureLevel;
 	UWorld* World;
 
-	static void CVarSinkFunction();
-	static FAutoConsoleVariableSink CVarSink;
-
 	uint32 GetNextTransitionFence() { return ++TransitionFence; }
 	uint32 TransitionFence = 0;
 
diff --git a/Engine/Source/Runtime/RenderCore/Private/RenderUtils.cpp b/Engine/Source/Runtime/RenderCore/Private/RenderUtils.cpp
--- a/Engine/Source/Runtime/RenderCore/Private/RenderUtils.cpp
+++ b/Engine/Source/Runtime/RenderCore/Private/RenderUtils.cpp
@@ -1586,7 +1586,8 @@ RENDERCORE_API bool AreSkinCacheShadersEnabled(EShaderPlatform Platform)
 RENDERCORE_API bool IsGPUSkinCacheAllowed(EShaderPlatform Platform)
 {
 	static FShaderPlatformCachedIniValue<bool> PerPlatformCVar(TEXT("r.SkinCache.Allow"));
-	return PerPlatformCVar.Get(Platform);
+	static FShaderPlatformCachedIniValue<bool> PerPlatformMobileCVar(TEXT("r.Mobile.SkinCache.Allow"));
+	return IsMobilePlatform(Platform) ? PerPlatformMobileCVar.Get(Platform) : PerPlatformCVar.Get(Platform);
 }
 
 RENDERCORE_API bool IsGPUSkinCacheAvailable(EShaderPlatform Platform)
