Add extra profiling marker for SceneCapture updates

#rnx
#Synced-CL 48568815

diff --git a/Engine/Source/Runtime/Engine/Private/Components/SceneCaptureComponent.cpp b/Engine/Source/Runtime/Engine/Private/Components/SceneCaptureComponent.cpp
--- a/Engine/Source/Runtime/Engine/Private/Components/SceneCaptureComponent.cpp
+++ b/Engine/Source/Runtime/Engine/Private/Components/SceneCaptureComponent.cpp
@@ -39,6 +39,7 @@
 #include "Engine/TextureRenderTarget2D.h"
 #include "UnrealEngine.h"
 #include "SceneRenderBuilderInterface.h"
+#include "ProfilingDebugging/CpuProfilerTrace.h"
 
 #include UE_INLINE_GENERATED_CPP_BY_NAME(SceneCaptureComponent)
 
@@ -515,6 +516,8 @@ void USceneCaptureComponent::UpdateDeferredCaptures(FSceneInterface* Scene)
 
 void SceneCaptureUpdateDeferredCapturesInternal(FSceneInterface* Scene, TConstArrayView<const FSceneViewFamily*> ViewFamilies, ISceneRenderBuilder& SceneRenderBuilder)
 {
+	TRACE_CPUPROFILER_EVENT_SCOPE(SceneCaptureUpdateDeferredCapturesInternal);
+
 	FScopeLock ScopeLock(&SceneCapturesToUpdateMapCS);
 	UWorld* World = Scene->GetWorld();
 	if (!World || SceneCapturesToUpdateMap.Num() == 0)
