Combine SDR/HDR tonemapping method settings on FPostProcessSettings for simplicity

#Synced-CL 48567412

diff --git a/Engine/Source/Runtime/Engine/Classes/Engine/Scene.h b/Engine/Source/Runtime/Engine/Classes/Engine/Scene.h
--- a/Engine/Source/Runtime/Engine/Classes/Engine/Scene.h
+++ b/Engine/Source/Runtime/Engine/Classes/Engine/Scene.h
@@ -773,9 +773,7 @@ struct FPostProcessSettings
 	uint8 bOverride_ToneCurveAmount : 1;
 
 	UPROPERTY(EditAnywhere, BlueprintReadWrite, Category=Overrides, meta=(PinHiddenByDefault, InlineEditConditionToggle))
-	uint8 bOverride_TonemappingMethodSDR:1;
-	UPROPERTY(EditAnywhere, BlueprintReadWrite, Category=Overrides, meta=(PinHiddenByDefault, InlineEditConditionToggle))
-	uint8 bOverride_TonemappingMethodHDR:1;
+	uint8 bOverride_TonemappingMethod:1;
 	UPROPERTY(EditAnywhere, BlueprintReadWrite, Category=Overrides, meta=(PinHiddenByDefault, InlineEditConditionToggle))
 	uint8 bOverride_FilmSlope:1;
 	UPROPERTY(EditAnywhere, BlueprintReadWrite, Category=Overrides, meta=(PinHiddenByDefault, InlineEditConditionToggle))
@@ -1546,12 +1544,9 @@ struct FPostProcessSettings
 	UPROPERTY(interp, BlueprintReadWrite, Category = "Color Grading|Misc", meta = (ClampMin = "0.0", UIMax = "1.0", editcondition = "bOverride_ToneCurveAmount"))
 	float ToneCurveAmount;
 
-	/** Controls the method of tonemapping in SDR. */
-	UPROPERTY(interp, BlueprintReadWrite, Category="Film", meta=(editcondition = "bOverride_TonemappingMethodSDR", DisplayName = "SDR Method"))
-	ETonemappingMethod TonemappingMethodSDR;
-	/** Controls the method of tonemapping in HDR. */
-	UPROPERTY(interp, BlueprintReadWrite, Category="Film", meta=(editcondition = "bOverride_TonemappingMethodHDR", DisplayName = "HDR Method"))
-	ETonemappingMethod TonemappingMethodHDR;
+	/** Controls the method of tonemapping. */
+	UPROPERTY(interp, BlueprintReadWrite, Category="Film", meta=(editcondition = "bOverride_TonemappingMethod", DisplayName = "Method"))
+	ETonemappingMethod TonemappingMethod;
 	/** Controls the overall steepness of the tonemapper curve.  Larger values increase scene contrast and smaller values reduce contrast. */
 	UPROPERTY(interp, BlueprintReadWrite, Category="Film", meta=(UIMin = "0.0", UIMax = "1.0", editcondition = "bOverride_FilmSlope", DisplayName = "Slope"))
 	float FilmSlope;
diff --git a/Engine/Source/Runtime/Engine/Private/PostProcessUtils.cpp b/Engine/Source/Runtime/Engine/Private/PostProcessUtils.cpp
--- a/Engine/Source/Runtime/Engine/Private/PostProcessUtils.cpp
+++ b/Engine/Source/Runtime/Engine/Private/PostProcessUtils.cpp
@@ -53,8 +53,7 @@ bool FPostProcessUtils::OverridePostProcessSettings(FPostProcessSettings& ThisFr
 		UE_SET_PP(ExpandGamut);
 		UE_SET_PP(ToneCurveAmount);
 
-		UE_SET_PP(TonemappingMethodSDR);
-		UE_SET_PP(TonemappingMethodHDR);
+		UE_SET_PP(TonemappingMethod);
 		UE_SET_PP(FilmSlope);
 		UE_SET_PP(FilmToe);
 		UE_SET_PP(FilmShoulder);
@@ -360,8 +359,7 @@ bool FPostProcessUtils::BlendPostProcessSettings(FPostProcessSettings& ThisFrom,
 		UE_LERP_PP(ExpandGamut);
 		UE_LERP_PP(ToneCurveAmount);
 
-		UE_SET_PP(TonemappingMethodSDR);
-		UE_SET_PP(TonemappingMethodHDR);
+		UE_SET_PP(TonemappingMethod);
 		UE_LERP_PP(FilmSlope);
 		UE_LERP_PP(FilmToe);
 		UE_LERP_PP(FilmShoulder);
diff --git a/Engine/Source/Runtime/Engine/Private/Scene.cpp b/Engine/Source/Runtime/Engine/Private/Scene.cpp
--- a/Engine/Source/Runtime/Engine/Private/Scene.cpp
+++ b/Engine/Source/Runtime/Engine/Private/Scene.cpp
@@ -437,8 +437,7 @@ FPostProcessSettings::FPostProcessSettings()
 	ExpandGamut = 1.0f;
 	ToneCurveAmount = 1.0;
 
-	TonemappingMethodSDR = ETonemappingMethod::Filmic;
-	TonemappingMethodHDR = ETonemappingMethod::Filmic;
+	TonemappingMethod = ETonemappingMethod::Filmic;
 
 	// default values:
 
diff --git a/Engine/Source/Runtime/Engine/Private/SceneView.cpp b/Engine/Source/Runtime/Engine/Private/SceneView.cpp
--- a/Engine/Source/Runtime/Engine/Private/SceneView.cpp
+++ b/Engine/Source/Runtime/Engine/Private/SceneView.cpp
@@ -1618,8 +1618,7 @@ void FSceneView::OverridePostProcessSettings(const FPostProcessSettings& Src, fl
 		LERP_PP(ExpandGamut);
 		LERP_PP(ToneCurveAmount);
 
-		SET_PP(TonemappingMethodSDR);
-		SET_PP(TonemappingMethodHDR);
+		SET_PP(TonemappingMethod);
 		LERP_PP(FilmSlope);
 		LERP_PP(FilmToe);
 		LERP_PP(FilmShoulder);
diff --git a/Engine/Source/Runtime/Renderer/Private/PostProcess/PostProcessCombineLUTs.cpp b/Engine/Source/Runtime/Renderer/Private/PostProcess/PostProcessCombineLUTs.cpp
--- a/Engine/Source/Runtime/Renderer/Private/PostProcess/PostProcessCombineLUTs.cpp
+++ b/Engine/Source/Runtime/Renderer/Private/PostProcess/PostProcessCombineLUTs.cpp
@@ -385,15 +385,13 @@ struct FCachedLUTSettings
 		const int32 TonemappingMethodOverride = CVarLUTTonemappingMethod.GetValueOnRenderThread();
 		const uint32 TonemappingMethod = (TonemappingMethodOverride >= 0) ?
 			static_cast<uint32>(TonemappingMethodOverride) :
-			(static_cast<uint32>(ViewFamily.RenderTarget->GetSceneHDREnabled() ?
-				Settings.TonemappingMethodHDR :
-				Settings.TonemappingMethodSDR));
+			static_cast<uint32>(Settings.TonemappingMethod);
 		UPDATE_CACHE_SETTINGS(Parameters.TonemappingMethod, TonemappingMethod, bHasChanged);
 
 		const int32 TonemappingSDRInvEOTFOverride = CVarLUTTonemappingSDRInvEOTF.GetValueOnRenderThread();
 		const uint32 TonemappingSDRInvEOTF = (TonemappingSDRInvEOTFOverride >= 0) ?
 			static_cast<uint32>(TonemappingSDRInvEOTFOverride) :
-			static_cast<uint32>(Settings.TonemappingMethodSDR == ETonemappingMethod::StandardACES);
+			static_cast<uint32>(Settings.TonemappingMethod == ETonemappingMethod::StandardACES);
 		UPDATE_CACHE_SETTINGS(Parameters.TonemappingSDRInvEOTF, TonemappingSDRInvEOTF, bHasChanged);
 
 		UPDATE_CACHE_SETTINGS(Parameters.FilmSlope, Settings.FilmSlope, bHasChanged);
