Fix Specular profile not being applied on low end mobile platform.

#Synced-CL 48602645

diff --git a/Engine/Source/Runtime/Engine/Private/Rendering/SpecularProfile.cpp b/Engine/Source/Runtime/Engine/Private/Rendering/SpecularProfile.cpp
--- a/Engine/Source/Runtime/Engine/Private/Rendering/SpecularProfile.cpp
+++ b/Engine/Source/Runtime/Engine/Private/Rendering/SpecularProfile.cpp
@@ -50,7 +50,7 @@ static bool ForceUpdateSpecularProfile()
 
 static bool IsSpecularProfileSupport(EShaderPlatform InShaderPlatform)
 {
-	return IsFeatureLevelSupported(InShaderPlatform, ERHIFeatureLevel::SM5) && Substrate::IsSpecularProfileEnabled(InShaderPlatform);
+	return IsFeatureLevelSupported(InShaderPlatform, ERHIFeatureLevel::ES3_1) && Substrate::IsSpecularProfileEnabled(InShaderPlatform);
 }
 
 class FSpecularProfileCopyCS : public FGlobalShader
