Add a VS specific material Uniform Buffer
Use r.ShaderCompiler.UseMaterialVSUniformBuffer=0 and r.MaterialEditor.AllowSkipAlphaComputationForMaterialBlend=0 to restore the previous behaviour

This UB replaces the current Material Uniform Buffer for vertex shaders. Other frequencies (PS/CS/MS ...) still use the same material UB. This is done to remove coupling between what the VS needs and what the other stages need. This can remove a lot of vertex shaders after deduplication (ie, FShaderCompilerOutput::GenerateOutputHash will generate the same output hash)

The tracking happens first in FHLSLMaterialTranslator where we know the frequency at which we want to evaluate material properties. After walking the graph, we iterate on the different code chunks generated and make sure that what is used at the VS frequency is tagged.
Because of that, we need to be careful to not emit code chunks that do not contribute to the final value of the property. This is why UMaterialExpressionBlendMaterialAttributes::Compile was modified to skip the computation of the alpha value if both A and B are the same: we would otherwise tag some uniforms to be evaluated at the VS frequency, even though the code performing the material blend was not changing the return value.
Once translation is done, FUniformExpressionSet::CreateBufferStructs takes care of generating the HLSL declaration of the material UB, as well as keeping track of the CPU-side offsets between the canonical material UB and the VS-compacted one. That information is stored in CompactUniformsVSOptional.
At draw time, FMaterialRenderProxy::EvaluateUniformExpressions takes care of creating the compact VS UB if needed, and we use FUniformExpressionSet::CompactUniformBuffer to copy the relevant data. This avoids calling FillUniformBuffer multiple times.

#Synced-CL 48597399

diff --git a/Engine/Shaders/Public/IndirectVirtualTextureCommon.ush b/Engine/Shaders/Public/IndirectVirtualTextureCommon.ush
--- a/Engine/Shaders/Public/IndirectVirtualTextureCommon.ush
+++ b/Engine/Shaders/Public/IndirectVirtualTextureCommon.ush
@@ -13,18 +13,15 @@ struct FIndirectVirtualTextureInfo
 	uint BlockHeight;
 };
 
-FIndirectVirtualTextureUniform GetIndirectVirtualTextureUniform(uint Index)
+FIndirectVirtualTextureUniform GetIndirectVirtualTextureUniform(uint4 MaterialTextureCollectionPackedUniforms0, uint4 MaterialTextureCollectionPackedUniforms1
+																		  , uint4 MaterialTextureCollectionPackedUniforms2, uint4 MaterialTextureCollectionPackedUniforms3)
 {
-	// Uniform layouts with struct's doesn't appear supported, so, load
-	// it manually for now. When it's supported, this can be removed.
-	uint Offset = IndirectVirtualTextureUniformDWord4Count * Index;
-	
 	FIndirectVirtualTextureUniform Out = (FIndirectVirtualTextureUniform)0;
 #if TEXTURE_COLLECTION_PACKED_UNIFORMS
-	Out.UniformCountSub1          = Material.TextureCollectionPackedUniforms[Offset + 0].w;
-	Out.PackedPageTableUniform[0] = Material.TextureCollectionPackedUniforms[Offset + 1];
-	Out.PackedPageTableUniform[1] = Material.TextureCollectionPackedUniforms[Offset + 2];
-	Out.PackedUniform             = Material.TextureCollectionPackedUniforms[Offset + 3];
+	Out.UniformCountSub1          = MaterialTextureCollectionPackedUniforms0.w;
+	Out.PackedPageTableUniform[0] = MaterialTextureCollectionPackedUniforms1;
+	Out.PackedPageTableUniform[1] = MaterialTextureCollectionPackedUniforms2;
+	Out.PackedUniform             = MaterialTextureCollectionPackedUniforms3;
 #endif // TEXTURE_COLLECTION_PACKED_UNIFORMS
 	return Out;
 }
diff --git a/Engine/Shaders/Public/ShaderVersion.ush b/Engine/Shaders/Public/ShaderVersion.ush
--- a/Engine/Shaders/Public/ShaderVersion.ush
+++ b/Engine/Shaders/Public/ShaderVersion.ush
@@ -3,4 +3,4 @@
 // Each time the console command is executed it generates a new GUID. As this file is included
 // in Platform.ush (which should be included in any shader) it allows to invalidate the shader DDC.
 // 
-#pragma message("UESHADERMETADATA_VERSION 1EBEB311-1396-4095-A654-42314B22EFD0")
\ No newline at end of file
+#pragma message("UESHADERMETADATA_VERSION F1D733B6-E167-4498-8806-20F99812F8FE")
\ No newline at end of file
diff --git a/Engine/Source/Developer/Apple/MetalShaderFormat/Private/MetalCompileShaderMSC.cpp b/Engine/Source/Developer/Apple/MetalShaderFormat/Private/MetalCompileShaderMSC.cpp
--- a/Engine/Source/Developer/Apple/MetalShaderFormat/Private/MetalCompileShaderMSC.cpp
+++ b/Engine/Source/Developer/Apple/MetalShaderFormat/Private/MetalCompileShaderMSC.cpp
@@ -299,7 +299,8 @@ static void ProcessReflection(ID3D12ShaderReflection* ShaderReflection, const ui
 								FString(VariableDesc.Name),
 								VariableDesc.StartOffset,
 								VariableDesc.Size,
-								Output
+								Output,
+								Input
 							);
 						}
 					}
@@ -307,7 +308,7 @@ static void ProcessReflection(ID3D12ShaderReflection* ShaderReflection, const ui
 				
 				// Regular uniform buffer - we only care about the binding index
 				CCHeaderWriter.WriteUniformBlock(*UniformBufferName, BindIndex);
-				HandleReflectedUniformBuffer(UniformBufferName, BindIndex, Output);
+				HandleReflectedUniformBuffer(UniformBufferName, BindIndex, Output, Input);
 			}
 			NumCBVs = FMath::Max(NumCBVs, BindIndex + BindDesc.BindCount);
 		}
diff --git a/Engine/Source/Developer/Apple/MetalShaderFormat/Private/MetalShaderCompiler.cpp b/Engine/Source/Developer/Apple/MetalShaderFormat/Private/MetalShaderCompiler.cpp
--- a/Engine/Source/Developer/Apple/MetalShaderFormat/Private/MetalShaderCompiler.cpp
+++ b/Engine/Source/Developer/Apple/MetalShaderFormat/Private/MetalShaderCompiler.cpp
@@ -390,7 +390,7 @@ void BuildMetalShaderOutput(
 		}
 		UsedUniformBufferSlots[UBIndex] = true;
 
-		HandleReflectedUniformBuffer(UniformBlock.Name, UBIndex, ShaderOutput);
+		HandleReflectedUniformBuffer(UniformBlock.Name, UBIndex, ShaderOutput, ShaderInput);
 	}
 	
 	if (bOutOfBounds)
diff --git a/Engine/Source/Developer/ShaderCompilerCommon/Private/ShaderCompilerCommon.cpp b/Engine/Source/Developer/ShaderCompilerCommon/Private/ShaderCompilerCommon.cpp
--- a/Engine/Source/Developer/ShaderCompilerCommon/Private/ShaderCompilerCommon.cpp
+++ b/Engine/Source/Developer/ShaderCompilerCommon/Private/ShaderCompilerCommon.cpp
@@ -605,7 +605,8 @@ void HandleReflectedUniformBufferConstantBufferMember(
 	FStringView InMemberName,
 	int32 ReflectionOffset,
 	int32 ReflectionSize,
-	FShaderCompilerOutput& Output
+	FShaderCompilerOutput& Output,
+	const FShaderCompilerInput& Input
 )
 {
 	FStringView MemberName = InMemberName;
@@ -617,6 +618,21 @@ void HandleReflectedUniformBufferConstantBufferMember(
 		bAdd |= (ParameterType != EShaderParameterType::LooseData);
 	}
 
+	if (EnumHasAnyFlags(Reason, EUniformBufferMemberReflectionReason::MaterialUBBindError))
+	{
+#define SVARG(N) N.Len(), N.GetData()
+		FString ErrorMessage = FString::Printf(TEXT("%s: Unauthorized to use %.*s from %.*s at frequency %s"), 
+			*Input.GenerateShaderName(), 
+			SVARG(InMemberName),
+			SVARG(UniformBufferName),
+			GetShaderFrequencyString(Output.Target.GetFrequency()));
+#undef SVARG
+
+		FShaderCompilerError Error(MoveTemp(ErrorMessage));
+		Output.Errors.Add(Error);
+		Output.bSucceeded = false;
+	}
+
 	if (bAdd)
 	{
 		Output.ParameterMap.AddParameterAllocation(
@@ -668,12 +684,25 @@ void HandleReflectedRootConstantBuffer(
 		EShaderParameterType::LooseData);
 }
 
+bool VerifyBoundUniform(const FShaderCompilerInput& Input, FStringView UniformBufferName)
+{
+	EShaderFrequency ShaderFrequency = Input.Target.GetFrequency();
+	const TArray<FString>& DisallowedUniformBufferNames = (ShaderFrequency == SF_Vertex) ? Input.VSDisallowedUniformBufferNames : Input.NonVSDisallowedUniformBufferNames;
+	if (DisallowedUniformBufferNames.Contains(UniformBufferName))
+	{
+		return false;
+	}
+
+	return true;
+}
+
 void HandleReflectedUniformBuffer(
 	FStringView UniformBufferName,
 	int32 ReflectionSlot,
 	int32 BaseIndex,
 	int32 BufferSize,
-	FShaderCompilerOutput& CompilerOutput
+	FShaderCompilerOutput& CompilerOutput,
+	const FShaderCompilerInput& Input
 )
 {
 	FStringView AdjustedUniformBufferName(UE::ShaderCompilerCommon::RemoveConstantBufferPrefix(UniformBufferName));
@@ -685,6 +714,21 @@ void HandleReflectedUniformBuffer(
 		BufferSize,
 		EShaderParameterType::UniformBuffer
 	);
+
+	if (!VerifyBoundUniform(Input, AdjustedUniformBufferName))
+	{
+#define SVARG(N) N.Len(), N.GetData()
+		FString ErrorMessage = FString::Printf(TEXT("%s: Unauthorized to use from %.*s at frequency %s"), 
+			*Input.GenerateShaderName(),
+			SVARG(AdjustedUniformBufferName),
+			GetShaderFrequencyString(CompilerOutput.Target.GetFrequency()));
+#undef SVARG
+
+		FShaderCompilerError Error(MoveTemp(ErrorMessage));
+		CompilerOutput.Errors.Add(Error);
+		CompilerOutput.bSucceeded = false;
+	}
+
 }
 
 EUniformBufferMemberReflectionReason ShouldReflectUniformBufferMembers(const FShaderCompilerInput& Input, FStringView UniformBufferName)
@@ -701,6 +745,11 @@ EUniformBufferMemberReflectionReason ShouldReflectUniformBufferMembers(const FSh
 		Reason |= EUniformBufferMemberReflectionReason::NeedsReflection;
 	}
 
+	if (!VerifyBoundUniform(Input, UniformBufferName))
+	{
+		Reason |= EUniformBufferMemberReflectionReason::MaterialUBBindError;
+	}
+
 	return Reason;
 }
 
diff --git a/Engine/Source/Developer/ShaderCompilerCommon/Public/ShaderCompilerCommon.h b/Engine/Source/Developer/ShaderCompilerCommon/Public/ShaderCompilerCommon.h
--- a/Engine/Source/Developer/ShaderCompilerCommon/Public/ShaderCompilerCommon.h
+++ b/Engine/Source/Developer/ShaderCompilerCommon/Public/ShaderCompilerCommon.h
@@ -196,9 +196,10 @@ namespace UE::ShaderCompilerCommon
 
 enum class EUniformBufferMemberReflectionReason
 {
-	None            = 0,
-	NeedsReflection = 1 << 0,
-	Bindless        = 1 << 1,
+	None            	= 0,
+	NeedsReflection 	= 1 << 0,
+	Bindless        	= 1 << 1,
+	MaterialUBBindError = 1 << 2,
 };
 ENUM_CLASS_FLAGS(EUniformBufferMemberReflectionReason);
 
@@ -222,7 +223,8 @@ extern SHADERCOMPILERCOMMON_API void HandleReflectedUniformBufferConstantBufferM
 	FStringView MemberName,
 	int32 ReflectionOffset,
 	int32 ReflectionSize,
-	FShaderCompilerOutput& CompilerOutput
+	FShaderCompilerOutput& CompilerOutput,
+	const FShaderCompilerInput& Input
 );
 
 extern SHADERCOMPILERCOMMON_API void HandleReflectedRootConstantBufferMember(
@@ -244,17 +246,18 @@ extern SHADERCOMPILERCOMMON_API void HandleReflectedUniformBuffer(
 	int32 ReflectionSlot,
 	int32 BaseIndex,
 	int32 BufferSize,
-	FShaderCompilerOutput& CompilerOutput
+	FShaderCompilerOutput& CompilerOutput,
+	const FShaderCompilerInput& Input
 );
 
-inline void HandleReflectedUniformBuffer(FStringView UniformBufferName, int32 ReflectionSlot, int32 BufferSize, FShaderCompilerOutput& CompilerOutput)
+inline void HandleReflectedUniformBuffer(FStringView UniformBufferName, int32 ReflectionSlot, int32 BufferSize, FShaderCompilerOutput& CompilerOutput, const FShaderCompilerInput& Input)
 {
-	HandleReflectedUniformBuffer(UniformBufferName, ReflectionSlot, 0, BufferSize, CompilerOutput);
+	HandleReflectedUniformBuffer(UniformBufferName, ReflectionSlot, 0, BufferSize, CompilerOutput, Input);
 }
 
-inline void HandleReflectedUniformBuffer(FStringView UniformBufferName, int32 ReflectionSlot, FShaderCompilerOutput& CompilerOutput)
+inline void HandleReflectedUniformBuffer(FStringView UniformBufferName, int32 ReflectionSlot, FShaderCompilerOutput& CompilerOutput, const FShaderCompilerInput& Input)
 {
-	HandleReflectedUniformBuffer(UniformBufferName, ReflectionSlot, 0, CompilerOutput);
+	HandleReflectedUniformBuffer(UniformBufferName, ReflectionSlot, 0, CompilerOutput, Input);
 }
 
 extern SHADERCOMPILERCOMMON_API void HandleReflectedShaderResource(
diff --git a/Engine/Source/Developer/ShaderFormatOpenGL/Private/OpenGLShaderCompiler.cpp b/Engine/Source/Developer/ShaderFormatOpenGL/Private/OpenGLShaderCompiler.cpp
--- a/Engine/Source/Developer/ShaderFormatOpenGL/Private/OpenGLShaderCompiler.cpp
+++ b/Engine/Source/Developer/ShaderFormatOpenGL/Private/OpenGLShaderCompiler.cpp
@@ -673,7 +673,7 @@ void BuildShaderOutputInternal(
 		uint16 UBIndex = UniformBlock.Index;
 
 		UsedUniformBufferSlots[UBIndex] = true;
-		HandleReflectedUniformBuffer(UniformBlock.Name, UBIndex, ShaderOutput);
+		HandleReflectedUniformBuffer(UniformBlock.Name, UBIndex, ShaderOutput, ShaderInput);
 		Header.Bindings.NumUniformBuffers++;
 	}
 
@@ -720,7 +720,7 @@ void BuildShaderOutputInternal(
 	for (auto& PackedUB : CCHeader.PackedUBs)
 	{
 		UsedUniformBufferSlots[PackedUB.Attribute.Index] = true;
-		HandleReflectedUniformBuffer(PackedUB.Attribute.Name, PackedUB.Attribute.Index, ShaderOutput);
+		HandleReflectedUniformBuffer(PackedUB.Attribute.Name, PackedUB.Attribute.Index, ShaderOutput, ShaderInput);
 		Header.Bindings.NumUniformBuffers++;
 
 		// Nothing else...
diff --git a/Engine/Source/Developer/VulkanShaderFormat/Private/SpirVShaderCompiler.inl b/Engine/Source/Developer/VulkanShaderFormat/Private/SpirVShaderCompiler.inl
--- a/Engine/Source/Developer/VulkanShaderFormat/Private/SpirVShaderCompiler.inl
+++ b/Engine/Source/Developer/VulkanShaderFormat/Private/SpirVShaderCompiler.inl
@@ -806,7 +806,7 @@ static bool BuildShaderOutputFromSpirv(
 							if (InternalState.UseRootParametersStructure())
 							{
 								check(ReflectionSlot == FShaderParametersMetadata::kRootCBufferBindingIndex);
-								HandleReflectedUniformBuffer(ResourceName, ReflectionSlot, Output);
+								HandleReflectedUniformBuffer(ResourceName, ReflectionSlot, Output, InternalState.Input);
 							}
 
 							// Register all uniform buffer members of Globals as loose data
@@ -840,7 +840,7 @@ static bool BuildShaderOutputFromSpirv(
 						{
 							check(BindingTypeCount == ReflectionSlot);
 							check(!UsedUniformBufferSlots[ReflectionSlot]);
-							HandleReflectedUniformBuffer(ResourceName, ReflectionSlot, Output);
+							HandleReflectedUniformBuffer(ResourceName, ReflectionSlot, Output, InternalState.Input);
 							AddShaderValidationUBSize(BindingTypeCount, Binding->block.padded_size, Output);
 
 							const EUniformBufferMemberReflectionReason Reason = ShouldReflectUniformBufferMembers(InternalState.Input, ResourceName);
@@ -864,7 +864,8 @@ static bool BuildShaderOutputFromSpirv(
 										MemberName,
 										Member.absolute_offset,
 										Member.size,
-										Output
+										Output,
+										InternalState.Input
 									);
 								}
 							}
diff --git a/Engine/Source/Developer/Windows/ShaderFormatD3D/Private/D3DShaderCompiler.inl b/Engine/Source/Developer/Windows/ShaderFormatD3D/Private/D3DShaderCompiler.inl
--- a/Engine/Source/Developer/Windows/ShaderFormatD3D/Private/D3DShaderCompiler.inl
+++ b/Engine/Source/Developer/Windows/ShaderFormatD3D/Private/D3DShaderCompiler.inl
@@ -213,7 +213,7 @@ template <typename ID3D1xShaderReflection, typename D3D1x_SHADER_DESC, typename
 			{
 				// Track just the constant buffer itself.
 				AddShaderValidationUBSize(CBIndex, CBDesc.Size, Output);
-				HandleReflectedUniformBuffer(ConstantBufferName, CBIndex, Output);
+				HandleReflectedUniformBuffer(ConstantBufferName, CBIndex, Output, Input);
 				
 				CompileData.UsedUniformBufferSlots[CBIndex] = true;
 
@@ -238,7 +238,8 @@ template <typename ID3D1xShaderReflection, typename D3D1x_SHADER_DESC, typename
 								MemberName,
 								VariableDesc.StartOffset,
 								VariableDesc.Size,
-								Output
+								Output,
+								Input
 							);
 						}
 					}
diff --git a/Engine/Source/Runtime/Core/Private/UObject/FortniteMainBranchObjectVersion.cpp b/Engine/Source/Runtime/Core/Private/UObject/FortniteMainBranchObjectVersion.cpp
--- a/Engine/Source/Runtime/Core/Private/UObject/FortniteMainBranchObjectVersion.cpp
+++ b/Engine/Source/Runtime/Core/Private/UObject/FortniteMainBranchObjectVersion.cpp
@@ -10,10 +10,10 @@ TMap<FGuid, FGuid> FFortniteMainBranchObjectVersion::GetSystemGuids()
 
 	SystemGuids.Add(DevGuids.GLOBALSHADERMAP_DERIVEDDATA_VER, FGuid("EF39291B950C42C99FFE8D7AA98064BA"));
 	SystemGuids.Add(DevGuids.LANDSCAPE_MOBILE_COOK_VERSION, FGuid("32D02EF867C74B71A0D4E0FA41392732"));
-	SystemGuids.Add(DevGuids.MATERIALSHADERMAP_DERIVEDDATA_VER, FGuid("4DCE94C9E9994B3DA256F6A015B87CA9"));
+	SystemGuids.Add(DevGuids.MATERIALSHADERMAP_DERIVEDDATA_VER, FGuid("10F1046C8BCC4303A8F6DD4F55621F90"));
 	SystemGuids.Add(DevGuids.NIAGARASHADERMAP_DERIVEDDATA_VER, FGuid("CC229D919A214E59B79D01AC2E5BF981"));
 	SystemGuids.Add(DevGuids.Niagara_LatestScriptCompileVersion, FGuid("FB15C2870A4E4AE8B5F58566F0E3825E"));
-	SystemGuids.Add(DevGuids.MaterialTranslationDDCVersion, FGuid("E3ED11E905F64505BD5A6CB7E3599C10"));
+	SystemGuids.Add(DevGuids.MaterialTranslationDDCVersion, FGuid("A723F63E47F6470491077EE964C6726A"));
 
 	// These GUIDs are in AppendRestrictedFortniteMainStreamObjectVersionGuids as they should rarely be changed.
 	//SystemGuids.Add(DevGuids.NANITE_DERIVEDDATA_VER, FGuid("ACF98184978742A592537DA223A1E6BE"));
diff --git a/Engine/Source/Runtime/Engine/Private/Materials/HLSLMaterialTranslator.cpp b/Engine/Source/Runtime/Engine/Private/Materials/HLSLMaterialTranslator.cpp
--- a/Engine/Source/Runtime/Engine/Private/Materials/HLSLMaterialTranslator.cpp
+++ b/Engine/Source/Runtime/Engine/Private/Materials/HLSLMaterialTranslator.cpp
@@ -893,6 +893,14 @@ void FHLSLMaterialTranslator::GatherCustomVertexInterpolators(const ExpressionsA
 			{
 				CustomVertexInterpolators.AddUnique(Interpolator);
 				NextVertexInterpolatorIndex++;
+
+				for (FShaderCodeChunk& ShaderCodeChunk : CustomExpressionChunks)
+				{
+					if (ShaderCodeChunk.UniformExpression)
+					{
+						ShaderCodeChunk.UniformExpression->ShaderFrequencyMask |= (1 << SF_Vertex);
+					}
+				}
 			}
 
 			// Restore error handling
@@ -1027,7 +1035,19 @@ void FHLSLMaterialTranslator::CompileCustomOutputs(TArray<UMaterialExpressionCus
 							bCompilingPreviousFrame = PreviousFrameEvaluation != 0;
 							TArray<FShaderCodeChunk> CustomExpressionChunks;
 							AssignTempScope(CustomExpressionChunks);
-							CustomOutput->Compile(this, Index);
+							int32 CodeChunk = CustomOutput->Compile(this, Index);
+
+							if (CodeChunk != INDEX_NONE && ShaderFrequency == SF_Vertex)
+							{
+								for (FShaderCodeChunk& ShaderCodeChunk : CustomExpressionChunks)
+								{
+									if (ShaderCodeChunk.UniformExpression)
+									{
+										ShaderCodeChunk.UniformExpression->ShaderFrequencyMask |= (1 << SF_Vertex);
+									}
+								}
+							}
+
 						}
 					}
 
@@ -1111,6 +1131,17 @@ uint64 FHLSLMaterialTranslator::CompileMaterialAttributesCustomOutputs(uint64 Cu
 			int32 CodeChunk = Material->CompileCustomAttribute(CustomAttribute.AttributeID, this);
 			if (CodeChunk != INDEX_NONE)
 			{
+				if (ShaderFrequency == SF_Vertex)
+				{
+					for (FShaderCodeChunk& ShaderCodeChunk : CustomExpressionChunks)
+					{
+						if (ShaderCodeChunk.UniformExpression)
+						{
+							ShaderCodeChunk.UniformExpression->ShaderFrequencyMask |= (1 << SF_Vertex);
+						}
+					}
+				}
+
 				CustomOutput(CustomAttribute.AttributeName, HLSLTypeString(GetParameterType(CodeChunk)), CustomAttribute.FunctionName, CustomAttribute.OutputIndex, CodeChunk, EMaterialCustomOutputFlags::None);
 			}
 		}
@@ -1386,6 +1417,32 @@ EHLSLMaterialTranslatorResult FHLSLMaterialTranslator::Translate(bool bForceDisa
 	return TranslationResult;
 }
 
+static FString GenerateCombinedCodeChunk(const TCHAR* CommonCodeChunk, const TCHAR* VSCodeChunk)
+{
+	return FString::Printf(
+		"#if VERTEXSHADER" HLSL_LINE_TERMINATOR
+		"%s" HLSL_LINE_TERMINATOR
+		"#else" HLSL_LINE_TERMINATOR
+		"%s" HLSL_LINE_TERMINATOR
+		"#endif" HLSL_LINE_TERMINATOR,
+		VSCodeChunk,
+		CommonCodeChunk);
+}
+
+uint64 HasShaderFrequency(uint64 ShaderFrequencyMask, EShaderFrequency ShaderFrequency)
+{
+	return (ShaderFrequencyMask & (1ULL << ((uint64)ShaderFrequency))) ? 1 : 0;
+}
+
+static inline const TCHAR* GetCurrentMaterialUBNameInternal(EShaderFrequency ShaderFrequency)
+{
+	if (ShaderFrequency == SF_Vertex)
+	{
+		return TEXT("__Material_UB_VS__");
+	}
+	return TEXT("__Material_UB__");
+}
+
 void FHLSLMaterialTranslator::TranslateMaterial()
 {
 	TranslationResult = EHLSLMaterialTranslatorResult::Success;
@@ -2367,16 +2424,74 @@ void FHLSLMaterialTranslator::TranslateMaterial()
 		}
 	}
 
+	// Make sure every uniform expression is marked as accessed by the VS. This is needed below to send the information to the uniform expression set to create the compact version of the UB
+	{
+		static_assert(sizeof(EShaderFrequencyMask)*8 >= (1 << SF_NumBits));
+		const EShaderFrequencyMask VSFrequencyMask = (1 << SF_Vertex);
+		for (FShaderCodeChunk& ShaderCodeChunk : SharedPropertyCodeChunks[SF_Vertex])
+		{
+			if (ShaderCodeChunk.UniformExpression)
+			{
+				ShaderCodeChunk.UniformExpression->ShaderFrequencyMask |= VSFrequencyMask;
+				FMaterialUniformExpressionTexture* MaterialUniformExpressionTexture = ShaderCodeChunk.UniformExpression->GetTextureUniformExpression();
+				if (MaterialUniformExpressionTexture)
+				{
+					int32 MaterialCacheTagIndex = MaterialUniformExpressionTexture->GetMaterialCacheTagIndex();
+					if (MaterialCompilationOutput.UniformExpressionSet.MaterialCacheTagStacks.IsValidIndex(MaterialCacheTagIndex))
+					{
+						MaterialCompilationOutput.UniformExpressionSet.MaterialCacheTagStacks[MaterialCacheTagIndex].ShaderFrequencyMask |= VSFrequencyMask;
+					}
+				}
+			}
+		}
+	}
+
+	FString ResourcesStringPrefix = TEXT("#define __Material_UB__ Material" HLSL_LINE_TERMINATOR);
+
+	if (UseMaterialVSUniformBuffer())
+	{
+		ResourcesStringPrefix += GenerateCombinedCodeChunk(TEXT("#define __Material_UB_VS__ Material"), TEXT("#define __Material_UB_VS__ MaterialVS"));
+	}
+	else
+	{
+		// When we want to compile a Nanite compute shader which is also calling SF_Vertex level 
+		// functions, we actually remap __Material_UB_VS__ to the generic Material UB. 
+		// Since the indices are the same for VS-accessed uniforms, we only need to replace the UB name
+		ResourcesStringPrefix += TEXT("#define __Material_UB_VS__ Material" HLSL_LINE_TERMINATOR);
+	}
+
+	ResourcesStringPrefix += GenerateUniformExpressionMacroDefinitions(); 
+	ResourcesString = ResourcesStringPrefix + ResourcesString;
+
 	// Store the number of float4s
 	MaterialCompilationOutput.UniformExpressionSet.UniformPreshaderBufferSize = (UniformPreshaderOffset + 3u) / 4u;
 
+	FCompactUniformExpressionSet* CompactUniformsVS = nullptr;
+	if (UseMaterialVSUniformBuffer())
+	{
+		CompactUniformsVS = MaterialCompilationOutput.UniformExpressionSet.CreateCompactUniformsVS();
+		check(CompactUniformsVS);
+	}
+	else
+	{
+		MaterialCompilationOutput.UniformExpressionSet.CompactUniformsVSOptional.Empty();
+		check(UniformPreshaderOffsetVS == 0);
+	}
+
+	if (CompactUniformsVS)
+	{
+		CompactUniformsVS->UniformPreshaderBufferSize = (UniformPreshaderOffsetVS + 3u) / 4u;
+	}
+
 	for (uint32 TypeIndex = 0u; TypeIndex < NumMaterialTextureParameterTypes; ++TypeIndex)
 	{
 		MaterialCompilationOutput.UniformExpressionSet.UniformTextureParameters[TypeIndex].Empty(UniformTextureExpressions[TypeIndex].Num());
 		for (FMaterialUniformExpressionTexture* TextureExpression : UniformTextureExpressions[TypeIndex])
 		{
 			CHECK_DDC_QUERY_FINISHED_ELSE_RETURN();
-			TextureExpression->GetTextureParameterInfo(MaterialCompilationOutput.UniformExpressionSet.UniformTextureParameters[TypeIndex].AddDefaulted_GetRef());
+			FMaterialTextureParameterInfo& TextureParameter = MaterialCompilationOutput.UniformExpressionSet.UniformTextureParameters[TypeIndex].AddDefaulted_GetRef();
+			TextureExpression->GetTextureParameterInfo(TextureParameter);
+			TextureParameter.ShaderFrequencyMask = TextureExpression->ShaderFrequencyMask;
 		}
 	}
 
@@ -2384,14 +2499,18 @@ void FHLSLMaterialTranslator::TranslateMaterial()
 	for (FMaterialUniformExpressionExternalTexture* TextureExpression : UniformExternalTextureExpressions)
 	{
 		CHECK_DDC_QUERY_FINISHED_ELSE_RETURN();
-		TextureExpression->GetExternalTextureParameterInfo(MaterialCompilationOutput.UniformExpressionSet.UniformExternalTextureParameters.AddDefaulted_GetRef());
+		FMaterialExternalTextureParameterInfo& TextureParameter = MaterialCompilationOutput.UniformExpressionSet.UniformExternalTextureParameters.AddDefaulted_GetRef();
+		TextureExpression->GetExternalTextureParameterInfo(TextureParameter);
+		TextureParameter.ShaderFrequencyMask = TextureExpression->ShaderFrequencyMask;
 	}
 
 	MaterialCompilationOutput.UniformExpressionSet.UniformTextureCollectionParameters.Empty(UniformTextureCollectionExpressions.Num());
 	for (const FMaterialUniformExpressionTextureCollection* TextureExpression : UniformTextureCollectionExpressions)
 	{
 		CHECK_DDC_QUERY_FINISHED_ELSE_RETURN();
-		TextureExpression->GetTextureCollectionParameterInfo(MaterialCompilationOutput.UniformExpressionSet.UniformTextureCollectionParameters.AddDefaulted_GetRef());
+		FMaterialTextureCollectionParameterInfo& TextureParameter = MaterialCompilationOutput.UniformExpressionSet.UniformTextureCollectionParameters.AddDefaulted_GetRef();
+		TextureExpression->GetTextureCollectionParameterInfo(TextureParameter);
+		TextureParameter.ShaderFrequencyMask = TextureExpression->ShaderFrequencyMask;
 	}
 
 	if (bSubstrateEnabled)
@@ -2695,6 +2814,32 @@ void FHLSLMaterialTranslator::TranslateMaterial()
 
 	MaterialCompilationOutput.UniformExpressionSet.SetParameterCollections(ParameterCollections);
 
+	for (int VTStackIndex = 0; VTStackIndex < MaterialCompilationOutput.UniformExpressionSet.VTStacks.Num(); ++VTStackIndex)
+	{
+		FMaterialVirtualTextureStack& CurrentVTStack = MaterialCompilationOutput.UniformExpressionSet.VTStacks[VTStackIndex];
+		const uint32 NumLayers = CurrentVTStack.GetNumLayers();
+		uint32 NumLayersWithVSFrequency = 0;
+		uint32 NumLayersWithUniformExpression = 0;
+
+		for (uint32 LayerIndex = 0u; LayerIndex < NumLayers; ++LayerIndex)
+		{
+			const int32 ParameterIndex = CurrentVTStack.LayerUniformExpressionIndices[LayerIndex];
+			if (ParameterIndex != INDEX_NONE)
+			{
+				const FMaterialTextureParameterInfo& Parameter = MaterialCompilationOutput.UniformExpressionSet.GetTextureParameter(EMaterialTextureParameterType::Virtual, ParameterIndex);
+				++NumLayersWithUniformExpression;
+				if (HasShaderFrequency(Parameter.ShaderFrequencyMask, SF_Vertex))
+				{
+					++NumLayersWithVSFrequency;
+				}
+			}
+		}
+
+		// ensure that all layers share the same frequency
+		check(NumLayersWithVSFrequency == 0 || (NumLayersWithVSFrequency == NumLayersWithUniformExpression));
+		CurrentVTStack.ShaderFrequencyMask = (NumLayersWithVSFrequency > 0) ? (1<<SF_Vertex) : 0;
+	}
+
 	// Store the number of unique VT samples
 	MaterialCompilationOutput.EstimatedNumVirtualTextureLookups = NumVtSamples;
 }
@@ -2725,6 +2870,11 @@ void FHLSLMaterialTranslator::ValidateShadingModelsForFeatureLevel(const FMateri
 #endif
 }
 
+const TCHAR* FHLSLMaterialTranslator::GetCurrentMaterialUniformBufferName() const
+{
+	return GetCurrentMaterialUBNameInternal(ShaderFrequency);
+}
+
 void FHLSLMaterialTranslator::GetMaterialEnvironment(EShaderPlatform InPlatform, FShaderCompilerEnvironment& OutEnvironment)
 {
 	check(InPlatform == GetShaderPlatform());
@@ -2876,6 +3026,8 @@ void FHLSLMaterialTranslator::GetMaterialEnvironment(EShaderPlatform InPlatform,
 	OutEnvironment.SetDefine(TEXT("MATERIAL_ENABLE_TRANSLUCENT_HIGH_QUALITY_LOCAL_LIGHT_SHADOW"), EnvironmentDefines->bMaterialEnableTranslucentHighQualityLocalLightShadow);
 	OutEnvironment.SetDefine(TEXT("MATERIAL_ENABLE_TRANSLUCENT_HIGH_QUALITY_DIRECTIONAL_LIGHT_SHADOW"), EnvironmentDefines->bMaterialEnableTranslucentHighQualityDirectionalLightShadow);
 
+
+	check(MaterialCompilationOutput.UniformExpressionSet.VTStacks.Num() == EnvironmentDefines->VirtualPageTypes.Num());
 	for (int i = 0; i < EnvironmentDefines->VirtualPageTypes.Num(); ++i)
 	{
 		if (EnvironmentDefines->VirtualPageTypes[i] & FEnvironmentDefines::TABLE_MESH_PAINT)
@@ -2898,20 +3050,25 @@ void FHLSLMaterialTranslator::GetMaterialEnvironment(EShaderPlatform InPlatform,
 		}
 		else
 		{
+			FMaterialVirtualTextureStack& CurrentVTStack = MaterialCompilationOutput.UniformExpressionSet.VTStacks[i];
+			EShaderFrequency UniformExpressionVSOnlyFrequency = (HasShaderFrequency(CurrentVTStack.ShaderFrequencyMask, SF_Vertex)) ? SF_Vertex : SF_NumFrequencies;
+
+			const TCHAR* MaterialVSName = GetCurrentMaterialUBNameInternal(UniformExpressionVSOnlyFrequency);
+
 			// Setup page table defines to map each VT stack to either 1 or 2 page table textures, depending on how many layers it uses
-			FString PageTableValue = FString::Printf(TEXT("Material.VirtualTexturePageTable0_%d"), i);
+			FString PageTableValue = FString::Printf(TEXT("%s.VirtualTexturePageTable0_%d"), MaterialVSName, i);
 			if (EnvironmentDefines->VirtualPageTypes[i] & FEnvironmentDefines::LOCAL_TABLE1)
 			{
-				PageTableValue += FString::Printf(TEXT(", Material.VirtualTexturePageTable1_%d"), i);
+				PageTableValue += FString::Printf(TEXT(", %s.VirtualTexturePageTable1_%d"), MaterialVSName, i);
 			}
 			if (EnvironmentDefines->VirtualPageTypes[i] & FEnvironmentDefines::LOCAL_ADAPTIVE_INDIRECTION)
 			{
-				PageTableValue += FString::Printf(TEXT(", Material.VirtualTexturePageTableIndirection_%d"), i);
+				PageTableValue += FString::Printf(TEXT(", %s.VirtualTexturePageTableIndirection_%d"), MaterialVSName, i);
 			}
 			OutEnvironment.SetDefine(*FString::Printf(TEXT("VIRTUALTEXTURE_PAGETABLE_%d"), i), *PageTableValue);
 
 			// Setup page table uniform defines.
-			FString PageTableUniformValue = FString::Printf(TEXT("Material.VTPackedPageTableUniform[%d*2], Material.VTPackedPageTableUniform[%d*2+1]"), i, i);
+			FString PageTableUniformValue = FString::Printf(TEXT("%s.VTPackedPageTableUniform[%d*2], %s.VTPackedPageTableUniform[%d*2+1]"), MaterialVSName, i, MaterialVSName, i);
 			OutEnvironment.SetDefine(*FString::Printf(TEXT("VIRTUALTEXTURE_PAGETABLE_UNIFORM_%d"), i), *PageTableUniformValue);
 		}
 	}
@@ -3350,7 +3507,6 @@ FString FHLSLMaterialTranslator::GenerateFunctionCode(uint32 Index, ECompiledPar
 	return DerivativeVariations[Variation].TranslatedCodeChunkDefinitions[Index] + TEXT("	return ") + DerivativeVariations[Variation].TranslatedCodeChunks[Index] + TEXT(";");
 }
 
-
 const FShaderCodeChunk& FHLSLMaterialTranslator::AtParameterCodeChunk(int32 Index) const
 {
 	check(Index != INDEX_NONE);
@@ -4285,7 +4441,269 @@ int32 FHLSLMaterialTranslator::AddUniformExpression(FAddUniformExpressionScope&
 }
 
 // From MaterialUniformExpressions.cpp
-extern void WriteMaterialUniformAccess(UE::Shader::EValueComponentType ComponentType, uint32 NumComponents, uint32 UniformOffset, FStringBuilderBase& OutResult);
+extern void WriteMaterialUniformAccess(UE::Shader::EValueComponentType ComponentType, uint32 NumComponents, uint32 UniformOffset, const TCHAR* MaterialNamePrefix, FStringBuilderBase& OutResult);
+
+void SetupNumericUniformExpressionOffset(FUniformExpressionSet& UniformExpressionSet, 
+	FMaterialUniformExpression* CodeChunkUniformExpression, 
+	uint32& UniformPreshaderOffset,
+	const uint32 NumComponents,
+	bool bIsLWC,
+	int32 PreshaderGapInterval)
+{
+	if (ensureAlways(CodeChunkUniformExpression->UniformOffset == INDEX_NONE))
+	{
+		const uint32 RegisterOffset = UniformPreshaderOffset % 4;
+		if (!bIsLWC && RegisterOffset + NumComponents > 4u)
+		{
+			// If this uniform would span multiple registers, align offset to the next register to avoid this
+			UniformPreshaderOffset = Align(UniformPreshaderOffset, 4u);
+		}
+
+		// Optionally insert an empty vector element (gap) at the given interval, to work around a platform specific shader compiler bug
+		if (!PreshaderGapInterval)
+		{
+			PreshaderGapInterval = GPreshaderGapInterval;
+		}
+		if (PreshaderGapInterval > 0)
+		{
+			// Set a minimum of 4 (one out of every four elements is padding)
+			PreshaderGapInterval = FMath::Max(PreshaderGapInterval, 4);
+
+			// Check if we are on the vector element that is the end of the gap interval, or we will pass the end of the gap interval if the
+			// preshader is larger than a single vector.  The divides will change value when a modulus gap interval boundary is crossed.
+			uint32 PreshaderGapOffset = FMath::Max(4u, bIsLWC ? NumComponents * 2u : NumComponents);
+			if (UniformPreshaderOffset / (PreshaderGapInterval *4) != (UniformPreshaderOffset + PreshaderGapOffset) / (PreshaderGapInterval *4))
+			{
+				UniformPreshaderOffset += 4u;
+			}
+		}
+
+		const UE::Shader::EValueComponentType ComponentType = bIsLWC ? UE::Shader::EValueComponentType::Double : UE::Shader::EValueComponentType::Float;
+
+		UniformExpressionSet.WriteUniformPreshaderEntry(
+			UniformPreshaderOffset, UE::Shader::MakeValueType(ComponentType, NumComponents),
+			[&CodeChunkUniformExpression](UE::Shader::FPreshaderData& UniformPreshaderData) -> void
+			{
+				CodeChunkUniformExpression->WriteNumberOpcodes(UniformPreshaderData);
+			}
+		);
+
+		CodeChunkUniformExpression->UniformOffset = UniformPreshaderOffset;
+		UniformPreshaderOffset += bIsLWC ? NumComponents * 2u : NumComponents;
+	}
+}
+
+void SetupBooleanUniformExpressionOffset(FUniformExpressionSet& UniformExpressionSet, 
+	FMaterialUniformExpression* CodeChunkUniformExpression, 
+	uint32& UniformPreshaderOffset,
+	uint32& CurrentNumBoolComponents, 
+	uint32& CurrentBoolUniformOffset, 
+	const uint32 NumComponents)
+{
+	if (ensureAlways(CodeChunkUniformExpression->UniformOffset == INDEX_NONE))
+	{
+		// 'Bool' uniforms are packed into bits
+		if (CurrentNumBoolComponents + NumComponents > 32u)
+		{
+			CurrentBoolUniformOffset = UniformPreshaderOffset++;
+			CurrentNumBoolComponents = 0u;
+		}
+
+		int32 UniformOffset = CurrentBoolUniformOffset * 32u + CurrentNumBoolComponents;
+
+		UniformExpressionSet.WriteUniformPreshaderEntry(
+			UniformOffset, UE::Shader::MakeValueType(UE::Shader::EValueComponentType::Bool, NumComponents),
+			[&CodeChunkUniformExpression](UE::Shader::FPreshaderData& UniformPreshaderData) -> void
+			{
+				CodeChunkUniformExpression->WriteNumberOpcodes(UniformPreshaderData);
+			}
+		);
+
+		CodeChunkUniformExpression->UniformOffset = UniformOffset;
+		CurrentNumBoolComponents += NumComponents;
+	}
+
+}
+
+static void EmitBoolPreshaderBufferString(FStringBuilderBase& FormattedCode, const FMaterialUniformExpression* CodeChunkUniformExpression, const TCHAR* CurrentMaterialUBName)
+{
+	const uint32 UniformOffset = CodeChunkUniformExpression->UniformOffset / 32u;
+	const uint32 NumBoolComponents = CodeChunkUniformExpression->UniformOffset % 32u;
+
+	const uint32 RegisterIndex = UniformOffset / 4;
+	const uint32 RegisterOffset = UniformOffset % 4;
+	FormattedCode.Appendf(TEXT("UnpackUniform_%s(asuint(%s.PreshaderBuffer[%u][%u]), %u)"),
+		TEXT("bool"),
+		CurrentMaterialUBName,
+		RegisterIndex,
+		RegisterOffset,
+		NumBoolComponents);
+
+}
+
+static void EmitNumericPreshaderBufferString(FStringBuilderBase& FormattedCode, const EMaterialValueType MaterialValueType, const TCHAR* MaterialUBName ,const uint32 UniformOffset)
+{
+	const bool bIsLWC = IsLWCType(MaterialValueType);
+	const uint32 NumComponents = GetNumComponents(MaterialValueType);
+
+	if (bIsLWC)
+	{
+		if (NumComponents == 1)
+		{
+			FormattedCode.Append(TEXT("DFToWS(MakeDFScalar("));
+		}
+		else
+		{
+			FormattedCode.Appendf(TEXT("DFToWS(MakeDFVector%d("), NumComponents);
+		}
+
+		WriteMaterialUniformAccess(UE::Shader::EValueComponentType::Float, NumComponents, UniformOffset, MaterialUBName, FormattedCode); // High
+		FormattedCode.Append(TEXT(","));
+		WriteMaterialUniformAccess(UE::Shader::EValueComponentType::Float, NumComponents, UniformOffset + NumComponents, MaterialUBName, FormattedCode); // Low
+		FormattedCode.Append(TEXT("))"));
+	}
+	else
+	{
+		if (IsUIntNumericType(MaterialValueType))
+		{
+			FormattedCode.Append(TEXT("asuint("));
+		}
+
+		WriteMaterialUniformAccess(UE::Shader::EValueComponentType::Float, NumComponents, UniformOffset, MaterialUBName, FormattedCode);
+
+		if (IsUIntNumericType(MaterialValueType))
+		{
+			FormattedCode.Append(TEXT(")"));
+		}
+	}
+}
+
+static const TCHAR* GetMaterialPreshaderBufferPrefix()
+{
+	return TEXT("__Material_PreshaderBuffer_");
+}
+
+void FHLSLMaterialTranslator::GenerateUniformExpressionMacro(FStringBuilderBase& FormattedCode, const FShaderCodeChunk&	CodeChunk)
+{
+	// We first use UniformOffset as a unique id to emit the __Material_PreshaderBuffer_%u macro below.
+	// Once we're done traversing the graph, FHLSLMaterialTranslator::GenerateUniformExpressionMacroDefinitions 
+	// will set UniformOffset to its final value for runtime use, ie the preshader offset in the UB
+	if (CodeChunk.UniformExpression->UniformOffset == INDEX_NONE)
+	{
+		CodeChunk.UniformExpression->UniformOffset = UniformExpressionShaderFrequencies.Num();
+		UniformExpressionShaderFrequencies.Add(	FUniformExpressionInPreshader{ .UniformExpression = CodeChunk.UniformExpression, .CodeChunkMaterialValueType = CodeChunk.Type });
+	}
+
+	int32 PreshaderBufferWrapperIndex = CodeChunk.UniformExpression->UniformOffset;
+	FUniformExpressionInPreshader& UniformExpressionInPreshader = UniformExpressionShaderFrequencies[PreshaderBufferWrapperIndex];
+	FormattedCode.Appendf(TEXT("%s%u"), GetMaterialPreshaderBufferPrefix(), PreshaderBufferWrapperIndex);
+}
+
+FString FHLSLMaterialTranslator::GenerateUniformExpressionMacroDefinitions()
+{
+	TArray<int32> UniformWrapperIndices;
+	UniformWrapperIndices.Reserve(UniformExpressionShaderFrequencies.Num());
+	for (const FUniformExpressionInPreshader& UniformMetaData : UniformExpressionShaderFrequencies)
+	{
+		check(UniformMetaData.UniformExpression->UniformOffset != INDEX_NONE);
+		UniformWrapperIndices.Add(UniformMetaData.UniformExpression->UniformOffset);
+	}
+
+	for (const FUniformExpressionInPreshader& UniformMetaData : UniformExpressionShaderFrequencies)
+	{
+		UniformMetaData.UniformExpression->UniformOffset = INDEX_NONE;
+	}
+
+	FUniformExpressionSet& MaterialCompilationOutputUniformExpressionSet = MaterialCompilationOutput.UniformExpressionSet;
+
+	auto SetupUniformExpressionInPreshader = [this]
+		(const FUniformExpressionInPreshader& UniformMetaData)
+	{
+		const uint32 NumComponents = GetNumComponents(UniformMetaData.CodeChunkMaterialValueType);
+		const bool bIsLWC = IsLWCType(UniformMetaData.CodeChunkMaterialValueType);
+
+		if (UniformMetaData.CodeChunkMaterialValueType == MCT_Bool)
+		{
+			SetupBooleanUniformExpressionOffset(MaterialCompilationOutput.UniformExpressionSet,
+				UniformMetaData.UniformExpression,
+				UniformPreshaderOffset,
+				CurrentNumBoolComponents,
+				CurrentBoolUniformOffset,
+				NumComponents);
+
+		}
+		else if (IsFloatNumericType(UniformMetaData.CodeChunkMaterialValueType) || IsUIntNumericType(UniformMetaData.CodeChunkMaterialValueType))
+		{
+			SetupNumericUniformExpressionOffset(MaterialCompilationOutput.UniformExpressionSet, 
+				UniformMetaData.UniformExpression, 
+				UniformPreshaderOffset,
+				NumComponents,
+				bIsLWC,
+				Material->GetPreshaderGap());
+		}
+		else
+		{
+			checkNoEntry();
+		}
+	};
+
+	// First process parameters that will be used by the compact version so that we can get a single memcpy when compacting
+	if (UseMaterialVSUniformBuffer())
+	{
+		for (const FUniformExpressionInPreshader& UniformMetaData : UniformExpressionShaderFrequencies)
+		{
+			if (HasShaderFrequency(UniformMetaData.UniformExpression->ShaderFrequencyMask, SF_Vertex))
+			{
+				SetupUniformExpressionInPreshader(UniformMetaData);
+			}
+
+		}
+
+		UniformPreshaderOffsetVS = UniformPreshaderOffset;
+
+		for (const FUniformExpressionInPreshader& UniformMetaData : UniformExpressionShaderFrequencies)
+		{
+			if (!HasShaderFrequency(UniformMetaData.UniformExpression->ShaderFrequencyMask, SF_Vertex))
+			{
+				SetupUniformExpressionInPreshader(UniformMetaData);
+			}
+		}
+	}
+	else
+	{
+		for (const FUniformExpressionInPreshader& UniformMetaData : UniformExpressionShaderFrequencies)
+		{
+			SetupUniformExpressionInPreshader(UniformMetaData);
+		}
+	}
+
+	TStringBuilder<1024> OutResult;
+	check(UniformWrapperIndices.Num() == UniformExpressionShaderFrequencies.Num());
+	for (int32 UniformExpressionIndex = 0; UniformExpressionIndex < UniformExpressionShaderFrequencies.Num(); ++UniformExpressionIndex)
+	{
+		const FUniformExpressionInPreshader& UniformMetaData = UniformExpressionShaderFrequencies[UniformExpressionIndex];
+		EShaderFrequency UniformExpressionVSOnlyFrequency = HasShaderFrequency(UniformMetaData.UniformExpression->ShaderFrequencyMask, SF_Vertex) ? SF_Vertex : SF_NumFrequencies;
+		const TCHAR* MaterialVSName = GetCurrentMaterialUBNameInternal(UniformExpressionVSOnlyFrequency);
+
+		OutResult.Appendf(TEXT("#define %s%u "), GetMaterialPreshaderBufferPrefix(), UniformWrapperIndices[UniformExpressionIndex]);
+		if (UniformMetaData.CodeChunkMaterialValueType == MCT_Bool)
+		{
+			EmitBoolPreshaderBufferString(OutResult, UniformMetaData.UniformExpression, MaterialVSName);
+		}
+		else if (IsFloatNumericType(UniformMetaData.CodeChunkMaterialValueType) || IsUIntNumericType(UniformMetaData.CodeChunkMaterialValueType))
+		{
+			EmitNumericPreshaderBufferString(OutResult, UniformMetaData.CodeChunkMaterialValueType, MaterialVSName, UniformMetaData.UniformExpression->UniformOffset);
+		}
+		else
+		{
+			checkNoEntry();
+		}
+
+		OutResult.Append(TEXT("\n"));
+	}
+
+	return OutResult.ToString();
+}
 
 // AccessUniformExpression - Adds code to access the value of a uniform expression to the Code array and returns its index.
 int32 FHLSLMaterialTranslator::AccessUniformExpression(int32 Index)
@@ -4324,123 +4742,13 @@ int32 FHLSLMaterialTranslator::AccessUniformExpression(int32 Index)
 		check(CodeChunk.UniformExpression->GetType() == &FMaterialUniformExpressionStaticBoolParameter::StaticType);
 		FMaterialUniformExpressionStaticBoolParameter* StaticBoolParameter = static_cast<FMaterialUniformExpressionStaticBoolParameter*>(CodeChunk.UniformExpression.GetReference());
 		check(!bIsLWC);
-
-		const uint32 NumComponents = GetNumComponents(CodeChunk.Type);
-
-		if (CodeChunk.UniformExpression->UniformOffset == INDEX_NONE)
-		{
-			// 'Bool' uniforms are packed into bits
-			if (CurrentNumBoolComponents + NumComponents > 32u)
-			{
-				CurrentBoolUniformOffset = UniformPreshaderOffset++;
-				CurrentNumBoolComponents = 0u;
-			}
-
-			int32 UniformOffset = CurrentBoolUniformOffset * 32u + CurrentNumBoolComponents;
-
-			MaterialCompilationOutput.UniformExpressionSet.WriteUniformPreshaderEntry(
-				UniformOffset, UE::Shader::MakeValueType(UE::Shader::EValueComponentType::Bool, NumComponents),
-				[&CodeChunk](UE::Shader::FPreshaderData& UniformPreshaderData) -> void
-				{
-					CodeChunk.UniformExpression->WriteNumberOpcodes(UniformPreshaderData);
-				}
-			);
-
-			CodeChunk.UniformExpression->UniformOffset = UniformOffset;
-			CurrentNumBoolComponents += NumComponents;
-		}
-
-		const uint32 UniformOffset = CodeChunk.UniformExpression->UniformOffset / 32u;
-		const uint32 NumBoolComponents = CodeChunk.UniformExpression->UniformOffset % 32u;
-
-		const uint32 RegisterIndex = UniformOffset / 4;
-		const uint32 RegisterOffset = UniformOffset % 4;
-		FormattedCode.Appendf(TEXT("UnpackUniform_%s(asuint(Material.PreshaderBuffer[%u][%u]), %u)"),
-			TEXT("bool"),
-			RegisterIndex,
-			RegisterOffset,
-			NumBoolComponents);
+		GenerateUniformExpressionMacro(FormattedCode, CodeChunk);
 	}
 	else if (IsFloatNumericType(CodeChunk.Type) || IsUIntNumericType(CodeChunk.Type))
 	{
 		check(CodeChunk.DerivativeStatus != EDerivativeStatus::Valid);
-		const uint32 NumComponents = GetNumComponents(CodeChunk.Type);
-
-		if (CodeChunk.UniformExpression->UniformOffset == INDEX_NONE)
-		{
-			const uint32 RegisterOffset = UniformPreshaderOffset % 4;
-			if (!bIsLWC && RegisterOffset + NumComponents > 4u)
-			{
-				// If this uniform would span multiple registers, align offset to the next register to avoid this
-				UniformPreshaderOffset = Align(UniformPreshaderOffset, 4u);
-			}
-
-			// Optionally insert an empty vector element (gap) at the given interval, to work around a platform specific shader compiler bug
-			int32 PreshaderGapInterval = Material->GetPreshaderGap();
-			if (!PreshaderGapInterval)
-			{
-				PreshaderGapInterval = GPreshaderGapInterval;
-			}
-			if (PreshaderGapInterval > 0)
-			{
-				// Set a minimum of 4 (one out of every four elements is padding)
-				PreshaderGapInterval = FMath::Max(PreshaderGapInterval, 4);
-
-				// Check if we are on the vector element that is the end of the gap interval, or we will pass the end of the gap interval if the
-				// preshader is larger than a single vector.  The divides will change value when a modulus gap interval boundary is crossed.
-				uint32 PreshaderGapOffset = FMath::Max(4u, bIsLWC ? NumComponents * 2u : NumComponents);
-				if (UniformPreshaderOffset / (PreshaderGapInterval *4) != (UniformPreshaderOffset + PreshaderGapOffset) / (PreshaderGapInterval *4))
-				{
-					UniformPreshaderOffset += 4u;
-				}
-			}
 
-			const UE::Shader::EValueComponentType ComponentType = bIsLWC ? UE::Shader::EValueComponentType::Double : UE::Shader::EValueComponentType::Float;
-
-			MaterialCompilationOutput.UniformExpressionSet.WriteUniformPreshaderEntry(
-				UniformPreshaderOffset, UE::Shader::MakeValueType(ComponentType, NumComponents),
-				[&CodeChunk](UE::Shader::FPreshaderData& UniformPreshaderData) -> void
-				{
-					CodeChunk.UniformExpression->WriteNumberOpcodes(UniformPreshaderData);
-				}
-			);
-
-			CodeChunk.UniformExpression->UniformOffset = UniformPreshaderOffset;
-			UniformPreshaderOffset += bIsLWC ? NumComponents * 2u : NumComponents;
-		}
-
-		const uint32 UniformOffset = CodeChunk.UniformExpression->UniformOffset;
-		if (bIsLWC)
-		{
-			AddLWCFuncUsage(ELWCFunctionKind::Constructor);
-			if (NumComponents == 1)
-			{
-				FormattedCode.Append(TEXT("DFToWS(MakeDFScalar("));
-			}
-			else
-			{
-				FormattedCode.Appendf(TEXT("DFToWS(MakeDFVector%d("), NumComponents);
-			}
-
-			WriteMaterialUniformAccess(UE::Shader::EValueComponentType::Float, NumComponents, UniformOffset, FormattedCode); // High
-			FormattedCode.Append(TEXT(","));
-			WriteMaterialUniformAccess(UE::Shader::EValueComponentType::Float, NumComponents, UniformOffset + NumComponents, FormattedCode); // Low
-			FormattedCode.Append(TEXT("))"));
-		}
-		else
-		{
-			if (IsUIntNumericType(CodeChunk.Type))
-			{
-				FormattedCode.Append(TEXT("asuint("));
-			}
-
-			WriteMaterialUniformAccess(UE::Shader::EValueComponentType::Float, NumComponents, UniformOffset, FormattedCode);
-
-			if (IsUIntNumericType(CodeChunk.Type))
-			{
-				FormattedCode.Append(TEXT(")"));
-			}
-		}
+		GenerateUniformExpressionMacro(FormattedCode, CodeChunk);
 	}
 	else if(CodeChunk.Type & MCT_Texture)
 	{
@@ -4481,13 +4789,13 @@ int32 FHLSLMaterialTranslator::AccessUniformExpression(int32 Index)
 		};
 		if(GenerateCode)
 		{
-			FormattedCode.Appendf(TEXT("Material.%s_%u"), BaseName, TextureInputIndex);
+			FormattedCode.Appendf(TEXT("%s.%s_%u"), GetCurrentMaterialUniformBufferName(), BaseName, TextureInputIndex);
 		}
 	}
 	else if (CodeChunk.Type == MCT_SparseVolumeTexture)
 	{
 		int32 TextureInputIndex = UniformTextureExpressions[(uint32)EMaterialTextureParameterType::SparseVolume].AddUnique(TextureUniformExpression);
-		FormattedCode.Appendf(TEXT("SparseVolumeTextureUnpackUniforms(Material.SVTPackedUniform[%d*2], Material.SVTPackedUniform[%d*2+1])"), TextureInputIndex, TextureInputIndex);
+		FormattedCode.Appendf(TEXT("SparseVolumeTextureUnpackUniforms(%s.SVTPackedUniform[%d*2], %s.SVTPackedUniform[%d*2+1])"), GetCurrentMaterialUniformBufferName(), TextureInputIndex, GetCurrentMaterialUniformBufferName(), TextureInputIndex);
 	}
 	else if (CodeChunk.Type == MCT_TextureCollection)
 	{
@@ -4504,7 +4812,7 @@ int32 FHLSLMaterialTranslator::AccessUniformExpression(int32 Index)
 		}
 		
 		int32 TextureCollectionInputIndex = UniformTextureCollectionExpressions.AddUnique(TextureCollectionUniformExpression);
-		FormattedCode.Appendf(TEXT("Material.TextureCollection_%u"), TextureCollectionInputIndex);
+		FormattedCode.Appendf(TEXT("%s.TextureCollection_%u"), GetCurrentMaterialUniformBufferName(), TextureCollectionInputIndex);
 	}
 	else
 	{
@@ -7362,6 +7670,20 @@ static bool SamplerDebugSupported(EMaterialValueType TextureType, bool bVirtualT
 	return IsDebugTextureSampleEnabled() && (TextureType == MCT_Texture2D || bVirtualTexture) && !bDecal;
 }
 
+FString GetIndirectVirtualTextureUniformParameters(const TCHAR* MaterialUBName, int32 VirtualTextureIndex)
+{
+	if (!ensure(VirtualTextureIndex != INDEX_NONE))
+	{
+		VirtualTextureIndex = 0;
+	}
+
+	FString Offset = FString::Printf(TEXT("(IndirectVirtualTextureUniformDWord4Count * %i)"), VirtualTextureIndex);
+	return FString::Printf(
+		TEXT("%s.TextureCollectionPackedUniforms[%s + 0], %s.TextureCollectionPackedUniforms[%s + 1]"
+		", %s.TextureCollectionPackedUniforms[%s + 2], %s.TextureCollectionPackedUniforms[%s + 3]"), 
+		MaterialUBName, *Offset, MaterialUBName, *Offset, MaterialUBName, *Offset, MaterialUBName, *Offset);
+}
+
 int32 FHLSLMaterialTranslator::TextureSample(
 	int32 TextureIndex,
 	int32 CoordinateIndex,
@@ -7528,7 +7850,7 @@ int32 FHLSLMaterialTranslator::TextureSample(
 
 		SampleInfo.PageTableIndex = AddInternalCodeChunk(
 			MCT_Unexposed, TEXT("Texture2D<uint4>"),
-			*FString::Printf(TEXT("Material.MaterialCacheTagPageTable0_%d"), MaterialCacheTagIndex)
+			*FString::Printf(TEXT("%s.MaterialCacheTagPageTable0_%d"), GetCurrentMaterialUniformBufferName(), MaterialCacheTagIndex)
 		);
 		
 		SampleInfo.PageTableUniformIndex = AddCodeChunkInner(
@@ -7539,8 +7861,8 @@ int32 FHLSLMaterialTranslator::TextureSample(
 
 		// Physical texture indexing assumes TagIndex * Stride + LayerIndex
 		TextureName = FString::Printf(
-			TEXT("Material.MaterialCacheTagPhysical_%d, View.SharedBilinearClampedSampler"),
-			MaterialCacheTagIndex * MaterialCacheMaxRuntimeLayers + TextureUniform->GetTextureLayerIndex()
+			TEXT("%s.MaterialCacheTagPhysical_%d, View.SharedBilinearClampedSampler"),
+			GetCurrentMaterialUniformBufferName(), MaterialCacheTagIndex * MaterialCacheMaxRuntimeLayers + TextureUniform->GetTextureLayerIndex()
 		);
 
 		NumVtSamples++;
@@ -7562,11 +7884,11 @@ int32 FHLSLMaterialTranslator::TextureSample(
 		{
 			const bool   bUseAnisoSampler  = VirtualTextureScalability::IsAnisotropicFilteringEnabled() && MipValueMode != TMVM_MipLevel;
 			const TCHAR* SharedSamplerName = bUseAnisoSampler ? TEXT("View.SharedBilinearAnisoClampedSampler") : TEXT("View.SharedBilinearClampedSampler");
-			TextureName += FString::Printf(TEXT("Material.TextureCollectionPhysical_%d, GetMaterialSharedSampler(Material.TextureCollectionSampler_%d, %s)"), VirtualTextureIndex, VirtualTextureIndex, SharedSamplerName);
+			TextureName += FString::Printf(TEXT("%s.TextureCollectionPhysical_%d, GetMaterialSharedSampler(%s.TextureCollectionSampler_%d, %s)"), GetCurrentMaterialUniformBufferName(), VirtualTextureIndex, GetCurrentMaterialUniformBufferName(), VirtualTextureIndex, SharedSamplerName);
 		}
 		else
 		{
-			TextureName += FString::Printf(TEXT("Material.TextureCollectionPhysical_%d, Material.TextureCollectionSampler_%d"), VirtualTextureIndex, VirtualTextureIndex);
+			TextureName += FString::Printf(TEXT("%s.TextureCollectionPhysical_%d, %s.TextureCollectionSampler_%d"), GetCurrentMaterialUniformBufferName(), VirtualTextureIndex, GetCurrentMaterialUniformBufferName(), VirtualTextureIndex);
 		}
 		
 		NumVtSamples++;
@@ -7596,13 +7918,13 @@ int32 FHLSLMaterialTranslator::TextureSample(
 			// The wrap/clamp/mirror operation is handled in the shader explicitly.
 			const bool bUseAnisoSampler = VirtualTextureScalability::IsAnisotropicFilteringEnabled() && MipValueMode != TMVM_MipLevel;
 			const TCHAR* SharedSamplerName = bUseAnisoSampler ? TEXT("View.SharedBilinearAnisoClampedSampler") : TEXT("View.SharedBilinearClampedSampler");
-			TextureName += FString::Printf(TEXT("Material.VirtualTexturePhysical_%d, GetMaterialSharedSampler(Material.VirtualTexturePhysical_%dSampler, %s)")
-				, VirtualTextureIndex, VirtualTextureIndex, SharedSamplerName);
+			TextureName += FString::Printf(TEXT("%s.VirtualTexturePhysical_%d, GetMaterialSharedSampler(%s.VirtualTexturePhysical_%dSampler, %s)"),
+				GetCurrentMaterialUniformBufferName(), VirtualTextureIndex, GetCurrentMaterialUniformBufferName(), VirtualTextureIndex, SharedSamplerName);
 		}
 		else
 		{
-			TextureName += FString::Printf(TEXT("Material.VirtualTexturePhysical_%d, Material.VirtualTexturePhysical_%dSampler")
-				, VirtualTextureIndex, VirtualTextureIndex);
+			TextureName += FString::Printf(TEXT("%s.VirtualTexturePhysical_%d, %s.VirtualTexturePhysical_%dSampler"),
+				GetCurrentMaterialUniformBufferName(), VirtualTextureIndex, GetCurrentMaterialUniformBufferName(), VirtualTextureIndex);
 		}
 
 		const bool bIsStrictValidation = EnumHasAnyFlags(Material->GetMaterialTranslateValidationFlags(), EMaterialTranslateValidationFlags::Strict_RuntimeVirtualTexture);
@@ -7736,27 +8058,27 @@ int32 FHLSLMaterialTranslator::TextureSample(
 		}
 		else
 		{
-			const TCHAR* SharedSamplerName = nullptr;
+			FString SharedSamplerName;
 
 			if (SamplerSource == SSM_Wrap_WorldGroupSettings)
 			{
 				// Use the shared sampler to save sampler slots
-				SharedSamplerName = AutomaticViewMipBias ? TEXT("View.MaterialTextureBilinearWrapedSampler") : TEXT("Material.Wrap_WorldGroupSettings");
+				SharedSamplerName = AutomaticViewMipBias ? FString(TEXT("View.MaterialTextureBilinearWrapedSampler")) : FString::Printf(TEXT("%s.Wrap_WorldGroupSettings"), GetCurrentMaterialUniformBufferName());
 				RequiresManualViewMipBias = false;
 			}
 			else if (SamplerSource == SSM_Clamp_WorldGroupSettings)
 			{
 				// Use the shared sampler to save sampler slots
-				SharedSamplerName = AutomaticViewMipBias ? TEXT("View.MaterialTextureBilinearClampedSampler") : TEXT("Material.Clamp_WorldGroupSettings");
+				SharedSamplerName = AutomaticViewMipBias ? FString(TEXT("View.MaterialTextureBilinearClampedSampler")) : FString::Printf(TEXT("%s.Clamp_WorldGroupSettings"), GetCurrentMaterialUniformBufferName());
 				RequiresManualViewMipBias = false;
 			}
 			else if (SamplerSource == SSM_TerrainWeightmapGroupSettings)
 			{
-				SharedSamplerName = TEXT("View.LandscapeWeightmapSampler");
+				SharedSamplerName = FString(TEXT("View.LandscapeWeightmapSampler"));
 				RequiresManualViewMipBias = false;
 			}
 
-			if (SharedSamplerName)
+			if (!SharedSamplerName.IsEmpty())
 			{
 				if (bFromCollection)
 				{
@@ -7764,7 +8086,7 @@ int32 FHLSLMaterialTranslator::TextureSample(
 				}
 				else
 				{
-					SamplerStateCode = FString::Printf(TEXT("GetMaterialSharedSampler(%sSampler,%s)"), *TextureName, SharedSamplerName);
+					SamplerStateCode = FString::Printf(TEXT("GetMaterialSharedSampler(%sSampler,%s)"), *TextureName, *SharedSamplerName);
 				}
 			}
 		}
@@ -7948,7 +8270,7 @@ int32 FHLSLMaterialTranslator::TextureSample(
 			}
 
 			FMaterialUniformExpressionTexture* TextureUniform = UniformExpression->GetTextureUniformExpression();
-			VTPackedUniformName = FString::Printf(TEXT("VTUniform_Unpack(Material.MaterialCacheTagPageTableUniform[%d])"), TextureUniform->GetMaterialCacheTagIndex());
+			VTPackedUniformName = FString::Printf(TEXT("VTUniform_Unpack(%s.MaterialCacheTagPageTableUniform[%d])"), GetCurrentMaterialUniformBufferName(), TextureUniform->GetMaterialCacheTagIndex());
 		}
 		else if (bFromCollection)
 		{
@@ -7956,12 +8278,12 @@ int32 FHLSLMaterialTranslator::TextureSample(
 			
 			SampleInfo.PageTableIndex = AddInternalCodeChunk(
 				MCT_Unexposed, TEXT("Texture2D<uint4>"),
-				*FString::Printf(TEXT("Material.TextureCollectionPageTable_%i"), VirtualTextureIndex)
+				*FString::Printf(TEXT("%s.TextureCollectionPageTable_%i"), GetCurrentMaterialUniformBufferName(), VirtualTextureIndex)
 			);
 
 			FString Uniforms = GetParameterCode(AddInternalCodeChunk(
 				MCT_Unexposed, TEXT("FIndirectVirtualTextureUniform"),
-				*FString::Printf(TEXT("GetIndirectVirtualTextureUniform(%i)"), VirtualTextureIndex)
+				*FString::Printf(TEXT("GetIndirectVirtualTextureUniform(%s)"), *GetIndirectVirtualTextureUniformParameters(GetCurrentMaterialUniformBufferName(), VirtualTextureIndex))
 			));
 
 			// Manually unpack the page table to redirect the page boundaries
@@ -7981,7 +8303,7 @@ int32 FHLSLMaterialTranslator::TextureSample(
 		else 
 		{
 			check(VirtualTextureIndex >= 0);
-			VTPackedUniformName = FString::Printf(TEXT("VTUniform_Unpack(Material.VTPackedUniform[%d])"), VirtualTextureIndex);
+			VTPackedUniformName = FString::Printf(TEXT("VTUniform_Unpack(%s.VTPackedUniform[%d])"), GetCurrentMaterialUniformBufferName(), VirtualTextureIndex);
 
 			SampleInfo.CoordinateIndex = NonLWCCoordinateIndex;
 
@@ -8211,7 +8533,7 @@ int32 FHLSLMaterialTranslator::TextureFromCollection(int32 TextureCollectionCode
 		
 		FString Uniforms = GetParameterCode(AddInternalCodeChunk(
 			MCT_Unexposed, TEXT("FIndirectVirtualTextureUniform"),
-			*FString::Printf(TEXT("GetIndirectVirtualTextureUniform(%i)"), TextureCollectionExpression->GetTextureCollectionTypePrefixIndex())
+			*FString::Printf(TEXT("GetIndirectVirtualTextureUniform(%s)"), *GetIndirectVirtualTextureUniformParameters(GetCurrentMaterialUniformBufferName(), TextureCollectionExpression->GetTextureCollectionTypePrefixIndex()))
 		));
 		
 		return AddInternalCodeChunk(
@@ -8339,6 +8661,7 @@ FHLSLMaterialTranslator::FMaterialCacheTag& FHLSLMaterialTranslator::FindOrAddMa
 	// Add tag stack to expression set
 	FMaterialCacheTagStack& TagStack = MaterialCompilationOutput.UniformExpressionSet.MaterialCacheTagStacks.Emplace_GetRef();
 	TagStack.TagGuid = Layout.Guid;
+	TagStack.ShaderFrequencyMask = 0;
 
 	return Tag;
 }
@@ -8355,8 +8678,8 @@ int32 FHLSLMaterialTranslator::MaterialCacheTextureDescriptor(const FMaterialCac
 	// Though, it'll probably combine the reads anyway.
 	int32 CodeIndex = AddInlinedCodeChunkZeroDeriv(
 		MCT_TextureMaterialCache,
-		TEXT("Material.MaterialCacheTagBuffer_%u[GetMaterialCacheTagOffset(GetPrimitiveData(%s))]"),
-		Tag.StackIndex, *OverloadArgument
+		TEXT("%s.MaterialCacheTagBuffer_%u[GetMaterialCacheTagOffset(GetPrimitiveData(%s))]"),
+		GetCurrentMaterialUniformBufferName(), Tag.StackIndex, *OverloadArgument
 	);
 
 	// Assign allocated stack index through texture uniform
@@ -9037,8 +9360,8 @@ int32 FHLSLMaterialTranslator::TextureCollectionCount(int32 InTextureCollectionC
 
 		return AddCodeChunk(
 			EMaterialValueType(MCT_UInt1),
-			TEXT("GetIndirectVirtualTextureCount(GetIndirectVirtualTextureUniform(%u))"),
-			TextureCollectionExpression->GetTextureCollectionTypePrefixIndex()
+			TEXT("GetIndirectVirtualTextureCount(GetIndirectVirtualTextureUniform(%s))"),
+			*GetIndirectVirtualTextureUniformParameters(GetCurrentMaterialUniformBufferName(), TextureCollectionExpression->GetTextureCollectionTypePrefixIndex())
 		);
 	}
 	
@@ -17330,7 +17653,7 @@ void FHLSLMaterialTranslator::PrepareMaterialSourceStringParameters()
 
 	MaterialSourceTemplateParams.Add({ TEXT("get_material_world_position_offset_raw"), *GenerateFunctionCode(MP_WorldPositionOffset, BaseDerivativeVariation) });
 	MaterialSourceTemplateParams.Add({ TEXT("get_material_previous_world_position_offset_raw"), *GenerateFunctionCode(CompiledMP_PrevWorldPositionOffset, BaseDerivativeVariation) });
-
+	
 	// Print custom texture coordinate assignments, should be fine with regular derivatives
 	FString CustomUVAssignments;
 
@@ -17534,8 +17857,9 @@ int32 FHLSLMaterialTranslator::SparseVolumeTextureSamplePageTable(int32 SparseVo
 
 	AddEstimatedTextureSample();
 	const TCHAR* FunctionPostfix = bIsManualLinearMipMapSecondSample ? TEXT("SecondMipWrapper") : TEXT("");
-	FString SampleCode = FString::Printf(TEXT("SparseVolumeTextureSamplePageTable%s(Material.SparseVolumeTexturePageTable_%d, %s, %s, %s, %s, %s, %s)"),
+	FString SampleCode = FString::Printf(TEXT("SparseVolumeTextureSamplePageTable%s(%s.SparseVolumeTexturePageTable_%d, %s, %s, %s, %s, %s, %s)"),
 		FunctionPostfix,
+		GetCurrentMaterialUniformBufferName(),
 		SVTReferenceIndex, 
 		*GetParameterCode(SparseVolumeTextureIndex), 
 		*GetParameterCode(UVWAsFloat3Index), 
@@ -17581,10 +17905,14 @@ int32 FHLSLMaterialTranslator::SparseVolumeTextureSamplePhysicalTileData(int32 S
 
 	AddEstimatedTextureSample();
 	const TCHAR* FunctionPostfix = bIsManualLinearMipMapSecondSample ? TEXT("SecondMipWrapper") : TEXT("");
-	return AddCodeChunk(MCT_Float4, TEXT("SparseVolumeTextureSamplePhysicalTileData%s(Material.SparseVolumeTexturePhysicalA_%d, Material.SparseVolumeTexturePhysicalB_%d, Material.SparseVolumeTexturePhysical_%dSampler, %s, %s)"),
+	return AddCodeChunk(MCT_Float4, TEXT("SparseVolumeTextureSamplePhysicalTileData%s(%s.SparseVolumeTexturePhysicalA_%d, %s.SparseVolumeTexturePhysicalB_%d, %s.SparseVolumeTexturePhysical_%dSampler, %s, %s)"),
+		GetCurrentMaterialUniformBufferName(),
 		FunctionPostfix,
+		GetCurrentMaterialUniformBufferName(),
 		SVTReferenceIndex,
+		GetCurrentMaterialUniformBufferName(),
 		SVTReferenceIndex,
+		GetCurrentMaterialUniformBufferName(),
 		SVTReferenceIndex,
 		*GetParameterCode(VoxelCoordAsFloat3Index),
 		*GetParameterCode(IndexAsFloatIndex));
diff --git a/Engine/Source/Runtime/Engine/Private/Materials/HLSLMaterialTranslator.h b/Engine/Source/Runtime/Engine/Private/Materials/HLSLMaterialTranslator.h
--- a/Engine/Source/Runtime/Engine/Private/Materials/HLSLMaterialTranslator.h
+++ b/Engine/Source/Runtime/Engine/Private/Materials/HLSLMaterialTranslator.h
@@ -234,6 +234,9 @@ enum class EHLSLMaterialTranslatorResult
 	RetryWithoutDDC,
 };
 
+// Specify which shader frequency an entry appears in
+using FShaderFrequencyMask = uint64;
+
 class FHLSLMaterialTranslator : public FMaterialCompiler
 {
 	friend class FMaterialDerivativeAutogen;
@@ -339,13 +342,22 @@ protected:
 
 	TMap<FName, FMaterialLocalVariableEntry> LocalVariables[SF_NumFrequencies];
 
+	struct FUniformExpressionInPreshader
+	{
+		TRefCountPtr<FMaterialUniformExpression> UniformExpression;
+		EMaterialValueType CodeChunkMaterialValueType = (EMaterialValueType)0;
+	};
+
 	// Uniform expressions used across all material properties
 	TArray<FShaderCodeChunk> UniformExpressions;
 	TArray<TRefCountPtr<FMaterialUniformExpressionTexture> > UniformTextureExpressions[NumMaterialTextureParameterTypes];
 	TArray<TRefCountPtr<FMaterialUniformExpressionTextureCollection>> UniformTextureCollectionExpressions;
 	TArray<TRefCountPtr<FMaterialUniformExpressionExternalTexture>> UniformExternalTextureExpressions;
+	TArray<FUniformExpressionInPreshader> UniformExpressionShaderFrequencies;
+
 	TMap<UE::Shader::FValue, uint32> DefaultUniformValues;
 	uint32 UniformPreshaderOffset = 0u;
+	uint32 UniformPreshaderOffsetVS = 0u;
 	uint32 CurrentBoolUniformOffset = ~0u;
 	uint32 CurrentNumBoolComponents = 32u;
 
@@ -864,6 +876,8 @@ protected:
 
 	int32 AccessMaterialAttribute(int32 CodeIndex, const FGuid& AttributeID);
 
+	const TCHAR* GetCurrentMaterialUniformBufferName() const;
+
 	// GetParameterType
 	virtual EMaterialValueType GetParameterType(int32 Index) const override;
 
@@ -1554,6 +1568,10 @@ protected:
 	 */
 	void PrepareEnvironmentDefines();
 
+
+	void GenerateUniformExpressionMacro(FStringBuilderBase& FormattedCode, const FShaderCodeChunk&	CodeChunk);
+	FString GenerateUniformExpressionMacroDefinitions();
+
 	/** The final compilation output. */
 	FMaterialCompilationOutput& MaterialCompilationOutput;
 	FMaterialCompilationOutput DDCMaterialCompilationOutput;
diff --git a/Engine/Source/Runtime/Engine/Private/Materials/MaterialExpressions.cpp b/Engine/Source/Runtime/Engine/Private/Materials/MaterialExpressions.cpp
--- a/Engine/Source/Runtime/Engine/Private/Materials/MaterialExpressions.cpp
+++ b/Engine/Source/Runtime/Engine/Private/Materials/MaterialExpressions.cpp
@@ -353,6 +353,15 @@ FUObjectAnnotationSparseBool GMaterialFunctionsThatNeedSamplerFixup;
 FUObjectAnnotationSparseBool GMaterialFunctionsThatNeedFeatureLevelSM6Fix;
 
 static const TCHAR* CPD_UI_ErrorMessage = TEXT("Custom Primitive Data can't be used with the UI material domain.");
+
+bool GAllowSkipAlphaComputationForMaterialBlend = true;
+static FAutoConsoleVariableRef CVarAllowSkipAlphaComputationForMaterialBlend(
+	TEXT("r.MaterialEditor.AllowSkipAlphaComputationForMaterialBlend"),
+	GAllowSkipAlphaComputationForMaterialBlend,
+	TEXT("If true, Allow to skip alpha computation when evaluating material blend nodes if the 2 values are equal"),
+	ECVF_ReadOnly
+);
+
 #endif // #if WITH_EDITOR
 
 // Helper functions to handle determinism during cooking
@@ -7082,6 +7091,11 @@ FExpressionInput* UMaterialExpressionBlendMaterialAttributes::GetInput(int32 Inp
 	}
 }
 
+static inline bool AllowSkipAlphaComputationForMaterialBlend()
+{
+	return GAllowSkipAlphaComputationForMaterialBlend;
+}
+
 int32 UMaterialExpressionBlendMaterialAttributes::Compile(class FMaterialCompiler* Compiler, int32 OutputIndex)
 {
 	const FGuid AttributeID = Compiler->GetMaterialAttribute();
@@ -7125,6 +7139,11 @@ int32 UMaterialExpressionBlendMaterialAttributes::Compile(class FMaterialCompile
 	Compiler->SubstrateTreeStackPush(this, 1);
 	int32 ResultB = B.CompileWithDefault(Compiler, AttributeID);
 	Compiler->SubstrateTreeStackPop();
+	if (AllowSkipAlphaComputationForMaterialBlend() && (ResultA == ResultB))
+	{
+		return ResultA;
+	}
+
 	int32 ResultAlpha = Alpha.Compile(Compiler);
 
 	MaterialAttributeBlendFunction BlendFunction = FMaterialAttributeDefinitionMap::GetBlendFunction(AttributeID);
diff --git a/Engine/Source/Runtime/Engine/Private/Materials/MaterialRenderProxy.cpp b/Engine/Source/Runtime/Engine/Private/Materials/MaterialRenderProxy.cpp
--- a/Engine/Source/Runtime/Engine/Private/Materials/MaterialRenderProxy.cpp
+++ b/Engine/Source/Runtime/Engine/Private/Materials/MaterialRenderProxy.cpp
@@ -41,6 +41,7 @@ FUniformExpressionCache::~FUniformExpressionCache()
 {
 	ResetAllocatedVTs();
 	UniformBuffer.SafeRelease();
+	UniformBufferVS.SafeRelease();
 }
 
 class FUniformExpressionCacheAsyncUpdateTask
@@ -111,9 +112,9 @@ void FUniformExpressionCacheAsyncUpdateScope::WaitForTask()
 class FUniformExpressionCacheAsyncUpdater
 {
 public:
-	void Add(FUniformExpressionCache* UniformExpressionCache, const FUniformExpressionSet* UniformExpressionSet, const FRHIUniformBufferLayout* UniformBufferLayout, const FMaterialRenderContext& Context)
+	void Add(FUniformExpressionCache* UniformExpressionCache, const FUniformExpressionSet* UniformExpressionSet, const FRHIUniformBufferLayout* UniformBufferLayout, const FRHIUniformBufferLayout* InUniformBufferLayoutVS, const FMaterialRenderContext& Context)
 	{
-		Items.Emplace(UniformExpressionCache, UniformExpressionSet, UniformBufferLayout, Context);
+		Items.Emplace(UniformExpressionCache, UniformExpressionSet, UniformBufferLayout, InUniformBufferLayoutVS, Context);
 	}
 
 	void Update(FRHICommandListBase& RHICmdList)
@@ -144,6 +145,16 @@ public:
 						FMaterialRenderContext Context(Item.MaterialRenderProxy, *Item.Material, Item.bShowSelection);
 
 						Item.UniformExpressionSet->FillUniformBuffer(Context, Item.AllocatedVTs, Item.UniformBufferLayout, TempBuffer, Item.UniformBufferLayout->ConstantBufferSize);
+				
+						if (Item.UniformBufferVS)
+						{
+							check(Item.UniformBufferVS != Item.UniformBuffer);
+							check(Item.UniformBufferLayoutVS);
+							check(Item.UniformsVS);
+							uint8* UniformsVSTempBuffer = FMemStack::Get().PushBytes(Item.UniformBufferLayoutVS->ConstantBufferSize, SHADER_PARAMETER_STRUCT_ALIGNMENT);
+							FUniformExpressionSet::CompactUniformBuffer(UniformsVSTempBuffer, Item.UniformBufferLayoutVS->ConstantBufferSize, TempBuffer, Item.UniformBufferLayout->ConstantBufferSize, *Item.UniformsVS);
+							RHICmdList.UpdateUniformBuffer(Item.UniformBufferVS, UniformsVSTempBuffer);
+						}
 
 						RHICmdList.UpdateUniformBuffer(Item.UniformBuffer, TempBuffer);
 					}
@@ -160,22 +171,30 @@ private:
 	struct FItem
 	{
 		FItem() = default;
-		FItem(FUniformExpressionCache* InUniformExpressionCache, const FUniformExpressionSet* InUniformExpressionSet, const FRHIUniformBufferLayout* InUniformBufferLayout, const FMaterialRenderContext& Context)
+		FItem(FUniformExpressionCache* InUniformExpressionCache, const FUniformExpressionSet* InUniformExpressionSet, const FRHIUniformBufferLayout* InUniformBufferLayout, const FRHIUniformBufferLayout* InUniformBufferLayoutVS, const FMaterialRenderContext& Context)
 			: UniformBuffer(InUniformExpressionCache->UniformBuffer)
+			, UniformBufferVS(InUniformExpressionCache->UniformBufferVS)
 			, AllocatedVTs(InUniformExpressionCache->AllocatedVTs)
 			, UniformExpressionSet(InUniformExpressionSet)
 			, UniformBufferLayout(InUniformBufferLayout)
+			, UniformBufferLayoutVS(InUniformBufferLayoutVS)
 			, MaterialRenderProxy(Context.MaterialRenderProxy)
 			, Material(&Context.Material)
+			, UniformsVS( InUniformExpressionSet->GetCompactUniformsVS())
 			, bShowSelection(Context.bShowSelection)
-		{}
+		{
+			check((UniformBufferVS == nullptr) || (UniformBufferVS != UniformBuffer));
+		}
 
 		TRefCountPtr<FRHIUniformBuffer> UniformBuffer;
+		TRefCountPtr<FRHIUniformBuffer> UniformBufferVS;
 		TArray<IAllocatedVirtualTexture*, FConcurrentLinearArrayAllocator> AllocatedVTs;
 		const FUniformExpressionSet* UniformExpressionSet = nullptr;
 		const FRHIUniformBufferLayout* UniformBufferLayout = nullptr;
+		const FRHIUniformBufferLayout* UniformBufferLayoutVS = nullptr;
 		const FMaterialRenderProxy* MaterialRenderProxy = nullptr;
 		const FMaterial* Material = nullptr;
+		const FCompactUniformExpressionSet* UniformsVS = nullptr;
 		bool bShowSelection = false;
 	};
 
@@ -519,31 +538,41 @@ void FMaterialRenderProxy::EvaluateUniformExpressions(FRHICommandListBase& RHICm
 	}
 
 	const FRHIUniformBufferLayout* UniformBufferLayout = ShaderMap->GetUniformBufferLayout();
+	const FRHIUniformBufferLayout* UniformBufferLayoutVS = ShaderMap->GetUniformBufferLayoutVS();
 
-	if (IsValidRef(OutUniformExpressionCache.UniformBuffer))
+	auto ValidateLayout = [](FUniformBufferRHIRef& UniformBuffer, const FRHIUniformBufferLayout* UniformBufferLayout)
 	{
-		if (!OutUniformExpressionCache.UniformBuffer->IsValid())
+		if (IsValidRef(UniformBuffer))
 		{
-			UE_LOG(LogMaterial, Fatal, TEXT("The Uniformbuffer needs to be valid if it has been set"));
-		}
+			if (!UniformBuffer->IsValid())
+			{
+				UE_LOG(LogMaterial, Fatal, TEXT("The Uniformbuffer needs to be valid if it has been set"));
+			}
 
-		/**
-		* The actual pointer may not match because there are cases(in the editor, during the shader compilation) when material's shader map gets updated without proxy's cache
-		* getting invalidated, but the layout contents must match.
-		*/
+			/**
+				* The actual pointer may not match because there are cases(in the editor, during the shader compilation) when material's shader map gets updated without proxy's cache
+				* getting invalidated, but the layout contents must match.
+				*/
 #if WITH_EDITOR
-		/**
-		* If we are in the editor, this is likely occuring due to a mismatch of buffer invalidations and the type of shaders that need processing, so we can safely reset the
-		* Buffer layout to avoid any engine crashes. However we still want this check in project builds for safety when doing QA passes.
-		*/
-		if (OutUniformExpressionCache.UniformBuffer->GetLayoutPtr() != UniformBufferLayout || *OutUniformExpressionCache.UniformBuffer->GetLayoutPtr() != *UniformBufferLayout)
-		{
-			OutUniformExpressionCache.UniformBuffer = nullptr;
-		}
+			/**
+				* If we are in the editor, this is likely occuring due to a mismatch of buffer invalidations and the type of shaders that need processing, so we can safely reset the
+				* Buffer layout to avoid any engine crashes. However we still want this check in project builds for safety when doing QA passes.
+				*/
+			if (UniformBuffer->GetLayoutPtr() != UniformBufferLayout || *UniformBuffer->GetLayoutPtr() != *UniformBufferLayout)
+			{
+				UniformBuffer = nullptr;
+			}
 #else
-		check(OutUniformExpressionCache.UniformBuffer->GetLayoutPtr() == UniformBufferLayout || *OutUniformExpressionCache.UniformBuffer->GetLayoutPtr() == *UniformBufferLayout);
+			check(UniformBuffer->GetLayoutPtr() == UniformBufferLayout || *UniformBuffer->GetLayoutPtr() == *UniformBufferLayout);
 #endif
-	}
+		}	
+	};
+
+	ValidateLayout(OutUniformExpressionCache.UniformBuffer, UniformBufferLayout);
+	ValidateLayout(OutUniformExpressionCache.UniformBufferVS, UniformBufferLayoutVS);
+
+	const FCompactUniformExpressionSet* CompactUniformsVS = UniformExpressionSet.GetCompactUniformsVS();
+	bool bShouldHaveDedicatedMaterialUB = (CompactUniformsVS != nullptr && !CompactUniformsVS->CompactionSegments.IsEmpty());
 
 	if (Updater)
 	{
@@ -552,7 +581,19 @@ void FMaterialRenderProxy::EvaluateUniformExpressions(FRHICommandListBase& RHICm
 			OutUniformExpressionCache.UniformBuffer = RHICreateUniformBuffer(nullptr, UniformBufferLayout, UniformBuffer_MultiFrame);
 		}
 
-		Updater->Add(&OutUniformExpressionCache, &UniformExpressionSet, UniformBufferLayout, Context);
+		if (bShouldHaveDedicatedMaterialUB)
+		{
+			if (!IsValidRef(OutUniformExpressionCache.UniformBufferVS))
+			{
+				OutUniformExpressionCache.UniformBufferVS = RHICreateUniformBuffer(nullptr, UniformBufferLayoutVS, UniformBuffer_MultiFrame);
+			}
+		}
+		else
+		{
+			OutUniformExpressionCache.UniformBufferVS = nullptr;
+		}
+
+		Updater->Add(&OutUniformExpressionCache, &UniformExpressionSet, UniformBufferLayout, UniformBufferLayoutVS, Context);
 	}
 	else
 	{
@@ -560,13 +601,31 @@ void FMaterialRenderProxy::EvaluateUniformExpressions(FRHICommandListBase& RHICm
 		uint8* TempBuffer = FMemStack::Get().PushBytes(UniformBufferLayout->ConstantBufferSize, SHADER_PARAMETER_STRUCT_ALIGNMENT);
 		UniformExpressionSet.FillUniformBuffer(Context, OutUniformExpressionCache, UniformBufferLayout, TempBuffer, UniformBufferLayout->ConstantBufferSize);
 
-		if (IsValidRef(OutUniformExpressionCache.UniformBuffer))
+		auto CreateOrUpdateUniformBuffer = [](FRHICommandListBase& RHICmdList, FUniformBufferRHIRef& UniformBuffer, const uint8* TempBuffer, const FRHIUniformBufferLayout* UniformBufferLayout)
+		{
+			if (IsValidRef(UniformBuffer))
+			{
+				RHICmdList.UpdateUniformBuffer(UniformBuffer, TempBuffer);
+			}
+			else
+			{
+				UniformBuffer = RHICreateUniformBuffer(TempBuffer, UniformBufferLayout, UniformBuffer_MultiFrame);
+			}
+
+		};
+
+		CreateOrUpdateUniformBuffer(RHICmdList, OutUniformExpressionCache.UniformBuffer, TempBuffer, UniformBufferLayout);
+
+		if (bShouldHaveDedicatedMaterialUB)
 		{
-			RHICmdList.UpdateUniformBuffer(OutUniformExpressionCache.UniformBuffer, TempBuffer);
+			uint8* UniformsVSTempBuffer = FMemStack::Get().PushBytes(UniformBufferLayoutVS->ConstantBufferSize, SHADER_PARAMETER_STRUCT_ALIGNMENT);
+			FUniformExpressionSet::CompactUniformBuffer(UniformsVSTempBuffer, UniformBufferLayoutVS->ConstantBufferSize, TempBuffer, UniformBufferLayout->ConstantBufferSize, *CompactUniformsVS);
+				
+			CreateOrUpdateUniformBuffer(RHICmdList, OutUniformExpressionCache.UniformBufferVS, UniformsVSTempBuffer, UniformBufferLayoutVS);
 		}
 		else
 		{
-			OutUniformExpressionCache.UniformBuffer = RHICreateUniformBuffer(TempBuffer, UniformBufferLayout, UniformBuffer_MultiFrame);
+			OutUniformExpressionCache.UniformBufferVS = nullptr;
 		}
 	}
 
@@ -669,6 +728,7 @@ void FMaterialRenderProxy::InvalidateUniformExpressionCache(bool bRecreateUnifor
 			// This should only be done if the calling code is using FMaterialUpdateContext to recreate the rendering state of primitives using this material, 
 			// Since cached mesh commands also cache uniform buffer pointers.
 			UniformExpressionCache[i].UniformBuffer = nullptr;
+			UniformExpressionCache[i].UniformBufferVS = nullptr;
 		}
 	}
 }
diff --git a/Engine/Source/Runtime/Engine/Private/Materials/MaterialShader.cpp b/Engine/Source/Runtime/Engine/Private/Materials/MaterialShader.cpp
--- a/Engine/Source/Runtime/Engine/Private/Materials/MaterialShader.cpp
+++ b/Engine/Source/Runtime/Engine/Private/Materials/MaterialShader.cpp
@@ -1672,6 +1672,8 @@ void FMaterialShaderMapId::AppendKeyString(FString& KeyString, bool bIncludeSour
 	const_cast<FMaterialShaderMapId&>(*this).RecordAndEmit(Context);
 }
 
+extern bool GAllowSkipAlphaComputationForMaterialBlend;
+
 void FMaterialShaderMapId::RecordAndEmit(FMaterialKeyGeneratorContext& Context)
 {
 	check(Context.IsLoading() || IsContentValid());
@@ -1797,6 +1799,22 @@ void FMaterialShaderMapId::RecordAndEmit(FMaterialKeyGeneratorContext& Context)
 		FShaderKeyGenerator& KeyGen = Context.GetKeyGenIfEmittingChecked();
 		SubstrateCompilationConfig.Append(KeyGen);
 	}
+
+	// Changing these CVar should trigger a re-translate operation
+	{
+		uint32 MiscFlags = 0;
+		if (UseMaterialVSUniformBuffer())
+		{
+			MiscFlags |= 1 << 0;
+		}
+
+		if (GAllowSkipAlphaComputationForMaterialBlend)
+		{
+			MiscFlags |= 1 << 1;
+		}
+
+		Context.Emit(MiscFlags);
+	}
 }
 
 void FMaterialShaderMapId::SetShaderDependencies(const TArray<FShaderType*>& ShaderTypes, const TArray<const FShaderPipelineType*>& ShaderPipelineTypes, const TArray<FVertexFactoryType*>& VFTypes, EShaderPlatform ShaderPlatform)
@@ -4098,9 +4116,16 @@ void FMaterialShaderMap::InitalizeForODSC(EShaderPlatform TargetShaderPlatform,
 void FMaterialShaderMap::PostFinalizeContent()
 {
 	UniformBufferLayout.SafeRelease();
+	UniformBufferLayoutVS.SafeRelease();
+
 	if (GetContent())
 	{
 		UniformBufferLayout = RHICreateUniformBufferLayout(GetUniformExpressionSet().GetUniformBufferLayoutInitializer());
+		const FCompactUniformExpressionSet* CompactUniformsVS = GetUniformExpressionSet().GetCompactUniformsVS();
+		if (CompactUniformsVS)
+		{
+			UniformBufferLayoutVS = RHICreateUniformBufferLayout(CompactUniformsVS->UniformBufferLayoutInitializer);
+		}
 	}
 }
 
diff --git a/Engine/Source/Runtime/Engine/Private/Materials/MaterialShared.cpp b/Engine/Source/Runtime/Engine/Private/Materials/MaterialShared.cpp
--- a/Engine/Source/Runtime/Engine/Private/Materials/MaterialShared.cpp
+++ b/Engine/Source/Runtime/Engine/Private/Materials/MaterialShared.cpp
@@ -141,6 +141,8 @@ IMPLEMENT_TYPE_LAYOUT(FMaterialTextureCollectionParameterInfo);
 IMPLEMENT_TYPE_LAYOUT(FMaterialExternalTextureParameterInfo);
 IMPLEMENT_TYPE_LAYOUT(FMaterialVirtualTextureStack);
 IMPLEMENT_TYPE_LAYOUT(FMaterialCacheTagStack);
+IMPLEMENT_TYPE_LAYOUT(FCompactionMetaDataSegment);
+IMPLEMENT_TYPE_LAYOUT(FCompactUniformExpressionSet);
 
 struct FAllowCachingStaticParameterValues
 {
@@ -2554,13 +2556,16 @@ FMaterial::~FMaterial()
 
 /** Populates OutEnvironment with defines needed to compile shaders for this material. */
 void FMaterial::SetupMaterialEnvironment(
-	const FShaderParametersMetadata& InUniformBufferStruct,
+	TConstArrayView<FShaderParametersMetadata*> InUniformBufferStructs,
 	const FUniformExpressionSet& InUniformExpressionSet,
 	FShaderCompilerEnvironment& OutEnvironment
 	) const
 {
 	// Add the material uniform buffer definition.
-	FShaderUniformBufferParameter::ModifyCompilationEnvironment(TEXT("Material"), InUniformBufferStruct, GetShaderPlatform(), OutEnvironment);
+	for (const FShaderParametersMetadata* InUniformBufferStruct : InUniformBufferStructs)
+	{
+		FShaderUniformBufferParameter::ModifyCompilationEnvironment(InUniformBufferStruct->GetShaderVariableName(), *InUniformBufferStruct, GetShaderPlatform(), OutEnvironment);
+	}
 
 	if (GetMaterialDomain() == MD_Surface)
 	{
@@ -3583,9 +3588,12 @@ bool FMaterial::BeginCompileShaderMap(
 		FShaderCompileUtilities::ApplyDerivedDefines(*MaterialEnvironment, nullptr, (EShaderPlatform)GetShaderPlatform());
 
 		{
-			FShaderParametersMetadata* UniformBufferStruct = NewCompilationOutput.UniformExpressionSet.CreateBufferStruct();
-			SetupMaterialEnvironment(*UniformBufferStruct, NewCompilationOutput.UniformExpressionSet, *MaterialEnvironment);
-			delete UniformBufferStruct;
+			TArray<FShaderParametersMetadata*> UniformBufferStructs = NewCompilationOutput.UniformExpressionSet.CreateBufferStructs();
+			SetupMaterialEnvironment(UniformBufferStructs, NewCompilationOutput.UniformExpressionSet, *MaterialEnvironment);
+			for (FShaderParametersMetadata* UniformBufferStruct : UniformBufferStructs)
+			{
+				delete UniformBufferStruct;
+			}
 		}
 
 		// we can ignore requests for synch compilation if we are compiling for a different platform than we're running, or we're a commandlet that doesn't render (e.g. cooker)
@@ -4207,6 +4215,9 @@ FMaterialRenderContext::FMaterialRenderContext(
 FMaterialVirtualTextureStack::FMaterialVirtualTextureStack()
 	: NumLayers(0u)
 	, PreallocatedStackTextureIndex(INDEX_NONE)
+#if WITH_EDITORONLY_DATA
+	, ShaderFrequencyMask(0)
+#endif
 {
 	for (uint32 i = 0u; i < VIRTUALTEXTURE_SPACE_MAXLAYERS; ++i)
 	{
@@ -4217,6 +4228,9 @@ FMaterialVirtualTextureStack::FMaterialVirtualTextureStack()
 FMaterialVirtualTextureStack::FMaterialVirtualTextureStack(int32 InPreallocatedStackTextureIndex)
 	: NumLayers(0u)
 	, PreallocatedStackTextureIndex(InPreallocatedStackTextureIndex)
+#if WITH_EDITORONLY_DATA
+	, ShaderFrequencyMask(0)
+#endif
 {
 	for (uint32 i = 0u; i < VIRTUALTEXTURE_SPACE_MAXLAYERS; ++i)
 	{
diff --git a/Engine/Source/Runtime/Engine/Private/Materials/MaterialUniformExpressions.cpp b/Engine/Source/Runtime/Engine/Private/Materials/MaterialUniformExpressions.cpp
--- a/Engine/Source/Runtime/Engine/Private/Materials/MaterialUniformExpressions.cpp
+++ b/Engine/Source/Runtime/Engine/Private/Materials/MaterialUniformExpressions.cpp
@@ -28,7 +28,7 @@
 #include "MaterialCache/MaterialCacheTagSceneData.h"
 #include "Engine/VirtualTextureCollection.h"
 
-void WriteMaterialUniformAccess(UE::Shader::EValueComponentType ComponentType, uint32 NumComponents, uint32 UniformOffset, FStringBuilderBase& OutResult)
+void WriteMaterialUniformAccess(UE::Shader::EValueComponentType ComponentType, uint32 NumComponents, uint32 UniformOffset, const TCHAR* MaterialNamePrefix, FStringBuilderBase& OutResult)
 {
 	static const TCHAR IndexToMask[] = TEXT("xyzw");
 	uint32 RegisterIndex = UniformOffset / 4;
@@ -48,7 +48,7 @@ void WriteMaterialUniformAccess(UE::Shader::EValueComponentType ComponentType, u
 			bConstructor = true;
 		}
 
-		OutResult.Appendf(TEXT("Material.PreshaderBuffer[%u]"), RegisterIndex);
+		OutResult.Appendf(TEXT("%s.PreshaderBuffer[%u]"), MaterialNamePrefix, RegisterIndex);
 		// Can skip writing mask if we're taking all 4 components from the register
 		if (NumComponentsInRegister < 4u)
 		{
@@ -391,12 +391,190 @@ void FUniformExpressionSet::SetParameterCollections(TConstArrayView<const UMater
 	}
 }
 
-FShaderParametersMetadata* FUniformExpressionSet::CreateBufferStruct()
+template <class TMemberType>
+uint32 CountIndividualMembers(const TMemoryImageArray<TMemberType>& ParameterCollection, EShaderFrequencyMask ShaderFrequencyMask, TFunctionRef<bool(const TMemberType&)> Filter)
+{
+	uint32 NumElements = 0;
+	for (const TMemberType& ParameterCollectionMember : ParameterCollection)
+	{
+		if (!Filter(ParameterCollectionMember))
+		{
+			continue;
+		}
+
+		if (!!(ParameterCollectionMember.ShaderFrequencyMask & ShaderFrequencyMask))
+		{
+			++NumElements;
+		}
+	}
+
+	return NumElements;
+}
+
+template<class TMemberType>
+uint32 CountIndividualMembersSimple(const TMemoryImageArray<TMemberType>& ParameterCollection, EShaderFrequencyMask ShaderFrequencyMask)
+{
+	uint32 NumElements = 0;
+	for (const TMemberType& ParameterCollectionMember : ParameterCollection)
+	{
+		if (!!(ParameterCollectionMember.ShaderFrequencyMask & ShaderFrequencyMask))
+		{
+			++NumElements;
+		}
+	}
+
+	return NumElements;
+}
+
+void AddCopySegment(TArray<FCompactionMetaDataSegment>& CopySegments, 
+	const FShaderParametersMetadata::FMember& SrcMember,
+	const FShaderParametersMetadata::FMember& DestMember,
+	uint32 DestMemberSize)
+{
+	FCompactionMetaDataSegment PerFrequencyUniformMetadata;
+	PerFrequencyUniformMetadata.SrcOffset = SrcMember.GetOffset();
+	PerFrequencyUniformMetadata.DestOffset = DestMember.GetOffset();
+	PerFrequencyUniformMetadata.SizeToCopy = DestMemberSize;
+	CopySegments.Add(PerFrequencyUniformMetadata);
+}
+
+static FShaderParametersMetadata* CreateCompactShaderParametersMetadata(FCompactUniformExpressionSet* CompactUniformExpressions,
+	const TCHAR* UniformBufferLayoutInitializerName,
+	const TCHAR* LayoutName, const TCHAR* StructTypeName, const TCHAR* ShaderVariableName,
+	const TArray<FShaderParametersMetadata::FMember>& MembersVS, uint32 StructSize, const TArray<FCompactionMetaDataSegment>& CopySegments)
+{
+	if (CompactUniformExpressions == nullptr)
+	{
+		return nullptr;
+	}
+
+	CompactUniformExpressions->UniformBufferLayoutInitializer = FRHIUniformBufferLayoutInitializer(UniformBufferLayoutInitializerName);
+
+	FShaderParametersMetadata* NewUniformBufferStruct = new FShaderParametersMetadata(
+		FShaderParametersMetadata::EUseCase::DataDrivenUniformBuffer,
+		EUniformBufferBindingFlags::Shader,
+		LayoutName,
+		StructTypeName,
+		ShaderVariableName,
+		nullptr,
+		UE_LOG_SOURCE_FILE(__FILE__),
+		__LINE__,
+		StructSize,
+		MembersVS,
+		false,
+		&CompactUniformExpressions->UniformBufferLayoutInitializer);
+
+	// Sort by SrcOffset because the destination buffer is compact
+	Algo::SortBy(CopySegments, &FCompactionMetaDataSegment::SrcOffset);
+
+	// CopySegments contains the mapping for every member between the canonical UB and the compact one
+	// Try to merge copy sections that are contiguous in the canonical UB
+	FCompactionMetaDataSegment* FinalCompactionMetaData = &CompactUniformExpressions->CompactionSegments.AddZeroed_GetRef();
+	for (const FCompactionMetaDataSegment& CopySegment : CopySegments)
+	{
+		if (FinalCompactionMetaData->SizeToCopy != 0)
+		{
+			uint64 ExpectedSrcOffset = FinalCompactionMetaData->SrcOffset + FinalCompactionMetaData->SizeToCopy;
+			// Check if there's a hole between the current aggregated segment, and the new property
+			if (ExpectedSrcOffset == CopySegment.SrcOffset)
+			{
+				// Since everything is contiguous, we can add the current property to the aggregated segment
+				// Make sure that the destination offsets align. Holes are only supposed to exist in the canonical UB
+				check(CopySegment.DestOffset == FinalCompactionMetaData->DestOffset + FinalCompactionMetaData->SizeToCopy);
+				FinalCompactionMetaData->SizeToCopy += CopySegment.SizeToCopy;
+			}
+			else
+			{
+				// If we have a hole, start a new segment that we will fill below
+				FinalCompactionMetaData = &CompactUniformExpressions->CompactionSegments.AddZeroed_GetRef();
+			}
+		}
+
+		// new segment initialized from the current property's data. Subsequent members will try to merge with it
+		if (FinalCompactionMetaData->SizeToCopy == 0)
+		{
+			FinalCompactionMetaData->SizeToCopy = CopySegment.SizeToCopy;
+			FinalCompactionMetaData->SrcOffset = CopySegment.SrcOffset;
+			FinalCompactionMetaData->DestOffset = CopySegment.DestOffset;
+		}
+	}
+
+	return NewUniformBufferStruct;
+
+}
+
+#if WITH_EDITOR
+static bool IsNumericMember(const FShaderParametersMetadata::FMember& Member)
+{
+	switch (Member.GetBaseType())
+	{
+		case UBMT_INT32:
+		case UBMT_UINT32:
+		case UBMT_FLOAT32:
+			return true;
+		default:
+			return false;
+	}
+
+}
+
+TArray<FShaderParametersMetadata*> FUniformExpressionSet::CreateBufferStructs()
 {
 	// Make sure FUniformExpressionSet::CreateDebugLayout() is in sync
 	TArray<FShaderParametersMetadata::FMember> Members;
 	uint32 NextMemberOffset = 0;
+
+	TArray<FShaderParametersMetadata::FMember> MembersVS;
+	uint32 NextMemberOffsetVS = 0;
 	
+	TArray<FCompactionMetaDataSegment> CopySegments;
+
+	EShaderFrequencyMask ShaderFrequencyMaskVS = 1;
+
+	bool bAllowNumericsInCompactVS = true;
+
+	FCompactUniformExpressionSet* CompactUniformsVS = GetCompactUniformsVS();
+	bool bAddVSMembers = (CompactUniformsVS != nullptr);
+
+	auto ExecuteAddToMembersVS =
+	[&Members, &MembersVS, &bAllowNumericsInCompactVS, &NextMemberOffsetVS, &NextMemberOffset, &CopySegments, bAddVSMembers]
+		(uint32 DestNumMembers, uint32 DestCPUSize)
+	{
+		if (!bAddVSMembers)
+		{
+			return;
+		}
+
+		const FShaderParametersMetadata::FMember& SrcMember = Members.Last();
+
+		if (!bAllowNumericsInCompactVS)
+		{
+			check(!IsNumericMember(SrcMember));
+		}
+
+		new(MembersVS) FShaderParametersMetadata::FMember(SrcMember.GetName(), SrcMember.GetShaderType(), SrcMember.GetFileLine(), NextMemberOffsetVS, SrcMember.GetBaseType(), SrcMember.GetPrecision(),  SrcMember.GetNumRows(), SrcMember.GetNumColumns(), DestNumMembers, SrcMember.GetStructMetadata());
+		NextMemberOffsetVS += DestCPUSize;
+			
+		uint32 SrcSize = NextMemberOffset - SrcMember.GetOffset();
+		check(DestCPUSize <= SrcSize);
+		AddCopySegment(CopySegments, SrcMember, MembersVS.Last(), DestCPUSize);	
+	};
+
+	// As soon as one element of the array has the VS frequency, add the whole array: This makes indexing consistent between the 2 UBs, and HLSL generation easier
+	// Preshaders are the exception, since they do change quite a lot and impact VS often. The rest is much rarer
+	auto TryAddToMembersVS =
+	[&Members, &MembersVS, &NextMemberOffsetVS, &NextMemberOffset, ShaderFrequencyMaskVS, &ExecuteAddToMembersVS]
+		<class TMemberType>
+		(const TMemoryImageArray<TMemberType>& Container)
+	{
+		uint32 NumMembers = CountIndividualMembersSimple(Container, ShaderFrequencyMaskVS);
+		if (NumMembers > 0)
+		{
+			uint32 SrcSize = NextMemberOffset - Members.Last().GetOffset();
+			ExecuteAddToMembersVS(Members.Last().GetNumElements(), SrcSize);
+		}
+	};
+
 	if (!UniformTextureCollectionParameters.IsEmpty())
 	{
 		uint32 BindlessCollectionCount = 0;
@@ -408,7 +586,15 @@ FShaderParametersMetadata* FUniformExpressionSet::CreateBufferStruct()
 		{
 			static_assert(sizeof(FUintVector4) * UE::HLSL::IndirectVirtualTextureUniformDWord4Count == sizeof(UE::HLSL::FIndirectVirtualTextureUniform), "Unexpected size");
 			new(Members) FShaderParametersMetadata::FMember(TEXT("TextureCollectionPackedUniforms"), TEXT("FIndirectVirtualTextureUniform"), __LINE__, NextMemberOffset, UBMT_UINT32, EShaderPrecisionModifier::Float, 1, 4, UE::HLSL::IndirectVirtualTextureUniformDWord4Count * VirtualCollectionCount, NULL);
-			NextMemberOffset += sizeof(UE::HLSL::FIndirectVirtualTextureUniform) * UniformTextureCollectionParameters.Num();
+			NextMemberOffset += sizeof(UE::HLSL::FIndirectVirtualTextureUniform) * VirtualCollectionCount;
+
+			TFunctionRef<bool(const FMaterialTextureCollectionParameterInfo &)> AcceptElement([](const auto& A) { return A.bIsVirtualCollection; });
+			uint32 NumMembers = CountIndividualMembers(UniformTextureCollectionParameters, ShaderFrequencyMaskVS, AcceptElement);
+			if (NumMembers > 0)
+			{
+				uint32 SrcSize = NextMemberOffset - Members.Last().GetOffset();
+				ExecuteAddToMembersVS(Members.Last().GetNumElements(), SrcSize);
+			}
 		}
 	}
 	
@@ -417,12 +603,14 @@ FShaderParametersMetadata* FUniformExpressionSet::CreateBufferStruct()
 		// 2x uint4 per VTStack
 		new(Members) FShaderParametersMetadata::FMember(TEXT("VTPackedPageTableUniform"),TEXT(""),__LINE__,NextMemberOffset, UBMT_UINT32, EShaderPrecisionModifier::Float, 1, 4, VTStacks.Num() * 2, NULL);
 		NextMemberOffset += VTStacks.Num() * sizeof(FUintVector4) * 2;
+		TryAddToMembersVS(VTStacks);
 	}
 
 	if (MaterialCacheTagStacks.Num())
 	{
 		new(Members) FShaderParametersMetadata::FMember(TEXT("MaterialCacheTagPageTableUniform"),TEXT(""), __LINE__, NextMemberOffset, UBMT_UINT32, EShaderPrecisionModifier::Float, 1, 4, MaterialCacheTagStacks.Num(), NULL);
 		NextMemberOffset += MaterialCacheTagStacks.Num() * sizeof(FUintVector4);
+		TryAddToMembersVS(MaterialCacheTagStacks);
 	}
 
 	const int32 NumVirtualTextures = UniformTextureParameters[(uint32)EMaterialTextureParameterType::Virtual].Num();
@@ -431,6 +619,7 @@ FShaderParametersMetadata* FUniformExpressionSet::CreateBufferStruct()
 		// 1x uint4 per Virtual Texture
 		new(Members) FShaderParametersMetadata::FMember(TEXT("VTPackedUniform"), TEXT(""),__LINE__,NextMemberOffset, UBMT_UINT32, EShaderPrecisionModifier::Float, 1, 4, NumVirtualTextures, NULL);
 		NextMemberOffset += NumVirtualTextures * sizeof(FUintVector4);
+		TryAddToMembersVS(UniformTextureParameters[(uint32)EMaterialTextureParameterType::Virtual]);
 	}
 
 	const int32 NumSparseVolumeTextures = UniformTextureParameters[(uint32)EMaterialTextureParameterType::SparseVolume].Num();
@@ -439,15 +628,23 @@ FShaderParametersMetadata* FUniformExpressionSet::CreateBufferStruct()
 		// 2x uint4 per SVT
 		new(Members) FShaderParametersMetadata::FMember(TEXT("SVTPackedUniform"), TEXT(""), __LINE__, NextMemberOffset, UBMT_UINT32, EShaderPrecisionModifier::Float, 1, 4, NumSparseVolumeTextures * 2, NULL);
 		NextMemberOffset += NumSparseVolumeTextures * sizeof(FUintVector4) * 2;
+		TryAddToMembersVS(UniformTextureParameters[(uint32)EMaterialTextureParameterType::SparseVolume]);
 	}
 
 	if (UniformPreshaderBufferSize > 0u)
 	{
 		new(Members) FShaderParametersMetadata::FMember(TEXT("PreshaderBuffer"),TEXT(""),__LINE__,NextMemberOffset,UBMT_FLOAT32,EShaderPrecisionModifier::Float,1,4, UniformPreshaderBufferSize,NULL);
 		NextMemberOffset += UniformPreshaderBufferSize * sizeof(FVector4f);
+		if (CompactUniformsVS && CompactUniformsVS->UniformPreshaderBufferSize > 0)
+		{
+			ExecuteAddToMembersVS(CompactUniformsVS->UniformPreshaderBufferSize, CompactUniformsVS->UniformPreshaderBufferSize * sizeof(FVector4f));
+		}
 	}
 
+	bAllowNumericsInCompactVS = false;
+
 	check((NextMemberOffset % (2 * SHADER_PARAMETER_POINTER_ALIGNMENT)) == 0);
+	check((NextMemberOffsetVS % (2 * SHADER_PARAMETER_POINTER_ALIGNMENT)) == 0);
 
 	const uint32 MaxUniformExpressions = 128;
 
@@ -528,40 +725,70 @@ FShaderParametersMetadata* FUniformExpressionSet::CreateBufferStruct()
 	check(VTStacks.Num() <= MaxUniformExpressions);
 	check(MaterialCacheTagStacks.Num() <= MaxUniformExpressions);
 
+	auto AddMemberIfNeeded = [&Members, &MembersVS, &NextMemberOffsetVS, &NextMemberOffset, ShaderFrequencyMaskVS, &ExecuteAddToMembersVS]
+		(EShaderFrequencyMask ShaderFrequencyMask)
+	{
+		if (!!(ShaderFrequencyMask & ShaderFrequencyMaskVS))
+		{
+			check(Members.Last().GetNumElements() == 0);
+			ExecuteAddToMembersVS(Members.Last().GetNumElements(), NextMemberOffset - Members.Last().GetOffset());
+		}
+
+	};
+
+	auto AddTextureMember = [this, &Members, &MembersVS, &NextMemberOffsetVS, ShaderFrequencyMaskVS, &AddMemberIfNeeded]
+		(int32 i, EMaterialTextureParameterType MaterialTextureParameterType)
+	{
+		const FMaterialTextureParameterInfo& MaterialTextureParameterInfo = UniformTextureParameters[(uint32)MaterialTextureParameterType][i];
+		AddMemberIfNeeded(MaterialTextureParameterInfo.ShaderFrequencyMask);
+	};
+
 	for (int32 i = 0; i < UniformTextureParameters[(uint32)EMaterialTextureParameterType::Standard2D].Num(); ++i)
 	{
 		check((NextMemberOffset % SHADER_PARAMETER_POINTER_ALIGNMENT) == 0);
+		check((NextMemberOffsetVS % SHADER_PARAMETER_POINTER_ALIGNMENT) == 0);
 		new(Members) FShaderParametersMetadata::FMember(*Texture2DNames[i],TEXT("Texture2D"),__LINE__,NextMemberOffset,UBMT_TEXTURE,EShaderPrecisionModifier::Float,1,1,0,NULL);
 		NextMemberOffset += SHADER_PARAMETER_POINTER_ALIGNMENT;
+		AddTextureMember(i, EMaterialTextureParameterType::Standard2D);
 		new(Members) FShaderParametersMetadata::FMember(*Texture2DSamplerNames[i],TEXT("SamplerState"),__LINE__,NextMemberOffset,UBMT_SAMPLER,EShaderPrecisionModifier::Float,1,1,0,NULL);
 		NextMemberOffset += SHADER_PARAMETER_POINTER_ALIGNMENT;
+		AddTextureMember(i, EMaterialTextureParameterType::Standard2D);
 	}
 
 	for (int32 i = 0; i < UniformTextureParameters[(uint32)EMaterialTextureParameterType::Cube].Num(); ++i)
 	{
 		check((NextMemberOffset % SHADER_PARAMETER_POINTER_ALIGNMENT) == 0);
+		check((NextMemberOffsetVS % SHADER_PARAMETER_POINTER_ALIGNMENT) == 0);
 		new(Members) FShaderParametersMetadata::FMember(*TextureCubeNames[i],TEXT("TextureCube"),__LINE__,NextMemberOffset,UBMT_TEXTURE,EShaderPrecisionModifier::Float,1,1,0,NULL);
 		NextMemberOffset += SHADER_PARAMETER_POINTER_ALIGNMENT;
+		AddTextureMember(i, EMaterialTextureParameterType::Cube);
 		new(Members) FShaderParametersMetadata::FMember(*TextureCubeSamplerNames[i],TEXT("SamplerState"),__LINE__,NextMemberOffset,UBMT_SAMPLER,EShaderPrecisionModifier::Float,1,1,0,NULL);
 		NextMemberOffset += SHADER_PARAMETER_POINTER_ALIGNMENT;
+		AddTextureMember(i, EMaterialTextureParameterType::Cube);
 	}
 
 	for (int32 i = 0; i < UniformTextureParameters[(uint32)EMaterialTextureParameterType::Array2D].Num(); ++i)
 	{
 		check((NextMemberOffset % SHADER_PARAMETER_POINTER_ALIGNMENT) == 0);
+		check((NextMemberOffsetVS % SHADER_PARAMETER_POINTER_ALIGNMENT) == 0);
 		new(Members) FShaderParametersMetadata::FMember(*Texture2DArrayNames[i], TEXT("Texture2DArray"), __LINE__, NextMemberOffset, UBMT_TEXTURE, EShaderPrecisionModifier::Float, 1, 1, 0, NULL);
 		NextMemberOffset += SHADER_PARAMETER_POINTER_ALIGNMENT;
+		AddTextureMember(i, EMaterialTextureParameterType::Array2D);
 		new(Members) FShaderParametersMetadata::FMember(*Texture2DArraySamplerNames[i], TEXT("SamplerState"), __LINE__, NextMemberOffset, UBMT_SAMPLER, EShaderPrecisionModifier::Float, 1, 1, 0, NULL);
 		NextMemberOffset += SHADER_PARAMETER_POINTER_ALIGNMENT;
+		AddTextureMember(i, EMaterialTextureParameterType::Array2D);
 	}
 
 	for (int32 i = 0; i < UniformTextureParameters[(uint32)EMaterialTextureParameterType::ArrayCube].Num(); ++i)
 	{
 		check((NextMemberOffset % SHADER_PARAMETER_POINTER_ALIGNMENT) == 0);
+		check((NextMemberOffsetVS % SHADER_PARAMETER_POINTER_ALIGNMENT) == 0);
 		new(Members) FShaderParametersMetadata::FMember(*TextureCubeArrayNames[i], TEXT("TextureCubeArray"), __LINE__, NextMemberOffset, UBMT_TEXTURE, EShaderPrecisionModifier::Float, 1, 1, 0, NULL);
 		NextMemberOffset += SHADER_PARAMETER_POINTER_ALIGNMENT;
+		AddTextureMember(i, EMaterialTextureParameterType::ArrayCube);
 		new(Members) FShaderParametersMetadata::FMember(*TextureCubeArraySamplerNames[i], TEXT("SamplerState"), __LINE__, NextMemberOffset, UBMT_SAMPLER, EShaderPrecisionModifier::Float, 1, 1, 0, NULL);
 		NextMemberOffset += SHADER_PARAMETER_POINTER_ALIGNMENT;
+		AddTextureMember(i, EMaterialTextureParameterType::ArrayCube);
 	}
 
 	for (int32 i = 0; i < UniformTextureParameters[(uint32)EMaterialTextureParameterType::Volume].Num(); ++i)
@@ -569,10 +796,13 @@ FShaderParametersMetadata* FUniformExpressionSet::CreateBufferStruct()
 		const FMaterialTextureParameterInfo& Parameter = UniformTextureParameters[(uint32)EMaterialTextureParameterType::Volume][i];
 		
 		check((NextMemberOffset % SHADER_PARAMETER_POINTER_ALIGNMENT) == 0);
+		check((NextMemberOffsetVS % SHADER_PARAMETER_POINTER_ALIGNMENT) == 0);
 		new(Members) FShaderParametersMetadata::FMember(*VolumeTextureNames[i], TEXT("Texture3D"), __LINE__, NextMemberOffset, UBMT_TEXTURE, EShaderPrecisionModifier::Float, 1, 1, 0, NULL);
 		NextMemberOffset += SHADER_PARAMETER_POINTER_ALIGNMENT;
+		AddTextureMember(i, EMaterialTextureParameterType::Volume);
 		new(Members) FShaderParametersMetadata::FMember(*VolumeTextureSamplerNames[i], TEXT("SamplerState"), __LINE__, NextMemberOffset, UBMT_SAMPLER, EShaderPrecisionModifier::Float, 1, 1, 0, NULL);
 		NextMemberOffset += SHADER_PARAMETER_POINTER_ALIGNMENT;
+		AddTextureMember(i, EMaterialTextureParameterType::Volume);
 	}
 
 	for (int32 i = 0; i < UniformTextureParameters[(uint32)EMaterialTextureParameterType::SparseVolume].Num(); ++i)
@@ -581,57 +811,75 @@ FShaderParametersMetadata* FUniformExpressionSet::CreateBufferStruct()
 
 		// Page Table
 		check((NextMemberOffset % SHADER_PARAMETER_POINTER_ALIGNMENT) == 0);
+		check((NextMemberOffsetVS % SHADER_PARAMETER_POINTER_ALIGNMENT) == 0);
 		new(Members) FShaderParametersMetadata::FMember(*SparseVolumeTexturePageTableNames[i], TEXT("Texture3D<uint>"), __LINE__, NextMemberOffset, UBMT_TEXTURE, EShaderPrecisionModifier::Float, 1, 1, 0, NULL);
 		NextMemberOffset += SHADER_PARAMETER_POINTER_ALIGNMENT;
+		AddTextureMember(i, EMaterialTextureParameterType::SparseVolume);
 		
 		// Physical A
 		check((NextMemberOffset% SHADER_PARAMETER_POINTER_ALIGNMENT) == 0);
+		check((NextMemberOffsetVS % SHADER_PARAMETER_POINTER_ALIGNMENT) == 0);
 		new(Members) FShaderParametersMetadata::FMember(*SparseVolumeTexturePhysicalANames[i], TEXT("Texture3D"), __LINE__, NextMemberOffset, UBMT_TEXTURE, EShaderPrecisionModifier::Float, 1, 1, 0, NULL);
 		NextMemberOffset += SHADER_PARAMETER_POINTER_ALIGNMENT;
+		AddTextureMember(i, EMaterialTextureParameterType::SparseVolume);
 
 		// Physical B
 		check((NextMemberOffset% SHADER_PARAMETER_POINTER_ALIGNMENT) == 0);
+		check((NextMemberOffsetVS % SHADER_PARAMETER_POINTER_ALIGNMENT) == 0);
 		new(Members) FShaderParametersMetadata::FMember(*SparseVolumeTexturePhysicalBNames[i], TEXT("Texture3D"), __LINE__, NextMemberOffset, UBMT_TEXTURE, EShaderPrecisionModifier::Float, 1, 1, 0, NULL);
 		NextMemberOffset += SHADER_PARAMETER_POINTER_ALIGNMENT;
+		AddTextureMember(i, EMaterialTextureParameterType::SparseVolume);
 
 		// Sampler
 		check((NextMemberOffset% SHADER_PARAMETER_POINTER_ALIGNMENT) == 0);
+		check((NextMemberOffsetVS % SHADER_PARAMETER_POINTER_ALIGNMENT) == 0);
 		new(Members) FShaderParametersMetadata::FMember(*SparseVolumeTexturePhysicalSamplerNames[i], TEXT("SamplerState"), __LINE__, NextMemberOffset, UBMT_SAMPLER, EShaderPrecisionModifier::Float, 1, 1, 0, NULL);
 		NextMemberOffset += SHADER_PARAMETER_POINTER_ALIGNMENT;
+		AddTextureMember(i, EMaterialTextureParameterType::SparseVolume);
 	}
 
 	for (int32 i = 0; i < UniformExternalTextureParameters.Num(); ++i)
 	{
 		check((NextMemberOffset % SHADER_PARAMETER_POINTER_ALIGNMENT) == 0);
+		check((NextMemberOffsetVS % SHADER_PARAMETER_POINTER_ALIGNMENT) == 0);
 		new(Members) FShaderParametersMetadata::FMember(*ExternalTextureNames[i], TEXT("TextureExternal"), __LINE__, NextMemberOffset, UBMT_TEXTURE, EShaderPrecisionModifier::Float, 1, 1, 0, NULL);
 		NextMemberOffset += SHADER_PARAMETER_POINTER_ALIGNMENT;
+		AddMemberIfNeeded(UniformExternalTextureParameters[i].ShaderFrequencyMask);
+		
 		new(Members) FShaderParametersMetadata::FMember(*MediaTextureSamplerNames[i], TEXT("SamplerState"), __LINE__, NextMemberOffset, UBMT_SAMPLER, EShaderPrecisionModifier::Float, 1, 1, 0, NULL);
 		NextMemberOffset += SHADER_PARAMETER_POINTER_ALIGNMENT;
+		AddMemberIfNeeded(UniformExternalTextureParameters[i].ShaderFrequencyMask);
 	}
 
 	for (int32 i = 0; i < UniformTextureCollectionParameters.Num(); ++i)
 	{
 		check((NextMemberOffset % SHADER_PARAMETER_POINTER_ALIGNMENT) == 0);
+		check((NextMemberOffsetVS % SHADER_PARAMETER_POINTER_ALIGNMENT) == 0);
 
 		// Resource collections are RHI constructs, texture collections are software virtual tables
 		if (!UniformTextureCollectionParameters[i].bIsVirtualCollection)
 		{
 			new(Members) FShaderParametersMetadata::FMember(*TextureCollectionNames[i], TEXT("FResourceCollection"), __LINE__, NextMemberOffset, UBMT_RESOURCE_COLLECTION, EShaderPrecisionModifier::Float, 1, 1, 0, NULL);
 			NextMemberOffset += SHADER_PARAMETER_POINTER_ALIGNMENT;
+			AddMemberIfNeeded(UniformTextureCollectionParameters[i].ShaderFrequencyMask);
 		}
 		else
 		{
 			new(Members) FShaderParametersMetadata::FMember(*TextureCollectionNames[i], TEXT("FIndirectVirtualTexture"), __LINE__, NextMemberOffset, UBMT_SRV, EShaderPrecisionModifier::Float, 1, 1, 0, NULL);
 			NextMemberOffset += SHADER_PARAMETER_POINTER_ALIGNMENT;
+			AddMemberIfNeeded(UniformTextureCollectionParameters[i].ShaderFrequencyMask);
 			
 			new(Members) FShaderParametersMetadata::FMember(*TextureCollectionPageTableNames[i], TEXT("Texture2D<uint4>"), __LINE__, NextMemberOffset, UBMT_TEXTURE, EShaderPrecisionModifier::Float, 1, 1, 0, NULL);
 			NextMemberOffset += SHADER_PARAMETER_POINTER_ALIGNMENT;
+			AddMemberIfNeeded(UniformTextureCollectionParameters[i].ShaderFrequencyMask);
 		
 			new(Members) FShaderParametersMetadata::FMember(*TextureCollectionPhysicalNames[i], TEXT("Texture2D"), __LINE__, NextMemberOffset, UBMT_SRV, EShaderPrecisionModifier::Float, 1, 1, 0, NULL);
 			NextMemberOffset += SHADER_PARAMETER_POINTER_ALIGNMENT;
+			AddMemberIfNeeded(UniformTextureCollectionParameters[i].ShaderFrequencyMask);
 	
 			new(Members) FShaderParametersMetadata::FMember(*TextureCollectionSamplerNames[i], TEXT("SamplerState"), __LINE__, NextMemberOffset, UBMT_SAMPLER, EShaderPrecisionModifier::Float, 1, 1, 0, NULL);
 			NextMemberOffset += SHADER_PARAMETER_POINTER_ALIGNMENT;
+			AddMemberIfNeeded(UniformTextureCollectionParameters[i].ShaderFrequencyMask);
 		}
 	}
 
@@ -639,52 +887,79 @@ FShaderParametersMetadata* FUniformExpressionSet::CreateBufferStruct()
 	{
 		const FMaterialVirtualTextureStack& Stack = VTStacks[i];
 		check((NextMemberOffset % SHADER_PARAMETER_POINTER_ALIGNMENT) == 0);
+		check((NextMemberOffsetVS % SHADER_PARAMETER_POINTER_ALIGNMENT) == 0);
 		new(Members) FShaderParametersMetadata::FMember(*VirtualTexturePageTableNames0[i], TEXT("Texture2D<uint4>"), __LINE__, NextMemberOffset, UBMT_TEXTURE, EShaderPrecisionModifier::Float, 1, 1, 0, NULL);
 		NextMemberOffset += SHADER_PARAMETER_POINTER_ALIGNMENT;
+		AddMemberIfNeeded(Stack.ShaderFrequencyMask);
 		if (Stack.GetNumLayers() > 4u)
 		{
 			new(Members) FShaderParametersMetadata::FMember(*VirtualTexturePageTableNames1[i], TEXT("Texture2D<uint4>"), __LINE__, NextMemberOffset, UBMT_TEXTURE, EShaderPrecisionModifier::Float, 1, 1, 0, NULL);
 			NextMemberOffset += SHADER_PARAMETER_POINTER_ALIGNMENT;
+			AddMemberIfNeeded(Stack.ShaderFrequencyMask);
 		}
 		new(Members) FShaderParametersMetadata::FMember(*VirtualTexturePageTableIndirectionNames[i], TEXT("Texture2D<uint>"), __LINE__, NextMemberOffset, UBMT_TEXTURE, EShaderPrecisionModifier::Float, 1, 1, 0, NULL);
 		NextMemberOffset += SHADER_PARAMETER_POINTER_ALIGNMENT;
+		AddMemberIfNeeded(Stack.ShaderFrequencyMask);
 	}
 
 	for (int32 TagIndex = 0; TagIndex < MaterialCacheTagStacks.Num(); ++TagIndex)
 	{
 		check((NextMemberOffset % SHADER_PARAMETER_POINTER_ALIGNMENT) == 0);
-		
+		check((NextMemberOffsetVS % SHADER_PARAMETER_POINTER_ALIGNMENT) == 0);
+
 		new(Members) FShaderParametersMetadata::FMember(*MaterialCacheTagBuffer[TagIndex], TEXT("StructuredBuffer<FMaterialCacheTagEntry>"), __LINE__, NextMemberOffset, UBMT_SRV, EShaderPrecisionModifier::Float, 1, 1, 0, NULL);
 		NextMemberOffset += SHADER_PARAMETER_POINTER_ALIGNMENT;
+		AddMemberIfNeeded(MaterialCacheTagStacks[TagIndex].ShaderFrequencyMask);
 		
 		new(Members) FShaderParametersMetadata::FMember(*MaterialCacheTagPageTableNames0[TagIndex], TEXT("Texture2D<uint4>"), __LINE__, NextMemberOffset, UBMT_TEXTURE, EShaderPrecisionModifier::Float, 1, 1, 0, NULL);
 		NextMemberOffset += SHADER_PARAMETER_POINTER_ALIGNMENT;
+		AddMemberIfNeeded(MaterialCacheTagStacks[TagIndex].ShaderFrequencyMask);
 
 		for (int32 LayerIndex = 0; LayerIndex < MaterialCacheMaxRuntimeLayers; LayerIndex++)
 		{
 			new(Members) FShaderParametersMetadata::FMember(*MaterialCacheTagPhysicalNames[TagIndex * MaterialCacheMaxRuntimeLayers + LayerIndex], TEXT("Texture2D"), __LINE__, NextMemberOffset, UBMT_SRV, EShaderPrecisionModifier::Float, 1, 1, 0, NULL);
 			NextMemberOffset += SHADER_PARAMETER_POINTER_ALIGNMENT;
+			AddMemberIfNeeded(MaterialCacheTagStacks[TagIndex].ShaderFrequencyMask);
 		}
 	}
 
 	for (int32 i = 0; i < UniformTextureParameters[(uint32)EMaterialTextureParameterType::Virtual].Num(); ++i)
 	{
 		check((NextMemberOffset % SHADER_PARAMETER_POINTER_ALIGNMENT) == 0);
+		check((NextMemberOffsetVS % SHADER_PARAMETER_POINTER_ALIGNMENT) == 0);
 
 		// VT physical textures are bound as SRV, allows aliasing the same underlying texture with both sRGB/non-sRGB views
 		new(Members) FShaderParametersMetadata::FMember(*VirtualTexturePhysicalNames[i], TEXT("Texture2D"), __LINE__, NextMemberOffset, UBMT_SRV, EShaderPrecisionModifier::Float, 1, 1, 0, NULL);
 		NextMemberOffset += SHADER_PARAMETER_POINTER_ALIGNMENT;
+		AddTextureMember(i, EMaterialTextureParameterType::Virtual);
 		new(Members) FShaderParametersMetadata::FMember(*VirtualTexturePhysicalSamplerNames[i], TEXT("SamplerState"), __LINE__, NextMemberOffset, UBMT_SAMPLER, EShaderPrecisionModifier::Float, 1, 1, 0, NULL);
 		NextMemberOffset += SHADER_PARAMETER_POINTER_ALIGNMENT;
+		AddTextureMember(i, EMaterialTextureParameterType::Virtual);
 	}
 
+	// Check if we need any texture at the VS level, and include the default samplers if so
+	bool bRequiresWorldGroupSettingsSamplers = false;
+	for (const FShaderParametersMetadata::FMember& MemberVS : MembersVS)
+	{
+		if (!IsNumericMember(MemberVS))
+		{
+			bRequiresWorldGroupSettingsSamplers = true;
+			break;
+		}
+	}
+
+	EShaderFrequencyMask ShaderFrequencyMaskWorldGroupSettings = bRequiresWorldGroupSettingsSamplers ? ShaderFrequencyMaskVS : 0;
+
 	new(Members) FShaderParametersMetadata::FMember(TEXT("Wrap_WorldGroupSettings"),TEXT("SamplerState"), __LINE__, NextMemberOffset,UBMT_SAMPLER,EShaderPrecisionModifier::Float,1,1,0,NULL);
 	NextMemberOffset += SHADER_PARAMETER_POINTER_ALIGNMENT;
+	AddMemberIfNeeded(ShaderFrequencyMaskWorldGroupSettings);
 
 	new(Members) FShaderParametersMetadata::FMember(TEXT("Clamp_WorldGroupSettings"),TEXT("SamplerState"), __LINE__, NextMemberOffset,UBMT_SAMPLER,EShaderPrecisionModifier::Float,1,1,0,NULL);
 	NextMemberOffset += SHADER_PARAMETER_POINTER_ALIGNMENT;
+	AddMemberIfNeeded(ShaderFrequencyMaskWorldGroupSettings);
 
 	const uint32 StructSize = Align(NextMemberOffset, SHADER_PARAMETER_STRUCT_ALIGNMENT);
+	const uint32 StructSizeVS = Align(NextMemberOffsetVS, SHADER_PARAMETER_STRUCT_ALIGNMENT);
 
 	UniformBufferLayoutInitializer = FRHIUniformBufferLayoutInitializer(TEXT("Material"));
 
@@ -702,8 +977,33 @@ FShaderParametersMetadata* FUniformExpressionSet::CreateBufferStruct()
 		false,
 		&UniformBufferLayoutInitializer);
 
-	return UniformBufferStruct;
+	TArray<FShaderParametersMetadata*> AllUniformBufferStructs;
+	AllUniformBufferStructs.Add(UniformBufferStruct);
+
+	if (StructSizeVS > 0)
+	{
+		check(CompactUniformsVS != nullptr)
+		check(!MembersVS.IsEmpty());
+		check(!CopySegments.IsEmpty());
+
+		FShaderParametersMetadata* ShaderParametersMetadataVS = CreateCompactShaderParametersMetadata(CompactUniformsVS, TEXT("MaterialVS"), TEXT("MaterialVS"), TEXT("MaterialUniformsVS"), TEXT("MaterialVS"), MembersVS, StructSizeVS, CopySegments);
+
+		if (ShaderParametersMetadataVS)
+		{
+			check(CompactUniformsVS->UniformBufferLayoutInitializer.ConstantBufferSize <= UniformBufferLayoutInitializer.ConstantBufferSize);
+			AllUniformBufferStructs.Add(ShaderParametersMetadataVS);
+		}
+	}
+	else
+	{
+		check(MembersVS.IsEmpty());
+		check(CopySegments.IsEmpty());
+		CompactUniformsVSOptional.Empty();
+	}
+
+	return AllUniformBufferStructs;
 }
+#endif
 
 FUniformExpressionSet::FVTPackedStackAndLayerIndex FUniformExpressionSet::GetVTStackAndLayerIndex(int32 UniformExpressionIndex) const
 {
@@ -944,6 +1244,27 @@ void FUniformExpressionSet::FillUniformBuffer(const FMaterialRenderContext& Mate
 	FillUniformBuffer(MaterialRenderContext, UniformExpressionCache.AllocatedVTs, UniformBufferLayout, TempBuffer, TempBufferSize);
 }
 
+void FUniformExpressionSet::CompactUniformBuffer(uint8* DestBuffer, int32 DestBufferSize, const uint8* SrcBuffer, int32 SrcBufferSize, const FCompactUniformExpressionSet& CompactUniforms)
+{
+	FMemoryView         SrcView = MakeMemoryView(SrcBuffer, SrcBufferSize);
+	FMutableMemoryView  DestView = MakeMemoryView(DestBuffer, DestBufferSize);
+
+	for (const FCompactionMetaDataSegment& CompactionSegment : CompactUniforms.CompactionSegments)
+	{
+		FMemoryView        SrcSegment = SrcView.Mid(CompactionSegment.SrcOffset,  CompactionSegment.SizeToCopy);
+		FMutableMemoryView DestSegment = DestView.Mid(CompactionSegment.DestOffset, CompactionSegment.SizeToCopy);
+
+		// Make sure that we won't stomp memory if compaction metadata is corrupted. TMemoryView::Mid adjusts the size based on the parent's capacity
+		if ( !ensureAlways(SrcSegment.GetSize() == CompactionSegment.SizeToCopy) ||
+			 !ensureAlways(DestSegment.GetSize() == CompactionSegment.SizeToCopy))
+		{
+			continue;
+		}
+
+		DestSegment.CopyFrom(SrcSegment);
+	}
+}
+
 static void CopyValueToUniformBuffer(const UE::Shader::FValue& Value, float* PreshaderBuffer, uint32 BufferOffset)
 {
 	using namespace UE::Shader;
@@ -1935,6 +2256,9 @@ void FMaterialUniformExpressionTexture::GetTextureParameterInfo(FMaterialTexture
 	OutParameter.TextureIndex = TextureIndex;
 	OutParameter.SamplerSource = SamplerSource;
 	OutParameter.VirtualTextureLayerIndex = TextureLayerIndex;	// VirtualTextureLayerIndex will be 255 if PageTableLayerIndex==INDEX_NONE.
+#if WITH_EDITORONLY_DATA
+	OutParameter.ShaderFrequencyMask = 0;
+#endif
 }
 
 bool FMaterialUniformExpressionTexture::IsIdentical(const FMaterialUniformExpression* OtherExpression) const
@@ -1981,6 +2305,9 @@ void FMaterialUniformExpressionTextureCollection::GetTextureCollectionParameterI
 {
 	OutParameter.TextureCollectionIndex = this->TextureCollectionIndex;
 	OutParameter.bIsVirtualCollection = this->bIsVirtualCollection;
+#if WITH_EDITORONLY_DATA
+	OutParameter.ShaderFrequencyMask = 0;
+#endif
 }
 
 // FMaterialUniformExpressionTextureCollectionParameter
@@ -2050,6 +2377,9 @@ void FMaterialUniformExpressionExternalTexture::GetExternalTextureParameterInfo(
 {
 	OutParameter.ExternalTextureGuid = ExternalTextureGuid;
 	OutParameter.SourceTextureIndex = SourceTextureIndex;
+#if WITH_EDITORONLY_DATA
+	OutParameter.ShaderFrequencyMask = 0;
+#endif
 }
 
 FMaterialUniformExpressionExternalTextureParameter::FMaterialUniformExpressionExternalTextureParameter()
@@ -2064,6 +2394,9 @@ void FMaterialUniformExpressionExternalTextureParameter::GetExternalTextureParam
 {
 	Super::GetExternalTextureParameterInfo(OutParameter);
 	OutParameter.ParameterName = NameToScriptName(ParameterName);
+#if WITH_EDITORONLY_DATA
+	OutParameter.ShaderFrequencyMask = 0;
+#endif
 }
 
 bool FMaterialUniformExpressionExternalTextureParameter::IsIdentical(const FMaterialUniformExpression* OtherExpression) const
diff --git a/Engine/Source/Runtime/Engine/Private/Materials/MaterialUniformExpressions.h b/Engine/Source/Runtime/Engine/Private/Materials/MaterialUniformExpressions.h
--- a/Engine/Source/Runtime/Engine/Private/Materials/MaterialUniformExpressions.h
+++ b/Engine/Source/Runtime/Engine/Private/Materials/MaterialUniformExpressions.h
@@ -76,6 +76,8 @@ public:
 
 	/** Index of this uniform in the material translator's list of unique expressions */
 	int32 UniformIndex = INDEX_NONE;
+
+	EShaderFrequencyMask ShaderFrequencyMask = 0;
 };
 
 /**
diff --git a/Engine/Source/Runtime/Engine/Private/ShaderCompiler/ShaderCompiler.cpp b/Engine/Source/Runtime/Engine/Private/ShaderCompiler/ShaderCompiler.cpp
--- a/Engine/Source/Runtime/Engine/Private/ShaderCompiler/ShaderCompiler.cpp
+++ b/Engine/Source/Runtime/Engine/Private/ShaderCompiler/ShaderCompiler.cpp
@@ -185,6 +185,19 @@ static TAutoConsoleVariable<int> CVarShaderCompilerLogSlowJobThreshold(
 );
 
 
+bool GUseMaterialVSUniformBuffer = true;
+static FAutoConsoleVariableRef CVarUseMaterialVSUniformBuffer(
+	TEXT("r.ShaderCompiler.UseMaterialVSUniformBuffer"),
+	GUseMaterialVSUniformBuffer,
+	TEXT("If true, translate/compile vertex shaders with a dedicated uniform buffer"),
+	ECVF_ReadOnly
+);
+
+bool UseMaterialVSUniformBuffer()
+{
+	return GUseMaterialVSUniformBuffer;
+}
+
 bool AreShaderErrorsFatal()
 {
 	return CVarAreShaderErrorsFatal.GetValueOnAnyThread();
@@ -3265,6 +3278,18 @@ void GlobalBeginCompileShader(
 	Input.bIncludeUsedOutputs = false;
 	Input.bBindlessEnabled = UE::ShaderCompiler::ShouldCompileWithBindlessEnabled(ShaderPlatform, Input);
 
+	if (UseMaterialVSUniformBuffer())
+	{
+		Input.VSDisallowedUniformBufferNames.Add(TEXT("Material"));
+		Input.NonVSDisallowedUniformBufferNames.Add(TEXT("MaterialVS"));
+	}
+	else
+	{
+		// Make sure we don't try to include the MaterialVS by mistake
+		Input.VSDisallowedUniformBufferNames.Add(TEXT("MaterialVS"));
+		Input.NonVSDisallowedUniformBufferNames.Add(TEXT("MaterialVS"));
+	}
+
 	if (GDumpShaderDebugInfoBindless)
 	{
 		FStringBuilderBase Builder;
diff --git a/Engine/Source/Runtime/Engine/Public/MaterialShared.h b/Engine/Source/Runtime/Engine/Public/MaterialShared.h
--- a/Engine/Source/Runtime/Engine/Public/MaterialShared.h
+++ b/Engine/Source/Runtime/Engine/Public/MaterialShared.h
@@ -379,6 +379,8 @@ struct FMaterialRenderContext
 	ENGINE_API FGuid GetExternalTextureGuid(const FGuid& ExternalTextureGuid, const FName& ParameterName, int32 SourceTextureIndex) const;
 };
 
+using EShaderFrequencyMask = uint16;
+
 class FMaterialVirtualTextureStack
 {
 	DECLARE_TYPE_LAYOUT(FMaterialVirtualTextureStack, NonVirtual);
@@ -428,6 +430,7 @@ public:
 				return false;
 			}
 		}
+		// Note: ShaderFrequencyMask is intentionally ignored — it only affects the editor-side decision to include the parameter in the VS compact uniform buffer.
 		return true;
 	}
 
@@ -442,6 +445,7 @@ public:
 	LAYOUT_ARRAY(int32, LayerUniformExpressionIndices, VIRTUALTEXTURE_SPACE_MAXLAYERS);
 	/** Index of a texture reference if we create a stack from a single known texture that has it's own layer stack. */
 	LAYOUT_FIELD(int32, PreallocatedStackTextureIndex);
+	LAYOUT_FIELD_EDITORONLY(EShaderFrequencyMask, ShaderFrequencyMask);
 };
 
 class FMaterialCacheTagStack
@@ -451,6 +455,7 @@ class FMaterialCacheTagStack
 public:
 	friend bool operator==(const FMaterialCacheTagStack& Lhs, const FMaterialCacheTagStack& Rhs)
 	{
+		// Note: ShaderFrequencyMask is intentionally ignored — it only affects the editor-side decision to include the parameter in the VS compact uniform buffer.
 		return Lhs.TagGuid != Rhs.TagGuid;
 	}
 
@@ -460,6 +465,7 @@ public:
 	}
 	
 	LAYOUT_FIELD(FGuid, TagGuid);
+	LAYOUT_FIELD_EDITORONLY(EShaderFrequencyMask, ShaderFrequencyMask);
 };
 
 /**
@@ -537,6 +543,7 @@ class FMaterialTextureParameterInfo
 public:
 	friend inline bool operator==(const FMaterialTextureParameterInfo& Lhs, const FMaterialTextureParameterInfo& Rhs)
 	{
+        // Note: ShaderFrequencyMask is intentionally ignored — it only affects the editor-side decision to include the parameter in the VS compact uniform buffer.
 		return Lhs.ParameterInfo == Rhs.ParameterInfo && Lhs.TextureIndex == Rhs.TextureIndex && Lhs.SamplerSource == Rhs.SamplerSource && Lhs.VirtualTextureLayerIndex == Rhs.VirtualTextureLayerIndex;
 	}
 	friend inline bool operator!=(const FMaterialTextureParameterInfo& Lhs, const FMaterialTextureParameterInfo& Rhs)
@@ -553,6 +560,7 @@ public:
 	LAYOUT_FIELD_INITIALIZED(int32, TextureIndex, INDEX_NONE);
 	LAYOUT_FIELD(TEnumAsByte<ESamplerSourceMode>, SamplerSource);
 	LAYOUT_FIELD_INITIALIZED(uint8, VirtualTextureLayerIndex, 0u);
+	LAYOUT_FIELD_EDITORONLY(EShaderFrequencyMask, ShaderFrequencyMask);
 };
 
 class FMaterialTextureCollectionParameterInfo
@@ -561,6 +569,7 @@ class FMaterialTextureCollectionParameterInfo
 public:
 	friend inline bool operator==(const FMaterialTextureCollectionParameterInfo& Lhs, const FMaterialTextureCollectionParameterInfo& Rhs)
 	{
+        // Note: ShaderFrequencyMask is intentionally ignored — it only affects the editor-side decision to include the parameter in the VS compact uniform buffer.
 		return Lhs.ParameterInfo == Rhs.ParameterInfo && Lhs.TextureCollectionIndex == Rhs.TextureCollectionIndex && Lhs.bIsVirtualCollection == Rhs.bIsVirtualCollection;
 	}
 	friend inline bool operator!=(const FMaterialTextureCollectionParameterInfo& Lhs, const FMaterialTextureCollectionParameterInfo& Rhs)
@@ -574,6 +583,7 @@ public:
 	LAYOUT_FIELD(FHashedMaterialParameterInfo, ParameterInfo);
 	LAYOUT_FIELD(int32, TextureCollectionIndex);
 	LAYOUT_FIELD(bool, bIsVirtualCollection);
+	LAYOUT_FIELD_EDITORONLY(EShaderFrequencyMask, ShaderFrequencyMask);
 };
 
 class FMaterialExternalTextureParameterInfo
@@ -582,6 +592,7 @@ class FMaterialExternalTextureParameterInfo
 public:
 	friend inline bool operator==(const FMaterialExternalTextureParameterInfo& Lhs, const FMaterialExternalTextureParameterInfo& Rhs)
 	{
+        // Note: ShaderFrequencyMask is intentionally ignored — it only affects the editor-side decision to include the parameter in the VS compact uniform buffer.
 		return Lhs.SourceTextureIndex == Rhs.SourceTextureIndex && Lhs.ExternalTextureGuid == Rhs.ExternalTextureGuid && Lhs.ParameterName == Rhs.ParameterName;
 	}
 	friend inline bool operator!=(const FMaterialExternalTextureParameterInfo& Lhs, const FMaterialExternalTextureParameterInfo& Rhs)
@@ -594,6 +605,7 @@ public:
 	LAYOUT_FIELD(FScriptName, ParameterName);
 	LAYOUT_FIELD(FGuid, ExternalTextureGuid);
 	LAYOUT_FIELD(int32, SourceTextureIndex);
+	LAYOUT_FIELD_EDITORONLY(EShaderFrequencyMask, ShaderFrequencyMask);
 };
 
 class FUniformParameterOverrides
@@ -652,6 +664,29 @@ private:
 	TMap<FTextureParameterKey, UTexture*> RenderThreadTextureOverrides[NumMaterialTextureParameterTypes];
 };
 
+/** Metadata allowing to create a dense uniform buffer from the generic one*/
+class FCompactionMetaDataSegment
+{
+	DECLARE_TYPE_LAYOUT(FCompactionMetaDataSegment, NonVirtual);
+public:
+	LAYOUT_FIELD(uint16, SrcOffset);
+	LAYOUT_FIELD(uint16, DestOffset);
+	LAYOUT_FIELD(uint16, SizeToCopy);
+};
+
+
+/** Stores information per frequency about the values stored in FUniformExpressionSet*/
+class FCompactUniformExpressionSet
+{
+	DECLARE_TYPE_LAYOUT(FCompactUniformExpressionSet, NonVirtual);
+public:
+	FCompactUniformExpressionSet() : UniformPreshaderBufferSize(0u) {}
+
+	LAYOUT_FIELD(FRHIUniformBufferLayoutInitializer, UniformBufferLayoutInitializer);
+	LAYOUT_FIELD(TMemoryImageArray<FCompactionMetaDataSegment>, CompactionSegments);
+	LAYOUT_FIELD(uint32, UniformPreshaderBufferSize);
+};
+
 /** Stores all uniform expressions for a material generated from a material translation. */
 class FUniformExpressionSet
 {
@@ -663,7 +698,9 @@ public:
 	bool operator==(const FUniformExpressionSet& ReferenceSet) const;
 	ENGINE_API FString GetSummaryString() const;
 
-	FShaderParametersMetadata* CreateBufferStruct();
+#if WITH_EDITOR
+	TArray<FShaderParametersMetadata*> CreateBufferStructs();
+#endif
 
 	void SetParameterCollections(TConstArrayView<const class UMaterialParameterCollection*> Collections);
 
@@ -671,6 +708,8 @@ public:
 
 	ENGINE_API void FillUniformBuffer(const FMaterialRenderContext& MaterialRenderContext, TConstArrayView<IAllocatedVirtualTexture*> AllocatedVTs, const FRHIUniformBufferLayout* UniformBufferLayout, uint8* TempBuffer, int TempBufferSize) const;
 
+	static void CompactUniformBuffer(uint8* DestBuffer, int32 DestBufferSize, const uint8* SrcBuffer, int32 SrcBufferSize, const FCompactUniformExpressionSet& CompactUniforms);
+
 	// Get a combined hash of all referenced Texture2D's underlying RHI textures, going through TextureReferences. Can be used to tell if any texture has gone through texture streaming mip changes recently.
 	ENGINE_API uint32 GetReferencedTexture2DRHIHash(const FMaterialRenderContext& MaterialRenderContext) const;
 
@@ -742,6 +781,24 @@ public:
 
 	void CountTextureCollections(uint32& BindlessCollections, uint32& VirtualCollections) const;
 
+	const FCompactUniformExpressionSet* GetCompactUniformsVS() const 
+	{ 
+		check(CompactUniformsVSOptional.Num() <= 1); 
+		return CompactUniformsVSOptional.IsEmpty() ? nullptr : &CompactUniformsVSOptional[0]; 
+	}
+
+	FCompactUniformExpressionSet* GetCompactUniformsVS() 
+	{ 
+		check(CompactUniformsVSOptional.Num() <= 1); 
+		return CompactUniformsVSOptional.IsEmpty() ? nullptr : &CompactUniformsVSOptional[0]; 
+	}
+
+	FCompactUniformExpressionSet* CreateCompactUniformsVS()
+	{
+		CompactUniformsVSOptional.Empty();
+		CompactUniformsVSOptional.Add(FCompactUniformExpressionSet());
+		return &CompactUniformsVSOptional[0];
+	}
 protected:
 	union FVTPackedStackAndLayerIndex
 	{
@@ -787,6 +844,11 @@ protected:
 	LAYOUT_FIELD(TMemoryImageArray<FGuid>, ParameterCollections);
 
 	LAYOUT_FIELD(FRHIUniformBufferLayoutInitializer, UniformBufferLayoutInitializer);
+	
+	// Using TMemoryImageArray here because CompactUniformsVS is optional for now, and we don't want to bloat FUniformExpressionSet when disabled
+	// FUniformExpressionSet are getting copied with shadermaps, so things are usually copied by value, so having a pointer makes things trickier
+	LAYOUT_FIELD(TMemoryImageArray<FCompactUniformExpressionSet>, CompactUniformsVSOptional);
+
 };
 
 /**
@@ -1771,6 +1833,7 @@ public:
 
 	const FUniformExpressionSet& GetUniformExpressionSet() const { return GetContent()->MaterialCompilationOutput.UniformExpressionSet; }
 	const FRHIUniformBufferLayout* GetUniformBufferLayout() const { return UniformBufferLayout; }
+	const FRHIUniformBufferLayout* GetUniformBufferLayoutVS() const { return UniformBufferLayoutVS; }
 
 	int32 GetNumRefs() const { return NumRefs; }
 	int32 GetRefCount() const { return NumRefs; }
@@ -1834,6 +1897,7 @@ private:
 #endif // WITH_EDITOR
 
 	FUniformBufferLayoutRHIRef UniformBufferLayout;
+	FUniformBufferLayoutRHIRef UniformBufferLayoutVS;
 
 #if ALLOW_SHADERMAP_DEBUG_DATA
 	float CompileTime;
@@ -3028,7 +3092,7 @@ private:
 
 	/** Populates OutEnvironment with defines needed to compile shaders for this material. */
 	void SetupMaterialEnvironment(
-		const FShaderParametersMetadata& InUniformBufferStruct,
+		TConstArrayView<FShaderParametersMetadata*> InUniformBufferStructs,
 		const FUniformExpressionSet& InUniformExpressionSet,
 		FShaderCompilerEnvironment& OutEnvironment
 		) const;
diff --git a/Engine/Source/Runtime/Engine/Public/Materials/MaterialRenderProxy.h b/Engine/Source/Runtime/Engine/Public/Materials/MaterialRenderProxy.h
--- a/Engine/Source/Runtime/Engine/Public/Materials/MaterialRenderProxy.h
+++ b/Engine/Source/Runtime/Engine/Public/Materials/MaterialRenderProxy.h
@@ -50,6 +50,8 @@ struct FUniformExpressionCache
 {
 	/** Material uniform buffer. */
 	FUniformBufferRHIRef UniformBuffer;
+	/** Material uniform buffer used only by VS */
+	FUniformBufferRHIRef UniformBufferVS;
 	/** Allocated virtual textures, one for each entry in FUniformExpressionSet::VTStacks */
 	TArray<IAllocatedVirtualTexture*> AllocatedVTs;
 	/** Allocated virtual textures that will need destroying during a call to ResetAllocatedVTs() */
diff --git a/Engine/Source/Runtime/Engine/Public/ShaderCompiler.h b/Engine/Source/Runtime/Engine/Public/ShaderCompiler.h
--- a/Engine/Source/Runtime/Engine/Public/ShaderCompiler.h
+++ b/Engine/Source/Runtime/Engine/Public/ShaderCompiler.h
@@ -51,6 +51,7 @@ DECLARE_LOG_CATEGORY_EXTERN(LogShaderCompilers, Log, All);
 
 bool AreShaderErrorsFatal();
 bool CreateShadersOnLoad();
+bool UseMaterialVSUniformBuffer();
 
 extern ENGINE_API bool IsShaderJobCacheDDCEnabled();
 extern ENGINE_API bool IsMaterialMapDDCEnabled();
diff --git a/Engine/Source/Runtime/RenderCore/Private/ShaderCore.cpp b/Engine/Source/Runtime/RenderCore/Private/ShaderCore.cpp
--- a/Engine/Source/Runtime/RenderCore/Private/ShaderCore.cpp
+++ b/Engine/Source/Runtime/RenderCore/Private/ShaderCore.cpp
@@ -4260,6 +4260,10 @@ FArchive& operator<<(FArchive& Ar, FShaderCompilerInput& Input)
 	Ar << Input.ShaderName;
 	Ar << Input.SupportedHardwareMask;
 	Ar << Input.UsedOutputs;
+	// VSDisallowedUniformBufferNames and NonVSDisallowedUniformBufferNames are not used in FShaderCompileJob::GetInputHash because these do not influence
+	// the result of the shader compilation and are only there to perform validation wrt what material uniform buffer is used
+	Ar << Input.VSDisallowedUniformBufferNames;
+	Ar << Input.NonVSDisallowedUniformBufferNames;
 	Ar << Input.DumpDebugInfoRootPath;
 	Ar << Input.DumpDebugInfoPath;
 	Ar << Input.DebugInfoFlags;
diff --git a/Engine/Source/Runtime/RenderCore/Public/ShaderCompilerCore.h b/Engine/Source/Runtime/RenderCore/Public/ShaderCompilerCore.h
--- a/Engine/Source/Runtime/RenderCore/Public/ShaderCompilerCore.h
+++ b/Engine/Source/Runtime/RenderCore/Public/ShaderCompilerCore.h
@@ -233,6 +233,11 @@ struct FShaderCompilerInput
 	// Shader pipeline information
 	TArray<FString> UsedOutputs;
 
+	// Uniform buffer that are (dis)allowed to be used for (non) VS frequencies: relies on reflection data to fail 
+	// compilation even though the shader may actually compile. This ensures we don't start using a uniform buffer by mistake
+	TArray<FString> VSDisallowedUniformBufferNames;
+	TArray<FString> NonVSDisallowedUniformBufferNames;
+
 	// Dump debug path (up to platform) e.g. "D:/Project/Saved/ShaderDebugInfo/PCD3D_SM5"
 	FString DumpDebugInfoRootPath;
 	// only used if enabled by r.DumpShaderDebugInfo (platform/groupname) e.g. ""
diff --git a/Engine/Source/Runtime/Renderer/Private/ShaderBaseClasses.cpp b/Engine/Source/Runtime/Renderer/Private/ShaderBaseClasses.cpp
--- a/Engine/Source/Runtime/Renderer/Private/ShaderBaseClasses.cpp
+++ b/Engine/Source/Runtime/Renderer/Private/ShaderBaseClasses.cpp
@@ -112,6 +112,7 @@ FMaterialShader::FMaterialShader(const FMaterialShaderType::CompiledShaderInitia
 
 	// Bind the material uniform buffer parameter.
 	MaterialUniformBuffer.Bind(Initializer.ParameterMap,TEXT("Material"));
+	MaterialUniformBufferVS.Bind(Initializer.ParameterMap,TEXT("MaterialVS"));
 
 	for (int32 CollectionIndex = 0; CollectionIndex < Initializer.UniformExpressionSet.ParameterCollections.Num(); CollectionIndex++)
 	{
@@ -285,6 +286,10 @@ void FMaterialShader::SetParameters(
 #endif
 
 	SetUniformBufferParameter(BatchedParameters, MaterialUniformBuffer, UniformExpressionCache->UniformBuffer);
+	if (UniformExpressionCache->UniformBufferVS)
+	{
+		SetUniformBufferParameter(BatchedParameters, MaterialUniformBufferVS, UniformExpressionCache->UniformBufferVS);
+	}
 
 #if !(UE_BUILD_TEST || UE_BUILD_SHIPPING || !WITH_EDITOR)
 	VerifyExpressionAndShaderMaps(MaterialRenderProxy, Material, UniformExpressionCache);
@@ -349,6 +354,10 @@ void FMaterialShader::GetShaderBindings(
 #endif
 
 	ShaderBindings.Add(MaterialUniformBuffer, UniformExpressionCache.UniformBuffer);
+	if (UniformExpressionCache.UniformBufferVS)
+	{
+		ShaderBindings.Add(MaterialUniformBufferVS, UniformExpressionCache.UniformBufferVS);
+	}
 
 	{
 		const TArray<FGuid>& ParameterCollections = UniformExpressionCache.ParameterCollections;
diff --git a/Engine/Source/Runtime/Renderer/Public/MaterialShader.h b/Engine/Source/Runtime/Renderer/Public/MaterialShader.h
--- a/Engine/Source/Runtime/Renderer/Public/MaterialShader.h
+++ b/Engine/Source/Runtime/Renderer/Public/MaterialShader.h
@@ -120,6 +120,7 @@ private:
 	LAYOUT_FIELD(TMemoryImageArray<FShaderUniformBufferParameter>, ParameterCollectionUniformBuffers);
 
 	LAYOUT_FIELD(FShaderUniformBufferParameter, MaterialUniformBuffer);
+	LAYOUT_FIELD(FShaderUniformBufferParameter, MaterialUniformBufferVS);
 
 	// Only needed to avoid unbound parameter error
 	// This texture is bound as an UAV (RWTexture) and so it must be bound together with any RT. So it actually bound but not as part of the material
