Fix editor compilation

#Synced-CL 48597399

diff --git a/Engine/Source/Runtime/Engine/Private/Materials/MaterialUniformExpressions.cpp b/Engine/Source/Runtime/Engine/Private/Materials/MaterialUniformExpressions.cpp
--- a/Engine/Source/Runtime/Engine/Private/Materials/MaterialUniformExpressions.cpp
+++ b/Engine/Source/Runtime/Engine/Private/Materials/MaterialUniformExpressions.cpp
@@ -391,13 +391,13 @@ void FUniformExpressionSet::SetParameterCollections(TConstArrayView<const UMater
 	}
 }
 
-template <class TMemberType>
-uint32 CountIndividualMembers(const TMemoryImageArray<TMemberType>& ParameterCollection, EShaderFrequencyMask ShaderFrequencyMask, TFunctionRef<bool(const TMemberType&)> Filter)
+#if WITH_EDITOR
+uint32 CountTextureCollectionParameterInfo(const TMemoryImageArray<FMaterialTextureCollectionParameterInfo>& ParameterCollection, EShaderFrequencyMask ShaderFrequencyMask)
 {
 	uint32 NumElements = 0;
-	for (const TMemberType& ParameterCollectionMember : ParameterCollection)
+	for (const FMaterialTextureCollectionParameterInfo& ParameterCollectionMember : ParameterCollection)
 	{
-		if (!Filter(ParameterCollectionMember))
+		if (!ParameterCollectionMember.bIsVirtualCollection)
 		{
 			continue;
 		}
@@ -503,7 +503,6 @@ static FShaderParametersMetadata* CreateCompactShaderParametersMetadata(FCompact
 
 }
 
-#if WITH_EDITOR
 static bool IsNumericMember(const FShaderParametersMetadata::FMember& Member)
 {
 	switch (Member.GetBaseType())
@@ -588,8 +587,7 @@ TArray<FShaderParametersMetadata*> FUniformExpressionSet::CreateBufferStructs()
 			new(Members) FShaderParametersMetadata::FMember(TEXT("TextureCollectionPackedUniforms"), TEXT("FIndirectVirtualTextureUniform"), __LINE__, NextMemberOffset, UBMT_UINT32, EShaderPrecisionModifier::Float, 1, 4, UE::HLSL::IndirectVirtualTextureUniformDWord4Count * VirtualCollectionCount, NULL);
 			NextMemberOffset += sizeof(UE::HLSL::FIndirectVirtualTextureUniform) * VirtualCollectionCount;
 
-			TFunctionRef<bool(const FMaterialTextureCollectionParameterInfo &)> AcceptElement([](const auto& A) { return A.bIsVirtualCollection; });
-			uint32 NumMembers = CountIndividualMembers(UniformTextureCollectionParameters, ShaderFrequencyMaskVS, AcceptElement);
+			uint32 NumMembers = CountTextureCollectionParameterInfo(UniformTextureCollectionParameters, ShaderFrequencyMaskVS);
 			if (NumMembers > 0)
 			{
 				uint32 SrcSize = NextMemberOffset - Members.Last().GetOffset();
