Substrate
- Allow to select either slab or hair BSDF. This show how it is possible to do, while relying on a parameter blending sub tree.
- Reworked how a graph is validated using a allowed list of nodes approach
- Removed no longer necessary compiler error handled in a more general fashion by white lists.

#rnx
#tests PC editor

#Synced-CL 48604927

diff --git a/Engine/Source/Editor/MaterialEditor/Private/SMaterialEditorSubstrateWidget.cpp b/Engine/Source/Editor/MaterialEditor/Private/SMaterialEditorSubstrateWidget.cpp
--- a/Engine/Source/Editor/MaterialEditor/Private/SMaterialEditorSubstrateWidget.cpp
+++ b/Engine/Source/Editor/MaterialEditor/Private/SMaterialEditorSubstrateWidget.cpp
@@ -681,6 +681,7 @@ void SMaterialEditorSubstrateWidget::Tick(const FGeometry& AllottedGeometry, con
 				{
 				case SUBSTRATE_MATERIAL_TYPE_SINGLESLAB:			MaterialType = TEXT("Single Slab"); break;
 				case SUBSTRATE_MATERIAL_TYPE_MULTIPLESLABS:			MaterialType = TEXT("Multiple Slabs"); break;
+				case SUBSTRATE_MATERIAL_TYPE_MULTIPLEBSDFS:			MaterialType = TEXT("Multiple BSDFs"); break;
 				case SUBSTRATE_MATERIAL_TYPE_VOLUMETRICFOGCLOUD:	MaterialType = TEXT("Volumetric (Fog or Cloud)"); break;
 				case SUBSTRATE_MATERIAL_TYPE_UNLIT:					MaterialType = TEXT("Unlit"); break;
 				case SUBSTRATE_MATERIAL_TYPE_HAIR:					MaterialType = TEXT("Hair"); break;
diff --git a/Engine/Source/Runtime/Core/Private/UObject/FortniteMainBranchObjectVersion.cpp b/Engine/Source/Runtime/Core/Private/UObject/FortniteMainBranchObjectVersion.cpp
--- a/Engine/Source/Runtime/Core/Private/UObject/FortniteMainBranchObjectVersion.cpp
+++ b/Engine/Source/Runtime/Core/Private/UObject/FortniteMainBranchObjectVersion.cpp
@@ -13,7 +13,7 @@ TMap<FGuid, FGuid> FFortniteMainBranchObjectVersion::GetSystemGuids()
 	SystemGuids.Add(DevGuids.MATERIALSHADERMAP_DERIVEDDATA_VER, FGuid("10F1046C8BCC4303A8F6DD4F55621F90"));
 	SystemGuids.Add(DevGuids.NIAGARASHADERMAP_DERIVEDDATA_VER, FGuid("CC229D919A214E59B79D01AC2E5BF981"));
 	SystemGuids.Add(DevGuids.Niagara_LatestScriptCompileVersion, FGuid("FB15C2870A4E4AE8B5F58566F0E3825E"));
-	SystemGuids.Add(DevGuids.MaterialTranslationDDCVersion, FGuid("A723F63E47F6470491077EE964C6726A"));
+	SystemGuids.Add(DevGuids.MaterialTranslationDDCVersion, FGuid("A723F63E47F6470491077EE964C6726B"));
 
 	// These GUIDs are in AppendRestrictedFortniteMainStreamObjectVersionGuids as they should rarely be changed.
 	//SystemGuids.Add(DevGuids.NANITE_DERIVEDDATA_VER, FGuid("ACF98184978742A592537DA223A1E6BE"));
diff --git a/Engine/Source/Runtime/Engine/Private/Materials/HLSLMaterialTranslator.cpp b/Engine/Source/Runtime/Engine/Private/Materials/HLSLMaterialTranslator.cpp
--- a/Engine/Source/Runtime/Engine/Private/Materials/HLSLMaterialTranslator.cpp
+++ b/Engine/Source/Runtime/Engine/Private/Materials/HLSLMaterialTranslator.cpp
@@ -13747,6 +13747,259 @@ bool FHLSLMaterialTranslator::FSubstrateCompilationContext::SubstrateGenerateDer
 		return false;
 	}
 
+	
+	//
+	// Validate the tree operators with respect to how BSDF and operator can flow.
+	// Also evaluate a bunch of data required.
+	//
+	enum EMaterialTypeValidation : uint32
+	{
+		MTV_BSDF_None				= 0x000,
+		MTV_BSDF_Unlit				= 0x001,
+		MTV_BSDF_Slab				= 0x002,
+		MTV_BSDF_Hair				= 0x004,
+		MTV_BSDF_Eye				= 0x008,
+		MTV_BSDF_SLW				= 0x010,
+		MTV_BSDF_Toon				= 0x020,
+		MTV_BSDF_VFogCloud			= 0x040,
+		MTV_UI						= 0x080,
+		MTV_Decal					= 0x100,
+		MTV_PostProcess				= 0x200,
+		MTV_LightFunction			= 0x400,
+
+		MTV_ValidMultipleBSDFs		= (MTV_BSDF_Slab | MTV_BSDF_Hair | MTV_BSDF_Toon),	// SUBSTRATE_TODO Unlit cannot be in there yet due to the specific define for it (SUBSTRATE_OPTIMIZED_UNLIT)
+		MTV_WeightOp_AllowedList	= (MTV_BSDF_Unlit | MTV_BSDF_Slab),
+		MTV_SelectOp_AllowedList	= (MTV_ValidMultipleBSDFs),							// This only works because parameter blending is enforced so there will always be a single BSDF of a single type only. And hair or toon BSDFs will never pass through other operator allowed list.
+		MTV_HorizoOp_AllowedList	= MTV_BSDF_Slab,
+		MTV_VerticOp_AllowedList	= MTV_BSDF_Slab,
+		MTV_DecalOp_AllowedList		= MTV_BSDF_Slab,
+		MTV_AddOp_AllowedList		= MTV_BSDF_Slab
+	};
+	struct FSubstrateTreeValidation
+	{
+		uint32 MaterialType = MTV_BSDF_None;
+		bool bValid = true;
+	};
+	FSubstrateTreeValidation SubstrateTreeValidation;
+	{
+
+		auto MergeSubstrateTreeValidationParameterBlended = [](const FSubstrateTreeValidation& A, const FSubstrateTreeValidation& B)
+		{
+			FSubstrateTreeValidation Result;
+			Result.MaterialType = A.MaterialType | B.MaterialType;
+			return Result;
+		};
+		auto MergeSubstrateTreeValidation = [](const FSubstrateTreeValidation& A, const FSubstrateTreeValidation& B)
+		{
+			FSubstrateTreeValidation Result;
+			Result.MaterialType = A.MaterialType | B.MaterialType;
+			return Result;
+		};
+
+		auto ValidateAllowedList = [&](FSubstrateOperator& CurrentOperator, FSubstrateTreeValidation SubstrateTreeValidation)
+		{
+			if (CurrentOperator.OperatorType == SUBSTRATE_OPERATOR_BSDF)
+			{
+				return true;
+			}
+
+			FString OperatorName;
+			uint32 AllowedListMask = 0;
+			if (CurrentOperator.OperatorType == SUBSTRATE_OPERATOR_WEIGHT && CurrentOperator.SubUsage == SUBSTRATE_OPERATOR_SUBUSAGE_DECAL)
+			{
+				OperatorName = TEXT("ConvertToDecal");
+				AllowedListMask = MTV_DecalOp_AllowedList;
+			}
+			else if (CurrentOperator.OperatorType == SUBSTRATE_OPERATOR_WEIGHT)
+			{
+				OperatorName = TEXT("Weight");
+				AllowedListMask = MTV_WeightOp_AllowedList;
+			}
+			else if (CurrentOperator.OperatorType == SUBSTRATE_OPERATOR_HORIZONTAL)
+			{
+				OperatorName = TEXT("Horizontal");
+				AllowedListMask = MTV_HorizoOp_AllowedList;
+			}
+			else if (CurrentOperator.OperatorType == SUBSTRATE_OPERATOR_VERTICAL)
+			{
+				OperatorName = TEXT("Vertical");
+				AllowedListMask = MTV_VerticOp_AllowedList;
+			}
+			else if (CurrentOperator.OperatorType == SUBSTRATE_OPERATOR_ADD)
+			{
+				OperatorName = TEXT("Add");
+				AllowedListMask = MTV_AddOp_AllowedList;
+			}
+			else if (CurrentOperator.OperatorType == SUBSTRATE_OPERATOR_SELECT)
+			{
+				OperatorName = TEXT("Select");
+				AllowedListMask = MTV_SelectOp_AllowedList;
+			}
+			else
+			{
+				Compiler->Errorf(TEXT("Substrate error: unhandled operator type from ValidateAllowedList. (asset: %s).\r\n"), *CompilerMaterial->GetDebugName());
+				return false;
+			}
+
+			uint32 FoundNotAllowedTypeMask = SubstrateTreeValidation.MaterialType & (~AllowedListMask);
+
+			if (FoundNotAllowedTypeMask != 0)
+			{
+				FString NotAllowedTypeNames;
+				NotAllowedTypeNames += "{";
+				if (!!(FoundNotAllowedTypeMask & MTV_BSDF_Unlit))
+				{
+					NotAllowedTypeNames += TEXT("Unlit ");
+				}
+				if (!!(FoundNotAllowedTypeMask & MTV_BSDF_Slab))
+				{
+					NotAllowedTypeNames += TEXT("Slab ");
+				}
+				if (!!(FoundNotAllowedTypeMask & MTV_BSDF_Hair))
+				{
+					NotAllowedTypeNames += TEXT("Hair ");
+				}
+				if (!!(FoundNotAllowedTypeMask & MTV_BSDF_Eye))
+				{
+					NotAllowedTypeNames += TEXT("Eye ");
+				}
+				if (!!(FoundNotAllowedTypeMask & MTV_BSDF_SLW))
+				{
+					NotAllowedTypeNames += TEXT("SLW ");
+				}
+				if (!!(FoundNotAllowedTypeMask & MTV_BSDF_Toon))
+				{
+					NotAllowedTypeNames += TEXT("Toon ");
+				}
+				if (!!(FoundNotAllowedTypeMask & MTV_BSDF_VFogCloud))
+				{
+					NotAllowedTypeNames += TEXT("VFogCloud ");
+				}
+				if (!!(FoundNotAllowedTypeMask & MTV_UI))
+				{
+					NotAllowedTypeNames += TEXT("UI ");
+				}
+				if (!!(FoundNotAllowedTypeMask & MTV_Decal))
+				{
+					NotAllowedTypeNames += TEXT("Decal ");
+				}
+				if (!!(FoundNotAllowedTypeMask & MTV_PostProcess))
+				{
+					NotAllowedTypeNames += TEXT("PostProcess ");
+				}
+				if (!!(FoundNotAllowedTypeMask & MTV_LightFunction))
+				{
+					NotAllowedTypeNames += TEXT("LightFunction ");
+				}
+				NotAllowedTypeNames += "}";
+
+				Compiler->Errorf(TEXT("Substrate error: Operator %s cannot be used with the following data flowing through %s. (asset: %s).\r\n"), *OperatorName, *NotAllowedTypeNames, *CompilerMaterial->GetDebugName());
+				return false;
+			}
+			return true;
+		};
+
+		std::function<FSubstrateTreeValidation(FSubstrateOperator&)> WalkOperators = [&](FSubstrateOperator& CurrentOperator) -> FSubstrateTreeValidation
+		{
+			FSubstrateTreeValidation Out;
+			switch (CurrentOperator.OperatorType)
+			{
+			case SUBSTRATE_OPERATOR_VERTICAL:
+			case SUBSTRATE_OPERATOR_HORIZONTAL:
+			case SUBSTRATE_OPERATOR_ADD:
+			{
+				FSubstrateTreeValidation A = WalkOperators(SubstrateMaterialExpressionRegisteredOperators[CurrentOperator.LeftIndex]);
+				FSubstrateTreeValidation B = WalkOperators(SubstrateMaterialExpressionRegisteredOperators[CurrentOperator.RightIndex]);
+				Out = MergeSubstrateTreeValidation(A, B);
+				break;
+			}
+			case SUBSTRATE_OPERATOR_SELECT:
+			{
+				if (CurrentOperator.bNodeRequestParameterBlending)
+				{
+					FSubstrateTreeValidation A = WalkOperators(SubstrateMaterialExpressionRegisteredOperators[CurrentOperator.LeftIndex]);
+					FSubstrateTreeValidation B = WalkOperators(SubstrateMaterialExpressionRegisteredOperators[CurrentOperator.RightIndex]);
+					Out = MergeSubstrateTreeValidation(A, B);
+				}
+				else
+				{
+					Compiler->Errorf(TEXT("Substrate error: Select node can only be parameter blended as of today. (asset: %s).\r\n"), *CompilerMaterial->GetDebugName());
+				}
+				break;
+			}
+			case SUBSTRATE_OPERATOR_WEIGHT:
+			{
+				Out = WalkOperators(SubstrateMaterialExpressionRegisteredOperators[CurrentOperator.LeftIndex]);
+				break;
+			}
+			case SUBSTRATE_OPERATOR_BSDF:
+			{
+				if (CurrentOperator.BSDFType == SUBSTRATE_BSDF_TYPE_UNLIT)
+				{
+					if (CurrentOperator.SubUsage == SUBSTRATE_OPERATOR_SUBUSAGE_NONE)
+					{
+						Out.MaterialType = MTV_BSDF_Unlit;
+					}
+					else if (CurrentOperator.SubUsage == SUBSTRATE_OPERATOR_SUBUSAGE_UI)
+					{
+						Out.MaterialType = MTV_UI;
+					}
+					else if (CurrentOperator.SubUsage == SUBSTRATE_OPERATOR_SUBUSAGE_LIGHTFUNCTION)
+					{
+						Out.MaterialType = MTV_LightFunction;
+					}
+					else if (CurrentOperator.SubUsage == SUBSTRATE_OPERATOR_SUBUSAGE_POSTPROCESS)
+					{
+						Out.MaterialType = MTV_PostProcess;
+					}
+					else
+					{
+						Compiler->Errorf(TEXT("Substrate error: unhandled Unlit BSDF sub type from WalkParameterBlendedSubTree. (asset: %s).\r\n"), *CompilerMaterial->GetDebugName());
+					}
+				}
+				else if (CurrentOperator.BSDFType == SUBSTRATE_BSDF_TYPE_SLAB)
+				{
+					Out.MaterialType = MTV_BSDF_Slab;
+				}
+				else if (CurrentOperator.BSDFType == SUBSTRATE_BSDF_TYPE_VOLUMETRICFOGCLOUD)
+				{
+					Out.MaterialType = MTV_BSDF_VFogCloud;
+				}
+				else if (CurrentOperator.BSDFType == SUBSTRATE_BSDF_TYPE_HAIR)
+				{
+					Out.MaterialType = MTV_BSDF_Hair;
+				}
+				else if (CurrentOperator.BSDFType == SUBSTRATE_BSDF_TYPE_SINGLELAYERWATER)
+				{
+					Out.MaterialType = MTV_BSDF_SLW;
+				}
+				else if (CurrentOperator.BSDFType == SUBSTRATE_BSDF_TYPE_TOON)
+				{
+					Out.MaterialType = MTV_BSDF_Toon;
+				}
+				else if (CurrentOperator.BSDFType == SUBSTRATE_BSDF_TYPE_EYE)
+				{
+					Out.MaterialType = MTV_BSDF_Eye;
+				}
+				else
+				{
+					Compiler->Errorf(TEXT("Substrate error: unhandled BSDF type from WalkParameterBlendedSubTree. (asset: %s).\r\n"), *CompilerMaterial->GetDebugName());
+				}
+				break;
+			}
+			}
+			Out.bValid = ValidateAllowedList(CurrentOperator, Out);
+			return Out;
+		};
+
+		SubstrateTreeValidation = WalkOperators(*SubstrateMaterialRootOperator);
+		if (!SubstrateTreeValidation.bValid)
+		{
+			Compiler->Errorf(TEXT("Substrate error: invalid Substrate node graph. (asset: %s).\r\n"), *CompilerMaterial->GetDebugName());
+			return false;
+		}
+	}
+
 	// Evaluate if simplification is needed.
 	SubstrateSimplificationStatus.bRunFullSimplification = Compiler->SubstrateCompilationConfig.bFullSimplify || CompilationContextIndex == ESubstrateCompilationContext::SCC_FullySimplified;
 	if (!SubstrateSimplificationStatus.bRunFullSimplification)
@@ -13845,7 +14098,7 @@ bool FHLSLMaterialTranslator::FSubstrateCompilationContext::SubstrateGenerateDer
 		}
 		else
 		{
-			Compiler->Errorf(TEXT("Unkown Substrate material simplification status error for %s (asset: %s).\r\n"), *CompilerMaterial->GetDebugName(), *CompilerMaterial->GetAssetPath().ToString());
+			Compiler->Errorf(TEXT("Unkown Substrate material simplification status error for %s (asset: %s).\r\n"), *CompilerMaterial->GetDebugName(), *CompilerMaterial->GetDebugName());
 			return false;
 		}
 
@@ -13857,23 +14110,21 @@ bool FHLSLMaterialTranslator::FSubstrateCompilationContext::SubstrateGenerateDer
 		// Allocate BSDFIndex at the same time.
 		// Check BSDF that should be unique.
 		//
-		bool bHasUnlit = false;
-		bool bHasVFogCloud = false;
-		bool bHasHair = false;
-		bool bHasEye = false;
-		bool bHasSLW = false;
-		bool bHasToon = false;
-		bool bHasSlab = false;
-		bool bHasUI = false;
-		bool bHasDecal = false;
-		bool bHasPostProcess = false;
-		bool bHasLightFunction = false;
+		bool bHasUnlit			= !!(SubstrateTreeValidation.MaterialType & MTV_BSDF_Unlit);
+		bool bHasVFogCloud		= !!(SubstrateTreeValidation.MaterialType & MTV_BSDF_VFogCloud);
+		bool bHasHair			= !!(SubstrateTreeValidation.MaterialType & MTV_BSDF_Hair);
+		bool bHasEye			= !!(SubstrateTreeValidation.MaterialType & MTV_BSDF_Eye);
+		bool bHasSLW			= !!(SubstrateTreeValidation.MaterialType & MTV_BSDF_SLW);
+		bool bHasToon			= !!(SubstrateTreeValidation.MaterialType & MTV_BSDF_Toon);
+		bool bHasSlab			= !!(SubstrateTreeValidation.MaterialType & MTV_BSDF_Slab);
+		bool bHasUI				= !!(SubstrateTreeValidation.MaterialType & MTV_UI);
+		bool bHasDecal			= !!(SubstrateTreeValidation.MaterialType & MTV_Decal);
+		bool bHasPostProcess	= !!(SubstrateTreeValidation.MaterialType & MTV_PostProcess);
+		bool bHasLightFunction	= !!(SubstrateTreeValidation.MaterialType & MTV_LightFunction);
 		{
 			int VOpTopBranchCountTaken = 0;
 			int VOpBottomBranchCountTaken = 0;
 			bool bSubstrateUsesVerticalLayering = false;
-			bool bOperatorEncountered = false;
-			bool bOperatorEncounteredButNotWeight = false;
 
 			std::function<void(FSubstrateOperator&, bool)> WalkOperators = [&](FSubstrateOperator& CurrentOperator, bool bInsideParameterBlendingSubTree) -> void
 			{
@@ -13881,7 +14132,9 @@ bool FHLSLMaterialTranslator::FSubstrateCompilationContext::SubstrateGenerateDer
 				const bool bRootOfParameterBlendingSubTree		= bCurrentOpRequestParameterBlending && !bInsideParameterBlendingSubTree;
 				const bool bUseParameterBlending				= bCurrentOpRequestParameterBlending || bInsideParameterBlendingSubTree;
 
-				if (CurrentOperator.BSDFType == SUBSTRATE_BSDF_TYPE_SLAB || CurrentOperator.BSDFType == SUBSTRATE_BSDF_TYPE_UNLIT)
+				if (CurrentOperator.BSDFType == SUBSTRATE_BSDF_TYPE_SLAB																// Any operators
+					|| CurrentOperator.BSDFType == SUBSTRATE_BSDF_TYPE_UNLIT															// Weight operator
+					|| CurrentOperator.BSDFType == SUBSTRATE_BSDF_TYPE_HAIR || CurrentOperator.BSDFType == SUBSTRATE_BSDF_TYPE_TOON)	// Select operator
 				{
 					// Update the parameter blending data for BSDFs supporting operators.
 					// We also need to do this for UNLIT since it supports Coverage operator.
@@ -13889,20 +14142,12 @@ bool FHLSLMaterialTranslator::FSubstrateCompilationContext::SubstrateGenerateDer
 					CurrentOperator.bRootOfParameterBlendingSubTree = bRootOfParameterBlendingSubTree;
 				}
 
-				// this can show up on Unlit and Weight operators for Decals.
-				bHasUI				|= CurrentOperator.SubUsage == SUBSTRATE_OPERATOR_SUBUSAGE_UI;
-				bHasDecal			|= CurrentOperator.SubUsage == SUBSTRATE_OPERATOR_SUBUSAGE_DECAL;
-				bHasPostProcess		|= CurrentOperator.SubUsage == SUBSTRATE_OPERATOR_SUBUSAGE_POSTPROCESS;
-				bHasLightFunction	|= CurrentOperator.SubUsage == SUBSTRATE_OPERATOR_SUBUSAGE_LIGHTFUNCTION;
-
 				switch (CurrentOperator.OperatorType)
 				{
 				case SUBSTRATE_OPERATOR_VERTICAL:
 				{
 					WalkOperators(SubstrateMaterialExpressionRegisteredOperators[CurrentOperator.LeftIndex], bUseParameterBlending);
 					WalkOperators(SubstrateMaterialExpressionRegisteredOperators[CurrentOperator.RightIndex], bUseParameterBlending);
-					bOperatorEncountered = true;
-					bOperatorEncounteredButNotWeight = true;
 					break;
 				}
 				case SUBSTRATE_OPERATOR_HORIZONTAL:
@@ -13911,14 +14156,11 @@ bool FHLSLMaterialTranslator::FSubstrateCompilationContext::SubstrateGenerateDer
 				{
 					WalkOperators(SubstrateMaterialExpressionRegisteredOperators[CurrentOperator.LeftIndex], bUseParameterBlending);
 					WalkOperators(SubstrateMaterialExpressionRegisteredOperators[CurrentOperator.RightIndex], bUseParameterBlending);
-					bOperatorEncountered = true;
-					bOperatorEncounteredButNotWeight = true;
 					break;
 				}
 				case SUBSTRATE_OPERATOR_WEIGHT:
 				{
 					WalkOperators(SubstrateMaterialExpressionRegisteredOperators[CurrentOperator.LeftIndex], bUseParameterBlending);
-					bOperatorEncountered = true;
 					break;
 				}
 				case SUBSTRATE_OPERATOR_BSDF:
@@ -13961,46 +14203,6 @@ bool FHLSLMaterialTranslator::FSubstrateCompilationContext::SubstrateGenerateDer
 				Compiler->bSubstrateOutputsOpaqueRoughRefractions = bSubstrateUsesVerticalLayering && bIsOpaqueOrMasked;
 			}
 			bSubstrateMaterialIsUnlitNode = bHasUnlit;
-
-			if ((bHasVFogCloud || bHasHair || bHasEye || bHasSLW || bHasToon) && bOperatorEncountered)
-			{
-				Compiler->Errorf(TEXT("Fog/Cloud, Hair, Eye or SingleLayerWater cannot be used with operators. See %s (asset: %s).\r\n"), *CompilerMaterial->GetDebugName(), *CompilerMaterial->GetAssetPath().ToString());
-				return false;
-			}
-
-			if ((bHasUI ||  bHasPostProcess || bHasLightFunction) && bOperatorEncountered)
-			{
-				Compiler->Errorf(TEXT("UI, Post Process or Light Function materials cannot be used with operators. See %s (asset: %s).\r\n"), *CompilerMaterial->GetDebugName(), *CompilerMaterial->GetAssetPath().ToString());
-				return false;
-			}
-
-			if (bHasUnlit && bOperatorEncounteredButNotWeight)
-			{
-				Compiler->Errorf(TEXT("Unlit can only be used with coverage operators.\r\n \
-					If you want to blend an unlit material with other slabs, replace the unlit node with a slab having a black albedo, black F0, and an emissive value.\r\n \
-					See %s (asset: %s).\r\n"), *CompilerMaterial->GetDebugName(), *CompilerMaterial->GetAssetPath().ToString());
-				return false;
-			}
-
-			if ((bHasUnlit || bHasVFogCloud || bHasHair || bHasEye || bHasSLW || bHasToon) && SubstrateMaterialEffectiveClosureCount > 1)
-			{
-				Compiler->Errorf(TEXT("Unlit, Fog/Cloud, Hair or SingleLayerWater must result in a single one Closure (Found %d). See %s (asset: %s).\r\n"), SubstrateMaterialEffectiveClosureCount, *CompilerMaterial->GetDebugName(), *CompilerMaterial->GetAssetPath().ToString());
-				// Even though technically we could support Unlit parameter blending with Slab.
-				return false;
-			}
-
-			if (bHasDecal && bHasUnlit)
-			{
-				// Because Unlit does not support parameter blending.
-				Compiler->Errorf(TEXT("Decals only support Substrate Slabs and Shading Model nodes, not Unlit node. See %s (asset: %s).\r\n"), *CompilerMaterial->GetDebugName(), *CompilerMaterial->GetAssetPath().ToString());
-				return false;
-			}
-
-			if (SubstrateMaterialEffectiveClosureCount > SUBSTRATE_MAX_CLOSURE_COUNT)
-			{
-				Compiler->Errorf(TEXT("Material tries to register more BSDF than can be supported (%d > %d). See %s (asset: %s).\r\n"), SubstrateMaterialEffectiveClosureCount, SUBSTRATE_MAX_CLOSURE_COUNT, *CompilerMaterial->GetDebugName(), *CompilerMaterial->GetAssetPath().ToString());
-				return false;
-			}
 		}
 
 		//
@@ -14477,13 +14679,6 @@ bool FHLSLMaterialTranslator::FSubstrateCompilationContext::SubstrateGenerateDer
 					break;
 				} // case SUBSTRATE_OPERATOR_BSDF
 				} // switch (It.OperatorType)
-
-				const uint32 ClosureRequestedSizeByte = SubstrateMaterialRequestedSizeByte - PreSubstrateMaterialRequestedSizeByte;
-				if (ClosureRequestedSizeByte > ClosureMaxByteCountForOffset)
-				{
-					Compiler->Errorf(TEXT("A closure is requesting more bytes than our closure offset system can handle (%d/%d bytes). Notify your rendering engineer. Material %s (asset: %s).\r\n"), ClosureRequestedSizeByte, ClosureMaxByteCountForOffset, *CompilerMaterial->GetDebugName(), *CompilerMaterial->GetAssetPath().ToString(), OperatorCount, SUBSTRATE_MAX_OPERATOR_COUNT);
-					return false;
-				}
 			}
 
 			if (OperatorCount > SUBSTRATE_MAX_OPERATOR_COUNT)
@@ -14530,6 +14725,8 @@ bool FHLSLMaterialTranslator::FSubstrateCompilationContext::SubstrateGenerateDer
 
 				Compiler->MaterialCompilationOutput.SubstrateMaterialCompilationOutput.bIsThin = CompilerMaterial->IsThinSurface() ? 1 : 0;
 
+				uint32 BSDFTypeCount = FMath::CountBits(SubstrateTreeValidation.MaterialType & MTV_ValidMultipleBSDFs);
+
 				// The order of ifs here is important.
 				if (CompilerMaterial->IsLightFunction())
 				{
@@ -14547,10 +14744,19 @@ bool FHLSLMaterialTranslator::FSubstrateCompilationContext::SubstrateGenerateDer
 				{
 					Compiler->MaterialCompilationOutput.SubstrateMaterialCompilationOutput.MaterialType = SUBSTRATE_MATERIAL_TYPE_DECAL;
 				}
-				else if (SubstrateMaterialEffectiveClosureCount > 1)
+				else if (SubstrateMaterialEffectiveClosureCount > 1 && (SubstrateTreeValidation.MaterialType == (SubstrateTreeValidation.MaterialType & MTV_BSDF_Slab)))
 				{
+					// Only slabs, not hair nor toon.
 					Compiler->MaterialCompilationOutput.SubstrateMaterialCompilationOutput.MaterialType = SUBSTRATE_MATERIAL_TYPE_MULTIPLESLABS;
 				}
+				else if (SubstrateMaterialEffectiveClosureCount > 1) // Only slabs can have multiple closures because the select nodde will always parameter blend slabs into a single one
+				{
+					Compiler->MaterialCompilationOutput.SubstrateMaterialCompilationOutput.MaterialType = SUBSTRATE_MATERIAL_TYPE_MULTIPLESLABS;
+				}
+				else if (BSDFTypeCount > 1)
+				{
+					Compiler->MaterialCompilationOutput.SubstrateMaterialCompilationOutput.MaterialType = SUBSTRATE_MATERIAL_TYPE_MULTIPLEBSDFS;
+				}
 				else if (bHasUnlit)
 				{
 					Compiler->MaterialCompilationOutput.SubstrateMaterialCompilationOutput.MaterialType = SUBSTRATE_MATERIAL_TYPE_UNLIT;
@@ -15179,21 +15385,35 @@ int32 FHLSLMaterialTranslator::SubstrateUIBSDF(int32 EmissiveColor, int32 Opacit
 
 int32 FHLSLMaterialTranslator::SubstrateHairBSDF(int32 BaseColor, int32 Scatter, int32 Specular, int32 Roughness, int32 Backlit, int32 EmissiveColor, int32 Tangent, const FString& SharedLocalBasisIndexMacro, FSubstrateOperator* PromoteToOperator)
 {
+	if (PromoteToOperator)
+	{
+		return AddCodeChunk(
+			MCT_Substrate, TEXT("Parameters.%s.PromoteParameterBlendedBSDFToOperator(GetSubstrateHairBSDF(%s, %s, %s, %s, %s, %s, %s), %u, %u, %u, %u) /* Tangent:%s */"),
+			*GetParametersSubstrateTreeName(CurrentSubstrateCompilationContext),
+			*SubstrateGetCastParameterCode(BaseColor, MCT_Float3),
+			*SubstrateGetCastParameterCode(Scatter, MCT_Float),
+			*SubstrateGetCastParameterCode(Specular, MCT_Float),
+			*SubstrateGetCastParameterCode(Roughness, MCT_Float),
+			*SubstrateGetCastParameterCode(Backlit, MCT_Float),
+			*SubstrateGetCastParameterCode(EmissiveColor, MCT_Float3),
+			*SharedLocalBasisIndexMacro,
+			PromoteToOperator->Index,
+			PromoteToOperator->BSDFIndex,
+			PromoteToOperator->LayerDepth,
+			PromoteToOperator->bIsBottom ? 1 : 0,
+			*GetParameterCode(Tangent)
+		);
+	}
+
 	return AddCodeChunk(
-		MCT_Substrate, TEXT("Parameters.%s.PromoteParameterBlendedBSDFToOperator(GetSubstrateHairBSDF(%s, %s, %s, %s, %s, %s, %s), %u, %u, %u, %u) /* Tangent:%s */"),
-		*GetParametersSubstrateTreeName(CurrentSubstrateCompilationContext),
-		*SubstrateGetCastParameterCode(BaseColor,			MCT_Float3),
-		*SubstrateGetCastParameterCode(Scatter,				MCT_Float),
-		*SubstrateGetCastParameterCode(Specular,			MCT_Float),
-		*SubstrateGetCastParameterCode(Roughness,			MCT_Float),
-		*SubstrateGetCastParameterCode(Backlit,				MCT_Float),
-		*SubstrateGetCastParameterCode(EmissiveColor,		MCT_Float3),
-		*SharedLocalBasisIndexMacro,
-		PromoteToOperator->Index,
-		PromoteToOperator->BSDFIndex,
-		PromoteToOperator->LayerDepth,
-		PromoteToOperator->bIsBottom ? 1 : 0,
-		*GetParameterCode(Tangent)
+		MCT_Substrate, TEXT("GetSubstrateHairBSDF(%s, %s, %s, %s, %s, %s, %s)"),
+		*SubstrateGetCastParameterCode(BaseColor, MCT_Float3),
+		*SubstrateGetCastParameterCode(Scatter, MCT_Float),
+		*SubstrateGetCastParameterCode(Specular, MCT_Float),
+		*SubstrateGetCastParameterCode(Roughness, MCT_Float),
+		*SubstrateGetCastParameterCode(Backlit, MCT_Float),
+		*SubstrateGetCastParameterCode(EmissiveColor, MCT_Float3),
+		*SharedLocalBasisIndexMacro
 	);
 }
 
@@ -15253,14 +15473,34 @@ int32 FHLSLMaterialTranslator::SubstrateToonBSDF(
 	int32 EmissiveColor, int32 Normal,
 	const FString& SharedLocalBasisIndexMacro, FSubstrateOperator* PromoteToOperator)
 {
-	if (ToonProfileId <= INDEX_NONE)
+	if (PromoteToOperator)
 	{
-		Errorf(TEXT("Substrate error: invalid Toon profile specified (asset: %s [%s]: ).\r\n"), *Material->GetDebugName(), *Material->GetAssetPath().ToString());
-		return INDEX_NONE;
+		if (ToonProfileId <= INDEX_NONE)
+		{
+			Errorf(TEXT("Substrate error: invalid Toon profile specified (asset: %s [%s]: ).\r\n"), *Material->GetDebugName(), *Material->GetAssetPath().ToString());
+			return INDEX_NONE;
+		}
+		return AddCodeChunk(
+			MCT_Substrate, TEXT("Parameters.%s.PromoteParameterBlendedBSDFToOperator(GetSubstrateToonBSDF(%s, %s, %s, %s, %s, %s, %s, Parameters.%s.Types), %u, %u, %u, %u) /* Normal:%s */"),
+			*GetParametersSubstrateTreeName(CurrentSubstrateCompilationContext),
+			*SubstrateGetCastParameterCode(ToonProfileId, MCT_Float),
+			*SubstrateGetCastParameterCode(DiffuseColor, MCT_Float3),
+			*SubstrateGetCastParameterCode(SpecularColor, MCT_Float3),
+			*SubstrateGetCastParameterCode(Roughness, MCT_Float),
+			*SubstrateGetCastParameterCode(Anisotropy, MCT_Float),
+			*SubstrateGetCastParameterCode(EmissiveColor, MCT_Float3),
+			*SharedLocalBasisIndexMacro,
+			*GetParametersSharedLocalBasesName(CurrentSubstrateCompilationContext),
+			PromoteToOperator->Index,
+			PromoteToOperator->BSDFIndex,
+			PromoteToOperator->LayerDepth,
+			PromoteToOperator->bIsBottom ? 1 : 0,
+			*GetParameterCode(Normal)
+		);
 	}
+
 	return AddCodeChunk(
-		MCT_Substrate, TEXT("Parameters.%s.PromoteParameterBlendedBSDFToOperator(GetSubstrateToonBSDF(%s, %s, %s, %s, %s, %s, %s, Parameters.%s.Types), %u, %u, %u, %u) /* Normal:%s */"),
-		*GetParametersSubstrateTreeName(CurrentSubstrateCompilationContext),
+		MCT_Substrate, TEXT("GetSubstrateToonBSDF(%s, %s, %s, %s, %s, %s, %s, Parameters.%s.Types)"),
 		*SubstrateGetCastParameterCode(ToonProfileId, MCT_Float),
 		*SubstrateGetCastParameterCode(DiffuseColor, MCT_Float3),
 		*SubstrateGetCastParameterCode(SpecularColor, MCT_Float3),
@@ -15268,12 +15508,7 @@ int32 FHLSLMaterialTranslator::SubstrateToonBSDF(
 		*SubstrateGetCastParameterCode(Anisotropy, MCT_Float),
 		*SubstrateGetCastParameterCode(EmissiveColor, MCT_Float3),
 		*SharedLocalBasisIndexMacro,
-		*GetParametersSharedLocalBasesName(CurrentSubstrateCompilationContext),
-		PromoteToOperator->Index,
-		PromoteToOperator->BSDFIndex,
-		PromoteToOperator->LayerDepth,
-		PromoteToOperator->bIsBottom ? 1 : 0,
-		*GetParameterCode(Normal)
+		*GetParametersSharedLocalBasesName(CurrentSubstrateCompilationContext)
 	);
 }
 
diff --git a/Engine/Source/Runtime/Engine/Private/Materials/MaterialExpressionSubstrate.cpp b/Engine/Source/Runtime/Engine/Private/Materials/MaterialExpressionSubstrate.cpp
--- a/Engine/Source/Runtime/Engine/Private/Materials/MaterialExpressionSubstrate.cpp
+++ b/Engine/Source/Runtime/Engine/Private/Materials/MaterialExpressionSubstrate.cpp
@@ -2084,15 +2084,6 @@ int32 UMaterialExpressionSubstrateHairBSDF::Compile(class FMaterialCompiler* Com
 	FSubstrateOperator& SubstrateOperator = Compiler->SubstrateCompilationGetOperator(Compiler->SubstrateTreeStackGetPathUniqueId());
 	SubstrateOperator.BSDFRegisteredSharedLocalBasis = NewRegisteredSharedLocalBasis;
 
-	if (SubstrateOperator.bUseParameterBlending)
-	{
-		return Compiler->Errorf(TEXT("Substrate Hair BSDF node cannot be used with parameter blending."));
-	}
-	else if (SubstrateOperator.bRootOfParameterBlendingSubTree)
-	{
-		return Compiler->Errorf(TEXT("Substrate Hair BSDF node cannot be the root of a parameter blending sub tree."));
-	}
-
 	int32 OutputCodeChunk = Compiler->SubstrateHairBSDF(
 		CompileWithDefaultFloat3(Compiler, BaseColor, 0.0f, 0.0f, 0.0f),
 		CompileWithDefaultFloat1(Compiler, Scatter, 0.0f),
@@ -2102,7 +2093,7 @@ int32 UMaterialExpressionSubstrateHairBSDF::Compile(class FMaterialCompiler* Com
 		CompileWithDefaultFloat3(Compiler, EmissiveColor, 0.0f, 0.0f, 0.0f),
 		TangentCodeChunk,
 		Compiler->GetSubstrateSharedLocalBasisIndexMacro(NewRegisteredSharedLocalBasis),
-		&SubstrateOperator);
+		!SubstrateOperator.bUseParameterBlending || (SubstrateOperator.bUseParameterBlending && SubstrateOperator.bRootOfParameterBlendingSubTree) ? &SubstrateOperator : nullptr);
 
 	return OutputCodeChunk;
 }
@@ -2452,15 +2443,6 @@ int32 UMaterialExpressionSubstrateToonBSDF::Compile(class FMaterialCompiler* Com
 	FSubstrateOperator& SubstrateOperator = Compiler->SubstrateCompilationGetOperator(Compiler->SubstrateTreeStackGetPathUniqueId());
 	SubstrateOperator.BSDFRegisteredSharedLocalBasis = NewRegisteredSharedLocalBasis;
 
-	if (SubstrateOperator.bUseParameterBlending)
-	{
-		return Compiler->Errorf(TEXT("Substrate Toon BSDF node cannot be used with parameter blending."));
-	}
-	else if (SubstrateOperator.bRootOfParameterBlendingSubTree)
-	{
-		return Compiler->Errorf(TEXT("Substrate Toon BSDF node cannot be the root of a parameter blending sub tree."));
-	}
-
 	int32 ToonProfileIdCodeChunk = INDEX_NONE;
 	if (ToonProfile != nullptr)
 	{
@@ -2481,7 +2463,7 @@ int32 UMaterialExpressionSubstrateToonBSDF::Compile(class FMaterialCompiler* Com
 		CompileWithDefaultFloat3(Compiler, EmissiveColor, 0.0f, 0.0f, 0.0f),
 		NormalCodeChunk,
 		Compiler->GetSubstrateSharedLocalBasisIndexMacro(NewRegisteredSharedLocalBasis),
-		&SubstrateOperator);
+		!SubstrateOperator.bUseParameterBlending || (SubstrateOperator.bUseParameterBlending && SubstrateOperator.bRootOfParameterBlendingSubTree) ? &SubstrateOperator : nullptr);
 
 	return OutputCodeChunk;
 }
diff --git a/Engine/Source/Runtime/Engine/Public/Rendering/SubstrateMaterialShared.h b/Engine/Source/Runtime/Engine/Public/Rendering/SubstrateMaterialShared.h
--- a/Engine/Source/Runtime/Engine/Public/Rendering/SubstrateMaterialShared.h
+++ b/Engine/Source/Runtime/Engine/Public/Rendering/SubstrateMaterialShared.h
@@ -175,6 +175,7 @@ public:
 #define SUBSTRATE_MATERIAL_TYPE_UI					9
 #define SUBSTRATE_MATERIAL_TYPE_DECAL				10
 #define SUBSTRATE_MATERIAL_TYPE_TOON				11
+#define SUBSTRATE_MATERIAL_TYPE_MULTIPLEBSDFS		12
 
 struct FSubstrateMaterialCompilationOutput
 {
