Fix clear coat rendering on Mobile when using Substrate with Substrate's ClearCoat node.

Mobile renderer marks shading model using stencil value. This requires the CPU shading model value to be correct value. The Substrate's ClearCoat node was settings its shading model value to DefaultLit instead of ClearCoat, causing incorrect stencil writting.

#Synced-CL 48602645

diff --git a/Engine/Source/Runtime/Engine/Private/Materials/MaterialExpressionSubstrate.cpp b/Engine/Source/Runtime/Engine/Private/Materials/MaterialExpressionSubstrate.cpp
--- a/Engine/Source/Runtime/Engine/Private/Materials/MaterialExpressionSubstrate.cpp
+++ b/Engine/Source/Runtime/Engine/Private/Materials/MaterialExpressionSubstrate.cpp
@@ -1477,7 +1477,7 @@ void UMaterialExpressionSubstrateSimpleClearCoatBSDF::GatherSubstrateMaterialInf
 	if (Normal.IsConnected()) { SubstrateMaterialInfo.AddPropertyConnected(MP_Normal); }
 	if (EmissiveColor.IsConnected()) { SubstrateMaterialInfo.AddPropertyConnected(MP_EmissiveColor); }
 
-	SubstrateMaterialInfo.AddShadingModel(SSM_DefaultLit);
+	SubstrateMaterialInfo.AddShadingModel(SSM_ClearCoat);
 	SubstrateMaterialInfo.AddGuid(MaterialExpressionGuid);
 }
 
