Convert USceneComponent::AttachParent from TObjectPtr to TWeakObjectPtr to handle the case where a component get destroyed, marking it as garbage, but some other components still have their AttachParent set to the destroyed component without being actually attached.
This worked before "kill pending kill" where the GC would null-out references to pending kill objects, but now the GC reports it as a reachable garbage reference instead.
Converting the pointer to a weak ref instead is closer to the expected behavio, without having to do expensive cleanup when destroying a component.

#Synced-CL 48612369

diff --git a/Engine/Source/Runtime/Engine/Classes/Components/SceneComponent.h b/Engine/Source/Runtime/Engine/Classes/Components/SceneComponent.h
--- a/Engine/Source/Runtime/Engine/Classes/Components/SceneComponent.h
+++ b/Engine/Source/Runtime/Engine/Classes/Components/SceneComponent.h
@@ -106,7 +106,7 @@ private:
 
 	/** What we are currently attached to. If valid, RelativeLocation etc. are used relative to this object */
 	UPROPERTY(ReplicatedUsing = OnRep_AttachParent)
-	TObjectPtr<USceneComponent> AttachParent;
+	TWeakObjectPtr<USceneComponent> AttachParent;
 
 	/** Optional socket name on AttachParent that we are attached to. */
 	UPROPERTY(ReplicatedUsing = OnRep_AttachSocketName)
@@ -1696,7 +1696,7 @@ inline const TArray<TObjectPtr<USceneComponent>>& USceneComponent::GetAttachChil
 
 inline USceneComponent* USceneComponent::GetAttachParent() const
 {
-	return AttachParent;
+	return AttachParent.GetEvenIfUnreachable();
 }
 
 inline FName USceneComponent::GetAttachSocketName() const
diff --git a/Engine/Source/Runtime/Engine/Private/Components/SceneComponent.cpp b/Engine/Source/Runtime/Engine/Private/Components/SceneComponent.cpp
--- a/Engine/Source/Runtime/Engine/Private/Components/SceneComponent.cpp
+++ b/Engine/Source/Runtime/Engine/Private/Components/SceneComponent.cpp
@@ -942,7 +942,7 @@ void USceneComponent::EndPlay(EEndPlayReason::Type Reason)
 		// Detach components which are in different streaming levels so that this level can be properly garbage collected.
 		// Note that we explicitly want to check the outer hierarchy and not the owning package because we want references that participate in GC.
 		UObject* Outermost = GetOutermostObject();
-		if (AttachParent && AttachParent->GetOutermostObject() != Outermost)
+		if (GetAttachParent() && GetAttachParent()->GetOutermostObject() != Outermost)
 		{
 			DetachFromComponent(FDetachmentTransformRules::KeepWorldTransform);
 		}
@@ -3784,7 +3784,7 @@ void USceneComponent::OnRep_AttachParent()
 
 void USceneComponent::OnRep_AttachSocketName()
 {
-	if (IsValid(AttachParent))
+	if (GetAttachParent())
 	{
 		bNetUpdateAttachment = true;
 	}
@@ -3899,7 +3899,10 @@ void USceneComponent::PostRepNotifies()
 {
 	if (bNetUpdateAttachment)
 	{
-		Exchange(NetOldAttachParent, AttachParent);
+		USceneComponent* NetOldAttachParentSwapTemp = NetOldAttachParent;
+		NetOldAttachParent = GetAttachParent();
+		AttachParent = NetOldAttachParentSwapTemp;
+
 		Exchange(NetOldAttachSocketName, AttachSocketName);
 		
 		// Note: This is a local fix for JIRA UE-43355.
@@ -3917,7 +3920,7 @@ void USceneComponent::PostRepNotifies()
 		}
 
 		// Check if this is a detach
-		if (AttachParent && !bShouldBeAttached)
+		if (GetAttachParent() && !bShouldBeAttached)
 		{
 			ensureMsgf(NetOldAttachParent == nullptr, TEXT("Local modification of AttachParent detected for replicated component %s, disable replication or execute detachment on host."), *GetFullName());
 			DetachFromComponent(FDetachmentTransformRules::KeepWorldTransform);
