Don't gather internal actor descs when registering World Partition's main Actor Desc Container Instance

diff --git a/Engine/Source/Runtime/Engine/Private/WorldPartition/ActorDescContainer.cpp b/Engine/Source/Runtime/Engine/Private/WorldPartition/ActorDescContainer.cpp
--- a/Engine/Source/Runtime/Engine/Private/WorldPartition/ActorDescContainer.cpp
+++ b/Engine/Source/Runtime/Engine/Private/WorldPartition/ActorDescContainer.cpp
@@ -96,6 +96,7 @@ void UActorDescContainer::Initialize(const FInitializeParams& InitParams)
 		}
 
 		// Gather non-external actors
+		if (InitParams.bShouldGatherInternalActors)
 		{
 			TRACE_CPUPROFILER_EVENT_SCOPE(GetInternalAssets);
 
diff --git a/Engine/Source/Runtime/Engine/Private/WorldPartition/ActorDescContainerInstance.cpp b/Engine/Source/Runtime/Engine/Private/WorldPartition/ActorDescContainerInstance.cpp
--- a/Engine/Source/Runtime/Engine/Private/WorldPartition/ActorDescContainerInstance.cpp
+++ b/Engine/Source/Runtime/Engine/Private/WorldPartition/ActorDescContainerInstance.cpp
@@ -38,6 +38,7 @@ void UActorDescContainerInstance::RegisterContainer(const FInitializeParams& InP
 	ContainerInitParams.ContentBundleGuid = InParams.ContentBundleGuid;
 	ContainerInitParams.ExternalDataLayerAsset = InParams.ExternalDataLayerAsset;
 	ContainerInitParams.bShouldRegisterEditorDeletages = InParams.bShouldRegisterEditorDeletages;
+	ContainerInitParams.bShouldGatherInternalActors = InParams.bShouldGatherInternalActors;
 
 	SetContainer(UActorDescContainerSubsystem::GetChecked().RegisterContainer(ContainerInitParams));
 }
diff --git a/Engine/Source/Runtime/Engine/Private/WorldPartition/WorldPartition.cpp b/Engine/Source/Runtime/Engine/Private/WorldPartition/WorldPartition.cpp
--- a/Engine/Source/Runtime/Engine/Private/WorldPartition/WorldPartition.cpp
+++ b/Engine/Source/Runtime/Engine/Private/WorldPartition/WorldPartition.cpp
@@ -784,6 +784,7 @@ void UWorldPartition::Initialize(UWorld* InWorld, const FTransform& InTransform)
 	{
 		UActorDescContainerInstance::FInitializeParams ContainerInitParams(UActorDescContainerInstance::GetContainerPackageNameFromWorld(OuterWorld));
 		ContainerInitParams.SetShouldRegisterEditorDeletages(bIsEditor);
+		ContainerInitParams.bShouldGatherInternalActors = false;
 
 		ActorDescContainerInstance = RegisterActorDescContainerInstance(ContainerInitParams);
 
@@ -2237,6 +2238,7 @@ UActorDescContainerInstance* UWorldPartition::RegisterActorDescContainerInstance
 		InitParams.ContentBundleGuid = InParams.ContentBundleGuid;
 		InitParams.ExternalDataLayerAsset = InParams.ExternalDataLayerAsset;
 		InitParams.bShouldRegisterEditorDeletages = InParams.bShouldRegisterEditorDeletages;
+		InitParams.bShouldGatherInternalActors = InParams.bShouldGatherInternalActors;
 		
 		const FWorldDataLayersActorDesc* WorldDataLayerActorsDesc = nullptr;
 		InitParams.FilterActorDescFunc = [this, &WorldDataLayerActorsDesc, &InParams](const FWorldPartitionActorDesc* ActorDesc)
diff --git a/Engine/Source/Runtime/Engine/Public/WorldPartition/ActorDescContainerInitParams.h b/Engine/Source/Runtime/Engine/Public/WorldPartition/ActorDescContainerInitParams.h
--- a/Engine/Source/Runtime/Engine/Public/WorldPartition/ActorDescContainerInitParams.h
+++ b/Engine/Source/Runtime/Engine/Public/WorldPartition/ActorDescContainerInitParams.h
@@ -40,6 +40,9 @@ struct FActorDescContainerInitParams
 
 	/** If the container should bind to editor events */
 	bool bShouldRegisterEditorDeletages = true;
+	
+	/** Whether non-external actors should be gathered */
+	bool bShouldGatherInternalActors = true;
 
 	/** The associated External Data Layer Asset */
 	const UExternalDataLayerAsset* ExternalDataLayerAsset = nullptr;
diff --git a/Engine/Source/Runtime/Engine/Public/WorldPartition/ActorDescContainerInstance.h b/Engine/Source/Runtime/Engine/Public/WorldPartition/ActorDescContainerInstance.h
--- a/Engine/Source/Runtime/Engine/Public/WorldPartition/ActorDescContainerInstance.h
+++ b/Engine/Source/Runtime/Engine/Public/WorldPartition/ActorDescContainerInstance.h
@@ -80,6 +80,7 @@ public:
 
 		bool bCreateContainerInstanceHierarchy;
 		bool bShouldRegisterEditorDeletages = true;
+		bool bShouldGatherInternalActors = true;
 
 		/* Custom filter function used to filter actors descriptors. */
 		TUniqueFunction<bool(const FWorldPartitionActorDesc*)> FilterActorDescFunc;
