Fix TaskSyncManager to correctly shutdown PIE worlds as they don't go through the normal DestroyWorldContext flow

#Synced-CL 48620592

diff --git a/Engine/Source/Runtime/Engine/Private/TaskSyncManager.cpp b/Engine/Source/Runtime/Engine/Private/TaskSyncManager.cpp
--- a/Engine/Source/Runtime/Engine/Private/TaskSyncManager.cpp
+++ b/Engine/Source/Runtime/Engine/Private/TaskSyncManager.cpp
@@ -7,6 +7,10 @@
 #include "Engine/World.h"
 #include "Misc/ScopeLock.h"
 
+#if WITH_EDITOR
+#include "Editor.h"
+#endif
+
 #include UE_INLINE_GENERATED_CPP_BY_NAME(TaskSyncManager)
 
 DEFINE_LOG_CATEGORY_STATIC(LogTaskSync, Log, All);
@@ -969,6 +973,9 @@ FTaskSyncManager::FTaskSyncManager()
 		GEngine->OnWorldContextDestroyed().AddRaw(this, &FTaskSyncManager::OnWorldContextRemove);
 	}
 	FWorldDelegates::OnPostWorldInitialization.AddRaw(this, &FTaskSyncManager::OnPostWorldInitialization);
+#if WITH_EDITOR	
+	FEditorDelegates::EndPIE.AddRaw(this, &FTaskSyncManager::EndPIE);
+#endif
 }
 
 FTaskSyncManager::~FTaskSyncManager()
@@ -978,6 +985,9 @@ FTaskSyncManager::~FTaskSyncManager()
 		GEngine->OnWorldContextDestroyed().RemoveAll(this);
 	}
 	FWorldDelegates::OnPostWorldInitialization.RemoveAll(this);
+#if WITH_EDITOR	
+	FEditorDelegates::EndPIE.RemoveAll(this);
+#endif
 }
 
 void FTaskSyncManager::RegisterSyncPointInternal(const FSyncPointDescription& InDescription)
@@ -1448,6 +1458,20 @@ void FTaskSyncManager::OnWorldContextRemove(FWorldContext& InWorldContext)
 	}
 }
 
+#if WITH_EDITOR
+void FTaskSyncManager::EndPIE(bool bStartSimulate)
+{
+	// PIE shutdown doesn't properly call OnWorldContextDestroyed, so iterate the list ourselves
+	for (const FWorldContext& Context : GEngine->GetWorldContexts())
+	{
+		if (Context.WorldType == EWorldType::PIE)
+		{
+			OnWorldContextRemove(const_cast<FWorldContext&>(Context));
+		}
+	}
+}
+#endif
+
 FWorldContextId FTaskSyncManager::GetCurrentWorldContext() const
 {
 	// Does this make sense to be different per thread? Looking up the world is slow and unnecessary in cooked
diff --git a/Engine/Source/Runtime/Engine/Public/TaskSyncManager.h b/Engine/Source/Runtime/Engine/Public/TaskSyncManager.h
--- a/Engine/Source/Runtime/Engine/Public/TaskSyncManager.h
+++ b/Engine/Source/Runtime/Engine/Public/TaskSyncManager.h
@@ -969,6 +969,9 @@ private:
 	void InitializeForWorld(FWorldContextId WorldContext);
 	void OnPostWorldInitialization(UWorld* World, const UWorld::InitializationValues IVS);
 	void InitializeBatchForNewWorld(FBatchData& BatchData);
+#if WITH_EDITOR
+	void EndPIE(bool bStartSimulate);
+#endif
 
 	FActiveSyncPoint* GetOrCreateSyncPoint(FBatchData& BatchData, FRegisteredSyncPointData& SyncData);
 	FActiveSyncPoint* GetOrCreateTickGroupWork(FBatchData& BatchData, ETickingGroup TickGroup);
