Restore static mesh UV element stream setup due to issues with reading half2 from half4 declared element stream on certain platforms

diff --git a/Engine/Plugins/FX/Niagara/Source/NiagaraVertexFactories/Private/NiagaraMeshVertexFactory.cpp b/Engine/Plugins/FX/Niagara/Source/NiagaraVertexFactories/Private/NiagaraMeshVertexFactory.cpp
--- a/Engine/Plugins/FX/Niagara/Source/NiagaraVertexFactories/Private/NiagaraMeshVertexFactory.cpp
+++ b/Engine/Plugins/FX/Niagara/Source/NiagaraVertexFactories/Private/NiagaraMeshVertexFactory.cpp
@@ -128,8 +128,7 @@ void FNiagaraMeshVertexFactory::GetVertexElements(ERHIFeatureLevel::Type Feature
 		if (Data.TextureCoordinates.Num())
 		{
 			const int32 BaseTexCoordAttribute = 4;
-			uint8 TexCoordStreamIndex = AddVertexStream(Data.TextureCoordinates[0], InOutStreams);
-			AddStaticMeshTextureCoordinateElements(TexCoordStreamIndex, BaseTexCoordAttribute, Data.TextureCoordinates, Elements);
+			AddStaticMeshTextureCoordinateElements(BaseTexCoordAttribute, Data.TextureCoordinates, Elements, InOutStreams);
 		}
 	}
 }
diff --git a/Engine/Source/Runtime/Engine/Private/GPUSkinVertexFactory.cpp b/Engine/Source/Runtime/Engine/Private/GPUSkinVertexFactory.cpp
--- a/Engine/Source/Runtime/Engine/Private/GPUSkinVertexFactory.cpp
+++ b/Engine/Source/Runtime/Engine/Private/GPUSkinVertexFactory.cpp
@@ -915,8 +915,7 @@ void TGPUSkinVertexFactory<BoneInfluenceType>::GetVertexElements(ERHIFeatureLeve
 	if (GPUSkinData.TextureCoordinates.Num())
 	{
 		const uint8 BaseTexCoordAttribute = 5;				
-		uint8 TexCoordStreamIndex = AddVertexStream(GPUSkinData.TextureCoordinates[0], InOutStreams);
-		AddStaticMeshTextureCoordinateElements(TexCoordStreamIndex, BaseTexCoordAttribute, GPUSkinData.TextureCoordinates, OutElements);
+		AddStaticMeshTextureCoordinateElements(BaseTexCoordAttribute, GPUSkinData.TextureCoordinates, OutElements, InOutStreams);
 	}
 
 	if (GPUSkinData.ColorComponentsSRV == nullptr)
diff --git a/Engine/Source/Runtime/Engine/Private/GeometryCacheVertexFactory.cpp b/Engine/Source/Runtime/Engine/Private/GeometryCacheVertexFactory.cpp
--- a/Engine/Source/Runtime/Engine/Private/GeometryCacheVertexFactory.cpp
+++ b/Engine/Source/Runtime/Engine/Private/GeometryCacheVertexFactory.cpp
@@ -231,8 +231,7 @@ void FGeometryCacheVertexVertexFactory::GetVertexElements(
 	if (StreamComponentData.TextureCoordinates.Num())
 	{
 		const int32 BaseTexCoordAttribute = 5;	
-		uint8 TexCoordStreamIndex = AddVertexStream(StreamComponentData.TextureCoordinates[0], InOutStreams);
-		AddStaticMeshTextureCoordinateElements(TexCoordStreamIndex, BaseTexCoordAttribute, StreamComponentData.TextureCoordinates, OutElements);
+		AddStaticMeshTextureCoordinateElements(BaseTexCoordAttribute, StreamComponentData.TextureCoordinates, OutElements, InOutStreams);
 	}
 }
 
diff --git a/Engine/Source/Runtime/Engine/Private/GeometryCollectionRendering.cpp b/Engine/Source/Runtime/Engine/Private/GeometryCollectionRendering.cpp
--- a/Engine/Source/Runtime/Engine/Private/GeometryCollectionRendering.cpp
+++ b/Engine/Source/Runtime/Engine/Private/GeometryCollectionRendering.cpp
@@ -212,8 +212,7 @@ void FGeometryCollectionVertexFactory::InitRHI(FRHICommandListBase& RHICmdList)
 	if (Data.TextureCoordinates.Num())
 	{
 		const int32 BaseTexCoordAttribute = 4;			
-		uint8 TexCoordStreamIndex = AddVertexStream(Data.TextureCoordinates[0], Streams);
-		AddStaticMeshTextureCoordinateElements(TexCoordStreamIndex, BaseTexCoordAttribute, Data.TextureCoordinates, Elements);
+		AddStaticMeshTextureCoordinateElements(BaseTexCoordAttribute, Data.TextureCoordinates, Elements, Streams);
 	}
 
 	if (IsStaticLightingAllowed())
diff --git a/Engine/Source/Runtime/Engine/Private/InstancedStaticMesh.cpp b/Engine/Source/Runtime/Engine/Private/InstancedStaticMesh.cpp
--- a/Engine/Source/Runtime/Engine/Private/InstancedStaticMesh.cpp
+++ b/Engine/Source/Runtime/Engine/Private/InstancedStaticMesh.cpp
@@ -854,8 +854,7 @@ void FInstancedStaticMeshVertexFactory::GetVertexElements(
 		if (Data.TextureCoordinates.Num())
 		{
 			const int32 BaseTexCoordAttribute = 4;				
-			uint8 TexCoordStreamIndex = AddVertexStream(Data.TextureCoordinates[0], Streams);
-			AddStaticMeshTextureCoordinateElements(TexCoordStreamIndex, BaseTexCoordAttribute, Data.TextureCoordinates, Elements);
+			AddStaticMeshTextureCoordinateElements(BaseTexCoordAttribute, Data.TextureCoordinates, Elements, Streams);
 		}
 
 		// PreSkinPosition attribute is only used for GPUSkinPassthrough variation of local vertex factory.
diff --git a/Engine/Source/Runtime/Engine/Private/LocalVertexFactory.cpp b/Engine/Source/Runtime/Engine/Private/LocalVertexFactory.cpp
--- a/Engine/Source/Runtime/Engine/Private/LocalVertexFactory.cpp
+++ b/Engine/Source/Runtime/Engine/Private/LocalVertexFactory.cpp
@@ -507,48 +507,27 @@ void FLocalVertexFactory::InitRHI(FRHICommandListBase& RHICmdList)
 }
 
 void AddStaticMeshTextureCoordinateElements(
-	uint8 TexCoordStreamIndex, 
 	int32 BaseTexCoordAttribute, 
 	const TArrayView<FVertexStreamComponent>& TextureCoordinates, 
-	FVertexDeclarationElementList& Elements)
+	FVertexDeclarationElementList& Elements,
+	FVertexFactory::FVertexStreamList& InOutStreams)
 {
-	if (TextureCoordinates.IsEmpty())
+	for (int32 CoordinateIndex = 0; CoordinateIndex < TextureCoordinates.Num(); ++CoordinateIndex)
 	{
-		return; 
+		Elements.Add(FVertexFactory::AccessStreamComponent(
+			TextureCoordinates[CoordinateIndex],
+			BaseTexCoordAttribute + CoordinateIndex,
+			InOutStreams
+		));
 	}
 
-	// Validate all texture coordinates share the same stream data
-	for (int32 TextureCoordIndex = 1; TextureCoordIndex < TextureCoordinates.Num(); ++TextureCoordIndex)
+	for (int32 CoordinateIndex = TextureCoordinates.Num(); CoordinateIndex < MAX_STATIC_TEXCOORDS / 2; ++CoordinateIndex)
 	{
-		check(TextureCoordinates[0].VertexBuffer == TextureCoordinates[TextureCoordIndex].VertexBuffer);
-		check(TextureCoordinates[0].Stride == TextureCoordinates[TextureCoordIndex].Stride);
-		check(TextureCoordinates[0].StreamOffset == TextureCoordinates[TextureCoordIndex].StreamOffset);
-		check(TextureCoordinates[0].VertexStreamUsage == TextureCoordinates[TextureCoordIndex].VertexStreamUsage);
-	}
-
-	EVertexElementType FirstElementType = TextureCoordinates[0].Type;
-			
-	EVertexElementType UVDoubleWideVertexElementType = VET_None;
-	uint32 DoubleUVSizeInBytes = 0;
-	if (FirstElementType == VET_Float2 || FirstElementType == VET_Float4)
-	{
-		DoubleUVSizeInBytes = sizeof(TStaticMeshVertexUVsTypeSelector<EStaticMeshVertexUVType::HighPrecision>::UVsTypeT) * 2;
-		UVDoubleWideVertexElementType = VET_Float4;
-	}
-	else
-	{
-		DoubleUVSizeInBytes = sizeof(TStaticMeshVertexUVsTypeSelector<EStaticMeshVertexUVType::Default>::UVsTypeT) * 2;
-		UVDoubleWideVertexElementType = VET_Half4;
-	}
-			
-	for (int32 CoordinateIndex = 0; CoordinateIndex < MAX_STATIC_TEXCOORDS / 2; ++CoordinateIndex)
-	{
-		Elements.Add(FVertexElement(
-			TexCoordStreamIndex, 
-			DoubleUVSizeInBytes * CoordinateIndex, 
-			UVDoubleWideVertexElementType, 
-			BaseTexCoordAttribute + CoordinateIndex, 
-			TextureCoordinates[0].Stride));
+		Elements.Add(FVertexFactory::AccessStreamComponent(
+			TextureCoordinates[TextureCoordinates.Num() - 1],
+			BaseTexCoordAttribute + CoordinateIndex,
+			InOutStreams
+		));
 	}
 }
 
@@ -608,8 +587,7 @@ void FLocalVertexFactory::GetVertexElements(
 		if (Data.TextureCoordinates.Num())
 		{
 			const int32 BaseTexCoordAttribute = 4;
-			uint8 TexCoordStreamIndex = AddVertexStream(Data.TextureCoordinates[0], InOutStreams);
-			AddStaticMeshTextureCoordinateElements(TexCoordStreamIndex, BaseTexCoordAttribute, Data.TextureCoordinates, Elements);
+			AddStaticMeshTextureCoordinateElements(BaseTexCoordAttribute, Data.TextureCoordinates, Elements, InOutStreams);
 		}
 
 		// TODO: should also check if VFType supports 'SupportsGPUSkinPassThrough'
diff --git a/Engine/Source/Runtime/Engine/Private/MeshParticleVertexFactory.cpp b/Engine/Source/Runtime/Engine/Private/MeshParticleVertexFactory.cpp
--- a/Engine/Source/Runtime/Engine/Private/MeshParticleVertexFactory.cpp
+++ b/Engine/Source/Runtime/Engine/Private/MeshParticleVertexFactory.cpp
@@ -174,8 +174,7 @@ void FMeshParticleVertexFactory::GetVertexElements(ERHIFeatureLevel::Type Featur
 	if(Data.TextureCoordinates.Num())
 	{
 		const int32 BaseTexCoordAttribute = 4;		
-		uint8 TexCoordStreamIndex = AddVertexStream(Data.TextureCoordinates[0], InOutStreams);
-		AddStaticMeshTextureCoordinateElements(TexCoordStreamIndex, BaseTexCoordAttribute, Data.TextureCoordinates, Elements);
+		AddStaticMeshTextureCoordinateElements(BaseTexCoordAttribute, Data.TextureCoordinates, Elements, InOutStreams);
 	}
 }
 
diff --git a/Engine/Source/Runtime/Engine/Public/LocalVertexFactory.h b/Engine/Source/Runtime/Engine/Public/LocalVertexFactory.h
--- a/Engine/Source/Runtime/Engine/Public/LocalVertexFactory.h
+++ b/Engine/Source/Runtime/Engine/Public/LocalVertexFactory.h
@@ -46,10 +46,10 @@ extern ENGINE_API TUniformBufferRef<FLocalVertexFactoryUniformShaderParameters>
 	);
 
 extern ENGINE_API void AddStaticMeshTextureCoordinateElements(
-	uint8 TexCoordStreamIndex, 
 	int32 BaseTexCoordAttribute, 
 	const TArrayView<FVertexStreamComponent>& TextureCoordinates, 
-	FVertexDeclarationElementList& Elements);
+	FVertexDeclarationElementList& Elements,
+	FVertexFactory::FVertexStreamList& InOutStreams);
 
 /**
  * A vertex factory which simply transforms explicit vertex attributes from local to world space.
diff --git a/Engine/Source/Runtime/RenderCore/Public/VertexFactory.h b/Engine/Source/Runtime/RenderCore/Public/VertexFactory.h
--- a/Engine/Source/Runtime/RenderCore/Public/VertexFactory.h
+++ b/Engine/Source/Runtime/RenderCore/Public/VertexFactory.h
@@ -778,6 +778,14 @@ protected:
 	 */
 	RENDERCORE_API FVertexElement AccessStreamComponent(const FVertexStreamComponent& Component, uint8 AttributeIndex, EVertexInputStreamType InputStreamType);
 	
+	/**
+	* Initializes the vertex declaration.
+	* @param Elements - The elements of the vertex declaration.
+		*/
+	RENDERCORE_API void InitDeclaration(const FVertexDeclarationElementList& Elements, EVertexInputStreamType StreamType = EVertexInputStreamType::Default);
+
+public:
+
 	/**
 	 * Creates a vertex element for a vertex stream components.  Adds a unique stream index for the vertex buffer used by the component.
 	 * @param Component - The vertex stream component.
@@ -803,12 +811,6 @@ protected:
 		return (uint8)InOutStreams.AddUnique(VertexStream);
 	}
 
-	/**
-	 * Initializes the vertex declaration.
-	 * @param Elements - The elements of the vertex declaration.
-	 */
-	RENDERCORE_API void InitDeclaration(const FVertexDeclarationElementList& Elements, EVertexInputStreamType StreamType = EVertexInputStreamType::Default);
-
 	/**
 	 * Information needed to set a vertex stream.
 	 */
@@ -832,6 +834,8 @@ protected:
 
 	typedef TArray<FVertexStream, TInlineAllocator<8> > FVertexStreamList;
 	
+protected:
+
 	/** The vertex streams used to render the factory. */
 	FVertexStreamList Streams;
 	
