Fix internationalization of the Tilde key on non-US keyboards

- Change to GetKeyboardLayoutName to get the keyboard layout where GetKeyboardLayout would not work
- Relax the check on the number of console keys so that an additional console key does not disable internationalization of Tilde
- Update config to remove unneeded overrides

#rnx

#Synced-CL 48600773

diff --git a/Engine/Source/Runtime/Engine/Private/UserInterface/InputSettings.cpp b/Engine/Source/Runtime/Engine/Private/UserInterface/InputSettings.cpp
--- a/Engine/Source/Runtime/Engine/Private/UserInterface/InputSettings.cpp
+++ b/Engine/Source/Runtime/Engine/Private/UserInterface/InputSettings.cpp
@@ -155,43 +155,51 @@ void UInputSettings::AddInternationalConsoleKey()
 {
 #if PLATFORM_WINDOWS
 	// If the console key is set to the default we'll see about adding the keyboard default
-	// If they've mapped any additional keys, we'll just assume they've set it up in a way they desire
-	if (ConsoleKeys.Num() == 1 && ConsoleKeys[0] == EKeys::Tilde)
+	if (ConsoleKeys.Num() >= 1 && ConsoleKeys[0] == EKeys::Tilde)
 	{
-		FKey DefaultConsoleKey = EKeys::Tilde;
-		switch (PRIMARYLANGID(LOWORD(GetKeyboardLayout(0))))
-		{
-		case LANG_FRENCH:
-		case LANG_HUNGARIAN:
-			DefaultConsoleKey = FInputKeyManager::Get().GetKeyFromCodes(VK_OEM_7, 0);
-			break;
-
-		case LANG_GERMAN:
-			DefaultConsoleKey = EKeys::Caret;
-			break;
-
-		case LANG_ITALIAN:
-			DefaultConsoleKey = EKeys::Backslash;
-			break;
-
-		case LANG_SPANISH:
-			DefaultConsoleKey = FInputKeyManager::Get().GetKeyFromCodes(VK_OEM_5, 0);
-			break;
-
-		case LANG_SLOVAK:
-		case LANG_SWEDISH:
-			DefaultConsoleKey = EKeys::Section;
-			break;
-
-		case LANG_JAPANESE:
-		case LANG_RUSSIAN:
-			DefaultConsoleKey = FInputKeyManager::Get().GetKeyFromCodes(VK_OEM_3, 0);
-			break;
-		}
+		CHAR KeyboardLayoutName[KL_NAMELENGTH];
 
-		if (DefaultConsoleKey != EKeys::Tilde && DefaultConsoleKey.IsValid())
+		// We need to use this call because it gives the actual layout, rather than the locale
+		if (GetKeyboardLayoutNameA(KeyboardLayoutName))
 		{
-			ConsoleKeys.Add(DefaultConsoleKey);
+			FKey DefaultConsoleKey = EKeys::Tilde;
+
+			const int32 KeyboardLayoutId = strtol(KeyboardLayoutName, nullptr, 16);
+
+			switch (PRIMARYLANGID(LOWORD(KeyboardLayoutId)))
+			{
+			case LANG_FRENCH:
+			case LANG_HUNGARIAN:
+				DefaultConsoleKey = FInputKeyManager::Get().GetKeyFromCodes(VK_OEM_7, 0);
+				break;
+
+			case LANG_GERMAN:
+				DefaultConsoleKey = EKeys::Caret;
+				break;
+
+			case LANG_ITALIAN:
+				DefaultConsoleKey = EKeys::Backslash;
+				break;
+
+			case LANG_SPANISH:
+				DefaultConsoleKey = FInputKeyManager::Get().GetKeyFromCodes(VK_OEM_5, 0);
+				break;
+
+			case LANG_SLOVAK:
+			case LANG_SWEDISH:
+				DefaultConsoleKey = EKeys::Section;
+				break;
+
+			case LANG_JAPANESE:
+			case LANG_RUSSIAN:
+				DefaultConsoleKey = FInputKeyManager::Get().GetKeyFromCodes(VK_OEM_3, 0);
+				break;
+			}
+
+			if (DefaultConsoleKey != EKeys::Tilde && DefaultConsoleKey.IsValid())
+			{
+				ConsoleKeys.Add(DefaultConsoleKey);
+			}
 		}
 	}
 #endif
