When setting overlay materials via blueprints, force the MaterialSlotsOverlayMaterial array to populate enough entries if the base mesh overlay count allows for the requested slot number.

#Synced-CL 48656912

diff --git a/Engine/Source/Runtime/Engine/Private/Components/MeshComponent.cpp b/Engine/Source/Runtime/Engine/Private/Components/MeshComponent.cpp
--- a/Engine/Source/Runtime/Engine/Private/Components/MeshComponent.cpp
+++ b/Engine/Source/Runtime/Engine/Private/Components/MeshComponent.cpp
@@ -324,7 +324,22 @@ void UMeshComponent::SetOverlayMaterial(UMaterialInterface* NewOverlayMaterial,
 
 	if (bSetMaterialSlot)
 	{
-		if (MaterialSlotsOverlayMaterial.IsValidIndex(SlotIndex))
+		if (SlotIndex < 0)
+		{
+			return;
+		}
+
+		bool bValid = MaterialSlotsOverlayMaterial.IsValidIndex(SlotIndex);
+		if (!bValid)
+		{
+			if (SlotIndex < GetNumMaterials())
+			{
+				MaterialSlotsOverlayMaterial.SetNum(SlotIndex + 1);
+				bValid = true;
+			}
+		}
+
+		if (bValid)
 		{
 			OverrideOverlaySlot = &MaterialSlotsOverlayMaterial[SlotIndex];
 		}
