Fix substrate blend mode behavior affecting legacy materials.

[RN] minor | Rendering Architecture
#Synced-CL 48625785

diff --git a/Engine/Source/Runtime/Engine/Private/Materials/Material.cpp b/Engine/Source/Runtime/Engine/Private/Materials/Material.cpp
--- a/Engine/Source/Runtime/Engine/Private/Materials/Material.cpp
+++ b/Engine/Source/Runtime/Engine/Private/Materials/Material.cpp
@@ -5107,7 +5107,8 @@ bool UMaterial::CanEditChange(const FProperty* InProperty) const
 		const UMaterialEditorOnlyData* EditorOnly = GetEditorOnlyData();
 		FString PropertyName = InProperty->GetName();
 		const bool bSubstrateEnabled = Substrate::IsSubstrateEnabled();
-		const bool bSubstrateHiddenMaterialAssetConversion = bSubstrateEnabled && Substrate::IsHiddenMaterialAssetConversionEnabled() && !EditorOnly->FrontMaterial.IsConnected();
+		const bool bFrontMaterialConnected = EditorOnly->FrontMaterial.IsConnected();
+		const bool bSubstrateHiddenMaterialAssetConversion = bSubstrateEnabled && Substrate::IsHiddenMaterialAssetConversionEnabled() && !bFrontMaterialConnected;
 
 		if (PropertyName == GET_MEMBER_NAME_STRING_CHECKED(UMaterial, MaterialDomain))
 		{
@@ -5195,7 +5196,7 @@ bool UMaterial::CanEditChange(const FProperty* InProperty) const
 
 		if (PropertyName == GET_MEMBER_NAME_STRING_CHECKED(UMaterial, BlendMode))
 		{
-			if (bSubstrateEnabled)
+			if (bSubstrateEnabled && bFrontMaterialConnected)
 			{
 				return ((MaterialDomain != MD_PostProcess && MaterialDomain != MD_LightFunction && MaterialDomain != MD_Volume) || IsPostProcessMaterialOutputingAlpha());
 			}
