Reset compilation output state before rerunning material translation after a retry is requested.
Also don't use move on the material DDC key string and pass as const ref. We will need again it for any retry.
#rnx
#Synced-CL 48656603

diff --git a/Engine/Source/Runtime/Engine/Private/Materials/HLSLMaterialTranslator.cpp b/Engine/Source/Runtime/Engine/Private/Materials/HLSLMaterialTranslator.cpp
--- a/Engine/Source/Runtime/Engine/Private/Materials/HLSLMaterialTranslator.cpp
+++ b/Engine/Source/Runtime/Engine/Private/Materials/HLSLMaterialTranslator.cpp
@@ -674,7 +674,7 @@ FHLSLMaterialTranslator::FHLSLMaterialTranslator(FMaterial* InMaterial,
 	ERHIFeatureLevel::Type InFeatureLevel,
 	const ITargetPlatform* InTargetPlatform, //if InTargetPlatform is nullptr, we use the current active
 	const FSubstrateCompilationConfig* InSubstrateCompilationConfig,
-	FString MaterialTranslationDDCKeyString)
+	const FString& MaterialTranslationDDCKeyString)
 :	ShaderFrequency(SF_Pixel)
 ,	MaterialProperty(MP_EmissiveColor)
 ,	CurrentScopeChunks(nullptr)
diff --git a/Engine/Source/Runtime/Engine/Private/Materials/HLSLMaterialTranslator.h b/Engine/Source/Runtime/Engine/Private/Materials/HLSLMaterialTranslator.h
--- a/Engine/Source/Runtime/Engine/Private/Materials/HLSLMaterialTranslator.h
+++ b/Engine/Source/Runtime/Engine/Private/Materials/HLSLMaterialTranslator.h
@@ -670,7 +670,7 @@ public:
 		ERHIFeatureLevel::Type InFeatureLevel,
 		const ITargetPlatform* InTargetPlatform = nullptr,
 		const FSubstrateCompilationConfig* InSubstrateCompilationConfig = nullptr,
-		FString MaterialTranslationDDCKeyString = {});
+		const FString& MaterialTranslationDDCKeyString = {});
 
 	~FHLSLMaterialTranslator();
 
diff --git a/Engine/Source/Runtime/Engine/Private/Materials/MaterialShared.cpp b/Engine/Source/Runtime/Engine/Private/Materials/MaterialShared.cpp
--- a/Engine/Source/Runtime/Engine/Private/Materials/MaterialShared.cpp
+++ b/Engine/Source/Runtime/Engine/Private/Materials/MaterialShared.cpp
@@ -3398,7 +3398,7 @@ bool FMaterial::Translate_Legacy(const FMaterialShaderMapId& ShaderMapId,
 {
 	FString MaterialTranslationDDCKeyString = GetMaterialShaderMapKeyString(ShaderMapId, FMaterialShaderParameters(this), GetShaderPlatform(), false);
 	
-	FHLSLMaterialTranslator MaterialTranslator(this, OutCompilationOutput, InStaticParameters, GetShaderPlatform(), GetQualityLevel(), ShaderMapId.FeatureLevel, InTargetPlatform, &ShaderMapId.SubstrateCompilationConfig, MoveTemp(MaterialTranslationDDCKeyString));
+	FHLSLMaterialTranslator MaterialTranslator(this, OutCompilationOutput, InStaticParameters, GetShaderPlatform(), GetQualityLevel(), ShaderMapId.FeatureLevel, InTargetPlatform, &ShaderMapId.SubstrateCompilationConfig, MaterialTranslationDDCKeyString);
 	EHLSLMaterialTranslatorResult Result = MaterialTranslator.Translate(false);
 
 	// If the DDC result was invalid we need to invoke translation again turning the DDC off.
@@ -3407,8 +3407,9 @@ bool FMaterial::Translate_Legacy(const FMaterialShaderMapId& ShaderMapId,
 		// FHLSLMaterialTranslator is designed as single use. After a call to Translate() no other calls are allowed.
 		// Destruct the current instance and create a new one before translating the material again, forcing the translator
 		// to translate the material instead of accessing the DDC cache.
+		OutCompilationOutput = {};
 		MaterialTranslator.~FHLSLMaterialTranslator();
-		new (&MaterialTranslator) FHLSLMaterialTranslator(this, OutCompilationOutput, InStaticParameters, GetShaderPlatform(), GetQualityLevel(), ShaderMapId.FeatureLevel, InTargetPlatform, &ShaderMapId.SubstrateCompilationConfig, MoveTemp(MaterialTranslationDDCKeyString));
+		new (&MaterialTranslator) FHLSLMaterialTranslator(this, OutCompilationOutput, InStaticParameters, GetShaderPlatform(), GetQualityLevel(), ShaderMapId.FeatureLevel, InTargetPlatform, &ShaderMapId.SubstrateCompilationConfig, MaterialTranslationDDCKeyString);
 		Result = MaterialTranslator.Translate(true);
 	}
 
