Fixing a LWC issue with character movement and trying to transition from falling to landed.

#Synced-CL 48505955

diff --git a/Engine/Source/Runtime/Engine/Private/Components/CharacterMovementComponent.cpp b/Engine/Source/Runtime/Engine/Private/Components/CharacterMovementComponent.cpp
--- a/Engine/Source/Runtime/Engine/Private/Components/CharacterMovementComponent.cpp
+++ b/Engine/Source/Runtime/Engine/Private/Components/CharacterMovementComponent.cpp
@@ -7294,7 +7294,7 @@ bool UCharacterMovementComponent::IsValidLandingSpot(const FVector& CapsuleLocat
 		CharacterOwner->GetCapsuleComponent()->GetScaledCapsuleSize(PawnRadius, PawnHalfHeight);
 
 		// Reject hits that are above our lower hemisphere (can happen when sliding down a vertical surface).
-		const float LowerHemisphereZ = GetGravitySpaceZ(Hit.Location) - PawnHalfHeight + PawnRadius;
+		const FVector::FReal LowerHemisphereZ = GetGravitySpaceZ(Hit.Location) - PawnHalfHeight + PawnRadius;
 		if (GetGravitySpaceZ(Hit.ImpactPoint) >= LowerHemisphereZ)
 		{
 			return false;
