Downgrade build machine shader Warnings to Display to work around some issues until we can investigate further.

#Synced-CL 48662114

diff --git a/Engine/Source/Runtime/Engine/Private/ShaderCompiler/ShaderCompiler.cpp b/Engine/Source/Runtime/Engine/Private/ShaderCompiler/ShaderCompiler.cpp
--- a/Engine/Source/Runtime/Engine/Private/ShaderCompiler/ShaderCompiler.cpp
+++ b/Engine/Source/Runtime/Engine/Private/ShaderCompiler/ShaderCompiler.cpp
@@ -1282,12 +1282,12 @@ void FShaderCompilingManager::CalculateNumberOfCompilingThreads(int32 NumberOfCo
 				}
 				else
 				{
-					UE_LOG(LogShaderCompilers, Warning, TEXT("Machine has %d GBs of RAM, cooker might take %f GBs, but not enough memory left for the OS! (Requested %f GBs for the OS)"), PhysicalGBRam, CookerMemoryUsedInGB, MemoryToLeaveForTheOSInGB);
+					UE_LOG(LogShaderCompilers, Display, TEXT("Machine has %d GBs of RAM, cooker might take %f GBs, but not enough memory left for the OS! (Requested %f GBs for the OS)"), PhysicalGBRam, CookerMemoryUsedInGB, MemoryToLeaveForTheOSInGB);
 				}
 			}
 			else
 			{
-				UE_LOG(LogShaderCompilers, Warning, TEXT("Machine has %d GBs of RAM, but cooker might take %f GBs!"), PhysicalGBRam, CookerMemoryUsedInGB);
+				UE_LOG(LogShaderCompilers, Display, TEXT("Machine has %d GBs of RAM, but cooker might take %f GBs!"), PhysicalGBRam, CookerMemoryUsedInGB);
 			}
 			if (MemoryUsedPerSCWProcessInGB > 0.0f)
 			{
