Substrate - added support for translucency alpha-to-coverage to Toon shading model.

#tests PC Editor dx12
#rnx
#Synced-CL 48664805

diff --git a/Engine/Source/Runtime/Engine/Private/Materials/HLSLMaterialTranslator.cpp b/Engine/Source/Runtime/Engine/Private/Materials/HLSLMaterialTranslator.cpp
--- a/Engine/Source/Runtime/Engine/Private/Materials/HLSLMaterialTranslator.cpp
+++ b/Engine/Source/Runtime/Engine/Private/Materials/HLSLMaterialTranslator.cpp
@@ -13768,7 +13768,7 @@ bool FHLSLMaterialTranslator::FSubstrateCompilationContext::SubstrateGenerateDer
 		MTV_LightFunction			= 0x400,
 
 		MTV_ValidMultipleBSDFs		= (MTV_BSDF_Slab | MTV_BSDF_Hair | MTV_BSDF_Toon),	// SUBSTRATE_TODO Unlit cannot be in there yet due to the specific define for it (SUBSTRATE_OPTIMIZED_UNLIT)
-		MTV_WeightOp_AllowedList	= (MTV_BSDF_Unlit | MTV_BSDF_Slab),
+		MTV_WeightOp_AllowedList	= (MTV_BSDF_Unlit | MTV_BSDF_Slab | MTV_BSDF_Toon),
 		MTV_SelectOp_AllowedList	= (MTV_ValidMultipleBSDFs),							// This only works because parameter blending is enforced so there will always be a single BSDF of a single type only. And hair or toon BSDFs will never pass through other operator allowed list.
 		MTV_HorizoOp_AllowedList	= MTV_BSDF_Slab,
 		MTV_VerticOp_AllowedList	= MTV_BSDF_Slab,
