Re-enables new material translator even when Substrate is turned on on per project.

#rnx
#Synced-CL 48660151

diff --git a/Engine/Source/Runtime/Engine/Private/Materials/MaterialAggregate.cpp b/Engine/Source/Runtime/Engine/Private/Materials/MaterialAggregate.cpp
--- a/Engine/Source/Runtime/Engine/Private/Materials/MaterialAggregate.cpp
+++ b/Engine/Source/Runtime/Engine/Private/Materials/MaterialAggregate.cpp
@@ -38,6 +38,7 @@ EMaterialAggregateAttributeType MaterialValueTypeToMaterialAggregateAttributeTyp
 
 		case MCT_ShadingModel: return EMaterialAggregateAttributeType::ShadingModel;
 		case MCT_MaterialAttributes: return EMaterialAggregateAttributeType::MaterialAttributes;
+		case MCT_Substrate: return EMaterialAggregateAttributeType::SubstrateData;
 
 		default:
 			checkf(false, TEXT("This material value type '%d' cannot be expressed to a material aggregate attribute type."), (int32)Type);
@@ -66,6 +67,8 @@ EMaterialValueType FMaterialAggregateAttribute::ToMaterialValueType() const
 		
 		case EMaterialAggregateAttributeType::ShadingModel: return MCT_ShadingModel;
 		case EMaterialAggregateAttributeType::MaterialAttributes: return MCT_MaterialAttributes;
+		case EMaterialAggregateAttributeType::SubstrateData: return MCT_Substrate;
+
 		case EMaterialAggregateAttributeType::Aggregate:
 			return Aggregate.Get() == UMaterialAggregate::GetMaterialAttributes() ? MCT_MaterialAttributes : MCT_Unknown;
 		
@@ -139,7 +142,7 @@ static FMaterialAggregateAttribute MakeAttributeFromMaterialProperty(EMaterialPr
 	Attribute.Type = MaterialValueTypeToMaterialAggregateAttributeType((Property == MP_SubsurfaceColor) ? MCT_Float4 : FMaterialAttributeDefinitionMap::GetValueType(Property));
 	
 	//ShadingModel defaults to 0, i.e. unlit. This is overriden for the old material system via editor only data, but we can't access that here so we force the value to Lit(1.0)
-	Attribute.DefaultValue = (Property == MP_ShadingModel) ? FVector4f(1.0f, 0.0f, 0.0f, 0.0f) : FMaterialAttributeDefinitionMap::GetDefaultValue(Property);
+	Attribute.DefaultValue = (Property == MP_ShadingModel) ? FVector4f(1.0f, 0.0f, 0.0f, 0.0f) : FMaterialAttributeDefinitionMap::GetDefaultValue(Property); // Should that be changed for substrate?
 	
 	return Attribute;
 }
diff --git a/Engine/Source/Runtime/Engine/Private/Materials/MaterialIREmitter.cpp b/Engine/Source/Runtime/Engine/Private/Materials/MaterialIREmitter.cpp
--- a/Engine/Source/Runtime/Engine/Private/Materials/MaterialIREmitter.cpp
+++ b/Engine/Source/Runtime/Engine/Private/Materials/MaterialIREmitter.cpp
@@ -502,6 +502,7 @@ FType FEmitter::GetMaterialAggregateAttributeType(const UMaterialAggregate* Aggr
 		case EMaterialAggregateAttributeType::Float3: return FType::MakeVector(EScalarKind::Float, 3);
 		case EMaterialAggregateAttributeType::Float4: return FType::MakeVector(EScalarKind::Float, 4);
 		case EMaterialAggregateAttributeType::MaterialAttributes: return FType::MakeAggregate(UMaterialAggregate::GetMaterialAttributes());
+		case EMaterialAggregateAttributeType::SubstrateData: return FType::MakeSubstrateData();
 		case EMaterialAggregateAttributeType::Aggregate: return FType::MakeAggregate(Aggregate->Attributes[AttributeIndex].Aggregate.Get());
 		default: UE_MIR_UNREACHABLE();
 	}
@@ -1362,7 +1363,7 @@ void FEmitter::SetCustomOutputs(FStringView Name, TConstArrayView<FValueRef> Arg
 	}
 
 	// Avoid work if some error occurred.
-	if (CurrentExpressionHasErrors())
+	if (CurrentExpressionHasErrors() || IsAnyNotValid(Args))
 	{
 		return;
 	}
@@ -1588,6 +1589,7 @@ static FValueRef EmitAttributeDefaultValue(FEmitter& Emitter, const UMaterial* M
 		case EMaterialAggregateAttributeType::Float4: return Emitter.ConstantFloat4(UE::Math::TVector4<TFloat>{ Attribute.DefaultValue.X, Attribute.DefaultValue.Y, Attribute.DefaultValue.Z, Attribute.DefaultValue.W });
 		case EMaterialAggregateAttributeType::ShadingModel: return Emitter.ConstantInt((TInteger)Material->GetShadingModels().GetFirstShadingModel());
 		case EMaterialAggregateAttributeType::MaterialAttributes: return Emitter.Aggregate(UMaterialAggregate::GetMaterialAttributes());
+		case EMaterialAggregateAttributeType::SubstrateData: return Emitter.PromoteSubstrateParameter();	// is that really a default value ???
 		case EMaterialAggregateAttributeType::Aggregate: return Emitter.Aggregate(Attribute.Aggregate);
 		default: UE_MIR_UNREACHABLE();
 	}
diff --git a/Engine/Source/Runtime/Engine/Private/Materials/MaterialIRToHLSLTranslator.cpp b/Engine/Source/Runtime/Engine/Private/Materials/MaterialIRToHLSLTranslator.cpp
--- a/Engine/Source/Runtime/Engine/Private/Materials/MaterialIRToHLSLTranslator.cpp
+++ b/Engine/Source/Runtime/Engine/Private/Materials/MaterialIRToHLSLTranslator.cpp
@@ -939,9 +939,9 @@ struct FPrivate : FMaterialIRToHLSLTranslation
 
 			switch (Argument.MaterialProperty)
 			{
-				case MP_Normal: LowerValue(SubstrateParameter->WorldSpaceTangentsAndNormals[0]); break;
-				case MP_Tangent: LowerValue(SubstrateParameter->WorldSpaceTangentsAndNormals[1]); break;
-				
+				case MP_Normal: Printer << TEXTVIEW("float3(0,0,1)") << NewLine; break;
+				case MP_Tangent:  Printer << TEXTVIEW("float3(1,0,0)") << NewLine; break;
+
 				default:
 					if (Argument.MaterialProperty != MP_MAX && SubstrateParameter->PropertyArgs[Argument.MaterialProperty])
 					{
@@ -2230,7 +2230,13 @@ struct FPrivate : FMaterialIRToHLSLTranslation
 				const EMaterialValueType PropertyType = FMaterialAttributeDefinitionMap::GetValueType(AttributeID);
 				Printer << MIR::FType::FromMaterialValueType(PropertyType) << TEXT(' ') << *FMaterialAttributeDefinitionMap::GetAttributeName(AttributeID) << EndOfStatement;
 			}
-			Printer << CloseBlock << L';';
+			Printer << CloseBlock << TEXT(";\n");
+
+			// temporary placeholder support for Substrate
+			Printer << TEXTVIEW("#if SUBSTRATE_USE_FULLYSIMPLIFIED_MATERIAL == 1") << NewLine;
+			Printer << GetHLSLTypeString(FMaterialAttributeDefinitionMap::GetValueType(MP_FrontMaterial)) << TEXTVIEW(" FullySimplifiedFrontMaterial") << EndOfStatement;
+			Printer << TEXTVIEW("#endif") << NewLine;
+
 			OutputShaderTemplateParams.Add(TEXT("material_declarations"), MoveTemp(Printer.Buffer));
 		}
 
diff --git a/Engine/Source/Runtime/Engine/Private/Materials/MaterialShared.cpp b/Engine/Source/Runtime/Engine/Private/Materials/MaterialShared.cpp
--- a/Engine/Source/Runtime/Engine/Private/Materials/MaterialShared.cpp
+++ b/Engine/Source/Runtime/Engine/Private/Materials/MaterialShared.cpp
@@ -3543,8 +3543,8 @@ bool FMaterial::Translate(const FMaterialShaderMapId& InShaderMapId,
 		Interface->MaterialInsight.Reset(new FMaterialInsights);
 	}
 
-	// Use the new translator if the shader map uses it. Note: the new translator does not support Substrate yet, so if Substrate is on, fallback to the old translator.
-	if (InShaderMapId.bUsingNewHLSLGenerator && !Substrate::IsSubstrateEnabled())
+	// Use the new translator if the shader map uses it.
+	if (InShaderMapId.bUsingNewHLSLGenerator)
 	{
 		return Translate_New(InShaderMapId, InStaticParameters, InTargetPlatform, OutCompilationOutput, OutMaterialEnvironment);
 	}
diff --git a/Engine/Source/Runtime/Engine/Public/Materials/MaterialAggregate.h b/Engine/Source/Runtime/Engine/Public/Materials/MaterialAggregate.h
--- a/Engine/Source/Runtime/Engine/Public/Materials/MaterialAggregate.h
+++ b/Engine/Source/Runtime/Engine/Public/Materials/MaterialAggregate.h
@@ -35,6 +35,9 @@ enum class EMaterialAggregateAttributeType
 	// Attribute is of type MaterialAttributes aggregate.
 	MaterialAttributes,
 
+	// Attribute is of type SubstrateData aggregate.
+	SubstrateData,
+
 	// Attribute is of specified user aggregate type.
 	Aggregate,
 };
