Conditionally removed Instanced Static Mesh Vertex Factory for platforms supporting full GPU-Scene (r.InstancedStaticMeshes.UseInstancedStaticMeshVertexFactory, on by default).

* Moved per instance instance dither parameters to Local VF UB to make it possible for HISM to update these without needing a separate VF with loose paramters for every ISM.

* Removed USE_DITHERED_LOD_TRANSITION_FOR_INSTANCED as this has been defunct since 2017.

* Removed some references to UserData within ISM that are no longer needed since there is only ever one.

* Removed shader code dealing with selection (this was already redundant since my previous CL).

* Redirected the ISM proxy to use the base SM vertex factory & removed the loose UBs when the separate ISM VF is disabled.

* Added InstancingFadeOutParams to the Primitive UB  / GPU Scene to make it possible to implement the GetPerInstanceFadeAmount without the separate VF (also enables a path for ISKM to do the same).

* Tidied up some of the interpolants in Local VF. Several are removed when ISM VF is not enabled and instead the attributes are fetched in the PS (see LOCAL_VF_INTERPOLATE_INSTANCE_).

* Local VF now interpolates InstanceId instead of PrimitiveId, note that for SHADER_USES_PRIMITIVE_UBO the DrawInstanceId is actually what's being piped around.

#Synced-CL 43344078

diff --git a/Engine/Plugins/Enterprise/LidarPointCloud/Shaders/Private/LidarPointCloudVertexFactory.ush b/Engine/Plugins/Enterprise/LidarPointCloud/Shaders/Private/LidarPointCloudVertexFactory.ush
--- a/Engine/Plugins/Enterprise/LidarPointCloud/Shaders/Private/LidarPointCloudVertexFactory.ush
+++ b/Engine/Plugins/Enterprise/LidarPointCloud/Shaders/Private/LidarPointCloudVertexFactory.ush
@@ -501,6 +501,8 @@ FMaterialPixelParameters GetMaterialPixelParameters(FVertexFactoryInterpolantsVS
 }
 
 #if USE_INSTANCING
+#define VF_IMPLEMENTED_GET_INSTANCE_HIT_PROXY_ID
+
 float4 VertexFactoryGetInstanceHitProxyId(FVertexFactoryInput Input, FVertexFactoryIntermediates Intermediates) { return 0; }
 #endif
 
diff --git a/Engine/Plugins/Experimental/GPULightmass/Source/GPULightmass/Private/Scene/InstancedStaticMesh.cpp b/Engine/Plugins/Experimental/GPULightmass/Source/GPULightmass/Private/Scene/InstancedStaticMesh.cpp
--- a/Engine/Plugins/Experimental/GPULightmass/Source/GPULightmass/Private/Scene/InstancedStaticMesh.cpp
+++ b/Engine/Plugins/Experimental/GPULightmass/Source/GPULightmass/Private/Scene/InstancedStaticMesh.cpp
@@ -147,7 +147,7 @@ TArray<FMeshBatch> FInstanceGroupRenderState::GetMeshBatchesForGBufferRendering(
 
 		FMeshBatchElement& MeshBatchElement = MeshBatch.Elements[0];
 
-		if (CoordsForCulling.MipLevel == -1)
+		if (CoordsForCulling.MipLevel == -1 || !InstancedRenderData.IsValid())
 		{
 			// No culling, should be for ray tracing scene
 			// YW: While it seems more 'correct' to use FISMVF, what it contains is basically view-dependent instance fading loose parameters,
diff --git a/Engine/Plugins/Experimental/GPULightmass/Source/GPULightmass/Private/Scene/Scene.cpp b/Engine/Plugins/Experimental/GPULightmass/Source/GPULightmass/Private/Scene/Scene.cpp
--- a/Engine/Plugins/Experimental/GPULightmass/Source/GPULightmass/Private/Scene/Scene.cpp
+++ b/Engine/Plugins/Experimental/GPULightmass/Source/GPULightmass/Private/Scene/Scene.cpp
@@ -1018,7 +1018,10 @@ void FScene::AddGeometryInstanceFromComponent(UInstancedStaticMeshComponent* InC
 	InstanceRenderState.ComponentUObject = Instance->ComponentUObject;
 	InstanceRenderState.RenderData = Instance->ComponentUObject->GetStaticMesh()->GetRenderData();
 	FInstancedStaticMeshSceneProxyDesc ProxyDesc(Instance->ComponentUObject);
-	InstanceRenderState.InstancedRenderData = MakeUnique<FInstancedStaticMeshRenderData>(&ProxyDesc, FeatureLevel);
+	if (DoesPlatformUseInstancedStaticMeshVertexFactory(GMaxRHIShaderPlatform))
+	{
+		InstanceRenderState.InstancedRenderData = MakeUnique<FInstancedStaticMeshRenderData>(&ProxyDesc, FeatureLevel);
+	}
 	InstanceRenderState.InstanceDataSceneProxy = ProxyDesc.InstanceDataSceneProxy;
 	InstanceRenderState.LocalToWorld = InComponent->GetRenderMatrix();
 	InstanceRenderState.WorldBounds = InComponent->Bounds;
@@ -1048,7 +1051,11 @@ void FScene::AddGeometryInstanceFromComponent(UInstancedStaticMeshComponent* InC
 			RelevantRectLightsToAddOnRenderThread
 		](FRHICommandListImmediate& RHICmdList) mutable
 	{
-		InstanceRenderState.InstancedRenderData->BindBuffersToVertexFactories(RHICmdList, nullptr);
+		if (InstanceRenderState.InstancedRenderData)
+		{
+			InstanceRenderState.InstancedRenderData->BindBuffersToVertexFactories(RHICmdList, nullptr);
+		}
+
 		if (FInstanceDataUpdateTaskInfo *TaskInfo = InstanceRenderState.InstanceDataSceneProxy->GetUpdateTaskInfo())
 		{
 			TaskInfo->WaitForUpdateCompletion();
@@ -1128,8 +1135,10 @@ void FScene::RemoveGeometryInstanceFromComponent(UInstancedStaticMeshComponent*
 	ENQUEUE_RENDER_COMMAND(RenderThreadRemove)(
 		[ElementId, &RenderState = RenderState](FRHICommandListImmediate&) mutable
 	{
-		RenderState.InstanceGroupRenderStates.Elements[ElementId].InstancedRenderData->ReleaseResources(nullptr, nullptr);
-
+		if (RenderState.InstanceGroupRenderStates.Elements[ElementId].InstancedRenderData)
+		{
+			RenderState.InstanceGroupRenderStates.Elements[ElementId].InstancedRenderData->ReleaseResources(nullptr, nullptr);
+		}
 		for (FLightmapRenderStateRef& Lightmap : RenderState.InstanceGroupRenderStates.Elements[ElementId].LODLightmapRenderStates)
 		{
 			if (Lightmap.IsValid())
diff --git a/Engine/Plugins/Experimental/Water/Shaders/Private/WaterMeshVertexFactory.ush b/Engine/Plugins/Experimental/Water/Shaders/Private/WaterMeshVertexFactory.ush
--- a/Engine/Plugins/Experimental/Water/Shaders/Private/WaterMeshVertexFactory.ush
+++ b/Engine/Plugins/Experimental/Water/Shaders/Private/WaterMeshVertexFactory.ush
@@ -443,6 +443,9 @@ FVertexFactoryIntermediates GetVertexFactoryIntermediates(FVertexFactoryInput In
 }
 
 #if HIT_PROXY_SHADER
+
+#define VF_IMPLEMENTED_GET_INSTANCE_HIT_PROXY_ID
+
 float4 VertexFactoryGetInstanceHitProxyId(FVertexFactoryInput Input, FVertexFactoryIntermediates Intermediates)
 {
 #if WITH_WATER_SELECTION_SUPPORT_VF
diff --git a/Engine/Plugins/FX/Niagara/Shaders/Private/NiagaraMeshVertexFactory.ush b/Engine/Plugins/FX/Niagara/Shaders/Private/NiagaraMeshVertexFactory.ush
--- a/Engine/Plugins/FX/Niagara/Shaders/Private/NiagaraMeshVertexFactory.ush
+++ b/Engine/Plugins/FX/Niagara/Shaders/Private/NiagaraMeshVertexFactory.ush
@@ -862,11 +862,6 @@ uint VertexFactoryGetPrimitiveId(FVertexFactoryInterpolantsVSToPS Interpolants)
 #endif
 }
 
-float4 VertexFactoryGetInstanceHitProxyId(FVertexFactoryInput Input, FVertexFactoryIntermediates Intermediates)
-{
-	return 0;
-}
-
 float4 VertexFactoryGetTranslatedPrimitiveVolumeBounds(FVertexFactoryInterpolantsVSToPS Interpolants)
 {
 	return float4(0,0,0,0);
diff --git a/Engine/Shaders/Private/GeometryCollectionVertexFactory.ush b/Engine/Shaders/Private/GeometryCollectionVertexFactory.ush
--- a/Engine/Shaders/Private/GeometryCollectionVertexFactory.ush
+++ b/Engine/Shaders/Private/GeometryCollectionVertexFactory.ush
@@ -510,7 +510,7 @@ FVertexFactoryInterpolantsVSToPS VertexFactoryGetInterpolantsVSToPS(FVertexFacto
 	SetTangents(Interpolants, Intermediates.TangentToWorld[0], Intermediates.TangentToWorld[2], Intermediates.TangentToWorldSign);
 	SetColor(Interpolants, Intermediates.Color);
 
-	SetPrimitiveId(Interpolants, Intermediates.SceneData.PrimitiveId);
+	SetSceneData(Interpolants, Intermediates.SceneData);
 
 	return Interpolants;
 }
@@ -601,7 +601,7 @@ FVertexFactoryRayTracingInterpolants VertexFactoryInterpolate(FVertexFactoryRayT
 	FVertexFactoryRayTracingInterpolants O = (FVertexFactoryRayTracingInterpolants)0;
 
 #if VF_USE_PRIMITIVE_SCENE_DATA
-	O.InterpolantsVSToPS.PrimitiveId = a.InterpolantsVSToPS.PrimitiveId;
+	O.InterpolantsVSToPS.InstanceId = a.InterpolantsVSToPS.InstanceId;
 #if NEEDS_LIGHTMAP_COORDINATE
 	O.InterpolantsVSToPS.LightmapDataIndex = a.InterpolantsVSToPS.LightmapDataIndex;
 #endif
diff --git a/Engine/Shaders/Private/HairStrands/HairCardsVertexFactory.ush b/Engine/Shaders/Private/HairStrands/HairCardsVertexFactory.ush
--- a/Engine/Shaders/Private/HairStrands/HairCardsVertexFactory.ush
+++ b/Engine/Shaders/Private/HairStrands/HairCardsVertexFactory.ush
@@ -445,7 +445,7 @@ FVertexFactoryInterpolantsVSToPS VertexFactoryGetInterpolantsVSToPS(FVertexFacto
 
 	GetRootUVsAtlasUVs(Input, Interpolants.HairPrimitiveUV, Interpolants.HairPrimitiveRootUV);
 
-	SetPrimitiveId(Interpolants, Intermediates.SceneData.PrimitiveId);
+	SetSceneData(Interpolants, Intermediates.SceneData);
 
 #if NUM_TEX_COORD_INTERPOLATORS && NUM_MATERIAL_TEXCOORDS_VERTEX
 	float2 TexCoords[NUM_MATERIAL_TEXCOORDS_VERTEX];
@@ -548,7 +548,7 @@ FVertexFactoryRayTracingInterpolants VertexFactoryInterpolate(FVertexFactoryRayT
 	FVertexFactoryRayTracingInterpolants O = (FVertexFactoryRayTracingInterpolants)0;
 
 #if VF_USE_PRIMITIVE_SCENE_DATA
-	O.InterpolantsVSToPS.PrimitiveId = a.InterpolantsVSToPS.PrimitiveId;
+	O.InterpolantsVSToPS.InstanceId = a.InterpolantsVSToPS.InstanceId;
 #if NEEDS_LIGHTMAP_COORDINATE
 	O.InterpolantsVSToPS.LightmapDataIndex = a.InterpolantsVSToPS.LightmapDataIndex;
 #endif
diff --git a/Engine/Shaders/Private/HairStrands/HairStrandsMaterialCommon.ush b/Engine/Shaders/Private/HairStrands/HairStrandsMaterialCommon.ush
--- a/Engine/Shaders/Private/HairStrands/HairStrandsMaterialCommon.ush
+++ b/Engine/Shaders/Private/HairStrands/HairStrandsMaterialCommon.ush
@@ -155,7 +155,7 @@ FEvaluateOutput Evaluate(
 		Interpolants.TangentToWorld2 = normalize(Interpolants.TangentToWorld2);
 
 #if VF_USE_PRIMITIVE_SCENE_DATA
-		Interpolants.PrimitiveId	 = PrimitiveId;
+		Interpolants.InstanceId	 = GetPrimitiveData(PrimitiveId).InstanceSceneDataOffset;
 #endif // VF_USE_PRIMITIVE_SCENE_DATA
 		Interpolants.HairControlPointId = HairControlPointId0;
 		Interpolants.HairPrimitiveUV = float2(SegmentU, SegmentV);
diff --git a/Engine/Shaders/Private/HairStrands/HairStrandsVertexFactory.ush b/Engine/Shaders/Private/HairStrands/HairStrandsVertexFactory.ush
--- a/Engine/Shaders/Private/HairStrands/HairStrandsVertexFactory.ush
+++ b/Engine/Shaders/Private/HairStrands/HairStrandsVertexFactory.ush
@@ -836,7 +836,7 @@ FVertexFactoryInterpolantsVSToPS VertexFactoryGetInterpolantsVSToPS(FVertexFacto
 	
 	Interpolants.HairControlPointId = Intermediates.HairControlPointId;
 	Interpolants.HairPrimitiveUV = Intermediates.HairPrimitiveUV;
-	SetPrimitiveId(Interpolants, Intermediates.SceneData.PrimitiveId);
+	SetSceneData(Interpolants, Intermediates.SceneData);
 
 	return Interpolants;
 }
@@ -949,7 +949,7 @@ FVertexFactoryRayTracingInterpolants VertexFactoryInterpolate(FVertexFactoryRayT
 	FVertexFactoryRayTracingInterpolants O = (FVertexFactoryRayTracingInterpolants)0;
 
 #if VF_USE_PRIMITIVE_SCENE_DATA
-	O.InterpolantsVSToPS.PrimitiveId = a.InterpolantsVSToPS.PrimitiveId;
+	O.InterpolantsVSToPS.InstanceId = a.InterpolantsVSToPS.InstanceId;
 #if NEEDS_LIGHTMAP_COORDINATE
 	O.InterpolantsVSToPS.LightmapDataIndex = a.InterpolantsVSToPS.LightmapDataIndex;
 #endif
diff --git a/Engine/Shaders/Private/HitProxyPixelShader.usf b/Engine/Shaders/Private/HitProxyPixelShader.usf
--- a/Engine/Shaders/Private/HitProxyPixelShader.usf
+++ b/Engine/Shaders/Private/HitProxyPixelShader.usf
@@ -18,9 +18,7 @@ float4 HitProxyId;
 
 void Main(
 	FVertexFactoryInterpolantsVSToPS Interpolants,
-#if PER_INSTANCE_HITPROXY_ID || USE_INSTANCE_CULLING
 	float4 InstanceHitProxyId : HIT_PROXY_ID,
-#endif
 	float4 PixelPosition	: TEXCOORD6,
 #if USE_WORLD_POSITION_EXCLUDING_SHADER_OFFSETS
 	float3 PixelPositionExcludingWPO : TEXCOORD7,
@@ -51,9 +49,5 @@ void Main(
 
 	GetMaterialCoverageAndClipping(MaterialParameters, PixelMaterialInputs);
 
-#if PER_INSTANCE_HITPROXY_ID || USE_INSTANCE_CULLING
 	OutColor = HitProxyId + InstanceHitProxyId;
-#else
-	OutColor = HitProxyId;
-#endif
 }
diff --git a/Engine/Shaders/Private/HitProxyVertexShader.usf b/Engine/Shaders/Private/HitProxyVertexShader.usf
--- a/Engine/Shaders/Private/HitProxyVertexShader.usf
+++ b/Engine/Shaders/Private/HitProxyVertexShader.usf
@@ -11,9 +11,6 @@
 #include "/Engine/Generated/Material.ush"
 #include "/Engine/Generated/VertexFactory.ush"
 
-// Whether HitProxy IDs are provided per instance or not
-#define PER_INSTANCE_HITPROXY_ID	(USE_INSTANCING || USE_PER_VERTEX_HITPROXY_ID || USE_VERTEXFACTORY_HITPROXY_ID)
-
 #if USE_PER_VERTEX_HITPROXY_ID
 Buffer<float4> VertexFetch_HitProxyIdBuffer;
 #endif
@@ -32,9 +29,7 @@ Buffer<float4> VertexFetch_HitProxyIdBuffer;
 struct FHitProxyVSToPS
 {
 	FVertexFactoryInterpolantsVSToPS FactoryInterpolants;
-#if PER_INSTANCE_HITPROXY_ID || USE_INSTANCE_CULLING
 	float4 InstanceHitProxyId : HIT_PROXY_ID;
-#endif
 	float4 PixelPosition : TEXCOORD6;
 #if USE_WORLD_POSITION_EXCLUDING_SHADER_OFFSETS
 	float3 PixelPositionExcludingWPO : TEXCOORD7;
@@ -65,15 +60,13 @@ void Main(
 	FMaterialVertexParameters VertexParameters = GetMaterialVertexParameters(Input, VFIntermediates, WorldPosition.xyz, TangentToLocal);
 	WorldPosition.xyz += GetMaterialWorldPositionOffset(VertexParameters);
 	ApplyMaterialFirstPersonTransform(VertexParameters, WorldPosition.xyz);
-
+	Output.InstanceHitProxyId = VertexFactoryGetInstanceHitProxyId(Input, VFIntermediates);
 #if USE_PER_VERTEX_HITPROXY_ID
 	#if MANUAL_VERTEX_FETCH
 		Output.InstanceHitProxyId = VertexFetch_HitProxyIdBuffer[(VertexFactoryGetVertexFetchParameter(VF_VertexOffset) + Input.VertexId)] FMANUALFETCH_COLOR_COMPONENT_SWIZZLE;
 	#else
 		Output.InstanceHitProxyId = float4(0,0,0,0);
 	#endif
-#elif USE_INSTANCING || USE_VERTEXFACTORY_HITPROXY_ID|| USE_INSTANCE_CULLING
-	Output.InstanceHitProxyId = VertexFactoryGetInstanceHitProxyId(Input, VFIntermediates);
 #endif
 
 	{
diff --git a/Engine/Shaders/Private/LocalVertexFactory.ush b/Engine/Shaders/Private/LocalVertexFactory.ush
--- a/Engine/Shaders/Private/LocalVertexFactory.ush
+++ b/Engine/Shaders/Private/LocalVertexFactory.ush
@@ -315,13 +315,13 @@ struct FVertexFactoryIntermediates
 
 #if USE_INSTANCE_CULLING
 	float4 HitProxyId;
-	float IsSelected;
 #endif
 
 	float IsVisible;
-
-#if USE_INSTANCING || USE_INSTANCE_CULLING
+#if USE_INSTANCING
 	float4 InstanceLightMapAndShadowMapUVBias;
+#endif
+#if USE_INSTANCING || USE_INSTANCE_CULLING
 	// x = per-instance fade out amount, y = hide/show flag, z dither fade cutoff, w - CustomData index
 	float4 PerInstanceParams;
 #endif
@@ -493,24 +493,6 @@ float3 GetInstanceOrigin(FVertexFactoryIntermediates Intermediates)
 
 #endif	// USE_INSTANCING
 
-#if USE_INSTANCE_CULLING
-
-float2 GetInstanceShadowMapBias(FVertexFactoryIntermediates Intermediates)
-{
-	// GPUCULL_TODO: This data is not yet pushed to GPU-Scene for non-nanite instances.
-	//return GetInstanceData(Intermediates).InstanceLightMapAndShadowMapUVBias.zw;
-	return Intermediates.InstanceLightMapAndShadowMapUVBias.zw;
-}
-
-float2 GetInstanceLightMapBias(FVertexFactoryIntermediates Intermediates)
-{
-	// GPUCULL_TODO: This data is not yet pushed to GPU-Scene for non-nanite instances.
-	//return GetInstanceData(Intermediates).InstanceLightMapAndShadowMapUVBias.xy;
-	return Intermediates.InstanceLightMapAndShadowMapUVBias.xy;
-}
-
-#endif	// USE_INSTANCE_CULLING
-
 /** Converts from vertex factory specific interpolants to a FMaterialPixelParameters, which is used by material inputs. */
 FMaterialPixelParameters GetMaterialPixelParameters(FVertexFactoryInterpolantsVSToPS Interpolants, float4 SvPosition)
 {
@@ -564,23 +546,44 @@ FMaterialPixelParameters GetMaterialPixelParameters(FVertexFactoryInterpolantsVS
 #endif	// LIGHTMAP_UV_ACCESS
 
 	Result.TwoSidedSign = 1;
-	Result.PrimitiveId = GetPrimitiveId(Interpolants);
 
-#if USES_PER_INSTANCE_CUSTOM_DATA && VF_USE_PRIMITIVE_SCENE_DATA
+#if VF_USE_PRIMITIVE_SCENE_DATA
+	FInstanceSceneData InstanceData = GetInstanceSceneData(Interpolants.InstanceId);
+	Result.PrimitiveId = InstanceData.PrimitiveId;
+
+#if USES_PER_INSTANCE_CUSTOM_DATA
+	Result.CustomDataOffset = InstanceData.CustomDataOffset;
+	Result.CustomDataCount  = InstanceData.CustomDataCount;
+#endif
+	Result.PerInstanceRandom = InstanceData.RandomID;
+#else
+	Result.PrimitiveId = 0;
+#endif
+
+#if LOCAL_VF_INTERPOLATE_CUSTOM_DATA_OFFSET
 	Result.CustomDataOffset = Interpolants.CustomDataOffset;
 	Result.CustomDataCount  = Interpolants.CustomDataCount;
 #endif
 
-#if NEEDS_PER_INSTANCE_RANDOM_PS
+#if LOCAL_VF_INTERPOLATE_PER_INSTANCE_RANDOM
 	Result.PerInstanceRandom = Interpolants.PerInstanceRandom;
 #endif
 
+#if LOCAL_VF_INTERPOLATE_INSTANCE_TRANSFORMS
 #if HAS_INSTANCE_LOCAL_TO_WORLD_PS
 	Result.InstanceLocalToWorld = GetInstanceLocalToWorld(Interpolants);
 #endif
 #if HAS_INSTANCE_WORLD_TO_LOCAL_PS
 	Result.InstanceWorldToLocal = GetInstanceWorldToLocal(Interpolants);
 #endif
+#else
+#if HAS_INSTANCE_LOCAL_TO_WORLD_PS
+	Result.InstanceLocalToWorld = InstanceData.LocalToWorld;
+#endif
+#if HAS_INSTANCE_WORLD_TO_LOCAL_PS
+	Result.InstanceWorldToLocal = InstanceData.WorldToLocal;
+#endif
+#endif
 
 #if NEEDS_PARTICLE_LOCAL_TO_WORLD || NEEDS_PARTICLE_WORLD_TO_LOCAL
 	FPrimitiveSceneData PrimitiveData = GetPrimitiveData(Result.PrimitiveId);
@@ -905,40 +908,42 @@ FVertexFactoryIntermediates GetVertexFactoryIntermediates(FVertexFactoryInput In
 	Intermediates.PrimitiveLocalInstanceIndex = Intermediates.Common.SceneData.InstanceId - PrimitiveData.InstanceSceneDataOffset;
 
 #if USE_INSTANCE_CULLING
-	Intermediates.InstanceLightMapAndShadowMapUVBias = InstanceData.LightMapAndShadowMapUVBias;
 
 #if USE_EDITOR_SHADERS
-	Intermediates.IsSelected = InstanceData.EditorData.bIsSelected ? 1.0f : 0.0f;
 	Intermediates.HitProxyId = float4(InstanceData.EditorData.HitProxyId, 0.0f);
 #else // !USE_EDITOR_SHADERS
-	Intermediates.IsSelected = 0.0f;
 	Intermediates.HitProxyId = float4(0.0f, 0.0f, 0.0f, 0.0f);
 #endif // USE_EDITOR_SHADERS
-
+	
+	Intermediates.PerInstanceParams = float4(1.0f, Intermediates.IsVisible, 0.0f, 0.0f);
 #if NEEDS_PER_INSTANCE_PARAMS
 	{
-		// This repackaging is stupid-seeming but these params are propably packed off to some interpolator or something
-		
+		float3 InstanceTranslatedLocation = TransformLocalToTranslatedWorld(InstanceData.LocalBoundsCenter, InstanceData.LocalToWorld).xyz;
+#if USES_PER_INSTANCE_FADE_AMOUNT
+#if IS_INSTANCED_STATIC_MESH_VF
 		// PerInstanceParams.x: per-instance fade out amount
-		float3 InstanceTranslatedLocation = TransformLocalToTranslatedWorld(GetInstanceData(Intermediates).LocalBoundsCenter, GetInstanceData(Intermediates).LocalToWorld).xyz;
 		Intermediates.PerInstanceParams.x = 1.0 - saturate((length(InstanceTranslatedLocation) - InstancedVFLooseParameters.InstancingFadeOutParams.x) * InstancedVFLooseParameters.InstancingFadeOutParams.y);
-		
-		// InstancingFadeOutParams.z,w are RenderSelected and RenderDeselected respectively.
-
-		// PerInstanceParams.y = hide / show flag, 
-		Intermediates.IsVisible = lerp(InstancedVFLooseParameters.InstancingFadeOutParams.w, InstancedVFLooseParameters.InstancingFadeOutParams.z, Intermediates.IsSelected);
-		Intermediates.PerInstanceParams.y = Intermediates.IsVisible;
-
-		// PerInstanceParams.z dither fade cutoff
+#else
+		Intermediates.PerInstanceParams.x = 1.0 - saturate((length(InstanceTranslatedLocation) - PrimitiveData.InstancingFadeOutParams.x) * PrimitiveData.InstancingFadeOutParams.y);
+#endif
+#endif
 	#if USE_DITHERED_LOD_TRANSITION
-		float RandomLOD = InstancedVFLooseParameters.InstancingViewZCompareZero.w * InstanceData.RandomID; 
-		float ViewZZero = length(InstanceTranslatedLocation - InstancedVFLooseParameters.InstancingTranslatedWorldViewOriginZero.xyz) + RandomLOD;
-		float ViewZOne = length(InstanceTranslatedLocation - InstancedVFLooseParameters.InstancingTranslatedWorldViewOriginOne.xyz) + RandomLOD;
-		Intermediates.PerInstanceParams.z =
-			dot(float3(ViewZZero.xxx > InstancedVFLooseParameters.InstancingViewZCompareZero.xyz), InstancedVFLooseParameters.InstancingViewZConstant.xyz) * InstancedVFLooseParameters.InstancingTranslatedWorldViewOriginZero.w +
-			dot(float3(ViewZOne.xxx > InstancedVFLooseParameters.InstancingViewZCompareOne.xyz), InstancedVFLooseParameters.InstancingViewZConstant.xyz) * InstancedVFLooseParameters.InstancingTranslatedWorldViewOriginOne.w;
-		Intermediates.PerInstanceParams.y *= abs(Intermediates.PerInstanceParams.z) < .999;
-		Intermediates.IsVisible = Intermediates.PerInstanceParams.y;
+
+#if !IS_INSTANCED_STATIC_MESH_VF
+		// Note: in the old path (with a separate VF) we fetch the parameters from InstancedVFLooseParameters, but that is not possible with the shared VF
+#define InstancedVFLooseParameters LocalVF
+#endif
+		if (InstancedVFLooseParameters.InstancingViewZCompareOne.w != 0.0f)
+		{
+			float RandomLOD = InstancedVFLooseParameters.InstancingViewZCompareZero.w * InstanceData.RandomID; 
+			float ViewZZero = length(InstanceTranslatedLocation - InstancedVFLooseParameters.InstancingTranslatedWorldViewOriginZero.xyz) + RandomLOD;
+			float ViewZOne = length(InstanceTranslatedLocation - InstancedVFLooseParameters.InstancingTranslatedWorldViewOriginOne.xyz) + RandomLOD;
+			Intermediates.PerInstanceParams.z =
+				dot(float3(ViewZZero.xxx > InstancedVFLooseParameters.InstancingViewZCompareZero.xyz), InstancedVFLooseParameters.InstancingViewZConstant.xyz) * InstancedVFLooseParameters.InstancingTranslatedWorldViewOriginZero.w +
+				dot(float3(ViewZOne.xxx > InstancedVFLooseParameters.InstancingViewZCompareOne.xyz), InstancedVFLooseParameters.InstancingViewZConstant.xyz) * InstancedVFLooseParameters.InstancingTranslatedWorldViewOriginOne.w;
+			Intermediates.PerInstanceParams.y *= abs(Intermediates.PerInstanceParams.z) < .999;
+			Intermediates.IsVisible = Intermediates.PerInstanceParams.y;
+		}
 	#else
 		Intermediates.PerInstanceParams.z = 0;
 	#endif
@@ -949,8 +954,6 @@ FVertexFactoryIntermediates GetVertexFactoryIntermediates(FVertexFactoryInput In
 			Intermediates.bEvaluateWorldPositionOffset = false;
 		}
 	}
-#else
-	Intermediates.PerInstanceParams = float4(0.f,Intermediates.IsVisible,0.f,0.f);
 #endif // NEEDS_PER_INSTANCE_PARAMS
 
 #endif // USE_INSTANCE_CULLING
@@ -981,10 +984,8 @@ FVertexFactoryIntermediates GetVertexFactoryIntermediates(FVertexFactoryInput In
 	#if	NEEDS_PER_INSTANCE_PARAMS	
 		// x = per-instance fade out factor, y = zero or one depending of if it is shown at all, z is dither cutoff 
 		// PerInstanceParams.y stores a hide/show flag for this instance
-		float SelectedValue = GetInstanceSelected(Intermediates);
 		Intermediates.PerInstanceParams.x = 1.0 - saturate((InstanceViewDistance - InstancedVFLooseParameters.InstancingFadeOutParams.x) * InstancedVFLooseParameters.InstancingFadeOutParams.y);
 		// InstancingFadeOutParams.z,w are RenderSelected and RenderDeselected respectively.
-		Intermediates.PerInstanceParams.y = InstancedVFLooseParameters.InstancingFadeOutParams.z * SelectedValue + InstancedVFLooseParameters.InstancingFadeOutParams.w * (1-SelectedValue);
 		#if USE_DITHERED_LOD_TRANSITION
 			float RandomLOD = InstancedVFLooseParameters.InstancingViewZCompareZero.w * InstanceData.RandomID;
 			float ViewZZero = length(InstanceTranslatedLocation - InstancedVFLooseParameters.InstancingTranslatedWorldViewOriginZero.xyz) + RandomLOD;
@@ -1130,11 +1131,15 @@ FVertexFactoryInterpolantsVSToPS VertexFactoryGetInterpolantsVSToPS(FVertexFacto
 
 	FPrimitiveSceneData PrimitiveData = GetPrimitiveData(Intermediates);
 
+	const bool bHasPerInstanceCoordinateScaleBias = (InstanceData.Flags & INSTANCE_SCENE_DATA_FLAG_HAS_LIGHTSHADOW_UV_BIAS);
+
 	uint LightmapDataIndex = PrimitiveData.LightmapDataIndex + LocalVF.LODLightmapDataIndex;
 	uint LightmapUVIndex = PrimitiveData.LightmapUVIndex;
 	float4 LightMapCoordinateScaleBias = VF_GPUSCENE_GET_LIGHTMAP_UV_SCALE_BIAS(Input, LightmapDataIndex);
-
-	#if USE_INSTANCING || USE_INSTANCE_CULLING
+	#if USE_INSTANCE_CULLING
+		const float2 InstanceLightMapScaleBias = CondMask(bHasPerInstanceCoordinateScaleBias, InstanceData.LightMapAndShadowMapUVBias.xy, LightMapCoordinateScaleBias.zw);
+		LightMapCoordinate = LightMapCoordinateInput * LightMapCoordinateScaleBias.xy + InstanceLightMapScaleBias;
+	#elif USE_INSTANCING
 		LightMapCoordinate = LightMapCoordinateInput * LightMapCoordinateScaleBias.xy + GetInstanceLightMapBias(Intermediates);
 	#else
 		LightMapCoordinate = LightMapCoordinateInput * LightMapCoordinateScaleBias.xy + LightMapCoordinateScaleBias.zw;
@@ -1142,7 +1147,10 @@ FVertexFactoryInterpolantsVSToPS VertexFactoryGetInterpolantsVSToPS(FVertexFacto
 	#if STATICLIGHTING_TEXTUREMASK
 		float4 ShadowMapCoordinateScaleBias = VF_GPUSCENE_GET_SHADOWMAP_UV_SCALE_BIAS(Input, LightmapDataIndex);
 
-		#if USE_INSTANCING || USE_INSTANCE_CULLING
+		#if USE_INSTANCE_CULLING
+			const float2 InstanceShadowMapScaleBias = CondMask(bHasPerInstanceCoordinateScaleBias, InstanceData.LightMapAndShadowMapUVBias.zw, ShadowMapCoordinateScaleBias.zw);
+			ShadowMapCoordinate = LightMapCoordinateInput * ShadowMapCoordinateScaleBias.xy + InstanceShadowMapScaleBias;
+		#elif USE_INSTANCING
 			ShadowMapCoordinate = LightMapCoordinateInput * ShadowMapCoordinateScaleBias.xy + GetInstanceShadowMapBias(Intermediates);
 		#else
 			ShadowMapCoordinate = LightMapCoordinateInput * ShadowMapCoordinateScaleBias.xy + ShadowMapCoordinateScaleBias.zw;
@@ -1160,12 +1168,12 @@ FVertexFactoryInterpolantsVSToPS VertexFactoryGetInterpolantsVSToPS(FVertexFacto
 	Interpolants.PerInstanceParams = Intermediates.PerInstanceParams;
 #endif
 
-#if USES_PER_INSTANCE_CUSTOM_DATA && VF_USE_PRIMITIVE_SCENE_DATA
+#if LOCAL_VF_INTERPOLATE_CUSTOM_DATA_OFFSET
 	Interpolants.CustomDataOffset = InstanceData.CustomDataOffset;
 	Interpolants.CustomDataCount  = InstanceData.CustomDataCount;
 #endif
 
-#if NEEDS_PER_INSTANCE_RANDOM_PS
+#if LOCAL_VF_INTERPOLATE_PER_INSTANCE_RANDOM
 	#if USE_INSTANCING
 		Interpolants.PerInstanceRandom = GetInstanceRandom(Intermediates);
 	#else
@@ -1173,15 +1181,15 @@ FVertexFactoryInterpolantsVSToPS VertexFactoryGetInterpolantsVSToPS(FVertexFacto
 	#endif
 #endif
 
-	SetPrimitiveId(Interpolants, Intermediates.Common.SceneData.PrimitiveId);
-
+	SetSceneData(Interpolants, Intermediates.Common.SceneData);
+#if LOCAL_VF_INTERPOLATE_INSTANCE_TRANSFORMS
 #if HAS_INSTANCE_LOCAL_TO_WORLD_PS
 	SetInstanceLocalToWorld(Interpolants, VertexParameters.InstanceLocalToWorld);
 #endif
 #if HAS_INSTANCE_WORLD_TO_LOCAL_PS
 	SetInstanceWorldToLocal(Interpolants, VertexParameters.InstanceWorldToLocal);
 #endif
-
+#endif
 	return Interpolants;
 }
 
@@ -1308,21 +1316,23 @@ float4 VertexFactoryGetPreviousWorldPosition(FVertexFactoryInput Input, FVertexF
 		FVertexFactoryRayTracingInterpolants O = (FVertexFactoryRayTracingInterpolants)0;
 
 #if VF_USE_PRIMITIVE_SCENE_DATA
-		O.InterpolantsVSToPS.PrimitiveId = a.InterpolantsVSToPS.PrimitiveId;
-    #if NEEDS_LIGHTMAP_COORDINATE
-        O.InterpolantsVSToPS.LightmapDataIndex = a.InterpolantsVSToPS.LightmapDataIndex;
-    #endif
+		O.InterpolantsVSToPS.InstanceId = a.InterpolantsVSToPS.InstanceId;
 #endif	
 
-#if USES_PER_INSTANCE_CUSTOM_DATA && VF_USE_PRIMITIVE_SCENE_DATA
+#if LOCAL_VF_INTERPOLATE_LIGHTMAP_DATA_INDEX
+		O.InterpolantsVSToPS.LightmapDataIndex = a.InterpolantsVSToPS.LightmapDataIndex;
+#endif
+
+#if LOCAL_VF_INTERPOLATE_CUSTOM_DATA_OFFSET
 		O.InterpolantsVSToPS.CustomDataOffset = a.InterpolantsVSToPS.CustomDataOffset;
 		O.InterpolantsVSToPS.CustomDataCount  = a.InterpolantsVSToPS.CustomDataCount;
 #endif
 
-#if NEEDS_PER_INSTANCE_RANDOM_PS
+#if LOCAL_VF_INTERPOLATE_PER_INSTANCE_RANDOM
 		O.InterpolantsVSToPS.PerInstanceRandom = a.InterpolantsVSToPS.PerInstanceRandom;
 #endif
 
+#if LOCAL_VF_INTERPOLATE_INSTANCE_TRANSFORMS
 #if HAS_INSTANCE_LOCAL_TO_WORLD_PS
 		O.InterpolantsVSToPS.InstanceLocalToWorld = a.InterpolantsVSToPS.InstanceLocalToWorld;
 #endif
@@ -1330,6 +1340,7 @@ float4 VertexFactoryGetPreviousWorldPosition(FVertexFactoryInput Input, FVertexF
 #if HAS_INSTANCE_WORLD_TO_LOCAL_PS
 		O.InterpolantsVSToPS.InstanceWorldToLocal = a.InterpolantsVSToPS.InstanceWorldToLocal;
 #endif
+#endif
 
 
 		// Do we really need to interpolate TangentToWorld2 here? It should be replaced by the
@@ -1362,24 +1373,6 @@ float4 VertexFactoryGetPreviousWorldPosition(FVertexFactoryInput Input, FVertexF
 	}
 #endif // #if NEEDS_VERTEX_FACTORY_INTERPOLATION
 
-#if USE_INSTANCE_CULLING
-float4 VertexFactoryGetInstanceHitProxyId(FVertexFactoryInput Input, FVertexFactoryIntermediates Intermediates)
-{
-	return Intermediates.HitProxyId;
-}
-#endif
-
-#if USE_INSTANCING
-float4 VertexFactoryGetInstanceHitProxyId(FVertexFactoryInput Input, FVertexFactoryIntermediates Intermediates)
-{
-		float R = Intermediates.Common.InstanceInput.Transform1.w - 256.0 * GetInstanceSelected(Intermediates);
-		float G = Intermediates.Common.InstanceInput.Transform2.w;
-		float B = Intermediates.Common.InstanceInput.Transform3.w;
-
-	return float4(R/255.0, G/255.0, B/255.0, 0);
-}
-#endif	// USE_INSTANCING
-
 float4 VertexFactoryGetTranslatedPrimitiveVolumeBounds(FVertexFactoryInterpolantsVSToPS Interpolants)
 {
 	FPrimitiveSceneData PrimitiveData = GetPrimitiveData(GetPrimitiveId(Interpolants));
diff --git a/Engine/Shaders/Private/LocalVertexFactoryCommon.ush b/Engine/Shaders/Private/LocalVertexFactoryCommon.ush
--- a/Engine/Shaders/Private/LocalVertexFactoryCommon.ush
+++ b/Engine/Shaders/Private/LocalVertexFactoryCommon.ush
@@ -4,8 +4,18 @@
 	LocalVertexFactoryCommon.usf: Local vertex factory common functionality
 =============================================================================*/
 
+#ifndef IS_INSTANCED_STATIC_MESH_VF
+	#define IS_INSTANCED_STATIC_MESH_VF 0
+#endif
+
 #define NEEDS_PER_INSTANCE_PARAMS ((USE_INSTANCING || USE_INSTANCE_CULLING) && (USES_PER_INSTANCE_FADE_AMOUNT || USES_SPEEDTREE || USE_DITHERED_LOD_TRANSITION || (!VF_USE_PRIMITIVE_SCENE_DATA && USES_PER_INSTANCE_CUSTOM_DATA)))
 
+// Preserve behavior with separate ISM VF (primarily for mobile GPU-Scene platforms)
+#define LOCAL_VF_INTERPOLATE_INSTANCE_TRANSFORMS (IS_INSTANCED_STATIC_MESH_VF || !(VF_USE_PRIMITIVE_SCENE_DATA))
+#define LOCAL_VF_INTERPOLATE_PER_INSTANCE_RANDOM (IS_INSTANCED_STATIC_MESH_VF && NEEDS_PER_INSTANCE_RANDOM_PS)
+#define LOCAL_VF_INTERPOLATE_CUSTOM_DATA_OFFSET  (IS_INSTANCED_STATIC_MESH_VF && USES_PER_INSTANCE_CUSTOM_DATA && VF_USE_PRIMITIVE_SCENE_DATA)
+#define LOCAL_VF_INTERPOLATE_LIGHTMAP_DATA_INDEX (NEEDS_LIGHTMAP_COORDINATE && VF_USE_PRIMITIVE_SCENE_DATA)
+
 struct FVertexFactoryInterpolantsVSToPS
 {
 	TANGENTTOWORLD_INTERPOLATOR_BLOCK
@@ -30,16 +40,18 @@ struct FVertexFactoryInterpolantsVSToPS
 #endif
 
 #if VF_USE_PRIMITIVE_SCENE_DATA
-	nointerpolation uint PrimitiveId : PRIMITIVE_ID;
-	#if NEEDS_LIGHTMAP_COORDINATE
-		nointerpolation uint LightmapDataIndex : LIGHTMAP_ID;
-	#endif
-	#if USES_PER_INSTANCE_CUSTOM_DATA
-		nointerpolation uint CustomDataOffset : CUSTOM_DATA_OFFSET;
-		nointerpolation uint CustomDataCount : CUSTOM_DATA_COUNT;
-	#endif
-#endif
-#if NEEDS_PER_INSTANCE_RANDOM_PS
+	nointerpolation uint InstanceId : INSTANCE_ID;
+#endif
+
+#if LOCAL_VF_INTERPOLATE_LIGHTMAP_DATA_INDEX
+	nointerpolation uint LightmapDataIndex : LIGHTMAP_ID;
+#endif
+
+#if LOCAL_VF_INTERPOLATE_CUSTOM_DATA_OFFSET
+	nointerpolation uint CustomDataOffset : CUSTOM_DATA_OFFSET;
+	nointerpolation uint CustomDataCount : CUSTOM_DATA_COUNT;
+#endif
+#if LOCAL_VF_INTERPOLATE_PER_INSTANCE_RANDOM
 	nointerpolation float PerInstanceRandom : PER_INSTANCE_RANDOM;
 #endif
 
@@ -55,12 +67,14 @@ struct FVertexFactoryInterpolantsVSToPS
 	nointerpolation float HairPrimitiveGroupIndex	: HAIR_PRIMITIVE_GROUPINDEX;
 #endif
 
+#if LOCAL_VF_INTERPOLATE_INSTANCE_TRANSFORMS
 #if HAS_INSTANCE_LOCAL_TO_WORLD_PS
 	nointerpolation float3x4 InstanceLocalToWorld : INSTANCE_LOCAL_TO_WORLD;
 #endif
 #if HAS_INSTANCE_WORLD_TO_LOCAL_PS
 	nointerpolation float3x4 InstanceWorldToLocal : INSTANCE_WORLD_TO_LOCAL;
 #endif
+#endif
 };
 
 #if NUM_TEX_COORD_INTERPOLATORS || USE_PARTICLE_SUBUVS
@@ -106,7 +120,7 @@ void GetLightMapCoordinates(FVertexFactoryInterpolantsVSToPS Interpolants, out f
 	LightmapUV0 = Interpolants.LightMapCoordinate.xy * float2( 1, 0.5 );
 	LightmapUV1 = LightmapUV0 + float2( 0, 0.5 );
 
-#if VF_USE_PRIMITIVE_SCENE_DATA && NEEDS_LIGHTMAP_COORDINATE
+#if LOCAL_VF_INTERPOLATE_LIGHTMAP_DATA_INDEX
 	LightmapDataIndex = Interpolants.LightmapDataIndex;
 #else
 	LightmapDataIndex = 0;
@@ -115,7 +129,7 @@ void GetLightMapCoordinates(FVertexFactoryInterpolantsVSToPS Interpolants, out f
 
 void GetShadowMapCoordinate(FVertexFactoryInterpolantsVSToPS Interpolants, out float2 ShadowMapCoordinate, out uint LightmapDataIndex)
 {
-#if VF_USE_PRIMITIVE_SCENE_DATA && NEEDS_LIGHTMAP_COORDINATE
+#if LOCAL_VF_INTERPOLATE_LIGHTMAP_DATA_INDEX
 	LightmapDataIndex = Interpolants.LightmapDataIndex;
 #else
 	LightmapDataIndex = 0;
@@ -152,27 +166,38 @@ void SetTangents(inout FVertexFactoryInterpolantsVSToPS Interpolants, float3 InT
 uint GetPrimitiveId(FVertexFactoryInterpolantsVSToPS Interpolants)
 {
 #if VF_USE_PRIMITIVE_SCENE_DATA
-	return Interpolants.PrimitiveId;
+#if SHADER_USES_PRIMITIVE_UBO
+	// Note: SHADER_USES_PRIMITIVE_UBO pushes the DrawInstanceId through InstanceId
+	return Interpolants.InstanceId;
+#else
+	FInstanceSceneData InstanceData = GetInstanceSceneData(Interpolants.InstanceId);
+	return InstanceData.PrimitiveId;
+#endif
 #else
 	return 0;
 #endif
 }
 
-void SetPrimitiveId(inout FVertexFactoryInterpolantsVSToPS Interpolants, uint PrimitiveId)
+void SetSceneData(inout FVertexFactoryInterpolantsVSToPS Interpolants, FSceneDataIntermediates SceneData)
 {
 #if VF_USE_PRIMITIVE_SCENE_DATA
-	Interpolants.PrimitiveId = PrimitiveId;
+#if SHADER_USES_PRIMITIVE_UBO
+	// Note: SHADER_USES_PRIMITIVE_UBO pushes the DrawInstanceId through SceneData.PrimitiveId
+	Interpolants.InstanceId = SceneData.PrimitiveId;
+#else
+	Interpolants.InstanceId = SceneData.InstanceId;
+#endif
 #endif
 }
 
 void SetLightmapDataIndex(inout FVertexFactoryInterpolantsVSToPS Interpolants, uint LightmapDataIndex)
 {
-#if VF_USE_PRIMITIVE_SCENE_DATA && NEEDS_LIGHTMAP_COORDINATE
+#if LOCAL_VF_INTERPOLATE_LIGHTMAP_DATA_INDEX
 	Interpolants.LightmapDataIndex = LightmapDataIndex;
 #endif
 }
 
-#if HAS_INSTANCE_LOCAL_TO_WORLD_PS
+#if LOCAL_VF_INTERPOLATE_INSTANCE_TRANSFORMS && HAS_INSTANCE_LOCAL_TO_WORLD_PS
 void SetInstanceLocalToWorld(inout FVertexFactoryInterpolantsVSToPS Interpolants, FDFMatrix InstanceLocalToWorld)
 {
 	float4x4 InstanceLocalToTranslatedWorld = DFFastToTranslatedWorld(InstanceLocalToWorld, ResolvedView.PrevPreViewTranslation);
@@ -189,7 +214,7 @@ FDFMatrix GetInstanceLocalToWorld(FVertexFactoryInterpolantsVSToPS Interpolants)
 }
 #endif
 
-#if HAS_INSTANCE_WORLD_TO_LOCAL_PS
+#if LOCAL_VF_INTERPOLATE_INSTANCE_TRANSFORMS && HAS_INSTANCE_WORLD_TO_LOCAL_PS
 
 void SetInstanceWorldToLocal(inout FVertexFactoryInterpolantsVSToPS Interpolants, FDFInverseMatrix InstanceWorldToLocal)
 {
diff --git a/Engine/Shaders/Private/MaterialTemplate.ush b/Engine/Shaders/Private/MaterialTemplate.ush
--- a/Engine/Shaders/Private/MaterialTemplate.ush
+++ b/Engine/Shaders/Private/MaterialTemplate.ush
@@ -1492,6 +1492,7 @@ float GetPerInstanceCustomData(FMaterialPixelParameters Parameters, int Index, f
 #if VF_USE_PRIMITIVE_SCENE_DATA && USES_PER_INSTANCE_CUSTOM_DATA
 	const uint FloatIndex = uint(Index);
 	#if SHADER_USES_PRIMITIVE_UBO
+		// Note: Not actually primitive ID here (DrawInstanceId)
 		return LoadInstanceCustomDataElementUBO(Parameters.PrimitiveId, FloatIndex, DefaultValue);
 	#else
 		BRANCH
@@ -1522,6 +1523,7 @@ float GetPerInstanceCustomData(FMaterialVertexParameters Parameters, int Index,
 #if VF_USE_PRIMITIVE_SCENE_DATA && USES_PER_INSTANCE_CUSTOM_DATA
 	const uint FloatIndex = uint(Index);
 	#if SHADER_USES_PRIMITIVE_UBO
+		// Note: Not actually primitive ID here (DrawInstanceId)
 		return LoadInstanceCustomDataElementUBO(Parameters.PrimitiveId, FloatIndex, DefaultValue);
 	#else
 		BRANCH
@@ -1967,7 +1969,7 @@ float GetPerInstanceRandom(FMaterialPixelParameters Parameters)
 /** Get the per-instance fade-out amount when instancing */
 float GetPerInstanceFadeAmount(FMaterialPixelParameters Parameters)
 {
-#if USE_INSTANCING || USE_INSTANCE_CULLING
+#if USES_PER_INSTANCE_FADE_AMOUNT && (USE_INSTANCING || USE_INSTANCE_CULLING)
 	return float(Parameters.PerInstanceParams.x);
 #else
 	return float(1.0);
@@ -1977,7 +1979,7 @@ float GetPerInstanceFadeAmount(FMaterialPixelParameters Parameters)
 /** Get the per-instance fade-out amount when instancing */
 float GetPerInstanceFadeAmount(FMaterialVertexParameters Parameters)
 {
-#if USE_INSTANCING || USE_INSTANCE_CULLING
+#if USES_PER_INSTANCE_FADE_AMOUNT && (USE_INSTANCING || USE_INSTANCE_CULLING)
 	return float(Parameters.PerInstanceParams.x);
 #else
 	return float(1.0);
@@ -4448,7 +4450,11 @@ void ClipLODTransition(FMaterialPixelParameters Parameters, float DitherFactor)
 #if (USE_INSTANCING  || USE_INSTANCE_CULLING) && USE_DITHERED_LOD_TRANSITION
 void ClipLODTransition(FMaterialPixelParameters Parameters)
 {
-	ClipLODTransition(Parameters, Parameters.PerInstanceParams.z);
+	float LODFactor = select(PrimitiveDither.LODFactor != 0.0f, PrimitiveDither.LODFactor, Parameters.PerInstanceParams.z);
+	if (LODFactor != 0.0)
+	{
+		ClipLODTransition(Parameters, LODFactor);
+	}
 }
 #elif USE_DITHERED_LOD_TRANSITION && !USE_STENCIL_LOD_DITHER
 void ClipLODTransition(FMaterialPixelParameters Parameters)
diff --git a/Engine/Shaders/Private/SceneData.ush b/Engine/Shaders/Private/SceneData.ush
--- a/Engine/Shaders/Private/SceneData.ush
+++ b/Engine/Shaders/Private/SceneData.ush
@@ -52,17 +52,7 @@
 	#define VF_USE_PRIMITIVE_SCENE_DATA 0
 #endif
 
-#ifndef USE_DITHERED_LOD_TRANSITION
-	#if USE_INSTANCING || USE_INSTANCE_CULLING
-		#ifdef USE_DITHERED_LOD_TRANSITION_FOR_INSTANCED
-			#define USE_DITHERED_LOD_TRANSITION USE_DITHERED_LOD_TRANSITION_FOR_INSTANCED
-		#endif	
-	#else
-		#ifdef USE_DITHERED_LOD_TRANSITION_FROM_MATERIAL
-			#define USE_DITHERED_LOD_TRANSITION USE_DITHERED_LOD_TRANSITION_FROM_MATERIAL
-		#endif
-	#endif
-#endif
+#define USE_DITHERED_LOD_TRANSITION USE_DITHERED_LOD_TRANSITION_FROM_MATERIAL
 
 #define INVALID_INSTANCE_PAYLOAD_OFFSET (0xFFFFFFFFu)
 
@@ -133,6 +123,7 @@ struct FPrimitiveSceneData
 	uint        MaterialCacheDescriptor;
 	float4		CustomPrimitiveData[NUM_CUSTOM_PRIMITIVE_DATA]; // TODO: Move to associated array to shrink primitive data and pack cachelines more effectively
 	float		AnimationMinScreenSize;
+	float2		InstancingFadeOutParams;
 };
 
 float3 UnpackColorRGB24(float PackedColorFloat)
@@ -210,8 +201,9 @@ FPrimitiveSceneData GetPrimitiveDataFromUniformBuffer()
 	PrimitiveData.PixelProgrammableDistanceSquared	= Primitive.PixelProgrammableDistanceSquared;
 	PrimitiveData.MaterialDisplacementFadeOutSize	= Primitive.MaterialDisplacementFadeOutSize;
 	PrimitiveData.MeshPaintTextureDescriptor		= Primitive.MeshPaintTextureDescriptor;
-	PrimitiveData.MaterialCacheDescriptor            = Primitive.MaterialCacheDescriptor;
+	PrimitiveData.MaterialCacheDescriptor			= Primitive.MaterialCacheDescriptor;
 	PrimitiveData.AnimationMinScreenSize			= Primitive.AnimationMinScreenSize;
+	PrimitiveData.InstancingFadeOutParams			= Primitive.InstancingFadeOutParams;
 	
 	UNROLL
 	for (uint DataIndex = 0; DataIndex < NUM_CUSTOM_PRIMITIVE_DATA; ++DataIndex)
@@ -478,8 +470,9 @@ FPrimitiveSceneData GetPrimitiveData(uint PrimitiveId)
 	PrimitiveData.MeshPaintTextureDescriptor			= asuint(LoadPrimitivePrimitiveSceneDataElement(PrimitiveIndex, 33).xy);
 	PrimitiveData.NaniteAssemblyTransformOffset			= asuint(LoadPrimitivePrimitiveSceneDataElement(PrimitiveIndex, 33).z);
 	PrimitiveData.NaniteAssemblyTransformCount			= asuint(LoadPrimitivePrimitiveSceneDataElement(PrimitiveIndex, 33).w);
-	PrimitiveData.MaterialCacheDescriptor                = asuint(LoadPrimitivePrimitiveSceneDataElement(PrimitiveIndex, 34).x);
+	PrimitiveData.MaterialCacheDescriptor				= asuint(LoadPrimitivePrimitiveSceneDataElement(PrimitiveIndex, 34).x);
 	PrimitiveData.AnimationMinScreenSize				= LoadPrimitivePrimitiveSceneDataElement(PrimitiveIndex, 34).y;
+	PrimitiveData.InstancingFadeOutParams				= LoadPrimitivePrimitiveSceneDataElement(PrimitiveIndex, 34).zw;
 	// NOTE: Please make sure GetPrimitiveDataFromUniformBuffer() gets updated as well when adding new members here
 
 	// TODO: Move to associated array (and editor data) to shrink primitive data and better pack cachelines
diff --git a/Engine/Shaders/Private/VertexFactoryDefaultInterface.ush b/Engine/Shaders/Private/VertexFactoryDefaultInterface.ush
--- a/Engine/Shaders/Private/VertexFactoryDefaultInterface.ush
+++ b/Engine/Shaders/Private/VertexFactoryDefaultInterface.ush
@@ -66,3 +66,17 @@ uint VertexFactoryGetInstanceIdLoadIndex(FPositionAndNormalOnlyVertexFactoryInpu
 #endif // POSITION_ONLY
 
 #endif // VF_IMPLEMENTED_GET_VIEW_INDEX
+
+#ifndef VF_IMPLEMENTED_GET_INSTANCE_HIT_PROXY_ID
+#define VF_IMPLEMENTED_GET_INSTANCE_HIT_PROXY_ID
+
+float4 VertexFactoryGetInstanceHitProxyId(FVertexFactoryInput Input, FVertexFactoryIntermediates Intermediates)
+{
+#if USE_EDITOR_SHADERS
+	return float4(GetSceneDataIntermediates(Intermediates).InstanceData.EditorData.HitProxyId, 0.0f);
+#else
+	return float4(0.0f, 0.0f, 0.0f, 0.0f);
+#endif
+}
+
+#endif // VF_IMPLEMENTED_GET_INSTANCE_HIT_PROXY_ID
diff --git a/Engine/Source/Runtime/Engine/Classes/Engine/InstancedStaticMesh.h b/Engine/Source/Runtime/Engine/Classes/Engine/InstancedStaticMesh.h
--- a/Engine/Source/Runtime/Engine/Classes/Engine/InstancedStaticMesh.h
+++ b/Engine/Source/Runtime/Engine/Classes/Engine/InstancedStaticMesh.h
@@ -48,11 +48,7 @@ BEGIN_GLOBAL_SHADER_PARAMETER_STRUCT(FInstancedStaticMeshVertexFactoryUniformSha
 END_GLOBAL_SHADER_PARAMETER_STRUCT()
 
 BEGIN_GLOBAL_SHADER_PARAMETER_STRUCT(FInstancedStaticMeshVFLooseUniformShaderParameters, ENGINE_API)
-	SHADER_PARAMETER(FVector4f, InstancingViewZCompareZero)
-	SHADER_PARAMETER(FVector4f, InstancingViewZCompareOne)
-	SHADER_PARAMETER(FVector4f, InstancingViewZConstant)
-	SHADER_PARAMETER(FVector4f, InstancingTranslatedWorldViewOriginZero)
-	SHADER_PARAMETER(FVector4f, InstancingTranslatedWorldViewOriginOne)
+	SHADER_PARAMETER_STRUCT_INCLUDE(FLocalVertexFactoryInstanceDitherParameters, InstanceDitherParameters)
 	SHADER_PARAMETER(FVector4f, InstancingFadeOutParams)
 END_GLOBAL_SHADER_PARAMETER_STRUCT()
 
@@ -230,8 +226,6 @@ struct FInstancingUserData
 
 	int32 MinLOD;
 
-	bool bRenderSelected;
-	bool bRenderUnselected;
 	FVector AverageInstancesScale;
 	FVector InstancingOffset;
 };
@@ -515,6 +509,7 @@ public:
 	ENGINE_API virtual HHitProxy* CreateHitProxies(UPrimitiveComponent* Component,TArray<TRefCountPtr<HHitProxy> >& OutHitProxies) override;
 
 	ENGINE_API virtual bool GetInstanceDrawDistanceMinMax(FVector2f& OutDistanceMinMax) const override;
+	ENGINE_API virtual bool GetInstancingFadeOutParams(FVector2f& OutInstancingFadeOutParams) const override;
 
 	ENGINE_API virtual float GetLodScreenSizeScale() const override;
 	ENGINE_API virtual float GetGpuLodInstanceRadius() const override;
@@ -527,13 +522,18 @@ public:
 	ENGINE_API virtual void GetPreSkinnedLocalBounds(FBoxSphereBounds& OutBounds) const override;
 
 protected:
-	ENGINE_API FInstancedStaticMeshVFLooseUniformShaderParametersRef CreateLooseUniformBuffer(const FSceneView* View, const FInstancingUserData* InstancingUserData, uint32 InstancedLODRange, uint32 InstancedLODIndex, EUniformBufferUsage UniformBufferUsage) const;
+	ENGINE_API FLocalVertexFactoryInstanceDitherParameters GetInstanceDitherParameters(const FSceneView* View, uint32 InstancedLODRange, uint32 InstancedLODIndex) const;
+	ENGINE_API FVector2f CalcInstancingFadeOutParams() const;
+	ENGINE_API FInstancedStaticMeshVFLooseUniformShaderParametersRef CreateLooseUniformBuffer(const FSceneView* View, uint32 InstancedLODRange, uint32 InstancedLODIndex, EUniformBufferUsage UniformBufferUsage) const;
 
 	/** Cache of the StaticMesh asset, needed to release SpeedTree resources*/
 	UStaticMesh* StaticMesh;
 
 	/** Per component render data */
-	FInstancedStaticMeshRenderData InstancedRenderData;
+	TUniquePtr<FInstancedStaticMeshRenderData> InstancedRenderData;
+
+	/** Helper function to get the number of VFs in the render data */
+	int32 GetNumRenderDataVertexFactories() const;
 
 #if WITH_EDITOR
 	/* If we we have any selected instances */
diff --git a/Engine/Source/Runtime/Engine/Private/HierarchicalInstancedStaticMesh.cpp b/Engine/Source/Runtime/Engine/Private/HierarchicalInstancedStaticMesh.cpp
--- a/Engine/Source/Runtime/Engine/Private/HierarchicalInstancedStaticMesh.cpp
+++ b/Engine/Source/Runtime/Engine/Private/HierarchicalInstancedStaticMesh.cpp
@@ -1226,9 +1226,11 @@ void FHierarchicalStaticMeshSceneProxy::FillDynamicMeshElements(const FSceneView
 	SCOPE_CYCLE_COUNTER(STAT_FoliageBatchTime);
 	int64 TotalTriangles = 0;
 
-	int32 OnlyLOD = FMath::Min<int32>(CVarOnlyLOD.GetValueOnRenderThread(),InstancedRenderData.VertexFactories.Num() - 1);
+	const int32 NumVertexFactories = GetNumRenderDataVertexFactories();
+
+	int32 OnlyLOD = FMath::Min<int32>(CVarOnlyLOD.GetValueOnRenderThread(), NumVertexFactories - 1);
 	int32 FirstLOD = FMath::Max((OnlyLOD < 0) ? 0 : OnlyLOD, static_cast<int32>(this->GetCurrentFirstLODIdx_Internal()));
-	int32 LastLODPlusOne = (OnlyLOD < 0) ? InstancedRenderData.VertexFactories.Num() : (OnlyLOD+1);
+	int32 LastLODPlusOne = (OnlyLOD < 0) ? NumVertexFactories : (OnlyLOD+1);
 
 	const bool bUseGPUScene = UseGPUScene(GetScene().GetShaderPlatform(), GetScene().GetFeatureLevel());
 
@@ -1249,7 +1251,11 @@ void FHierarchicalStaticMeshSceneProxy::FillDynamicMeshElements(const FSceneView
 		const int32 TotalNumSections = LODModel.Sections.Num();
 
 		{
-			FInstancedStaticMeshVFLooseUniformShaderParametersRef LooseUniformBuffer = CreateLooseUniformBuffer(View, ElementParams.PassUserData, InstancedLODRange, LODIndex, EUniformBufferUsage::UniformBuffer_SingleFrame);
+			FInstancedStaticMeshVFLooseUniformShaderParametersRef LooseUniformBuffer;
+			if (InstancedRenderData)
+			{
+				LooseUniformBuffer = CreateLooseUniformBuffer(View, InstancedLODRange, LODIndex, EUniformBufferUsage::UniformBuffer_SingleFrame);
+			}
 
 			for (int32 SectionIndex = 0; SectionIndex < TotalNumSections; SectionIndex++)
 			{
@@ -1276,9 +1282,19 @@ void FHierarchicalStaticMeshSceneProxy::FillDynamicMeshElements(const FSceneView
 					MeshBatch.bUseSelectionOutline = ElementParams.bBatchRenderSelection;
 					MeshBatch.bUseWireframeSelectionColoring = ElementParams.bBatchRenderSelection;
 					MeshBatch.bUseAsOccluder = ShouldUseAsOccluder();
-					MeshBatch.VertexFactory = &InstancedRenderData.VertexFactories[LODIndex];
-
 					FMeshBatchElement& MeshBatchElement = MeshBatch.Elements[0];
+					if (InstancedRenderData)
+					{
+						MeshBatch.VertexFactory = &InstancedRenderData->VertexFactories[LODIndex];
+					}
+					// For dither LOD to work we need to override the VF UB with per-view dynamic data
+					else if (bDitherLODEnabled)
+					{
+						const FLocalVertexFactory* VF = static_cast<const FLocalVertexFactory*>(MeshBatch.VertexFactory);
+						FLocalVertexFactoryInstanceDitherParameters InstanceDitherParameters = GetInstanceDitherParameters(View, InstancedLODRange, LODIndex);
+						MeshBatchElement.VertexFactoryUserData = CreateLocalVFUniformBuffer(VF, VF->GetLightMapCoordinateIndex(), nullptr, 0, 0, &InstanceDitherParameters, EUniformBufferUsage::UniformBuffer_SingleFrame);
+					}
+
 					MeshBatchElement.UserData = ElementParams.PassUserData;
 					MeshBatchElement.bUserDataIsColorVertexBuffer = false;
 					MeshBatchElement.MaxScreenSize = 1.0;
@@ -1299,7 +1315,7 @@ void FHierarchicalStaticMeshSceneProxy::FillDynamicMeshElements(const FSceneView
 							if (ElementParams.FinalCullDistance > 9.9E8)
 							{
 								UE_LOG(LogStaticMesh, Display, TEXT("lod:%1d/%1d   section:%1d/%1d   runs:%4d   inst:%8d   tris:%9lld   cast shadow:%1d   cull:-NONE!!-   shadow:%1d     %s %s"),
-									LODIndex, InstancedRenderData.VertexFactories.Num(), SectionIndex, TotalNumSections, RunArray.Num() / 2,
+									LODIndex, NumVertexFactories, SectionIndex, TotalNumSections, RunArray.Num() / 2,
 									TotalInstances, Tris, (int)MeshBatch.CastShadow, ElementParams.ShadowFrustum,
 									*StaticMesh->GetPathName(),
 									*MeshBatch.MaterialRenderProxy->GetIncompleteMaterialWithFallback(ElementParams.FeatureLevel).GetFriendlyName());
@@ -1307,7 +1323,7 @@ void FHierarchicalStaticMeshSceneProxy::FillDynamicMeshElements(const FSceneView
 							else
 							{
 								UE_LOG(LogStaticMesh, Display, TEXT("lod:%1d/%1d   section:%1d/%1d   runs:%4d   inst:%8d   tris:%9lld   cast shadow:%1d   cull:%8.0f   shadow:%1d     %s %s"),
-									LODIndex, InstancedRenderData.VertexFactories.Num(), SectionIndex, TotalNumSections, RunArray.Num() / 2,
+									LODIndex, NumVertexFactories, SectionIndex, TotalNumSections, RunArray.Num() / 2,
 									TotalInstances, Tris, (int)MeshBatch.CastShadow, ElementParams.FinalCullDistance, ElementParams.ShadowFrustum,
 									*StaticMesh->GetPathName(),
 									*MeshBatch.MaterialRenderProxy->GetIncompleteMaterialWithFallback(ElementParams.FeatureLevel).GetFriendlyName());
@@ -1346,8 +1362,9 @@ void FHierarchicalStaticMeshSceneProxy::FillDynamicMeshElements(const FSceneView
 						}
 					}
 				}
-				else
+				else if (ensure(InstancedRenderData))
 				{
+					ensure(DoesPlatformUseInstancedStaticMeshVertexFactory(GetScene().GetShaderPlatform()));
 					int32 NumBatches = 1;
 					int32 CurrentRun = 0;
 					int32 CurrentInstance = 0;
@@ -1374,7 +1391,7 @@ void FHierarchicalStaticMeshSceneProxy::FillDynamicMeshElements(const FSceneView
 						}
 						checkSlow(MeshElement.GetNumPrimitives() > 0);
 
-						MeshElement.VertexFactory = &InstancedRenderData.VertexFactories[LODIndex];
+						MeshElement.VertexFactory = &InstancedRenderData->VertexFactories[LODIndex];
 						FMeshBatchElement& BatchElement0 = MeshElement.Elements[0];
 
 						BatchElement0.UserData = ElementParams.PassUserData;
@@ -1401,7 +1418,7 @@ void FHierarchicalStaticMeshSceneProxy::FillDynamicMeshElements(const FSceneView
 								if (ElementParams.FinalCullDistance > 9.9E8)
 								{
 									UE_LOG(LogStaticMesh, Display, TEXT("lod:%1d/%1d   section:%1d/%1d   runs:%4d   inst:%8d   tris:%9lld   cast shadow:%1d   cull:-NONE!!-   shadow:%1d     %s %s"),
-										LODIndex, InstancedRenderData.VertexFactories.Num(), SectionIndex, TotalNumSections, RunArray.Num() / 2,
+										LODIndex, InstancedRenderData->VertexFactories.Num(), SectionIndex, TotalNumSections, RunArray.Num() / 2,
 										RemainingInstances, Tris, (int)MeshElement.CastShadow, ElementParams.ShadowFrustum,
 										*StaticMesh->GetPathName(),
 										*MeshElement.MaterialRenderProxy->GetIncompleteMaterialWithFallback(ElementParams.FeatureLevel).GetFriendlyName());
@@ -1409,7 +1426,7 @@ void FHierarchicalStaticMeshSceneProxy::FillDynamicMeshElements(const FSceneView
 								else
 								{
 									UE_LOG(LogStaticMesh, Display, TEXT("lod:%1d/%1d   section:%1d/%1d   runs:%4d   inst:%8d   tris:%9lld   cast shadow:%1d   cull:%8.0f   shadow:%1d     %s %s"),
-										LODIndex, InstancedRenderData.VertexFactories.Num(), SectionIndex, TotalNumSections, RunArray.Num() / 2,
+										LODIndex, InstancedRenderData->VertexFactories.Num(), SectionIndex, TotalNumSections, RunArray.Num() / 2,
 										RemainingInstances, Tris, (int)MeshElement.CastShadow, ElementParams.FinalCullDistance, ElementParams.ShadowFrustum,
 										*StaticMesh->GetPathName(),
 										*MeshElement.MaterialRenderProxy->GetIncompleteMaterialWithFallback(ElementParams.FeatureLevel).GetFriendlyName());
@@ -1526,7 +1543,7 @@ void FHierarchicalStaticMeshSceneProxy::GetDynamicMeshElements(const TArray<cons
 			ElementParams.bIsWireframe = ViewFamily.EngineShowFlags.Wireframe;
 			ElementParams.bUseHoveredMaterial = IsHovered();
 			ElementParams.bUseInstanceRuns = (CVarFoliageUseInstanceRuns.GetValueOnRenderThread() > 0);
-			ElementParams.FeatureLevel = InstancedRenderData.FeatureLevel;
+			ElementParams.FeatureLevel = GetScene().GetFeatureLevel();
 			ElementParams.ViewIndex = ViewIndex;
 			ElementParams.View = View;
 
diff --git a/Engine/Source/Runtime/Engine/Private/InstancedStaticMesh.cpp b/Engine/Source/Runtime/Engine/Private/InstancedStaticMesh.cpp
--- a/Engine/Source/Runtime/Engine/Private/InstancedStaticMesh.cpp
+++ b/Engine/Source/Runtime/Engine/Private/InstancedStaticMesh.cpp
@@ -101,10 +101,11 @@ TAutoConsoleVariable<int32> CVarMinLOD(
 	TEXT("Used to discard the top LODs for performance evaluation. -1: Disable all effects of this cvar."),
 	ECVF_Scalability | ECVF_Default);
 
+static FAutoConsoleVariableDeprecated CVarCullAllInVertexShader_Deprecated(TEXT("foliage.CullAllInVertexShader"), TEXT("foliage.CullAllInVertexShader_Unused"), TEXT("5.8"));
 static TAutoConsoleVariable<int32> CVarCullAllInVertexShader(
-	TEXT("foliage.CullAllInVertexShader"),
+	TEXT("foliage.CullAllInVertexShader_Unused"),
 	0,
-	TEXT("Debugging, if this is greater than 0, cull all instances in the vertex shader."));
+	TEXT("This cvar is no longer implemented as the support would add overhead and complexity."));
 
 static TAutoConsoleVariable<int32> CVarRayTracingRenderInstances(
 	TEXT("r.RayTracing.Geometry.InstancedStaticMeshes"),
@@ -696,12 +697,15 @@ void FStaticMeshInstanceData::Serialize(FArchive& Ar)
  */
 bool FInstancedStaticMeshVertexFactory::ShouldCompilePermutation(const FVertexFactoryShaderPermutationParameters& Parameters)
 {
-	return (Parameters.MaterialParameters.bIsUsedWithInstancedStaticMeshes || Parameters.MaterialParameters.bIsSpecialEngineMaterial)
+	return DoesPlatformUseInstancedStaticMeshVertexFactory(Parameters.Platform)
+			&& (Parameters.MaterialParameters.bIsUsedWithInstancedStaticMeshes || Parameters.MaterialParameters.bIsSpecialEngineMaterial)
 			&& FLocalVertexFactory::ShouldCompilePermutation(Parameters);
 }
 
 void FInstancedStaticMeshVertexFactory::ModifyCompilationEnvironment(const FVertexFactoryShaderPermutationParameters& Parameters, FShaderCompilerEnvironment& OutEnvironment)
 {
+	OutEnvironment.SetDefine(TEXT("IS_INSTANCED_STATIC_MESH_VF"), TEXT("1"));
+
 	if (RHISupportsManualVertexFetch(Parameters.Platform))
 	{
 		OutEnvironment.SetDefineIfUnset(TEXT("MANUAL_VERTEX_FETCH"), TEXT("1"));
@@ -717,16 +721,6 @@ void FInstancedStaticMeshVertexFactory::ModifyCompilationEnvironment(const FVert
 		OutEnvironment.SetDefine(TEXT("USE_INSTANCING"), TEXT("1"));
 	}
 
-	if (IsFeatureLevelSupported(Parameters.Platform, ERHIFeatureLevel::SM5))
-	{
-		OutEnvironment.SetDefine(TEXT("USE_DITHERED_LOD_TRANSITION_FOR_INSTANCED"), true);
-	}
-	else
-	{
-		// On mobile dithered LOD transition has to be explicitly enabled in material and project settings
-		OutEnvironment.SetDefine(TEXT("USE_DITHERED_LOD_TRANSITION_FOR_INSTANCED"), static_cast<bool>(Parameters.MaterialParameters.bIsDitheredLODTransition));
-	}
-
 	FLocalVertexFactory::ModifyCompilationEnvironment(Parameters, OutEnvironment);
 }
 
@@ -1104,7 +1098,7 @@ void FInstancedStaticMeshSceneProxy::GetDynamicMeshElements(const TArray<const F
 				const int32 LODIndex = FMath::Max(GetLOD(View), FirstLODIndex);
 				const FStaticMeshLODResources& LODModel = RenderData->LODResources[LODIndex];
 
-				FInstancedStaticMeshVFLooseUniformShaderParametersRef LooseUniformBuffer = CreateLooseUniformBuffer(View, &UserData_AllInstances, /*InstancedLODRange=*/0, LODIndex, EUniformBufferUsage::UniformBuffer_SingleFrame);
+				FInstancedStaticMeshVFLooseUniformShaderParametersRef LooseUniformBuffer = CreateLooseUniformBuffer(View, /*InstancedLODRange=*/0, LODIndex, EUniformBufferUsage::UniformBuffer_SingleFrame);
 
 				for (int32 SectionIndex = 0; SectionIndex < LODModel.Sections.Num(); SectionIndex++)
 				{
@@ -1181,11 +1175,11 @@ static uint32 GetTypeHash(const FInstancedStaticMeshVFLooseUniformShaderParamete
 {
 	uint32 HashParameters = 0;
 	HashParameters = HashCombine(HashParameters, GetTypeHash(Params.InstancingFadeOutParams));
-	HashParameters = HashCombine(HashParameters, GetTypeHash(Params.InstancingTranslatedWorldViewOriginOne));
-	HashParameters = HashCombine(HashParameters, GetTypeHash(Params.InstancingTranslatedWorldViewOriginZero));
-	HashParameters = HashCombine(HashParameters, GetTypeHash(Params.InstancingViewZCompareOne));
-	HashParameters = HashCombine(HashParameters, GetTypeHash(Params.InstancingViewZCompareZero));
-	HashParameters = HashCombine(HashParameters, GetTypeHash(Params.InstancingViewZConstant));
+	HashParameters = HashCombine(HashParameters, GetTypeHash(Params.InstanceDitherParameters.InstancingTranslatedWorldViewOriginOne));
+	HashParameters = HashCombine(HashParameters, GetTypeHash(Params.InstanceDitherParameters.InstancingTranslatedWorldViewOriginZero));
+	HashParameters = HashCombine(HashParameters, GetTypeHash(Params.InstanceDitherParameters.InstancingViewZCompareOne));
+	HashParameters = HashCombine(HashParameters, GetTypeHash(Params.InstanceDitherParameters.InstancingViewZCompareZero));
+	HashParameters = HashCombine(HashParameters, GetTypeHash(Params.InstanceDitherParameters.InstancingViewZConstant));
 	return HashParameters;
 }
 
@@ -1194,11 +1188,11 @@ static bool operator== (const FInstancedStaticMeshVFLooseUniformShaderParameters
 	static_assert(sizeof(FInstancedStaticMeshVFLooseUniformShaderParameters) == 96, "FInstancedStaticMeshVFLooseUniformShaderParameters size has changed, make sure operator== is updated");
 	
 	return A.InstancingFadeOutParams == B.InstancingFadeOutParams
-		&& A.InstancingTranslatedWorldViewOriginOne == B.InstancingTranslatedWorldViewOriginOne
-		&& A.InstancingTranslatedWorldViewOriginZero == B.InstancingTranslatedWorldViewOriginZero
-		&& A.InstancingViewZCompareOne == B.InstancingViewZCompareOne
-		&& A.InstancingViewZCompareZero == B.InstancingViewZCompareZero
-		&& A.InstancingViewZConstant == B.InstancingViewZConstant;
+		&& A.InstanceDitherParameters.InstancingTranslatedWorldViewOriginOne == B.InstanceDitherParameters.InstancingTranslatedWorldViewOriginOne
+		&& A.InstanceDitherParameters.InstancingTranslatedWorldViewOriginZero == B.InstanceDitherParameters.InstancingTranslatedWorldViewOriginZero
+		&& A.InstanceDitherParameters.InstancingViewZCompareOne == B.InstanceDitherParameters.InstancingViewZCompareOne
+		&& A.InstanceDitherParameters.InstancingViewZCompareZero == B.InstanceDitherParameters.InstancingViewZCompareZero
+		&& A.InstanceDitherParameters.InstancingViewZConstant == B.InstanceDitherParameters.InstancingViewZConstant;
 }
 
 // For deduplicating FInstancedStaticMeshVFLooseUniformShaderParameters UBs
@@ -1242,10 +1236,12 @@ public:
 
 static TGlobalResource<FInstancedStaticMeshVFUniformsCache> GInstancedStaticMeshVFUniformsCache;
 
-FInstancedStaticMeshVFLooseUniformShaderParametersRef FInstancedStaticMeshSceneProxy::CreateLooseUniformBuffer(const FSceneView* View, const FInstancingUserData* InstancingUserData, uint32 InstancedLODRange, uint32 InstancedLODIndex, EUniformBufferUsage UniformBufferUsage) const
+FLocalVertexFactoryInstanceDitherParameters FInstancedStaticMeshSceneProxy::GetInstanceDitherParameters(const FSceneView* View, uint32 InstancedLODRange, uint32 InstancedLODIndex) const
 {
-	FInstancedStaticMeshVFLooseUniformShaderParameters LooseParameters;
+	FLocalVertexFactoryInstanceDitherParameters Result = MakeNullLocalVertexFactoryInstanceDitherParameters();
 
+	// InstancedLODRange is only set for HierarchicalInstancedStaticMeshes
+	if (InstancedLODRange)
 	{
 		FVector4f InstancingViewZCompareZero(MIN_flt, MIN_flt, MAX_flt, 1.0f);
 		FVector4f InstancingViewZCompareOne(MIN_flt, MIN_flt, MAX_flt, 0.0f);
@@ -1254,142 +1250,131 @@ FInstancedStaticMeshVFLooseUniformShaderParametersRef FInstancedStaticMeshSceneP
 		FVector4f InstancingTranslatedWorldViewOriginOne(ForceInit);
 		InstancingTranslatedWorldViewOriginOne.W = 1.0f;
 
-		// InstancedLODRange is only set for HierarchicalInstancedStaticMeshes
-		if (InstancingUserData && InstancedLODRange)
+		check(View);
+		int32 FirstLOD = UserData_AllInstances.MinLOD;
+
+		int32 DebugMin = FMath::Min(CVarMinLOD.GetValueOnRenderThread(), UserData_AllInstances.MeshRenderData->LODResources.Num() - 1);
+		if (DebugMin >= 0)
 		{
-			check(View);
-			int32 FirstLOD = InstancingUserData->MinLOD;
+			FirstLOD = FMath::Max(FirstLOD, DebugMin);
+		}
 
-			int32 DebugMin = FMath::Min(CVarMinLOD.GetValueOnRenderThread(), InstancingUserData->MeshRenderData->LODResources.Num() - 1);
-			if (DebugMin >= 0)
-			{
-				FirstLOD = FMath::Max(FirstLOD, DebugMin);
-			}
+		FBoxSphereBounds ScaledBounds = UserData_AllInstances.MeshRenderData->Bounds.TransformBy(FTransform(FQuat::Identity, FVector::ZeroVector, UserData_AllInstances.AverageInstancesScale));
+		float SphereRadius = ScaledBounds.SphereRadius;
+		float MinSize = View->ViewMatrices.IsPerspectiveProjection() ? CVarFoliageMinimumScreenSize.GetValueOnRenderThread() : 0.0f;
+		float LODScale = UserData_AllInstances.LODDistanceScale;
+		float LODRandom = CVarRandomLODRange.GetValueOnRenderThread();
+		float MaxDrawDistanceScale = GetCachedScalabilityCVars().ViewDistanceScale;
+
+		if (InstancedLODIndex)
+		{
+			InstancingViewZConstant.X = -1.0f;
+		}
+		else
+		{
+			// this is the first LOD, so we don't have a fade-in region
+			InstancingViewZConstant.X = 0.0f;
+		}
+		InstancingViewZConstant.Y = 0.0f;
+		InstancingViewZConstant.Z = 1.0f;
+
+		// now we subtract off the lower segments, since they will be incorporated 
+		InstancingViewZConstant.Y -= InstancingViewZConstant.X;
+		InstancingViewZConstant.Z -= InstancingViewZConstant.X + InstancingViewZConstant.Y;
+		//not using W InstancingViewZConstant.W -= InstancingViewZConstant.X + InstancingViewZConstant.Y + InstancingViewZConstant.Z;
 
-			FBoxSphereBounds ScaledBounds = InstancingUserData->MeshRenderData->Bounds.TransformBy(FTransform(FQuat::Identity, FVector::ZeroVector, InstancingUserData->AverageInstancesScale));
-			float SphereRadius = ScaledBounds.SphereRadius;
-			float MinSize = View->ViewMatrices.IsPerspectiveProjection() ? CVarFoliageMinimumScreenSize.GetValueOnRenderThread() : 0.0f;
-			float LODScale = InstancingUserData->LODDistanceScale;
-			float LODRandom = CVarRandomLODRange.GetValueOnRenderThread();
-			float MaxDrawDistanceScale = GetCachedScalabilityCVars().ViewDistanceScale;
+		for (int32 SampleIndex = 0; SampleIndex < 2; SampleIndex++)
+		{
+			FVector4f& InstancingViewZCompare(SampleIndex ? InstancingViewZCompareOne : InstancingViewZCompareZero);
 
-			if (InstancedLODIndex)
+			float FinalCull = MAX_flt;
+			if (MinSize > 0.0)
 			{
-				InstancingViewZConstant.X = -1.0f;
+				FinalCull = ComputeBoundsDrawDistance(MinSize, SphereRadius, View->ViewMatrices.GetProjectionMatrix()) * LODScale;
 			}
-			else
+			if (UserData_AllInstances.EndCullDistance > 0.0f)
 			{
-				// this is the first LOD, so we don't have a fade-in region
-				InstancingViewZConstant.X = 0.0f;
+				FinalCull = FMath::Min(FinalCull, UserData_AllInstances.EndCullDistance * MaxDrawDistanceScale);
 			}
-			InstancingViewZConstant.Y = 0.0f;
-			InstancingViewZConstant.Z = 1.0f;
-
-			// now we subtract off the lower segments, since they will be incorporated 
-			InstancingViewZConstant.Y -= InstancingViewZConstant.X;
-			InstancingViewZConstant.Z -= InstancingViewZConstant.X + InstancingViewZConstant.Y;
-			//not using W InstancingViewZConstant.W -= InstancingViewZConstant.X + InstancingViewZConstant.Y + InstancingViewZConstant.Z;
 
-			for (int32 SampleIndex = 0; SampleIndex < 2; SampleIndex++)
+			InstancingViewZCompare.Z = FinalCull;
+			if (int(InstancedLODIndex) < UserData_AllInstances.MeshRenderData->LODResources.Num() - 1)
 			{
-				FVector4f& InstancingViewZCompare(SampleIndex ? InstancingViewZCompareOne : InstancingViewZCompareZero);
-
-				float FinalCull = MAX_flt;
-				if (MinSize > 0.0)
-				{
-					FinalCull = ComputeBoundsDrawDistance(MinSize, SphereRadius, View->ViewMatrices.GetProjectionMatrix()) * LODScale;
-				}
-				if (InstancingUserData->EndCullDistance > 0.0f)
-				{
-					FinalCull = FMath::Min(FinalCull, InstancingUserData->EndCullDistance * MaxDrawDistanceScale);
-				}
-
-				InstancingViewZCompare.Z = FinalCull;
-				if (int(InstancedLODIndex) < InstancingUserData->MeshRenderData->LODResources.Num() - 1)
-				{
-					float NextCut = ComputeBoundsDrawDistance(InstancingUserData->MeshRenderData->ScreenSize[InstancedLODIndex + 1].GetValue(), SphereRadius, View->ViewMatrices.GetProjectionMatrix()) * LODScale;
-					InstancingViewZCompare.Z = FMath::Min(NextCut, FinalCull);
-				}
+				float NextCut = ComputeBoundsDrawDistance(UserData_AllInstances.MeshRenderData->ScreenSize[InstancedLODIndex + 1].GetValue(), SphereRadius, View->ViewMatrices.GetProjectionMatrix()) * LODScale;
+				InstancingViewZCompare.Z = FMath::Min(NextCut, FinalCull);
+			}
 
-				InstancingViewZCompare.X = MIN_flt;
-				if (int(InstancedLODIndex) > FirstLOD)
+			InstancingViewZCompare.X = MIN_flt;
+			if (int(InstancedLODIndex) > FirstLOD)
+			{
+				float CurCut = ComputeBoundsDrawDistance(UserData_AllInstances.MeshRenderData->ScreenSize[InstancedLODIndex].GetValue(), SphereRadius, View->ViewMatrices.GetProjectionMatrix()) * LODScale;
+				if (CurCut < FinalCull)
 				{
-					float CurCut = ComputeBoundsDrawDistance(InstancingUserData->MeshRenderData->ScreenSize[InstancedLODIndex].GetValue(), SphereRadius, View->ViewMatrices.GetProjectionMatrix()) * LODScale;
-					if (CurCut < FinalCull)
-					{
-						InstancingViewZCompare.Y = CurCut;
-					}
-					else
-					{
-						// this LOD is completely removed by one of the other two factors
-						InstancingViewZCompare.Y = MIN_flt;
-						InstancingViewZCompare.Z = MIN_flt;
-					}
+					InstancingViewZCompare.Y = CurCut;
 				}
 				else
 				{
-					// this is the first LOD, so we don't have a fade-in region
+					// this LOD is completely removed by one of the other two factors
 					InstancingViewZCompare.Y = MIN_flt;
+					InstancingViewZCompare.Z = MIN_flt;
 				}
 			}
+			else
+			{
+				// this is the first LOD, so we don't have a fade-in region
+				InstancingViewZCompare.Y = MIN_flt;
+			}
+		}
 
-			const FVector PreViewTranslation = View->ViewMatrices.GetPreViewTranslation();
-			InstancingTranslatedWorldViewOriginZero = (FVector3f)(View->GetTemporalLODOrigin(0) + PreViewTranslation); // LWC_TODO: precision loss
-			InstancingTranslatedWorldViewOriginOne = (FVector3f)(View->GetTemporalLODOrigin(1) + PreViewTranslation); // LWC_TODO: precision loss
+		const FVector PreViewTranslation = View->ViewMatrices.GetPreViewTranslation();
+		InstancingTranslatedWorldViewOriginZero = (FVector3f)(View->GetTemporalLODOrigin(0) + PreViewTranslation); // LWC_TODO: precision loss
+		InstancingTranslatedWorldViewOriginOne = (FVector3f)(View->GetTemporalLODOrigin(1) + PreViewTranslation); // LWC_TODO: precision loss
 
-			float Alpha = View->GetTemporalLODTransition();
-			InstancingTranslatedWorldViewOriginZero.W = 1.0f - Alpha;
-			InstancingTranslatedWorldViewOriginOne.W = Alpha;
+		float Alpha = View->GetTemporalLODTransition();
+		InstancingTranslatedWorldViewOriginZero.W = 1.0f - Alpha;
+		InstancingTranslatedWorldViewOriginOne.W = Alpha;
 
-			InstancingViewZCompareZero.W = LODRandom;
-		}
+		InstancingViewZCompareZero.W = LODRandom;
 
-		LooseParameters.InstancingViewZCompareZero = InstancingViewZCompareZero;
-		LooseParameters.InstancingViewZCompareOne = InstancingViewZCompareOne;
-		LooseParameters.InstancingViewZConstant = InstancingViewZConstant;
-		LooseParameters.InstancingTranslatedWorldViewOriginZero = InstancingTranslatedWorldViewOriginZero;
-		LooseParameters.InstancingTranslatedWorldViewOriginOne = InstancingTranslatedWorldViewOriginOne;
+		Result.InstancingViewZCompareZero = InstancingViewZCompareZero;
+		Result.InstancingViewZCompareOne = InstancingViewZCompareOne;
+		Result.InstancingViewZCompareOne.W = 1.0f;
+		Result.InstancingViewZConstant = InstancingViewZConstant;
+		Result.InstancingTranslatedWorldViewOriginZero = InstancingTranslatedWorldViewOriginZero;
+		Result.InstancingTranslatedWorldViewOriginOne = InstancingTranslatedWorldViewOriginOne;
 	}
 
+	return Result;
+}
+
+FVector2f FInstancedStaticMeshSceneProxy::CalcInstancingFadeOutParams() const
+{
+	const float MaxDrawDistanceScale = GetCachedScalabilityCVars().ViewDistanceScale;
+	const float StartDistance = UserData_AllInstances.StartCullDistance * MaxDrawDistanceScale;
+	const float EndDistance = UserData_AllInstances.EndCullDistance * MaxDrawDistanceScale;
+
+	FVector2f Result = FVector2f(StartDistance, 0.f);
+	if (EndDistance > 0.0f)
 	{
-		FVector4f InstancingFadeOutParams(MAX_flt, 0.f, 1.f, 1.f);
-		if (InstancingUserData)
+		if (EndDistance > StartDistance)
 		{
-			const float MaxDrawDistanceScale = GetCachedScalabilityCVars().ViewDistanceScale;
-			const float StartDistance = InstancingUserData->StartCullDistance * MaxDrawDistanceScale;
-			const float EndDistance = InstancingUserData->EndCullDistance * MaxDrawDistanceScale;
-
-			InstancingFadeOutParams.X = StartDistance;
-			if (EndDistance > 0)
-			{
-				if (EndDistance > StartDistance)
-				{
-					InstancingFadeOutParams.Y = 1.f / (float)(EndDistance - StartDistance);
-				}
-				else
-				{
-					InstancingFadeOutParams.Y = 1.f;
-				}
-			}
-			else
-			{
-				InstancingFadeOutParams.Y = 0.f;
-			}
-			if (CVarCullAllInVertexShader.GetValueOnRenderThread() > 0)
-			{
-				InstancingFadeOutParams.Z = 0.0f;
-				InstancingFadeOutParams.W = 0.0f;
-			}
-			else
-			{
-				InstancingFadeOutParams.Z = InstancingUserData->bRenderSelected ? 1.f : 0.f;
-				InstancingFadeOutParams.W = InstancingUserData->bRenderUnselected ? 1.f : 0.f;
-			}
+			Result.Y = 1.f / (float)(EndDistance - StartDistance);
+		}
+		else
+		{
+			Result.Y = 1.f;
 		}
+	}
+	return Result;
+}
 
-		LooseParameters.InstancingFadeOutParams = InstancingFadeOutParams;
+FInstancedStaticMeshVFLooseUniformShaderParametersRef FInstancedStaticMeshSceneProxy::CreateLooseUniformBuffer(const FSceneView* View, uint32 InstancedLODRange, uint32 InstancedLODIndex, EUniformBufferUsage UniformBufferUsage) const
+{
+	FInstancedStaticMeshVFLooseUniformShaderParameters LooseParameters;
 
-	}
+	LooseParameters.InstanceDitherParameters = GetInstanceDitherParameters(View, InstancedLODRange, InstancedLODIndex);
+	const FVector2f InstancingFadeOutParams = CalcInstancingFadeOutParams();
+	LooseParameters.InstancingFadeOutParams = FVector4f(InstancingFadeOutParams.X, InstancingFadeOutParams.Y, 1.0f, 1.0f);
 
 	if (IsInRenderingThread())
 	{
@@ -1435,7 +1420,6 @@ FInstancedStaticMeshSceneProxy::FInstancedStaticMeshSceneProxy(UInstancedStaticM
 FInstancedStaticMeshSceneProxy::FInstancedStaticMeshSceneProxy(const FInstancedStaticMeshSceneProxyDesc& InProxyDesc, ERHIFeatureLevel::Type InFeatureLevel)
 	:	FStaticMeshSceneProxy(InProxyDesc, true)
 	,	StaticMesh(InProxyDesc.GetStaticMesh())
-	,	InstancedRenderData(&InProxyDesc, InFeatureLevel)
 #if WITH_EDITOR
 	,	bHasSelectedInstances(false)
 #endif
@@ -1449,7 +1433,10 @@ FInstancedStaticMeshSceneProxy::FInstancedStaticMeshSceneProxy(const FInstancedS
 #if WITH_EDITOR
 	bHasSelectedInstances = InProxyDesc.bHasSelectedInstances;
 #endif
-
+	if (DoesPlatformUseInstancedStaticMeshVertexFactory(GetScene().GetShaderPlatform()))
+	{
+		InstancedRenderData = MakeUnique<FInstancedStaticMeshRenderData>(&InProxyDesc, InFeatureLevel);
+	}
 	SetupProxy(InProxyDesc);
 }
 
@@ -1528,8 +1515,6 @@ void FInstancedStaticMeshSceneProxy::SetupProxy(const FInstancedStaticMeshSceneP
 	UserData_AllInstances.LODDistanceScale = 1.0f;
 	UserData_AllInstances.InstancingOffset = InProxyDesc.GetStaticMesh()->GetBoundingBox().GetCenter();
 	UserData_AllInstances.MinLOD = ClampedMinLOD;
-	UserData_AllInstances.bRenderSelected = true;
-	UserData_AllInstances.bRenderUnselected = true;
 	UserData_AllInstances.RenderData = nullptr;
 
 	// Only used by HISM, and thus set in the descendant ctor
@@ -1556,14 +1541,25 @@ void FInstancedStaticMeshSceneProxy::SetupProxy(const FInstancedStaticMeshSceneP
 
 bool FInstancedStaticMeshSceneProxy::CanSetupInstancedMeshBatch(int32 LODIndex) const
 {
-	// TODO: If this fires somehow our LODLooseUniformBuffers created in CreateRenderThreadResources is empty, a full dmp should be investigated
-	check(LODLooseUniformBuffers.Num() == InstancedRenderData.VertexFactories.Num());
-	return LODLooseUniformBuffers.Contains(LODIndex) && LODIndex < InstancedRenderData.VertexFactories.Num();
+	if (InstancedRenderData)
+	{
+		// TODO: If this fires somehow our LODLooseUniformBuffers created in CreateRenderThreadResources is empty, a full dmp should be investigated
+		check(LODLooseUniformBuffers.Num() == GetNumRenderDataVertexFactories());
+		return LODLooseUniformBuffers.Contains(LODIndex) && LODIndex < GetNumRenderDataVertexFactories();
+	}
+	return LODIndex < GetNumRenderDataVertexFactories();
 }
 
 bool FInstancedStaticMeshSceneProxy::CanSetupRayTracingInstancedMeshBatch(int32 LODIndex) const
 {
-	return LODIndex < InstancedRenderData.RayTracingProxyVertexFactories.Num();
+	if (InstancedRenderData)
+	{
+		return LODIndex < InstancedRenderData->RayTracingProxyVertexFactories.Num();
+	}
+	else
+	{
+		return LODIndex < RenderData->RayTracingProxy->LODVertexFactories->Num();
+	}
 }
 
 void FInstancedStaticMeshSceneProxy::CreateRenderThreadResources(FRHICommandListBase& RHICmdList)
@@ -1571,22 +1567,26 @@ void FInstancedStaticMeshSceneProxy::CreateRenderThreadResources(FRHICommandList
 	check(InstanceDataSceneProxy);
 	FStaticMeshSceneProxy::CreateRenderThreadResources(RHICmdList);
 
-	InstancedRenderData.BindBuffersToVertexFactories(RHICmdList, InstanceDataSceneProxy->GetLegacyInstanceBuffer());
-	
-	for (int32 i = 0; i < LODs.Num(); ++i)
+	if (InstancedRenderData)
 	{
-		FInstancedStaticMeshVFLooseUniformShaderParametersRef LooseUniformBuffer = CreateLooseUniformBuffer(nullptr, &UserData_AllInstances, 0, i, EUniformBufferUsage::UniformBuffer_MultiFrame);
-		LODLooseUniformBuffers.Add(i, LooseUniformBuffer);
-	}
-	
+		InstancedRenderData->BindBuffersToVertexFactories(RHICmdList, InstanceDataSceneProxy->GetLegacyInstanceBuffer());
+		for (int32 i = 0; i < LODs.Num(); ++i)
+		{
+			FInstancedStaticMeshVFLooseUniformShaderParametersRef LooseUniformBuffer = CreateLooseUniformBuffer(nullptr, 0, i, EUniformBufferUsage::UniformBuffer_MultiFrame);
+			LODLooseUniformBuffers.Add(i, LooseUniformBuffer);
+		}
+	}	
 	// currently we setup loose uniform buffer for ray tracing based on InstancedLODIndex = 0
 	// since currently non of the uniforms depend on the LOD index when InstancedLODRange = 0
-	RayTracingLooseUniformBuffer = CreateLooseUniformBuffer(nullptr, &UserData_AllInstances, /*InstancedLODRange*/ 0, /*InstancedLODIndex*/ 0, EUniformBufferUsage::UniformBuffer_MultiFrame);
+	RayTracingLooseUniformBuffer = CreateLooseUniformBuffer(nullptr, /*InstancedLODRange*/ 0, /*InstancedLODIndex*/ 0, EUniformBufferUsage::UniformBuffer_MultiFrame);
 }
 
 void FInstancedStaticMeshSceneProxy::DestroyRenderThreadResources()
 {
-	InstancedRenderData.ReleaseResources(&GetScene(), StaticMesh);
+	if (InstancedRenderData)
+	{
+		InstancedRenderData->ReleaseResources(&GetScene(), StaticMesh);
+	}
 	FStaticMeshSceneProxy::DestroyRenderThreadResources();
 
 #if RHI_RAYTRACING
@@ -1596,25 +1596,29 @@ void FInstancedStaticMeshSceneProxy::DestroyRenderThreadResources()
 
 void FInstancedStaticMeshSceneProxy::UpdateInstances_RenderThread(FRHICommandListBase& RHICmdList, const FBoxSphereBounds& InBounds, const FBoxSphereBounds& InLocalBounds)
 {
-	if (FStaticMeshInstanceBuffer *LegacyInstanceBuffer = InstanceDataSceneProxy->GetLegacyInstanceBuffer())
+	if (InstancedRenderData)
 	{
-		InstancedRenderData.BindBuffersToVertexFactories(RHICmdList, LegacyInstanceBuffer);
+		if (FStaticMeshInstanceBuffer *LegacyInstanceBuffer = InstanceDataSceneProxy->GetLegacyInstanceBuffer())
+		{
+			InstancedRenderData->BindBuffersToVertexFactories(RHICmdList, LegacyInstanceBuffer);
+		}
 	}
-
 	return FPrimitiveSceneProxy::UpdateInstances_RenderThread(RHICmdList, InBounds, InLocalBounds);
 }
 
 void FInstancedStaticMeshSceneProxy::SetupInstancedMeshBatch(int32 LODIndex, int32 BatchIndex, FMeshBatch& OutMeshBatch) const
 {
-	OutMeshBatch.VertexFactory = &InstancedRenderData.VertexFactories[LODIndex];
-
 	FMeshBatchElement& BatchElement0 = OutMeshBatch.Elements[0];
+	if (InstancedRenderData)
+	{
+		OutMeshBatch.VertexFactory = &InstancedRenderData->VertexFactories[LODIndex];
+		BatchElement0.LooseParametersUniformBuffer = LODLooseUniformBuffers[LODIndex];
+	}
 	BatchElement0.UserData = (void*)&UserData_AllInstances;
 	BatchElement0.bUserDataIsColorVertexBuffer = false;
 	BatchElement0.InstancedLODIndex = LODIndex;
 	BatchElement0.UserIndex = 0;
 	BatchElement0.PrimitiveUniformBuffer = GetUniformBuffer();
-	BatchElement0.LooseParametersUniformBuffer = LODLooseUniformBuffers[LODIndex];
 	BatchElement0.bForceInstanceCulling = true; // force ISM through Generic path even for a single instance cases
 	
 	if (OutMeshBatch.MaterialRenderProxy)
@@ -1636,8 +1640,10 @@ void FInstancedStaticMeshSceneProxy::SetupInstancedMeshBatch(int32 LODIndex, int
 
 void FInstancedStaticMeshSceneProxy::SetupRayTracingInstancedMeshBatch(int32 LODIndex, int32 BatchIndex, FMeshBatch& OutMeshBatch) const
 {
-	OutMeshBatch.VertexFactory = &InstancedRenderData.RayTracingProxyVertexFactories[LODIndex];
-
+	if (InstancedRenderData)
+	{
+		OutMeshBatch.VertexFactory = &InstancedRenderData->RayTracingProxyVertexFactories[LODIndex];
+	}
 	FMeshBatchElement& BatchElement0 = OutMeshBatch.Elements[0];
 	BatchElement0.UserData = (void*)&UserData_AllInstances;
 	BatchElement0.bUserDataIsColorVertexBuffer = false;
@@ -1750,6 +1756,16 @@ bool FInstancedStaticMeshSceneProxy::GetInstanceDrawDistanceMinMax(FVector2f& Ou
 	}
 }
 
+bool FInstancedStaticMeshSceneProxy::GetInstancingFadeOutParams(FVector2f& OutInstancingFadeOutParams) const
+{
+	if (UserData_AllInstances.StartCullDistance > 0)
+	{
+		OutInstancingFadeOutParams = CalcInstancingFadeOutParams();;
+		return true;
+	}
+	return false;
+}
+
 float FInstancedStaticMeshSceneProxy::GetLodScreenSizeScale() const
 {
 	return InstanceLODDistanceScale > 0.f ? 1.f / InstanceLODDistanceScale : 1.f;
@@ -1788,6 +1804,18 @@ void FInstancedStaticMeshSceneProxy::SetEvaluateWorldPositionOffsetInRayTracing(
 #endif // RHI_RAYTRACING
 }
 
+int32 FInstancedStaticMeshSceneProxy::GetNumRenderDataVertexFactories() const
+{
+	if (InstancedRenderData)
+	{
+		return InstancedRenderData->VertexFactories.Num();
+	}
+	else
+	{
+		return RenderData->LODVertexFactories.Num();
+	}
+}
+
 #if RHI_RAYTRACING
 void FInstancedStaticMeshSceneProxy::GetDynamicRayTracingInstances(FRayTracingInstanceCollector& Collector)
 {
@@ -1878,7 +1906,15 @@ void FInstancedStaticMeshSceneProxy::GetDynamicRayTracingInstances(FRayTracingIn
 	TArray<FMeshBatch> RayTracingMeshBatches;
 
 	{
-		int32 SectionCount = (*InstancedRenderData.RayTracingProxyLODs)[LODIndex].Sections->Num();
+		int32 SectionCount = 0;
+		if (InstancedRenderData)
+		{
+			SectionCount = (*InstancedRenderData->RayTracingProxyLODs)[LODIndex].Sections->Num();
+		}
+		else
+		{
+			SectionCount = RenderData->RayTracingProxy->LODs[LODIndex].Sections->Num();
+		}
 
 		for (int32 SectionIdx = 0; SectionIdx < SectionCount; ++SectionIdx)
 		{
@@ -1898,13 +1934,15 @@ void FInstancedStaticMeshSceneProxy::GetDynamicRayTracingInstances(FRayTracingIn
 
 			if (bResult)
 			{
-				if (RenderData->RayTracingProxy->bUsingRenderingLODs)
+				if (InstancedRenderData)
 				{
-					check(MeshBatch.VertexFactory == &InstancedRenderData.VertexFactories[LODIndex]);
+					auto VF = RenderData->RayTracingProxy->bUsingRenderingLODs ? &InstancedRenderData->VertexFactories[LODIndex] : &InstancedRenderData->RayTracingProxyVertexFactories[LODIndex];
+					check(MeshBatch.VertexFactory == VF);
 				}
 				else
 				{
-					check(MeshBatch.VertexFactory == &InstancedRenderData.RayTracingProxyVertexFactories[LODIndex]);
+					auto VF = RenderData->RayTracingProxy->bUsingRenderingLODs ? &RenderData->LODVertexFactories[LODIndex].VertexFactory : &(*RenderData->RayTracingProxy->LODVertexFactories)[LODIndex].VertexFactory;
+					check(MeshBatch.VertexFactory == VF);
 				}
 
 				MeshBatch.ReverseCulling = bReverseCulling;
@@ -2100,7 +2138,10 @@ void FInstancedStaticMeshSceneProxy::GetDynamicRayTracingInstances(FRayTracingIn
 		FMatrix InstanceToWorld = InstanceSceneDataBuffers->GetInstanceToWorld(InstanceIndex);
 		const uint32 DynamicInstanceIdx = InstanceIndex % SimulatedInstances;
 
-		if (bDynamicRayTracingGeometry && InstancedRenderData.RayTracingProxyVertexFactories[LODIndex].GetType()->SupportsRayTracingDynamicGeometry())
+		const bool bSupportsRayTracingDynamicGeometry = InstancedRenderData 
+			? InstancedRenderData->RayTracingProxyVertexFactories[LODIndex].GetType()->SupportsRayTracingDynamicGeometry()
+			: (*RenderData->RayTracingProxy->LODVertexFactories)[LODIndex].VertexFactory.GetType()->SupportsRayTracingDynamicGeometry();
+		if (bDynamicRayTracingGeometry && bSupportsRayTracingDynamicGeometry)
 		{
 			FRayTracingInstance* DynamicInstance = nullptr;
 
@@ -2178,7 +2219,6 @@ void FInstancedStaticMeshSceneProxy::GetDynamicRayTracingInstances(FRayTracingIn
 	}
 }
 
-
 void FInstancedStaticMeshSceneProxy::SetupRayTracingDynamicInstances(FRHICommandListBase& RHICmdList, int32 NumDynamicInstances, int32 LODIndex)
 {
 	if (NumDynamicInstances == RayTracingDynamicInstanceData.Num() && LODIndex == CachedRayTracingLODIndex)
@@ -2263,15 +2303,17 @@ void FInstancedStaticMeshSceneProxy::SetInstanceCullDistance_RenderThread(float
 	UserData_AllInstances.StartCullDistance = StartCullDistance;
 	UserData_AllInstances.EndCullDistance = EndCullDistance;
 
-	for (int32 i = 0; i < LODs.Num(); ++i)
+	if (InstancedRenderData)
 	{
-		FInstancedStaticMeshVFLooseUniformShaderParametersRef LooseUniformBuffer = CreateLooseUniformBuffer(nullptr, &UserData_AllInstances, 0, i, EUniformBufferUsage::UniformBuffer_MultiFrame);
-		LODLooseUniformBuffers.Add(i, LooseUniformBuffer);
+		for (int32 i = 0; i < LODs.Num(); ++i)
+		{
+			FInstancedStaticMeshVFLooseUniformShaderParametersRef LooseUniformBuffer = CreateLooseUniformBuffer(nullptr, 0, i, EUniformBufferUsage::UniformBuffer_MultiFrame);
+			LODLooseUniformBuffers.Add(i, LooseUniformBuffer);
+		}
 	}
-
 	// currently we setup loose uniform buffer for ray tracing based on InstancedLODIndex = 0
 	// since currently non of the uniforms depend on the LOD index when InstancedLODRange = 0
-	RayTracingLooseUniformBuffer = CreateLooseUniformBuffer(nullptr, &UserData_AllInstances, /*InstancedLODRange*/ 0, /*InstancedLODIndex*/ 0, EUniformBufferUsage::UniformBuffer_MultiFrame);
+	RayTracingLooseUniformBuffer = CreateLooseUniformBuffer(nullptr, /*InstancedLODRange*/ 0, /*InstancedLODIndex*/ 0, EUniformBufferUsage::UniformBuffer_MultiFrame);
 }
 
 void FInstancedStaticMeshSceneProxy::GetPreSkinnedLocalBounds(FBoxSphereBounds& OutBounds) const
diff --git a/Engine/Source/Runtime/Engine/Private/LocalVertexFactory.cpp b/Engine/Source/Runtime/Engine/Private/LocalVertexFactory.cpp
--- a/Engine/Source/Runtime/Engine/Private/LocalVertexFactory.cpp
+++ b/Engine/Source/Runtime/Engine/Private/LocalVertexFactory.cpp
@@ -69,6 +69,17 @@ void FLocalVertexFactoryShaderParametersBase::Bind(const FShaderParameterMap& Pa
 
 IMPLEMENT_GLOBAL_SHADER_PARAMETER_STRUCT(FLocalVertexFactoryUniformShaderParameters, "LocalVF");
 
+FLocalVertexFactoryInstanceDitherParameters MakeNullLocalVertexFactoryInstanceDitherParameters()
+{
+	FLocalVertexFactoryInstanceDitherParameters InstanceDitherParameters;
+	InstanceDitherParameters.InstancingViewZCompareZero = 				FVector4f(MIN_flt, MIN_flt, MAX_flt, 1.0f);
+	InstanceDitherParameters.InstancingViewZCompareOne = 				FVector4f(MIN_flt, MIN_flt, MAX_flt, 0.0f);
+	InstanceDitherParameters.InstancingViewZConstant = 					FVector4f(ForceInit);
+	InstanceDitherParameters.InstancingTranslatedWorldViewOriginZero =	FVector4f(ForceInit);
+	InstanceDitherParameters.InstancingTranslatedWorldViewOriginOne = 	FVector4f(0.0f, 0.0f, 0.0f, 1.0f);
+	return InstanceDitherParameters;
+}
+
 void GetLocalVFUniformShaderParameters(
 	FLocalVertexFactoryUniformShaderParameters& UniformParameters,
 	const FLocalVertexFactory* LocalVertexFactory, 
@@ -121,6 +132,7 @@ void GetLocalVFUniformShaderParameters(
 	UniformParameters.VertexFetch_Parameters = {ColorIndexMask, NumTexCoords, LightMapCoordinateIndex, EffectiveBaseVertexIndex};
 	UniformParameters.PreSkinBaseVertexIndex = EffectivePreSkinBaseVertexIndex;
 	UniformParameters.bIsGPUSkinPassThrough = LocalVertexFactory->IsGPUSkinPassThrough() ? 1 : 0;
+	UniformParameters.InstanceDitherParameters = MakeNullLocalVertexFactoryInstanceDitherParameters();
 }
 
 TUniformBufferRef<FLocalVertexFactoryUniformShaderParameters> CreateLocalVFUniformBuffer(
@@ -128,12 +140,18 @@ TUniformBufferRef<FLocalVertexFactoryUniformShaderParameters> CreateLocalVFUnifo
 	uint32 LODLightmapDataIndex, 
 	FColorVertexBuffer* OverrideColorVertexBuffer, 
 	int32 BaseVertexIndex,
-	int32 PreSkinBaseVertexIndex
+	int32 PreSkinBaseVertexIndex,
+	const FLocalVertexFactoryInstanceDitherParameters* OptionalInstanceDitherParameters,
+	EUniformBufferUsage UniformBufferUsage 
 	)
 {
 	FLocalVertexFactoryUniformShaderParameters UniformParameters;
 	GetLocalVFUniformShaderParameters(UniformParameters, LocalVertexFactory, LODLightmapDataIndex, OverrideColorVertexBuffer, BaseVertexIndex, PreSkinBaseVertexIndex);
-	return TUniformBufferRef<FLocalVertexFactoryUniformShaderParameters>::CreateUniformBufferImmediate(UniformParameters, UniformBuffer_MultiFrame);
+	if (OptionalInstanceDitherParameters)
+	{
+		UniformParameters.InstanceDitherParameters = *OptionalInstanceDitherParameters;
+	}
+	return TUniformBufferRef<FLocalVertexFactoryUniformShaderParameters>::CreateUniformBufferImmediate(UniformParameters, UniformBufferUsage);
 }
 
 void FLocalVertexFactoryShaderParametersBase::GetElementShaderBindingsBase(
@@ -269,6 +287,22 @@ FLocalVertexFactory::FLocalVertexFactory(ERHIFeatureLevel::Type InFeatureLevel,
 
 FLocalVertexFactory::~FLocalVertexFactory() = default;
 
+static bool GUseInstancedStaticMeshVertexFactory = false;
+FAutoConsoleVariableRef CVarUseInstancedStaticMeshVertexFactory(
+	TEXT("r.InstancedStaticMeshes.UseInstancedStaticMeshVertexFactory"),
+	GUseInstancedStaticMeshVertexFactory,
+	TEXT("If enabled (default is off), all platforms use a separate vertex factory for the instanced static meshes which increases the number of shader permutations significantly.\n")
+	TEXT("  When disabled some platforms - notably those without full GPU scene - will still enable the separate ISM VF."),
+	ECVF_ReadOnly);
+
+bool DoesPlatformUseInstancedStaticMeshVertexFactory(const FStaticShaderPlatform ShaderPlatform)
+{
+	const bool bPlatformSupportsFullGPUScene = UseGPUScene(ShaderPlatform) && !PlatformGPUSceneUsesUniformBufferView(ShaderPlatform);
+	// If the platform doesn't support full GPU-Scene it must use the separate ISM VF
+	return GUseInstancedStaticMeshVertexFactory
+		|| !bPlatformSupportsFullGPUScene;
+}
+
 /**
  * Should we cache the material's shadertype on this platform with this vertex factory? 
  */
@@ -306,6 +340,13 @@ void FLocalVertexFactory::ModifyCompilationEnvironment(const FVertexFactoryShade
 	const bool bVFSupportsPrimtiveSceneData = Parameters.VertexFactoryType->SupportsPrimitiveIdStream() && UseGPUScene(Parameters.Platform, GetMaxSupportedFeatureLevel(Parameters.Platform));
 	OutEnvironment.SetDefine(TEXT("VF_SUPPORTS_PRIMITIVE_SCENE_DATA"), bVFSupportsPrimtiveSceneData);
 
+	if (!DoesPlatformUseInstancedStaticMeshVertexFactory(Parameters.Platform))
+	{
+		OutEnvironment.SetDefine(TEXT("USE_INSTANCE_CULLING"), TEXT("1"));
+		// Note: is defaulted off, but set explicitly here anyway for clarity
+		OutEnvironment.SetDefine(TEXT("IS_INSTANCED_STATIC_MESH_VF"), TEXT("0"));
+	}
+
 	// When combining ray tracing and WPO, leave the mesh in local space for consistency with how shading normals are calculated.
 	// See UE-139634 for the case that lead to this.
 	OutEnvironment.SetDefine(TEXT("RAY_TRACING_DYNAMIC_MESH_IN_LOCAL_SPACE"), TEXT("1"));
diff --git a/Engine/Source/Runtime/Engine/Private/PrimitiveSceneProxy.cpp b/Engine/Source/Runtime/Engine/Private/PrimitiveSceneProxy.cpp
--- a/Engine/Source/Runtime/Engine/Private/PrimitiveSceneProxy.cpp
+++ b/Engine/Source/Runtime/Engine/Private/PrimitiveSceneProxy.cpp
@@ -853,6 +853,11 @@ void FPrimitiveSceneProxy::BuildUniformShaderParameters(FPrimitiveUniformShaderP
 		Builder.InstanceDrawDistance(InstanceDrawDistanceMinMax);
 	}
 
+	// Note, always set the value to useful defaults (InstancingFadeOutParams.y == 0.0f means it works out to fully visible)
+	FVector2f InstancingFadeOutParams = FVector2f::ZeroVector;
+	GetInstancingFadeOutParams(InstancingFadeOutParams);
+	Builder.InstancingFadeOutParams(InstancingFadeOutParams);
+
 	float WPODisableDistance;
 	if (GetInstanceWorldPositionOffsetDisableDistance(WPODisableDistance))
 	{
diff --git a/Engine/Source/Runtime/Engine/Private/PrimitiveUniformShaderParameters.cpp b/Engine/Source/Runtime/Engine/Private/PrimitiveUniformShaderParameters.cpp
--- a/Engine/Source/Runtime/Engine/Private/PrimitiveUniformShaderParameters.cpp
+++ b/Engine/Source/Runtime/Engine/Private/PrimitiveUniformShaderParameters.cpp
@@ -192,8 +192,8 @@ void FPrimitiveSceneShaderData::Setup(const FPrimitiveUniformShaderParameters& P
 		PrimitiveUniformShaderParameters.MaterialCacheDescriptor,
 		// Note: AnimationMinScreenSize could be packed to fewer bits as it needs to represent values in the range 0-1, and special case -1.
 		PrimitiveUniformShaderParameters.AnimationMinScreenSize,
-		0u /*unused*/,
-		0u /*unused*/);
+		PrimitiveUniformShaderParameters.InstancingFadeOutParams.X,
+		PrimitiveUniformShaderParameters.InstancingFadeOutParams.Y);
 
 	// Set all the custom primitive data float4. This matches the loop in SceneData.ush
 	const int32 CustomPrimitiveDataStartIndex = 35;
diff --git a/Engine/Source/Runtime/Engine/Public/InstancedStaticMeshComponentHelper.h b/Engine/Source/Runtime/Engine/Public/InstancedStaticMeshComponentHelper.h
--- a/Engine/Source/Runtime/Engine/Public/InstancedStaticMeshComponentHelper.h
+++ b/Engine/Source/Runtime/Engine/Public/InstancedStaticMeshComponentHelper.h
@@ -50,36 +50,43 @@ private:
 template<class T>
 void FInstancedStaticMeshComponentHelper::CollectPSOPrecacheData(const T& Component, const FPSOPrecacheParams& BasePrecachePSOParams, FMaterialInterfacePSOPrecacheParamsList& OutParams)
 {
-	if (Component.GetStaticMesh() == nullptr || Component.GetStaticMesh()->GetRenderData() == nullptr)
+	if (DoesPlatformUseInstancedStaticMeshVertexFactory(GMaxRHIShaderPlatform))
 	{
-		return;
-	}
+		if (Component.GetStaticMesh() == nullptr || Component.GetStaticMesh()->GetRenderData() == nullptr)
+		{
+			return;
+		}
 
-	const bool bCanUseGPUScene = UseGPUScene(GMaxRHIShaderPlatform, GMaxRHIFeatureLevel);
-	FStaticMeshInstanceBuffer* InstanceBuffer = bCanUseGPUScene ? nullptr : &FInstancedStaticMeshComponentHelper::DummyStaticMeshInstanceBuffer;
-	int32 LightMapCoordinateIndex = Component.GetStaticMesh()->GetLightMapCoordinateIndex();
+		const bool bCanUseGPUScene = UseGPUScene(GMaxRHIShaderPlatform, GMaxRHIFeatureLevel);
+		FStaticMeshInstanceBuffer* InstanceBuffer = bCanUseGPUScene ? nullptr : &FInstancedStaticMeshComponentHelper::DummyStaticMeshInstanceBuffer;
+		int32 LightMapCoordinateIndex = Component.GetStaticMesh()->GetLightMapCoordinateIndex();
 
-	auto ISMC_GetElements = [LightMapCoordinateIndex, InstanceBuffer, &Component](const FStaticMeshLODResources& LODRenderData, int32 LODIndex, bool bSupportsManualVertexFetch, FVertexDeclarationElementList& Elements)
-	{
-		FInstancedStaticMeshDataType InstanceData;
-		FInstancedStaticMeshVertexFactory::FDataType Data;
-		const FColorVertexBuffer* ColorVertexBuffer = LODRenderData.bHasColorVertexData ? &(LODRenderData.VertexBuffers.ColorVertexBuffer) : nullptr;
-		if (Component.LODData.IsValidIndex(LODIndex) && Component.LODData[LODIndex].OverrideVertexColors)
+		auto ISMC_GetElements = [LightMapCoordinateIndex, InstanceBuffer, &Component](const FStaticMeshLODResources& LODRenderData, int32 LODIndex, bool bSupportsManualVertexFetch, FVertexDeclarationElementList& Elements)
 		{
-			ColorVertexBuffer = Component.LODData[LODIndex].OverrideVertexColors;
-		}
-		FInstancedStaticMeshVertexFactory::InitInstancedStaticMeshVertexFactoryComponents(LODRenderData.VertexBuffers, ColorVertexBuffer, InstanceBuffer, nullptr /*VertexFactory*/, LightMapCoordinateIndex, bSupportsManualVertexFetch, Data, InstanceData);
-		FInstancedStaticMeshVertexFactory::GetVertexElements(GMaxRHIFeatureLevel, EVertexInputStreamType::Default, bSupportsManualVertexFetch, Data, InstanceData, Elements);
-	};
+			FInstancedStaticMeshDataType InstanceData;
+			FInstancedStaticMeshVertexFactory::FDataType Data;
+			const FColorVertexBuffer* ColorVertexBuffer = LODRenderData.bHasColorVertexData ? &(LODRenderData.VertexBuffers.ColorVertexBuffer) : nullptr;
+			if (Component.LODData.IsValidIndex(LODIndex) && Component.LODData[LODIndex].OverrideVertexColors)
+			{
+				ColorVertexBuffer = Component.LODData[LODIndex].OverrideVertexColors;
+			}
+			FInstancedStaticMeshVertexFactory::InitInstancedStaticMeshVertexFactoryComponents(LODRenderData.VertexBuffers, ColorVertexBuffer, InstanceBuffer, nullptr /*VertexFactory*/, LightMapCoordinateIndex, bSupportsManualVertexFetch, Data, InstanceData);
+			FInstancedStaticMeshVertexFactory::GetVertexElements(GMaxRHIFeatureLevel, EVertexInputStreamType::Default, bSupportsManualVertexFetch, Data, InstanceData, Elements);
+		};
 
-	Nanite::FMaterialAudit NaniteMaterials{};
-	if (Nanite::FNaniteResourcesHelper::ShouldCreateNaniteProxy(Component, &NaniteMaterials))
-	{
-		FStaticMeshComponentHelper::CollectPSOPrecacheDataImpl(Component, &FNaniteVertexFactory::StaticType, BasePrecachePSOParams, ISMC_GetElements, OutParams);
+		Nanite::FMaterialAudit NaniteMaterials{};
+		if (Nanite::FNaniteResourcesHelper::ShouldCreateNaniteProxy(Component, &NaniteMaterials))
+		{
+			FStaticMeshComponentHelper::CollectPSOPrecacheDataImpl(Component, &FNaniteVertexFactory::StaticType, BasePrecachePSOParams, ISMC_GetElements, OutParams);
+		}
+		else
+		{
+			FStaticMeshComponentHelper::CollectPSOPrecacheDataImpl(Component, &FInstancedStaticMeshVertexFactory::StaticType, BasePrecachePSOParams, ISMC_GetElements, OutParams);
+		}
 	}
 	else
 	{
-		FStaticMeshComponentHelper::CollectPSOPrecacheDataImpl(Component, &FInstancedStaticMeshVertexFactory::StaticType, BasePrecachePSOParams, ISMC_GetElements, OutParams);
+		FStaticMeshComponentHelper::CollectPSOPrecacheData(Component, BasePrecachePSOParams, OutParams);
 	}
 }
 
diff --git a/Engine/Source/Runtime/Engine/Public/LocalVertexFactory.h b/Engine/Source/Runtime/Engine/Public/LocalVertexFactory.h
--- a/Engine/Source/Runtime/Engine/Public/LocalVertexFactory.h
+++ b/Engine/Source/Runtime/Engine/Public/LocalVertexFactory.h
@@ -16,8 +16,17 @@ struct FMeshBatchElement;
 	LocalVertexFactory.h: Local vertex factory definitions.
 =============================================================================*/
 
+BEGIN_SHADER_PARAMETER_STRUCT(FLocalVertexFactoryInstanceDitherParameters, ENGINE_API)
+	SHADER_PARAMETER(FVector4f, InstancingViewZCompareZero)
+	SHADER_PARAMETER(FVector4f, InstancingViewZCompareOne)
+	SHADER_PARAMETER(FVector4f, InstancingViewZConstant)
+	SHADER_PARAMETER(FVector4f, InstancingTranslatedWorldViewOriginZero)
+	SHADER_PARAMETER(FVector4f, InstancingTranslatedWorldViewOriginOne)
+END_SHADER_PARAMETER_STRUCT()
+
 BEGIN_GLOBAL_SHADER_PARAMETER_STRUCT(FLocalVertexFactoryUniformShaderParameters,ENGINE_API)
 	SHADER_PARAMETER(FIntVector4,VertexFetch_Parameters)
+	SHADER_PARAMETER_STRUCT_INCLUDE(FLocalVertexFactoryInstanceDitherParameters, InstanceDitherParameters)
 	SHADER_PARAMETER(int32, PreSkinBaseVertexIndex)
 	SHADER_PARAMETER(uint32,LODLightmapDataIndex)
 	SHADER_PARAMETER(uint32,bIsGPUSkinPassThrough)
@@ -28,6 +37,8 @@ BEGIN_GLOBAL_SHADER_PARAMETER_STRUCT(FLocalVertexFactoryUniformShaderParameters,
 	SHADER_PARAMETER_SRV(Buffer<float4>, VertexFetch_ColorComponentsBuffer)
 END_GLOBAL_SHADER_PARAMETER_STRUCT()
 
+ENGINE_API FLocalVertexFactoryInstanceDitherParameters MakeNullLocalVertexFactoryInstanceDitherParameters();
+
 extern ENGINE_API void GetLocalVFUniformShaderParameters(
 	FLocalVertexFactoryUniformShaderParameters& UniformParameters,
 	const class FLocalVertexFactory* VertexFactory, 
@@ -42,7 +53,9 @@ extern ENGINE_API TUniformBufferRef<FLocalVertexFactoryUniformShaderParameters>
 	uint32 LODLightmapDataIndex, 
 	class FColorVertexBuffer* OverrideColorVertexBuffer, 
 	int32 BaseVertexIndex,
-	int32 PreSkinBaseVertexIndex
+	int32 PreSkinBaseVertexIndex,
+	const FLocalVertexFactoryInstanceDitherParameters* OptionalInstanceDitherParameters = nullptr,
+	EUniformBufferUsage UniformBufferUsage = UniformBuffer_MultiFrame
 	);
 
 extern ENGINE_API void AddStaticMeshTextureCoordinateElements(
@@ -51,6 +64,10 @@ extern ENGINE_API void AddStaticMeshTextureCoordinateElements(
 	FVertexDeclarationElementList& Elements,
 	FVertexFactory::FVertexStreamList& InOutStreams);
 
+/**
+ */
+ENGINE_API bool DoesPlatformUseInstancedStaticMeshVertexFactory(const FStaticShaderPlatform Platform);
+
 /**
  * A vertex factory which simply transforms explicit vertex attributes from local to world space.
  */
diff --git a/Engine/Source/Runtime/Engine/Public/PrimitiveSceneProxy.h b/Engine/Source/Runtime/Engine/Public/PrimitiveSceneProxy.h
--- a/Engine/Source/Runtime/Engine/Public/PrimitiveSceneProxy.h
+++ b/Engine/Source/Runtime/Engine/Public/PrimitiveSceneProxy.h
@@ -1039,6 +1039,14 @@ public:
 		OutDistanceMinMax = FVector2f(0.0f);
 		return false;
 	}
+	
+	/**
+	 * Note: should not modify the OutInstancingFadeOutParams output unless it actually has a valid output - the default is initialized.
+	 */
+	virtual bool GetInstancingFadeOutParams(FVector2f& OutInstancingFadeOutParams) const
+	{
+		return false;
+	}
 
 	/**
 	 * Retrieves the per-instance world position offset disable distance
diff --git a/Engine/Source/Runtime/Engine/Public/PrimitiveUniformShaderParameters.h b/Engine/Source/Runtime/Engine/Public/PrimitiveUniformShaderParameters.h
--- a/Engine/Source/Runtime/Engine/Public/PrimitiveUniformShaderParameters.h
+++ b/Engine/Source/Runtime/Engine/Public/PrimitiveUniformShaderParameters.h
@@ -62,6 +62,7 @@ BEGIN_GLOBAL_SHADER_PARAMETER_STRUCT(FPrimitiveUniformShaderParameters,ENGINE_AP
 	SHADER_PARAMETER(float,			PixelProgrammableDistanceSquared)
 	SHADER_PARAMETER(float,			MaterialDisplacementFadeOutSize)
 	SHADER_PARAMETER(float,			AnimationMinScreenSize)
+	SHADER_PARAMETER(FVector2f,		InstancingFadeOutParams)
 	SHADER_PARAMETER(FUintVector2,	MeshPaintTextureDescriptor)
 	SHADER_PARAMETER(uint32,        MaterialCacheDescriptor)
 	SHADER_PARAMETER(uint32,		NaniteRayTracingDataOffset)
diff --git a/Engine/Source/Runtime/Engine/Public/PrimitiveUniformShaderParametersBuilder.h b/Engine/Source/Runtime/Engine/Public/PrimitiveUniformShaderParametersBuilder.h
--- a/Engine/Source/Runtime/Engine/Public/PrimitiveUniformShaderParametersBuilder.h
+++ b/Engine/Source/Runtime/Engine/Public/PrimitiveUniformShaderParametersBuilder.h
@@ -156,6 +156,7 @@ public:
 	PRIMITIVE_UNIFORM_BUILDER_METHOD(FUintVector2,		MeshPaintTextureDescriptor);
 	PRIMITIVE_UNIFORM_BUILDER_METHOD(uint32,	        MaterialCacheDescriptor);
 	PRIMITIVE_UNIFORM_BUILDER_METHOD(float,				AnimationMinScreenSize);
+	PRIMITIVE_UNIFORM_BUILDER_METHOD(FVector2f,			InstancingFadeOutParams);
 
 #undef PRIMITIVE_UNIFORM_BUILDER_FLAG_METHOD
 #undef PRIMITIVE_UNIFORM_BUILDER_METHOD
