Disable collision and overlaps on SplineComponent, which does not set up a valid BodyInstance and hides collision settings in the details panel

#tests PIE
[RN] fix Framework/Gameplay Disable collision and overlaps on SplineComponent, any subclasses with collision will need to explicitly enable collision now.
#Synced-CL 48614854

diff --git a/Engine/Source/Runtime/Engine/Private/Components/SplineComponent.cpp b/Engine/Source/Runtime/Engine/Private/Components/SplineComponent.cpp
--- a/Engine/Source/Runtime/Engine/Private/Components/SplineComponent.cpp
+++ b/Engine/Source/Runtime/Engine/Private/Components/SplineComponent.cpp
@@ -319,6 +319,9 @@ USplineComponent::USplineComponent(const FObjectInitializer& ObjectInitializer)
 	, ScaleVisualizationWidth(30.0f)
 #endif
 {
+	SetCollisionEnabled(ECollisionEnabled::NoCollision);
+	SetGenerateOverlapEvents(false);
+
 	SetDefaultSpline();
 
 #if WITH_EDITORONLY_DATA
