[NetPhysics] Adding an option to scale input decay with measured input lag

#Synced-CL 48614097

diff --git a/Engine/Source/Runtime/Engine/Private/PhysicsEngine/NetworkPhysicsComponent.cpp b/Engine/Source/Runtime/Engine/Private/PhysicsEngine/NetworkPhysicsComponent.cpp
--- a/Engine/Source/Runtime/Engine/Private/PhysicsEngine/NetworkPhysicsComponent.cpp
+++ b/Engine/Source/Runtime/Engine/Private/PhysicsEngine/NetworkPhysicsComponent.cpp
@@ -96,10 +96,16 @@ namespace PhysicsReplicationCVars
 		static FAutoConsoleVariableRef CVarApplyInputDecayOverSetTime(TEXT("np2.Resim.ApplyInputDecayOverSetTime"), bApplyInputDecayOverSetTime, TEXT("When true, apply the Input Decay Curve over a set amount of time instead of over the start of input prediction and end of resim which is variable each resimulation"));
 		float InputDecaySetTime = 0.15f;
 		static FAutoConsoleVariableRef CVarInputDecaySetTime(TEXT("np2.Resim.InputDecaySetTime"), InputDecaySetTime, TEXT("Applied when np2.Resim.ApplyInputDecayOverSetTime is true, read there for more info. Set time to apply Input Decay Curve over while predicting inputs during resimulation"));
+		bool bEnableLagScaledInputDecay = false;
+		static FAutoConsoleVariableRef CVarEnableLagScaledInputDecay(TEXT("np2.Resim.EnableLagScaledInputDecay"), bEnableLagScaledInputDecay, TEXT("If true, scales input decay as a proportion of measured input lag compared to the reference input lag (np2.Resim.InputDecayReferenceLagMs)"));
+		float InputDecayReferenceLagMs = 100.0f;
+		static FAutoConsoleVariableRef CVarInputDecayReferenceLagMs(TEXT("np2.Resim.InputDecayReferenceLagMs"), InputDecayReferenceLagMs, TEXT("The reference input lag in milliseconds used to scale input decay. Only applies if np2.Resim.EnableLagScaledInputDecay is true"));
 		bool bApplyInputDecaySimProxyInputAtRuntime = false;
 		static FAutoConsoleVariableRef CVarApplyInputDecaySimProxyInputAtRuntime(TEXT("np2.Resim.ApplyInputDecaySimProxyInputAtRuntime"), bApplyInputDecaySimProxyInputAtRuntime, TEXT("When true, apply input decay on inputs applied on simulated proxies at runtime (outside of resimulation)"));
 		float InputDecaySimProxyInputAtRuntime = 0.25f;
 		static FAutoConsoleVariableRef CVarInputDecaySimProxyInputAtRuntime(TEXT("np2.Resim.InputDecaySimProxyInputAtRuntime"), InputDecaySimProxyInputAtRuntime, TEXT("The amount of input decay to apply on simulated proxy inputs at runtime (outside of resim), if np2.Resim.ApplyInputDecaySimProxyInputAtRuntime is true"));
+		bool bEnableLagScalingSimProxyRuntimeInputDecay = false;
+		static FAutoConsoleVariableRef CVarEnableLagScalingSimProxyRuntimeInputDecay(TEXT("np2.Resim.EnableLagScalingSimProxyRuntimeInputDecay"), bEnableLagScalingSimProxyRuntimeInputDecay, TEXT("When true, scale 'runtime' simulated proxy input decay with measured lag, just like resim input decay"));
 
 		bool bSimulationDecayEnable = false;
 		static FAutoConsoleVariableRef CVarEnableSimulationDecay(TEXT("np2.Resim.SimulationDecay.Enable"), bSimulationDecayEnable, TEXT("When enabled, apply simulation decay on simulated proxies by essentially multiplying their physics step by a time scale value. Effectively forward predicting sim-proxies less, which can make desync less during forward prediction."));
@@ -2276,7 +2282,8 @@ void FAsyncNetworkPhysicsComponent::OnPreProcessInputs_Internal(const int32 Phys
 								// Perform input decay on sim-proxy runtime (non-resim) input
 								if (bApplySimProxyInput && ComponentSettings.GetApplyInputDecaySimProxyInputAtRuntime())
 								{
-									PhysicsData->DecayData(ComponentSettings.GetInputDecaySimProxyInputAtRuntime());
+									const float RuntimeInputDecay = GetCurrentSimProxyInputDecayAtRuntime();
+									PhysicsData->DecayData(RuntimeInputDecay);
 								}
 
 								if (bIsUsingLegacyData)
@@ -2319,7 +2326,8 @@ void FAsyncNetworkPhysicsComponent::OnPreProcessInputs_Internal(const int32 Phys
 				// Perform input decay on sim-proxy runtime (non-resim) input
 				if (PhysicsReplicationCVars::ResimulationCVars::bEnableInputDecay && !bIsSolverResim && bApplySimProxyInput && ComponentSettings.GetApplyInputDecaySimProxyInputAtRuntime())
 				{
-					PhysicsData->DecayData(ComponentSettings.GetInputDecaySimProxyInputAtRuntime());
+					const float RuntimeInputDecay = GetCurrentSimProxyInputDecayAtRuntime();
+					PhysicsData->DecayData(RuntimeInputDecay);
 				}
 
 				if (bIsUsingLegacyData)
@@ -2765,9 +2773,51 @@ const float FAsyncNetworkPhysicsComponent::GetCurrentInputDecay() const
 	const float PredictionAlpha = MaxPredictedInputs > 0 ? FMath::Clamp(NumPredictedInputs / MaxPredictedInputs, 0.0f, 1.0f) : 0.0f;
 
 	// Get decay from curve
-	const float InputDecay = InputDecayCurve.GetRichCurveConst()->Eval(PredictionAlpha);
+	float InputDecay = InputDecayCurve.GetRichCurveConst()->Eval(PredictionAlpha);
 
-	return InputDecay;
+	// If enabled, scale input decay as a function of measured lag
+	if (ComponentSettings.GetEnableLerpInputDecayWithLag())
+	{
+		const float MeasuredInputDelaySeconds = GetInputLagFromLastReceivedTargetState();
+		const float ReferenceInputDelaySeconds = ComponentSettings.GetInputDecayReferenceLagMs() * 0.001f;
+		// We do not clamp this lerp value to 1.0; if the measured delay is greater than the reference one, we apply a greater decay
+		const float LerpValue = (ReferenceInputDelaySeconds >= UE_SMALL_NUMBER) ? (MeasuredInputDelaySeconds / ReferenceInputDelaySeconds) : 1.0f;
+		InputDecay *= LerpValue;
+	}
+
+	return FMath::Clamp(InputDecay, 0.0f, 1.0f);
+}
+
+const float FAsyncNetworkPhysicsComponent::GetCurrentSimProxyInputDecayAtRuntime() const
+{
+	const FNetworkPhysicsSettingsNetworkPhysicsComponent& ComponentSettings = GetComponentSettings();
+	float SimProxyRuntimeInputDecay = ComponentSettings.GetInputDecaySimProxyInputAtRuntime();
+	if (ComponentSettings.GetEnableLagScalingSimProxyRuntimeInputDecay())
+	{
+		const float MeasuredInputDelaySeconds = GetInputLagFromLastReceivedTargetState();
+		const float ReferenceInputDelaySeconds = ComponentSettings.GetInputDecayReferenceLagMs() * 0.001f;
+		// We do not clamp this lerp value to 1.0; if the measured delay is greater than the reference one, we apply a greater decay
+		float LerpValue = (ReferenceInputDelaySeconds >= UE_SMALL_NUMBER) ? (MeasuredInputDelaySeconds / ReferenceInputDelaySeconds) : 1.0f;
+		SimProxyRuntimeInputDecay *= LerpValue;
+	}
+	return FMath::Clamp(SimProxyRuntimeInputDecay, 0.0f, 1.0f);
+}
+
+const float FAsyncNetworkPhysicsComponent::GetInputLagFromLastReceivedTargetState() const
+{
+	const Chaos::FPhysicsSolverBase* BaseSolver = GetSolver();
+	if (const Chaos::FRewindData* RewindData = BaseSolver ? BaseSolver->GetRewindData() : nullptr)
+	{
+		// Todo: we should not have to get a read lock to get this information bu instead measure the latest lag upon reception
+		Chaos::FReadPhysicsObjectInterface_Internal Interface = Chaos::FPhysicsObjectInternalInterface::GetRead();
+		if (Chaos::FGeometryParticleHandle* Particle = Interface.GetParticle(PhysicsObject))
+		{
+			const int32 MeasuredInputDelayFrames = RewindData->CurrentFrame() - RewindData->GetLatestReceivedTargetStateFrame(*Particle);
+			const float MeasuredInputDelaySeconds = FMath::Max(MeasuredInputDelayFrames, 0.0f) * BaseSolver->GetAsyncDeltaTime();
+			return MeasuredInputDelaySeconds;
+		}
+	}
+	return 0.0f;
 }
 
 void FAsyncNetworkPhysicsComponent::UpdateDynamicInputReplicationScaling()
diff --git a/Engine/Source/Runtime/Engine/Public/Physics/NetworkPhysicsComponent.h b/Engine/Source/Runtime/Engine/Public/Physics/NetworkPhysicsComponent.h
--- a/Engine/Source/Runtime/Engine/Public/Physics/NetworkPhysicsComponent.h
+++ b/Engine/Source/Runtime/Engine/Public/Physics/NetworkPhysicsComponent.h
@@ -1974,10 +1974,16 @@ private:
 	/** Returns the current amount of input decay during resimulation as a magnitude from 0.0 to 1.0. Returns 0 if not currently resimulating. */
 	UE_DEPRECATED(5.8, "Use function that doesn't take a parameter, the decay curve will be based on the current frame number compared with the latest input frame number available")
 	const float GetCurrentInputDecay(const FNetworkPhysicsPayload* PhysicsData) { return GetCurrentInputDecay(); }
-	
+
 	/** Returns the current amount of input decay during resimulation as a magnitude from 0.0 to 1.0. Returns 0 if not currently resimulating. */
 	const float GetCurrentInputDecay() const;
 
+	/** Returns the current amount of input decay for sim proxies during simulation as a magnitude from 0.0 to 1.0. */
+	const float GetCurrentSimProxyInputDecayAtRuntime() const;
+
+	/** Returns the input lag (in seconds) as measured based on the difference between the frame associated with the last received target and the frame at which it was received */
+	const float GetInputLagFromLastReceivedTargetState() const;
+
 	/** Populate input data in AsyncOutput to send over the network */
 	void SendInputData_Internal(FAsyncNetworkPhysicsComponentOutput& AsyncOutput, const int32 PhysicsStep);
 
diff --git a/Engine/Source/Runtime/Engine/Public/Physics/NetworkPhysicsSettingsComponent.h b/Engine/Source/Runtime/Engine/Public/Physics/NetworkPhysicsSettingsComponent.h
--- a/Engine/Source/Runtime/Engine/Public/Physics/NetworkPhysicsSettingsComponent.h
+++ b/Engine/Source/Runtime/Engine/Public/Physics/NetworkPhysicsSettingsComponent.h
@@ -79,9 +79,12 @@ namespace PhysicsReplicationCVars
 		extern bool bApplySimProxyInputAtRuntime;
 		extern bool bTriggerResimOnInputReceive;
 		extern bool bApplyInputDecaySimProxyInputAtRuntime;
+		extern bool bEnableLagScalingSimProxyRuntimeInputDecay;
 		extern float InputDecaySimProxyInputAtRuntime;
 		extern bool bApplyInputDecayOverSetTime;
 		extern float InputDecaySetTime;
+		extern bool bEnableLagScaledInputDecay;
+		extern float InputDecayReferenceLagMs;
 	}
 }
 
@@ -510,8 +513,11 @@ struct FNetworkPhysicsSettingsNetworkPhysicsComponent
 		, bOverrideApplySimProxyInputAtRuntime(0)
 		, bOverrideApplyInputDecaySimProxyInputAtRuntime(0)
 		, bOverrideInputDecaySimProxyInputAtRuntime(0)
+		, bOverrideEnableLagScalingSimProxyRuntimeInputDecay(0)
 		, bOverrideApplyInputDecayOverSetTime(0)
 		, bOverrideInputDecaySetTime(0)
+		, bOverrideEnableLagScaledInputDecay(0)
+		, bOverrideInputDecayReferenceLagMs(0)
 	{};
 
 	GENERATED_BODY()
@@ -562,9 +568,15 @@ struct FNetworkPhysicsSettingsNetworkPhysicsComponent
 	UPROPERTY(config, EditDefaultsOnly, Category = "Overrides", Meta = (InlineEditConditionToggle))
 	uint32 bOverrideInputDecaySimProxyInputAtRuntime : 1;
 	UPROPERTY(config, EditDefaultsOnly, Category = "Overrides", Meta = (InlineEditConditionToggle))
+	uint32 bOverrideEnableLagScalingSimProxyRuntimeInputDecay : 1;
+	UPROPERTY(config, EditDefaultsOnly, Category = "Overrides", Meta = (InlineEditConditionToggle))
 	uint32 bOverrideApplyInputDecayOverSetTime : 1;
 	UPROPERTY(config, EditDefaultsOnly, Category = "Overrides", Meta = (InlineEditConditionToggle))
 	uint32 bOverrideInputDecaySetTime : 1;
+	UPROPERTY(config, EditDefaultsOnly, Category = "Overrides", Meta = (InlineEditConditionToggle))
+	uint32 bOverrideEnableLagScaledInputDecay : 1;
+	UPROPERTY(config, EditDefaultsOnly, Category = "Overrides", Meta = (InlineEditConditionToggle))
+	uint32 bOverrideInputDecayReferenceLagMs : 1;
 
 	/* Overrides CVar : np2.Resim.RedundantInputs -- How many extra inputs to send with each unreliable network message, to account for packetloss.From owning client to server and server to owning client.
 	 * NOTE: This is disabled while np2.Resim.DynamicInputReplicationScaling.Enabled is enabled. */
@@ -656,6 +668,11 @@ struct FNetworkPhysicsSettingsNetworkPhysicsComponent
 	float InputDecaySimProxyInputAtRuntime = PhysicsReplicationCVars::ResimulationCVars::InputDecaySimProxyInputAtRuntime;
 	const float GetInputDecaySimProxyInputAtRuntime() const { return bOverrideInputDecaySimProxyInputAtRuntime ? InputDecaySimProxyInputAtRuntime : PhysicsReplicationCVars::ResimulationCVars::InputDecaySimProxyInputAtRuntime; }
 
+	// Overrides CVar: np2.Resim.EnableLagScalingSimProxyRuntimeInputDecay -- When true, scales the input decay applied on simulated proxies at runtime (outside of resimulation) as a function of measured lag
+	UPROPERTY(config, EditDefaultsOnly, Category = "Overrides", Meta = (EditCondition = "bOverrideEnableLagScalingSimProxyRuntimeInputDecay"))
+	bool bEnableLagScalingSimProxyRuntimeInputDecay = PhysicsReplicationCVars::ResimulationCVars::bEnableLagScalingSimProxyRuntimeInputDecay;
+	const bool GetEnableLagScalingSimProxyRuntimeInputDecay() const { return bOverrideEnableLagScalingSimProxyRuntimeInputDecay ? bEnableLagScalingSimProxyRuntimeInputDecay : PhysicsReplicationCVars::ResimulationCVars::bEnableLagScalingSimProxyRuntimeInputDecay; }
+
 	// Overrides CVar: np2.Resim.ApplyInputDecayOverSetTime -- When true, apply the Input Decay Curve over a set amount of time instead of over the start of input prediction and end of resim which is variable each resimulation
 	UPROPERTY(config, EditDefaultsOnly, Category = "Overrides", Meta = (EditCondition = "bOverrideApplyInputDecayOverSetTime"))
 	bool bApplyInputDecayOverSetTime = PhysicsReplicationCVars::ResimulationCVars::bApplyInputDecayOverSetTime;
@@ -672,6 +689,16 @@ struct FNetworkPhysicsSettingsNetworkPhysicsComponent
 	UPROPERTY(config, EditDefaultsOnly, Category = "Overrides", Meta = (XAxisName="LerpAlpha", YAxisName="DecayValue"))
 	FRuntimeFloatCurve InputDecayCurve;
 	const FRuntimeFloatCurve& GetInputDecayCurve() const { return InputDecayCurve; }
+
+	// Overrides CVar: np2.Resim.EnableLagScaledInputDecay -- When true, scales input decay as a function of measured input lag (the greater the lag, the greater the decay). The decay will be exactly the tuned InputDecay when the measured lag is equal to the reference input lag (InputDecayReferenceLagMs)
+	UPROPERTY(config, EditDefaultsOnly, Category = "Overrides", Meta = (EditCondition = "bOverrideEnableLagScaledInputDecay"))
+	bool bEnableLagScaledInputDecay = PhysicsReplicationCVars::ResimulationCVars::bEnableLagScaledInputDecay;
+	const bool GetEnableLerpInputDecayWithLag() const { return bOverrideEnableLagScaledInputDecay ? bEnableLagScaledInputDecay : PhysicsReplicationCVars::ResimulationCVars::bEnableLagScaledInputDecay; }
+
+	// Overrides CVar: np2.Resim.InputDecayReferenceLagMs -- Applied when np2.Resim.EnableLagScaledInputDecay is true. Reference lag at which the decay will be exactly InputDecay. Linearly interpolated below and beyond, reaching 0 decay at 0 lag.
+	UPROPERTY(config, EditDefaultsOnly, Category = "Overrides", Meta = (EditCondition = "bOverrideInputDecayReferenceLagMs"))
+	float InputDecayReferenceLagMs = PhysicsReplicationCVars::ResimulationCVars::InputDecayReferenceLagMs;
+	const float GetInputDecayReferenceLagMs() const { return bOverrideInputDecayReferenceLagMs ? InputDecayReferenceLagMs : PhysicsReplicationCVars::ResimulationCVars::InputDecayReferenceLagMs; }
 };
 
 /*
diff --git a/Engine/Source/Runtime/Experimental/Chaos/Private/RewindData.cpp b/Engine/Source/Runtime/Experimental/Chaos/Private/RewindData.cpp
--- a/Engine/Source/Runtime/Experimental/Chaos/Private/RewindData.cpp
+++ b/Engine/Source/Runtime/Experimental/Chaos/Private/RewindData.cpp
@@ -2560,6 +2560,19 @@ void FRewindData::SetTargetStateAtFrame(FGeometryParticleHandle& Handle, const i
 		Position, Quaternion, LinVelocity, AngVelocity, bShouldSleep);
 }
 
+int32 FRewindData::GetLatestReceivedTargetStateFrame(Chaos::FGeometryParticleHandle& Handle) const
+{
+	if (const FDirtyParticleInfo* DirtyParticleInfo = FindDirtyObj(Handle))
+	{
+		FFrameAndPhase FrameAndPhase;
+		if (DirtyParticleInfo->GetHistory().TargetPositions.GetHeadFrameAndPhase(FrameAndPhase))
+		{
+			return FrameAndPhase.Frame;
+		}
+	}
+	return LatestFrame;
+}
+
 void FRewindData::RequestResimulation(int32 Frame, Chaos::FGeometryParticleHandle* Particle)
 {
 	if (IsResim() || LatestFrame < ResimCooldownUntilFrame)
diff --git a/Engine/Source/Runtime/Experimental/Chaos/Public/RewindData.h b/Engine/Source/Runtime/Experimental/Chaos/Public/RewindData.h
--- a/Engine/Source/Runtime/Experimental/Chaos/Public/RewindData.h
+++ b/Engine/Source/Runtime/Experimental/Chaos/Public/RewindData.h
@@ -1760,6 +1760,8 @@ public:
 	void CHAOS_API SetTargetStateAtFrame(FGeometryParticleHandle& Handle, const int32 Frame, FFrameAndPhase::EParticleHistoryPhase Phase,
 		const FVector& Position, const FQuat& Quaternion, const FVector& LinVelocity, const FVector& AngVelocity, const bool bShouldSleep);
 
+	int32 CHAOS_API GetLatestReceivedTargetStateFrame(FGeometryParticleHandle& Handle) const;
+
 	/** Extract some history information before cleaning/pushing state*/
 	void ExtractHistoryState(FGeometryParticleHandle& Handle, int32& PositionValidCount, int32& VelocityValidCount, int32& PositionNextIterator, int32& VelocityNextIterator)
 	{
@@ -2120,6 +2122,12 @@ private:
 		return DirtyObjs.Find(&Handle);
 	}
 
+	template <typename TDirtyObjs, typename TObj>
+	const auto* FindDirtyObjImp(const TDirtyObjs& DirtyObjs, TObj& Handle) const
+	{
+		return DirtyObjs.Find(&Handle);
+	}
+
 	FDirtyParticleInfo* FindDirtyObj(const FGeometryParticleHandle& Handle)
 	{
 		return FindDirtyObjImp(DirtyParticles, Handle);
@@ -2130,6 +2138,16 @@ private:
 		return FindDirtyObjImp(DirtyJoints, Handle);
 	}
 
+	const FDirtyParticleInfo* FindDirtyObj(const FGeometryParticleHandle& Handle) const
+	{
+		return FindDirtyObjImp(DirtyParticles, Handle);
+	}
+
+	const FDirtyJointInfo* FindDirtyObj(const FPBDJointConstraintHandle& Handle) const
+	{
+		return FindDirtyObjImp(DirtyJoints, Handle);
+	}
+
 	template <typename TDirtyObjs, typename TObj>
 	auto& FindOrAddDirtyObjImp(TDirtyObjs & DirtyObjs, TObj & Handle, const int32 InitializedOnFrame = INDEX_NONE)
 	{
