Iris: convert replication bridge to support any type of object container, not just levels.
Levels are still supported the old way, but entity streaming will use the new path, using the streaming cells package as the container object.

#Synced-CL 48615930

diff --git a/Engine/Source/Runtime/Engine/Classes/Engine/NetConnection.h b/Engine/Source/Runtime/Engine/Classes/Engine/NetConnection.h
--- a/Engine/Source/Runtime/Engine/Classes/Engine/NetConnection.h
+++ b/Engine/Source/Runtime/Engine/Classes/Engine/NetConnection.h
@@ -882,6 +882,9 @@ public:
 	/** Called by PlayerController to tell connection about client level visibility change */
 	ENGINE_API void UpdateLevelVisibility(const struct FUpdateLevelVisibilityLevelInfo& LevelVisibility);
 
+	/** Called by PlayerController to tell connection about client container visibility change */
+	ENGINE_API void UpdateContainerVisibility(UObject* Container, bool bVisible);
+
 	const TSet<FName>& GetClientMakingVisibleLevelNames() const { return ClientMakingVisibleLevelNames; }
 	
 #if DO_ENABLE_NET_TEST
@@ -1720,6 +1723,9 @@ private:
 	/** Called by PlayerController to tell connection about client level visibility change */
 	void UpdateLevelVisibilityInternal(const struct FUpdateLevelVisibilityLevelInfo& LevelVisibility);
 
+	/** Called by PlayerController to tell connection about client container visibility change */
+	void UpdateContainerVisibilityInternal(UObject* Container, bool bVisible);
+
 	/** Called internally to update cached acceleration map */
 	bool UpdateCachedLevelVisibility(const FName& PackageName) const;
 
diff --git a/Engine/Source/Runtime/Engine/Classes/GameFramework/PlayerController.h b/Engine/Source/Runtime/Engine/Classes/GameFramework/PlayerController.h
--- a/Engine/Source/Runtime/Engine/Classes/GameFramework/PlayerController.h
+++ b/Engine/Source/Runtime/Engine/Classes/GameFramework/PlayerController.h
@@ -1539,6 +1539,9 @@ public:
 	UFUNCTION(reliable, server, WithValidation, SealedEvent)
 	ENGINE_API void ServerUpdateMultipleLevelsVisibility( const TArray<FUpdateLevelVisibilityLevelInfo>& LevelVisibilities );
 
+	UFUNCTION(Reliable, Server, SealedEvent)
+	ENGINE_API void ServerUpdateContainerVisibility(UObject* Container, bool bVisible);
+
 	/** Used by client to request server to confirm current viewtarget (server will respond with ClientSetViewTarget() ). */
 	UFUNCTION(reliable, server, WithValidation)
 	ENGINE_API void ServerVerifyViewTarget();
diff --git a/Engine/Source/Runtime/Engine/Classes/GameFramework/UpdateLevelVisibilityLevelInfo.h b/Engine/Source/Runtime/Engine/Classes/GameFramework/UpdateLevelVisibilityLevelInfo.h
--- a/Engine/Source/Runtime/Engine/Classes/GameFramework/UpdateLevelVisibilityLevelInfo.h
+++ b/Engine/Source/Runtime/Engine/Classes/GameFramework/UpdateLevelVisibilityLevelInfo.h
@@ -63,11 +63,11 @@ struct FUpdateLevelVisibilityLevelInfo
 	}
 
 	/**
-	 * @param Level				Level to pull PackageName and FileName from.
+	 * @param Container			Container to pull PackageName and FileName from.
 	 * @param bInIsVisible		Default value for bIsVisible.
 	 * @param bInTryMakeVisible	Whether the level is trying to be made visible or not.
 	 */
-	ENGINE_API FUpdateLevelVisibilityLevelInfo(const class ULevel* const Level, const bool bInIsVisible, const bool bInTryMakeVisible = false);
+	ENGINE_API FUpdateLevelVisibilityLevelInfo(const class UObject* const Container, const bool bInIsVisible, const bool bInTryMakeVisible = false);
 
 	/** The name of the package for the level whose status changed. */
 	UPROPERTY()
diff --git a/Engine/Source/Runtime/Engine/Private/GameFramework/UpdateLevelVisibilityLevelInfo.cpp b/Engine/Source/Runtime/Engine/Private/GameFramework/UpdateLevelVisibilityLevelInfo.cpp
--- a/Engine/Source/Runtime/Engine/Private/GameFramework/UpdateLevelVisibilityLevelInfo.cpp
+++ b/Engine/Source/Runtime/Engine/Private/GameFramework/UpdateLevelVisibilityLevelInfo.cpp
@@ -17,7 +17,7 @@ namespace PlayerControllerCVars
 	);
 }
 
-FUpdateLevelVisibilityLevelInfo::FUpdateLevelVisibilityLevelInfo(const ULevel* const Level, const bool bInIsVisible, const bool bInTryMakeVisible)
+FUpdateLevelVisibilityLevelInfo::FUpdateLevelVisibilityLevelInfo(const UObject* const Container, const bool bInIsVisible, const bool bInTryMakeVisible)
 	: CookPackageName(NAME_None)
 	, bIsVisible(bInIsVisible)
 	, bTryMakeVisible(bInTryMakeVisible)
@@ -26,15 +26,18 @@ FUpdateLevelVisibilityLevelInfo::FUpdateLevelVisibilityLevelInfo(const ULevel* c
 	// For backward compatibility, bTryMakeVisible was added instead of converting bIsVisible to an enum.
 	// Make sure we don't receive the invalid state (bIsVisible == true) && (bTryMakeVisible == true)
 	check(!bTryMakeVisible || (bIsVisible != bTryMakeVisible));
-	const UPackage* const LevelPackage = Level->GetOutermost();
-	PackageName = LevelPackage->GetFName();
+	const UPackage* const ContainerPackage = Container->GetOutermost();
+	PackageName = ContainerPackage->GetFName();
 #if WITH_EDITOR
-	CookPackageName = Level->GetCookPackageName();
+	if (const ULevel* Level = Cast<ULevel>(Container))
+	{
+		CookPackageName = Level->GetCookPackageName();
+	}
 #endif
 
 	// When packages are duplicated for PIE, they may not have a FileName.
 	// For now, just revert to the old behavior.
-	FName LoadedPathName = LevelPackage->GetLoadedPath().GetPackageFName();
+	FName LoadedPathName = ContainerPackage->GetLoadedPath().GetPackageFName();
 	FileName = (LoadedPathName.IsNone()) ? PackageName : LoadedPathName;
 }
 
diff --git a/Engine/Source/Runtime/Engine/Private/Net/Iris/ReplicationSystem/EngineReplicationBridge.cpp b/Engine/Source/Runtime/Engine/Private/Net/Iris/ReplicationSystem/EngineReplicationBridge.cpp
--- a/Engine/Source/Runtime/Engine/Private/Net/Iris/ReplicationSystem/EngineReplicationBridge.cpp
+++ b/Engine/Source/Runtime/Engine/Private/Net/Iris/ReplicationSystem/EngineReplicationBridge.cpp
@@ -485,7 +485,7 @@ UE::Net::FNetRefHandle UEngineReplicationBridge::StartReplicatingActor(AActor* A
 	}
 
 	// Setup Level filtering
-	AddRootObjectToLevelGroup(Actor, Actor->GetLevel());
+	AddRootObjectToContainerGroup(Actor, Actor->GetLevel());
 
 	// If we have registered sub objects we replicate them as well
 	const FSubObjectRegistry& ActorSubObjects = FSubObjectRegistryGetter::GetSubObjects(Actor);
@@ -623,7 +623,7 @@ void UEngineReplicationBridge::StopReplicatingActor(AActor* Actor, const FStopRe
 		const UE::Net::FDestructionParameters DestructionParameters
 		{
 			.Location = Actor->GetActorLocation(),
-			.Level = Actor->GetLevel(),
+			.Container = Actor->GetLevel(),
 			.bUseDistanceBasedPrioritization = (Actor->bAlwaysRelevant == false),
 			.NetFactoryId = ActorFactoryId,
 		};
@@ -775,8 +775,9 @@ void UEngineReplicationBridge::SetNetDriver(UNetDriver* const InNetDriver)
 
 			checkf(bRequiredChannelIsConfigured, TEXT("ObjectReplication needs the netdriver channel %s to work. Add this channel to the netdriver channel definitions config"), *RequiredChannelName.ToString());
 
-		bAllowExecutePreReplication = InNetDriver->IsPreReplicationEnabled();
-		bAllowReplicationOfActorsWithAnyRole = InNetDriver->IsReplicationOfActorsWithAnyRoleEnabled();}
+			bAllowExecutePreReplication = InNetDriver->IsPreReplicationEnabled();
+			bAllowReplicationOfActorsWithAnyRole = InNetDriver->IsReplicationOfActorsWithAnyRoleEnabled();
+		}
 	}
 }
 
@@ -801,7 +802,7 @@ bool UEngineReplicationBridge::RemapPathForPIE(uint32 ConnectionId, FString& Str
 	}
 }
 
-bool UEngineReplicationBridge::ObjectLevelHasFinishedLoading(UObject* Object) const
+bool UEngineReplicationBridge::ObjectContainerHasFinishedLoading(UObject* Object) const
 {
 	const UWorld* DriverWorld = Object && NetDriver ? NetDriver->GetWorld() : nullptr;
 	if (DriverWorld)
@@ -967,19 +968,19 @@ void UEngineReplicationBridge::ActorChangedLevel(const AActor* Actor, const ULev
 		const FName PreviousLevelPackageName = PreviousLevelPackage->GetFName();
 
 		const UE::Net::FNetRefHandle ActorRefHandle = GetReplicatedRefHandle(Actor);
-		UE::Net::FNetObjectGroupHandle PreviousLevelGroup = GetLevelGroup(PreviousLevel);
-		if (GetReplicationSystem()->IsInGroup(PreviousLevelGroup, ActorRefHandle))
+		UE::Net::FNetObjectGroupHandle PreviousContainerGroup = GetContainerGroup(PreviousLevel);
+		if (GetReplicationSystem()->IsInGroup(PreviousContainerGroup, ActorRefHandle))
 		{
-			UE_LOG_BRIDGEID(LogIris, Verbose, TEXT("ActorChangedLevel: removing %s from GroupIndex: %u PreviousLevel: %s"), *ActorRefHandle.ToString(), PreviousLevelGroup.GetGroupIndex(), ToCStr(PreviousLevelPackageName.ToString()));
-			GetReplicationSystem()->RemoveFromGroup(PreviousLevelGroup, ActorRefHandle);
+			UE_LOG_BRIDGEID(LogIris, Verbose, TEXT("ActorChangedLevel: removing %s from GroupIndex: %u PreviousLevel: %s"), *ActorRefHandle.ToString(), PreviousContainerGroup.GetGroupIndex(), ToCStr(PreviousLevelPackageName.ToString()));
+			GetReplicationSystem()->RemoveFromGroup(PreviousContainerGroup, ActorRefHandle);
 		}
-		else if (PreviousLevelGroup.IsValid())
+		else if (PreviousContainerGroup.IsValid())
 		{
-			UE_LOG_BRIDGEID(LogIris, Warning, TEXT("ActorChangedLevel: %s not found in GroupIndex: %u PreviousLevel: %s"), *ActorRefHandle.ToString(), PreviousLevelGroup.GetGroupIndex(), ToCStr(PreviousLevelPackageName.ToString()));
+			UE_LOG_BRIDGEID(LogIris, Warning, TEXT("ActorChangedLevel: %s not found in GroupIndex: %u PreviousLevel: %s"), *ActorRefHandle.ToString(), PreviousContainerGroup.GetGroupIndex(), ToCStr(PreviousLevelPackageName.ToString()));
 		}
 	}
 
-	AddRootObjectToLevelGroup(Actor, Actor->GetLevel());
+	AddRootObjectToContainerGroup(Actor, Actor->GetLevel());
 }
 
 void UEngineReplicationBridge::OnNetUpdateFrequencyChanged(const AActor* Actor)
@@ -995,7 +996,14 @@ void UEngineReplicationBridge::OnNetUpdateFrequencyChanged(const AActor* Actor)
 	}
 }
 
+PRAGMA_DISABLE_DEPRECATION_WARNINGS
 void UEngineReplicationBridge::AddRootObjectToLevelGroup(const UObject* RootObject, const ULevel* Level)
+{
+	AddRootObjectToContainerGroup(RootObject, Level);
+}
+PRAGMA_ENABLE_DEPRECATION_WARNINGS
+
+void UEngineReplicationBridge::AddRootObjectToContainerGroup(const UObject* RootObject, const UObject* Container)
 {
 	if (!UE::Net::Private::ShouldIncludeRootObjectInLevelGroups(RootObject))
 	{
@@ -1003,38 +1011,63 @@ void UEngineReplicationBridge::AddRootObjectToLevelGroup(const UObject* RootObje
 	}
 	
 	// For seamless travel we need to maintain a level group for all kinds of levels, including persistent ones.
-	if (Level && (UE::Net::Private::bAlwaysCreateLevelFilteringGroupsForPersistentLevels || !Level->IsPersistentLevel() || (Level != NetDriver->GetWorld()->PersistentLevel)))
+	if (const ULevel* Level = Cast<ULevel>(Container))
 	{
-		const UPackage* const LevelPackage = Level->GetOutermost();
-		const FString PackageName = LevelPackage->GetName();
-
-		UE::Net::FNetObjectGroupHandle LevelGroup = GetLevelGroup(Level);
-		if (!LevelGroup.IsValid())
+		if (UE::Net::Private::bAlwaysCreateLevelFilteringGroupsForPersistentLevels || !Level->IsPersistentLevel() || (Level != NetDriver->GetWorld()->PersistentLevel))
 		{
-			// When a level is unloaded and reloaded immediately and when LevelStreaming.ShouldReuseUnloadedButStillAroundLevels is false, 
-			// we can end up with 2 different sets of level pointers that are have the same original package name.
-			// Make the package name unique so a new group gets assigned to it.
-			FString LevelGroupPackageUniqueName = PackageName;
-			int32 Index = 0;
-			while (GetReplicationSystem()->FindGroup(*LevelGroupPackageUniqueName).IsValid())
+			const UPackage* const LevelPackage = Level->GetOutermost();
+			const FString PackageName = LevelPackage->GetName();
+
+			UE::Net::FNetObjectGroupHandle LevelGroup = GetContainerGroup(Level);
+			if (!LevelGroup.IsValid())
 			{
-				++Index;
-				LevelGroupPackageUniqueName = FString::Printf(TEXT("%s_%d"), *PackageName, Index);
+				// When a level is unloaded and reloaded immediately and when LevelStreaming.ShouldReuseUnloadedButStillAroundLevels is false, 
+				// we can end up with 2 different sets of level pointers that are have the same original package name.
+				// Make the package name unique so a new group gets assigned to it.
+				FString LevelGroupPackageUniqueName = PackageName;
+				int32 Index = 0;
+				while (GetReplicationSystem()->FindGroup(*LevelGroupPackageUniqueName).IsValid())
+				{
+					++Index;
+					LevelGroupPackageUniqueName = FString::Printf(TEXT("%s_%d"), *PackageName, Index);
+				}
+
+				LevelGroup = CreateContainerGroup(Level, *LevelGroupPackageUniqueName);
+
+				UE_LOG_BRIDGEID(LogIris, Verbose, TEXT("Created new GroupIndex: %u for Level: %s using Package Unique Name: %s"), LevelGroup.GetGroupIndex(), ToCStr(PackageName), ToCStr(LevelGroupPackageUniqueName));
+
+				// Update the filtering status of the group based on current level visibility for all connections
+				NetDriver->UpdateGroupFilterStatusForLevel(Level, LevelGroup);
 			}
 
-			LevelGroup = CreateLevelGroup(Level, *LevelGroupPackageUniqueName);
+			const UE::Net::FNetRefHandle RootObjectRefHandle = GetReplicatedRefHandle(RootObject);
 
-			UE_LOG_BRIDGEID(LogIris, Verbose, TEXT("Created new GroupIndex: %u for Level: %s using Package Unique Name: %s"), LevelGroup.GetGroupIndex(), ToCStr(PackageName), ToCStr(LevelGroupPackageUniqueName));
+			// Add object to group
+			UE_LOG_BRIDGEID(LogIris, Verbose, TEXT("Added %s to GroupIndex: %u Level: %s"), *RootObjectRefHandle.ToString(), LevelGroup.GetGroupIndex(), ToCStr(PackageName));
+			GetReplicationSystem()->AddToGroup(LevelGroup, RootObjectRefHandle);
+		}
+	}
+	else if (const UPackage* Package = Cast<UPackage>(Container))
+	{
+		UE::Net::FNetObjectGroupHandle ContainerGroup = GetContainerGroup(Package);
 
-			// Update the filtering status of the group based on current level visibility for all connections
-			NetDriver->UpdateGroupFilterStatusForLevel(Level, LevelGroup);
+		if (!ContainerGroup.IsValid())
+		{
+			const FName ContainerName = Package->GetFName();
+			ContainerGroup = CreateContainerGroup(Package, ContainerName);
+
+			UE_LOG_BRIDGEID(LogIris, Verbose, TEXT("Created new GroupIndex: %u for Level: %s using Package Unique Name: %s"), ContainerGroup.GetGroupIndex(), ToCStr(ContainerName.ToString()), ToCStr(ContainerName.ToString()));
 		}
 
 		const UE::Net::FNetRefHandle RootObjectRefHandle = GetReplicatedRefHandle(RootObject);
 
 		// Add object to group
-		UE_LOG_BRIDGEID(LogIris, Verbose, TEXT("Added %s to GroupIndex: %u Level: %s"), *RootObjectRefHandle.ToString(), LevelGroup.GetGroupIndex(), ToCStr(PackageName));
-		GetReplicationSystem()->AddToGroup(LevelGroup, RootObjectRefHandle);
+		UE_LOG_BRIDGEID(LogIris, Verbose, TEXT("Added %s to GroupIndex: %u Container: %s"), *RootObjectRefHandle.ToString(), ContainerGroup.GetGroupIndex(), ToCStr(Package->GetPackage()->GetName()));
+		GetReplicationSystem()->AddToGroup(ContainerGroup, RootObjectRefHandle);
+	}
+	else
+	{
+		ensureMsgf(false, TEXT("Iris requires replicated container group to be a ULevel or a UPackage, Container: %s"), ToCStr(GetFullNameSafe(Container)));
 	}
 }
 
diff --git a/Engine/Source/Runtime/Engine/Private/Net/Iris/ReplicationSystem/NetRootObjectAdapter.cpp b/Engine/Source/Runtime/Engine/Private/Net/Iris/ReplicationSystem/NetRootObjectAdapter.cpp
--- a/Engine/Source/Runtime/Engine/Private/Net/Iris/ReplicationSystem/NetRootObjectAdapter.cpp
+++ b/Engine/Source/Runtime/Engine/Private/Net/Iris/ReplicationSystem/NetRootObjectAdapter.cpp
@@ -207,7 +207,7 @@ void FNetRootObjectAdapter::StartReplication(ULevel* Level)
 		bIsReplicating = true;
 
 		// Setup level filtering
-		Bridge->AddRootObjectToLevelGroup(ReplicatedObject, Level);
+		Bridge->AddRootObjectToContainerGroup(ReplicatedObject, Level);
 
 		//TODO: Support a subobject list in the adapter ?
 	});
diff --git a/Engine/Source/Runtime/Engine/Private/NetConnection.cpp b/Engine/Source/Runtime/Engine/Private/NetConnection.cpp
--- a/Engine/Source/Runtime/Engine/Private/NetConnection.cpp
+++ b/Engine/Source/Runtime/Engine/Private/NetConnection.cpp
@@ -1733,7 +1733,7 @@ FNetLevelVisibilityTransactionId UNetConnection::UpdateLevelStreamStatusChangedT
 			if (const ULevel* Level = LevelObject->GetLoadedLevel())
 			{
 				const UObjectReplicationBridge* ReplicationBridge = ReplicationSystem->GetReplicationBridge();
-				if (UE::Net::FNetObjectGroupHandle GroupHandle = ReplicationBridge->GetLevelGroup(Level); GroupHandle.IsValid())
+				if (UE::Net::FNetObjectGroupHandle GroupHandle = ReplicationBridge->GetContainerGroup(Level); GroupHandle.IsValid())
 				{
 					ReplicationSystem->SetGroupFilterStatus(GroupHandle, GetConnectionHandle().GetParentConnectionId(), UE::Net::ENetFilterStatus::Disallow);
 				}
@@ -1887,7 +1887,7 @@ void UNetConnection::UpdateLevelVisibilityInternal(const FUpdateLevelVisibilityL
 				if (Level)
 				{
 					const UObjectReplicationBridge* ReplicationBridge = ReplicationSystem->GetReplicationBridge();
-					if (UE::Net::FNetObjectGroupHandle GroupHandle = ReplicationBridge->GetLevelGroup(Level); GroupHandle.IsValid())
+					if (UE::Net::FNetObjectGroupHandle GroupHandle = ReplicationBridge->GetContainerGroup(Level); GroupHandle.IsValid())
 					{
 						ReplicationSystem->SetGroupFilterStatus(GroupHandle, GetConnectionHandle().GetParentConnectionId(), UE::Net::ENetFilterStatus::Allow);
 					}
@@ -1983,7 +1983,7 @@ void UNetConnection::UpdateLevelVisibilityInternal(const FUpdateLevelVisibilityL
 					// The reason for this is that the client currently destroys the instances rather then managing this through the replication system
 					// If we implement a way to re-instantiate instances on the client we might be able to persist the state
 					const UObjectReplicationBridge* ReplicationBridge = ReplicationSystem->GetReplicationBridge();
-					if (UE::Net::FNetObjectGroupHandle GroupHandle = ReplicationBridge->GetLevelGroup(Level); GroupHandle.IsValid())
+					if (UE::Net::FNetObjectGroupHandle GroupHandle = ReplicationBridge->GetContainerGroup(Level); GroupHandle.IsValid())
 					{
 						ReplicationSystem->SetGroupFilterStatus(GroupHandle, GetConnectionHandle().GetParentConnectionId(), UE::Net::ENetFilterStatus::Disallow);
 					}
@@ -2031,6 +2031,24 @@ void UNetConnection::UpdateLevelVisibilityInternal(const FUpdateLevelVisibilityL
 	}
 }
 
+void UNetConnection::UpdateContainerVisibility(UObject* Container, bool bVisible)
+{
+	UpdateContainerVisibilityInternal(Container, bVisible);
+}
+
+void UNetConnection::UpdateContainerVisibilityInternal(UObject* Container, bool bVisible)
+{
+	if (UReplicationSystem* ReplicationSystem = Driver->GetReplicationSystem())
+	{
+		const UObjectReplicationBridge* ReplicationBridge = ReplicationSystem->GetReplicationBridge();
+		UE::Net::FNetObjectGroupHandle ContainerGroupHandle = ReplicationBridge->GetContainerGroup(Container);
+		if (ContainerGroupHandle.IsValid())
+		{
+			ReplicationSystem->SetGroupFilterStatus(ContainerGroupHandle, GetConnectionHandle().GetParentConnectionId(), bVisible ? UE::Net::ENetFilterStatus::Allow : UE::Net::ENetFilterStatus::Disallow);
+		}
+	}
+}
+
 void UNetConnection::SetClientWorldPackageName(FName NewClientWorldPackageName)
 {
 	ClientWorldPackageName = NewClientWorldPackageName;
diff --git a/Engine/Source/Runtime/Engine/Private/NetDriver.cpp b/Engine/Source/Runtime/Engine/Private/NetDriver.cpp
--- a/Engine/Source/Runtime/Engine/Private/NetDriver.cpp
+++ b/Engine/Source/Runtime/Engine/Private/NetDriver.cpp
@@ -4351,7 +4351,7 @@ void UNetDriver::NotifyActorDestroyed(AActor* ThisActor, bool IsSeamlessTravel,
 						const UE::Net::FDestructionParameters Params
 						{ 
 							.Location = ThisActor->GetActorLocation(), 
-							.Level = ThisActor->GetLevel(), 
+							.Container = ThisActor->GetLevel(), 
 							.bUseDistanceBasedPrioritization = true,
 							.NetFactoryId = Bridge->GetActorFactoryId()
 						};
@@ -6860,7 +6860,7 @@ void UNetDriver::CreateReplicatedStaticActorDestructionInfo(ULevel* Level, const
 			const UE::Net::FDestructionParameters Params
 			{
 				.Location = Info.DestroyedPosition,
-				.Level = Level,				
+				.Container = Level,				
 				.bUseDistanceBasedPrioritization = true,
 				.NetFactoryId = Bridge->GetActorFactoryId()
 			};
@@ -7351,10 +7351,10 @@ void UNetDriver::DisallowReplicationOfAllLevelGroupsForConnection(UE::Net::FConn
 	}
 
 	const uint32 ConnectionId = ConnectionHandle.GetParentConnectionId();
-	const TMap<FObjectKey, UE::Net::FNetObjectGroupHandle>& LevelGroups = Bridge->GetAllLevelGroups();
-	for (const TMap<FObjectKey, UE::Net::FNetObjectGroupHandle>::ElementType& LevelAndGroup : LevelGroups)
+	const TMap<FObjectKey, UE::Net::FNetObjectGroupHandle>& ContainerGroups = Bridge->GetAllContainerGroups();
+	for (const TMap<FObjectKey, UE::Net::FNetObjectGroupHandle>::ElementType& ContaineAndGroups : ContainerGroups)
 	{
-		ReplicationSystem->SetGroupFilterStatus(LevelAndGroup.Value, ConnectionId, UE::Net::ENetFilterStatus::Disallow);
+		ReplicationSystem->SetGroupFilterStatus(ContaineAndGroups.Value, ConnectionId, UE::Net::ENetFilterStatus::Disallow);
 	}
 }
 
@@ -7385,13 +7385,13 @@ void UNetDriver::UpdateGroupFilterStatusForAllLevelsForConnection(UE::Net::FConn
 	}
 
 	const uint32 ConnectionId = ConnectionHandle.GetParentConnectionId();
-	const TMap<FObjectKey, UE::Net::FNetObjectGroupHandle>& LevelGroups = Bridge->GetAllLevelGroups();
-	for (const TMap<FObjectKey, UE::Net::FNetObjectGroupHandle>::ElementType& LevelAndGroup : LevelGroups)
+	const TMap<FObjectKey, UE::Net::FNetObjectGroupHandle>& ContainerGroups = Bridge->GetAllContainerGroups();
+	for (const TMap<FObjectKey, UE::Net::FNetObjectGroupHandle>::ElementType& ContainerAndGroup : ContainerGroups)
 	{
-		if (const ULevel* Level = Cast<ULevel>(LevelAndGroup.Key.ResolveObjectPtr()))
+		if (const ULevel* Level = Cast<ULevel>(ContainerAndGroup.Key.ResolveObjectPtr()))
 		{
 			const bool bIsVisible = (Level->IsPersistentLevel() || Connection->ClientVisibleLevelNames.Contains(Level->GetOutermost()->GetFName()));
-			ReplicationSystem->SetGroupFilterStatus(LevelAndGroup.Value, ConnectionId, bIsVisible ? UE::Net::ENetFilterStatus::Allow : UE::Net::ENetFilterStatus::Disallow);
+			ReplicationSystem->SetGroupFilterStatus(ContainerAndGroup.Value, ConnectionId, bIsVisible ? UE::Net::ENetFilterStatus::Allow : UE::Net::ENetFilterStatus::Disallow);
 		}
 	}
 }
@@ -7411,12 +7411,12 @@ void UNetDriver::UpdateGroupFilterStatusForAllLevelsForAllConnections()
 		return;
 	}
 
-	const TMap<FObjectKey, UE::Net::FNetObjectGroupHandle>& LevelGroups = Bridge->GetAllLevelGroups();
-	for (const TMap<FObjectKey, UE::Net::FNetObjectGroupHandle>::ElementType& LevelAndGroup : LevelGroups)
+	const TMap<FObjectKey, UE::Net::FNetObjectGroupHandle>& ContainerGroups = Bridge->GetAllContainerGroups();
+	for (const TMap<FObjectKey, UE::Net::FNetObjectGroupHandle>::ElementType& ContainerAndGroup : ContainerGroups)
 	{
-		if (const ULevel* Level = Cast<ULevel>(LevelAndGroup.Key.ResolveObjectPtr()))
+		if (const ULevel* Level = Cast<ULevel>(ContainerAndGroup.Key.ResolveObjectPtr()))
 		{
-			UpdateGroupFilterStatusForLevel(Level, LevelAndGroup.Value);
+			UpdateGroupFilterStatusForLevel(Level, ContainerAndGroup.Value);
 		}
 	}
 }
diff --git a/Engine/Source/Runtime/Engine/Private/PlayerController.cpp b/Engine/Source/Runtime/Engine/Private/PlayerController.cpp
--- a/Engine/Source/Runtime/Engine/Private/PlayerController.cpp
+++ b/Engine/Source/Runtime/Engine/Private/PlayerController.cpp
@@ -491,6 +491,14 @@ bool APlayerController::ServerUpdateLevelVisibility_Validate(const FUpdateLevelV
 	return true;
 }
 
+void APlayerController::ServerUpdateContainerVisibility_Implementation(UObject* Container, bool bVisible)
+{
+	if (UNetConnection* Connection = Cast<UNetConnection>(Player))
+	{
+		Connection->UpdateContainerVisibility(Container, bVisible);
+	}
+}
+
 void APlayerController::ServerUpdateMultipleLevelsVisibility_Implementation( const TArray<FUpdateLevelVisibilityLevelInfo>& LevelVisibilities )
 {
 	for(const FUpdateLevelVisibilityLevelInfo& LevelVisibility : LevelVisibilities)
diff --git a/Engine/Source/Runtime/Engine/Public/Net/Iris/ReplicationSystem/EngineReplicationBridge.h b/Engine/Source/Runtime/Engine/Public/Net/Iris/ReplicationSystem/EngineReplicationBridge.h
--- a/Engine/Source/Runtime/Engine/Public/Net/Iris/ReplicationSystem/EngineReplicationBridge.h
+++ b/Engine/Source/Runtime/Engine/Public/Net/Iris/ReplicationSystem/EngineReplicationBridge.h
@@ -129,7 +129,11 @@ public:
 	ENGINE_API virtual void SendErrorWithNetRefHandle(UE::Net::ENetRefHandleError ErrorType, FNetRefHandle RefHandle, uint32 ConnectionId, TConstArrayView<const FNetRefHandle> ExtraNetRefHandle = {}) override;
 
 	/** Add the rootobject to the level's filter group so it will only be relevant if the connection has that level streamed in. */
+	UE_DEPRECATED(5.8, "Call AddRootObjectToContainerGroup instead.")
 	ENGINE_API void AddRootObjectToLevelGroup(const UObject* RootObject, const ULevel* Level);
+
+	/** Add the root object to the container's filter group so it will only be relevant if the connection has that container streamed in. */
+	ENGINE_API void AddRootObjectToContainerGroup(const UObject* RootObject, const UObject* Container);
 	
 	/** Updates the level group for an actor that changed levels */
 	void ActorChangedLevel(const AActor* Actor, const ULevel* PreviousLevel);
@@ -164,7 +168,7 @@ protected:
 	virtual void Deinitialize() override;
 	virtual void GetInitialDependencies(FNetRefHandle Handle, FNetDependencyInfoArray& OutDependencies) const override;
 	virtual bool RemapPathForPIE(uint32 ConnectionId, FString& Path, bool bReading) const override;
-	virtual bool ObjectLevelHasFinishedLoading(UObject* Object) const override;
+	virtual bool ObjectContainerHasFinishedLoading(UObject* Object) const override;
 	virtual void OnProtocolMismatchDetected(FNetRefHandle ObjectHandle) override;
 	virtual void OnProtocolMismatchReported(FNetRefHandle RefHandle, uint32 ConnectionId) override;
 	virtual bool CanCreateDestructionInfo() const override;
diff --git a/Engine/Source/Runtime/Net/Iris/Private/Iris/ReplicationSystem/ObjectReferenceCache.cpp b/Engine/Source/Runtime/Net/Iris/Private/Iris/ReplicationSystem/ObjectReferenceCache.cpp
--- a/Engine/Source/Runtime/Net/Iris/Private/Iris/ReplicationSystem/ObjectReferenceCache.cpp
+++ b/Engine/Source/Runtime/Net/Iris/Private/Iris/ReplicationSystem/ObjectReferenceCache.cpp
@@ -1077,9 +1077,9 @@ UObject* FObjectReferenceCache::ResolveObjectReferenceHandleInternal(FNetRefHand
 	//}
 
 	UObjectReplicationBridge* ReplicationBridge = ResolveContext.ReplicationSystem->GetReplicationBridgeAs<UObjectReplicationBridge>();
-	if (Object && !ReplicationBridge->ObjectLevelHasFinishedLoading(Object))
+	if (Object && !ReplicationBridge->ObjectContainerHasFinishedLoading(Object))
 	{
-		UE_LOG_REFERENCECACHE(Verbose, TEXT("GetObjectFromRefHandle: Forcing object to NULL since level is not loaded yet. Object: %s" ), *Object->GetFullName());
+		UE_LOG_REFERENCECACHE(Verbose, TEXT("GetObjectFromRefHandle: Forcing object to NULL since container is not loaded yet. Object: %s" ), *Object->GetFullName());
 
 		return nullptr;
 	}
diff --git a/Engine/Source/Runtime/Net/Iris/Private/Iris/ReplicationSystem/ObjectReplicationBridge.cpp b/Engine/Source/Runtime/Net/Iris/Private/Iris/ReplicationSystem/ObjectReplicationBridge.cpp
--- a/Engine/Source/Runtime/Net/Iris/Private/Iris/ReplicationSystem/ObjectReplicationBridge.cpp
+++ b/Engine/Source/Runtime/Net/Iris/Private/Iris/ReplicationSystem/ObjectReplicationBridge.cpp
@@ -378,12 +378,12 @@ UE::Net::FNetRefHandle UObjectReplicationBridge::StoreDestructionInfo(FNetRefHan
 	};
 
 	// Create handle for DestructionInfo to be able to scope destruction infos
-	const FNetObjectGroupHandle LevelGroupHandle = GetLevelGroup(Parameters.Level);
+	const FNetObjectGroupHandle ContainerGroupHandle = GetContainerGroup(Parameters.Container);
 	FNetRefHandle DestructionInfoHandle = NetRefHandleManager->CreateHandleForDestructionInfo(Handle, CreateNetObjectParams);
 
-	if (LevelGroupHandle.IsValid())
+	if (ContainerGroupHandle.IsValid())
 	{
-		GetReplicationSystem()->AddToGroup(LevelGroupHandle, DestructionInfoHandle);
+		GetReplicationSystem()->AddToGroup(ContainerGroupHandle, DestructionInfoHandle);
 	}
 
 	const FInternalNetRefIndex InternalReplicationIndex = NetRefHandleManager->GetInternalIndex(DestructionInfoHandle);
@@ -392,7 +392,7 @@ UE::Net::FNetRefHandle UObjectReplicationBridge::StoreDestructionInfo(FNetRefHan
 	FStaticDestructionInfo PersistentDestructionInfo
 	{
 		.StaticRef = FObjectReferenceCache::MakeNetObjectReference(Handle),
-		.LevelGroupHandle = LevelGroupHandle,
+		.ContainerGroupHandle = ContainerGroupHandle,
 		.NetFactoryId = Parameters.NetFactoryId,
 		.InternalReplicationIndex = InternalReplicationIndex
 	};
@@ -416,7 +416,7 @@ UE::Net::FNetRefHandle UObjectReplicationBridge::StoreDestructionInfo(FNetRefHan
 		GetReplicationSystem()->SetPrioritizer(DestructionInfoHandle, DefaultSpatialNetObjectPrioritizerHandle);
 	}
 
-	UE_LOG_OBJECTREPLICATIONBRIDGE(Verbose, TEXT("StoreDestructionInfo %s (InternalIndex: %u ) for %s GroupIndex: %u"), *DestructionInfoHandle.ToString(), InternalReplicationIndex,  *PrintObjectFromNetRefHandle(Handle), LevelGroupHandle.GetGroupIndex());
+	UE_LOG_OBJECTREPLICATIONBRIDGE(Verbose, TEXT("StoreDestructionInfo %s (InternalIndex: %u ) for %s GroupIndex: %u"), *DestructionInfoHandle.ToString(), InternalReplicationIndex,  *PrintObjectFromNetRefHandle(Handle), ContainerGroupHandle.GetGroupIndex());
 	
 	return DestructionInfoHandle;
 }
diff --git a/Engine/Source/Runtime/Net/Iris/Private/Iris/ReplicationSystem/ReplicationBridge.cpp b/Engine/Source/Runtime/Net/Iris/Private/Iris/ReplicationSystem/ReplicationBridge.cpp
--- a/Engine/Source/Runtime/Net/Iris/Private/Iris/ReplicationSystem/ReplicationBridge.cpp
+++ b/Engine/Source/Runtime/Net/Iris/Private/Iris/ReplicationSystem/ReplicationBridge.cpp
@@ -795,29 +795,29 @@ void UReplicationBridge::ClearNetPushIdOnFragments(const TArrayView<const UE::Ne
 #endif
 }
 
-void UReplicationBridge::NotifyStreamingLevelUnload(const UObject* Level)
+void UReplicationBridge::NotifyStreamingContainerUnload(const UObject* Container)
 {
-	// Destroy group associated with level
-	UE::Net::FNetObjectGroupHandle LevelGroupHandle;
-	if (LevelGroups.RemoveAndCopyValue(FObjectKey(Level), LevelGroupHandle))
+	// Destroy group associated with container
+	UE::Net::FNetObjectGroupHandle ContainerGroupHandle;
+	if (ContainerGroups.RemoveAndCopyValue(FObjectKey(Container), ContainerGroupHandle))
 	{
-		RemoveDestructionInfosForGroup(LevelGroupHandle);
-		ReplicationSystem->DestroyGroup(LevelGroupHandle);
+		RemoveDestructionInfosForGroup(ContainerGroupHandle);
+		ReplicationSystem->DestroyGroup(ContainerGroupHandle);
 	}
 }
 
-UE::Net::FNetObjectGroupHandle UReplicationBridge::CreateLevelGroup(const UObject* Level, FName PackageName)
+UE::Net::FNetObjectGroupHandle UReplicationBridge::CreateContainerGroup(const UObject* Container, FName PackageName)
 {
 	using namespace UE::Net;
 
-	FNetObjectGroupHandle LevelGroupHandle = ReplicationSystem->CreateGroup(PackageName);
-	if (ensure(LevelGroupHandle.IsValid()))
+	FNetObjectGroupHandle ContainerGroupHandle = ReplicationSystem->CreateGroup(PackageName);
+	if (ensure(ContainerGroupHandle.IsValid()))
 	{
-		ReplicationSystem->AddExclusionFilterGroup(LevelGroupHandle);
-		LevelGroups.Emplace(FObjectKey(Level), LevelGroupHandle);
+		ReplicationSystem->AddExclusionFilterGroup(ContainerGroupHandle);
+		ContainerGroups.Emplace(FObjectKey(Container), ContainerGroupHandle);
 	}
 
-	return LevelGroupHandle;
+	return ContainerGroupHandle;
 }
 
 UE::Net::FNetObjectFactoryId UReplicationBridge::GetNetObjectFactoryId(FNetRefHandle RefHandle) const
@@ -825,15 +825,15 @@ UE::Net::FNetObjectFactoryId UReplicationBridge::GetNetObjectFactoryId(FNetRefHa
 	return NetRefHandleManager->GetReplicatedObjectData(NetRefHandleManager->GetInternalIndex(RefHandle)).NetFactoryId;
 }
 
-UE::Net::FNetObjectGroupHandle UReplicationBridge::GetLevelGroup(const UObject* Level) const
+UE::Net::FNetObjectGroupHandle UReplicationBridge::GetContainerGroup(const UObject* Container) const
 {
-	const UE::Net::FNetObjectGroupHandle* LevelGroupHandle = LevelGroups.Find(FObjectKey(Level));
-	return (LevelGroupHandle != nullptr ? *LevelGroupHandle : UE::Net::FNetObjectGroupHandle());
+	const UE::Net::FNetObjectGroupHandle* ContainerGroupHandle = ContainerGroups.Find(FObjectKey(Container));
+	return (ContainerGroupHandle != nullptr ? *ContainerGroupHandle : UE::Net::FNetObjectGroupHandle());
 }
 
-const TMap<FObjectKey, UE::Net::FNetObjectGroupHandle>& UReplicationBridge::GetAllLevelGroups() const
+const TMap<FObjectKey, UE::Net::FNetObjectGroupHandle>& UReplicationBridge::GetAllContainerGroups() const
 {
-	return LevelGroups;
+	return ContainerGroups;
 }
 
 void UReplicationBridge::DestroyGlobalNetHandle(UE::Net::Private::FInternalNetRefIndex InternalReplicationIndex)
@@ -875,11 +875,11 @@ void UReplicationBridge::PreSeamlessTravelGarbageCollect()
 {
 	RemoveDestructionInfosForGroup(UE::Net::FNetObjectGroupHandle());
 
-	for (const TPair<FObjectKey, UE::Net::FNetObjectGroupHandle>& LevelAndGroup : LevelGroups)
+	for (const TPair<FObjectKey, UE::Net::FNetObjectGroupHandle>& ContainerAndGroup : ContainerGroups)
 	{
-		ReplicationSystem->DestroyGroup(LevelAndGroup.Value);
+		ReplicationSystem->DestroyGroup(ContainerAndGroup.Value);
 	}
-	LevelGroups.Empty();
+	ContainerGroups.Empty();
 
 	OnPreSeamlessTravelGarbageCollect();
 }
diff --git a/Engine/Source/Runtime/Net/Iris/Private/Iris/ReplicationSystem/ReplicationSystem.cpp b/Engine/Source/Runtime/Net/Iris/Private/Iris/ReplicationSystem/ReplicationSystem.cpp
--- a/Engine/Source/Runtime/Net/Iris/Private/Iris/ReplicationSystem/ReplicationSystem.cpp
+++ b/Engine/Source/Runtime/Net/Iris/Private/Iris/ReplicationSystem/ReplicationSystem.cpp
@@ -1158,7 +1158,7 @@ void UReplicationSystem::CollectGarbage()
 
 void UReplicationSystem::NotifyStreamingLevelUnload(const UObject* Level)
 {
-	Impl->ReplicationSystemInternal.GetReplicationBridge()->NotifyStreamingLevelUnload(Level);
+	Impl->ReplicationSystemInternal.GetReplicationBridge()->NotifyStreamingContainerUnload(Level);
 }
 
 void UReplicationSystem::PreSeamlessTravelGarbageCollect()
diff --git a/Engine/Source/Runtime/Net/Iris/Public/Iris/ReplicationSystem/ObjectReplicationBridge.h b/Engine/Source/Runtime/Net/Iris/Public/Iris/ReplicationSystem/ObjectReplicationBridge.h
--- a/Engine/Source/Runtime/Net/Iris/Public/Iris/ReplicationSystem/ObjectReplicationBridge.h
+++ b/Engine/Source/Runtime/Net/Iris/Public/Iris/ReplicationSystem/ObjectReplicationBridge.h
@@ -56,8 +56,8 @@ struct FStaticDestructionInfo
 	/** The reference to the static object that is considered destroyed */
 	UE::Net::FNetObjectReference StaticRef;
 	
-	/** The level the object was a part of */
-	UE::Net::FNetObjectGroupHandle LevelGroupHandle;
+	/** The container the object was a part of */
+	UE::Net::FNetObjectGroupHandle ContainerGroupHandle;
 	
 	/** The NetFactory assigned to this object */
 	UE::Net::FNetObjectFactoryId NetFactoryId = UE::Net::InvalidNetObjectFactoryId;
@@ -174,7 +174,11 @@ public:
 	virtual bool RemapPathForPIE(uint32 ConnectionId, FString& Path, bool bReading) const { return false; }
 
 	/** Returns true of the level that the Object belongs to has finished loading. */
-	virtual bool ObjectLevelHasFinishedLoading(UObject* Object) const { return true; }
+	UE_DEPRECATED(5.8, "Call ObjectContainerHasFinishedLoading instead.")
+	virtual bool ObjectLevelHasFinishedLoading(UObject* Object) const { return ObjectContainerHasFinishedLoading(Object); }
+
+	/** Returns true of the container that the Object belongs to has finished loading. */
+	virtual bool ObjectContainerHasFinishedLoading(UObject* Object) const { return true; }
 
 	/**
 	 * Adds a dependent object. A dependent object can replicate separately or if a parent replicates.
@@ -638,7 +642,7 @@ private:
 	// When we flush objects, we might need to defer sending creation info.
 	TMap<FNetRefHandle, TUniquePtr<const UE::Net::FNetObjectCreationHeader>> CachedCreationHeaders;
 
-	// Track the static objects that are destroyed so clients can be told to delete them after loading the level
+	// Track the static objects that are destroyed so clients can be told to delete them after loading the container
 	TMap<FNetRefHandle, UE::Net::Private::FStaticDestructionInfo> StaticObjectsPendingDestroy;
 
 	// Objects which has object references and could be affected by garbage collection.
diff --git a/Engine/Source/Runtime/Net/Iris/Public/Iris/ReplicationSystem/ReplicationBridge.h b/Engine/Source/Runtime/Net/Iris/Public/Iris/ReplicationSystem/ReplicationBridge.h
--- a/Engine/Source/Runtime/Net/Iris/Public/Iris/ReplicationSystem/ReplicationBridge.h
+++ b/Engine/Source/Runtime/Net/Iris/Public/Iris/ReplicationSystem/ReplicationBridge.h
@@ -87,8 +87,8 @@ namespace UE::Net
 	{
 		/** The location of the object. Used for distance based prioritization. */
 		FVector Location;
-		/** The level the object is placed in. */
-		const UObject* Level = nullptr;
+		/** The container the object is placed in. */
+		const UObject* Container = nullptr;
 		/** Whether to use distance based priority for the destruction of the object. */
 		bool bUseDistanceBasedPrioritization = false;
 		/** The NetFactory that the replicated object would be assigned to. */
@@ -138,11 +138,11 @@ public:
 	/** Returns true if the handle is replicated. */
 	IRISCORE_API bool IsReplicatedHandle(FNetRefHandle Handle) const;
 
-	/** Get the group associated with the level in order to control connection filtering for it. */
-	IRISCORE_API UE::Net::FNetObjectGroupHandle GetLevelGroup(const UObject* Level) const;
+	/** Get the group associated with the container in order to control connection filtering for it. */
+	IRISCORE_API UE::Net::FNetObjectGroupHandle GetContainerGroup(const UObject* Container) const;
 
-	/** Get all level groups. The object key is to a ULevel. The TMap must be processed immediately or copied- DO NOT CACHE. */
-	IRISCORE_API const TMap<FObjectKey, UE::Net::FNetObjectGroupHandle>& GetAllLevelGroups() const;
+	/** Get all container groups. The object key is to a UObject. The TMap must be processed immediately or copied- DO NOT CACHE. */
+	IRISCORE_API const TMap<FObjectKey, UE::Net::FNetObjectGroupHandle>& GetAllContainerGroups() const;
 
 	/** Returns true when we are in the middle of processing incoming data. */
 	bool IsInReceiveUpdate() const { return bInReceiveUpdate; }
@@ -232,11 +232,11 @@ protected:
 	/** Return the NetFactoryId assigned to a replicated object. */
 	IRISCORE_API UE::Net::FNetObjectFactoryId GetNetObjectFactoryId(FNetRefHandle RefHandle) const;
 
-	/** Creates a group for a level for object filtering purposes. */
-	IRISCORE_API UE::Net::FNetObjectGroupHandle CreateLevelGroup(const UObject* Level, FName PackageName);
+	/** Creates a group for a container for object filtering purposes. */
+	IRISCORE_API UE::Net::FNetObjectGroupHandle CreateContainerGroup(const UObject* Container, FName PackageName);
 
-	/** Destroys the group associated with the level. */
-	IRISCORE_API void DestroyLevelGroup(const UObject* Level);
+	/** Destroys the group associated with the container. */
+	IRISCORE_API void DestroyContainerGroup(const UObject* Container);
 
 	/** Called when a remote connection detected a protocol mismatch when trying to instantiate the NetRefHandle replicated object. */
 	virtual void OnProtocolMismatchReported(FNetRefHandle RefHandle, uint32 ConnectionId) {}
@@ -293,10 +293,10 @@ private:
 	void InternalDestroySubObjects(FNetRefHandle OwnerHandle, EEndReplicationFlags Flags);
 
 	/**
-	 * Called from ReplicationSystem when a streaming level is about to unload.
-	 * Will remove the group associated with the level and remove destruction infos.
+	 * Called from ReplicationSystem when a streaming container is about to unload.
+	 * Will remove the group associated with the container and remove destruction infos.
 	 */
-	void NotifyStreamingLevelUnload(const UObject* Level);
+	void NotifyStreamingContainerUnload(const UObject* Container);
 
 	void DestroyLocalNetHandle(FNetRefHandle Handle, EEndReplicationFlags Flags);
 
@@ -325,7 +325,7 @@ private:
 	UE::Net::Private::FObjectReferenceCache* ObjectReferenceCache;
 	UE::Net::Private::FNetObjectGroups* Groups;
 
-	TMap<FObjectKey, UE::Net::FNetObjectGroupHandle> LevelGroups;
+	TMap<FObjectKey, UE::Net::FNetObjectGroupHandle> ContainerGroups;
 
 	/** Tracks if we are in the middle of processing incoming data */
 	bool bInReceiveUpdate = false;
