Add FScopedNetEmulationOverride for easy, temporary network emulation overrides in test cases.

#Synced-CL Unknown

diff --git a/Engine/Source/Runtime/Engine/Private/Tests/NetTestHelpers.cpp b/Engine/Source/Runtime/Engine/Private/Tests/NetTestHelpers.cpp
--- a/Engine/Source/Runtime/Engine/Private/Tests/NetTestHelpers.cpp
+++ b/Engine/Source/Runtime/Engine/Private/Tests/NetTestHelpers.cpp
@@ -5,6 +5,7 @@
 #include "CoreGlobals.h"
 #include "Engine/NetworkObjectList.h"
 #include "Engine/NetConnection.h"
+#include "Engine/NetDriver.h"
 #include "Engine/ChildConnection.h"
 #include "Engine/Engine.h"
 #include "Engine/GameInstance.h"
@@ -29,11 +30,37 @@
 #include "UObject/UObjectGlobals.h"
 #include "UObject/Package.h"
 
-#if WITH_EDITOR
-
 namespace UE::Net
 {
 
+FScopedNetEmulationOverride::FScopedNetEmulationOverride(UNetDriver* InNetDriver, const FPacketSimulationSettings& InSettings)
+	: WeakNetDriver(InNetDriver)
+#if DO_ENABLE_NET_TEST
+	, SavedSettings(InNetDriver->PacketSimulationSettings)
+#endif
+{
+#if DO_ENABLE_NET_TEST
+	InNetDriver->PacketSimulationSettings = InSettings;
+	InNetDriver->OnPacketSimulationSettingsChanged();
+#else
+	ensureMsgf(false, TEXT("FScopedNetEmulationOverride will have no effect since DO_ENABLE_NET_TEST is not defined."));
+#endif
+}
+
+FScopedNetEmulationOverride::~FScopedNetEmulationOverride()
+{
+#if DO_ENABLE_NET_TEST
+	UNetDriver* NetDriver = WeakNetDriver.Get();
+	if (ensureMsgf(NetDriver, TEXT("~FScopedNetEmulationOverride: WeakNetDriver is invalid.")))
+	{
+		NetDriver->PacketSimulationSettings = SavedSettings;
+		NetDriver->OnPacketSimulationSettingsChanged();
+	}
+#endif
+}
+
+#if WITH_EDITOR
+
 bool IsClientConnected(FTestWorldInstance& Instance)
 {
 	UNetDriver* NetDriver = Instance.GetNetDriver();
@@ -730,6 +757,6 @@ FScopedNetTestPIERestoration::~FScopedNetTestPIERestoration()
 	GIsPlayInEditorWorld = OldGIsPlayInEditorWorld;
 }
 
-} // end namespace UE::Net
+#endif // WITH_EDITOR
 
-#endif // WITH_EDITOR
\ No newline at end of file
+} // end namespace UE::Net
diff --git a/Engine/Source/Runtime/Engine/Public/Tests/NetTestHelpers.h b/Engine/Source/Runtime/Engine/Public/Tests/NetTestHelpers.h
--- a/Engine/Source/Runtime/Engine/Public/Tests/NetTestHelpers.h
+++ b/Engine/Source/Runtime/Engine/Public/Tests/NetTestHelpers.h
@@ -16,11 +16,27 @@ class UObjectReplicationBridge;
 
 struct FGameInstancePIEParameters;
 
-#if WITH_EDITOR
-
 namespace UE::Net
 {
 
+/**
+ * Temporarily overrides network emulation settings for a specific NetDriver while in scope.
+ * Only works when DO_ENABLE_NET_TEST is true.
+ */
+class [[nodiscard]] FScopedNetEmulationOverride
+{
+public:
+	ENGINE_API FScopedNetEmulationOverride(UNetDriver* InNetDriver, const FPacketSimulationSettings& InSettings);
+	ENGINE_API ~FScopedNetEmulationOverride();
+
+private:
+	TWeakObjectPtr<UNetDriver> WeakNetDriver;
+	FPacketSimulationSettings SavedSettings;
+};
+
+
+#if WITH_EDITOR
+
 template<class T>
 int32 NumActors(UWorld* World)
 {
@@ -354,6 +370,6 @@ inline bool FTestWorlds::TickUntil(const PredicateT& Predicate, int32 MaxTicks)
 ENGINE_API bool IsClientConnected(FTestWorldInstance& Instance);
 ENGINE_API bool IsServerConnected(FTestWorldInstance& Instance);
 
-} // end namespace UE::Net
+#endif //WITH_EDITOR
 
-#endif //WITH_EDITOR
\ No newline at end of file
+} // end namespace UE::Net
\ No newline at end of file
