[UWorld] Resolve race condition with UniformBuffors being deallocated but the resources not being removed from the pool of available cached MeshDrawCommands.

#Synced-CL 48543872

diff --git a/Engine/Source/Runtime/Engine/Private/World.cpp b/Engine/Source/Runtime/Engine/Private/World.cpp
--- a/Engine/Source/Runtime/Engine/Private/World.cpp
+++ b/Engine/Source/Runtime/Engine/Private/World.cpp
@@ -1967,13 +1967,19 @@ UMaterialParameterCollectionInstance* UWorld::CreateParameterCollectionInstance(
 
 void UWorld::OnPostGC()
 {
+	bool bRemovedSomething = false;
 	for (int32 InstanceIndex = ParameterCollectionInstances.Num()-1; InstanceIndex >= 0; InstanceIndex--)
 	{
 		if (!ParameterCollectionInstances[InstanceIndex]->IsCollectionValid())
 		{
 			ParameterCollectionInstances.RemoveAt(InstanceIndex);
+			bRemovedSomething = true;
 		}
 	}
+	if (bRemovedSomething)
+	{
+		UpdateParameterCollectionInstances(false, false);
+	}
 }
 
 
