StaticMeshComponent - Fix crash on nullptr StaticMesh

diff --git a/Engine/Source/Runtime/Engine/Private/Components/StaticMeshComponent.cpp b/Engine/Source/Runtime/Engine/Private/Components/StaticMeshComponent.cpp
--- a/Engine/Source/Runtime/Engine/Private/Components/StaticMeshComponent.cpp
+++ b/Engine/Source/Runtime/Engine/Private/Components/StaticMeshComponent.cpp
@@ -3188,7 +3188,7 @@ void UStaticMeshComponent::SetMeshPaintTextureOverride(UTexture* OverrideTexture
 
 int32 UStaticMeshComponent::GetMeshPaintTextureCoordinateIndex() const
 {
-	const int32 NumTexCoords = StaticMesh->GetNumTexCoords(0);
+	const int32 NumTexCoords = StaticMesh ? StaticMesh->GetNumTexCoords(0) : 0;
 	if (bOverrideMeshPaintTextureCoordinateIndex)
 	{
 		return FMath::Clamp(OverriddenMeshPaintTextureCoordinateIndex, 0, NumTexCoords - 1);
