Protect USoundModulatorBase pointers passed to the Audio Thread command queue via AudioComponent Routing update calls with weak pointer checks and early out if any go crufty

[RN] Protected against thread unsafe behavior when passing Sound Modulator updates from AudioComponents to actively playing sounds

diff --git a/Engine/Source/Runtime/Engine/Private/Components/AudioComponent.cpp b/Engine/Source/Runtime/Engine/Private/Components/AudioComponent.cpp
--- a/Engine/Source/Runtime/Engine/Private/Components/AudioComponent.cpp
+++ b/Engine/Source/Runtime/Engine/Private/Components/AudioComponent.cpp
@@ -1,6 +1,9 @@
 // Copyright Epic Games, Inc. All Rights Reserved.
 
 #include "Components/AudioComponent.h"
+
+#include "Algo/AllOf.h"
+#include "Algo/Transform.h"
 #include "Audio/ActorSoundParameterInterface.h"
 #include "Audio/AudioTraceUtil.h"
 #include "AudioDevice.h"
@@ -74,6 +77,39 @@ static const TCHAR* GAudioSpriteAssetNameAutoActivate = TEXT("/Engine/EditorReso
 static const TCHAR* GAudioSpriteAssetName = TEXT("/Engine/EditorResources/AudioIcons/S_AudioComponent.S_AudioComponent");
 #endif
 
+namespace AudioComponentPrivate
+{
+	struct FModulatorWeakAccessor
+	{
+	private:
+		TSet<TWeakObjectPtr<USoundModulatorBase>> WeakModulators;
+
+	public:
+		FModulatorWeakAccessor(const TSet<TObjectPtr<USoundModulatorBase>>& Modulators)
+		{
+			Algo::Transform(Modulators, WeakModulators, [](const TObjectPtr<USoundModulatorBase>& Obj)
+			{
+				return TWeakObjectPtr<USoundModulatorBase>(Obj.Get());
+			});
+		}
+
+		bool TryGetSafe(TSet<TObjectPtr<USoundModulatorBase>>& OutModulators) const
+		{
+			OutModulators.Reset();
+
+			auto IsValid = [](const TWeakObjectPtr<USoundModulatorBase>& WeakPtr) { return WeakPtr.IsValid(); };
+			if (!Algo::AllOf(WeakModulators, IsValid))
+			{
+				return false;
+			}
+
+			auto AsObjectPtr = [](const TWeakObjectPtr<USoundModulatorBase>& WeakPtr) { return WeakPtr.Get(); };
+			Algo::Transform(WeakModulators, OutModulators, AsObjectPtr);
+			return true;
+		}
+	};
+} // namespace AudioComponentPrivate
+
 UInitialActiveSoundParams::UInitialActiveSoundParams(const FObjectInitializer& ObjectInitializer)
 	: Super(ObjectInitializer)
 {
@@ -1891,6 +1927,45 @@ bool UAudioComponent::GetCookedEnvelopeDataForAllPlayingSounds(TArray<FSoundWave
 
 void UAudioComponent::SetModulationRouting(const TSet<USoundModulatorBase*>& Modulators, const EModulationDestination Destination, const EModulationRouting RoutingMethod)
 {
+	struct FModulatorWeakRoutingAccessor
+	{
+	private:
+		TArray<TWeakObjectPtr<UObject>> WeakCollection;
+		const FSoundModulationDefaultRoutingSettings Settings;
+
+	public:
+		FModulatorWeakRoutingAccessor(const FSoundModulationDefaultRoutingSettings& RoutingSettings)
+			: Settings(RoutingSettings)
+		{
+			auto GatherModulators = [&](const TSet<TObjectPtr<USoundModulatorBase>>& InModulators)
+			{
+				for (const TObjectPtr<USoundModulatorBase>& Modulator : InModulators)
+				{
+					if (USoundModulatorBase* ModPtr = Modulator.Get())
+					{
+						WeakCollection.Add(TWeakObjectPtr<UObject>(ModPtr));
+					}
+				}
+			};
+
+			GatherModulators(RoutingSettings.VolumeModulationDestination.Modulators);
+			GatherModulators(RoutingSettings.PitchModulationDestination.Modulators);
+			GatherModulators(RoutingSettings.HighpassModulationDestination.Modulators);
+			GatherModulators(RoutingSettings.LowpassModulationDestination.Modulators);
+		}
+
+		const FSoundModulationDefaultRoutingSettings* GetSafe() const
+		{
+			auto IsValid = [](const TWeakObjectPtr<UObject>& WeakPtr) { return WeakPtr.IsValid(); };
+			if (Algo::AllOf(WeakCollection, IsValid))
+			{
+				return &Settings;
+			}
+
+			return nullptr;
+		}
+	};
+
 	FAudioDevice* AudioDevice = GetAudioDevice();
 	if (!AudioDevice)
 	{
@@ -1924,15 +1999,20 @@ void UAudioComponent::SetModulationRouting(const TSet<USoundModulatorBase*>& Mod
 	}
 
 	// Tell the active sounds on the component to use the new Modulation Routing
-	AudioDevice->SendCommandToActiveSounds(AudioComponentID, [NewRouting = ModulationRouting](FActiveSound& ActiveSound)
+	AudioDevice->SendCommandToActiveSounds(AudioComponentID, [NewRouting = FModulatorWeakRoutingAccessor(ModulationRouting)](FActiveSound& ActiveSound)
 	{
-		ActiveSound.SetNewModulationRouting(NewRouting);
+		if (const FSoundModulationDefaultRoutingSettings* SafeSettings = NewRouting.GetSafe())
+		{
+			ActiveSound.SetNewModulationRouting(*SafeSettings);
+		}
 	});
 
 }
 
 void UAudioComponent::AddModulationRouting(const TSet<USoundModulatorBase*>& Modulators, const EModulationDestination Destination)
 {
+	using namespace AudioComponentPrivate;
+
 	FAudioDevice* AudioDevice = GetAudioDevice();
 	if (!AudioDevice)
 	{
@@ -1972,14 +2052,20 @@ void UAudioComponent::AddModulationRouting(const TSet<USoundModulatorBase*>& Mod
 	}
 
 	// Tell the active sounds on the component to use the new Modulation Routing
-	AudioDevice->SendCommandToActiveSounds(AudioComponentID, [NewModulators = ConvertedModulators, Destination](FActiveSound& ActiveSound)
+	AudioDevice->SendCommandToActiveSounds(AudioComponentID, [Accessor = FModulatorWeakAccessor(ConvertedModulators), Destination](FActiveSound& ActiveSound)
 	{
-		ActiveSound.AddModulationRouting(NewModulators, Destination);
+		TSet<TObjectPtr<USoundModulatorBase>> NewModulators;
+		if (Accessor.TryGetSafe(NewModulators))
+		{
+			ActiveSound.AddModulationRouting(NewModulators, Destination);
+		}
 	});
 }
 
 void UAudioComponent::RemoveModulationRouting(const TSet<USoundModulatorBase*>& Modulators, const EModulationDestination Destination)
 {
+	using namespace AudioComponentPrivate;
+
 	FAudioDevice* AudioDevice = GetAudioDevice();
 	if (!AudioDevice)
 	{
@@ -2014,10 +2100,14 @@ void UAudioComponent::RemoveModulationRouting(const TSet<USoundModulatorBase*>&
 		}
 	}
 
-	// Tell the active sounds on the component to use the new Modulation Routing
-	AudioDevice->SendCommandToActiveSounds(AudioComponentID, [NewModulators = ConvertedModulators, Destination](FActiveSound& ActiveSound)
+	// Tell the active sounds on the component to remove the given modulators
+	AudioDevice->SendCommandToActiveSounds(AudioComponentID, [Accessor = FModulatorWeakAccessor(ConvertedModulators), Destination](FActiveSound& ActiveSound)
 	{
-		ActiveSound.RemoveModulationRouting(NewModulators, Destination);
+		TSet<TObjectPtr<USoundModulatorBase>> ModulatorsToRemove;
+		if (Accessor.TryGetSafe(ModulatorsToRemove))
+		{
+			ActiveSound.RemoveModulationRouting(ModulatorsToRemove, Destination);
+		}
 	});
 }
 
