New Material Translator:
- Added support for three additional Custom Output based expressions.

#rnx
#Synced-CL 48664487

diff --git a/Engine/Source/Runtime/Engine/Private/Materials/MaterialExpressionsIR.cpp b/Engine/Source/Runtime/Engine/Private/Materials/MaterialExpressionsIR.cpp
--- a/Engine/Source/Runtime/Engine/Private/Materials/MaterialExpressionsIR.cpp
+++ b/Engine/Source/Runtime/Engine/Private/Materials/MaterialExpressionsIR.cpp
@@ -4281,7 +4281,7 @@ void UMaterialExpressionSingleLayerWaterMaterialOutput::Build(MIR::FEmitter& Em)
 
 void UMaterialExpressionSubsurfaceMediumMaterialOutput::Build(MIR::FEmitter& Em)
 {
-	FValueRef  OutputValues[] = 
+	FValueRef OutputValues[] = 
 	{
 		Em.CastToFloat(Em.Input(&MeanFreePath), 3),
 		Em.CastToFloat(Em.Input(&ScatteringDistribution), 1),
@@ -4289,4 +4289,38 @@ void UMaterialExpressionSubsurfaceMediumMaterialOutput::Build(MIR::FEmitter& Em)
 	Em.SetCustomOutputs(TEXTVIEW("SubsurfaceMediumMaterialOutput"), OutputValues, MIR::MaterialOutputStage_Pixel);
 }
 
+void UMaterialExpressionThinTranslucentMaterialOutput::Build(MIR::FEmitter& Em)
+{
+	FValueRef OutputValues[] =
+	{
+		Em.CastToFloat(Em.InputDefaultFloat3(&TransmittanceColor, {0.5f, 0.5f, 0.5f}), 3),
+		Em.CastToFloat(Em.InputDefaultFloat(&SurfaceCoverage, 1.0f), 1),
+	};
+	Em.SetCustomOutputs(TEXTVIEW("ThinTranslucentMaterialOutput"), OutputValues, MIR::MaterialOutputStage_Pixel);
+}
+
+void UMaterialExpressionVolumetricAdvancedMaterialOutput::Build(MIR::FEmitter& Em)
+{
+	FValueRef OutputValues[] =
+	{
+		Em.CastToFloat(Em.InputDefaultFloat(&PhaseG, ConstPhaseG), 1),
+		Em.CastToFloat(Em.InputDefaultFloat(&PhaseG2, ConstPhaseG2), 1),
+		Em.CastToFloat(Em.InputDefaultFloat(&PhaseBlend, ConstPhaseBlend), 1),
+		Em.CastToFloat(Em.InputDefaultFloat(&MultiScatteringContribution, ConstMultiScatteringContribution), 1),
+		Em.CastToFloat(Em.InputDefaultFloat(&MultiScatteringOcclusion, ConstMultiScatteringOcclusion), 1),
+		Em.CastToFloat(Em.InputDefaultFloat(&MultiScatteringEccentricity, ConstMultiScatteringEccentricity), 1),
+		Em.CastToFloat(Em.InputDefaultFloat4(&ConservativeDensity, FVector4f{ 1.0f, 1.0f, 1.0f, 1.0f }), 4),
+	};
+	Em.SetCustomOutputs(TEXTVIEW("VolumetricAdvancedMaterialOutput"), OutputValues, MIR::MaterialOutputStage_Pixel);
+}
+
+void UMaterialExpressionAbsorptionMediumMaterialOutput::Build(MIR::FEmitter& Em)
+{
+	FValueRef OutputValues[] =
+	{
+		Em.CastToFloat(Em.InputDefaultFloat3(&TransmittanceColor, { 1.0f, 1.0f, 1.0f }), 3),
+	};
+	Em.SetCustomOutputs(TEXTVIEW("AbsorptionMediumMaterialOutput"), OutputValues, MIR::MaterialOutputStage_Pixel);
+}
+
 #endif // WITH_EDITOR
diff --git a/Engine/Source/Runtime/Engine/Public/Materials/MaterialExpressionAbsorptionMediumMaterialOutput.h b/Engine/Source/Runtime/Engine/Public/Materials/MaterialExpressionAbsorptionMediumMaterialOutput.h
--- a/Engine/Source/Runtime/Engine/Public/Materials/MaterialExpressionAbsorptionMediumMaterialOutput.h
+++ b/Engine/Source/Runtime/Engine/Public/Materials/MaterialExpressionAbsorptionMediumMaterialOutput.h
@@ -20,6 +20,7 @@ class UMaterialExpressionAbsorptionMediumMaterialOutput : public UMaterialExpres
 public:
 #if WITH_EDITOR
 	//~ Begin UMaterialExpression Interface
+	virtual void Build(MIR::FEmitter& Emitter) override;
 	virtual int32 Compile(class FMaterialCompiler* Compiler, int32 OutputIndex) override;
 	virtual void GetCaption(TArray<FString>& OutCaptions) const override;
 	//~ End UMaterialExpression Interface
diff --git a/Engine/Source/Runtime/Engine/Public/Materials/MaterialExpressionThinTranslucentMaterialOutput.h b/Engine/Source/Runtime/Engine/Public/Materials/MaterialExpressionThinTranslucentMaterialOutput.h
--- a/Engine/Source/Runtime/Engine/Public/Materials/MaterialExpressionThinTranslucentMaterialOutput.h
+++ b/Engine/Source/Runtime/Engine/Public/Materials/MaterialExpressionThinTranslucentMaterialOutput.h
@@ -24,6 +24,7 @@ class UMaterialExpressionThinTranslucentMaterialOutput : public UMaterialExpress
 public:
 #if WITH_EDITOR
 	//~ Begin UMaterialExpression Interface
+	virtual void Build(MIR::FEmitter& Emitter) override;
 	virtual int32 Compile(class FMaterialCompiler* Compiler, int32 OutputIndex) override;
 	virtual void GetCaption(TArray<FString>& OutCaptions) const override;
 	virtual EMaterialValueType GetInputValueType(int32 InputIndex) override;
diff --git a/Engine/Source/Runtime/Engine/Public/Materials/MaterialExpressionVolumetricAdvancedMaterialOutput.h b/Engine/Source/Runtime/Engine/Public/Materials/MaterialExpressionVolumetricAdvancedMaterialOutput.h
--- a/Engine/Source/Runtime/Engine/Public/Materials/MaterialExpressionVolumetricAdvancedMaterialOutput.h
+++ b/Engine/Source/Runtime/Engine/Public/Materials/MaterialExpressionVolumetricAdvancedMaterialOutput.h
@@ -94,6 +94,7 @@ class UMaterialExpressionVolumetricAdvancedMaterialOutput : public UMaterialExpr
 public:
 #if WITH_EDITOR
 	//~ Begin UMaterialExpression Interface
+	virtual void Build(MIR::FEmitter& Emitter) override;
 	virtual int32 Compile(class FMaterialCompiler* Compiler, int32 OutputIndex) override;
 	virtual void GetCaption(TArray<FString>& OutCaptions) const override;
 	//~ End UMaterialExpression Interface
