[FastGeo] Fix MT issue in editor when UMaterialInterface::GetTextureDensity is called from multiple thread (which internally calls SortTextureStreamingData and modifies the material instance).

#Synced-CL 48520266

diff --git a/Engine/Source/Runtime/Engine/Private/Materials/MaterialInterface.cpp b/Engine/Source/Runtime/Engine/Private/Materials/MaterialInterface.cpp
--- a/Engine/Source/Runtime/Engine/Private/Materials/MaterialInterface.cpp
+++ b/Engine/Source/Runtime/Engine/Private/Materials/MaterialInterface.cpp
@@ -49,6 +49,7 @@
 #include "UnrealEngine.h"
 #include "UObject/ObjectEditorOptionalSupport.h"
 #include "UObject/ICookInfo.h"
+#include "Async/UniqueLock.h"
 
 #if WITH_EDITOR
 #include "ObjectCacheEventSink.h"
@@ -1975,6 +1976,8 @@ bool FMaterialTextureInfo::IsValid(bool bCheckTextureIndex) const
 void UMaterialInterface::SortTextureStreamingData(bool bForceSort, bool bFinalSort)
 {
 #if WITH_EDITOR
+	UE::TUniqueLock Lock(TextureStreamingDataMutex);
+
 	// In cook that was already done in the save.
 	if (!bTextureStreamingDataSorted || bForceSort)
 	{
@@ -2038,6 +2041,10 @@ extern 	TAutoConsoleVariable<int32> CVarStreamingUseMaterialData;
 
 bool UMaterialInterface::FindTextureStreamingDataIndexRange(FName TextureName, int32& LowerIndex, int32& HigherIndex) const
 {
+#if WITH_EDITOR
+	UE::TUniqueLock Lock(TextureStreamingDataMutex);
+#endif
+
 #if WITH_EDITORONLY_DATA
 	// Because of redirectors (when textures are renammed), the texture names might be invalid and we need to udpate the data at every load.
 	// Normally we would do that in the post load, but since the process needs to resolve the SoftObjectPaths, this is forbidden at that place.
@@ -2071,6 +2078,10 @@ bool UMaterialInterface::FindTextureStreamingDataIndexRange(FName TextureName, i
 
 void UMaterialInterface::SetTextureStreamingData(const TArray<FMaterialTextureInfo>& InTextureStreamingData)
 {
+#if WITH_EDITOR
+	UE::TUniqueLock Lock(TextureStreamingDataMutex);
+#endif
+
 	TextureStreamingData = InTextureStreamingData;
 #if WITH_EDITORONLY_DATA
 	bTextureStreamingDataSorted = false;
@@ -2101,6 +2112,10 @@ float UMaterialInterface::GetTextureDensityWithCache(const FTextureSamplingInfo&
 
 float UMaterialInterface::GetTextureDensity(FName TextureName, const FMeshUVChannelInfo& UVChannelData) const
 {
+#if WITH_EDITOR
+	UE::TUniqueLock Lock(TextureStreamingDataMutex);
+#endif
+
 	ensure(UVChannelData.bInitialized);
 
 	int32 LowerIndex = INDEX_NONE;
diff --git a/Engine/Source/Runtime/Engine/Public/Materials/MaterialInterface.h b/Engine/Source/Runtime/Engine/Public/Materials/MaterialInterface.h
--- a/Engine/Source/Runtime/Engine/Public/Materials/MaterialInterface.h
+++ b/Engine/Source/Runtime/Engine/Public/Materials/MaterialInterface.h
@@ -4,6 +4,7 @@
 #pragma once
 
 #include "Async/TaskGraphInterfaces.h"
+#include "Async/RecursiveMutex.h"
 #include "CoreMinimal.h"
 #include "Delegates/Delegate.h"
 #include "Materials/MaterialParameters.h"
@@ -413,6 +414,11 @@ private:
 	/** Feature levels to force to compile. */
 	uint32 FeatureLevelsToForceCompile;
 
+#if WITH_EDITOR
+	/** Used for multithread access of TextureStreamingData in SortTextureStreamingData in editor (called by GetTextureDensity). */
+	mutable UE::FRecursiveMutex TextureStreamingDataMutex;
+#endif
+
 public:
 
 	/** Whether this material interface is included in the base game (and not in a DLC) */
