Skeletal Mesh: Use safe prop name referencing

#rnx
#Synced-CL 48733304

diff --git a/Engine/Source/Runtime/Engine/Private/SkeletalMesh.cpp b/Engine/Source/Runtime/Engine/Private/SkeletalMesh.cpp
--- a/Engine/Source/Runtime/Engine/Private/SkeletalMesh.cpp
+++ b/Engine/Source/Runtime/Engine/Private/SkeletalMesh.cpp
@@ -1443,6 +1443,13 @@ int32 USkeletalMesh::UnStackPostEditChange()
 
 void USkeletalMesh::PostEditChangeProperty(FPropertyChangedEvent& PropertyChangedEvent)
 {
+	PRAGMA_DISABLE_DEPRECATION_WARNINGS
+	static const FName bSupportRayTracingPropName = GET_MEMBER_NAME_CHECKED(USkeletalMesh, bSupportRayTracing);
+	static const FName RayTracingMinLODPropName = GET_MEMBER_NAME_CHECKED(USkeletalMesh, RayTracingMinLOD);
+	static const FName ClothLODBiasModePropName = GET_MEMBER_NAME_CHECKED(USkeletalMesh, ClothLODBiasMode);
+	static const FName PhysicsAssetPropName = GET_MEMBER_NAME_CHECKED(USkeletalMesh, PhysicsAsset);
+	PRAGMA_ENABLE_DEPRECATION_WARNINGS
+	
 	check(IsInGameThread());
 
 	if (PostEditChangeStackCounter > 0)
@@ -1458,11 +1465,12 @@ void USkeletalMesh::PostEditChangeProperty(FPropertyChangedEvent& PropertyChange
 
 	FProperty* PropertyThatChanged = PropertyChangedEvent.Property;
 	
+	
 	if (GIsEditor &&
 		PropertyThatChanged &&
-		(PropertyThatChanged->GetFName() == FName("bSupportRayTracing") ||
-		 PropertyThatChanged->GetFName() == FName("RayTracingMinLOD") ||
-		 PropertyThatChanged->GetFName() == FName("ClothLODBiasMode")))
+		(PropertyThatChanged->GetFName() == bSupportRayTracingPropName ||
+		 PropertyThatChanged->GetFName() == RayTracingMinLODPropName ||
+		 PropertyThatChanged->GetFName() == ClothLODBiasModePropName))
 	{
 		// Update the extra cloth deformer mapping LOD bias using this cloth entry
 		for (UClothingAssetBase* const ClothingAsset : GetMeshClothingAssets())
@@ -1489,7 +1497,7 @@ void USkeletalMesh::PostEditChangeProperty(FPropertyChangedEvent& PropertyChange
 
 	if( GIsEditor &&
 		PropertyThatChanged &&
-		PropertyThatChanged->GetFName() == FName("PhysicsAsset") )
+		PropertyThatChanged->GetFName() == PhysicsAssetPropName )
 	{
 		RefreshSkelMeshOnPhysicsAssetChange(this);
 	}
@@ -1566,9 +1574,7 @@ void USkeletalMesh::PostEditChangeProperty(FPropertyChangedEvent& PropertyChange
 			Datum->PostEditChangeOwner(PropertyChangedEvent);
 		}
 	}
-	PRAGMA_DISABLE_DEPRECATION_WARNINGS
-	for (UAssetUserData* Datum : AssetUserDataEditorOnly)
-	PRAGMA_ENABLE_DEPRECATION_WARNINGS
+	for (UAssetUserData* Datum : GetAssetUserDataEditorOnly())
 	{
 		if (Datum != nullptr)
 		{
