- Add OnHandleSystemHang to CoreDelegates to be invoked when a fatal engine Hang occurred.
- Add Hang type to Analytics Session system.

This Prevents cases of AbnormalShutdown reports being sent by CrashReportClient when a ReportHang was already received and engine will exit.

#Synced-CL 48600773

diff --git a/Engine/Source/Runtime/Analytics/AnalyticsET/Private/AnalyticsSessionSummaryManager.cpp b/Engine/Source/Runtime/Analytics/AnalyticsET/Private/AnalyticsSessionSummaryManager.cpp
--- a/Engine/Source/Runtime/Analytics/AnalyticsET/Private/AnalyticsSessionSummaryManager.cpp
+++ b/Engine/Source/Runtime/Analytics/AnalyticsET/Private/AnalyticsSessionSummaryManager.cpp
@@ -148,18 +148,21 @@ const TCHAR* LexToString(EAnalyticsSessionShutdownType ShutdownTypeCode)
 	case EAnalyticsSessionShutdownType::Shutdown:
 		return TEXT("Shutdown");
 
-	case EAnalyticsSessionShutdownType::Crashed:
-		return TEXT("Crashed");
-
 	case EAnalyticsSessionShutdownType::Terminated:
 		return TEXT("Terminated");
 
 	case EAnalyticsSessionShutdownType::Debugged:
 		return TEXT("Debugger");
 
+	case EAnalyticsSessionShutdownType::Crashed:
+		return TEXT("Crashed");
+
 	case EAnalyticsSessionShutdownType::Abnormal:
 		return TEXT("AbnormalShutdown");
 
+	case EAnalyticsSessionShutdownType::Hang:
+		return TEXT("Hang");
+
 	case EAnalyticsSessionShutdownType::Unknown:
 	default:
 		return TEXT("Unknown");
diff --git a/Engine/Source/Runtime/Analytics/AnalyticsET/Public/AnalyticsSessionSummaryManager.h b/Engine/Source/Runtime/Analytics/AnalyticsET/Public/AnalyticsSessionSummaryManager.h
--- a/Engine/Source/Runtime/Analytics/AnalyticsET/Public/AnalyticsSessionSummaryManager.h
+++ b/Engine/Source/Runtime/Analytics/AnalyticsET/Public/AnalyticsSessionSummaryManager.h
@@ -28,6 +28,7 @@ enum class EAnalyticsSessionShutdownType : int32
 	Crashed,
 	Abnormal,
 	Unknown,
+	Hang,
 };
 
 /**
diff --git a/Engine/Source/Runtime/Core/Private/HAL/ThreadHeartBeat.cpp b/Engine/Source/Runtime/Core/Private/HAL/ThreadHeartBeat.cpp
--- a/Engine/Source/Runtime/Core/Private/HAL/ThreadHeartBeat.cpp
+++ b/Engine/Source/Runtime/Core/Private/HAL/ThreadHeartBeat.cpp
@@ -253,6 +253,7 @@ void FORCENOINLINE FThreadHeartBeat::OnHang(double HangDuration, uint32 ThreadTh
 		if (bHangsAreFatal)
 		{
 			GLog->Panic();
+			FCoreDelegates::OnHandleSystemHang.Broadcast();
 		}
 		else
 		{
diff --git a/Engine/Source/Runtime/Core/Private/Misc/CoreDelegates.cpp b/Engine/Source/Runtime/Core/Private/Misc/CoreDelegates.cpp
--- a/Engine/Source/Runtime/Core/Private/Misc/CoreDelegates.cpp
+++ b/Engine/Source/Runtime/Core/Private/Misc/CoreDelegates.cpp
@@ -86,6 +86,7 @@ TMulticastDelegate<void()> FCoreDelegates::OnSafeFrameChangedEvent;
 TMulticastDelegate<void(const ANSICHAR*, const ANSICHAR*, int32, const TCHAR*, const TCHAR*)> FCoreDelegates::OnEnsureFailed;
 TMulticastDelegate<void()> FCoreDelegates::OnHandleSystemEnsure;
 TMulticastDelegate<void()> FCoreDelegates::OnHandleSystemError;
+TMulticastDelegate<void()> FCoreDelegates::OnHandleSystemHang;
 
 TMulticastDelegate<void(AActor*)> FCoreDelegates::OnActorLabelChanged;
 
diff --git a/Engine/Source/Runtime/Core/Public/Misc/CoreDelegates.h b/Engine/Source/Runtime/Core/Public/Misc/CoreDelegates.h
--- a/Engine/Source/Runtime/Core/Public/Misc/CoreDelegates.h
+++ b/Engine/Source/Runtime/Core/Public/Misc/CoreDelegates.h
@@ -166,6 +166,10 @@ public:
 	// Callback when an error (crash) has occurred
 	static CORE_API TMulticastDelegate<void()> OnHandleSystemError;
 
+	// Callback when a fatal hang has occurred.
+	// WARNING: Use AddRaw/AddStatic for this delegate. Other bindings could be unsafe and cause deadlocks, e.g. AddUObject invoked during GC.
+	static CORE_API TMulticastDelegate<void()> OnHandleSystemHang;
+
 	// Called when an actor label is changed
 	static CORE_API TMulticastDelegate<void(AActor*)> OnActorLabelChanged;
 
diff --git a/Engine/Source/Runtime/Engine/Private/Analytics/EngineAnalyticsSessionSummary.cpp b/Engine/Source/Runtime/Engine/Private/Analytics/EngineAnalyticsSessionSummary.cpp
--- a/Engine/Source/Runtime/Engine/Private/Analytics/EngineAnalyticsSessionSummary.cpp
+++ b/Engine/Source/Runtime/Engine/Private/Analytics/EngineAnalyticsSessionSummary.cpp
@@ -34,7 +34,8 @@ namespace EngineAnalyticsProperties
 	static const TAnalyticsProperty<FString> UserActivity        = TEXT("CurrentUserActivity");
 	static const TAnalyticsProperty<int32>  MonitorExitCode      = TEXT("MonitorExitCode");
 	static const TAnalyticsProperty<float> AverageFPS            = TEXT("AverageFPS");
-	static const TAnalyticsProperty<bool> IsTerminating          = TEXT("IsTerminating");
+	static const TAnalyticsProperty<bool> IsTerminating			 = TEXT("IsTerminating");
+	static const TAnalyticsProperty<bool> IsHang				 = TEXT("IsHang");
 	static const TAnalyticsProperty<bool> WasShutdown            = TEXT("WasShutdown");
 	static const TAnalyticsProperty<bool> IsCrashing             = TEXT("IsCrash");
 	static const TAnalyticsProperty<bool> IsGpuCrash             = TEXT("GPUCrash");
@@ -109,15 +110,15 @@ namespace EngineAnalyticsProperties
 		{
 			return false;
 		}
-		else if (static_cast<EAnalyticsSessionShutdownType>(ProposedShutdownType) == EAnalyticsSessionShutdownType::Debugged) // Debugger can overwrite Shutdown|Terminated|Abnormal
+		else if (static_cast<EAnalyticsSessionShutdownType>(ProposedShutdownType) == EAnalyticsSessionShutdownType::Debugged) // Debugger can overwrite Shutdown|Terminated|Abnormal|Hang
 		{
 			return true;
 		}
-		else if (static_cast<EAnalyticsSessionShutdownType>(*ActualShutdownType) == EAnalyticsSessionShutdownType::Debugged) // Shutdown|Terminated|Abnormal cannot overwrite the debugger.
+		else if (static_cast<EAnalyticsSessionShutdownType>(*ActualShutdownType) == EAnalyticsSessionShutdownType::Debugged) // Shutdown|Terminated|Abnormal|Hang cannot overwrite the debugger.
 		{
 			return false;
 		}
-		return true; // Shutdown|Terminate|Abnormal can overwrite each other.
+		return true; // Shutdown|Terminate|Abnormal|Hang can overwrite each other.
 	}
 } // namespace EngineAnalyticsProperties
 
@@ -202,6 +203,7 @@ FEngineAnalyticsSessionSummary::FEngineAnalyticsSessionSummary(TSharedPtr<IAnaly
 	EngineAnalyticsProperties::UserActivity.Set(GetStore(), EngineAnalyticsProperties::GetUserActivityString(FUserActivityTracking::GetUserActivity()));
 	EngineAnalyticsProperties::AverageFPS.Set(GetStore(), EngineAnalyticsProperties::GetAverageFPS());
 	EngineAnalyticsProperties::IsTerminating.Set(GetStore(),  false);
+	EngineAnalyticsProperties::IsHang.Set(GetStore(), false);
 	EngineAnalyticsProperties::WasShutdown.Set(GetStore(), false);
 	EngineAnalyticsProperties::IsCrashing.Set(GetStore(), false);
 	EngineAnalyticsProperties::IsGpuCrash.Set(GetStore(), false);
@@ -217,6 +219,7 @@ FEngineAnalyticsSessionSummary::FEngineAnalyticsSessionSummary(TSharedPtr<IAnaly
 
 	// Listen to interesting events.
 	FCoreDelegates::OnHandleSystemError.AddRaw(this, &FEngineAnalyticsSessionSummary::OnCrashing); // WARNING: Don't assume this function is only called from game thread.
+	FCoreDelegates::OnHandleSystemHang.AddRaw(this, &FEngineAnalyticsSessionSummary::OnHung); // WARNING: Don't assume this function is only called from game thread.
 	FCoreDelegates::GetApplicationWillTerminateDelegate().AddRaw(this, &FEngineAnalyticsSessionSummary::OnTerminate); // WARNING: Don't assume this function is only called from game thread.
 	FCoreDelegates::IsVanillaProductChanged.AddRaw(this, &FEngineAnalyticsSessionSummary::OnVanillaStateChanged);
 	FCoreDelegates::OnUserLoginChangedEvent.AddRaw(this, &FEngineAnalyticsSessionSummary::OnUserLoginChanged);
@@ -231,6 +234,7 @@ void FEngineAnalyticsSessionSummary::Shutdown()
 		ShutdownInternal();
 
 		FCoreDelegates::OnHandleSystemError.RemoveAll(this);
+		FCoreDelegates::OnHandleSystemHang.RemoveAll(this);
 		FCoreDelegates::GetApplicationWillTerminateDelegate().RemoveAll(this);
 		FCoreDelegates::IsVanillaProductChanged.RemoveAll(this);
 		FCoreDelegates::OnUserLoginChangedEvent.RemoveAll(this);
@@ -283,6 +287,13 @@ void FEngineAnalyticsSessionSummary::OnCrashing()
 	UpdateSessionProgress(/*bFlushAsync*/false, /*FlushTimeout*/FTimespan::FromMilliseconds(100), /*bCrashing*/true);
 }
 
+void FEngineAnalyticsSessionSummary::OnHung()
+{
+	EngineAnalyticsProperties::IsHang.Set(GetStore(), true);
+	EngineAnalyticsProperties::ShutdownType.Set(GetStore(), (int32)EAnalyticsSessionShutdownType::Hang, EngineAnalyticsProperties::ShouldOverwriteShutdownType);
+	UpdateSessionProgress(/*bFlushAsync*/false, /*FlushTimeout*/FTimespan::FromMilliseconds(100));
+}
+
 void FEngineAnalyticsSessionSummary::OnTerminate()
 {
 	EngineAnalyticsProperties::IsTerminating.Set(GetStore(), true);
diff --git a/Engine/Source/Runtime/Engine/Public/Analytics/EngineAnalyticsSessionSummary.h b/Engine/Source/Runtime/Engine/Public/Analytics/EngineAnalyticsSessionSummary.h
--- a/Engine/Source/Runtime/Engine/Public/Analytics/EngineAnalyticsSessionSummary.h
+++ b/Engine/Source/Runtime/Engine/Public/Analytics/EngineAnalyticsSessionSummary.h
@@ -59,6 +59,9 @@ private:
 	/** Invoked when the engine is crashing.. */
 	ENGINE_API void OnCrashing();
 
+	/** Invoked when the engine fatally hung */
+	ENGINE_API void OnHung();
+
 	/** Invoked when the engine is requested to terminate. */
 	ENGINE_API void OnTerminate();
 
diff --git a/Engine/Source/Runtime/RenderCore/Private/RenderingThread.cpp b/Engine/Source/Runtime/RenderCore/Private/RenderingThread.cpp
--- a/Engine/Source/Runtime/RenderCore/Private/RenderingThread.cpp
+++ b/Engine/Source/Runtime/RenderCore/Private/RenderingThread.cpp
@@ -1137,6 +1137,7 @@ static void HandleRenderTaskHang(uint32 ThreadThatHung, double HangDuration)
 		*ThreadName, HangDuration, LINE_TERMINATOR, *StackTrimmed, LINE_TERMINATOR);
 	
 	GLog->Panic();
+	FCoreDelegates::OnHandleSystemHang.Broadcast();
 	ReportHang(*ErrorMessage, StackFrames, NumStackFrames, ThreadThatHung);
 	if (FApp::CanEverRender())
 	{
