Refactored FPSOPrecacheComponentData to allow for non UObject's to use FPSOPrecacheFinishedTask.

#Synced-CL 48520266

diff --git a/Engine/Source/Runtime/Engine/Public/PSOPrecacheComponent.h b/Engine/Source/Runtime/Engine/Public/PSOPrecacheComponent.h
--- a/Engine/Source/Runtime/Engine/Public/PSOPrecacheComponent.h
+++ b/Engine/Source/Runtime/Engine/Public/PSOPrecacheComponent.h
@@ -113,11 +113,22 @@ private:
 
 #if UE_WITH_PSO_PRECACHING
 
+template<typename T, typename Enable = void>
+struct TPSOWeakPtrTraits
+{
+	using WeakType = TWeakObjectPtr<T>;
+
+	static WeakType Make(T* Ptr) { return WeakType(Ptr); }
+	static T* Get(const WeakType& W) { return W.Get(); }
+};
+
 template<typename ComponentType, typename PrecachedFinishedOnComponentCallback>
 struct FPSOPrecacheFinishedTask
 {
+	using WeakType = typename TPSOWeakPtrTraits<ComponentType>::WeakType;
+
 	explicit FPSOPrecacheFinishedTask(ComponentType* Component, int32 InJobSetThatJustCompleted, PrecachedFinishedOnComponentCallback InCallback)
-		: WeakComponent(Component),
+		: WeakComponent(TPSOWeakPtrTraits<ComponentType>::Make(Component)),
 		JobSetThatJustCompleted(InJobSetThatJustCompleted),
 		Callback(InCallback)
 	{
@@ -129,7 +140,7 @@ struct FPSOPrecacheFinishedTask
 
 	void DoTask(ENamedThreads::Type CurrentThread, const FGraphEventRef& MyCompletionGraphEvent)
 	{
-		if (ComponentType* Component = WeakComponent.Get())
+		if (ComponentType* Component = TPSOWeakPtrTraits<ComponentType>::Get(WeakComponent))
 		{
 			QUICK_SCOPE_CYCLE_COUNTER(STAT_PSOPrecacheFinishedTask);
 			FPSOPrecacheComponentData& PSOPrecacheComponentData = Component->GetPSOPrecacheComponentData();
@@ -139,7 +150,7 @@ struct FPSOPrecacheFinishedTask
 		}
 	}
 		
-	TWeakObjectPtr<ComponentType> WeakComponent;
+	WeakType WeakComponent;
 	int32 JobSetThatJustCompleted;
 	PrecachedFinishedOnComponentCallback Callback;
 };
