Fix secondary issues with the localized console key

- Fix duplicate keys added to settings
- Fix German keyboard support, Caret is a dead key which likely doesn't produce an actual input we can use, using the scancode instead appears to fix it

#rnx
#tests Tested on a German keyboard
#Synced-CL 48642164

diff --git a/Engine/Source/Runtime/Engine/Private/UserInterface/InputSettings.cpp b/Engine/Source/Runtime/Engine/Private/UserInterface/InputSettings.cpp
--- a/Engine/Source/Runtime/Engine/Private/UserInterface/InputSettings.cpp
+++ b/Engine/Source/Runtime/Engine/Private/UserInterface/InputSettings.cpp
@@ -174,7 +174,7 @@ void UInputSettings::AddInternationalConsoleKey()
 				break;
 
 			case LANG_GERMAN:
-				DefaultConsoleKey = EKeys::Caret;
+				DefaultConsoleKey = FInputKeyManager::Get().GetKeyFromCodes(VK_OEM_5, 0);
 				break;
 
 			case LANG_ITALIAN:
@@ -198,7 +198,7 @@ void UInputSettings::AddInternationalConsoleKey()
 
 			if (DefaultConsoleKey != EKeys::Tilde && DefaultConsoleKey.IsValid())
 			{
-				ConsoleKeys.Add(DefaultConsoleKey);
+				ConsoleKeys.AddUnique(DefaultConsoleKey);
 			}
 		}
 	}
