Fixed uninitialized members in FPSOPrecacheComponentData.

diff --git a/Engine/Source/Runtime/Engine/Public/PSOPrecacheComponent.h b/Engine/Source/Runtime/Engine/Public/PSOPrecacheComponent.h
--- a/Engine/Source/Runtime/Engine/Public/PSOPrecacheComponent.h
+++ b/Engine/Source/Runtime/Engine/Public/PSOPrecacheComponent.h
@@ -94,10 +94,10 @@ struct FPSOPrecacheComponentData
 private:
 
 	/** Helper flag to check if PSOs have been precached already */
-	uint8 bPSOPrecacheCalled : 1;
+	uint8 bPSOPrecacheCalled : 1 = false;
 
 	/** PSOs requested priority */
-	EPSOPrecachePriority PSOPrecacheRequestPriority : 2;
+	EPSOPrecachePriority PSOPrecacheRequestPriority : 2 = EPSOPrecachePriority::Medium;
 	static_assert((int)EPSOPrecachePriority::Highest < 1 << 2);
 
 	int32 LatestPSOPrecacheJobSet : 29 = 0;
