SpringMath: Added inertialization functions

#virtualized
#Synced-CL 48762567

diff --git a/Engine/Source/Runtime/Engine/Classes/Kismet/BlueprintSpringMathLibrary.h b/Engine/Source/Runtime/Engine/Classes/Kismet/BlueprintSpringMathLibrary.h
--- a/Engine/Source/Runtime/Engine/Classes/Kismet/BlueprintSpringMathLibrary.h
+++ b/Engine/Source/Runtime/Engine/Classes/Kismet/BlueprintSpringMathLibrary.h
@@ -6,6 +6,25 @@
 
 #include "BlueprintSpringMathLibrary.generated.h"
 
+/** A simple struct that stores velocities for an FTransform. */
+USTRUCT(BlueprintType, Experimental)
+struct FSpringTransformVelocity
+{
+	GENERATED_BODY()
+
+	// Change in Location over time in cm/s
+	UPROPERTY(VisibleAnywhere, BlueprintReadWrite, Category = "Transform", meta= (ForceUnits = "cm/s"))
+	FVector LinearVelocity = FVector::ZeroVector;
+
+	// Change in rotation over time in deg/s
+	UPROPERTY(VisibleAnywhere, BlueprintReadWrite, Category = "Transform", meta = (ForceUnits = "deg/s"))
+	FVector AngularVelocity = FVector::ZeroVector;
+
+	// Change in scale over time
+	UPROPERTY(VisibleAnywhere, BlueprintReadWrite, Category = "Transform")
+	FVector ScalarVelocity = FVector::ZeroVector;
+};
+
 UCLASS(MinimalAPI, Experimental, meta=(BlueprintThreadSafe))
 class UBlueprintSpringMathLibrary : public UBlueprintFunctionLibrary
 {
@@ -86,6 +105,30 @@ public:
 	UFUNCTION(BlueprintCallable, meta=(Category="Math|Spring|Experimental"))
 	static ENGINE_API void CriticalSpringDampRotator(UPARAM(ref) FRotator& InOutRotation, UPARAM(ref) FVector& InOutAngularVelocity, const FRotator& TargetRotation, float DeltaTime, float SmoothingTime = 0.2f);
 
+	UE_EXPERIMENTAL(5.8, "SpringMath is experimental")
+	/** Interpolates the value InOutScale towards TargetScale with the motion of a critically damped spring.
+	 *
+	 * @param InOutScale The value to be damped
+	 * @param InOutScalarVelocity The speed of the value to be damped
+	 * @param TargetScale The goal to damp towards
+	 * @param SmoothingTime The smoothing time to use for the spring. Longer times lead to more damped behaviour. A time of 0 will snap InOutScale to TargetScale
+	 * @param DeltaTime Timestep in seconds
+	 */
+	UFUNCTION(BlueprintCallable, meta = (Category = "Math|Spring|Experimental"))
+	static ENGINE_API void CriticalSpringDampScale(UPARAM(ref) FVector& InOutScale, UPARAM(ref) FVector& InOutScalarVelocity, const FVector& TargetScale, float DeltaTime, float SmoothingTime = 0.2f);
+
+	UE_EXPERIMENTAL(5.8, "SpringMath is experimental")
+	/** Interpolates the value InOutTransform towards TargetTransform with the motion of a critically damped spring.
+	 *
+	 * @param InOutTransform The value to be damped
+	 * @param InOutVelocity The speed of the value to be damped
+	 * @param TargetTransform The goal to damp towards
+	 * @param SmoothingTime The smoothing time to use for the spring. Longer times lead to more damped behaviour. A time of 0 will snap InOutTransform to TargetTransform
+	 * @param DeltaTime Timestep in seconds
+	 */
+	UFUNCTION(BlueprintCallable, meta = (Category = "Math|Spring|Experimental"))
+	static ENGINE_API void CriticalSpringDampTransform(UPARAM(ref) FTransform& InOutTransform, UPARAM(ref) FSpringTransformVelocity& InOutVelocity, const FTransform& TargetTransform, float DeltaTime, float SmoothingTime = 0.2f);
+
 
 	// Double Spring //
 
@@ -173,6 +216,34 @@ public:
 	UFUNCTION(BlueprintCallable, meta = (Category = "Math|Spring|Experimental"))
 	static ENGINE_API void CriticalDoubleSpringDampRotator(UPARAM(ref) FRotator& InOutRotation, UPARAM(ref) FVector& InOutAngularVelocity, UPARAM(ref) FRotator& InOutIntermediateRotation, UPARAM(ref) FVector& InOutIntermediateAngularVelocity, const FRotator& TargetRotation, float DeltaTime, float SmoothingTime = 0.2f);
 
+	UE_EXPERIMENTAL(5.8, "SpringMath is experimental")
+	/** Interpolates the value InOutScale towards TargetScale using a critically damped spring via an intermediate state - producing a more S-shaped curve.
+	 *
+	 * @param InOutScale The value to be damped
+	 * @param InOutScalarVelocity The speed of the value to be damped
+	 * @param InOutIntermediateScale The intermediate state of the double spring
+	 * @param InOutIntermediateScalarVelocity The intermediate velocity of the double spring
+	 * @param TargetScale The goal to damp towards
+	 * @param SmoothingTime The smoothing time to use for the spring. Longer times lead to more damped behaviour. A time of 0 will snap InOutScale to TargetScale
+	 * @param DeltaTime Timestep in seconds
+	 */
+	UFUNCTION(BlueprintCallable, meta = (Category = "Math|Spring|Experimental"))
+	static ENGINE_API void CriticalDoubleSpringDampScale(UPARAM(ref) FVector& InOutScale, UPARAM(ref) FVector& InOutScalarVelocity, UPARAM(ref) FVector& InOutIntermediateScale, UPARAM(ref) FVector& InOutIntermediateScalarVelocity, const FVector& TargetScale, float DeltaTime, float SmoothingTime = 0.2f);
+
+	UE_EXPERIMENTAL(5.8, "SpringMath is experimental")
+	/** Interpolates the value InOutTransform towards TargetTransform using a critically damped spring via an intermediate state - producing a more S-shaped curve.
+	 *
+	 * @param InOutTransform The value to be damped
+	 * @param InOutVelocity The speed of the value to be damped
+	 * @param InOutIntermediateTransform The intermediate state of the double spring
+	 * @param InOutIntermediateVelocity The intermediate velocity of the double spring
+	 * @param TargetTransform The goal to damp towards
+	 * @param SmoothingTime The smoothing time to use for the spring. Longer times lead to more damped behaviour. A time of 0 will snap InOutTransform to TargetTransform
+	 * @param DeltaTime Timestep in seconds
+	 */
+	UFUNCTION(BlueprintCallable, meta = (Category = "Math|Spring|Experimental"))
+	static ENGINE_API void CriticalDoubleSpringDampTransform(UPARAM(ref) FTransform& InOutTransform, UPARAM(ref) FSpringTransformVelocity& InOutVelocity, UPARAM(ref) FTransform& InOutIntermediateTransform, UPARAM(ref) FSpringTransformVelocity& InOutIntermediateTransformVelocity, const FTransform& TargetTransform, float DeltaTime, float SmoothingTime = 0.2f);
+
 
 	// Velocity Spring //
 
@@ -302,6 +373,32 @@ public:
 	UFUNCTION(BlueprintPure, meta=(Category="Math|Spring|Experimental"))
 	static ENGINE_API FRotator DampRotator(const FRotator& Rotation, const FRotator& TargetRotation, float DeltaTime, float SmoothingTime = 0.2f);
 
+	UE_EXPERIMENTAL(5.8, "SpringMath is experimental")
+	/**
+	 * Smooths a scale using exponential damping towards a target.
+	 *
+	 *
+	 * @param  Scale The scale to be smoothed
+	 * @param  TargetScale The target to smooth towards
+	 * @param  DeltaTime Time interval
+	 * @param  SmoothingTime Timescale over which to smooth. Larger values result in more smoothed behaviour. Can be zero.
+	 */
+	UFUNCTION(BlueprintPure, meta = (Category = "Math|Spring|Experimental"))
+	static ENGINE_API FVector DampScale(const FVector& Scale, const FVector& TargetScale, float DeltaTime, float SmoothingTime = 0.2f);
+
+	UE_EXPERIMENTAL(5.8, "SpringMath is experimental")
+	/**
+	 * Smooths a transform using exponential damping towards a target.
+	 *
+	 *
+	 * @param  Transform The transform to be smoothed
+	 * @param  TargetTransform The target to smooth towards
+	 * @param  DeltaTime Time interval
+	 * @param  SmoothingTime Timescale over which to smooth. Larger values result in more smoothed behaviour. Can be zero.
+	 */
+	UFUNCTION(BlueprintPure, meta = (Category = "Math|Spring|Experimental"))
+	static ENGINE_API FTransform DampTransform(const FTransform& Transform, const FTransform& TargetTransform, float DeltaTime, float SmoothingTime = 0.2f);
+
 
 	// Character //
 
@@ -338,6 +435,7 @@ public:
 
 	// Conversion Methods //
 
+	UE_EXPERIMENTAL(5.7, "SpringMath is experimental")
 	/** Convert from smoothing time to spring strength.
 	* 
 	* @param SmoothingTime The smoothing time of the spring in seconds. It is the time by which the output lags the input when critically damped, when the input is changing at a constant rate.
@@ -346,6 +444,7 @@ public:
 	UFUNCTION(BlueprintPure, meta=(Category="Math|Spring|Experimental"))
 	static ENGINE_API float ConvertSmoothingTimeToStrength(float SmoothingTime);
 
+	UE_EXPERIMENTAL(5.7, "SpringMath is experimental")
 	/** Convert from spring strength to smoothing time.
  	* 
  	* @param Strength The spring strength. This corresponds to the undamped frequency of the spring in hz.
@@ -354,6 +453,7 @@ public:
 	UFUNCTION(BlueprintPure, meta=(Category="Math|Spring|Experimental"))
 	static ENGINE_API float ConvertStrengthToSmoothingTime(float Strength);
 
+	UE_EXPERIMENTAL(5.7, "SpringMath is experimental")
 	/** Convert a halflife to a smoothing time
  	* 
  	* @param HalfLife The half life of the spring. How long it takes the value to get halfway towards the target.
@@ -362,6 +462,7 @@ public:
 	UFUNCTION(BlueprintPure, meta=(Category="Math|Spring|Experimental"))
 	static ENGINE_API float ConvertHalfLifeToSmoothingTime(float HalfLife);
 
+	UE_EXPERIMENTAL(5.7, "SpringMath is experimental")
 	/** Convert a smoothing time to a half life
  	* 
  	* @param SmoothingTime The smoothing time of the spring in seconds. It is the time by which the output lags the input when critically damped, when the input is changing at a constant rate. 
@@ -369,4 +470,726 @@ public:
  	*/
 	UFUNCTION(BlueprintPure, meta=(Category="Math|Spring|Experimental"))
 	static ENGINE_API float ConvertSmoothingTimeToHalfLife(float SmoothingTime);
+
+
+	// Velocity Tracking //
+
+	UE_EXPERIMENTAL(5.8, "SpringMath is experimental")
+	/** Update the value and velocity using the finite difference  */
+	UFUNCTION(BlueprintCallable, meta = (Category = "Math|Spring|Experimental"))
+	static ENGINE_API void TrackVelocityFloat(UPARAM(ref) float& InOutValue, UPARAM(ref) float& InOutVelocity, float InValue, float DeltaTime);
+
+	UE_EXPERIMENTAL(5.8, "SpringMath is experimental")
+	/** Update the value and velocity using the finite difference  */
+	UFUNCTION(BlueprintCallable, meta = (Category = "Math|Spring|Experimental"))
+	static ENGINE_API void TrackVelocityVector(UPARAM(ref) FVector& InOutValue, UPARAM(ref) FVector& InOutVelocity, FVector InValue, float DeltaTime);
+
+	UE_EXPERIMENTAL(5.8, "SpringMath is experimental")
+	/** Update the value and velocity using the finite difference  */
+	UFUNCTION(BlueprintCallable, meta = (Category = "Math|Spring|Experimental"))
+	static ENGINE_API void TrackVelocityVector2D(UPARAM(ref) FVector2D& InOutValue, UPARAM(ref) FVector2D& InOutVelocity, FVector2D InValue, float DeltaTime);
+
+	UE_EXPERIMENTAL(5.8, "SpringMath is experimental")
+	/** Update the value and velocity using the finite difference  */
+	UFUNCTION(BlueprintCallable, meta = (Category = "Math|Spring|Experimental"))
+	static ENGINE_API void TrackVelocityQuat(UPARAM(ref) FQuat& InOutValue, UPARAM(ref) FVector& InOutVelocity, FQuat InValue, float DeltaTime);
+
+	UE_EXPERIMENTAL(5.8, "SpringMath is experimental")
+	/** Update the value and velocity using the finite difference  */
+	UFUNCTION(BlueprintCallable, meta = (Category = "Math|Spring|Experimental"))
+	static ENGINE_API void TrackVelocityRotator(UPARAM(ref) FRotator& InOutValue, UPARAM(ref) FVector& InOutVelocity, FRotator InValue, float DeltaTime);
+
+	UE_EXPERIMENTAL(5.8, "SpringMath is experimental")
+	/** Update the value and velocity using the finite difference  */
+	UFUNCTION(BlueprintCallable, meta = (Category = "Math|Spring|Experimental"))
+	static ENGINE_API void TrackVelocityAngle(UPARAM(ref) float& InOutValue, UPARAM(ref) float& InOutVelocity, float InValue, float DeltaTime);
+
+	UE_EXPERIMENTAL(5.8, "SpringMath is experimental")
+	/** Update the value and velocity using the finite difference  */
+	UFUNCTION(BlueprintCallable, meta = (Category = "Math|Spring|Experimental"))
+	static ENGINE_API void TrackVelocityScale(UPARAM(ref) FVector& InOutValue, UPARAM(ref) FVector& InOutVelocity, FVector InValue, float DeltaTime);
+
+	UE_EXPERIMENTAL(5.8, "SpringMath is experimental")
+	/** Update the value and velocity using the finite difference  */
+	UFUNCTION(BlueprintCallable, meta = (Category = "Math|Spring|Experimental"))
+	static ENGINE_API void TrackVelocityTransform(UPARAM(ref) FTransform& InOutValue, UPARAM(ref) FSpringTransformVelocity& InOutVelocity, FTransform InValue, float DeltaTime);
+
+	// Spring Inertialization //
+
+	UE_EXPERIMENTAL(5.8, "SpringMath is experimental")
+	/** Applies a spring inertialization to a float value, adding the appropriate offset given the time since transition.
+	*
+	* @param InValue				The value to apply the inertialization to
+	* @param InValueOffset			Offset value at the time of transition
+	* @param InVelocityOffset		Offset velocity at the time of transition
+	* @param TimeSinceTransition	Time since the last transition
+	* @param SmoothingTime			Timescale over which to apply the offset. Larger values result in longer blend times. Can be zero.
+	* @return						The value with appropriate offset applied
+	*/
+	UFUNCTION(BlueprintPure, meta = (Category = "Math|Spring|Experimental"))
+	static ENGINE_API float SpringInertializeApplyToFloat(float InValue, float InValueOffset, float InVelocityOffset, float TimeSinceTransition, float SmoothingTime = 0.2f);
+
+	UE_EXPERIMENTAL(5.8, "SpringMath is experimental")
+	/** Transitions a spring inertialization, updating the appropriate offsets.
+	*
+	* @param InOutValueOffset		Offset value to update
+	* @param InOutVelocityOffset	Offset velocity to update
+	* @param InSrcValue				Value of the source (value transitioning from)
+	* @param InSrcVelocity			Velocity of the source (velocity transitioning from)
+	* @param InDstValue				Value of the destination (value transitioning to)
+	* @param InDstVelocity			Velocity of the destination (velocity transitioning to)
+	* @param TimeSinceTransition	Time since the last transition
+	* @param SmoothingTime			Timescale over which to apply the offset. Larger values result in longer blend times. Can be zero.
+	*/
+	UFUNCTION(BlueprintCallable, meta = (Category = "Math|Spring|Experimental"))
+	static ENGINE_API void SpringInertializeTransitionFloat(UPARAM(ref) float& InOutValueOffset, UPARAM(ref) float& InOutVelocityOffset, float InSrcValue, float InSrcVelocity, float InDstValue, float InDstVelocity, float TimeSinceTransition, float SmoothingTime = 0.2f);
+
+	UE_EXPERIMENTAL(5.8, "SpringMath is experimental")
+	/** Applies a spring inertialization to a vector value, adding the appropriate offset given the time since transition.
+	*
+	* @param InValue				The value to apply the inertialization to
+	* @param InValueOffset			Offset value at the time of transition
+	* @param InVelocityOffset		Offset velocity at the time of transition
+	* @param TimeSinceTransition	Time since the last transition
+	* @param SmoothingTime			Timescale over which to apply the offset. Larger values result in longer blend times. Can be zero.
+	* @return						The value with appropriate offset applied
+	*/
+	UFUNCTION(BlueprintPure, meta = (Category = "Math|Spring|Experimental"))
+	static ENGINE_API FVector SpringInertializeApplyToVector(FVector InValue, FVector InValueOffset, FVector InVelocityOffset, float TimeSinceTransition, float SmoothingTime = 0.2f);
+
+	UE_EXPERIMENTAL(5.8, "SpringMath is experimental")
+	/** Transitions a spring inertialization, updating the appropriate offsets.
+	*
+	* @param InOutValueOffset		Offset value to update
+	* @param InOutVelocityOffset	Offset velocity to update
+	* @param InSrcValue				Value of the source (value transitioning from)
+	* @param InSrcVelocity			Velocity of the source (velocity transitioning from)
+	* @param InDstValue				Value of the destination (value transitioning to)
+	* @param InDstVelocity			Velocity of the destination (velocity transitioning to)
+	* @param TimeSinceTransition	Time since the last transition
+	* @param SmoothingTime			Timescale over which to apply the offset. Larger values result in longer blend times. Can be zero.
+	*/
+	UFUNCTION(BlueprintCallable, meta = (Category = "Math|Spring|Experimental"))
+	static ENGINE_API void SpringInertializeTransitionVector(UPARAM(ref) FVector& InOutValueOffset, UPARAM(ref) FVector& InOutVelocityOffset, FVector InSrcValue, FVector InSrcVelocity, FVector InDstValue, FVector InDstVelocity, float TimeSinceTransition, float SmoothingTime = 0.2f);
+
+	UE_EXPERIMENTAL(5.8, "SpringMath is experimental")
+	/** Applies a spring inertialization to a Vector2D value, adding the appropriate offset given the time since transition.
+	*
+	* @param InValue				The value to apply the inertialization to
+	* @param InValueOffset			Offset value at the time of transition
+	* @param InVelocityOffset		Offset velocity at the time of transition
+	* @param TimeSinceTransition	Time since the last transition
+	* @param SmoothingTime			Timescale over which to apply the offset. Larger values result in longer blend times. Can be zero.
+	* @return						The value with appropriate offset applied
+	*/
+	UFUNCTION(BlueprintPure, meta = (Category = "Math|Spring|Experimental"))
+	static ENGINE_API FVector2D SpringInertializeApplyToVector2D(FVector2D InValue, FVector2D InValueOffset, FVector2D InVelocityOffset, float TimeSinceTransition, float SmoothingTime = 0.2f);
+
+	UE_EXPERIMENTAL(5.8, "SpringMath is experimental")
+	/** Transitions a spring inertialization, updating the appropriate offsets.
+	*
+	* @param InOutValueOffset		Offset value to update
+	* @param InOutVelocityOffset	Offset velocity to update
+	* @param InSrcValue				Value of the source (value transitioning from)
+	* @param InSrcVelocity			Velocity of the source (velocity transitioning from)
+	* @param InDstValue				Value of the destination (value transitioning to)
+	* @param InDstVelocity			Velocity of the destination (velocity transitioning to)
+	* @param TimeSinceTransition	Time since the last transition
+	* @param SmoothingTime			Timescale over which to apply the offset. Larger values result in longer blend times. Can be zero.
+	*/
+	UFUNCTION(BlueprintCallable, meta = (Category = "Math|Spring|Experimental"))
+	static ENGINE_API void SpringInertializeTransitionVector2D(UPARAM(ref) FVector2D& InOutValueOffset, UPARAM(ref) FVector2D& InOutVelocityOffset, FVector2D InSrcValue, FVector2D InSrcVelocity, FVector2D InDstValue, FVector2D InDstVelocity, float TimeSinceTransition, float SmoothingTime = 0.2f);
+
+	UE_EXPERIMENTAL(5.8, "SpringMath is experimental")
+	/** Applies a spring inertialization to a quaternion, applying the appropriate offset given the time since transition.
+	*
+	* @param InRotation					The rotation to apply the inertialization to
+	* @param InRotationOffset			Offset rotation at the time of transition
+	* @param InAngularVelocityOffset	Offset angular velocity at the time of transition
+	* @param TimeSinceTransition		Time since the last transition
+	* @param SmoothingTime				Timescale over which to apply the offset. Larger values result in longer blend times. Can be zero.
+	* @return							The rotation with appropriate offset applied
+	*/
+	UFUNCTION(BlueprintPure, meta = (Category = "Math|Spring|Experimental"))
+	static ENGINE_API FQuat SpringInertializeApplyToQuat(FQuat InRotation, FQuat InRotationOffset, FVector InAngularVelocityOffset, float TimeSinceTransition, float SmoothingTime = 0.2f);
+
+	UE_EXPERIMENTAL(5.8, "SpringMath is experimental")
+	/** Transitions a spring inertialization, updating the appropriate offsets.
+	*
+	* @param InOutRotationOffset		Offset rotation to update
+	* @param InOutAngularVelocityOffset	Offset angular velocity to update
+	* @param InSrcRotation				Rotation of the source (rotation transitioning from)
+	* @param InSrcAngularVelocity		Angular velocity of the source (angular velocity transitioning from)
+	* @param InDstRotation				Rotation of the destination (rotation transitioning to)
+	* @param InDstAngularVelocity		Angular velocity of the destination (angular velocity transitioning to)
+	* @param TimeSinceTransition		Time since the last transition
+	* @param SmoothingTime				Timescale over which to apply the offset. Larger values result in longer blend times. Can be zero.
+	*/
+	UFUNCTION(BlueprintCallable, meta = (Category = "Math|Spring|Experimental"))
+	static ENGINE_API void SpringInertializeTransitionQuat(UPARAM(ref) FQuat& InOutRotationOffset, UPARAM(ref) FVector& InOutAngularVelocityOffset, FQuat InSrcRotation, FVector InSrcAngularVelocity, FQuat InDstRotation, FVector InDstAngularVelocity, float TimeSinceTransition, float SmoothingTime = 0.2f);
+
+	UE_EXPERIMENTAL(5.8, "SpringMath is experimental")
+	/** Applies a spring inertialization to a rotator, applying the appropriate offset given the time since transition.
+	*
+	* @param InRotation					The rotation to apply the inertialization to
+	* @param InRotationOffset			Offset rotation at the time of transition
+	* @param InAngularVelocityOffset	Offset angular velocity at the time of transition
+	* @param TimeSinceTransition		Time since the last transition
+	* @param SmoothingTime				Timescale over which to apply the offset. Larger values result in longer blend times. Can be zero.
+	* @return							The rotation with appropriate offset applied
+	*/
+	UFUNCTION(BlueprintPure, meta = (Category = "Math|Spring|Experimental"))
+	static ENGINE_API FRotator SpringInertializeApplyToRotator(FRotator InRotation, FRotator InRotationOffset, FVector InAngularVelocityOffset, float TimeSinceTransition, float SmoothingTime = 0.2f);
+
+	UE_EXPERIMENTAL(5.8, "SpringMath is experimental")
+	/** Transitions a spring inertialization, updating the appropriate offsets.
+	*
+	* @param InOutRotationOffset		Offset rotation to update
+	* @param InOutAngularVelocityOffset	Offset angular velocity to update
+	* @param InSrcRotation				Rotation of the source (rotation transitioning from)
+	* @param InSrcAngularVelocity		Angular velocity of the source (angular velocity transitioning from)
+	* @param InDstRotation				Rotation of the destination (rotation transitioning to)
+	* @param InDstAngularVelocity		Angular velocity of the destination (angular velocity transitioning to)
+	* @param TimeSinceTransition		Time since the last transition
+	* @param SmoothingTime				Timescale over which to apply the offset. Larger values result in longer blend times. Can be zero.
+	*/
+	UFUNCTION(BlueprintCallable, meta = (Category = "Math|Spring|Experimental"))
+	static ENGINE_API void SpringInertializeTransitionRotator(UPARAM(ref) FRotator& InOutRotationOffset, UPARAM(ref) FVector& InOutAngularVelocityOffset, FRotator InSrcRotation, FVector InSrcAngularVelocity, FRotator InDstRotation, FVector InDstAngularVelocity, float TimeSinceTransition, float SmoothingTime = 0.2f);
+
+	UE_EXPERIMENTAL(5.8, "SpringMath is experimental")
+	/** Applies a spring inertialization to a angle, adding the appropriate offset given the time since transition.
+	*
+	* @param InAngle					The angle to apply the inertialization to
+	* @param InAngleOffset				Offset angle at the time of transition
+	* @param InAngularVelocityOffset	Offset angular velocity at the time of transition
+	* @param TimeSinceTransition		Time since the last transition
+	* @param SmoothingTime				Timescale over which to apply the offset. Larger values result in longer blend times. Can be zero.
+	* @return							The angle with appropriate offset applied
+	*/
+	UFUNCTION(BlueprintPure, meta = (Category = "Math|Spring|Experimental"))
+	static ENGINE_API float SpringInertializeApplyToAngle(float InAngle, float InAngleOffset, float InAngularVelocityOffset, float TimeSinceTransition, float SmoothingTime = 0.2f);
+
+	UE_EXPERIMENTAL(5.8, "SpringMath is experimental")
+	/** Transitions a spring inertialization, updating the appropriate offsets.
+	*
+	* @param InOutAngleOffset			Offset angle to update
+	* @param InOutAngularVelocityOffset	Offset angular velocity to update
+	* @param InSrcAngle					Angle of the source (angle transitioning from)
+	* @param InSrcAngularVelocity		Angular velocity of the source (angular velocity transitioning from)
+	* @param InDstAngle					Angle of the destination (angle transitioning to)
+	* @param InDstAngularVelocity		Angular velocity of the destination (angular velocity transitioning to)
+	* @param TimeSinceTransition		Time since the last transition
+	* @param SmoothingTime				Timescale over which to apply the offset. Larger values result in longer blend times. Can be zero.
+	*/
+	UFUNCTION(BlueprintCallable, meta = (Category = "Math|Spring|Experimental"))
+	static ENGINE_API void SpringInertializeTransitionAngle(UPARAM(ref) float& InOutAngleOffset, UPARAM(ref) float& InOutAngularVelocityOffset, float InSrcAngle, float InSrcAngularVelocity, float InDstAngle, float InDstAngularVelocity, float TimeSinceTransition, float SmoothingTime = 0.2f);
+
+	UE_EXPERIMENTAL(5.8, "SpringMath is experimental")
+	/** Applies a spring inertialization to a scale, applying the appropriate offset given the time since transition.
+	*
+	* @param InScale					The scale to apply the inertialization to
+	* @param InScaleOffset				Offset scale at the time of transition
+	* @param InScalarVelocityOffset		Offset scalar velocity at the time of transition
+	* @param TimeSinceTransition		Time since the last transition
+	* @param SmoothingTime				Timescale over which to apply the offset. Larger values result in longer blend times. Can be zero.
+	* @return							The scale with appropriate offset applied
+	*/
+	UFUNCTION(BlueprintPure, meta = (Category = "Math|Spring|Experimental"))
+	static ENGINE_API FVector SpringInertializeApplyToScale(FVector InScale, FVector InScaleOffset, FVector InScalarVelocityOffset, float TimeSinceTransition, float SmoothingTime = 0.2f);
+
+	UE_EXPERIMENTAL(5.8, "SpringMath is experimental")
+	/** Transitions a spring inertialization, updating the appropriate offsets.
+	*
+	* @param InOutScaleOffset			Offset scale to update
+	* @param InOutScalarVelocityOffset	Offset scalar velocity to update
+	* @param InSrcScale					Scale of the source (scale transitioning from)
+	* @param InSrcScalarVelocity		Scalar velocity of the source (scalar velocity transitioning from)
+	* @param InDstScale					Scale of the destination (scale transitioning to)
+	* @param InDstScalarVelocity		Scalar velocity of the destination (scalar velocity transitioning to)
+	* @param TimeSinceTransition		Time since the last transition
+	* @param SmoothingTime				Timescale over which to apply the offset. Larger values result in longer blend times. Can be zero.
+	*/
+	UFUNCTION(BlueprintCallable, meta = (Category = "Math|Spring|Experimental"))
+	static ENGINE_API void SpringInertializeTransitionScale(UPARAM(ref) FVector& InOutScaleOffset, UPARAM(ref) FVector& InOutScalarVelocityOffset, FVector InSrcScale, FVector InSrcScalarVelocity, FVector InDstScale, FVector InDstScalarVelocity, float TimeSinceTransition, float SmoothingTime = 0.2f);
+
+	UE_EXPERIMENTAL(5.8, "SpringMath is experimental")
+	/** Applies a spring inertialization to a transform, applying the appropriate offset given the time since transition.
+	*
+	* @param InTransform				The transform to apply the inertialization to
+	* @param InTransformOffset			Offset transform at the time of transition
+	* @param InVelocityOffset			Offset velocity at the time of transition
+	* @param TimeSinceTransition		Time since the last transition
+	* @param SmoothingTime				Timescale over which to apply the offset. Larger values result in longer blend times. Can be zero.
+	* @return							The transform with appropriate offset applied
+	*/
+	UFUNCTION(BlueprintPure, meta = (Category = "Math|Spring|Experimental"))
+	static ENGINE_API FTransform SpringInertializeApplyToTransform(FTransform InTransform, FTransform InTransformOffset, FSpringTransformVelocity InVelocityOffset, float TimeSinceTransition, float SmoothingTime = 0.2f);
+
+	UE_EXPERIMENTAL(5.8, "SpringMath is experimental")
+	/** Transitions a spring inertialization, updating the appropriate offsets.
+	*
+	* @param InOutTransformOffset		Offset transform to update
+	* @param InOutVelocityOffset		Offset velocity to update
+	* @param InSrcTransform				Transform of the source (transform transitioning from)
+	* @param InSrcVelocity				Velocity of the source (velocity transitioning from)
+	* @param InDstTransform				Transform of the destination (transform transitioning to)
+	* @param InDstVelocity				Velocity of the destination (velocity transitioning to)
+	* @param TimeSinceTransition		Time since the last transition
+	* @param SmoothingTime				Timescale over which to apply the offset. Larger values result in longer blend times. Can be zero.
+	*/
+	UFUNCTION(BlueprintCallable, meta = (Category = "Math|Spring|Experimental"))
+	static ENGINE_API void SpringInertializeTransitionTransform(UPARAM(ref) FTransform& InOutTransformOffset, UPARAM(ref) FSpringTransformVelocity& InOutVelocityOffset, FTransform InSrcTransform, FSpringTransformVelocity InSrcVelocity, FTransform InDstTransform, FSpringTransformVelocity InDstVelocity, float TimeSinceTransition, float SmoothingTime = 0.2f);
+
+	// Cubic Inertialization //
+
+	UE_EXPERIMENTAL(5.8, "SpringMath is experimental")
+	/** Applies a cubic inertialization to a float value, adding the appropriate offset given the time since transition.
+	*
+	* @param InValue				The value to apply the inertialization to
+	* @param InValueOffset			Offset value at the time of transition
+	* @param InVelocityOffset		Offset velocity at the time of transition
+	* @param TimeSinceTransition	Time since the last transition
+	* @param BlendTime				Timescale over which to apply the offset. Larger values result in longer blend times. Can be zero.
+	* @return						The value with appropriate offset applied
+	*/
+	UFUNCTION(BlueprintPure, meta = (Category = "Math|Spring|Experimental"))
+	static ENGINE_API float CubicInertializeApplyToFloat(float InValue, float InValueOffset, float InVelocityOffset, float TimeSinceTransition, float BlendTime = 0.2f);
+
+	UE_EXPERIMENTAL(5.8, "SpringMath is experimental")
+	/** Transitions a cubic inertialization, updating the appropriate offsets.
+	*
+	* @param InOutValueOffset		Offset value to update
+	* @param InOutVelocityOffset	Offset velocity to update
+	* @param InSrcValue				Value of the source (value transitioning from)
+	* @param InSrcVelocity			Velocity of the source (velocity transitioning from)
+	* @param InDstValue				Value of the destination (value transitioning to)
+	* @param InDstVelocity			Velocity of the destination (velocity transitioning to)
+	* @param TimeSinceTransition	Time since the last transition
+	* @param BlendTime				Timescale over which to apply the offset. Larger values result in longer blend times. Can be zero.
+	*/
+	UFUNCTION(BlueprintCallable, meta = (Category = "Math|Spring|Experimental"))
+	static ENGINE_API void CubicInertializeTransitionFloat(UPARAM(ref) float& InOutValueOffset, UPARAM(ref) float& InOutVelocityOffset, float InSrcValue, float InSrcVelocity, float InDstValue, float InDstVelocity, float TimeSinceTransition, float BlendTime = 0.2f);
+
+	UE_EXPERIMENTAL(5.8, "SpringMath is experimental")
+	/** Applies a cubic inertialization to a vector value, adding the appropriate offset given the time since transition.
+	*
+	* @param InValue				The value to apply the inertialization to
+	* @param InValueOffset			Offset value at the time of transition
+	* @param InVelocityOffset		Offset velocity at the time of transition
+	* @param TimeSinceTransition	Time since the last transition
+	* @param BlendTime				Timescale over which to apply the offset. Larger values result in longer blend times. Can be zero.
+	* @return						The value with appropriate offset applied
+	*/
+	UFUNCTION(BlueprintPure, meta = (Category = "Math|Spring|Experimental"))
+	static ENGINE_API FVector CubicInertializeApplyToVector(FVector InValue, FVector InValueOffset, FVector InVelocityOffset, float TimeSinceTransition, float BlendTime = 0.2f);
+
+	UE_EXPERIMENTAL(5.8, "SpringMath is experimental")
+	/** Transitions a cubic inertialization, updating the appropriate offsets.
+	*
+	* @param InOutValueOffset		Offset value to update
+	* @param InOutVelocityOffset	Offset velocity to update
+	* @param InSrcValue				Value of the source (value transitioning from)
+	* @param InSrcVelocity			Velocity of the source (velocity transitioning from)
+	* @param InDstValue				Value of the destination (value transitioning to)
+	* @param InDstVelocity			Velocity of the destination (velocity transitioning to)
+	* @param TimeSinceTransition	Time since the last transition
+	* @param BlendTime				Timescale over which to apply the offset. Larger values result in longer blend times. Can be zero.
+	*/
+	UFUNCTION(BlueprintCallable, meta = (Category = "Math|Spring|Experimental"))
+	static ENGINE_API void CubicInertializeTransitionVector(UPARAM(ref) FVector& InOutValueOffset, UPARAM(ref) FVector& InOutVelocityOffset, FVector InSrcValue, FVector InSrcVelocity, FVector InDstValue, FVector InDstVelocity, float TimeSinceTransition, float BlendTime = 0.2f);
+
+	UE_EXPERIMENTAL(5.8, "SpringMath is experimental")
+	/** Applies a cubic inertialization to a Vector2D value, adding the appropriate offset given the time since transition.
+	*
+	* @param InValue				The value to apply the inertialization to
+	* @param InValueOffset			Offset value at the time of transition
+	* @param InVelocityOffset		Offset velocity at the time of transition
+	* @param TimeSinceTransition	Time since the last transition
+	* @param BlendTime				Timescale over which to apply the offset. Larger values result in longer blend times. Can be zero.
+	* @return						The value with appropriate offset applied
+	*/
+	UFUNCTION(BlueprintPure, meta = (Category = "Math|Spring|Experimental"))
+	static ENGINE_API FVector2D CubicInertializeApplyToVector2D(FVector2D InValue, FVector2D InValueOffset, FVector2D InVelocityOffset, float TimeSinceTransition, float BlendTime = 0.2f);
+
+	UE_EXPERIMENTAL(5.8, "SpringMath is experimental")
+	/** Transitions a cubic inertialization, updating the appropriate offsets.
+	*
+	* @param InOutValueOffset		Offset value to update
+	* @param InOutVelocityOffset	Offset velocity to update
+	* @param InSrcValue				Value of the source (value transitioning from)
+	* @param InSrcVelocity			Velocity of the source (velocity transitioning from)
+	* @param InDstValue				Value of the destination (value transitioning to)
+	* @param InDstVelocity			Velocity of the destination (velocity transitioning to)
+	* @param TimeSinceTransition	Time since the last transition
+	* @param BlendTime				Timescale over which to apply the offset. Larger values result in longer blend times. Can be zero.
+	*/
+	UFUNCTION(BlueprintCallable, meta = (Category = "Math|Spring|Experimental"))
+	static ENGINE_API void CubicInertializeTransitionVector2D(UPARAM(ref) FVector2D& InOutValueOffset, UPARAM(ref) FVector2D& InOutVelocityOffset, FVector2D InSrcValue, FVector2D InSrcVelocity, FVector2D InDstValue, FVector2D InDstVelocity, float TimeSinceTransition, float BlendTime = 0.2f);
+
+	UE_EXPERIMENTAL(5.8, "SpringMath is experimental")
+	/** Applies a cubic inertialization to a quaternion, applying the appropriate offset given the time since transition.
+	*
+	* @param InRotation					The rotation to apply the inertialization to
+	* @param InRotationOffset			Offset rotation at the time of transition
+	* @param InAngularVelocityOffset	Offset angular velocity at the time of transition
+	* @param TimeSinceTransition		Time since the last transition
+	* @param BlendTime					Timescale over which to apply the offset. Larger values result in longer blend times. Can be zero.
+	* @return							The rotation with appropriate offset applied
+	*/
+	UFUNCTION(BlueprintPure, meta = (Category = "Math|Spring|Experimental"))
+	static ENGINE_API FQuat CubicInertializeApplyToQuat(FQuat InRotation, FQuat InRotationOffset, FVector InAngularVelocityOffset, float TimeSinceTransition, float BlendTime = 0.2f);
+
+	UE_EXPERIMENTAL(5.8, "SpringMath is experimental")
+	/** Transitions a cubic inertialization, updating the appropriate offsets.
+	*
+	* @param InOutRotationOffset		Offset rotation to update
+	* @param InOutAngularVelocityOffset	Offset angular velocity to update
+	* @param InSrcRotation				Rotation of the source (rotation transitioning from)
+	* @param InSrcAngularVelocity		Angular velocity of the source (angular velocity transitioning from)
+	* @param InDstRotation				Rotation of the destination (rotation transitioning to)
+	* @param InDstAngularVelocity		Angular velocity of the destination (angular velocity transitioning to)
+	* @param TimeSinceTransition		Time since the last transition
+	* @param BlendTime					Timescale over which to apply the offset. Larger values result in longer blend times. Can be zero.
+	*/
+	UFUNCTION(BlueprintCallable, meta = (Category = "Math|Spring|Experimental"))
+	static ENGINE_API void CubicInertializeTransitionQuat(UPARAM(ref) FQuat& InOutRotationOffset, UPARAM(ref) FVector& InOutAngularVelocityOffset, FQuat InSrcRotation, FVector InSrcAngularVelocity, FQuat InDstRotation, FVector InDstAngularVelocity, float TimeSinceTransition, float BlendTime = 0.2f);
+
+	UE_EXPERIMENTAL(5.8, "SpringMath is experimental")
+	/** Applies a cubic inertialization to a rotator, applying the appropriate offset given the time since transition.
+	*
+	* @param InRotation					The rotation to apply the inertialization to
+	* @param InRotationOffset			Offset rotation at the time of transition
+	* @param InAngularVelocityOffset	Offset angular velocity at the time of transition
+	* @param TimeSinceTransition		Time since the last transition
+	* @param BlendTime					Timescale over which to apply the offset. Larger values result in longer blend times. Can be zero.
+	* @return							The rotation with appropriate offset applied
+	*/
+	UFUNCTION(BlueprintPure, meta = (Category = "Math|Spring|Experimental"))
+	static ENGINE_API FRotator CubicInertializeApplyToRotator(FRotator InRotation, FRotator InRotationOffset, FVector InAngularVelocityOffset, float TimeSinceTransition, float BlendTime = 0.2f);
+
+	UE_EXPERIMENTAL(5.8, "SpringMath is experimental")
+	/** Transitions a cubic inertialization, updating the appropriate offsets.
+	*
+	* @param InOutRotationOffset		Offset rotation to update
+	* @param InOutAngularVelocityOffset	Offset angular velocity to update
+	* @param InSrcRotation				Rotation of the source (rotation transitioning from)
+	* @param InSrcAngularVelocity		Angular velocity of the source (angular velocity transitioning from)
+	* @param InDstRotation				Rotation of the destination (rotation transitioning to)
+	* @param InDstAngularVelocity		Angular velocity of the destination (angular velocity transitioning to)
+	* @param TimeSinceTransition		Time since the last transition
+	* @param BlendTime					Timescale over which to apply the offset. Larger values result in longer blend times. Can be zero.
+	*/
+	UFUNCTION(BlueprintCallable, meta = (Category = "Math|Spring|Experimental"))
+	static ENGINE_API void CubicInertializeTransitionRotator(UPARAM(ref) FRotator& InOutRotationOffset, UPARAM(ref) FVector& InOutAngularVelocityOffset, FRotator InSrcRotation, FVector InSrcAngularVelocity, FRotator InDstRotation, FVector InDstAngularVelocity, float TimeSinceTransition, float BlendTime = 0.2f);
+
+	UE_EXPERIMENTAL(5.8, "SpringMath is experimental")
+	/** Applies a cubic inertialization to an angle, applying the appropriate offset given the time since transition.
+	*
+	* @param InAngle					The angle to apply the inertialization to
+	* @param InAngleOffset				Offset angle at the time of transition
+	* @param InAngularVelocityOffset	Offset angular velocity at the time of transition
+	* @param TimeSinceTransition		Time since the last transition
+	* @param BlendTime					Timescale over which to apply the offset. Larger values result in longer blend times. Can be zero.
+	* @return							The angle with appropriate offset applied
+	*/
+	UFUNCTION(BlueprintPure, meta = (Category = "Math|Spring|Experimental"))
+	static ENGINE_API float CubicInertializeApplyToAngle(float InAngle, float InAngleOffset, float InAngularVelocityOffset, float TimeSinceTransition, float BlendTime = 0.2f);
+
+	UE_EXPERIMENTAL(5.8, "SpringMath is experimental")
+	/** Transitions a cubic inertialization, updating the appropriate offsets.
+	*
+	* @param InOutAngleOffset			Offset angle to update
+	* @param InOutAngularVelocityOffset	Offset angular velocity to update
+	* @param InSrcAngle					Angle of the source (angle transitioning from)
+	* @param InSrcAngularVelocity		Angular velocity of the source (angular velocity transitioning from)
+	* @param InDstAngle					Angle of the destination (angle transitioning to)
+	* @param InDstAngularVelocity		Angular velocity of the destination (angular velocity transitioning to)
+	* @param TimeSinceTransition		Time since the last transition
+	* @param BlendTime					Timescale over which to apply the offset. Larger values result in longer blend times. Can be zero.
+	*/
+	UFUNCTION(BlueprintCallable, meta = (Category = "Math|Spring|Experimental"))
+	static ENGINE_API void CubicInertializeTransitionAngle(UPARAM(ref) float& InOutAngleOffset, UPARAM(ref) float& InOutAngularVelocityOffset, float InSrcAngle, float InSrcAngularVelocity, float InDstAngle, float InDstAngularVelocity, float TimeSinceTransition, float BlendTime = 0.2f);
+
+	UE_EXPERIMENTAL(5.8, "SpringMath is experimental")
+	/** Applies a cubic inertialization to a scale, applying the appropriate offset given the time since transition.
+	*
+	* @param InScale					The scale to apply the inertialization to
+	* @param InScaleOffset				Offset scale at the time of transition
+	* @param InScalarVelocityOffset		Offset scalar velocity at the time of transition
+	* @param TimeSinceTransition		Time since the last transition
+	* @param BlendTime					Timescale over which to apply the offset. Larger values result in longer blend times. Can be zero.
+	* @return							The scale with appropriate offset applied
+	*/
+	UFUNCTION(BlueprintPure, meta = (Category = "Math|Spring|Experimental"))
+	static ENGINE_API FVector CubicInertializeApplyToScale(FVector InScale, FVector InScaleOffset, FVector InScalarVelocityOffset, float TimeSinceTransition, float BlendTime = 0.2f);
+
+	UE_EXPERIMENTAL(5.8, "SpringMath is experimental")
+	/** Transitions a cubic inertialization, updating the appropriate offsets.
+	*
+	* @param InOutScaleOffset			Offset scale to update
+	* @param InOutScalarVelocityOffset	Offset scalar velocity to update
+	* @param InSrcScale					Scale of the source (scale transitioning from)
+	* @param InSrcScalarVelocity		Scalar velocity of the source (scalar velocity transitioning from)
+	* @param InDstScale					Scale of the destination (scale transitioning to)
+	* @param InDstScalarVelocity		Scalar velocity of the destination (scalar velocity transitioning to)
+	* @param TimeSinceTransition		Time since the last transition
+	* @param BlendTime					Timescale over which to apply the offset. Larger values result in longer blend times. Can be zero.
+	*/
+	UFUNCTION(BlueprintCallable, meta = (Category = "Math|Spring|Experimental"))
+	static ENGINE_API void CubicInertializeTransitionScale(UPARAM(ref) FVector& InOutScaleOffset, UPARAM(ref) FVector& InOutScalarVelocityOffset, FVector InSrcScale, FVector InSrcScalarVelocity, FVector InDstScale, FVector InDstScalarVelocity, float TimeSinceTransition, float BlendTime = 0.2f);
+
+	UE_EXPERIMENTAL(5.8, "SpringMath is experimental")
+	/** Applies a cubic inertialization to a transform, applying the appropriate offset given the time since transition.
+	*
+	* @param InTransform				The transform to apply the inertialization to
+	* @param InTransformOffset			Offset transform at the time of transition
+	* @param InVelocityOffset			Offset velocity at the time of transition
+	* @param TimeSinceTransition		Time since the last transition
+	* @param BlendTime					Timescale over which to apply the offset. Larger values result in longer blend times. Can be zero.
+	* @return							The transform with appropriate offset applied
+	*/
+	UFUNCTION(BlueprintPure, meta = (Category = "Math|Spring|Experimental"))
+	static ENGINE_API FTransform CubicInertializeApplyToTransform(FTransform InTransform, FTransform InTransformOffset, FSpringTransformVelocity InVelocityOffset, float TimeSinceTransition, float BlendTime = 0.2f);
+
+	UE_EXPERIMENTAL(5.8, "SpringMath is experimental")
+	/** Transitions a cubic inertialization, updating the appropriate offsets.
+	*
+	* @param InOutTransformOffset		Offset transform to update
+	* @param InOutVelocityOffset		Offset velocity to update
+	* @param InSrcTransform				Transform of the source (transform transitioning from)
+	* @param InSrcVelocity				Velocity of the source (velocity transitioning from)
+	* @param InDstTransform				Transform of the destination (transform transitioning to)
+	* @param InDstVelocity				Velocity of the destination (velocity transitioning to)
+	* @param TimeSinceTransition		Time since the last transition
+	* @param BlendTime					Timescale over which to apply the offset. Larger values result in longer blend times. Can be zero.
+	*/
+	UFUNCTION(BlueprintCallable, meta = (Category = "Math|Spring|Experimental"))
+	static ENGINE_API void CubicInertializeTransitionTransform(UPARAM(ref) FTransform& InOutTransformOffset, UPARAM(ref) FSpringTransformVelocity& InOutVelocityOffset, FTransform InSrcTransform, FSpringTransformVelocity InSrcVelocity, FTransform InDstTransform, FSpringTransformVelocity InDstVelocity, float TimeSinceTransition, float BlendTime = 0.2f);
+
+	// Dead Blending //
+
+	UE_EXPERIMENTAL(5.8, "SpringMath is experimental")
+	/** Applies a dead blend to a float value, blending it with an extrapolation of the value and velocity at the point of the last transition.
+	*
+	* @param InValue					The value to apply the dead blend to
+	* @param InValueTransition			The value at the time of the last transition
+	* @param InVelocityTransition		The velocity at the time of the last transition
+	* @param TimeSinceTransition		Time since the last transition
+	* @param BlendTime					Timescale over which to apply the offset. Larger values result in longer blend times. Can be zero.
+	* @param SmoothingTime				Decay rate of the extrapolated velocity. Smaller values make the velocity decay faster.
+	* @return							The blended value
+	*/
+	UFUNCTION(BlueprintPure, meta = (Category = "Math|Spring|Experimental"))
+	static ENGINE_API float DeadBlendApplyToFloat(float InValue, float InValueTransition, float InVelocityTransition, float TimeSinceTransition, float BlendTime = 0.2f, float SmoothingTime = 0.2f);
+
+	UE_EXPERIMENTAL(5.8, "SpringMath is experimental")
+	/** Transitions a dead blend, computing the new value and velocity at the point of transition.
+	*
+	* @param InOutValueTransition		Value at the point of transition
+	* @param InOutVelocityTransition	Velocity at the point of transition
+	* @param InSrcValue					Value of the source (value transitioning from)
+	* @param InSrcVelocity				Velocity of the source (velocity transitioning from)
+	* @param TimeSinceTransition		Time since the last transition
+	* @param BlendTime					Timescale over which to apply the offset. Larger values result in longer blend times. Can be zero.
+	* @param SmoothingTime				Decay rate of the extrapolated velocity. Smaller values make the velocity decay faster.
+	*/
+	UFUNCTION(BlueprintCallable, meta = (Category = "Math|Spring|Experimental"))
+	static ENGINE_API void DeadBlendTransitionFloat(UPARAM(ref) float& InOutValueTransition, UPARAM(ref) float& InOutVelocityTransition, float InSrcValue, float InSrcVelocity, float TimeSinceTransition, float BlendTime = 0.2f, float SmoothingTime = 0.2f);
+
+	UE_EXPERIMENTAL(5.8, "SpringMath is experimental")
+	/** Applies a dead blend to a vector value, blending it with an extrapolation of the value and velocity at the point of the last transition.
+	*
+	* @param InValue					The value to apply the dead blend to
+	* @param InValueTransition			The value at the time of the last transition
+	* @param InVelocityTransition		The velocity at the time of the last transition
+	* @param TimeSinceTransition		Time since the last transition
+	* @param BlendTime					Timescale over which to apply the offset. Larger values result in longer blend times. Can be zero.
+	* @param SmoothingTime				Decay rate of the extrapolated velocity. Smaller values make the velocity decay faster.
+	* @return							The blended value
+	*/
+	UFUNCTION(BlueprintPure, meta = (Category = "Math|Spring|Experimental"))
+	static ENGINE_API FVector DeadBlendApplyToVector(FVector InValue, FVector InValueTransition, FVector InVelocityTransition, float TimeSinceTransition, float BlendTime = 0.2f, float SmoothingTime = 0.2f);
+
+	UE_EXPERIMENTAL(5.8, "SpringMath is experimental")
+	/** Transitions a dead blend, computing the new value and velocity at the point of transition.
+	*
+	* @param InOutValueTransition		Value at the point of transition
+	* @param InOutVelocityTransition	Velocity at the point of transition
+	* @param InSrcValue					Value of the source (value transitioning from)
+	* @param InSrcVelocity				Velocity of the source (velocity transitioning from)
+	* @param TimeSinceTransition		Time since the last transition
+	* @param BlendTime					Timescale over which to apply the offset. Larger values result in longer blend times. Can be zero.
+	* @param SmoothingTime				Decay rate of the extrapolated velocity. Smaller values make the velocity decay faster.
+	*/
+	UFUNCTION(BlueprintCallable, meta = (Category = "Math|Spring|Experimental"))
+	static ENGINE_API void DeadBlendTransitionVector(UPARAM(ref) FVector& InOutValueTransition, UPARAM(ref) FVector& InOutVelocityTransition, FVector InSrcValue, FVector InSrcVelocity, float TimeSinceTransition, float BlendTime = 0.2f, float SmoothingTime = 0.2f);
+
+	UE_EXPERIMENTAL(5.8, "SpringMath is experimental")
+	/** Applies a dead blend to a Vector2D value, blending it with an extrapolation of the value and velocity at the point of the last transition.
+	*
+	* @param InValue					The value to apply the dead blend to
+	* @param InValueTransition			The value at the time of the last transition
+	* @param InVelocityTransition		The velocity at the time of the last transition
+	* @param TimeSinceTransition		Time since the last transition
+	* @param BlendTime					Timescale over which to apply the offset. Larger values result in longer blend times. Can be zero.
+	* @param SmoothingTime				Decay rate of the extrapolated velocity. Smaller values make the velocity decay faster.
+	* @return							The blended value
+	*/
+	UFUNCTION(BlueprintPure, meta = (Category = "Math|Spring|Experimental"))
+	static ENGINE_API FVector2D DeadBlendApplyToVector2D(FVector2D InValue, FVector2D InValueTransition, FVector2D InVelocityTransition, float TimeSinceTransition, float BlendTime = 0.2f, float SmoothingTime = 0.2f);
+
+	UE_EXPERIMENTAL(5.8, "SpringMath is experimental")
+	/** Transitions a dead blend, computing the new value and velocity at the point of transition.
+	*
+	* @param InOutValueTransition		Value at the point of transition
+	* @param InOutVelocityTransition	Velocity at the point of transition
+	* @param InSrcValue					Value of the source (value transitioning from)
+	* @param InSrcVelocity				Velocity of the source (velocity transitioning from)
+	* @param TimeSinceTransition		Time since the last transition
+	* @param BlendTime					Timescale over which to apply the offset. Larger values result in longer blend times. Can be zero.
+	* @param SmoothingTime				Decay rate of the extrapolated velocity. Smaller values make the velocity decay faster.
+	*/
+	UFUNCTION(BlueprintCallable, meta = (Category = "Math|Spring|Experimental"))
+	static ENGINE_API void DeadBlendTransitionVector2D(UPARAM(ref) FVector2D& InOutValueTransition, UPARAM(ref) FVector2D& InOutVelocityTransition, FVector2D InSrcValue, FVector2D InSrcVelocity, float TimeSinceTransition, float BlendTime = 0.2f, float SmoothingTime = 0.2f);
+
+	UE_EXPERIMENTAL(5.8, "SpringMath is experimental")
+	/** Applies a dead blend to a quaternion, blending it with an extrapolation of the rotation and velocity at the point of the last transition.
+	*
+	* @param InRotation						The rotation to apply the dead blend to
+	* @param InRotationTransition			The rotation at the time of the last transition
+	* @param InAngularVelocityTransition	The angular velocity at the time of the last transition
+	* @param TimeSinceTransition			Time since the last transition
+	* @param BlendTime						Timescale over which to apply the offset. Larger values result in longer blend times. Can be zero.
+	* @param SmoothingTime					Decay rate of the extrapolated velocity. Smaller values make the velocity decay faster.
+	* @return								The blended rotation
+	*/
+	UFUNCTION(BlueprintPure, meta = (Category = "Math|Spring|Experimental"))
+	static ENGINE_API FQuat DeadBlendApplyToQuat(FQuat InRotation, FQuat InRotationTransition, FVector InAngularVelocityTransition, float TimeSinceTransition, float BlendTime = 0.2f, float SmoothingTime = 0.2f);
+
+	UE_EXPERIMENTAL(5.8, "SpringMath is experimental")
+	/** Transitions a dead blend, computing the new value and velocity at the point of transition.
+	*
+	* @param InOutRotationTransition		Rotation at the point of transition
+	* @param InOutAngularVelocityTransition	Angular velocity at the point of transition
+	* @param InSrcRotation					Rotation of the source (rotation transitioning from)
+	* @param InSrcAngularVelocity			Angular velocity of the source (angular velocity transitioning from)
+	* @param TimeSinceTransition			Time since the last transition
+	* @param BlendTime						Timescale over which to apply the offset. Larger values result in longer blend times. Can be zero.
+	* @param SmoothingTime					Decay rate of the extrapolated velocity. Smaller values make the velocity decay faster.
+	*/
+	UFUNCTION(BlueprintCallable, meta = (Category = "Math|Spring|Experimental"))
+	static ENGINE_API void DeadBlendTransitionQuat(UPARAM(ref) FQuat& InOutRotationTransition, UPARAM(ref) FVector& InOutAngularVelocityTransition, FQuat InSrcRotation, FVector InSrcAngularVelocity, float TimeSinceTransition, float BlendTime = 0.2f, float SmoothingTime = 0.2f);
+
+	UE_EXPERIMENTAL(5.8, "SpringMath is experimental")
+	/** Applies a dead blend to a rotator, blending it with an extrapolation of the rotation and velocity at the point of the last transition.
+	*
+	* @param InRotation						The rotation to apply the dead blend to
+	* @param InRotationTransition			The rotation at the time of the last transition
+	* @param InAngularVelocityTransition	The angular velocity at the time of the last transition
+	* @param TimeSinceTransition			Time since the last transition
+	* @param BlendTime						Timescale over which to apply the offset. Larger values result in longer blend times. Can be zero.
+	* @param SmoothingTime					Decay rate of the extrapolated velocity. Smaller values make the velocity decay faster.
+	* @return								The blended rotation
+	*/
+	UFUNCTION(BlueprintPure, meta = (Category = "Math|Spring|Experimental"))
+	static ENGINE_API FRotator DeadBlendApplyToRotator(FRotator InRotation, FRotator InRotationTransition, FVector InAngularVelocityTransition, float TimeSinceTransition, float BlendTime = 0.2f, float SmoothingTime = 0.2f);
+
+	UE_EXPERIMENTAL(5.8, "SpringMath is experimental")
+	/** Transitions a dead blend, computing the new value and velocity at the point of transition.
+	*
+	* @param InOutRotationTransition		Rotation at the point of transition
+	* @param InOutAngularVelocityTransition	Angular velocity at the point of transition
+	* @param InSrcRotation					Rotation of the source (rotation transitioning from)
+	* @param InSrcAngularVelocity			Angular velocity of the source (angular velocity transitioning from)
+	* @param TimeSinceTransition			Time since the last transition
+	* @param BlendTime						Timescale over which to apply the offset. Larger values result in longer blend times. Can be zero.
+	* @param SmoothingTime					Decay rate of the extrapolated velocity. Smaller values make the velocity decay faster.
+	*/
+	UFUNCTION(BlueprintCallable, meta = (Category = "Math|Spring|Experimental"))
+	static ENGINE_API void DeadBlendTransitionRotator(UPARAM(ref) FRotator& InOutRotationTransition, UPARAM(ref) FVector& InOutAngularVelocityTransition, FRotator InSrcRotation, FVector InSrcAngularVelocity, float TimeSinceTransition, float BlendTime = 0.2f, float SmoothingTime = 0.2f);
+
+	UE_EXPERIMENTAL(5.8, "SpringMath is experimental")
+	/** Applies a dead blend to an angle, blending it with an extrapolation of the angle and velocity at the point of the last transition.
+	*
+	* @param InAngle						The angle to apply the dead blend to
+	* @param InAngleTransition				The angle at the time of the last transition
+	* @param InAngularVelocityTransition	The angular velocity at the time of the last transition
+	* @param TimeSinceTransition			Time since the last transition
+	* @param BlendTime						Timescale over which to apply the offset. Larger values result in longer blend times. Can be zero.
+	* @param SmoothingTime					Decay rate of the extrapolated velocity. Smaller values make the velocity decay faster.
+	* @return								The blended angle
+	*/
+	UFUNCTION(BlueprintPure, meta = (Category = "Math|Spring|Experimental"))
+	static ENGINE_API float DeadBlendApplyToAngle(float InAngle, float InAngleTransition, float InAngularVelocityTransition, float TimeSinceTransition, float BlendTime = 0.2f, float SmoothingTime = 0.2f);
+
+	UE_EXPERIMENTAL(5.8, "SpringMath is experimental")
+	/** Transitions a dead blend, computing the new value and velocity at the point of transition.
+	*
+	* @param InOutAngleTransition			Angle at the point of transition
+	* @param InOutAngularVelocityTransition	Angular velocity at the point of transition
+	* @param InSrcAngle						Angle of the source (angle transitioning from)
+	* @param InSrcAngularVelocity			Angular velocity of the source (angular velocity transitioning from)
+	* @param TimeSinceTransition			Time since the last transition
+	* @param BlendTime						Timescale over which to apply the offset. Larger values result in longer blend times. Can be zero.
+	* @param SmoothingTime					Decay rate of the extrapolated velocity. Smaller values make the velocity decay faster.
+	*/
+	UFUNCTION(BlueprintCallable, meta = (Category = "Math|Spring|Experimental"))
+	static ENGINE_API void DeadBlendTransitionAngle(UPARAM(ref) float& InOutAngleTransition, UPARAM(ref) float& InOutAngularVelocityTransition, float InSrcAngle, float InSrcAngularVelocity, float TimeSinceTransition, float BlendTime = 0.2f, float SmoothingTime = 0.2f);
+
+	UE_EXPERIMENTAL(5.8, "SpringMath is experimental")
+	/** Applies a dead blend to a scale, blending it with an extrapolation of the scale and velocity at the point of the last transition.
+	*
+	* @param InScale						The scale to apply the dead blend to
+	* @param InScaleTransition				The scale at the time of the last transition
+	* @param InScalarVelocityTransition		The scalar velocity at the time of the last transition
+	* @param TimeSinceTransition			Time since the last transition
+	* @param BlendTime						Timescale over which to apply the offset. Larger values result in longer blend times. Can be zero.
+	* @param SmoothingTime					Decay rate of the extrapolated velocity. Smaller values make the velocity decay faster.
+	* @return								The blended scale
+	*/
+	UFUNCTION(BlueprintPure, meta = (Category = "Math|Spring|Experimental"))
+	static ENGINE_API FVector DeadBlendApplyToScale(FVector InScale, FVector InScaleTransition, FVector InScalarVelocityTransition, float TimeSinceTransition, float BlendTime = 0.2f, float SmoothingTime = 0.2f);
+
+	UE_EXPERIMENTAL(5.8, "SpringMath is experimental")
+	/** Transitions a dead blend, computing the new value and velocity at the point of transition.
+	*
+	* @param InOutScaleTransition			Scale at the point of transition
+	* @param InOutScalarVelocityTransition	Scalar velocity at the point of transition
+	* @param InSrcScale						Scale of the source (scale transitioning from)
+	* @param InSrcScalarVelocity			Scalar velocity of the source (scalar velocity transitioning from)
+	* @param TimeSinceTransition			Time since the last transition
+	* @param BlendTime						Timescale over which to apply the offset. Larger values result in longer blend times. Can be zero.
+	* @param SmoothingTime					Decay rate of the extrapolated velocity. Smaller values make the velocity decay faster.
+	*/
+	UFUNCTION(BlueprintCallable, meta = (Category = "Math|Spring|Experimental"))
+	static ENGINE_API void DeadBlendTransitionScale(UPARAM(ref) FVector& InOutScaleTransition, UPARAM(ref) FVector& InOutScalarVelocityTransition, FVector InSrcScale, FVector InSrcScalarVelocity, float TimeSinceTransition, float BlendTime = 0.2f, float SmoothingTime = 0.2f);
+
+	UE_EXPERIMENTAL(5.8, "SpringMath is experimental")
+	/** Applies a dead blend to a transform, blending it with an extrapolation of the transform and velocity at the point of the last transition.
+	*
+	* @param InTransform					The transform to apply the dead blend to
+	* @param InTransformTransition			The transform at the time of the last transition
+	* @param InVelocityTransition			The velocity at the time of the last transition
+	* @param TimeSinceTransition			Time since the last transition
+	* @param BlendTime						Timescale over which to apply the offset. Larger values result in longer blend times. Can be zero.
+	* @param SmoothingTime					Decay rate of the extrapolated velocity. Smaller values make the velocity decay faster.
+	* @return								The blended transform
+	*/
+	UFUNCTION(BlueprintPure, meta = (Category = "Math|Spring|Experimental"))
+	static ENGINE_API FTransform DeadBlendApplyToTransform(FTransform InTransform, FTransform InTransformTransition, FSpringTransformVelocity InVelocityTransition, float TimeSinceTransition, float BlendTime = 0.2f, float SmoothingTime = 0.2f);
+
+	UE_EXPERIMENTAL(5.8, "SpringMath is experimental")
+	/** Transitions a dead blend, computing the new value and velocity at the point of transition.
+	*
+	* @param InOutTransformTransition	Transform at the point of transition
+	* @param InOutVelocityTransition	Velocity at the point of transition
+	* @param InSrcTransform				Transform of the source (transform transitioning from)
+	* @param InSrcVelocity				Velocity of the source (velocity transitioning from)
+	* @param TimeSinceTransition		Time since the last transition
+	* @param BlendTime					Timescale over which to apply the offset. Larger values result in longer blend times. Can be zero.
+	* @param SmoothingTime				Decay rate of the extrapolated velocity. Smaller values make the velocity decay faster.
+	*/
+	UFUNCTION(BlueprintCallable, meta = (Category = "Math|Spring|Experimental"))
+	static ENGINE_API void DeadBlendTransitionTransform(UPARAM(ref) FTransform& InOutTransformTransition, UPARAM(ref) FSpringTransformVelocity& InOutVelocityTransition, FTransform InSrcTransform, FSpringTransformVelocity InSrcVelocity, float TimeSinceTransition, float BlendTime = 0.2f, float SmoothingTime = 0.2f);
+
 };
diff --git a/Engine/Source/Runtime/Engine/Private/BlueprintSpringMathLibrary.cpp b/Engine/Source/Runtime/Engine/Private/BlueprintSpringMathLibrary.cpp
--- a/Engine/Source/Runtime/Engine/Private/BlueprintSpringMathLibrary.cpp
+++ b/Engine/Source/Runtime/Engine/Private/BlueprintSpringMathLibrary.cpp
@@ -30,16 +30,14 @@ void UBlueprintSpringMathLibrary::CriticalSpringDampAngle(float& InOutAngle, flo
 	InOutAngularVelocity = FMath::RadiansToDegrees(InOutAngularVelocityRadians);
 }
 
-void UBlueprintSpringMathLibrary::CriticalSpringDampQuat(FQuat& InOutRotation, FVector& InOutAngularVelocity, const FQuat& TargetRotation,
-	float DeltaTime, float SmoothingTime)
+void UBlueprintSpringMathLibrary::CriticalSpringDampQuat(FQuat& InOutRotation, FVector& InOutAngularVelocity, const FQuat& TargetRotation, float DeltaTime, float SmoothingTime)
 {
 	FVector InOutAngularVelocityRadians = FMath::DegreesToRadians(InOutAngularVelocity);
 	SpringMath::CriticalSpringDamperQuat(InOutRotation, InOutAngularVelocityRadians, TargetRotation, SmoothingTime, DeltaTime);
 	InOutAngularVelocity = FMath::RadiansToDegrees(InOutAngularVelocityRadians);
 }
 
-void UBlueprintSpringMathLibrary::CriticalSpringDampRotator(FRotator& InOutRotation, FVector& InOutAngularVelocity, const FRotator& TargetRotation,
-	float DeltaTime, float SmoothingTime)
+void UBlueprintSpringMathLibrary::CriticalSpringDampRotator(FRotator& InOutRotation, FVector& InOutAngularVelocity, const FRotator& TargetRotation, float DeltaTime, float SmoothingTime)
 {
 	FQuat InOutRotationQuat = InOutRotation.Quaternion();
 	FVector InOutAngularVelocityRadians = FMath::DegreesToRadians(InOutAngularVelocity);
@@ -48,6 +46,24 @@ void UBlueprintSpringMathLibrary::CriticalSpringDampRotator(FRotator& InOutRotat
 	InOutAngularVelocity = FMath::RadiansToDegrees(InOutAngularVelocityRadians);
 }
 
+void UBlueprintSpringMathLibrary::CriticalSpringDampScale(FVector& InOutScale, FVector& InOutScalarVelocity, const FVector& TargetScale, float DeltaTime, float SmoothingTime)
+{
+	SpringMath::CriticalSpringDamperScale(InOutScale, InOutScalarVelocity, TargetScale, SmoothingTime, DeltaTime);
+}
+
+void UBlueprintSpringMathLibrary::CriticalSpringDampTransform(FTransform& InOutTransform, FSpringTransformVelocity& InOutVelocity, const FTransform& TargetTransform, float DeltaTime, float SmoothingTime)
+{
+	FVector InOutLocation = InOutTransform.GetLocation();
+	FQuat InOutRotation = InOutTransform.GetRotation();
+	FVector InOutScale = InOutTransform.GetScale3D();
+
+	CriticalSpringDampVector(InOutLocation, InOutVelocity.LinearVelocity, TargetTransform.GetLocation(), DeltaTime, SmoothingTime);
+	CriticalSpringDampQuat(InOutRotation, InOutVelocity.AngularVelocity, TargetTransform.GetRotation(), DeltaTime, SmoothingTime);
+	CriticalSpringDampScale(InOutScale, InOutVelocity.ScalarVelocity, TargetTransform.GetScale3D(), DeltaTime, SmoothingTime);
+
+	InOutTransform = FTransform(InOutRotation, InOutLocation, InOutScale);
+}
+
 void UBlueprintSpringMathLibrary::CriticalDoubleSpringDampVector(FVector& InOutX, FVector& InOutV, FVector& InOutXi, FVector& InOutVi, const FVector& TargetX, float DeltaTime, float SmoothingTime)
 {
 	SpringMath::CriticalDoubleSpringDamper(InOutX, InOutV, InOutXi, InOutVi, TargetX, SmoothingTime, DeltaTime);
@@ -100,6 +116,28 @@ void UBlueprintSpringMathLibrary::CriticalDoubleSpringDampRotator(FRotator& InOu
 	InOutIntermediateAngularVelocity = FMath::RadiansToDegrees(InOutIntermediateAngularVelocityRadians);
 }
 
+void UBlueprintSpringMathLibrary::CriticalDoubleSpringDampScale(FVector& InOutScale, FVector& InOutScalarVelocity, FVector& InOutIntermediateScale, FVector& InOutIntermediateScalarVelocity, const FVector& TargetScale, float DeltaTime, float SmoothingTime)
+{
+	SpringMath::CriticalDoubleSpringDamperScale(InOutScale, InOutScalarVelocity, InOutIntermediateScale, InOutIntermediateScalarVelocity, TargetScale, SmoothingTime, DeltaTime);
+}
+
+void UBlueprintSpringMathLibrary::CriticalDoubleSpringDampTransform(FTransform& InOutTransform, FSpringTransformVelocity& InOutVelocity, FTransform& InOutIntermediateTransform, FSpringTransformVelocity& InOutIntermediateTransformVelocity, const FTransform& TargetTransform, float DeltaTime, float SmoothingTime)
+{
+	FVector InOutLocation = InOutTransform.GetLocation();
+	FQuat InOutRotation = InOutTransform.GetRotation();
+	FVector InOutScale = InOutTransform.GetScale3D();
+	FVector InOutIntermediateLocation = InOutIntermediateTransform.GetLocation();
+	FQuat InOutIntermediateRotation = InOutIntermediateTransform.GetRotation();
+	FVector InOutIntermediateScale = InOutIntermediateTransform.GetScale3D();
+
+	CriticalDoubleSpringDampVector(InOutLocation, InOutVelocity.LinearVelocity, InOutIntermediateLocation, InOutIntermediateTransformVelocity.LinearVelocity, TargetTransform.GetLocation(), DeltaTime, SmoothingTime);
+	CriticalDoubleSpringDampQuat(InOutRotation, InOutVelocity.AngularVelocity, InOutIntermediateRotation, InOutIntermediateTransformVelocity.AngularVelocity, TargetTransform.GetRotation(), DeltaTime, SmoothingTime);
+	CriticalDoubleSpringDampScale(InOutScale, InOutVelocity.ScalarVelocity, InOutIntermediateScale, InOutIntermediateTransformVelocity.ScalarVelocity, TargetTransform.GetScale3D(), DeltaTime, SmoothingTime);
+
+	InOutTransform = FTransform(InOutRotation, InOutLocation, InOutScale);
+	InOutIntermediateTransform = FTransform(InOutIntermediateRotation, InOutIntermediateLocation, InOutIntermediateScale);
+}
+
 void UBlueprintSpringMathLibrary::VelocitySpringDampFloat(float& InOutX, float& InOutV, float& InOutVi, float TargetX, float MaxSpeed,
                                                           float DeltaTime, float SmoothingTime)
 {
@@ -178,6 +216,21 @@ FRotator UBlueprintSpringMathLibrary::DampRotator(const FRotator& Rotation, cons
 	return DampedRotation.Rotator();
 }
 
+FVector UBlueprintSpringMathLibrary::DampScale(const FVector& Scale, const FVector& TargetScale, float DeltaTime, float SmoothingTime)
+{
+	FVector DampedScale = Scale;
+	SpringMath::ExponentialSmoothingApproxScale(DampedScale, TargetScale, DeltaTime, SmoothingTime);
+	return DampedScale;
+}
+
+FTransform UBlueprintSpringMathLibrary::DampTransform(const FTransform& Transform, const FTransform& TargetTransform, float DeltaTime, float SmoothingTime)
+{
+	return FTransform(
+		DampQuat(Transform.GetRotation(), TargetTransform.GetRotation(), DeltaTime, SmoothingTime),
+		DampVector(Transform.GetLocation(), TargetTransform.GetLocation(), DeltaTime, SmoothingTime),
+		DampScale(Transform.GetScale3D(), TargetTransform.GetScale3D(), DeltaTime, SmoothingTime));
+}
+
 void UBlueprintSpringMathLibrary::SpringCharacterUpdate(FVector& InOutPosition, FVector& InOutVelocity, FVector& InOutAcceleration,
                                                         const FVector& TargetVelocity, float DeltaTime, float SmoothingTime)
 {
@@ -209,3 +262,413 @@ float UBlueprintSpringMathLibrary::ConvertSmoothingTimeToHalfLife(float Smoothin
 {
 	return SpringMath::SmoothingTimeToHalfLife(SmoothingTime);
 }
+
+void UBlueprintSpringMathLibrary::TrackVelocityFloat(float& InOutValue, float& InOutVelocity, float InValue, float DeltaTime)
+{
+	SpringMath::TrackVelocity(InOutValue, InOutVelocity, InValue, DeltaTime);
+}
+
+void UBlueprintSpringMathLibrary::TrackVelocityVector(FVector& InOutValue, FVector& InOutVelocity, FVector InValue, float DeltaTime)
+{
+	SpringMath::TrackVelocity(InOutValue, InOutVelocity, InValue, DeltaTime);
+}
+
+void UBlueprintSpringMathLibrary::TrackVelocityVector2D(FVector2D& InOutValue, FVector2D& InOutVelocity, FVector2D InValue, float DeltaTime)
+{
+	SpringMath::TrackVelocity(InOutValue, InOutVelocity, InValue, DeltaTime);
+}
+
+void UBlueprintSpringMathLibrary::TrackVelocityQuat(FQuat& InOutValue, FVector& InOutVelocity, FQuat InValue, float DeltaTime)
+{
+	FVector InOutVelocityRadians = FMath::DegreesToRadians(InOutVelocity);
+	SpringMath::TrackVelocityQuat(InOutValue, InOutVelocityRadians, InValue, DeltaTime);
+	InOutVelocity = FMath::RadiansToDegrees(InOutVelocityRadians);
+}
+
+void UBlueprintSpringMathLibrary::TrackVelocityRotator(FRotator& InOutValue, FVector& InOutVelocity, FRotator InValue, float DeltaTime)
+{
+	FQuat InOutValueQuat = InOutValue.Quaternion();
+	FVector InOutVelocityRadians = FMath::DegreesToRadians(InOutVelocity);
+	SpringMath::TrackVelocityQuat(InOutValueQuat, InOutVelocityRadians, InValue.Quaternion(), DeltaTime);
+	InOutValue = InOutValueQuat.Rotator();
+	InOutVelocity = FMath::RadiansToDegrees(InOutVelocityRadians);
+}
+
+void UBlueprintSpringMathLibrary::TrackVelocityAngle(float& InOutValue, float& InOutVelocity, float InValue, float DeltaTime)
+{
+	float InOutValueRadians = FMath::DegreesToRadians(InOutValue);
+	float InOutVelocityRadians = FMath::DegreesToRadians(InOutVelocity);
+	SpringMath::TrackVelocityAngle(InOutValueRadians, InOutVelocityRadians, FMath::DegreesToRadians(InValue), DeltaTime);
+	InOutValue = FMath::RadiansToDegrees(InOutValueRadians);
+	InOutVelocity = FMath::RadiansToDegrees(InOutVelocityRadians);
+}
+
+void UBlueprintSpringMathLibrary::TrackVelocityScale(FVector& InOutValue, FVector& InOutVelocity, FVector InValue, float DeltaTime)
+{
+	SpringMath::TrackVelocityScale(InOutValue, InOutVelocity, InValue, DeltaTime);
+}
+
+void UBlueprintSpringMathLibrary::TrackVelocityTransform(FTransform& InOutValue, FSpringTransformVelocity& InOutVelocity, FTransform InValue, float DeltaTime)
+{
+	FVector InOutLocation = InOutValue.GetLocation();
+	FQuat InOutRotation = InOutValue.GetRotation();
+	FVector InOutScale = InOutValue.GetScale3D();
+
+	TrackVelocityVector(InOutLocation, InOutVelocity.LinearVelocity, InValue.GetLocation(), DeltaTime);
+	TrackVelocityQuat(InOutRotation, InOutVelocity.AngularVelocity, InValue.GetRotation(), DeltaTime);
+	TrackVelocityScale(InOutScale, InOutVelocity.ScalarVelocity, InValue.GetScale3D(), DeltaTime);
+
+	InOutValue = FTransform(InOutRotation, InOutLocation, InOutScale);
+}
+
+float UBlueprintSpringMathLibrary::SpringInertializeApplyToFloat(float InValue, float InValueOffset, float InVelocityOffset, float TimeSinceTransition, float SmoothingTime)
+{
+	float ValueVelocity = 0.0f;
+	SpringMath::SpringInertializeApply(InValue, ValueVelocity, InValueOffset, InVelocityOffset, TimeSinceTransition, SmoothingTime);
+	return InValue;
+}
+
+void UBlueprintSpringMathLibrary::SpringInertializeTransitionFloat(float& InOutValueOffset, float& InOutVelocityOffset, float InSrcValue, float InSrcVelocity, float InDstValue, float InDstVelocity, float TimeSinceTransition, float SmoothingTime)
+{
+	SpringMath::SpringInertializeTransition(InOutValueOffset, InOutVelocityOffset, InSrcValue, InSrcVelocity, InDstValue, InDstVelocity, TimeSinceTransition, SmoothingTime);
+}
+
+FVector UBlueprintSpringMathLibrary::SpringInertializeApplyToVector(FVector InValue, FVector InValueOffset, FVector InVelocityOffset, float TimeSinceTransition, float SmoothingTime)
+{
+	FVector ValueVelocity = FVector::ZeroVector;
+	SpringMath::SpringInertializeApply(InValue, ValueVelocity, InValueOffset, InVelocityOffset, TimeSinceTransition, SmoothingTime);
+	return InValue;
+}
+
+void UBlueprintSpringMathLibrary::SpringInertializeTransitionVector(FVector& InOutValueOffset, FVector& InOutVelocityOffset, FVector InSrcValue, FVector InSrcVelocity, FVector InDstValue, FVector InDstVelocity, float TimeSinceTransition, float SmoothingTime)
+{
+	SpringMath::SpringInertializeTransition(InOutValueOffset, InOutVelocityOffset, InSrcValue, InSrcVelocity, InDstValue, InDstVelocity, TimeSinceTransition, SmoothingTime);
+}
+
+FVector2D UBlueprintSpringMathLibrary::SpringInertializeApplyToVector2D(FVector2D InValue, FVector2D InValueOffset, FVector2D InVelocityOffset, float TimeSinceTransition, float SmoothingTime)
+{
+	FVector2D ValueVelocity = FVector2D::ZeroVector;
+	SpringMath::SpringInertializeApply(InValue, ValueVelocity, InValueOffset, InVelocityOffset, TimeSinceTransition, SmoothingTime);
+	return InValue;
+}
+
+void UBlueprintSpringMathLibrary::SpringInertializeTransitionVector2D(FVector2D& InOutValueOffset, FVector2D& InOutVelocityOffset, FVector2D InSrcValue, FVector2D InSrcVelocity, FVector2D InDstValue, FVector2D InDstVelocity, float TimeSinceTransition, float SmoothingTime)
+{
+	SpringMath::SpringInertializeTransition(InOutValueOffset, InOutVelocityOffset, InSrcValue, InSrcVelocity, InDstValue, InDstVelocity, TimeSinceTransition, SmoothingTime);
+}
+
+FQuat UBlueprintSpringMathLibrary::SpringInertializeApplyToQuat(FQuat InRotation, FQuat InRotationOffset, FVector InAngularVelocityOffset, float TimeSinceTransition, float SmoothingTime)
+{
+	FVector ValueVelocity = FVector::ZeroVector;
+	SpringMath::SpringInertializeApplyQuat(InRotation, ValueVelocity, InRotationOffset, FMath::DegreesToRadians(InAngularVelocityOffset), TimeSinceTransition, SmoothingTime);
+	return InRotation;
+}
+
+void UBlueprintSpringMathLibrary::SpringInertializeTransitionQuat(FQuat& InOutRotationOffset, FVector& InOutAngularVelocityOffset, FQuat InSrcRotation, FVector InSrcAngularVelocity, FQuat InDstRotation, FVector InDstAngularVelocity, float TimeSinceTransition, float SmoothingTime)
+{
+	FVector InOutAngularVelocityOffsetRadians = FMath::DegreesToRadians(InOutAngularVelocityOffset);
+	SpringMath::SpringInertializeTransitionQuat(InOutRotationOffset, InOutAngularVelocityOffsetRadians, InSrcRotation, FMath::DegreesToRadians(InSrcAngularVelocity), InDstRotation, FMath::DegreesToRadians(InDstAngularVelocity), TimeSinceTransition, SmoothingTime);
+	InOutAngularVelocityOffset = FMath::RadiansToDegrees(InOutAngularVelocityOffsetRadians);
+}
+
+FRotator UBlueprintSpringMathLibrary::SpringInertializeApplyToRotator(FRotator InRotation, FRotator InRotationOffset, FVector InAngularVelocityOffset, float TimeSinceTransition, float SmoothingTime)
+{
+	FQuat InRotationQuat = InRotation.Quaternion();
+	FVector ValueVelocity = FVector::ZeroVector;
+	SpringMath::SpringInertializeApplyQuat(InRotationQuat, ValueVelocity, InRotationOffset.Quaternion(), FMath::DegreesToRadians(InAngularVelocityOffset), TimeSinceTransition, SmoothingTime);
+	return InRotationQuat.Rotator();
+}
+
+void UBlueprintSpringMathLibrary::SpringInertializeTransitionRotator(FRotator& InOutRotationOffset, FVector& InOutAngularVelocityOffset, FRotator InSrcRotation, FVector InSrcAngularVelocity, FRotator InDstRotation, FVector InDstAngularVelocity, float TimeSinceTransition, float SmoothingTime)
+{
+	FQuat InOutRotationOffsetQuat = InOutRotationOffset.Quaternion();
+	FVector InOutAngularVelocityOffsetRadians = FMath::DegreesToRadians(InOutAngularVelocityOffset);
+	SpringMath::SpringInertializeTransitionQuat(InOutRotationOffsetQuat, InOutAngularVelocityOffsetRadians, InSrcRotation.Quaternion(), FMath::DegreesToRadians(InSrcAngularVelocity), InDstRotation.Quaternion(), FMath::DegreesToRadians(InDstAngularVelocity), TimeSinceTransition, SmoothingTime);
+	InOutRotationOffset = InOutRotationOffsetQuat.Rotator();
+	InOutAngularVelocityOffset = FMath::RadiansToDegrees(InOutAngularVelocityOffsetRadians);
+}
+
+float UBlueprintSpringMathLibrary::SpringInertializeApplyToAngle(float InAngle, float InAngleOffset, float InAngularVelocityOffset, float TimeSinceTransition, float SmoothingTime)
+{
+	float InAngleRadians = FMath::DegreesToRadians(InAngle);
+	float ValueVelocity = 0.0f;
+	SpringMath::SpringInertializeApplyAngle(InAngleRadians, ValueVelocity, FMath::DegreesToRadians(InAngleOffset), FMath::DegreesToRadians(InAngularVelocityOffset), TimeSinceTransition, SmoothingTime);
+	return FMath::RadiansToDegrees(InAngleRadians);
+}
+
+void UBlueprintSpringMathLibrary::SpringInertializeTransitionAngle(float& InOutAngleOffset, float& InOutAngularVelocityOffset, float InSrcAngle, float InSrcAngularVelocity, float InDstAngle, float InDstAngularVelocity, float TimeSinceTransition, float SmoothingTime)
+{
+	float InOutAngleOffsetRadians = FMath::DegreesToRadians(InOutAngleOffset);
+	float InOutAngularVelocityOffsetRadians = FMath::DegreesToRadians(InOutAngularVelocityOffset);
+	SpringMath::SpringInertializeTransitionAngle(InOutAngleOffsetRadians, InOutAngularVelocityOffsetRadians, FMath::DegreesToRadians(InSrcAngle), FMath::DegreesToRadians(InSrcAngularVelocity), FMath::DegreesToRadians(InDstAngle), FMath::DegreesToRadians(InDstAngularVelocity), TimeSinceTransition, SmoothingTime);
+	InOutAngleOffset = FMath::RadiansToDegrees(InOutAngleOffsetRadians);
+	InOutAngularVelocityOffset = FMath::RadiansToDegrees(InOutAngularVelocityOffsetRadians);
+}
+
+FVector UBlueprintSpringMathLibrary::SpringInertializeApplyToScale(FVector InScale, FVector InScaleOffset, FVector InScalarVelocityOffset, float TimeSinceTransition, float SmoothingTime)
+{
+	FVector ValueVelocity = FVector::ZeroVector;
+	SpringMath::SpringInertializeApplyScale(InScale, ValueVelocity, InScaleOffset, InScalarVelocityOffset, TimeSinceTransition, SmoothingTime);
+	return InScale;
+}
+
+void UBlueprintSpringMathLibrary::SpringInertializeTransitionScale(FVector& InOutScaleOffset, FVector& InOutScalarVelocityOffset, FVector InSrcScale, FVector InSrcScalarVelocity, FVector InDstScale, FVector InDstScalarVelocity, float TimeSinceTransition, float SmoothingTime)
+{
+	SpringMath::SpringInertializeTransitionScale(InOutScaleOffset, InOutScalarVelocityOffset, InSrcScale, InSrcScalarVelocity, InDstScale, InDstScalarVelocity, TimeSinceTransition, SmoothingTime);
+}
+
+FTransform UBlueprintSpringMathLibrary::SpringInertializeApplyToTransform(FTransform InTransform, FTransform InTransformOffset, FSpringTransformVelocity InVelocityOffset, float TimeSinceTransition, float SmoothingTime)
+{
+	return FTransform(
+		SpringInertializeApplyToQuat(InTransform.GetRotation(), InTransformOffset.GetRotation(), InVelocityOffset.AngularVelocity, TimeSinceTransition, SmoothingTime),
+		SpringInertializeApplyToVector(InTransform.GetLocation(), InTransformOffset.GetLocation(), InVelocityOffset.LinearVelocity, TimeSinceTransition, SmoothingTime),
+		SpringInertializeApplyToScale(InTransform.GetScale3D(), InTransformOffset.GetScale3D(), InVelocityOffset.ScalarVelocity, TimeSinceTransition, SmoothingTime));
+}
+
+void UBlueprintSpringMathLibrary::SpringInertializeTransitionTransform(FTransform& InOutTransformOffset, FSpringTransformVelocity& InOutVelocityOffset, FTransform InSrcTransform, FSpringTransformVelocity InSrcVelocity, FTransform InDstTransform, FSpringTransformVelocity InDstVelocity, float TimeSinceTransition, float SmoothingTime)
+{
+	FVector InOutLocation = InOutTransformOffset.GetLocation();
+	FQuat InOutRotation = InOutTransformOffset.GetRotation();
+	FVector InOutScale = InOutTransformOffset.GetScale3D();
+
+	SpringInertializeTransitionVector(InOutLocation, InOutVelocityOffset.LinearVelocity, InSrcTransform.GetLocation(), InSrcVelocity.LinearVelocity, InDstTransform.GetLocation(), InDstVelocity.LinearVelocity, TimeSinceTransition, SmoothingTime);
+	SpringInertializeTransitionQuat(InOutRotation, InOutVelocityOffset.AngularVelocity, InSrcTransform.GetRotation(), InSrcVelocity.AngularVelocity, InDstTransform.GetRotation(), InDstVelocity.AngularVelocity, TimeSinceTransition, SmoothingTime);
+	SpringInertializeTransitionScale(InOutScale, InOutVelocityOffset.ScalarVelocity, InSrcTransform.GetScale3D(), InSrcVelocity.ScalarVelocity, InDstTransform.GetScale3D(), InDstVelocity.ScalarVelocity, TimeSinceTransition, SmoothingTime);
+
+	InOutTransformOffset = FTransform(InOutRotation, InOutLocation, InOutScale);
+}
+
+float UBlueprintSpringMathLibrary::CubicInertializeApplyToFloat(float InValue, float InValueOffset, float InVelocityOffset, float TimeSinceTransition, float BlendTime)
+{
+	float ValueVelocity = 0.0f;
+	SpringMath::CubicInertializeApply(InValue, ValueVelocity, InValueOffset, InVelocityOffset, TimeSinceTransition, BlendTime);
+	return InValue;
+}
+
+void UBlueprintSpringMathLibrary::CubicInertializeTransitionFloat(float& InOutValueOffset, float& InOutVelocityOffset, float InSrcValue, float InSrcVelocity, float InDstValue, float InDstVelocity, float TimeSinceTransition, float BlendTime)
+{
+	SpringMath::CubicInertializeTransition(InOutValueOffset, InOutVelocityOffset, InSrcValue, InSrcVelocity, InDstValue, InDstVelocity, TimeSinceTransition, BlendTime);
+}
+
+FVector UBlueprintSpringMathLibrary::CubicInertializeApplyToVector(FVector InValue, FVector InValueOffset, FVector InVelocityOffset, float TimeSinceTransition, float BlendTime)
+{
+	FVector ValueVelocity = FVector::ZeroVector;
+	SpringMath::CubicInertializeApply(InValue, ValueVelocity, InValueOffset, InVelocityOffset, TimeSinceTransition, BlendTime);
+	return InValue;
+}
+
+void UBlueprintSpringMathLibrary::CubicInertializeTransitionVector(FVector& InOutValueOffset, FVector& InOutVelocityOffset, FVector InSrcValue, FVector InSrcVelocity, FVector InDstValue, FVector InDstVelocity, float TimeSinceTransition, float BlendTime)
+{
+	SpringMath::CubicInertializeTransition(InOutValueOffset, InOutVelocityOffset, InSrcValue, InSrcVelocity, InDstValue, InDstVelocity, TimeSinceTransition, BlendTime);
+}
+
+FVector2D UBlueprintSpringMathLibrary::CubicInertializeApplyToVector2D(FVector2D InValue, FVector2D InValueOffset, FVector2D InVelocityOffset, float TimeSinceTransition, float BlendTime)
+{
+	FVector2D ValueVelocity = FVector2D::ZeroVector;
+	SpringMath::CubicInertializeApply(InValue, ValueVelocity, InValueOffset, InVelocityOffset, TimeSinceTransition, BlendTime);
+	return InValue;
+}
+
+void UBlueprintSpringMathLibrary::CubicInertializeTransitionVector2D(FVector2D& InOutValueOffset, FVector2D& InOutVelocityOffset, FVector2D InSrcValue, FVector2D InSrcVelocity, FVector2D InDstValue, FVector2D InDstVelocity, float TimeSinceTransition, float BlendTime)
+{
+	SpringMath::CubicInertializeTransition(InOutValueOffset, InOutVelocityOffset, InSrcValue, InSrcVelocity, InDstValue, InDstVelocity, TimeSinceTransition, BlendTime);
+}
+
+FQuat UBlueprintSpringMathLibrary::CubicInertializeApplyToQuat(FQuat InRotation, FQuat InRotationOffset, FVector InAngularVelocityOffset, float TimeSinceTransition, float BlendTime)
+{
+	FVector ValueVelocity = FVector::ZeroVector;
+	SpringMath::CubicInertializeApplyQuat(InRotation, ValueVelocity, InRotationOffset, FMath::DegreesToRadians(InAngularVelocityOffset), TimeSinceTransition, BlendTime);
+	return InRotation;
+}
+
+void UBlueprintSpringMathLibrary::CubicInertializeTransitionQuat(FQuat& InOutRotationOffset, FVector& InOutAngularVelocityOffset, FQuat InSrcRotation, FVector InSrcAngularVelocity, FQuat InDstRotation, FVector InDstAngularVelocity, float TimeSinceTransition, float BlendTime)
+{
+	FVector InOutAngularVelocityOffsetRadians = FMath::DegreesToRadians(InOutAngularVelocityOffset);
+	SpringMath::CubicInertializeTransitionQuat(InOutRotationOffset, InOutAngularVelocityOffsetRadians, InSrcRotation, FMath::DegreesToRadians(InSrcAngularVelocity), InDstRotation, FMath::DegreesToRadians(InDstAngularVelocity), TimeSinceTransition, BlendTime);
+	InOutAngularVelocityOffset = FMath::RadiansToDegrees(InOutAngularVelocityOffsetRadians);
+}
+
+FRotator UBlueprintSpringMathLibrary::CubicInertializeApplyToRotator(FRotator InRotation, FRotator InRotationOffset, FVector InAngularVelocityOffset, float TimeSinceTransition, float BlendTime)
+{
+	FQuat InRotationQuat = InRotation.Quaternion();
+	FVector ValueVelocity = FVector::ZeroVector;
+	SpringMath::CubicInertializeApplyQuat(InRotationQuat, ValueVelocity, InRotationOffset.Quaternion(), FMath::DegreesToRadians(InAngularVelocityOffset), TimeSinceTransition, BlendTime);
+	return InRotationQuat.Rotator();
+}
+
+void UBlueprintSpringMathLibrary::CubicInertializeTransitionRotator(FRotator& InOutRotationOffset, FVector& InOutAngularVelocityOffset, FRotator InSrcRotation, FVector InSrcAngularVelocity, FRotator InDstRotation, FVector InDstAngularVelocity, float TimeSinceTransition, float BlendTime)
+{
+	FQuat InOutRotationOffsetQuat = InOutRotationOffset.Quaternion();
+	FVector InOutAngularVelocityOffsetRadians = FMath::DegreesToRadians(InOutAngularVelocityOffset);
+	SpringMath::CubicInertializeTransitionQuat(InOutRotationOffsetQuat, InOutAngularVelocityOffsetRadians, InSrcRotation.Quaternion(), FMath::DegreesToRadians(InSrcAngularVelocity), InDstRotation.Quaternion(), FMath::DegreesToRadians(InDstAngularVelocity), TimeSinceTransition, BlendTime);
+	InOutRotationOffset = InOutRotationOffsetQuat.Rotator();
+	InOutAngularVelocityOffset = FMath::RadiansToDegrees(InOutAngularVelocityOffsetRadians);
+}
+
+float UBlueprintSpringMathLibrary::CubicInertializeApplyToAngle(float InAngle, float InAngleOffset, float InAngularVelocityOffset, float TimeSinceTransition, float BlendTime)
+{
+	float InAngleRadians = FMath::DegreesToRadians(InAngle);
+	float ValueVelocity = 0.0f;
+	SpringMath::CubicInertializeApplyAngle(InAngleRadians, ValueVelocity, FMath::DegreesToRadians(InAngleOffset), FMath::DegreesToRadians(InAngularVelocityOffset), TimeSinceTransition, BlendTime);
+	return FMath::RadiansToDegrees(InAngleRadians);
+}
+
+void UBlueprintSpringMathLibrary::CubicInertializeTransitionAngle(float& InOutAngleOffset, float& InOutAngularVelocityOffset, float InSrcAngle, float InSrcAngularVelocity, float InDstAngle, float InDstAngularVelocity, float TimeSinceTransition, float BlendTime)
+{
+	float InOutAngleOffsetRadians = FMath::DegreesToRadians(InOutAngleOffset);
+	float InOutAngularVelocityOffsetRadians = FMath::DegreesToRadians(InOutAngularVelocityOffset);
+	SpringMath::CubicInertializeTransitionAngle(InOutAngleOffsetRadians, InOutAngularVelocityOffsetRadians, FMath::DegreesToRadians(InSrcAngle), FMath::DegreesToRadians(InSrcAngularVelocity), FMath::DegreesToRadians(InDstAngle), FMath::DegreesToRadians(InDstAngularVelocity), TimeSinceTransition, BlendTime);
+	InOutAngleOffset = FMath::RadiansToDegrees(InOutAngleOffsetRadians);
+	InOutAngularVelocityOffset = FMath::RadiansToDegrees(InOutAngularVelocityOffsetRadians);
+}
+
+FVector UBlueprintSpringMathLibrary::CubicInertializeApplyToScale(FVector InScale, FVector InScaleOffset, FVector InScalarVelocityOffset, float TimeSinceTransition, float BlendTime)
+{
+	FVector ValueVelocity = FVector::ZeroVector;
+	SpringMath::CubicInertializeApplyScale(InScale, ValueVelocity, InScaleOffset, InScalarVelocityOffset, TimeSinceTransition, BlendTime);
+	return InScale;
+}
+
+void UBlueprintSpringMathLibrary::CubicInertializeTransitionScale(FVector& InOutScaleOffset, FVector& InOutScalarVelocityOffset, FVector InSrcScale, FVector InSrcScalarVelocity, FVector InDstScale, FVector InDstScalarVelocity, float TimeSinceTransition, float BlendTime)
+{
+	SpringMath::CubicInertializeTransitionScale(InOutScaleOffset, InOutScalarVelocityOffset, InSrcScale, InSrcScalarVelocity, InDstScale, InDstScalarVelocity, TimeSinceTransition, BlendTime);
+}
+
+FTransform UBlueprintSpringMathLibrary::CubicInertializeApplyToTransform(FTransform InTransform, FTransform InTransformOffset, FSpringTransformVelocity InVelocityOffset, float TimeSinceTransition, float BlendTime)
+{
+	return FTransform(
+		CubicInertializeApplyToQuat(InTransform.GetRotation(), InTransformOffset.GetRotation(), InVelocityOffset.AngularVelocity, TimeSinceTransition, BlendTime),
+		CubicInertializeApplyToVector(InTransform.GetLocation(), InTransformOffset.GetLocation(), InVelocityOffset.LinearVelocity, TimeSinceTransition, BlendTime),
+		CubicInertializeApplyToScale(InTransform.GetScale3D(), InTransformOffset.GetScale3D(), InVelocityOffset.ScalarVelocity, TimeSinceTransition, BlendTime));
+}
+
+void UBlueprintSpringMathLibrary::CubicInertializeTransitionTransform(FTransform& InOutTransformOffset, FSpringTransformVelocity& InOutVelocityOffset, FTransform InSrcTransform, FSpringTransformVelocity InSrcVelocity, FTransform InDstTransform, FSpringTransformVelocity InDstVelocity, float TimeSinceTransition, float BlendTime)
+{
+	FVector InOutLocation = InOutTransformOffset.GetLocation();
+	FQuat InOutRotation = InOutTransformOffset.GetRotation();
+	FVector InOutScale = InOutTransformOffset.GetScale3D();
+
+	CubicInertializeTransitionVector(InOutLocation, InOutVelocityOffset.LinearVelocity, InSrcTransform.GetLocation(), InSrcVelocity.LinearVelocity, InDstTransform.GetLocation(), InDstVelocity.LinearVelocity, TimeSinceTransition, BlendTime);
+	CubicInertializeTransitionQuat(InOutRotation, InOutVelocityOffset.AngularVelocity, InSrcTransform.GetRotation(), InSrcVelocity.AngularVelocity, InDstTransform.GetRotation(), InDstVelocity.AngularVelocity, TimeSinceTransition, BlendTime);
+	CubicInertializeTransitionScale(InOutScale, InOutVelocityOffset.ScalarVelocity, InSrcTransform.GetScale3D(), InSrcVelocity.ScalarVelocity, InDstTransform.GetScale3D(), InDstVelocity.ScalarVelocity, TimeSinceTransition, BlendTime);
+
+	InOutTransformOffset = FTransform(InOutRotation, InOutLocation, InOutScale);
+}
+
+
+float UBlueprintSpringMathLibrary::DeadBlendApplyToFloat(float InValue, float InValueTransition, float InVelocityTransition, float TimeSinceTransition, float BlendTime, float SmoothingTime)
+{
+	float VelocityValue = 0.0f;
+	SpringMath::DeadBlendApply(InValue, VelocityValue, InValueTransition, InVelocityTransition, TimeSinceTransition, BlendTime, SmoothingTime);
+	return InValue;
+}
+
+void UBlueprintSpringMathLibrary::DeadBlendTransitionFloat(float& InOutValueTransition, float& InOutVelocityTransition, float InSrcValue, float InSrcVelocity, float TimeSinceTransition, float BlendTime, float SmoothingTime)
+{
+	SpringMath::DeadBlendTransition(InOutValueTransition, InOutVelocityTransition, InSrcValue, InSrcVelocity, TimeSinceTransition, BlendTime, SmoothingTime);
+}
+
+FVector UBlueprintSpringMathLibrary::DeadBlendApplyToVector(FVector InValue, FVector InValueTransition, FVector InVelocityTransition, float TimeSinceTransition, float BlendTime, float SmoothingTime)
+{
+	FVector VelocityValue = FVector::ZeroVector;
+	SpringMath::DeadBlendApply(InValue, VelocityValue, InValueTransition, InVelocityTransition, TimeSinceTransition, BlendTime, SmoothingTime);
+	return InValue;
+}
+
+void UBlueprintSpringMathLibrary::DeadBlendTransitionVector(FVector& InOutValueTransition, FVector& InOutVelocityTransition, FVector InSrcValue, FVector InSrcVelocity, float TimeSinceTransition, float BlendTime, float SmoothingTime)
+{
+	SpringMath::DeadBlendTransition(InOutValueTransition, InOutVelocityTransition, InSrcValue, InSrcVelocity, TimeSinceTransition, BlendTime, SmoothingTime);
+}
+
+FVector2D UBlueprintSpringMathLibrary::DeadBlendApplyToVector2D(FVector2D InValue, FVector2D InValueTransition, FVector2D InVelocityTransition, float TimeSinceTransition, float BlendTime, float SmoothingTime)
+{
+	FVector2D VelocityValue = FVector2D::ZeroVector;
+	SpringMath::DeadBlendApply(InValue, VelocityValue, InValueTransition, InVelocityTransition, TimeSinceTransition, BlendTime, SmoothingTime);
+	return InValue;
+}
+
+void UBlueprintSpringMathLibrary::DeadBlendTransitionVector2D(FVector2D& InOutValueTransition, FVector2D& InOutVelocityTransition, FVector2D InSrcValue, FVector2D InSrcVelocity, float TimeSinceTransition, float BlendTime, float SmoothingTime)
+{
+	SpringMath::DeadBlendTransition(InOutValueTransition, InOutVelocityTransition, InSrcValue, InSrcVelocity, TimeSinceTransition, BlendTime, SmoothingTime);
+}
+
+FQuat UBlueprintSpringMathLibrary::DeadBlendApplyToQuat(FQuat InRotation, FQuat InRotationTransition, FVector InAngularVelocityTransition, float TimeSinceTransition, float BlendTime, float SmoothingTime)
+{
+	FVector VelocityValue = FVector::ZeroVector;
+	SpringMath::DeadBlendApplyQuat(InRotation, VelocityValue, InRotationTransition, FMath::DegreesToRadians(InAngularVelocityTransition), TimeSinceTransition, BlendTime, SmoothingTime);
+	return InRotation;
+}
+
+void UBlueprintSpringMathLibrary::DeadBlendTransitionQuat(FQuat& InOutRotationTransition, FVector& InOutAngularVelocityTransition, FQuat InSrcRotation, FVector InSrcAngularVelocity, float TimeSinceTransition, float BlendTime, float SmoothingTime)
+{
+	FVector InOutAngularVelocityTransitionRadians = FMath::DegreesToRadians(InOutAngularVelocityTransition);
+	SpringMath::DeadBlendTransitionQuat(InOutRotationTransition, InOutAngularVelocityTransitionRadians, InSrcRotation, FMath::DegreesToRadians(InSrcAngularVelocity), TimeSinceTransition, BlendTime, SmoothingTime);
+	InOutAngularVelocityTransition = FMath::RadiansToDegrees(InOutAngularVelocityTransitionRadians);
+}
+
+FRotator UBlueprintSpringMathLibrary::DeadBlendApplyToRotator(FRotator InRotation, FRotator InRotationTransition, FVector InAngularVelocityTransition, float TimeSinceTransition, float BlendTime, float SmoothingTime)
+{
+	FVector VelocityValue = FVector::ZeroVector;
+	FQuat InRotationQuat = InRotation.Quaternion();
+	SpringMath::DeadBlendApplyQuat(InRotationQuat, VelocityValue, InRotationTransition.Quaternion(), FMath::DegreesToRadians(InAngularVelocityTransition), TimeSinceTransition, BlendTime, SmoothingTime);
+	return InRotationQuat.Rotator();
+}
+
+void UBlueprintSpringMathLibrary::DeadBlendTransitionRotator(FRotator& InOutRotationTransition, FVector& InOutAngularVelocityTransition, FRotator InSrcRotation, FVector InSrcAngularVelocity, float TimeSinceTransition, float BlendTime, float SmoothingTime)
+{
+	FQuat InOutRotationTransitionQuat = InOutRotationTransition.Quaternion();
+	FVector InOutAngularVelocityTransitionRadians = FMath::DegreesToRadians(InOutAngularVelocityTransition);
+	SpringMath::DeadBlendTransitionQuat(InOutRotationTransitionQuat, InOutAngularVelocityTransitionRadians, InSrcRotation.Quaternion(), FMath::DegreesToRadians(InSrcAngularVelocity), TimeSinceTransition, BlendTime, SmoothingTime);
+	InOutRotationTransition = InOutRotationTransitionQuat.Rotator();
+	InOutAngularVelocityTransition = FMath::RadiansToDegrees(InOutAngularVelocityTransitionRadians);
+}
+
+float UBlueprintSpringMathLibrary::DeadBlendApplyToAngle(float InAngle, float InAngleTransition, float InAngularVelocityTransition, float TimeSinceTransition, float BlendTime, float SmoothingTime)
+{
+	float InAngleRadians = FMath::DegreesToRadians(InAngle);
+	float VelocityValue = 0.0f;
+	SpringMath::DeadBlendApplyAngle(InAngleRadians, VelocityValue, FMath::DegreesToRadians(InAngleTransition), FMath::DegreesToRadians(InAngularVelocityTransition), TimeSinceTransition, BlendTime, SmoothingTime);
+	return FMath::RadiansToDegrees(InAngleRadians);
+}
+
+void UBlueprintSpringMathLibrary::DeadBlendTransitionAngle(float& InOutAngleTransition, float& InOutAngularVelocityTransition, float InSrcAngle, float InSrcAngularVelocity, float TimeSinceTransition, float BlendTime, float SmoothingTime)
+{
+	float InOutAngleTransitionRadians = FMath::DegreesToRadians(InOutAngleTransition);
+	float InOutAngularVelocityTransitionRadians = FMath::DegreesToRadians(InOutAngularVelocityTransition);
+	SpringMath::DeadBlendTransitionAngle(InOutAngleTransitionRadians, InOutAngularVelocityTransitionRadians, FMath::DegreesToRadians(InSrcAngle), FMath::DegreesToRadians(InSrcAngularVelocity), TimeSinceTransition, BlendTime, SmoothingTime);
+	InOutAngleTransition = FMath::RadiansToDegrees(InOutAngleTransitionRadians);
+	InOutAngularVelocityTransition = FMath::RadiansToDegrees(InOutAngularVelocityTransitionRadians);
+}
+
+FVector UBlueprintSpringMathLibrary::DeadBlendApplyToScale(FVector InScale, FVector InScaleTransition, FVector InScalarVelocityTransition, float TimeSinceTransition, float BlendTime, float SmoothingTime)
+{
+	FVector VelocityValue = FVector::ZeroVector;
+	SpringMath::DeadBlendApplyScale(InScale, VelocityValue, InScaleTransition, InScalarVelocityTransition, TimeSinceTransition, BlendTime, SmoothingTime);
+	return InScale;
+}
+
+void UBlueprintSpringMathLibrary::DeadBlendTransitionScale(FVector& InOutScaleTransition, FVector& InOutScalarVelocityTransition, FVector InSrcScale, FVector InSrcScalarVelocity, float TimeSinceTransition, float BlendTime, float SmoothingTime)
+{
+	SpringMath::DeadBlendTransitionScale(InOutScaleTransition, InOutScalarVelocityTransition, InSrcScale, InSrcScalarVelocity, TimeSinceTransition, BlendTime, SmoothingTime);
+}
+
+FTransform UBlueprintSpringMathLibrary::DeadBlendApplyToTransform(FTransform InTransform, FTransform InTransformTransition, FSpringTransformVelocity InVelocityTransition, float TimeSinceTransition, float BlendTime, float SmoothingTime)
+{
+	return FTransform(
+		DeadBlendApplyToQuat(InTransform.GetRotation(), InTransformTransition.GetRotation(), InVelocityTransition.AngularVelocity, TimeSinceTransition, BlendTime, SmoothingTime),
+		DeadBlendApplyToVector(InTransform.GetLocation(), InTransformTransition.GetLocation(), InVelocityTransition.LinearVelocity, TimeSinceTransition, BlendTime, SmoothingTime),
+		DeadBlendApplyToScale(InTransform.GetScale3D(), InTransformTransition.GetScale3D(), InVelocityTransition.ScalarVelocity, TimeSinceTransition, BlendTime, SmoothingTime));
+}
+
+void UBlueprintSpringMathLibrary::DeadBlendTransitionTransform(FTransform& InOutTransformTransition, FSpringTransformVelocity& InOutVelocityTransition, FTransform InSrcTransform, FSpringTransformVelocity InSrcVelocity, float TimeSinceTransition, float BlendTime, float SmoothingTime)
+{
+	FVector InOutLocation = InOutTransformTransition.GetLocation();
+	FQuat InOutRotation = InOutTransformTransition.GetRotation();
+	FVector InOutScale = InOutTransformTransition.GetScale3D();
+
+	DeadBlendTransitionVector(InOutLocation, InOutVelocityTransition.LinearVelocity, InSrcTransform.GetLocation(), InSrcVelocity.LinearVelocity, TimeSinceTransition, BlendTime, SmoothingTime);
+	DeadBlendTransitionQuat(InOutRotation, InOutVelocityTransition.AngularVelocity, InSrcTransform.GetRotation(), InSrcVelocity.AngularVelocity, TimeSinceTransition, BlendTime, SmoothingTime);
+	DeadBlendTransitionScale(InOutScale, InOutVelocityTransition.ScalarVelocity, InSrcTransform.GetScale3D(), InSrcVelocity.ScalarVelocity, TimeSinceTransition, BlendTime, SmoothingTime);
+
+	InOutTransformTransition = FTransform(InOutRotation, InOutLocation, InOutScale);
+}
\ No newline at end of file
diff --git a/Engine/Source/Runtime/Engine/Public/Animation/SpringMath.h b/Engine/Source/Runtime/Engine/Public/Animation/SpringMath.h
--- a/Engine/Source/Runtime/Engine/Public/Animation/SpringMath.h
+++ b/Engine/Source/Runtime/Engine/Public/Animation/SpringMath.h
@@ -7,7 +7,11 @@
 #include "Math/Quat.h"
 
 // Collection of useful spring methods which can be used for damping, simulating characters etc.
+// 
 // Reference https://theorangeduck.com/page/spring-roll-call
+//		     https://theorangeduck.com/page/scalar-velocity
+//			 https://theorangeduck.com/page/dead-blending
+//			 https://theorangeduck.com/page/propagating-velocities-through-animation-systems
 
 struct SpringMath
 {
@@ -16,7 +20,46 @@ private:
 	{
 		return 4.0f / FMath::Max(SmoothingTime, UE_KINDA_SMALL_NUMBER);
 	}
-	
+
+	/** Divides one scale by another, clamping the Rhs to be greater than some minimum */
+	static inline FVector ScaleDivMax(FVector Lhs, FVector Rhs)
+	{
+		return FVector(
+			Lhs.X / FMath::Max(Rhs.X, UE_SMALL_NUMBER),
+			Lhs.Y / FMath::Max(Rhs.Y, UE_SMALL_NUMBER),
+			Lhs.Z / FMath::Max(Rhs.Z, UE_SMALL_NUMBER));
+	}
+
+	/** Takes the log of a scale, clamping the scale to ensure it is greater than zero */
+	static inline FVector ScaleLog(FVector V)
+	{
+		return FVector(
+			FMath::Loge(FMath::Max(V.X, UE_SMALL_NUMBER)),
+			FMath::Loge(FMath::Max(V.Y, UE_SMALL_NUMBER)),
+			FMath::Loge(FMath::Max(V.Z, UE_SMALL_NUMBER)));
+	}
+
+	/** Computes the exponent of a scale, ensuring the value is not too large to prevent it exploding */
+	static inline FVector ScaleExp(FVector V, double Max = 10.0)
+	{
+		return FVector(
+			FMath::Exp(FMath::Min(V.X, Max)),
+			FMath::Exp(FMath::Min(V.Y, Max)),
+			FMath::Exp(FMath::Min(V.Z, Max)));
+	}
+
+	/** Performs an Eerp on scales (interpolation in the multiplicative space) */
+	static inline FVector ScaleEerp(FVector A, FVector B, double Alpha)
+	{
+		return ScaleExp(FMath::Lerp(ScaleLog(A), ScaleLog(B), Alpha));
+	}
+
+	/** Gets the version of angle A which is closest to reference angle B */
+	static inline float AngleGetShortestArcWith(float A, float B)
+	{
+		return FMath::FindDeltaAngleRadians(B, A) + B;
+	}
+
 public:
 	/** Convert a smoothing time to a half life
 	 * 
@@ -67,7 +110,7 @@ public:
 	 * 
 	 * @tparam T The type to be damped 
 	 * @param InOutX The value to be damped
-	 * @param InOutV The speed of the value to be damped
+	 * @param InOutV The velocity of the value to be damped
 	 * @param TargetX The goal to damp towards
 	 * @param SmoothingTime The smoothing time to use for the spring. Longer times lead to more damped behaviour. A time of 0 will snap InOutX to TargetX
 	 * @param DeltaTime Timestep in seconds
@@ -93,7 +136,7 @@ public:
 	/** Specialized angle version of CriticalSpringDamper that handles angle wrapping.
 	 *
 	 * @param InOutAngleRadians The value to be damped
-	 * @param InOutAngularVelocityRadians The speed of the value to be damped
+	 * @param InOutAngularVelocityRadians The velocity of the value to be damped
 	 * @param TargetAngleRadians The goal to damp towards
 	 * @param SmoothingTime The smoothing time to use for the spring. Longer times lead to more damped behaviour. A time of 0 will snap InOutAngleRadians to TargetAngleRadians
 	 * @param DeltaTime Timestep in seconds
@@ -138,9 +181,7 @@ public:
 		float Y = SmoothingTimeToDamping(SmoothingTime) / 2.0f;
 		PRAGMA_ENABLE_EXPERIMENTAL_WARNINGS
 
-		FQuat Diff = InOutRotation * TargetRotation.Inverse();
-		Diff.EnforceShortestArcWith(FQuat::Identity);
-		FVector J0 = Diff.ToRotationVector();
+		FVector J0 = (InOutRotation * TargetRotation.Inverse()).GetShortestArcWith(FQuat::Identity).ToRotationVector();
 		FVector J1 = InOutAngularVelocityRadians + J0 * Y;
 
 		float EyDt = FMath::InvExpApprox(Y * DeltaTime);
@@ -149,12 +190,147 @@ public:
 		InOutAngularVelocityRadians = EyDt * (InOutAngularVelocityRadians - J1 * Y * DeltaTime);
 	}
 
+	/** Specialized scale version of CriticalSpringDamper, uses FVector for scalar velocity
+	 *
+	 * @param InOutScale The value to be damped
+	 * @param InOutScalarVelocity The scales velocity of the scale
+	 * @param TargetScale The target scale to damp towards
+	 * @param SmoothingTime The smoothing time to use for the spring. Longer times lead to more damped behaviour. A time of 0 will snap InOutScale to TargetScale
+	 * @param DeltaTime
+	 */
+	UE_EXPERIMENTAL(5.8, "SpringMath is experimental")
+	static void CriticalSpringDamperScale(
+		FVector& InOutScale,
+		FVector& InOutScalarVelocity,
+		const FVector& TargetScale,
+		float SmoothingTime,
+		float DeltaTime)
+	{
+		PRAGMA_DISABLE_EXPERIMENTAL_WARNINGS
+		float Y = SmoothingTimeToDamping(SmoothingTime) / 2.0f;
+		PRAGMA_ENABLE_EXPERIMENTAL_WARNINGS
+
+		FVector J0 = ScaleLog(ScaleDivMax(InOutScale, TargetScale));
+		FVector J1 = InOutScalarVelocity + J0 * Y;
+
+		float EyDt = FMath::InvExpApprox(Y * DeltaTime);
+
+		InOutScale = ScaleExp(EyDt * (J0 + J1 * DeltaTime))* TargetScale;
+		InOutScalarVelocity = EyDt * (InOutScalarVelocity - J1 * Y * DeltaTime);
+	}
+
+	/** Simplified version of CriticalSpringDamper where the target is assumed to be 0. This decays the value InOutX towards zero
+	 * with the motion of a critically damped spring. The velocity of InOutX is stored in InOutV.
+	 *
+	 * @tparam T The type to be decayed
+	 * @param InOutX The value to be decayed
+	 * @param InOutV The velocity of the value to be decayed
+	 * @param SmoothingTime The smoothing time to use for the spring. Longer times lead to more damped behaviour. A time of 0 will snap InOutX to zero
+	 * @param DeltaTime Timestep in seconds
+	 */
+	template <typename T>
+	UE_EXPERIMENTAL(5.8, "SpringMath is experimental")
+	static void CriticalDecay(
+		T& InOutX,
+		T& InOutV,
+		float SmoothingTime,
+		float DeltaTime)
+	{
+		float Y = SmoothingTimeToDamping(SmoothingTime) / 2.0f;
+		T J0 = InOutX;
+		T J1 = InOutV + J0 * Y;
+		float EyDt = FMath::InvExpApprox(Y * DeltaTime);
+
+		InOutX = EyDt * (J0 + J1 * DeltaTime);
+		InOutV = EyDt * (InOutV - J1 * Y * DeltaTime);
+	}
+
+	/** Specialized angle version of CriticalDecay that handles angle wrapping.
+	 *
+	 * @param InOutAngleRadians The value to be decayed
+	 * @param InOutAngularVelocityRadians The velocity of the value to be decayed
+	 * @param SmoothingTime The smoothing time to use for the spring. Longer times lead to more damped behaviour. A time of 0 will snap InOutAngleRadians to TargetAngleRadians
+	 * @param DeltaTime Timestep in seconds
+	 */
+	UE_EXPERIMENTAL(5.8, "SpringMath is experimental")
+	static void CriticalDecayAngle(
+		float& InOutAngleRadians,
+		float& InOutAngularVelocityRadians,
+		float SmoothingTime,
+		float DeltaTime)
+	{
+		PRAGMA_DISABLE_EXPERIMENTAL_WARNINGS
+		float Y = SmoothingTimeToDamping(SmoothingTime) / 2.0f;
+		PRAGMA_ENABLE_EXPERIMENTAL_WARNINGS
+
+		float J0 = InOutAngleRadians;
+		float J1 = InOutAngularVelocityRadians + J0 * Y;
+		float EyDt = FMath::InvExpApprox(Y * DeltaTime);
+
+		InOutAngleRadians = EyDt * (J0 + J1 * DeltaTime);
+		InOutAngularVelocityRadians = EyDt * (InOutAngularVelocityRadians - J1 * Y * DeltaTime);
+	}
+
+	/** Specialized quaternion version of CriticalDecay, uses FVector for angular velocity
+	 *
+	 * @param InOutRotation The value to be decayed
+	 * @param InOutAngularVelocityRadians The angular velocity of the rotation in radians
+	 * @param SmoothingTime The smoothing time to use for the spring. Longer times lead to more damped behaviour. A time of 0 will snap InOutRotation to TargetRotation
+	 * @param DeltaTime
+	 */
+	UE_EXPERIMENTAL(5.8, "SpringMath is experimental")
+	static void CriticalDecayQuat(
+		FQuat& InOutRotation,
+		FVector& InOutAngularVelocityRadians,
+		float SmoothingTime,
+		float DeltaTime)
+	{
+		PRAGMA_DISABLE_EXPERIMENTAL_WARNINGS
+		float Y = SmoothingTimeToDamping(SmoothingTime) / 2.0f;
+		PRAGMA_ENABLE_EXPERIMENTAL_WARNINGS
+
+		FVector J0 = InOutRotation.GetShortestArcWith(FQuat::Identity).ToRotationVector();
+		FVector J1 = InOutAngularVelocityRadians + J0 * Y;
+
+		float EyDt = FMath::InvExpApprox(Y * DeltaTime);
+
+		InOutRotation = FQuat::MakeFromRotationVector(EyDt * (J0 + J1 * DeltaTime));
+		InOutAngularVelocityRadians = EyDt * (InOutAngularVelocityRadians - J1 * Y * DeltaTime);
+	}
+
+	/** Specialized scale version of CriticalDecay
+	 *
+	 * @param InOutScale The value to be decayed
+	 * @param InOutScalarVelocity The scalar velocity of the scale
+	 * @param SmoothingTime The smoothing time to use for the spring. Longer times lead to more damped behaviour. A time of 0 will snap InOutRotation to TargetRotation
+	 * @param DeltaTime
+	 */
+	UE_EXPERIMENTAL(5.8, "SpringMath is experimental")
+	static void CriticalDecayScale(
+		FVector& InOutScale,
+		FVector& InOutScalarVelocity,
+		float SmoothingTime,
+		float DeltaTime)
+	{
+		PRAGMA_DISABLE_EXPERIMENTAL_WARNINGS
+		float Y = SmoothingTimeToDamping(SmoothingTime) / 2.0f;
+		PRAGMA_ENABLE_EXPERIMENTAL_WARNINGS
+
+		FVector J0 = ScaleLog(InOutScale);
+		FVector J1 = InOutScalarVelocity + J0 * Y;
+
+		float EyDt = FMath::InvExpApprox(Y * DeltaTime);
+
+		InOutScale = ScaleExp(EyDt * (J0 + J1 * DeltaTime));
+		InOutScalarVelocity = EyDt * (InOutScalarVelocity - J1 * Y * DeltaTime);
+	}
+
 	/** A double spring damper that interpolates the value InOutX towards TargetX via some intermediate state - producing a more S-shaped curve. 
 	 * The velocity of InOutX is stored in InOutV.
 	 *
 	 * @tparam T The type to be damped
 	 * @param InOutX The value to be damped
-	 * @param InOutV The speed of the value to be damped
+	 * @param InOutV The velocity of the value to be damped
 	 * @param InOutXi The value of the intermediate state
 	 * @param InOutVi The velocity of the intermediate state
 	 * @param TargetX The goal to damp towards
@@ -163,14 +339,14 @@ public:
 	 */
 	template <typename T>
 	UE_EXPERIMENTAL(5.7, "SpringMath is experimental")
-		static void CriticalDoubleSpringDamper(
-			T& InOutX,
-			T& InOutV,
-			T& InOutXi,
-			T& InOutVi,
-			T TargetX,
-			float SmoothingTime,
-			float DeltaTime)
+	static void CriticalDoubleSpringDamper(
+		T& InOutX,
+		T& InOutV,
+		T& InOutXi,
+		T& InOutVi,
+		T TargetX,
+		float SmoothingTime,
+		float DeltaTime)
 	{
 		PRAGMA_DISABLE_EXPERIMENTAL_WARNINGS
 		CriticalSpringDamper(InOutXi, InOutVi, TargetX, SmoothingTime / 2.0f, DeltaTime);
@@ -181,7 +357,7 @@ public:
 	/** Specialized angle version of CriticalDoubleSpringDamper that handles angle wrapping.
 	 *
 	 * @param InOutAngleRadians The value to be damped
-	 * @param InOutAngularVelocityRadians The speed of the value to be damped
+	 * @param InOutAngularVelocityRadians The velocity of the value to be damped
 	 * @param InOutIntermediateAngleRadians The value of the intermediate state
 	 * @param InOutIntermediateAngularVelocityRadians The velocity of the intermediate state
 	 * @param TargetAngleRadians The goal to damp towards
@@ -189,14 +365,14 @@ public:
 	 * @param DeltaTime Timestep in seconds
 	 */
 	UE_EXPERIMENTAL(5.7, "SpringMath is experimental")
-		static void CriticalDoubleSpringDamperAngle(
-			float& InOutAngleRadians,
-			float& InOutAngularVelocityRadians,
-			float& InOutIntermediateAngleRadians,
-			float& InOutIntermediateAngularVelocityRadians,
-			float TargetAngleRadians,
-			float SmoothingTime,
-			float DeltaTime)
+	static void CriticalDoubleSpringDamperAngle(
+		float& InOutAngleRadians,
+		float& InOutAngularVelocityRadians,
+		float& InOutIntermediateAngleRadians,
+		float& InOutIntermediateAngularVelocityRadians,
+		float TargetAngleRadians,
+		float SmoothingTime,
+		float DeltaTime)
 	{
 		PRAGMA_DISABLE_EXPERIMENTAL_WARNINGS
 		CriticalSpringDamperAngle(InOutIntermediateAngleRadians, InOutIntermediateAngularVelocityRadians, TargetAngleRadians, SmoothingTime / 2.0f, DeltaTime);
@@ -215,14 +391,14 @@ public:
 	 * @param DeltaTime
 	 */
 	UE_EXPERIMENTAL(5.7, "SpringMath is experimental")
-		static void CriticalDoubleSpringDamperQuat(
-			FQuat& InOutRotation,
-			FVector& InOutAngularVelocityRadians,
-			FQuat& InOutIntermediateRotation,
-			FVector& InOutIntermediateAngularVelocityRadians,
-			const FQuat& TargetRotation,
-			float SmoothingTime,
-			float DeltaTime)
+	static void CriticalDoubleSpringDamperQuat(
+		FQuat& InOutRotation,
+		FVector& InOutAngularVelocityRadians,
+		FQuat& InOutIntermediateRotation,
+		FVector& InOutIntermediateAngularVelocityRadians,
+		const FQuat& TargetRotation,
+		float SmoothingTime,
+		float DeltaTime)
 	{
 		PRAGMA_DISABLE_EXPERIMENTAL_WARNINGS
 		CriticalSpringDamperQuat(InOutIntermediateRotation, InOutIntermediateAngularVelocityRadians, TargetRotation, SmoothingTime / 2.0f, DeltaTime);
@@ -230,6 +406,32 @@ public:
 		PRAGMA_ENABLE_EXPERIMENTAL_WARNINGS
 	}
 
+	/** Specialized scale version of CriticalDoubleSpringDamper
+	 *
+	 * @param InOutScale The value to be damped
+	 * @param InOutScalarVelocity The scalar velocity of the scale
+	 * @param InOutIntermediateScale The value of the intermediate state
+	 * @param InOutIntermediateScalarVelocity The velocity of the intermediate state
+	 * @param TargetScale The target scale to damp towards
+	 * @param SmoothingTime The smoothing time to use for the spring. Longer times lead to more damped behaviour. A time of 0 will snap InOutScale to TargetScale
+	 * @param DeltaTime
+	 */
+	UE_EXPERIMENTAL(5.8, "SpringMath is experimental")
+	static void CriticalDoubleSpringDamperScale(
+		FVector& InOutScale,
+		FVector& InOutScalarVelocity,
+		FVector& InOutIntermediateScale,
+		FVector& InOutIntermediateScalarVelocity,
+		const FVector& TargetScale,
+		float SmoothingTime,
+		float DeltaTime)
+	{
+		PRAGMA_DISABLE_EXPERIMENTAL_WARNINGS
+		CriticalSpringDamperScale(InOutIntermediateScale, InOutIntermediateScalarVelocity, TargetScale, SmoothingTime / 2.0f, DeltaTime);
+		CriticalSpringDamperScale(InOutScale, InOutScalarVelocity, InOutIntermediateScale, SmoothingTime / 2.0f, DeltaTime);
+		PRAGMA_ENABLE_EXPERIMENTAL_WARNINGS
+	}
+
 	/** A velocity spring will damp towards a target that follows a fixed linear target velocity, allowing control of the interpolation speed
 	 * while still giving a smoothed behavior. A SmoothingTime of 0 will give a linear interpolation between X and TargetX
 	 * 
@@ -585,14 +787,14 @@ public:
 	 * @param InOutRotation			The value to be smoothed
 	 * @param InTargetRotation		The target to smooth towards
 	 * @param InDeltaTime			Time interval
-	 * @param InSmoothingTime		Timescale over which to smooth. Larger values result in more smoothed behaviour. Can be zero.
+	 * @param SmoothingTime			Timescale over which to smooth. Larger values result in more smoothed behaviour. Can be zero.
 	 */
 	UE_EXPERIMENTAL(5.7, "SpringMath is experimental")
-	static constexpr void ExponentialSmoothingApproxQuat(FQuat& InOutRotation, const FQuat& InTargetRotation, const float InDeltaTime, const float InSmoothingTime)
+	static constexpr void ExponentialSmoothingApproxQuat(FQuat& InOutRotation, const FQuat& InTargetRotation, const float InDeltaTime, const float SmoothingTime)
 	{
-		if (InSmoothingTime > UE_KINDA_SMALL_NUMBER)
+		if (SmoothingTime > UE_KINDA_SMALL_NUMBER)
 		{
-			InOutRotation = FQuat::Slerp(InOutRotation, InTargetRotation, 1.0f - FMath::InvExpApprox(InDeltaTime / InSmoothingTime));
+			InOutRotation = FQuat::Slerp(InOutRotation, InTargetRotation, 1.0f - FMath::InvExpApprox(InDeltaTime / SmoothingTime));
 		}
 		else
 		{
@@ -606,18 +808,947 @@ public:
 	 * @param InOutAngleRadians		The angle to be smoothed
 	 * @param InTargetAngleRadians	The target to smooth towards
 	 * @param InDeltaTime			Time interval
-	 * @param InSmoothingTime		Timescale over which to smooth. Larger values result in more smoothed behaviour. Can be zero.
+	 * @param SmoothingTime			Timescale over which to smooth. Larger values result in more smoothed behaviour. Can be zero.
 	 */
 	UE_EXPERIMENTAL(5.7, "SpringMath is experimental")
-	static constexpr void ExponentialSmoothingApproxAngle(float& InOutAngleRadians, const float& InTargetAngleRadians, const float InDeltaTime, const float InSmoothingTime)
+	static constexpr void ExponentialSmoothingApproxAngle(float& InOutAngleRadians, const float& InTargetAngleRadians, const float InDeltaTime, const float SmoothingTime)
 	{
-		if (InSmoothingTime > UE_KINDA_SMALL_NUMBER)
+		if (SmoothingTime > UE_KINDA_SMALL_NUMBER)
 		{
-			InOutAngleRadians += FMath::FindDeltaAngleRadians(InOutAngleRadians, InTargetAngleRadians) * (1.0f - FMath::InvExpApprox(InDeltaTime / InSmoothingTime));
+			InOutAngleRadians += FMath::FindDeltaAngleRadians(InOutAngleRadians, InTargetAngleRadians) * (1.0f - FMath::InvExpApprox(InDeltaTime / SmoothingTime));
 		}
 		else
 		{
 			InOutAngleRadians = InTargetAngleRadians;
 		}
 	}
+
+	/** Specialized scale damper, similar to FMath::ExponentialSmoothingApprox but for scales.
+	 *  Smooths a value using exponential damping towards a target.
+	 *
+	 * @param InOutScale			The value to be smoothed
+	 * @param InTargetScale			The target to smooth towards
+	 * @param InDeltaTime			Time interval
+	 * @param SmoothingTime			Timescale over which to smooth. Larger values result in more smoothed behaviour. Can be zero.
+	 */
+	UE_EXPERIMENTAL(5.8, "SpringMath is experimental")
+	static constexpr void ExponentialSmoothingApproxScale(FVector& InOutScale, const FVector& InTargetScale, const float InDeltaTime, const float SmoothingTime)
+	{
+		if (SmoothingTime > UE_KINDA_SMALL_NUMBER)
+		{
+			InOutScale = ScaleEerp(InOutScale, InTargetScale, 1.0f - FMath::InvExpApprox(InDeltaTime / SmoothingTime));
+		}
+		else
+		{
+			InOutScale = InTargetScale;
+		}
+	}
+
+	/** Tracks the velocity of the input value InValue via finite difference, updating InOutValue and InOutVelocity in place.
+	 *
+	 * @param InOutValue			Updates to be the same as InValue
+	 * @param InOutVelocity			Updates to be the velocity of InValue computed via finite difference
+	 * @param InValue				Input value
+	 * @param InDeltaTime			Time interval
+	 */
+	template <typename TValue>
+	UE_EXPERIMENTAL(5.8, "SpringMath is experimental")
+	static constexpr void TrackVelocity(TValue& InOutValue, TValue& InOutVelocity, const TValue& InValue, float DeltaTime)
+	{
+		InOutVelocity = DeltaTime > UE_SMALL_NUMBER ? (InValue - InOutValue) / DeltaTime : TValue(0.0);
+		InOutValue = InValue;
+	}
+
+	/** Tracks the velocity of the input rotation InValue via finite difference, updating InOutValue and InOutVelocity in place.
+	 *
+	 * @param InOutValue			Updates to be the same as InValue
+	 * @param InOutVelocity			Updates to be the velocity of InValue computed via finite difference (in rad/s)
+	 * @param InValue				Input rotation
+	 * @param InDeltaTime			Time interval
+	 */
+	UE_EXPERIMENTAL(5.8, "SpringMath is experimental")
+	static constexpr void TrackVelocityQuat(FQuat& InOutValue, FVector& InOutVelocity, FQuat InValue, float DeltaTime)
+	{
+		InOutVelocity = DeltaTime > UE_SMALL_NUMBER ? (InValue * InOutValue.Inverse()).GetShortestArcWith(FQuat::Identity).ToRotationVector() / DeltaTime : FVector::ZeroVector;
+		InOutValue = InValue;
+	}
+
+	/** Tracks the velocity of the input angle InValue via finite difference, updating InOutValue and InOutVelocity in place.
+	 *
+	 * @param InOutValue			Updates to be the same as InValue
+	 * @param InOutVelocity			Updates to be the velocity of InValue computed via finite difference (in rad/s)
+	 * @param InValue				Input angle
+	 * @param InDeltaTime			Time interval
+	 */
+	UE_EXPERIMENTAL(5.8, "SpringMath is experimental")
+	static constexpr void TrackVelocityAngle(float& InOutValue, float& InOutVelocity, float InValue, float DeltaTime)
+	{
+		InOutVelocity = DeltaTime > UE_SMALL_NUMBER ? FMath::FindDeltaAngleRadians(InOutValue, InValue) / DeltaTime : 0.0f;
+		InOutValue = InValue;
+	}
+
+	/** Tracks the velocity of the input scale InValue via finite difference, updating InOutValue and InOutVelocity in place.
+	 *
+	 * @param InOutValue			Updates to be the same as InValue
+	 * @param InOutVelocity			Updates to be the velocity of InValue computed via finite difference
+	 * @param InValue				Input scale
+	 * @param InDeltaTime			Time interval
+	 */
+	UE_EXPERIMENTAL(5.8, "SpringMath is experimental")
+	static constexpr void TrackVelocityScale(FVector& InOutValue, FVector& InOutVelocity, FVector InValue, float DeltaTime)
+	{
+		InOutVelocity = DeltaTime > UE_SMALL_NUMBER ? ScaleLog(ScaleDivMax(InValue, InOutValue)) / DeltaTime : FVector::ZeroVector;
+		InOutValue = InValue;
+	}
+
+	/** Applies a spring inertialization to the input value (and velocity), updating them in place by adding the appropriate offset given the time
+	 * since transition.
+	 *
+	 * @param InOutValue			The value to apply the inertialization to
+	 * @param InOutVelocity			The velocity to apply the inertialization to
+	 * @param InValueOffset			Offset value at the time of transition
+	 * @param InVelocityOffset		Offset velocity at the time of transition
+	 * @param TimeSinceTransition	Time since the last transition
+	 * @param SmoothingTime			Timescale over which to apply the offset. Larger values result in longer blend times. Can be zero.
+	 */
+	template <typename TValue>
+	UE_EXPERIMENTAL(5.8, "SpringMath is experimental")
+	static void SpringInertializeApply(
+		TValue& InOutValue,
+		TValue& InOutVelocity,
+		const TValue& InValueOffset,
+		const TValue& InVelocityOffset,
+		float TimeSinceTransition,
+		float SmoothingTime)
+	{
+		TValue ValueOffset = InValueOffset;
+		TValue VelocityOffset = InVelocityOffset;
+
+		PRAGMA_DISABLE_EXPERIMENTAL_WARNINGS
+		CriticalDecay<TValue>(ValueOffset, VelocityOffset, SmoothingTime, TimeSinceTransition);
+		PRAGMA_ENABLE_EXPERIMENTAL_WARNINGS
+
+		InOutValue += ValueOffset;
+		InOutVelocity += VelocityOffset;
+	}
+
+	/** Transitions a spring inertialization, computing a new offset value.
+	 *
+	 * @param InOutValueOffset		Offset value to update
+	 * @param InOutVelocityOffset	Offset velocity to update
+	 * @param InSrcValue			Value of the source (value transitioning from)
+	 * @param InSrcVelocity			Velocity of the source (velocity transitioning from)
+	 * @param InDstValue			Value of the destination (value transitioning to)
+	 * @param InDstVelocity			Velocity of the destination (velocity transitioning to)
+	 * @param TimeSinceTransition	Time since the last transition
+	 * @param SmoothingTime			Timescale over which to apply the offset. Larger values result in longer blend times. Can be zero.
+	 */
+	template <typename TValue>
+	UE_EXPERIMENTAL(5.8, "SpringMath is experimental")
+	static void SpringInertializeTransition(
+		TValue& InOutValueOffset,
+		TValue& InOutVelocityOffset,
+		const TValue& InSrcValue,
+		const TValue& InSrcVelocity,
+		const TValue& InDstValue,
+		const TValue& InDstVelocity,
+		float TimeSinceTransition,
+		float SmoothingTime)
+	{
+		TValue ValueOffset = InOutValueOffset;
+		TValue VelocityOffset = InOutVelocityOffset;
+
+		PRAGMA_DISABLE_EXPERIMENTAL_WARNINGS
+		CriticalDecay<TValue>(ValueOffset, VelocityOffset, SmoothingTime, TimeSinceTransition);
+		PRAGMA_ENABLE_EXPERIMENTAL_WARNINGS
+
+		InOutValueOffset = InSrcValue + ValueOffset - InDstValue;
+		InOutVelocityOffset = InSrcVelocity + VelocityOffset - InDstVelocity;
+	}
+
+	/** Specialization of SpringInertializeApply for quaternions. Stores velocities in rad/s.
+	 *
+	 * @param InOutValue			The value to apply the inertialization to
+	 * @param InOutVelocity			The velocity to apply the inertialization to
+	 * @param InValueOffset			Offset value at the time of transition
+	 * @param InVelocityOffset		Offset velocity at the time of transition
+	 * @param TimeSinceTransition	Time since the last transition
+	 * @param SmoothingTime			Timescale over which to apply the offset. Larger values result in longer blend times. Can be zero.
+	 */
+	UE_EXPERIMENTAL(5.8, "SpringMath is experimental")
+	static void SpringInertializeApplyQuat(
+		FQuat& InOutRotation,
+		FVector& InOutAngularVelocity,
+		const FQuat& InRotationOffset,
+		const FVector& InAngularVelocityOffset,
+		float TimeSinceTransition,
+		float SmoothingTime)
+	{
+		FQuat RotationOffset = InRotationOffset;
+		FVector AngularVelocityOffset = InAngularVelocityOffset;
+
+		PRAGMA_DISABLE_EXPERIMENTAL_WARNINGS
+		CriticalDecayQuat(RotationOffset, AngularVelocityOffset, SmoothingTime, TimeSinceTransition);
+		PRAGMA_ENABLE_EXPERIMENTAL_WARNINGS
+
+		InOutRotation = RotationOffset * InOutRotation;
+		InOutAngularVelocity += AngularVelocityOffset;
+	}
+
+	/** Specialization of SpringInertializeTransition for quaternions. Stores velocities in rad/s.
+	 *
+	 * @param InOutValueOffset		Offset value to update
+	 * @param InOutVelocityOffset	Offset velocity to update
+	 * @param InSrcValue			Value of the source (value transitioning from)
+	 * @param InSrcVelocity			Velocity of the source (velocity transitioning from)
+	 * @param InDstValue			Value of the destination (value transitioning to)
+	 * @param InDstVelocity			Velocity of the destination (velocity transitioning to)
+	 * @param TimeSinceTransition	Time since the last transition
+	 * @param SmoothingTime			Timescale over which to apply the offset. Larger values result in longer blend times. Can be zero.
+	 */
+	UE_EXPERIMENTAL(5.8, "SpringMath is experimental")
+	static void SpringInertializeTransitionQuat(
+		FQuat& InOutRotationOffset,
+		FVector& InOutAngularVelocityOffset,
+		const FQuat& InSrcRotation,
+		const FVector& InSrcAngularVelocity,
+		const FQuat& InDstRotation,
+		const FVector& InDstAngularVelocity,
+		float TimeSinceTransition,
+		float SmoothingTime)
+	{
+		FQuat RotationOffset = InOutRotationOffset;
+		FVector AngularVelocityOffset = InOutAngularVelocityOffset;
+
+		PRAGMA_DISABLE_EXPERIMENTAL_WARNINGS
+		CriticalDecayQuat(RotationOffset, AngularVelocityOffset, SmoothingTime, TimeSinceTransition);
+		PRAGMA_ENABLE_EXPERIMENTAL_WARNINGS
+
+		InOutRotationOffset = InSrcRotation * RotationOffset * InDstRotation.Inverse();
+		InOutAngularVelocityOffset = InSrcAngularVelocity + AngularVelocityOffset - InDstAngularVelocity;
+	}
+
+	/** Specialization of SpringInertializeApply for angles. Stores velocities in rad/s.
+	 *
+	 * @param InOutValue			The value to apply the inertialization to
+	 * @param InOutVelocity			The velocity to apply the inertialization to
+	 * @param InValueOffset			Offset value at the time of transition
+	 * @param InVelocityOffset		Offset velocity at the time of transition
+	 * @param TimeSinceTransition	Time since the last transition
+	 * @param SmoothingTime			Timescale over which to apply the offset. Larger values result in longer blend times. Can be zero.
+	 */
+	UE_EXPERIMENTAL(5.8, "SpringMath is experimental")
+	static void SpringInertializeApplyAngle(
+		float& InOutAngle,
+		float& InOutAngularVelocity,
+		float InAngleOffset,
+		float InAngularVelocityOffset,
+		float TimeSinceTransition,
+		float SmoothingTime)
+	{
+		float AngleOffset = InAngleOffset;
+		float AngularVelocityOffset = InAngularVelocityOffset;
+
+		PRAGMA_DISABLE_EXPERIMENTAL_WARNINGS
+		CriticalDecayAngle(AngleOffset, AngularVelocityOffset, SmoothingTime, TimeSinceTransition);
+		PRAGMA_ENABLE_EXPERIMENTAL_WARNINGS
+
+		InOutAngle += AngleOffset;
+		InOutAngularVelocity += AngularVelocityOffset;
+	}
+
+	/** Specialization of SpringInertializeTransition for angles. Stores velocities in rad/s.
+	 *
+	 * @param InOutValueOffset		Offset value to update
+	 * @param InOutVelocityOffset	Offset velocity to update
+	 * @param InSrcValue			Value of the source (value transitioning from)
+	 * @param InSrcVelocity			Velocity of the source (velocity transitioning from)
+	 * @param InDstValue			Value of the destination (value transitioning to)
+	 * @param InDstVelocity			Velocity of the destination (velocity transitioning to)
+	 * @param TimeSinceTransition	Time since the last transition
+	 * @param SmoothingTime			Timescale over which to apply the offset. Larger values result in longer blend times. Can be zero.
+	 */
+	UE_EXPERIMENTAL(5.8, "SpringMath is experimental")
+	static void SpringInertializeTransitionAngle(
+		float& InOutAngleOffset,
+		float& InOutAngularVelocityOffset,
+		float InSrcAngle,
+		float InSrcAngularVelocity,
+		float InDstAngle,
+		float InDstAngularVelocity,
+		float TimeSinceTransition,
+		float SmoothingTime)
+	{
+		float AngleOffset = InOutAngleOffset;
+		float AngularVelocityOffset = InOutAngularVelocityOffset;
+
+		PRAGMA_DISABLE_EXPERIMENTAL_WARNINGS
+		CriticalDecayAngle(AngleOffset, AngularVelocityOffset, SmoothingTime, TimeSinceTransition);
+		PRAGMA_ENABLE_EXPERIMENTAL_WARNINGS
+
+		InOutAngleOffset = FMath::FindDeltaAngleRadians(InDstAngle, InSrcAngle + AngleOffset);
+		InOutAngularVelocityOffset = InSrcAngularVelocity + AngularVelocityOffset - InDstAngularVelocity;
+	}
+
+	/** Specialization of SpringInertializeApply for scales.
+	 *
+	 * @param InOutValue			The value to apply the inertialization to
+	 * @param InOutVelocity			The velocity to apply the inertialization to
+	 * @param InValueOffset			Offset value at the time of transition
+	 * @param InVelocityOffset		Offset velocity at the time of transition
+	 * @param TimeSinceTransition	Time since the last transition
+	 * @param SmoothingTime			Timescale over which to apply the offset. Larger values result in longer blend times. Can be zero.
+	 */
+	UE_EXPERIMENTAL(5.8, "SpringMath is experimental")
+	static void SpringInertializeApplyScale(
+		FVector& InOutScale,
+		FVector& InOutScalarVelocity,
+		const FVector& InScaleOffset,
+		const FVector& InScalarVelocityOffset,
+		float TimeSinceTransition,
+		float SmoothingTime)
+	{
+		FVector ScaleOffset = InScaleOffset;
+		FVector ScalarVelocityOffset = InScalarVelocityOffset;
+
+		PRAGMA_DISABLE_EXPERIMENTAL_WARNINGS
+		CriticalDecayScale(ScaleOffset, ScalarVelocityOffset, SmoothingTime, TimeSinceTransition);
+		PRAGMA_ENABLE_EXPERIMENTAL_WARNINGS
+
+		InOutScale = ScaleOffset * InOutScale;
+		InOutScalarVelocity += ScalarVelocityOffset;
+	}
+
+	/** Specialization of SpringInertializeTransition for Scales
+	 *
+	 * @param InOutValueOffset		Offset value to update
+	 * @param InOutVelocityOffset	Offset velocity to update
+	 * @param InSrcValue			Value of the source (value transitioning from)
+	 * @param InSrcVelocity			Velocity of the source (velocity transitioning from)
+	 * @param InDstValue			Value of the destination (value transitioning to)
+	 * @param InDstVelocity			Velocity of the destination (velocity transitioning to)
+	 * @param TimeSinceTransition	Time since the last transition
+	 * @param SmoothingTime			Timescale over which to apply the offset. Larger values result in longer blend times. Can be zero.
+	 */
+	UE_EXPERIMENTAL(5.8, "SpringMath is experimental")
+	static void SpringInertializeTransitionScale(
+		FVector& InOutScaleOffset,
+		FVector& InOutScalarVelocityOffset,
+		const FVector& InSrcScale,
+		const FVector& InSrcScalarVelocity,
+		const FVector& InDstScale,
+		const FVector& InDstScalarVelocity,
+		float TimeSinceTransition,
+		float SmoothingTime)
+	{
+		FVector ScaleOffset = InOutScaleOffset;
+		FVector ScalarVelocityOffset = InOutScalarVelocityOffset;
+
+		PRAGMA_DISABLE_EXPERIMENTAL_WARNINGS
+		CriticalDecayScale(ScaleOffset, ScalarVelocityOffset, SmoothingTime, TimeSinceTransition);
+		PRAGMA_ENABLE_EXPERIMENTAL_WARNINGS
+
+		InOutScaleOffset = ScaleDivMax(InSrcScale * ScaleOffset, InDstScale);
+		InOutScalarVelocityOffset = InSrcScalarVelocity + ScalarVelocityOffset - InDstScalarVelocity;
+	}
+
+	/** Compute the cubic weights required to inertialize out an offset and velocity offset for a given time and blend time */
+	UE_EXPERIMENTAL(5.8, "SpringMath is experimental")
+	static inline void CubicDecayWeights(float& OutW0, float& OutW1, float& OutW2, float& OutW3, const float Time, const float BlendTime)
+	{
+		float NonZeroBlendTime = FMath::Max(BlendTime, UE_SMALL_NUMBER);
+		float T = FMath::Clamp(Time / NonZeroBlendTime, 0.0f, 1.0f);
+		OutW0 = 2.0f * T * T * T - 3.0f * T * T + 1.0f;
+		OutW1 = (T * T * T - 2.0f * T * T + T) * BlendTime;
+		OutW2 = (6.0f * T * T - 6.0f * T) / NonZeroBlendTime;
+		OutW3 = 3.0f * T * T - 4.0f * T + 1.0f;
+	}
+
+	/** Applies a cubic inertialization to the input value (and velocity), updating them in place by adding the appropriate offset given the time
+	 * since transition.
+	 *
+	 * @param InOutValue			The value to apply the inertialization to
+	 * @param InOutVelocity			The velocity to apply the inertialization to
+	 * @param InValueOffset			Offset value at the time of transition
+	 * @param InVelocityOffset		Offset velocity at the time of transition
+	 * @param TimeSinceTransition	Time since the last transition
+	 * @param BlendTime				Timescale over which to apply the offset. Larger values result in longer blend times. Can be zero.
+	 */
+	template <typename TValue>
+	UE_EXPERIMENTAL(5.8, "SpringMath is experimental")
+	static void CubicInertializeApply(
+		TValue& InOutValue,
+		TValue& InOutVelocity,
+		const TValue& InValueOffset,
+		const TValue& InVelocityOffset,
+		float TimeSinceTransition,
+		float BlendTime)
+	{
+		float W0 = 0.0f, W1 = 0.0f, W2 = 0.0f, W3 = 0.0f;
+		PRAGMA_DISABLE_EXPERIMENTAL_WARNINGS
+		CubicDecayWeights(W0, W1, W2, W3, TimeSinceTransition, BlendTime);
+		PRAGMA_ENABLE_EXPERIMENTAL_WARNINGS
+
+		InOutValue += W0 * InValueOffset + W1 * InVelocityOffset;
+		InOutVelocity += W2 * InValueOffset + W3 * InVelocityOffset;
+	}
+
+	/** Transitions a cubic inertialization, computing a new offset value.
+	 *
+	 * @param InOutValueOffset		Offset value to update
+	 * @param InOutVelocityOffset	Offset velocity to update
+	 * @param InSrcValue			Value of the source (value transitioning from)
+	 * @param InSrcVelocity			Velocity of the source (velocity transitioning from)
+	 * @param InDstValue			Value of the destination (value transitioning to)
+	 * @param InDstVelocity			Velocity of the destination (velocity transitioning to)
+	 * @param TimeSinceTransition	Time since the last transition
+	 * @param BlendTime				Timescale over which to apply the offset. Larger values result in longer blend times. Can be zero.
+	 */
+	template <typename TValue>
+	UE_EXPERIMENTAL(5.8, "SpringMath is experimental")
+	static void CubicInertializeTransition(
+		TValue& InOutValueOffset,
+		TValue& InOutVelocityOffset,
+		const TValue& InSrcValue,
+		const TValue& InSrcVelocity,
+		const TValue& InDstValue,
+		const TValue& InDstVelocity,
+		float TimeSinceTransition,
+		float BlendTime)
+	{
+		float W0 = 0.0f, W1 = 0.0f, W2 = 0.0f, W3 = 0.0f;
+		PRAGMA_DISABLE_EXPERIMENTAL_WARNINGS
+		CubicDecayWeights(W0, W1, W2, W3, TimeSinceTransition, BlendTime);
+		PRAGMA_ENABLE_EXPERIMENTAL_WARNINGS
+
+		TValue Value = InSrcValue + W0 * InOutValueOffset + W1 * InOutVelocityOffset;
+		TValue Velocity = InSrcVelocity + W2 * InOutValueOffset + W3 * InOutVelocityOffset;
+
+		InOutValueOffset = Value - InDstValue;
+		InOutVelocityOffset = Velocity - InDstVelocity;
+	}
+
+	/** Specialization of CubicInertializeApply for quaternions. Stores velocities in rad/s.
+	 *
+	 * @param InOutValue			The value to apply the inertialization to
+	 * @param InOutVelocity			The velocity to apply the inertialization to
+	 * @param InValueOffset			Offset value at the time of transition
+	 * @param InVelocityOffset		Offset velocity at the time of transition
+	 * @param TimeSinceTransition	Time since the last transition
+	 * @param BlendTime				Blend over which to apply the offset. Larger values result in longer blend times. Can be zero.
+	 */
+	UE_EXPERIMENTAL(5.8, "SpringMath is experimental")
+	static void CubicInertializeApplyQuat(
+		FQuat& InOutRotation,
+		FVector& InOutAngularVelocity,
+		const FQuat& InRotationOffset,
+		const FVector& InAngularVelocityOffset,
+		float TimeSinceTransition,
+		float BlendTime)
+	{
+		float W0 = 0.0f, W1 = 0.0f, W2 = 0.0f, W3 = 0.0f;
+		PRAGMA_DISABLE_EXPERIMENTAL_WARNINGS
+		CubicDecayWeights(W0, W1, W2, W3, TimeSinceTransition, BlendTime);
+		PRAGMA_ENABLE_EXPERIMENTAL_WARNINGS
+
+		InOutRotation = FQuat::MakeFromRotationVector(W0 * InRotationOffset.ToRotationVector() + W1 * InAngularVelocityOffset) * InOutRotation;
+		InOutAngularVelocity += W2 * InRotationOffset.ToRotationVector() + W3 * InAngularVelocityOffset;
+	}
+
+	/** Specialization of CubicInertializeTransition for quaternions. Stores velocities in rad/s.
+	 *
+	 * @param InOutValueOffset		Offset value to update
+	 * @param InOutVelocityOffset	Offset velocity to update
+	 * @param InSrcValue			Value of the source (value transitioning from)
+	 * @param InSrcVelocity			Velocity of the source (velocity transitioning from)
+	 * @param InDstValue			Value of the destination (value transitioning to)
+	 * @param InDstVelocity			Velocity of the destination (velocity transitioning to)
+	 * @param TimeSinceTransition	Time since the last transition
+	 * @param BlendTime				Timescale over which to apply the offset. Larger values result in longer blend times. Can be zero.
+	 */
+	UE_EXPERIMENTAL(5.8, "SpringMath is experimental")
+	static void CubicInertializeTransitionQuat(
+		FQuat& InOutRotationOffset,
+		FVector& InOutAngularVelocityOffset,
+		const FQuat& InSrcRotation,
+		const FVector& InSrcAngularVelocity,
+		const FQuat& InDstRotation,
+		const FVector& InDstAngularVelocity,
+		float TimeSinceTransition,
+		float BlendTime)
+	{
+		float W0 = 0.0f, W1 = 0.0f, W2 = 0.0f, W3 = 0.0f;
+		PRAGMA_DISABLE_EXPERIMENTAL_WARNINGS
+		CubicDecayWeights(W0, W1, W2, W3, TimeSinceTransition, BlendTime);
+		PRAGMA_ENABLE_EXPERIMENTAL_WARNINGS
+
+		FQuat Value = FQuat::MakeFromRotationVector(W0 * InOutRotationOffset.ToRotationVector() + W1 * InOutAngularVelocityOffset) * InSrcRotation;
+		FVector Velocity = InSrcAngularVelocity + W2 * InOutRotationOffset.ToRotationVector() + W3 * InOutAngularVelocityOffset;
+
+		InOutRotationOffset = (Value * InDstRotation.Inverse()).GetShortestArcWith(FQuat::Identity);
+		InOutAngularVelocityOffset = Velocity - InDstAngularVelocity;
+	}
+
+	/** Specialization of CubicInertializeApply for angles. Stores velocities in rad/s.
+	 *
+	 * @param InOutValue			The value to apply the inertialization to
+	 * @param InOutVelocity			The velocity to apply the inertialization to
+	 * @param InValueOffset			Offset value at the time of transition
+	 * @param InVelocityOffset		Offset velocity at the time of transition
+	 * @param TimeSinceTransition	Time since the last transition
+	 * @param BlendTime				Blend over which to apply the offset. Larger values result in longer blend times. Can be zero.
+	 */
+	UE_EXPERIMENTAL(5.8, "SpringMath is experimental")
+	static void CubicInertializeApplyAngle(
+		float& InOutAngle,
+		float& InOutAngularVelocity,
+		float InAngleOffset,
+		float InAngularVelocityOffset,
+		float TimeSinceTransition,
+		float BlendTime)
+	{
+		float W0 = 0.0f, W1 = 0.0f, W2 = 0.0f, W3 = 0.0f;
+		PRAGMA_DISABLE_EXPERIMENTAL_WARNINGS
+		CubicDecayWeights(W0, W1, W2, W3, TimeSinceTransition, BlendTime);
+		PRAGMA_ENABLE_EXPERIMENTAL_WARNINGS
+
+		InOutAngle += W0 * InAngleOffset + W1 * InAngularVelocityOffset;
+		InOutAngularVelocity += W2 * InAngleOffset + W3 * InAngularVelocityOffset;
+	}
+
+	/** Specialization of CubicInertializeTransition for angles. Stores velocities in rad/s.
+	 *
+	 * @param InOutValueOffset		Offset value to update
+	 * @param InOutVelocityOffset	Offset velocity to update
+	 * @param InSrcValue			Value of the source (value transitioning from)
+	 * @param InSrcVelocity			Velocity of the source (velocity transitioning from)
+	 * @param InDstValue			Value of the destination (value transitioning to)
+	 * @param InDstVelocity			Velocity of the destination (velocity transitioning to)
+	 * @param TimeSinceTransition	Time since the last transition
+	 * @param BlendTime				Timescale over which to apply the offset. Larger values result in longer blend times. Can be zero.
+	 */
+	UE_EXPERIMENTAL(5.8, "SpringMath is experimental")
+	static void CubicInertializeTransitionAngle(
+		float& InOutAngleOffset,
+		float& InOutAngularVelocityOffset,
+		float InSrcAngle,
+		float InSrcAngularVelocity,
+		float InDstAngle,
+		float InDstAngularVelocity,
+		float TimeSinceTransition,
+		float BlendTime)
+	{
+		float W0 = 0.0f, W1 = 0.0f, W2 = 0.0f, W3 = 0.0f;
+		PRAGMA_DISABLE_EXPERIMENTAL_WARNINGS
+		CubicDecayWeights(W0, W1, W2, W3, TimeSinceTransition, BlendTime);
+		PRAGMA_ENABLE_EXPERIMENTAL_WARNINGS
+
+		float Value = InSrcAngle + W0 * InOutAngleOffset + W1 * InOutAngularVelocityOffset;
+		float Velocity = InSrcAngularVelocity + W2 * InOutAngleOffset + W3 * InOutAngularVelocityOffset;
+
+		InOutAngleOffset = FMath::FindDeltaAngleRadians(InDstAngle, Value);
+		InOutAngularVelocityOffset = Velocity - InDstAngularVelocity;
+	}
+
+	/** Specialization of CubicInertializeApply for scales.
+	 *
+	 * @param InOutValue			The value to apply the inertialization to
+	 * @param InOutVelocity			The velocity to apply the inertialization to
+	 * @param InValueOffset			Offset value at the time of transition
+	 * @param InVelocityOffset		Offset velocity at the time of transition
+	 * @param TimeSinceTransition	Time since the last transition
+	 * @param BlendTime				Blend over which to apply the offset. Larger values result in longer blend times. Can be zero.
+	 */
+	UE_EXPERIMENTAL(5.8, "SpringMath is experimental")
+	static void CubicInertializeApplyScale(
+		FVector& InOutScale,
+		FVector& InOutScalarVelocity,
+		const FVector& InScaleOffset,
+		const FVector& InScalarVelocityOffset,
+		float TimeSinceTransition,
+		float BlendTime)
+	{
+		float W0 = 0.0f, W1 = 0.0f, W2 = 0.0f, W3 = 0.0f;
+		PRAGMA_DISABLE_EXPERIMENTAL_WARNINGS
+		CubicDecayWeights(W0, W1, W2, W3, TimeSinceTransition, BlendTime);
+		PRAGMA_ENABLE_EXPERIMENTAL_WARNINGS
+
+		InOutScale = ScaleExp(W0 * ScaleLog(InScaleOffset) + W1 * InScalarVelocityOffset) * InOutScale;
+		InOutScalarVelocity += W2 * ScaleLog(InScaleOffset) + W3 * InScalarVelocityOffset;
+	}
+
+	/** Specialization of CubicInertializeTransition for scales.
+	 *
+	 * @param InOutValueOffset		Offset value to update
+	 * @param InOutVelocityOffset	Offset velocity to update
+	 * @param InSrcValue			Value of the source (value transitioning from)
+	 * @param InSrcVelocity			Velocity of the source (velocity transitioning from)
+	 * @param InDstValue			Value of the destination (value transitioning to)
+	 * @param InDstVelocity			Velocity of the destination (velocity transitioning to)
+	 * @param TimeSinceTransition	Time since the last transition
+	 * @param BlendTime				Timescale over which to apply the offset. Larger values result in longer blend times. Can be zero.
+	 */
+	UE_EXPERIMENTAL(5.8, "SpringMath is experimental")
+	static void CubicInertializeTransitionScale(
+		FVector& InOutScaleOffset,
+		FVector& InOutScalarVelocityOffset,
+		const FVector& InSrcScale,
+		const FVector& InSrcScalarVelocity,
+		const FVector& InDstScale,
+		const FVector& InDstScalarVelocity,
+		float TimeSinceTransition,
+		float BlendTime)
+	{
+		float W0 = 0.0f, W1 = 0.0f, W2 = 0.0f, W3 = 0.0f;
+		PRAGMA_DISABLE_EXPERIMENTAL_WARNINGS
+		CubicDecayWeights(W0, W1, W2, W3, TimeSinceTransition, BlendTime);
+		PRAGMA_ENABLE_EXPERIMENTAL_WARNINGS
+
+		FVector Value = ScaleExp(W0 * ScaleLog(InOutScaleOffset) + W1 * InOutScalarVelocityOffset) * InSrcScale;
+		FVector Velocity = InSrcScalarVelocity + W2 * ScaleLog(InOutScaleOffset) + W3 * InOutScalarVelocityOffset;
+
+		InOutScaleOffset = ScaleDivMax(Value, InDstScale);
+		InOutScalarVelocityOffset = Velocity - InDstScalarVelocity;
+	}
+
+	/** Extrapolates a value and velocity in place assuming some decay rate on the velocity.
+	 *
+	 * @param InOutValue			Value to extrapolate from
+	 * @param InOutVelocity			Velocity to extrapolate with
+	 * @param Time					Amount of time to extrapolate forward
+	 * @param SmoothingTime			Decay rate of the extrapolated velocity. Smaller values make the velocity decay faster.
+	 */
+	template <typename TValue>
+	UE_EXPERIMENTAL(5.8, "SpringMath is experimental")
+	static inline void DeadBlendExtrapolate(
+		TValue& InOutValue,
+		TValue& InOutVelocity,
+		const float Time,
+		const float SmoothingTime)
+	{
+		float NonZeroSmoothingTime = FMath::Max(SmoothingTime, UE_SMALL_NUMBER);
+		float Decay = FMath::InvExpApprox(Time / NonZeroSmoothingTime);
+		InOutValue = InOutValue + ((InOutVelocity * NonZeroSmoothingTime) * (1.0f - Decay));
+		InOutVelocity = InOutVelocity * Decay;
+	}
+
+	/** Specialization of DeadBlendExtrapolate for quaternions. Angular velocities are stored in rad/s.
+	 *
+	 * @param InOutValue			Value to extrapolate from
+	 * @param InOutVelocity			Velocity to extrapolate with
+	 * @param Time					Amount of time to extrapolate forward
+	 * @param SmoothingTime			Decay rate of the extrapolated velocity. Smaller values make the velocity decay faster.
+	 */
+	UE_EXPERIMENTAL(5.8, "SpringMath is experimental")
+	static inline void DeadBlendExtrapolateQuat(
+		FQuat& InOutValue,
+		FVector& InOutVelocity,
+		const float Time,
+		const float SmoothingTime)
+	{
+		float NonZeroSmoothingTime = FMath::Max(SmoothingTime, UE_SMALL_NUMBER);
+		float Decay = FMath::InvExpApprox(Time / NonZeroSmoothingTime);
+		InOutValue = InOutValue * FQuat::MakeFromRotationVector((InOutVelocity * NonZeroSmoothingTime) * (1.0f - Decay));
+		InOutVelocity = InOutVelocity * Decay;
+	}
+
+	/** Specialization of DeadBlendExtrapolate for scales.
+	 *
+	 * @param InOutValue			Value to extrapolate from
+	 * @param InOutVelocity			Velocity to extrapolate with
+	 * @param Time					Amount of time to extrapolate forward
+	 * @param SmoothingTime			Decay rate of the extrapolated velocity. Smaller values make the velocity decay faster.
+	 */
+	UE_EXPERIMENTAL(5.8, "SpringMath is experimental")
+	static inline void DeadBlendExtrapolateScale(
+		FVector& InOutValue,
+		FVector& InOutVelocity,
+		const float Time,
+		const float SmoothingTime)
+	{
+		float NonZeroSmoothingTime = FMath::Max(SmoothingTime, UE_SMALL_NUMBER);
+		float Decay = FMath::InvExpApprox(Time / NonZeroSmoothingTime);
+		InOutValue = InOutValue * ScaleExp((InOutVelocity * NonZeroSmoothingTime) * (1.0f - Decay));
+		InOutVelocity = InOutVelocity * Decay;
+	}
+
+	/** Applies a dead blend to the input value (and velocity), updating them in place by blending them with an extrapolation of the value and velocity at the
+	 * point of the last transition. 
+	 *
+	 * @param InOutValue			The value to apply the inertialization to
+	 * @param InOutVelocity			The velocity to apply the inertialization to
+	 * @param InValueTransition		Value at the time of the last transition
+	 * @param InVelocityTransition	Velocity at the time of the last transition
+	 * @param TimeSinceTransition	Time since the last transition
+	 * @param BlendTime				Timescale over which to blend out the extrapolation. Larger values result in longer blend times. Can be zero.
+	 * @param SmoothingTime			Decay rate of the extrapolated velocity. Smaller values make the velocity decay faster.
+	 */
+	template <typename TValue>
+	UE_EXPERIMENTAL(5.8, "SpringMath is experimental")
+	static void DeadBlendApply(
+		TValue& InOutValue,
+		TValue& InOutVelocity,
+		const TValue& InValueTransition,
+		const TValue& InVelocityTransition,
+		float TimeSinceTransition,
+		float BlendTime,
+		float SmoothingTime)
+	{
+		TValue ExtrapolatedValue = InValueTransition, ExtrapolatedVelocity = InVelocityTransition;
+		PRAGMA_DISABLE_EXPERIMENTAL_WARNINGS
+		DeadBlendExtrapolate(ExtrapolatedValue, ExtrapolatedVelocity, TimeSinceTransition, SmoothingTime);
+		PRAGMA_ENABLE_EXPERIMENTAL_WARNINGS
+
+		float NonZeroBlendTime = FMath::Max(BlendTime, UE_SMALL_NUMBER);
+		float Alpha = FMath::Clamp(TimeSinceTransition / NonZeroBlendTime, 0.0f, 1.0f);
+
+		InOutValue = FMath::Lerp(ExtrapolatedValue, InOutValue, Alpha);
+		InOutVelocity = FMath::Lerp(ExtrapolatedVelocity, InOutVelocity, Alpha) + 
+			(TimeSinceTransition < BlendTime) * ((InOutValue - ExtrapolatedValue) / NonZeroBlendTime);
+	}
+
+	/** Transitions a dead blend, computing the new value and velocity at the point of transition.
+	 *
+	 * @param InOutValueTransition		Value at the point of transition
+	 * @param InOutVelocityTransition	Velocity at the point of transition
+	 * @param InSrcValue				Value of the source (value transitioning from)
+	 * @param InSrcVelocity				Velocity of the source (velocity transitioning from)
+	 * @param TimeSinceTransition		Time since the last transition
+	 * @param BlendTime					Timescale over which to apply the offset. Larger values result in longer blend times. Can be zero.
+	 * @param SmoothingTime				Decay rate of the extrapolated velocity. Smaller values make the velocity decay faster.
+	 */
+	template <typename TValue>
+	UE_EXPERIMENTAL(5.8, "SpringMath is experimental")
+	static void DeadBlendTransition(
+		TValue& InOutValueTransition,
+		TValue& InOutVelocityTransition,
+		const TValue& InSrcValue,
+		const TValue& InSrcVelocity,
+		float TimeSinceTransition,
+		float BlendTime,
+		float SmoothingTime)
+	{
+		TValue ExtrapolatedValue = InOutValueTransition, ExtrapolatedVelocity = InOutVelocityTransition;
+		PRAGMA_DISABLE_EXPERIMENTAL_WARNINGS
+		DeadBlendExtrapolate(ExtrapolatedValue, ExtrapolatedVelocity, TimeSinceTransition, SmoothingTime);
+		PRAGMA_ENABLE_EXPERIMENTAL_WARNINGS
+
+		float NonZeroBlendTime = FMath::Max(BlendTime, UE_SMALL_NUMBER);
+		float Alpha = FMath::Clamp(TimeSinceTransition / NonZeroBlendTime, 0.0f, 1.0f);
+
+		InOutValueTransition = FMath::Lerp(ExtrapolatedValue, InSrcValue, Alpha);
+		InOutVelocityTransition = FMath::Lerp(ExtrapolatedVelocity, InSrcVelocity, Alpha) + 
+			(TimeSinceTransition < BlendTime) * ((InSrcValue - ExtrapolatedValue) / NonZeroBlendTime);
+	}
+
+	/** Specialization of DeadBlendApply for quaternions. Stores velocities in rad/s.
+	 *
+	 * @param InOutValue			The value to apply the inertialization to
+	 * @param InOutVelocity			The velocity to apply the inertialization to
+	 * @param InValueTransition		Value at the time of the last transition
+	 * @param InVelocityTransition	Velocity at the time of the last transition
+	 * @param TimeSinceTransition	Time since the last transition
+	 * @param BlendTime				Timescale over which to blend out the extrapolation. Larger values result in longer blend times. Can be zero.
+	 * @param SmoothingTime			Decay rate of the extrapolated velocity. Smaller values make the velocity decay faster.
+	 */
+	UE_EXPERIMENTAL(5.8, "SpringMath is experimental")
+	static void DeadBlendApplyQuat(
+		FQuat& InOutValue,
+		FVector& InOutVelocity,
+		const FQuat& InValueTransition,
+		const FVector& InVelocityTransition,
+		float TimeSinceTransition,
+		float BlendTime,
+		float SmoothingTime)
+	{
+		FQuat ExtrapolatedValue = InValueTransition; FVector ExtrapolatedVelocity = InVelocityTransition;
+		PRAGMA_DISABLE_EXPERIMENTAL_WARNINGS
+		DeadBlendExtrapolateQuat(ExtrapolatedValue, ExtrapolatedVelocity, TimeSinceTransition, SmoothingTime);
+		PRAGMA_ENABLE_EXPERIMENTAL_WARNINGS
+
+		float NonZeroBlendTime = FMath::Max(BlendTime, UE_SMALL_NUMBER);
+		float Alpha = FMath::Clamp(TimeSinceTransition / NonZeroBlendTime, 0.0f, 1.0f);
+
+		FQuat Direction = (InValueTransition * InOutValue.Inverse()).GetShortestArcWith(FQuat::Identity);
+		FVector Diff = (InOutValue * ExtrapolatedValue.Inverse()).GetShortestArcWith(Direction).ToRotationVector();
+		
+		InOutValue = FQuat::MakeFromRotationVector(Diff * Alpha) * ExtrapolatedValue;
+		InOutVelocity = FMath::Lerp(ExtrapolatedVelocity, InOutVelocity, Alpha) + 
+			(TimeSinceTransition < BlendTime) * (Diff / NonZeroBlendTime);
+	}
+
+	/** Specialization of DeadBlendTransition for quaternions. Stores velocities in rad/s.
+	 *
+	 * @param InOutValueTransition		Value at the point of transition
+	 * @param InOutVelocityTransition	Velocity at the point of transition
+	 * @param InSrcValue				Value of the source (value transitioning from)
+	 * @param InSrcVelocity				Velocity of the source (velocity transitioning from)
+	 * @param TimeSinceTransition		Time since the last transition
+	 * @param BlendTime					Timescale over which to apply the offset. Larger values result in longer blend times. Can be zero.
+	 * @param SmoothingTime				Decay rate of the extrapolated velocity. Smaller values make the velocity decay faster.
+	 */
+	UE_EXPERIMENTAL(5.8, "SpringMath is experimental")
+	static void DeadBlendTransitionQuat(
+		FQuat& InOutValueTransition,
+		FVector& InOutVelocityTransition,
+		const FQuat& InSrcValue,
+		const FVector& InSrcVelocity,
+		float TimeSinceTransition,
+		float BlendTime,
+		float SmoothingTime)
+	{
+		FQuat ExtrapolatedValue = InOutValueTransition; FVector ExtrapolatedVelocity = InOutVelocityTransition;
+		PRAGMA_DISABLE_EXPERIMENTAL_WARNINGS
+		DeadBlendExtrapolateQuat(ExtrapolatedValue, ExtrapolatedVelocity, TimeSinceTransition, SmoothingTime);
+		PRAGMA_ENABLE_EXPERIMENTAL_WARNINGS
+
+		float NonZeroBlendTime = FMath::Max(BlendTime, UE_SMALL_NUMBER);
+		float Alpha = FMath::Clamp(TimeSinceTransition / NonZeroBlendTime, 0.0f, 1.0f);
+
+		FQuat Direction = (InOutValueTransition * InSrcValue.Inverse()).GetShortestArcWith(FQuat::Identity);
+		FVector Diff = (InSrcValue * ExtrapolatedValue.Inverse()).GetShortestArcWith(Direction).ToRotationVector();
+	
+		InOutValueTransition = FQuat::MakeFromRotationVector(Diff * Alpha) * ExtrapolatedValue;
+		InOutVelocityTransition = FMath::Lerp(ExtrapolatedVelocity, InSrcVelocity, Alpha) + 
+			(TimeSinceTransition < BlendTime) * (Diff / NonZeroBlendTime);
+	}
+
+	/** Specialization of DeadBlendApply for angles. Stores velocities in rad/s.
+	 *
+	 * @param InOutValue			The value to apply the inertialization to
+	 * @param InOutVelocity			The velocity to apply the inertialization to
+	 * @param InValueTransition		Value at the time of the last transition
+	 * @param InVelocityTransition	Velocity at the time of the last transition
+	 * @param TimeSinceTransition	Time since the last transition
+	 * @param BlendTime				Timescale over which to blend out the extrapolation. Larger values result in longer blend times. Can be zero.
+	 * @param SmoothingTime			Decay rate of the extrapolated velocity. Smaller values make the velocity decay faster.
+	 */
+	UE_EXPERIMENTAL(5.8, "SpringMath is experimental")
+	static void DeadBlendApplyAngle(
+		float& InOutValue,
+		float& InOutVelocity,
+		const float& InValueTransition,
+		const float& InVelocityTransition,
+		float TimeSinceTransition,
+		float BlendTime,
+		float SmoothingTime)
+	{
+		float ExtrapolatedValue = InValueTransition, ExtrapolatedVelocity = InVelocityTransition;
+		PRAGMA_DISABLE_EXPERIMENTAL_WARNINGS
+		DeadBlendExtrapolate(ExtrapolatedValue, ExtrapolatedVelocity, TimeSinceTransition, SmoothingTime);
+		PRAGMA_ENABLE_EXPERIMENTAL_WARNINGS
+
+		float NonZeroBlendTime = FMath::Max(BlendTime, UE_SMALL_NUMBER);
+		float Alpha = FMath::Clamp(TimeSinceTransition / NonZeroBlendTime, 0.0f, 1.0f);
+
+		float Direction = FMath::FindDeltaAngleRadians(ExtrapolatedValue, InValueTransition);
+		float Diff = AngleGetShortestArcWith(FMath::FindDeltaAngleRadians(ExtrapolatedValue, InOutValue), Direction);
+
+		InOutValue = (Diff * Alpha) + ExtrapolatedValue;
+		InOutVelocity = FMath::Lerp(ExtrapolatedVelocity, InOutVelocity, Alpha) + 
+			(TimeSinceTransition < BlendTime) * (Diff / NonZeroBlendTime);
+	}
+
+	/** Specialization of DeadBlendTransition for angles. Stores velocities in rad/s.
+	 *
+	 * @param InOutValueTransition		Value at the point of transition
+	 * @param InOutVelocityTransition	Velocity at the point of transition
+	 * @param InSrcValue				Value of the source (value transitioning from)
+	 * @param InSrcVelocity				Velocity of the source (velocity transitioning from)
+	 * @param TimeSinceTransition		Time since the last transition
+	 * @param BlendTime					Timescale over which to apply the offset. Larger values result in longer blend times. Can be zero.
+	 * @param SmoothingTime				Decay rate of the extrapolated velocity. Smaller values make the velocity decay faster.
+	 */
+	UE_EXPERIMENTAL(5.8, "SpringMath is experimental")
+	static void DeadBlendTransitionAngle(
+		float& InOutValueTransition,
+		float& InOutVelocityTransition,
+		const float& InSrcValue,
+		const float& InSrcVelocity,
+		float TimeSinceTransition,
+		float BlendTime,
+		float SmoothingTime)
+	{
+		float ExtrapolatedValue = InOutValueTransition, ExtrapolatedVelocity = InOutVelocityTransition;
+		PRAGMA_DISABLE_EXPERIMENTAL_WARNINGS
+		DeadBlendExtrapolate(ExtrapolatedValue, ExtrapolatedVelocity, TimeSinceTransition, SmoothingTime);
+		PRAGMA_ENABLE_EXPERIMENTAL_WARNINGS
+
+		float NonZeroBlendTime = FMath::Max(BlendTime, UE_SMALL_NUMBER);
+		float Alpha = FMath::Clamp(TimeSinceTransition / NonZeroBlendTime, 0.0f, 1.0f);
+
+		float Direction = FMath::FindDeltaAngleRadians(ExtrapolatedValue, InOutValueTransition);
+		float Diff = AngleGetShortestArcWith(FMath::FindDeltaAngleRadians(ExtrapolatedValue, InSrcValue), Direction);
+
+		InOutValueTransition = (Diff * Alpha) + ExtrapolatedValue;
+		InOutVelocityTransition = FMath::Lerp(ExtrapolatedVelocity, InSrcVelocity, Alpha) + 
+			(TimeSinceTransition < BlendTime) * (Diff / NonZeroBlendTime);
+	}
+
+	/** Specialization of DeadBlendApply for scales.
+	 *
+	 * @param InOutValue			The value to apply the inertialization to
+	 * @param InOutVelocity			The velocity to apply the inertialization to
+	 * @param InValueTransition		Value at the time of the last transition
+	 * @param InVelocityTransition	Velocity at the time of the last transition
+	 * @param TimeSinceTransition	Time since the last transition
+	 * @param BlendTime				Timescale over which to blend out the extrapolation. Larger values result in longer blend times. Can be zero.
+	 * @param SmoothingTime			Decay rate of the extrapolated velocity. Smaller values make the velocity decay faster.
+	 */
+	UE_EXPERIMENTAL(5.8, "SpringMath is experimental")
+	static void DeadBlendApplyScale(
+		FVector& InOutValue,
+		FVector& InOutVelocity,
+		const FVector& InValueTransition,
+		const FVector& InVelocityTransition,
+		float TimeSinceTransition,
+		float BlendTime,
+		float SmoothingTime)
+	{
+		FVector ExtrapolatedValue = InValueTransition, ExtrapolatedVelocity = InVelocityTransition;
+		PRAGMA_DISABLE_EXPERIMENTAL_WARNINGS
+		DeadBlendExtrapolateScale(ExtrapolatedValue, ExtrapolatedVelocity, TimeSinceTransition, SmoothingTime);
+		PRAGMA_ENABLE_EXPERIMENTAL_WARNINGS
+
+		float NonZeroBlendTime = FMath::Max(BlendTime, UE_SMALL_NUMBER);
+		float Alpha = FMath::Clamp(TimeSinceTransition / NonZeroBlendTime, 0.0f, 1.0f);
+
+		InOutValue = ScaleEerp(ExtrapolatedValue, InOutValue, Alpha);
+		InOutVelocity = FMath::Lerp(ExtrapolatedVelocity, InOutVelocity, Alpha) + 
+			(TimeSinceTransition < BlendTime) * (ScaleLog(ScaleDivMax(InOutValue, ExtrapolatedValue)) / NonZeroBlendTime);
+	}
+
+	/** Specialization of DeadBlendTransition for scales.
+	 *
+	 * @param InOutValueTransition		Value at the point of transition
+	 * @param InOutVelocityTransition	Velocity at the point of transition
+	 * @param InSrcValue				Value of the source (value transitioning from)
+	 * @param InSrcVelocity				Velocity of the source (velocity transitioning from)
+	 * @param TimeSinceTransition		Time since the last transition
+	 * @param BlendTime					Timescale over which to apply the offset. Larger values result in longer blend times. Can be zero.
+	 * @param SmoothingTime				Decay rate of the extrapolated velocity. Smaller values make the velocity decay faster.
+	 */
+	UE_EXPERIMENTAL(5.8, "SpringMath is experimental")
+	static void DeadBlendTransitionScale(
+		FVector& InOutValueTransition,
+		FVector& InOutVelocityTransition,
+		const FVector& InSrcValue,
+		const FVector& InSrcVelocity,
+		float TimeSinceTransition,
+		float BlendTime,
+		float SmoothingTime)
+	{
+		FVector ExtrapolatedValue = InOutValueTransition, ExtrapolatedVelocity = InOutVelocityTransition;
+		PRAGMA_DISABLE_EXPERIMENTAL_WARNINGS
+		DeadBlendExtrapolateScale(ExtrapolatedValue, ExtrapolatedVelocity, TimeSinceTransition, SmoothingTime);
+		PRAGMA_ENABLE_EXPERIMENTAL_WARNINGS
+
+		float NonZeroBlendTime = FMath::Max(BlendTime, UE_SMALL_NUMBER);
+		float Alpha = FMath::Clamp(TimeSinceTransition / NonZeroBlendTime, 0.0f, 1.0f);
+
+		InOutValueTransition = ScaleEerp(ExtrapolatedValue, InSrcValue, Alpha);
+		InOutVelocityTransition = FMath::Lerp(ExtrapolatedVelocity, InSrcVelocity, Alpha) + 
+			(TimeSinceTransition < BlendTime) * (ScaleLog(ScaleDivMax(InSrcValue, ExtrapolatedValue)) / NonZeroBlendTime);
+	}
+
 };
