CIS fix for linux for issue introduced in 48769505

#Synced-CL 48731319

diff --git a/Engine/Source/Runtime/Engine/Private/WorldPartition/HLOD/HLODActor.cpp b/Engine/Source/Runtime/Engine/Private/WorldPartition/HLOD/HLODActor.cpp
--- a/Engine/Source/Runtime/Engine/Private/WorldPartition/HLOD/HLODActor.cpp
+++ b/Engine/Source/Runtime/Engine/Private/WorldPartition/HLOD/HLODActor.cpp
@@ -729,12 +729,12 @@ void AWorldPartitionHLOD::UpdateHLODBuildReportContent()
 	HLODBuildReport += HLODBuildDetailsSectionHeader;
 	HLODBuildReport += LINE_TERMINATOR;
 
-	for (const TObjectPtr<UHLODRebuildPolicyData> HLODRebuildPolicyData : HLODRebuildPolicyDataSet)
+	for (const TObjectPtr<UHLODRebuildPolicyData>& HLODRebuildPolicyData : HLODRebuildPolicyDataSet)
 	{
 		FString HLODPolicyDataString = HLODRebuildPolicyData->ToString();
 		if (!HLODPolicyDataString.IsEmpty())
 		{
-			HLODBuildReport += HLODRebuildPolicyData->ToString().TrimEnd();
+			HLODBuildReport += HLODPolicyDataString.TrimEnd();
 			HLODBuildReport += LINE_TERMINATOR;
 		}
 	}
@@ -992,7 +992,7 @@ void AWorldPartitionHLOD::SetHLODRebuildPolicyDataSet(const FHLODRebuildPolicyDa
 {
 	HLODRebuildPolicyDataSet = InHLODRebuildPolicyDataSet;
 
-	for (UHLODRebuildPolicyData* HLODRebuildPolicyData : HLODRebuildPolicyDataSet)
+	for (TObjectPtr<UHLODRebuildPolicyData>& HLODRebuildPolicyData : HLODRebuildPolicyDataSet)
 	{
 		HLODRebuildPolicyData->Rename(nullptr, this, REN_None);
 	}
