CIS fix for static analysis warnings reported for 48769505

#Synced-CL 48731319

diff --git a/Engine/Source/Runtime/Engine/Private/WorldPartition/HLOD/HLODRebuildPolicy.cpp b/Engine/Source/Runtime/Engine/Private/WorldPartition/HLOD/HLODRebuildPolicy.cpp
--- a/Engine/Source/Runtime/Engine/Private/WorldPartition/HLOD/HLODRebuildPolicy.cpp
+++ b/Engine/Source/Runtime/Engine/Private/WorldPartition/HLOD/HLODRebuildPolicy.cpp
@@ -70,7 +70,10 @@ FHLODRebuildPolicyDataSet UHLODRebuildPolicy::ComputeDataForRebuildPolicies(cons
 		{
 			const FHLODRebuildPolicyDataSet& BaselineHLODRebuildPolicyDataSet = InHLODActor->GetHLODRebuildPolicyDataSet();
 			const TObjectPtr<UHLODRebuildPolicyData>* OldDataPtr = BaselineHLODRebuildPolicyDataSet.FindByPredicate([UHLODRebuildPolicyCDO](const TObjectPtr<UHLODRebuildPolicyData>& Entry) { return Entry.IsA(UHLODRebuildPolicyCDO->GetRebuildPolicyDataType()); });
-			OldData = OldDataPtr ? *OldDataPtr : nullptr;
+			if (OldDataPtr)
+			{
+				OldData = *OldDataPtr;
+			}
 		}
 		
 		UHLODRebuildPolicyData* NewData = UHLODRebuildPolicyCDO->ComputeDataForRebuildPolicy(InHLODActor, InSourceComponents, OldData);
@@ -96,11 +99,19 @@ EHLODRebuildPolicyDecision UHLODRebuildPolicy::Evaluate(const AWorldPartitionHLO
 
 		// Find old data
 		const TObjectPtr<UHLODRebuildPolicyData>* OldHLODRebuildPolicyDataPtr = InOldData.FindByPredicate([UHLODRebuildPolicyCDO](const TObjectPtr<UHLODRebuildPolicyData>& Entry) { return Entry.IsA(UHLODRebuildPolicyCDO->GetRebuildPolicyDataType()); });
-		const UHLODRebuildPolicyData* OldHLODRebuildPolicyData = OldHLODRebuildPolicyDataPtr ? *OldHLODRebuildPolicyDataPtr : nullptr;
+		const UHLODRebuildPolicyData* OldHLODRebuildPolicyData = nullptr;
+		if (OldHLODRebuildPolicyDataPtr)
+		{
+			OldHLODRebuildPolicyData = *OldHLODRebuildPolicyDataPtr;
+		}
 
 		// Find new Data
 		const TObjectPtr<UHLODRebuildPolicyData>* NewHLODRebuildPolicyDataPtr = InNewData.FindByPredicate([UHLODRebuildPolicyCDO](const TObjectPtr<UHLODRebuildPolicyData>& Entry) { return Entry.IsA(UHLODRebuildPolicyCDO->GetRebuildPolicyDataType()); });
-		const UHLODRebuildPolicyData* NewHLODRebuildPolicyData = NewHLODRebuildPolicyDataPtr ? *NewHLODRebuildPolicyDataPtr : nullptr;
+		const UHLODRebuildPolicyData* NewHLODRebuildPolicyData = nullptr;
+		if (NewHLODRebuildPolicyDataPtr)
+		{
+			NewHLODRebuildPolicyData = *NewHLODRebuildPolicyDataPtr;
+		}
 
 		if (!NewHLODRebuildPolicyData)
 		{
