Removed leftover pragma once in cpp

#Synced-CL 48731319

diff --git a/Engine/Source/Runtime/Engine/Private/WorldPartition/HLOD/HLODRebuildPolicy.cpp b/Engine/Source/Runtime/Engine/Private/WorldPartition/HLOD/HLODRebuildPolicy.cpp
--- a/Engine/Source/Runtime/Engine/Private/WorldPartition/HLOD/HLODRebuildPolicy.cpp
+++ b/Engine/Source/Runtime/Engine/Private/WorldPartition/HLOD/HLODRebuildPolicy.cpp
@@ -1,10 +1,7 @@
 // Copyright Epic Games, Inc. All Rights Reserved.
 
-#pragma once
-
 #include "WorldPartition/HLOD/HLODRebuildPolicy.h"
 #include "WorldPartition/HLOD/HLODActor.h"
-#include "Algo/RemoveIf.h"
 
 UHLODRebuildPolicyData::UHLODRebuildPolicyData(const FObjectInitializer& ObjectInitializer)
 	: Super(ObjectInitializer)
