Ensured static mesh Card Representation and Distance Field data are built async & waited on to prevent blocking during package serialization. Ensured data builds are static mesh specific. Removed unused thread parameter in respective Build functions.

Static mesh now checks if card/distance field build count "GetNumOutstandingTasks" during "IsCachedCookedPlatformDataLoaded" before package serialization. "BlockUntilBuildComplete" (per-static mesh) / "BlockUntilAllBuildsComplete" functions remain as a fallback, but should no longer create significant overhead.

#rnx

#Synced-CL 48675545

diff --git a/Engine/Source/Runtime/Engine/Private/DistanceFieldAtlas.cpp b/Engine/Source/Runtime/Engine/Private/DistanceFieldAtlas.cpp
--- a/Engine/Source/Runtime/Engine/Private/DistanceFieldAtlas.cpp
+++ b/Engine/Source/Runtime/Engine/Private/DistanceFieldAtlas.cpp
@@ -414,12 +414,97 @@ void FDistanceFieldAsyncQueue::OnAssetPostCompile(const TArray<FAssetCompileData
 	{
 		if (CompileData.Asset.IsValid() && CompileData.Asset.Get()->IsA<UStaticMesh>())
 		{
-			ProcessPendingTasks();
-			return;
+			const UStaticMesh* StaticMesh = Cast<UStaticMesh>(CompileData.Asset.Get());
+			ProcessPendingTasksForMesh(StaticMesh);
 		}
 	}
 }
 
+#if WITH_EDITOR
+void FDistanceFieldAsyncQueue::FinalizeAsyncTask(FAsyncDistanceFieldTask* InTask, FObjectCacheContextScope& InObjectCacheScope) const
+{
+	TRACE_CPUPROFILER_EVENT_SCOPE(FDistanceFieldAsyncQueue::FinalizeAsyncTask);
+
+	check(InTask);
+	check(IsInGameThread()); // Timer wont work properly outside of gamethread.
+
+	// We want to count each resource built from a DDC miss, so count each iteration of the loop separately.
+	COOK_STAT(auto Timer = DistanceFieldCookStats::UsageStats.TimeSyncWork());
+
+	if (InTask->AsyncTask)
+	{
+		InTask->AsyncTask->EnsureCompletion();
+		InTask->AsyncTask.Reset();
+	}
+
+	// Editor 'force delete' can null any UObject pointers which are seen by reference collecting (eg FProperty or serialized)
+	if (InTask->StaticMesh)
+	{
+		check(!InTask->StaticMesh->IsCompiling());
+
+		InTask->GeneratedVolumeData->bAsyncBuilding = false;
+
+		FStaticMeshRenderData* RenderData = InTask->StaticMesh->GetRenderData();
+		FDistanceFieldVolumeData* OldVolumeData = RenderData->LODResources[0].DistanceFieldData;
+
+		// Assign the new volume data, this is safe because the render thread makes a copy of the pointer at scene proxy creation time.
+		RenderData->LODResources[0].DistanceFieldData = InTask->GeneratedVolumeData;
+
+		// Renderstates are not initialized between UStaticMesh::PreEditChange() and UStaticMesh::PostEditChange()
+		if (RenderData->IsInitialized())
+		{
+			for (IStaticMeshComponent* Component : InObjectCacheScope.GetContext().GetStaticMeshComponents(InTask->StaticMesh))
+			{
+				IPrimitiveComponent* PrimitiveComponent = Component->GetPrimitiveComponentInterface();
+
+				if (PrimitiveComponent->IsRegistered() && PrimitiveComponent->IsRenderStateCreated())
+				{
+					PrimitiveComponent->MarkRenderStateDirty();
+				}
+			}
+		}
+
+		if (OldVolumeData)
+		{
+			// Rendering thread may still be referencing the old one, use the deferred cleanup interface to delete it next frame when it is safe
+			BeginCleanup(OldVolumeData);
+		}
+
+		// Need also to update platform render data if it's being cached
+		FStaticMeshRenderData* PlatformRenderData = RenderData->NextCachedRenderData.Get();
+		while (PlatformRenderData)
+		{
+			if (PlatformRenderData->LODResources[0].DistanceFieldData)
+			{
+				*PlatformRenderData->LODResources[0].DistanceFieldData = *InTask->GeneratedVolumeData;
+				// The old bulk data assignment operator doesn't copy over flags
+				PlatformRenderData->LODResources[0].DistanceFieldData->StreamableMips.ResetBulkDataFlags(InTask->GeneratedVolumeData->StreamableMips.GetBulkDataFlags());
+			}
+			PlatformRenderData = PlatformRenderData->NextCachedRenderData.Get();
+		}
+
+		{
+			TArray<uint8> DerivedData;
+			// Save built distance field volume to DDC
+			FMemoryWriter Ar(DerivedData, /*bIsPersistent=*/ true);
+			InTask->StaticMesh->GetRenderData()->LODResources[0].DistanceFieldData->Serialize(Ar, InTask->StaticMesh);
+			GetDerivedDataCacheRef().Put(*InTask->DDCKey, DerivedData, InTask->StaticMesh->GetPathName());
+			COOK_STAT(Timer.AddMiss(DerivedData.Num()));
+		}
+
+		if (!InTask->bGeneratedCardRepresentation)
+		{
+			BeginCacheMeshCardRepresentation(
+				InTask->StaticMesh,
+				InTask->StaticMesh->GetPlatformStaticMeshRenderData(InTask->StaticMesh, InTask->TargetPlatform),
+				InTask->DDCKey,
+				MoveTemp(InTask->SourceMeshData),
+				MoveTemp(InTask->PartData));
+		}
+	}
+}
+#endif
+
 FDistanceFieldAsyncQueue::FDistanceFieldAsyncQueue() 
 	: Notification(MakeUnique<FAsyncCompilationNotification>(GetAssetNameFormat()))
 {
@@ -460,9 +545,7 @@ void FDistanceFieldAsyncQueue::OnPostReachabilityAnalysis()
 
 void FAsyncDistanceFieldTaskWorker::DoWork()
 {
-	// Put on background thread to avoid interfering with game-thread bound tasks
-	FQueuedThreadPoolTaskGraphWrapper TaskGraphWrapper(ENamedThreads::AnyBackgroundThreadNormalTask);
-	GDistanceFieldAsyncQueue->Build(&Task, TaskGraphWrapper);
+	GDistanceFieldAsyncQueue->Build(&Task);
 }
 
 FName FDistanceFieldAsyncQueue::GetStaticAssetTypeName()
@@ -681,7 +764,26 @@ void FDistanceFieldAsyncQueue::ProcessPendingTasks()
 	for (auto It = PendingTasks.CreateIterator(); It; ++It)
 	{
 		FAsyncDistanceFieldTask* Task = *It;
-		if ((Task->GenerateSource == nullptr || !Task->GenerateSource->IsCompiling()) && (Task->StaticMesh == nullptr || !Task->StaticMesh ->IsCompiling()))
+		if ((Task->GenerateSource == nullptr || !Task->GenerateSource->IsCompiling()) && (Task->StaticMesh == nullptr || !Task->StaticMesh->IsCompiling()))
+		{
+			StartBackgroundTask(Task);
+			It.RemoveCurrent();
+		}
+	}
+}
+
+void FDistanceFieldAsyncQueue::ProcessPendingTasksForMesh(const UStaticMesh *InStaticMesh)
+{
+	TRACE_CPUPROFILER_EVENT_SCOPE(FDistanceFieldAsyncQueue::ProcessPendingTasksForMesh);
+
+	FScopeLock Lock(&CriticalSection);
+
+	for (auto It = PendingTasks.CreateIterator(); It; ++It)
+	{
+		FAsyncDistanceFieldTask* Task = *It;
+		if ((Task->StaticMesh == InStaticMesh || Task->GenerateSource == InStaticMesh) &&
+			(Task->GenerateSource == nullptr || !Task->GenerateSource->IsCompiling()) && 
+			(Task->StaticMesh == nullptr || !Task->StaticMesh->IsCompiling()))
 		{
 			StartBackgroundTask(Task);
 			It.RemoveCurrent();
@@ -730,10 +832,7 @@ void FDistanceFieldAsyncQueue::AddTask(FAsyncDistanceFieldTask* Task)
 	
 	if (!bIsCompiling && !bUseAsyncBuild)
 	{
-		// To avoid deadlocks, we must queue the inner build tasks on another thread pool, so use the task graph.
-		// Put on background thread to avoid interfering with game-thread bound tasks
-		FQueuedThreadPoolTaskGraphWrapper TaskGraphWrapper(ENamedThreads::AnyBackgroundThreadNormalTask);
-		Build(Task, TaskGraphWrapper);
+		Build(Task);
 	}
 #else
 	UE_LOG(LogStaticMesh,Fatal,TEXT("Tried to build a distance field without editor support (this should have been done during cooking)"));
@@ -779,6 +878,29 @@ void FDistanceFieldAsyncQueue::RescheduleBackgroundTask(FAsyncDistanceFieldTask*
 	}
 }
 
+int32 FDistanceFieldAsyncQueue::GetNumOutstandingTasks(const UStaticMesh* InStaticMesh) const
+{
+	FScopeLock Lock(&CriticalSection);
+
+	// Pending tasks are already included in referenced tasks
+	if (!ReferencedTasks.Num())
+	{
+		return 0;
+	}
+
+	int32 Count = 0;
+	for (auto It = ReferencedTasks.CreateConstIterator(); It; ++It)
+	{
+		const FAsyncDistanceFieldTask* Task = *It;
+		if (InStaticMesh == Task->GenerateSource || InStaticMesh == Task->StaticMesh)
+		{
+			++Count;
+		}
+	}
+
+	return Count;
+}
+
 void FDistanceFieldAsyncQueue::BlockUntilBuildComplete(UStaticMesh* StaticMesh, bool bWarnIfBlocked)
 {
 	TRACE_CPUPROFILER_EVENT_SCOPE(FDistanceFieldAsyncQueue::BlockUntilBuildComplete)
@@ -795,7 +917,7 @@ void FDistanceFieldAsyncQueue::BlockUntilBuildComplete(UStaticMesh* StaticMesh,
 	TSet<UStaticMesh*> RequiredFinishCompilation;
 	do 
 	{
-		ProcessAsyncTasks();
+		ProcessAsyncTasksForMesh(StaticMesh);
 
 		bReferenced = false;
 		RequiredFinishCompilation.Reset();
@@ -895,7 +1017,7 @@ void FDistanceFieldAsyncQueue::BlockUntilAllBuildsComplete()
 	} 
 }
 
-void FDistanceFieldAsyncQueue::Build(FAsyncDistanceFieldTask* Task, FQueuedThreadPool& BuildThreadPool)
+void FDistanceFieldAsyncQueue::Build(FAsyncDistanceFieldTask* Task)
 {
 #if WITH_EDITOR
 	// Editor 'force delete' can null any UObject pointers which are seen by reference collecting (eg FProperty or serialized)
@@ -955,7 +1077,6 @@ void FDistanceFieldAsyncQueue::Build(FAsyncDistanceFieldTask* Task, FQueuedThrea
 			CompletedTasks.Add(Task);
 		}
 	}
-
 #endif
 }
 
@@ -966,7 +1087,7 @@ void FDistanceFieldAsyncQueue::ProcessAsyncTasks(bool bLimitExecutionTime)
 
 	ProcessPendingTasks();
 
-	FObjectCacheContextScope ObjectCacheScope;
+	FObjectCacheContextScope ObjectCacheScope; // Declared here for caching optimization purposes
 	const double MaxProcessingTime = 0.016f;
 	double StartTime = FPlatformTime::Seconds();
 	bool bMadeProgress = false;
@@ -992,80 +1113,56 @@ void FDistanceFieldAsyncQueue::ProcessAsyncTasks(bool bLimitExecutionTime)
 		}
 		bMadeProgress = true;
 
-		// We want to count each resource built from a DDC miss, so count each iteration of the loop separately.
-		COOK_STAT(auto Timer = DistanceFieldCookStats::UsageStats.TimeSyncWork());
+		FinalizeAsyncTask(Task, ObjectCacheScope);
 
-		if (Task->AsyncTask)
-		{
-			Task->AsyncTask->EnsureCompletion();
-			Task->AsyncTask.Reset();
-		}
-
-		// Editor 'force delete' can null any UObject pointers which are seen by reference collecting (eg FProperty or serialized)
-		if (Task->StaticMesh)
-		{
-			check(!Task->StaticMesh->IsCompiling());
-
-			Task->GeneratedVolumeData->bAsyncBuilding = false;
-
-			FStaticMeshRenderData* RenderData = Task->StaticMesh->GetRenderData();
-			FDistanceFieldVolumeData* OldVolumeData = RenderData->LODResources[0].DistanceFieldData;
-
-			// Assign the new volume data, this is safe because the render thread makes a copy of the pointer at scene proxy creation time.
-			RenderData->LODResources[0].DistanceFieldData = Task->GeneratedVolumeData;
+		delete Task;
+	}
 
-			// Renderstates are not initialized between UStaticMesh::PreEditChange() and UStaticMesh::PostEditChange()
-			if (RenderData->IsInitialized())
-			{
-				for (IStaticMeshComponent* Component : ObjectCacheScope.GetContext().GetStaticMeshComponents(Task->StaticMesh))
-				{
-					IPrimitiveComponent* PrimitiveComponent = Component->GetPrimitiveComponentInterface();
+	if (bMadeProgress)
+	{
+		Notification->Update(GetNumRemainingAssets());
+	}
+#endif
+}
 
-					if (PrimitiveComponent->IsRegistered() && PrimitiveComponent->IsRenderStateCreated())
-					{
-						PrimitiveComponent->MarkRenderStateDirty();
-					}
-				}
-			}
+void FDistanceFieldAsyncQueue::ProcessAsyncTasksForMesh(const UStaticMesh* InStaticMesh, bool bLimitExecutionTime)
+{
+#if WITH_EDITOR
+	TRACE_CPUPROFILER_EVENT_SCOPE(FDistanceFieldAsyncQueue::ProcessAsyncTasksForMesh);
 
-			if (OldVolumeData)
-			{
-				// Rendering thread may still be referencing the old one, use the deferred cleanup interface to delete it next frame when it is safe
-				BeginCleanup(OldVolumeData);
-			}
+	ProcessPendingTasksForMesh(InStaticMesh);
 
-			// Need also to update platform render data if it's being cached
-			FStaticMeshRenderData* PlatformRenderData = RenderData->NextCachedRenderData.Get();
-			while (PlatformRenderData)
+	FObjectCacheContextScope ObjectCacheScope;  // Declared here for mem/caching optimization purposes
+	const double MaxProcessingTime = 0.016f;
+	double StartTime = FPlatformTime::Seconds();
+	bool bMadeProgress = false;
+	while (!bLimitExecutionTime || (FPlatformTime::Seconds() - StartTime) < MaxProcessingTime)
+	{
+		FAsyncDistanceFieldTask* Task = nullptr;
+		{
+			// Try and find a task that matches the static mesh
+			FScopeLock Lock(&CriticalSection);
+			for (auto CompletedTaskIt = CompletedTasks.CreateIterator(); CompletedTaskIt; ++CompletedTaskIt)
 			{
-				if (PlatformRenderData->LODResources[0].DistanceFieldData)
+				FAsyncDistanceFieldTask* TaskToCheck = *CompletedTaskIt;
+				
+				if (TaskToCheck != nullptr && (TaskToCheck->StaticMesh == InStaticMesh || TaskToCheck->GenerateSource == InStaticMesh))
 				{
-					*PlatformRenderData->LODResources[0].DistanceFieldData = *Task->GeneratedVolumeData;
-					// The old bulk data assignment operator doesn't copy over flags
-					PlatformRenderData->LODResources[0].DistanceFieldData->StreamableMips.ResetBulkDataFlags(Task->GeneratedVolumeData->StreamableMips.GetBulkDataFlags());
+					Task = TaskToCheck;
+					CompletedTaskIt.RemoveCurrent();
+					verify(ReferencedTasks.Remove(Task));
+					break;
 				}
-				PlatformRenderData = PlatformRenderData->NextCachedRenderData.Get();
-			}
-
-			{
-				TArray<uint8> DerivedData;
-				// Save built distance field volume to DDC
-				FMemoryWriter Ar(DerivedData, /*bIsPersistent=*/ true);
-				Task->StaticMesh->GetRenderData()->LODResources[0].DistanceFieldData->Serialize(Ar, Task->StaticMesh);
-				GetDerivedDataCacheRef().Put(*Task->DDCKey, DerivedData, Task->StaticMesh->GetPathName());
-				COOK_STAT(Timer.AddMiss(DerivedData.Num()));
 			}
+		}
 
-			if (!Task->bGeneratedCardRepresentation)
-			{
-				BeginCacheMeshCardRepresentation(
-					Task->StaticMesh,
-					Task->StaticMesh->GetPlatformStaticMeshRenderData(Task->StaticMesh, Task->TargetPlatform),
-					Task->DDCKey,
-					MoveTemp(Task->SourceMeshData),
-					MoveTemp(Task->PartData));
-			}
+		if (Task == nullptr)
+		{
+			break;
 		}
+		bMadeProgress = true;
+
+		FinalizeAsyncTask(Task, ObjectCacheScope);
 
 		delete Task;
 	}
diff --git a/Engine/Source/Runtime/Engine/Private/MeshCardRepresentation.cpp b/Engine/Source/Runtime/Engine/Private/MeshCardRepresentation.cpp
--- a/Engine/Source/Runtime/Engine/Private/MeshCardRepresentation.cpp
+++ b/Engine/Source/Runtime/Engine/Private/MeshCardRepresentation.cpp
@@ -457,9 +457,7 @@ void FCardRepresentationAsyncQueue::OnPostReachabilityAnalysis()
 }
 void FAsyncCardRepresentationTaskWorker::DoWork()
 {
-	// Put on background thread to avoid interfering with game-thread bound tasks
-	FQueuedThreadPoolTaskGraphWrapper TaskGraphWrapper(ENamedThreads::AnyBackgroundThreadNormalTask);
-	GCardRepresentationAsyncQueue->Build(&Task, TaskGraphWrapper);
+	GCardRepresentationAsyncQueue->Build(&Task);
 }
 
 FName FCardRepresentationAsyncQueue::GetStaticAssetTypeName()
@@ -604,8 +602,7 @@ void FCardRepresentationAsyncQueue::RunImmediateTask(FAsyncCardRepresentationTas
 		FPlatformAtomics::InterlockedCompareExchangePointer((void**)&MeshUtilities, &FModuleManager::Get().LoadModuleChecked<IMeshUtilities>(TEXT("MeshUtilities")), nullptr);
 	}
 
-	// Thread pool argument isn't actually used by the function
-	Build(Task, *ThreadPool.Get());
+	Build(Task);
 
 	// Array protection when called from multiple threads
 	FScopeLock Lock(&CriticalSection);
@@ -632,6 +629,24 @@ void FCardRepresentationAsyncQueue::ProcessPendingTasks()
 	}
 }
 
+void FCardRepresentationAsyncQueue::ProcessPendingTasksForMesh(const UStaticMesh* InStaticMesh)
+{
+	TRACE_CPUPROFILER_EVENT_SCOPE(FCardRepresentationAsyncQueue::ProcessPendingTasksForMesh);
+
+	FScopeLock Lock(&CriticalSection);
+
+	for (auto It = PendingTasks.CreateIterator(); It; ++It)
+	{
+		FAsyncCardRepresentationTask* Task = *It;
+		if ((Task->StaticMesh == InStaticMesh || Task->GenerateSource == InStaticMesh) &&
+			(Task->GenerateSource == nullptr || Task->GenerateSource->IsCompiling() == false))
+		{
+			StartBackgroundTask(Task);
+			It.RemoveCurrent();
+		}
+	}
+}
+
 void FCardRepresentationAsyncQueue::AddTask(FAsyncCardRepresentationTask* Task)
 {
 #if WITH_EDITOR
@@ -671,10 +686,7 @@ void FCardRepresentationAsyncQueue::AddTask(FAsyncCardRepresentationTask* Task)
 
 	if (!bIsCompiling && !bUseAsyncBuild)
 	{
-		// To avoid deadlocks, we must queue the inner build tasks on another thread pool, so use the task graph.
-		// Put on background thread to avoid interfering with game-thread bound tasks
-		FQueuedThreadPoolTaskGraphWrapper TaskGraphWrapper(ENamedThreads::AnyBackgroundThreadNormalTask);
-		Build(Task, TaskGraphWrapper);
+		Build(Task);
 	}
 #else
 	UE_LOG(LogStaticMesh,Fatal,TEXT("Tried to build a card representation without editor support (this should have been done during cooking)"));
@@ -776,6 +788,27 @@ void FCardRepresentationAsyncQueue::MarkCompilationAsCanceled(TArrayView<UObject
 	}
 }
 
+int32 FCardRepresentationAsyncQueue::GetNumOutstandingTasks(const UStaticMesh* InStaticMesh) const
+{
+	FScopeLock Lock(&CriticalSection);
+	if (!ReferencedTasks.Num())
+	{
+		return 0;
+	}
+
+	int32 Count = 0;
+	for (auto It = ReferencedTasks.CreateConstIterator(); It; ++It)
+	{
+		const FAsyncCardRepresentationTask* Task = *It;
+		if (InStaticMesh == Task->GenerateSource || InStaticMesh == Task->StaticMesh)
+		{
+			++Count;
+		}
+	}
+
+	return Count;
+}
+
 void FCardRepresentationAsyncQueue::BlockUntilBuildComplete(UStaticMesh* InStaticMesh, bool bWarnIfBlocked)
 {
 	TRACE_CPUPROFILER_EVENT_SCOPE(FCardRepresentationAsyncQueue::BlockUntilBuildComplete);
@@ -800,7 +833,7 @@ void FCardRepresentationAsyncQueue::BlockUntilBuildComplete(UStaticMesh* InStati
 	TSet<UStaticMesh*> RequiredFinishCompilation;
 	do 
 	{
-		ProcessAsyncTasks();
+		ProcessAsyncTasksForMesh(InStaticMesh);
 
 		bReferenced = false;
 
@@ -909,7 +942,7 @@ void FCardRepresentationAsyncQueue::BlockUntilAllBuildsComplete()
 	} 
 }
 
-void FCardRepresentationAsyncQueue::Build(FAsyncCardRepresentationTask* Task, FQueuedThreadPool& BuildThreadPool)
+void FCardRepresentationAsyncQueue::Build(FAsyncCardRepresentationTask* Task)
 {
 #if WITH_EDITOR
 	COOK_STAT(auto Timer = CardRepresentationCookStats::UsageStats.TimeSyncWork());
@@ -961,7 +994,7 @@ void FCardRepresentationAsyncQueue::ProcessAsyncTasks(bool bLimitExecutionTime)
 
 	ProcessPendingTasks();
 
-	FObjectCacheContextScope ObjectCacheScope;
+	FObjectCacheContextScope ObjectCacheScope; // Declared here for caching optimization purposes
 	const double MaxProcessingTime = 0.016f;
 	double StartTime = FPlatformTime::Seconds();
 	bool bMadeProgress = false;
@@ -987,74 +1020,137 @@ void FCardRepresentationAsyncQueue::ProcessAsyncTasks(bool bLimitExecutionTime)
 		}
 		bMadeProgress = true;
 
-		// We want to count each resource built from a DDC miss, so count each iteration of the loop separately.
-		COOK_STAT(auto Timer = CardRepresentationCookStats::UsageStats.TimeSyncWork());
+		FinalizeAndDeleteTask(Task, ObjectCacheScope);
+	}
 
-		if (Task->AsyncTask)
+	if (bMadeProgress)
+	{
+		Notification->Update(GetNumRemainingAssets());
+	}
+#endif
+}
+
+void FCardRepresentationAsyncQueue::ProcessAsyncTasksForMesh(const UStaticMesh* InStaticMesh, bool bLimitExecutionTime)
+{
+#if WITH_EDITOR
+	TRACE_CPUPROFILER_EVENT_SCOPE(FCardRepresentationAsyncQueue::ProcessAsyncTasksForMesh);
+
+	ProcessPendingTasksForMesh(InStaticMesh);
+
+	FObjectCacheContextScope ObjectCacheScope; // Declared here for mem/caching optimization purposes
+	const double MaxProcessingTime = 0.016f;
+	double StartTime = FPlatformTime::Seconds();
+	bool bMadeProgress = false;
+	while (!bLimitExecutionTime || (FPlatformTime::Seconds() - StartTime) < MaxProcessingTime)
+	{
+		FAsyncCardRepresentationTask* Task = nullptr;
 		{
-			Task->AsyncTask->EnsureCompletion();
-			Task->AsyncTask.Reset();
+			// Try and find a task that matches the static mesh
+			FScopeLock Lock(&CriticalSection);
+			for (auto CompletedTaskIt = CompletedTasks.CreateIterator(); CompletedTaskIt; ++CompletedTaskIt)
+			{
+				FAsyncCardRepresentationTask* TaskToCheck = *CompletedTaskIt;
+
+				if (TaskToCheck != nullptr && (TaskToCheck->StaticMesh == InStaticMesh || TaskToCheck->GenerateSource == InStaticMesh))
+				{
+					Task = TaskToCheck;
+					CompletedTaskIt.RemoveCurrent();
+					verify(ReferencedTasks.Remove(Task));
+					break;
+				}
+			}
 		}
 
-		// Editor 'force delete' can null any UObject pointers which are seen by reference collecting (eg UProperty or serialized)
-		if (Task->StaticMesh && Task->bSuccess)
+		if (Task == nullptr)
 		{
-			check(!Task->StaticMesh->IsCompiling());
+			break;
+		}
+		bMadeProgress = true;
 
-			FStaticMeshRenderData* RenderData = Task->StaticMesh->GetRenderData();
-			FCardRepresentationData* OldCardData = RenderData->LODResources[0].CardRepresentationData;
+		FinalizeAndDeleteTask(Task, ObjectCacheScope);
+	}
 
-			// Assign the new data, this is safe because the render threads makes a copy of the pointer at scene proxy creation time.
-			RenderData->LODResources[0].CardRepresentationData = Task->GeneratedCardRepresentation;
+	if (bMadeProgress)
+	{
+		Notification->Update(GetNumRemainingAssets());
+	}
+#endif
+}
 
-			// Any already created render state needs to be dirtied
-			if (RenderData->IsInitialized())
-			{
-				for (IStaticMeshComponent* Component : ObjectCacheScope.GetContext().GetStaticMeshComponents(Task->StaticMesh))
-				{
-					IPrimitiveComponent* PrimitiveComponent = Component->GetPrimitiveComponentInterface();
+#if WITH_EDITOR
+void FCardRepresentationAsyncQueue::FinalizeAndDeleteTask(FAsyncCardRepresentationTask* InTask, FObjectCacheContextScope& InObjectCacheScope)
+{
+	check(InTask);
+	check(IsInGameThread()); // Timer wont work properly outside of gamethread.
 
-					if (PrimitiveComponent->IsRegistered() && PrimitiveComponent->IsRenderStateCreated())
-					{
-						PrimitiveComponent->MarkRenderStateDirty();
-					}
-				}
-			}
+	if (!InTask->bSuccess)
+	{
+		CancelAndDeleteTask(TSet<FAsyncCardRepresentationTask*>({ InTask }));
+		return;
+	}
 
-			// Rendering thread may still be referencing the old one, use the deferred cleanup interface to delete it next frame when it is safe
-			BeginCleanup(OldCardData);
+	// We want to count each resource built from a DDC miss, so count each iteration of the loop separately.
+	COOK_STAT(auto Timer = CardRepresentationCookStats::UsageStats.TimeSyncWork());
+
+	if (InTask->AsyncTask)
+	{
+		InTask->AsyncTask->EnsureCompletion();
+		InTask->AsyncTask.Reset();
+	}
 
-			// Need also to update platform render data if it's being cached
-			FStaticMeshRenderData* PlatformRenderData = RenderData->NextCachedRenderData.Get();
-			while (PlatformRenderData)
+	// Editor 'force delete' can null any UObject pointers which are seen by reference collecting (eg UProperty or serialized)
+	if (InTask->StaticMesh && InTask->bSuccess)
+	{
+		check(!InTask->StaticMesh->IsCompiling());
+		
+		FStaticMeshRenderData* RenderData = InTask->StaticMesh->GetRenderData();
+		FCardRepresentationData* OldCardData = RenderData->LODResources[0].CardRepresentationData;
+
+		// Assign the new data, this is safe because the render threads makes a copy of the pointer at scene proxy creation time.
+		RenderData->LODResources[0].CardRepresentationData = InTask->GeneratedCardRepresentation;
+
+		// Any already created render state needs to be dirtied
+		if (RenderData->IsInitialized())
+		{
+			for (IStaticMeshComponent* Component : InObjectCacheScope.GetContext().GetStaticMeshComponents(InTask->StaticMesh))
 			{
-				if (PlatformRenderData->LODResources[0].CardRepresentationData)
+				IPrimitiveComponent* PrimitiveComponent = Component->GetPrimitiveComponentInterface();
+
+				if (PrimitiveComponent->IsRegistered() && PrimitiveComponent->IsRenderStateCreated())
 				{
-					*PlatformRenderData->LODResources[0].CardRepresentationData = *Task->GeneratedCardRepresentation;
+					PrimitiveComponent->MarkRenderStateDirty();
 				}
-				PlatformRenderData = PlatformRenderData->NextCachedRenderData.Get();
 			}
+		}
+
+		// Rendering thread may still be referencing the old one, use the deferred cleanup interface to delete it next frame when it is safe
+		BeginCleanup(OldCardData);
 
-			if (!MeshCardRepresentation::IsDebugMode())
+		// Need also to update platform render data if it's being cached
+		FStaticMeshRenderData* PlatformRenderData = RenderData->NextCachedRenderData.Get();
+		while (PlatformRenderData)
+		{
+			if (PlatformRenderData->LODResources[0].CardRepresentationData)
 			{
-				TArray<uint8> DerivedData;
-				// Save built data to DDC
-				FMemoryWriter Ar(DerivedData, /*bIsPersistent=*/ true);
-				Ar << *(Task->StaticMesh->GetRenderData()->LODResources[0].CardRepresentationData);
-				GetDerivedDataCacheRef().Put(*Task->DDCKey, DerivedData, Task->StaticMesh->GetPathName());
-				COOK_STAT(Timer.AddMiss(DerivedData.Num()));
+				*PlatformRenderData->LODResources[0].CardRepresentationData = *InTask->GeneratedCardRepresentation;
 			}
+			PlatformRenderData = PlatformRenderData->NextCachedRenderData.Get();
 		}
 
-		delete Task;
+		if (!MeshCardRepresentation::IsDebugMode())
+		{
+			TArray<uint8> DerivedData;
+			// Save built data to DDC
+			FMemoryWriter Ar(DerivedData, /*bIsPersistent=*/ true);
+			Ar << *(InTask->StaticMesh->GetRenderData()->LODResources[0].CardRepresentationData);
+			GetDerivedDataCacheRef().Put(*InTask->DDCKey, DerivedData, InTask->StaticMesh->GetPathName());
+			COOK_STAT(Timer.AddMiss(DerivedData.Num()));
+		}
 	}
 
-	if (bMadeProgress)
-	{
-		Notification->Update(GetNumRemainingAssets());
-	}
-#endif
+	delete InTask;
 }
+#endif
 
 void FCardRepresentationAsyncQueue::Shutdown()
 {
diff --git a/Engine/Source/Runtime/Engine/Private/StaticMesh.cpp b/Engine/Source/Runtime/Engine/Private/StaticMesh.cpp
--- a/Engine/Source/Runtime/Engine/Private/StaticMesh.cpp
+++ b/Engine/Source/Runtime/Engine/Private/StaticMesh.cpp
@@ -6580,6 +6580,12 @@ bool UStaticMesh::IsCachedCookedPlatformDataLoaded(const ITargetPlatform* Target
 		return false;
 	}
 
+	if ((GDistanceFieldAsyncQueue != nullptr && GDistanceFieldAsyncQueue->GetNumOutstandingTasks(this) > 0) ||
+		(GCardRepresentationAsyncQueue != nullptr && GCardRepresentationAsyncQueue->GetNumOutstandingTasks(this) > 0))
+	{
+		return false;
+	}
+
 	FStaticMeshRenderData& PlatformRenderData = GetPlatformStaticMeshRenderData(this, TargetPlatform);
 
 	bool bFailed = false;
diff --git a/Engine/Source/Runtime/Engine/Public/DistanceFieldAtlas.h b/Engine/Source/Runtime/Engine/Public/DistanceFieldAtlas.h
--- a/Engine/Source/Runtime/Engine/Public/DistanceFieldAtlas.h
+++ b/Engine/Source/Runtime/Engine/Public/DistanceFieldAtlas.h
@@ -25,6 +25,7 @@
 struct FAppendToClassSchemaContext;
 struct FAssetCompileData;
 class FAsyncCompilationNotification;
+class FObjectCacheContextScope;
 class FDistanceFieldVolumeData;
 class UStaticMesh;
 class UTexture2D;
@@ -374,6 +375,9 @@ public:
 	/** Called once per frame, fetches completed tasks and applies them to the scene. */
 	ENGINE_API void ProcessAsyncTasks(bool bLimitExecutionTime = false) override;
 
+	/** Fetches completed tasks for a specific static mesh and applies them to the scene */
+	ENGINE_API void ProcessAsyncTasksForMesh(const UStaticMesh* InStaticMesh, bool bLimitExecutionTime = false);
+
 	/** Blocks until it is safe to shut down (worker threads are idle). */
 	ENGINE_API void Shutdown() override;
 
@@ -383,6 +387,9 @@ public:
 		return ReferencedTasks.Num();
 	}
 
+	/** Returns the number of outstanding referenced tasks for a specific static mesh. */
+	ENGINE_API int32 GetNumOutstandingTasks(const UStaticMesh* InStaticMesh) const;
+
 	/** Get the name of the asset type this compiler handles */
 	ENGINE_API static FName GetStaticAssetTypeName();
 
@@ -395,11 +402,12 @@ private:
 
 	friend FAsyncDistanceFieldTaskWorker;
 	void ProcessPendingTasks();
+	void ProcessPendingTasksForMesh(const UStaticMesh* InStaticMesh);
 
 	TUniquePtr<FQueuedThreadPool> ThreadPool;
 
-	/** Builds a single task with the given threadpool.  Called from the worker thread. */
-	void Build(FAsyncDistanceFieldTask* Task, class FQueuedThreadPool& ThreadPool);
+	/** Builds a single task. Called from the worker thread. */
+	void Build(FAsyncDistanceFieldTask* Task);
 
 	/** Change the priority of the background task. */
 	void RescheduleBackgroundTask(FAsyncDistanceFieldTask* InTask, EQueuedWorkPriority InPriority);
@@ -428,6 +436,11 @@ private:
 	/** Get notified when static mesh finish compiling */
 	void OnAssetPostCompile(const TArray<FAssetCompileData>& CompiledAssets);
 
+#if WITH_EDITOR
+	/** Finishes a singular async task to completion, but does not delete it */
+	void FinalizeAsyncTask(FAsyncDistanceFieldTask* InTask, FObjectCacheContextScope& InObjectCacheScope) const;
+#endif
+
 	/** Game-thread managed list of tasks in the async system. */
 	TSet<FAsyncDistanceFieldTask*> ReferencedTasks;
 
diff --git a/Engine/Source/Runtime/Engine/Public/MeshCardBuild.h b/Engine/Source/Runtime/Engine/Public/MeshCardBuild.h
--- a/Engine/Source/Runtime/Engine/Public/MeshCardBuild.h
+++ b/Engine/Source/Runtime/Engine/Public/MeshCardBuild.h
@@ -15,6 +15,7 @@
 #include <atomic>
 
 class FAsyncCompilationNotification;
+class FObjectCacheContextScope;
 class UStaticMesh;
 struct FAppendToClassSchemaContext;
 struct FSignedDistanceFieldBuildSectionData;
@@ -275,6 +276,9 @@ public:
 	/** Called once per frame, fetches completed tasks and applies them to the scene. */
 	ENGINE_API void ProcessAsyncTasks(bool bLimitExecutionTime = false) override;
 
+	/** Fetches completed tasks for a specific static mesh and applies them to the scene */
+	ENGINE_API void ProcessAsyncTasksForMesh(const UStaticMesh* InStaticMesh, bool bLimitExecutionTime = false);
+
 	/** Blocks until it is safe to shut down (worker threads are idle). */
 	ENGINE_API void Shutdown() override;
 
@@ -284,6 +288,9 @@ public:
 		return ReferencedTasks.Num();
 	}
 
+	/** Returns the number of outstanding referenced tasks for a specific static mesh. */
+	ENGINE_API int32 GetNumOutstandingTasks(const UStaticMesh* InStaticMesh) const;
+
 	/** Get the name of the asset type this compiler handles */
 	ENGINE_API static FName GetStaticAssetTypeName();
 
@@ -298,11 +305,12 @@ private:
 	ENGINE_API void FinishAllCompilation() override;
 
 	void ProcessPendingTasks();
+	void ProcessPendingTasksForMesh(const UStaticMesh* InStaticMesh);
 
 	TUniquePtr<FQueuedThreadPool> ThreadPool;
 
-	/** Builds a single task with the given threadpool.  Called from the worker thread. */
-	void Build(FAsyncCardRepresentationTask* Task, class FQueuedThreadPool& ThreadPool);
+	/** Builds a single task. Called from the worker thread. */
+	void Build(FAsyncCardRepresentationTask* Task);
 
 	/** Change the priority of the background task. */
 	void RescheduleBackgroundTask(FAsyncCardRepresentationTask* InTask, EQueuedWorkPriority InPriority);
@@ -331,6 +339,11 @@ private:
 	/** Used to cancel tasks that are not needed anymore when garbage collection occurs */
 	void OnPostReachabilityAnalysis();
 
+#if WITH_EDITOR
+	/** Finishes a singular async task to completion and deletes its owner task */
+	void FinalizeAndDeleteTask(FAsyncCardRepresentationTask* InTask, FObjectCacheContextScope& InObjectCacheScope);
+#endif
+
 	/** Game-thread managed list of tasks in the async system. */
 	TSet<FAsyncCardRepresentationTask*> ReferencedTasks;
 
