Fix missing light transmission when light has transmission enabled.

The regression was introduced by CL 48552190 while refactoring the light component. The logic in the new function TransmitsThroughSurfaces() was incorrect.

#Synced-CL 48777343

diff --git a/Engine/Source/Runtime/Engine/Private/Components/LightComponent.cpp b/Engine/Source/Runtime/Engine/Private/Components/LightComponent.cpp
--- a/Engine/Source/Runtime/Engine/Private/Components/LightComponent.cpp
+++ b/Engine/Source/Runtime/Engine/Private/Components/LightComponent.cpp
@@ -538,7 +538,7 @@ bool ULightComponent::CastsDynamicShadow() const
 
 bool ULightComponent::TransmitsThroughSurfaces() const
 {
-	return bTransmission && CastsDynamicShadow() && !CastsStaticShadow();
+	return bTransmission && CastsDynamicShadow() && !HasStaticShadowing();
 }
 
 bool ULightComponent::GetContactShadowLengthInWS() const
