New material translator: added a **temporary** system to allow the IR builder to push integer #defines in the final HLSL. This solution is only temporary to allow Substrate work to proceed, but will be refactored soon.

#rnx
#Synced-CL 48664487

diff --git a/Engine/Source/Runtime/Engine/Private/Materials/MaterialIRModule.cpp b/Engine/Source/Runtime/Engine/Private/Materials/MaterialIRModule.cpp
--- a/Engine/Source/Runtime/Engine/Private/Materials/MaterialIRModule.cpp
+++ b/Engine/Source/Runtime/Engine/Private/Materials/MaterialIRModule.cpp
@@ -26,6 +26,7 @@ void FMaterialIRModule::Empty()
 	FunctionHLSLs.Empty();
 	ParameterCollections.Empty();
 	EnvironmentDefines.Empty();
+	IntegerEnvironmentDefines.Empty();
 	UserStrings.Empty();
 	ParameterIdToData.Empty();
 	ParameterInfoToId.Empty();
diff --git a/Engine/Source/Runtime/Engine/Private/Materials/MaterialIRToHLSLTranslator.cpp b/Engine/Source/Runtime/Engine/Private/Materials/MaterialIRToHLSLTranslator.cpp
--- a/Engine/Source/Runtime/Engine/Private/Materials/MaterialIRToHLSLTranslator.cpp
+++ b/Engine/Source/Runtime/Engine/Private/Materials/MaterialIRToHLSLTranslator.cpp
@@ -2352,6 +2352,7 @@ struct FPrivate : FMaterialIRToHLSLTranslation
 		OutEnvironment.SetDefine(TEXT("IS_MATERIAL_SHADER"), true);
 		OutEnvironment.SetDefine(TEXT("VIRTUAL_TEXTURE_OUTPUT"), CompilationOutput.bHasRuntimeVirtualTextureOutputNode != 0);
 
+		// @massimo.tristano refactor this
 		uint32 DynamicParametersMask = Module->GetStatistics().DynamicParticleParameterMask;
 		if (DynamicParametersMask)
 		{
@@ -2361,11 +2362,19 @@ struct FPrivate : FMaterialIRToHLSLTranslation
 
 		// Set all defines that are defined by the module.
 		// Any conditional exemption via material properties, such as 'Material->IsUsedWithInstancedStaticMeshes()', are handled during the material IR analysis.
+		// @massimo.tristano refactor this
 		for (const FName& EnvironmentDefine : Module->GetEnvironmentDefines())
 		{
 			OutEnvironment.SetDefine(EnvironmentDefine, true);
 		}
 
+		// Temporary solution. Define the integer environment defines
+		// @massimo.tristano refactor this
+		for (const auto& Pair : Module->GetIntegerEnvironmentDefines())
+		{
+			OutEnvironment.SetDefine(Pair.Key, Pair.Value);
+		}
+
 		FMaterialShadingModelField ShadingModels = Module->GetCompiledShadingModels();
 		ensure(ShadingModels.IsValid());
 
diff --git a/Engine/Source/Runtime/Engine/Public/Materials/MaterialIRModule.h b/Engine/Source/Runtime/Engine/Public/Materials/MaterialIRModule.h
--- a/Engine/Source/Runtime/Engine/Public/Materials/MaterialIRModule.h
+++ b/Engine/Source/Runtime/Engine/Public/Materials/MaterialIRModule.h
@@ -133,9 +133,15 @@ public:
 	// Retrieves the value a meterial property is assigned to.
 	const MIR::FValue* GetPropertyValue(EMaterialProperty Property) const { return PropertyValues[Property]; }
 
-	// Returns a list of all environment define names this module requires to be enabled for shader compilation.
+	// Returns a list of all environment define names this module requires to be enabled for shader compilation. @massimo.tristano refactor this
 	const TSet<FName>& GetEnvironmentDefines() const { return EnvironmentDefines; }
 
+	// Temporary. @massimo.tristano refactor this
+	TArray<TPair<FName, int32>> GetIntegerEnvironmentDefines() const { return IntegerEnvironmentDefines; }
+
+	// Temporary. @massimo.tristano refactor this
+	void AddIntegerEnvironmentDefine(FName Name, int32 Value) { IntegerEnvironmentDefines.Emplace(Name, Value); }
+
 	// Provides access to the translated material statistics.
 	const FStatistics& GetStatistics() const { return Statistics; }
 
@@ -245,6 +251,9 @@ private:
 	// Environment define names for shader compilation in this module.
 	TSet<FName> EnvironmentDefines;
 
+	// Temporary array of integer defines. @massimo.tristano refactor this into a better way to carry output flags from the builder to the backend.
+	TArray<TPair<FName, int32>> IntegerEnvironmentDefines;
+
 	// Parameter collections referenced by this module.
 	TArray<UMaterialParameterCollection*> ParameterCollections;
 
