NetTest
* Added more utility methods to NetTestHelpers

#rnx
#Synced-CL 48743431

diff --git a/Engine/Source/Runtime/Engine/Private/Tests/NetTestHelpers.cpp b/Engine/Source/Runtime/Engine/Private/Tests/NetTestHelpers.cpp
--- a/Engine/Source/Runtime/Engine/Private/Tests/NetTestHelpers.cpp
+++ b/Engine/Source/Runtime/Engine/Private/Tests/NetTestHelpers.cpp
@@ -402,6 +402,26 @@ int32 FTestWorlds::ConnectNewClient()
 	return bConnected ? (Clients.Num()-1) : INDEX_NONE;
 }
 
+int32 FTestWorlds::FindClientIdFromWorld(const UWorld* ClientWorld) const
+{
+	const int32 Index = Clients.IndexOfByPredicate([ClientWorld](const FTestWorldInstance& ClientInstance)
+	{
+		return ClientInstance.GetWorld() == ClientWorld;
+	});
+	
+	return Index;
+}
+
+FTestWorldInstance::FContext FTestWorlds::GetClientContext(int32 ClientId) const
+{
+	if (Clients.IsValidIndex(ClientId))
+	{
+		return Clients[ClientId].GetTestContext();
+	}
+
+	return {};
+}
+
 bool FTestWorlds::WaitForClientConnect(FTestWorldInstance& Client)
 {
 	bool bConnected = TickAllUntil([&Client]()
@@ -614,6 +634,12 @@ UObject* FTestWorlds::FindReplicatedObjectOnClient(UObject* ServerObject, uint32
 		return nullptr;
 	}
 
+	if (!ServerObject)
+	{
+		ensureMsgf(false, TEXT("Expected valid pointer"));
+		return nullptr;
+	}
+
 	if (ServerObject->GetWorld() != Server.GetWorld() && ServerObject->GetOuter() != GetTransientPackage())
 	{
 		ensureMsgf(false, TEXT("FTestWorlds::FindReplicatedObjectOnClient received object %s not part of the Server world"), *GetFullNameSafe(ServerObject));
diff --git a/Engine/Source/Runtime/Engine/Public/Tests/NetTestHelpers.h b/Engine/Source/Runtime/Engine/Public/Tests/NetTestHelpers.h
--- a/Engine/Source/Runtime/Engine/Public/Tests/NetTestHelpers.h
+++ b/Engine/Source/Runtime/Engine/Public/Tests/NetTestHelpers.h
@@ -227,6 +227,12 @@ struct FTestWorlds
 	/** Create a new client and connect it to the server. Returns the ClientID of the new client or INDEX_NONE if it failed to connect. */
 	ENGINE_API int32 ConnectNewClient();
 
+	/** Returns the ClientId associated with the passed world or INDEX_NONE if it's not a known client world. */
+	ENGINE_API int32 FindClientIdFromWorld(const UWorld* ClientWorld) const;
+
+	/** Returns the world context data for an existing client. */
+	ENGINE_API FTestWorldInstance::FContext GetClientContext(int32 ClientId) const;
+
 	/** Ticks all server & client worlds NumTick times synchronously. */
 	ENGINE_API void TickAll(int32 NumTicks=1);
 	ENGINE_API void TickServer();
