[Backout] - CL48761025
#rnx
Original CL Desc
-----------------------------------------------------------------
Convert USceneComponent::AttachParent from TObjectPtr to TWeakObjectPtr to handle the case where a component get destroyed, marking it as garbage, but some other components still have their AttachParent set to the destroyed component without being actually attached.
This worked before "kill pending kill" where the GC would null-out references to pending kill objects, but now the GC reports it as a reachable garbage reference instead.
Converting the pointer to a weak ref instead is closer to the expected behavio, without having to do expensive cleanup when destroying a component.

diff --git a/Engine/Source/Runtime/Engine/Classes/Components/SceneComponent.h b/Engine/Source/Runtime/Engine/Classes/Components/SceneComponent.h
--- a/Engine/Source/Runtime/Engine/Classes/Components/SceneComponent.h
+++ b/Engine/Source/Runtime/Engine/Classes/Components/SceneComponent.h
@@ -106,7 +106,7 @@ private:
 
 	/** What we are currently attached to. If valid, RelativeLocation etc. are used relative to this object */
 	UPROPERTY(ReplicatedUsing = OnRep_AttachParent)
-	TWeakObjectPtr<USceneComponent> AttachParent;
+	TObjectPtr<USceneComponent> AttachParent;
 
 	/** Optional socket name on AttachParent that we are attached to. */
 	UPROPERTY(ReplicatedUsing = OnRep_AttachSocketName)
@@ -1696,7 +1696,7 @@ inline const TArray<TObjectPtr<USceneComponent>>& USceneComponent::GetAttachChil
 
 inline USceneComponent* USceneComponent::GetAttachParent() const
 {
-	return AttachParent.GetEvenIfUnreachable();
+	return AttachParent;
 }
 
 inline FName USceneComponent::GetAttachSocketName() const
diff --git a/Engine/Source/Runtime/Engine/Private/Components/SceneComponent.cpp b/Engine/Source/Runtime/Engine/Private/Components/SceneComponent.cpp
--- a/Engine/Source/Runtime/Engine/Private/Components/SceneComponent.cpp
+++ b/Engine/Source/Runtime/Engine/Private/Components/SceneComponent.cpp
@@ -942,7 +942,7 @@ void USceneComponent::EndPlay(EEndPlayReason::Type Reason)
 		// Detach components which are in different streaming levels so that this level can be properly garbage collected.
 		// Note that we explicitly want to check the outer hierarchy and not the owning package because we want references that participate in GC.
 		UObject* Outermost = GetOutermostObject();
-		if (GetAttachParent() && GetAttachParent()->GetOutermostObject() != Outermost)
+		if (AttachParent && AttachParent->GetOutermostObject() != Outermost)
 		{
 			DetachFromComponent(FDetachmentTransformRules::KeepWorldTransform);
 		}
@@ -3784,7 +3784,7 @@ void USceneComponent::OnRep_AttachParent()
 
 void USceneComponent::OnRep_AttachSocketName()
 {
-	if (GetAttachParent())
+	if (IsValid(AttachParent))
 	{
 		bNetUpdateAttachment = true;
 	}
@@ -3899,10 +3899,7 @@ void USceneComponent::PostRepNotifies()
 {
 	if (bNetUpdateAttachment)
 	{
-		USceneComponent* NetOldAttachParentSwapTemp = NetOldAttachParent;
-		NetOldAttachParent = GetAttachParent();
-		AttachParent = NetOldAttachParentSwapTemp;
-
+		Exchange(NetOldAttachParent, AttachParent);
 		Exchange(NetOldAttachSocketName, AttachSocketName);
 		
 		// Note: This is a local fix for JIRA UE-43355.
@@ -3920,7 +3917,7 @@ void USceneComponent::PostRepNotifies()
 		}
 
 		// Check if this is a detach
-		if (GetAttachParent() && !bShouldBeAttached)
+		if (AttachParent && !bShouldBeAttached)
 		{
 			ensureMsgf(NetOldAttachParent == nullptr, TEXT("Local modification of AttachParent detected for replicated component %s, disable replication or execute detachment on host."), *GetFullName());
 			DetachFromComponent(FDetachmentTransformRules::KeepWorldTransform);
diff --git a/Engine/Source/Runtime/Engine/Private/WorldPartition/WorldPartitionLevelHelper.cpp b/Engine/Source/Runtime/Engine/Private/WorldPartition/WorldPartitionLevelHelper.cpp
--- a/Engine/Source/Runtime/Engine/Private/WorldPartition/WorldPartitionLevelHelper.cpp
+++ b/Engine/Source/Runtime/Engine/Private/WorldPartition/WorldPartitionLevelHelper.cpp
@@ -278,6 +278,12 @@ void FWorldPartitionLevelHelper::MoveExternalActorsToLevel(const TArray<FWorldPa
 			}
 			check(Actor->GetPackage() == LevelPackage);
 
+			if (PackageObjectMapping.bIsEditorOnly)
+			{
+				Actor->SetFlags(RF_Transient);
+				UE_LOG(LogWorldPartition, Log, TEXT("Marked actor %s transient as it was referenced by an editor-only context"), *Actor->GetPathName());
+			}
+
 			// Process objects found in the source actor package
 			if (bIsActoPackageExternal)
 			{
@@ -323,6 +329,12 @@ void FWorldPartitionLevelHelper::MoveExternalActorsToLevel(const TArray<FWorldPa
 
 							Object->Rename(*NewName.ToString(), LevelPackage, REN_NonTransactional | REN_DoNotDirty | REN_DontCreateRedirectors);
 						}
+
+						if (PackageObjectMapping.bIsEditorOnly)
+						{
+							Object->SetFlags(RF_Transient);
+							UE_LOG(LogWorldPartition, Log, TEXT("Marked actor object %s transient as it was referenced by an editor-only context"), *Object->GetPathName());
+						}
 					}
 				}
 			}
diff --git a/Engine/Source/Runtime/Engine/Private/WorldPartition/WorldPartitionRuntimeLevelStreamingCell.cpp b/Engine/Source/Runtime/Engine/Private/WorldPartition/WorldPartitionRuntimeLevelStreamingCell.cpp
--- a/Engine/Source/Runtime/Engine/Private/WorldPartition/WorldPartitionRuntimeLevelStreamingCell.cpp
+++ b/Engine/Source/Runtime/Engine/Private/WorldPartition/WorldPartitionRuntimeLevelStreamingCell.cpp
@@ -328,46 +328,43 @@ void UWorldPartitionRuntimeLevelStreamingCell::AddActorToCell(const FStreamingGe
 		ParentContainerInstance = CurrentContainerInstance->GetParentContainerInstance();
 	}
 
-	if (!IsRunningCookCommandlet())
+	for (const FGuid& EditorReferenceGuid : ActorDescView.GetEditorReferences())
 	{
-		for (const FGuid& EditorReferenceGuid : ActorDescView.GetEditorReferences())
-		{
-			FName ReferencePackage;
-			FName ReferencePath;
-			FTopLevelAssetPath ReferenceBaseClass;
-			FTopLevelAssetPath ReferenceNativeClass;
-
-			// Special case where the actor descriptor view has an invalid references, use invalid reference information as the actor guid isn't necessarily in the container instance.
-			if (const FStreamingGenerationActorDescView::FInvalidReference* InvalidReference = ActorDescView.GetInvalidReference(EditorReferenceGuid))
-			{
-				ReferencePackage = InvalidReference->ActorPackage;
-				ReferencePath = *InvalidReference->ActorSoftPath.ToString();
-				ReferenceBaseClass = InvalidReference->BaseClass;
-				ReferenceNativeClass = InvalidReference->NativeClass;
-			}
-			else
-			{
-				const FWorldPartitionActorDescInstance& ReferenceActorDesc = ContainerInstance->GetActorDescInstanceChecked(EditorReferenceGuid);
-				ReferencePackage = ReferenceActorDesc.GetActorPackage();
-				ReferencePath = *ReferenceActorDesc.GetActorSoftPath().ToString();
-				ReferenceBaseClass = ReferenceActorDesc.GetBaseClass();
-				ReferenceNativeClass = ReferenceActorDesc.GetNativeClass();
-			}
+		FName ReferencePackage;
+		FName ReferencePath;
+		FTopLevelAssetPath ReferenceBaseClass;
+		FTopLevelAssetPath ReferenceNativeClass;
 
-			Packages.Emplace(
-				ReferencePackage,
-				ReferencePath,
-				ReferenceBaseClass,
-				ReferenceNativeClass,
-				ContainerID,
-				ContainerTransform,
-				FTransform::Identity,
-				ContainerPackage,
-				GetWorld()->GetPackage()->GetFName(),
-				ContainerID.GetActorGuid(EditorReferenceGuid),
-				true
-			);
+		// Special case where the actor descriptor view has an invalid references, use invalid reference information as the actor guid isn't necessarily in the container instance.
+		if (const FStreamingGenerationActorDescView::FInvalidReference* InvalidReference = ActorDescView.GetInvalidReference(EditorReferenceGuid))
+		{
+			ReferencePackage = InvalidReference->ActorPackage;
+			ReferencePath = *InvalidReference->ActorSoftPath.ToString();
+			ReferenceBaseClass = InvalidReference->BaseClass;
+			ReferenceNativeClass = InvalidReference->NativeClass;
 		}
+		else
+		{
+			const FWorldPartitionActorDescInstance& ReferenceActorDesc = ContainerInstance->GetActorDescInstanceChecked(EditorReferenceGuid);
+			ReferencePackage = ReferenceActorDesc.GetActorPackage();
+			ReferencePath = *ReferenceActorDesc.GetActorSoftPath().ToString();
+			ReferenceBaseClass = ReferenceActorDesc.GetBaseClass();
+			ReferenceNativeClass = ReferenceActorDesc.GetNativeClass();
+		}
+		
+		Packages.Emplace(
+			ReferencePackage,
+			ReferencePath,
+			ReferenceBaseClass,
+			ReferenceNativeClass,
+			ContainerID,
+			ContainerTransform,
+			FTransform::Identity,
+			ContainerPackage,
+			GetWorld()->GetPackage()->GetFName(),
+			ContainerID.GetActorGuid(EditorReferenceGuid),
+			true
+		);
 	}
 
 	FWorldPartitionRuntimeCellObjectMapping ActorMapping(
