Remove Mip Bias Fade system
- This has been unused for a long time. The calculated mip biases were never forwarded to the renderer.
- The system was also making non-threadsafe calls to FApp::GetCurrentTime() in FMipBiasFade::SetNewMipCount, which was being called on the async loading thread.

Removed associated dead code:
- GRenderingRealtimeClock. Only used by the mip bias system.
- ToggleMipFade / PauseRenderClock console commands
- bWaitForLODTransition argument on UStreamableRenderAsset functions (and replaced the optional bools with a flags enum to make the meaning more explicit)

#Synced-CL 48803475

diff --git a/Engine/Plugins/Experimental/PanoramicCapture/Source/PanoramicCapture/Private/SceneCapturer.cpp b/Engine/Plugins/Experimental/PanoramicCapture/Source/PanoramicCapture/Private/SceneCapturer.cpp
--- a/Engine/Plugins/Experimental/PanoramicCapture/Source/PanoramicCapture/Private/SceneCapturer.cpp
+++ b/Engine/Plugins/Experimental/PanoramicCapture/Source/PanoramicCapture/Private/SceneCapturer.cpp
@@ -1117,7 +1117,7 @@ void USceneCapturer::Tick( float DeltaTime )
         {
             FlushRenderingCommands();
             CaptureGameMode->ClearPause();
-            //GPauseRenderingRealtimeClock = false;
+
             CaptureStep = ECaptureStep::Pause;
             FlushRenderingCommands();
         }
@@ -1132,7 +1132,7 @@ void USceneCapturer::Tick( float DeltaTime )
 			}
 
             CaptureGameMode->SetPause(CapturePlayerController);
-            //GPauseRenderingRealtimeClock = true;
+
             CaptureStep = ECaptureStep::SetStartPosition;
             FlushRenderingCommands();
         }
@@ -1338,7 +1338,6 @@ void USceneCapturer::Tick( float DeltaTime )
 				// ----------------------------------------------------------------------------------------
 				// EXIT 
 				CaptureGameMode->ClearPause();
-				//GPauseRenderingRealtimeClock = false;
 
 				FTimespan OverallDuration = FDateTime::UtcNow() - OverallStartTime;
 
diff --git a/Engine/Plugins/Experimental/WaterAdvanced/Source/WaterAdvanced/Private/ShallowWaterRiverActor.cpp b/Engine/Plugins/Experimental/WaterAdvanced/Source/WaterAdvanced/Private/ShallowWaterRiverActor.cpp
--- a/Engine/Plugins/Experimental/WaterAdvanced/Source/WaterAdvanced/Private/ShallowWaterRiverActor.cpp
+++ b/Engine/Plugins/Experimental/WaterAdvanced/Source/WaterAdvanced/Private/ShallowWaterRiverActor.cpp
@@ -1056,7 +1056,7 @@ void UShallowWaterRiverComponent::InitializeVirtualTexture(TObjectPtr<UTexture2D
 	InTexture->VirtualTextureStreaming = true;
 
 	InTexture->UpdateResource();	
-	InTexture->WaitForStreaming(true, true);	
+	InTexture->WaitForStreaming(UStreamableRenderAsset::ETickStreamingFlags::SendCompletionEvents);
 	InTexture->BlockOnAnyAsyncBuild();
 	InTexture->PostEditChange();
 }
diff --git a/Engine/Plugins/PCG/Source/PCG/Private/Data/PCGTextureData.cpp b/Engine/Plugins/PCG/Source/PCG/Private/Data/PCGTextureData.cpp
--- a/Engine/Plugins/PCG/Source/PCG/Private/Data/PCGTextureData.cpp
+++ b/Engine/Plugins/PCG/Source/PCG/Private/Data/PCGTextureData.cpp
@@ -872,7 +872,7 @@ bool UPCGTextureData::ReadbackFromGPUTexture()
 			bUpdatedReadbackTextureResource = true;
 		}
 
-		if (Texture->HasPendingInitOrStreaming(/*bWaitForLODTransition=*/true))
+		if (Texture->HasPendingInitOrStreaming())
 		{
 			return false;
 		}
diff --git a/Engine/Plugins/PCG/Source/PCG/Private/Elements/PCGGenerateLandscapeTextures.cpp b/Engine/Plugins/PCG/Source/PCG/Private/Elements/PCGGenerateLandscapeTextures.cpp
--- a/Engine/Plugins/PCG/Source/PCG/Private/Elements/PCGGenerateLandscapeTextures.cpp
+++ b/Engine/Plugins/PCG/Source/PCG/Private/Elements/PCGGenerateLandscapeTextures.cpp
@@ -53,12 +53,11 @@ namespace PCGGenerateLandscapeTextures
 
 	static bool IsTextureFullyStreamedIn(UTexture* InTexture)
 	{
-		const bool bCheckForLODTransition = true;
 		return InTexture &&
 #if WITH_EDITOR
 			!InTexture->IsDefaultTexture() &&
 #endif // WITH_EDITOR
-			!InTexture->HasPendingInitOrStreaming(bCheckForLODTransition) && InTexture->IsFullyStreamedIn();
+			!InTexture->HasPendingInitOrStreaming() && InTexture->IsFullyStreamedIn();
 	}
 }
 
diff --git a/Engine/Plugins/PCG/Source/PCG/Private/Elements/PCGSaveTextureToAsset.cpp b/Engine/Plugins/PCG/Source/PCG/Private/Elements/PCGSaveTextureToAsset.cpp
--- a/Engine/Plugins/PCG/Source/PCG/Private/Elements/PCGSaveTextureToAsset.cpp
+++ b/Engine/Plugins/PCG/Source/PCG/Private/Elements/PCGSaveTextureToAsset.cpp
@@ -131,7 +131,7 @@ bool FPCGSaveTextureToAssetElement::ExecuteInternal(FPCGContext* InContext) cons
 		Context->ExportedTexture = ExportedTexture;
 	}
 
-	if (Context->ExportedTexture->HasPendingInitOrStreaming(/*bWaitForLODTransition=*/true))
+	if (Context->ExportedTexture->HasPendingInitOrStreaming())
 	{
 		SleepUntilNextFrame();
 		return false;
@@ -178,7 +178,7 @@ bool FPCGSaveTextureToAssetElement::ReadbackInputTexture(FPCGSaveTextureToAssetC
 				InContext->bUpdatedReadbackTextureResource = true;
 			}
 
-			if (InputTexture->HasPendingInitOrStreaming(/*bWaitForLODTransition=*/true))
+			if (InputTexture->HasPendingInitOrStreaming())
 			{
 				return false;
 			}
diff --git a/Engine/Source/Editor/UnrealEd/Private/EditorEngine.cpp b/Engine/Source/Editor/UnrealEd/Private/EditorEngine.cpp
--- a/Engine/Source/Editor/UnrealEd/Private/EditorEngine.cpp
+++ b/Engine/Source/Editor/UnrealEd/Private/EditorEngine.cpp
@@ -2481,24 +2481,12 @@ void UEditorEngine::Tick( float DeltaSeconds, bool bIdleMode )
 	// Update constraints if dirtied.
 	EditorContext.World()->UpdateConstraintActors();
 
+	ENQUEUE_RENDER_COMMAND(TickPoolElements)([](FRHICommandListImmediate& RHICmdList)
 	{
-		// rendering thread commands
-
-		bool bPauseRenderingRealtimeClock = GPauseRenderingRealtimeClock;
-		float DeltaTime = DeltaSeconds;
-		ENQUEUE_RENDER_COMMAND(TickRenderingTimer)(
-			[bPauseRenderingRealtimeClock, DeltaTime](FRHICommandListImmediate& RHICmdList)
-			{
-				if(!bPauseRenderingRealtimeClock)
-				{
-					// Tick the GRenderingRealtimeClock, unless it's paused
-					GRenderingRealtimeClock.Tick(DeltaTime);
-				}
-				GRenderTargetPool.TickPoolElements();
-				FRDGBuilder::TickPoolElements();
-				ICustomResourcePool::TickPoolElements(RHICmdList);
-			});
-	}
+		GRenderTargetPool.TickPoolElements();
+		FRDGBuilder::TickPoolElements();
+		ICustomResourcePool::TickPoolElements(RHICmdList);
+	});
 
 	// After the play world has ticked, see if a request was made to end pie
 	if (bEndPlayMapThisFrame)
diff --git a/Engine/Source/Runtime/Engine/Classes/Engine/GameViewportClient.h b/Engine/Source/Runtime/Engine/Classes/Engine/GameViewportClient.h
--- a/Engine/Source/Runtime/Engine/Classes/Engine/GameViewportClient.h
+++ b/Engine/Source/Runtime/Engine/Classes/Engine/GameViewportClient.h
@@ -244,8 +244,6 @@ public:
 	ENGINE_API bool HandleGetAllLocationCommand( const TCHAR* Cmd, FOutputDevice& Ar );
 	ENGINE_API bool HandleGetAllRotationCommand( const TCHAR* Cmd, FOutputDevice& Ar );
 	ENGINE_API bool HandleTextureDefragCommand( const TCHAR* Cmd, FOutputDevice& Ar );
-	ENGINE_API bool HandleToggleMIPFadeCommand( const TCHAR* Cmd, FOutputDevice& Ar );
-	ENGINE_API bool HandlePauseRenderClockCommand( const TCHAR* Cmd, FOutputDevice& Ar );
 
 	/**
 	 * Adds a widget to the Slate viewport's overlay (i.e for in game UI or tools) at the specified Z-order
diff --git a/Engine/Source/Runtime/Engine/Classes/Engine/StreamableRenderAsset.h b/Engine/Source/Runtime/Engine/Classes/Engine/StreamableRenderAsset.h
--- a/Engine/Source/Runtime/Engine/Classes/Engine/StreamableRenderAsset.h
+++ b/Engine/Source/Runtime/Engine/Classes/Engine/StreamableRenderAsset.h
@@ -70,7 +70,13 @@ public:
 	ENGINE_API void CancelPendingStreamingRequest();
 
 	/** Whether a stream in / out request is pending. Only one request can exists at anytime for a streamable asset. This also includes the resource initialization! */
-	ENGINE_API bool HasPendingInitOrStreaming(bool bWaitForLODTransition = false) const;
+	ENGINE_API bool HasPendingInitOrStreaming() const;
+
+	UE_DEPRECATED(5.8, "HasPendingInitOrStreaming with the optional bool argument is deprecated. bWaitForLODTransition is unused. Remove it from the arguments list.")
+	bool HasPendingInitOrStreaming(bool /*bWaitForLODTransition*/) const
+	{
+		return HasPendingInitOrStreaming();
+	}
 
 	/** Whether there is a pending update and it is locked within an update step. Used to prevent dealocks in SuspendRenderAssetStreaming(). */
 	bool IsPendingStreamingRequestLocked() const;
@@ -204,13 +210,57 @@ public:
 
 	inline int32 GetStreamingIndex() const { return StreamingIndex; }
 
+	enum class ETickStreamingFlags
+	{
+		None = 0,
+		SendCompletionEvents = 1 << 0
+	};
+	FRIEND_ENUM_CLASS_FLAGS(ETickStreamingFlags);
+
 	/** Wait for any pending operation (like InitResource and Streaming) to complete. */
-	ENGINE_API void WaitForPendingInitOrStreaming(bool bWaitForLODTransition = false, bool bSendCompletionEvents = false);
+	ENGINE_API void WaitForPendingInitOrStreaming(ETickStreamingFlags Flags = ETickStreamingFlags::None);
 
 	/** Wait for any pending operation and make sure that the asset streamer has performed requested ops on this asset. */
-	ENGINE_API void WaitForStreaming(bool bWaitForLODTransition = false, bool bSendCompletionEvents = false);
+	ENGINE_API void WaitForStreaming(ETickStreamingFlags Flags = ETickStreamingFlags::None);
+
+	void TickStreaming(ETickStreamingFlags Flags = ETickStreamingFlags::None, TArray<UStreamableRenderAsset*>* DeferredTickCBAssets = nullptr);
+
+	UE_DEPRECATED(5.8, "WaitForPendingInitOrStreaming with a single optional bool argument is deprecated. bWaitForLODTransition is unused. Call WaitForPendingInitOrStreaming with the optional ETickStreamingFlags argument.")
+	void WaitForPendingInitOrStreaming(bool /*bWaitForLODTransition*/)
+	{
+		WaitForPendingInitOrStreaming(ETickStreamingFlags::None);
+	}
+
+	UE_DEPRECATED(5.8, "WaitForPendingInitOrStreaming with two optional bool arguments is deprecated. bWaitForLODTransition is unused. Call WaitForPendingInitOrStreaming with the optional ETickStreamingFlags argument.")
+	void WaitForPendingInitOrStreaming(bool /*bWaitForLODTransition*/, bool bSendCompletionEvents)
+	{
+		WaitForPendingInitOrStreaming(bSendCompletionEvents
+			? ETickStreamingFlags::SendCompletionEvents
+			: ETickStreamingFlags::None);
+	}
 
-	void TickStreaming(bool bSendCompletionEvents = false, TArray<UStreamableRenderAsset*>* DeferredTickCBAssets = nullptr);
+	UE_DEPRECATED(5.8, "WaitForStreaming with a single optional bool argument is deprecated. bWaitForLODTransition is unused. Call WaitForStreaming with the optional ETickStreamingFlags argument.")
+	void WaitForStreaming(bool /*bWaitForLODTransition*/)
+	{
+		WaitForStreaming(ETickStreamingFlags::None);
+	}
+
+	UE_DEPRECATED(5.8, "WaitForStreaming with two optional bool arguments is deprecated. bWaitForLODTransition is unused. Call WaitForStreaming with the optional ETickStreamingFlags argument.")
+	void WaitForStreaming(bool /*bWaitForLODTransition*/, bool bSendCompletionEvents)
+	{
+		WaitForStreaming(bSendCompletionEvents
+			? ETickStreamingFlags::SendCompletionEvents
+			: ETickStreamingFlags::None);
+	}
+
+	UE_DEPRECATED(5.8, "WaitForStreaming with the optional bSendCompletionEvents bool argument is deprecated. Use the overload of TickStreaming which takes an ETickStreamingFlags argument.")
+	void TickStreaming(bool bSendCompletionEvents, TArray<UStreamableRenderAsset*>* DeferredTickCBAssets = nullptr)
+	{
+		TickStreaming(bSendCompletionEvents
+			? ETickStreamingFlags::SendCompletionEvents
+			: ETickStreamingFlags::None,
+			DeferredTickCBAssets);
+	}
 
 	virtual EStreamableRenderAssetType GetRenderAssetType() const { return EStreamableRenderAssetType::None; }
 
@@ -234,6 +284,7 @@ protected:
 	// Also returns false if the render resource is non existent, to prevent stalling on an event that will never complete.
 	virtual bool HasPendingRenderResourceInitialization() const { return false; }
 
+	UE_DEPRECATED(5.8, "HasPendingLODTransition is deprecated and no longer called. Remove any remaining uses of this function.")
 	virtual bool HasPendingLODTransition() const { return false; }
 
 	struct FLODStreamingCallbackPayload
@@ -331,3 +382,5 @@ protected:
 	friend struct FStreamingRenderAsset;
 	friend class Nanite::FCoarseMeshStreamingManager;
 };
+
+ENUM_CLASS_FLAGS(UStreamableRenderAsset::ETickStreamingFlags);
diff --git a/Engine/Source/Runtime/Engine/Classes/Engine/Texture.h b/Engine/Source/Runtime/Engine/Classes/Engine/Texture.h
--- a/Engine/Source/Runtime/Engine/Classes/Engine/Texture.h
+++ b/Engine/Source/Runtime/Engine/Classes/Engine/Texture.h
@@ -2018,7 +2018,6 @@ public:
 	ENGINE_API virtual FIoFilenameHash GetMipIoFilenameHash(const int32 MipIndex) const final override;
 	ENGINE_API virtual bool DoesMipDataExist(const int32 MipIndex) const final override;
 	ENGINE_API virtual bool HasPendingRenderResourceInitialization() const final override;
-	ENGINE_API virtual bool HasPendingLODTransition() const final override;
 	ENGINE_API virtual void InvalidateLastRenderTimeForStreaming() final override;
 	ENGINE_API virtual float GetLastRenderTimeForStreaming() const final override;
 	ENGINE_API virtual bool ShouldMipLevelsBeForcedResident() const final override;
diff --git a/Engine/Source/Runtime/Engine/Private/ContentStreaming.cpp b/Engine/Source/Runtime/Engine/Private/ContentStreaming.cpp
--- a/Engine/Source/Runtime/Engine/Private/ContentStreaming.cpp
+++ b/Engine/Source/Runtime/Engine/Private/ContentStreaming.cpp
@@ -883,10 +883,6 @@ void FStreamingManagerCollection::UpdateResourceStreaming( float DeltaTime, bool
  */
 int32 FStreamingManagerCollection::StreamAllResources( float TimeLimit/*=0.0f*/ )
 {
-	// Disable mip-fading for upcoming texture updates.
-	float PrevMipLevelFadingState = GEnableMipLevelFading;
-	GEnableMipLevelFading = -1.0f;
-
 	FlushRenderingCommands();
 
 	if (FMath::IsNearlyZero(TimeLimit))
@@ -900,8 +896,6 @@ int32 FStreamingManagerCollection::StreamAllResources( float TimeLimit/*=0.0f*/
 	// Block till requests are finished, or time limit was reached.
 	int32 NumPendingRequests = BlockTillAllRequestsFinished(TimeLimit, true);
 
-	GEnableMipLevelFading = PrevMipLevelFadingState;
-
 	return NumPendingRequests;
 }
 
diff --git a/Engine/Source/Runtime/Engine/Private/GameEngine.cpp b/Engine/Source/Runtime/Engine/Private/GameEngine.cpp
--- a/Engine/Source/Runtime/Engine/Private/GameEngine.cpp
+++ b/Engine/Source/Runtime/Engine/Private/GameEngine.cpp
@@ -2067,27 +2067,16 @@ void UGameEngine::Tick( float DeltaSeconds, bool bIdleMode )
 		GameAudioDeviceManager->UpdateActiveAudioDevices(bIsAnyNonPreviewWorldUnpaused);
 	}
 
-	// rendering thread commands
-	{
-		bool bPauseRenderingRealtimeClock = GPauseRenderingRealtimeClock;
-		ENQUEUE_RENDER_COMMAND(TickRenderingTimer)(
-			[bPauseRenderingRealtimeClock, DeltaSeconds](FRHICommandListImmediate& RHICmdList)
-		{
-			if(!bPauseRenderingRealtimeClock)
-			{
-				// Tick the GRenderingRealtimeClock, unless it's paused
-				GRenderingRealtimeClock.Tick(DeltaSeconds);
-			}
-			
-			GRenderTargetPool.TickPoolElements();
-			FRDGBuilder::TickPoolElements();
-			ICustomResourcePool::TickPoolElements(RHICmdList);
-		});
+	ENQUEUE_RENDER_COMMAND(TickPoolElements)([](FRHICommandListImmediate& RHICmdList)
+	{			
+		GRenderTargetPool.TickPoolElements();
+		FRDGBuilder::TickPoolElements();
+		ICustomResourcePool::TickPoolElements(RHICmdList);
+	});
 
-		if (bRenderingSuspended)
-		{
-			GetRendererModule().PerFrameCleanupIfSkipRenderer();
-		}
+	if (bRenderingSuspended)
+	{
+		GetRendererModule().PerFrameCleanupIfSkipRenderer();
 	}
 
 #if WITH_EDITOR
diff --git a/Engine/Source/Runtime/Engine/Private/GameViewportClient.cpp b/Engine/Source/Runtime/Engine/Private/GameViewportClient.cpp
--- a/Engine/Source/Runtime/Engine/Private/GameViewportClient.cpp
+++ b/Engine/Source/Runtime/Engine/Private/GameViewportClient.cpp
@@ -3490,14 +3490,6 @@ bool UGameViewportClient::Exec_Runtime( UWorld* InWorld, const TCHAR* Cmd,FOutpu
 	{
 		return HandleTextureDefragCommand( Cmd, Ar );
 	}
-	else if (FParse::Command(&Cmd, TEXT("TOGGLEMIPFADE")))
-	{
-		return HandleToggleMIPFadeCommand( Cmd, Ar );
-	}
-	else if (FParse::Command(&Cmd, TEXT("PAUSERENDERCLOCK")))
-	{
-		return HandlePauseRenderClockCommand( Cmd, Ar );
-	}
 	else
 	{
 		return false;
@@ -4410,21 +4402,6 @@ bool UGameViewportClient::HandleTextureDefragCommand( const TCHAR* Cmd, FOutputD
 	return true;
 }
 
-bool UGameViewportClient::HandleToggleMIPFadeCommand( const TCHAR* Cmd, FOutputDevice& Ar )
-{
-	GEnableMipLevelFading = (GEnableMipLevelFading >= 0.0f) ? -1.0f : 1.0f;
-	Ar.Logf(TEXT("Mip-fading is now: %s"), (GEnableMipLevelFading >= 0.0f) ? TEXT("ENABLED") : TEXT("DISABLED"));
-	return true;
-}
-
-bool UGameViewportClient::HandlePauseRenderClockCommand( const TCHAR* Cmd, FOutputDevice& Ar )
-{
-	GPauseRenderingRealtimeClock = !GPauseRenderingRealtimeClock;
-	Ar.Logf(TEXT("The global realtime rendering clock is now: %s"), GPauseRenderingRealtimeClock ? TEXT("PAUSED") : TEXT("RUNNING"));
-	return true;
-}
-
-
 bool UGameViewportClient::RequestBugScreenShot(const TCHAR* Cmd, bool bDisplayHUDInfo)
 {
 	// Path/name is the first (and only supported) argument
diff --git a/Engine/Source/Runtime/Engine/Private/Rendering/StreamableTextureResource.cpp b/Engine/Source/Runtime/Engine/Private/Rendering/StreamableTextureResource.cpp
--- a/Engine/Source/Runtime/Engine/Private/Rendering/StreamableTextureResource.cpp
+++ b/Engine/Source/Runtime/Engine/Private/Rendering/StreamableTextureResource.cpp
@@ -144,7 +144,6 @@ FStreamableTextureResource::FStreamableTextureResource(UTexture* InOwner, const
 	SizeY = Mip0.SizeY;
 	SizeZ = Mip0.SizeZ;
 
-	MipFadeSetting = (LODGroup == TEXTUREGROUP_Lightmap || LODGroup == TEXTUREGROUP_Shadowmap) ? MipFade_Slow : MipFade_Normal;
 	CreationFlags = 
 		  (InOwner->SRGB ? TexCreate_SRGB : TexCreate_None)
 		| (ShouldAllowPlatformTiling(InOwner) ? TexCreate_OfflineProcessed : TexCreate_None)
@@ -223,9 +222,6 @@ void FStreamableTextureResource::InitRHI(FRHICommandListBase& RHICmdList)
 		CreateTexture(RHICmdList);
 	}
 
-	// Update mip-level fading.
-	MipBiasFade.SetNewMipCount( State.NumRequestedLODs, State.NumRequestedLODs, LastRenderTime, MipFadeSetting );
-
 	if (ensure(TextureReferenceRHI.IsValid()))
 	{
 		RHICmdList.UpdateTextureReference(TextureReferenceRHI, TextureRHI);
@@ -252,11 +248,8 @@ void FStreamableTextureResource::FinalizeStreaming(FRHICommandListBase& RHICmdLi
 	// The new mip count must match the streaming request.
 	State.NumRequestedLODs = InTextureRHI->GetNumMips();
 
-	// Update mip-level fading.
 	if (State.NumResidentLODs != State.NumRequestedLODs)
 	{
-		MipBiasFade.SetNewMipCount(FMath::Max<int32>(State.NumRequestedLODs, State.NumResidentLODs), State.NumRequestedLODs, LastRenderTime, MipFadeSetting);
-
 		STAT(DecrementTextureStats());
 		STAT(CalcRequestedMipsSize());
 		STAT(IncrementTextureStats());
diff --git a/Engine/Source/Runtime/Engine/Private/StreamableRenderAsset.cpp b/Engine/Source/Runtime/Engine/Private/StreamableRenderAsset.cpp
--- a/Engine/Source/Runtime/Engine/Private/StreamableRenderAsset.cpp
+++ b/Engine/Source/Runtime/Engine/Private/StreamableRenderAsset.cpp
@@ -232,7 +232,7 @@ private:
 };
 #endif // WITH_EDITOR
 
-void UStreamableRenderAsset::TickStreaming(bool bSendCompletionEvents, TArray<UStreamableRenderAsset*>* DeferredTickCBAssets)
+void UStreamableRenderAsset::TickStreaming(ETickStreamingFlags Flags, TArray<UStreamableRenderAsset*>* DeferredTickCBAssets)
 {
 	// if resident and requested mip counts match then no pending request is in flight
 	if (PendingUpdate)
@@ -255,7 +255,7 @@ void UStreamableRenderAsset::TickStreaming(bool bSendCompletionEvents, TArray<US
 			PendingUpdate.SafeRelease();
 
 #if WITH_EDITOR
-			if (StreamableRenderAsset::bAllowUpdateResourceSize && GIsEditor && bSendCompletionEvents)
+			if (StreamableRenderAsset::bAllowUpdateResourceSize && GIsEditor && EnumHasAnyFlags(Flags, ETickStreamingFlags::SendCompletionEvents))
 			{
 				// When all the requested mips are streamed in, generate an empty property changed event, to force the
 				// ResourceSize asset registry tag to be recalculated.
@@ -265,7 +265,7 @@ void UStreamableRenderAsset::TickStreaming(bool bSendCompletionEvents, TArray<US
 		}
 	}
 
-	if (bSendCompletionEvents)
+	if (EnumHasAnyFlags(Flags, ETickStreamingFlags::SendCompletionEvents))
 	{
 		TickMipLevelChangeCallbacks(DeferredTickCBAssets);
 	}
@@ -295,7 +295,7 @@ void UStreamableRenderAsset::CancelPendingStreamingRequest()
 	}
 }
 
-bool UStreamableRenderAsset::HasPendingInitOrStreaming(bool bWaitForLODTransition) const 
+bool UStreamableRenderAsset::HasPendingInitOrStreaming() const 
 {
 	if (!!PendingUpdate)
 	{
@@ -316,21 +316,11 @@ bool UStreamableRenderAsset::HasPendingInitOrStreaming(bool bWaitForLODTransitio
 				CachedSRRState.bHasPendingInitHint = false;
 			}
 		}
-		if (bWaitForLODTransition && CachedSRRState.bHasPendingLODTransitionHint)
-		{
-			if (HasPendingLODTransition())
-			{
-				return true;
-			}
-			else if (IsInGameThread())
-			{
-				CachedSRRState.bHasPendingLODTransitionHint = false;
-			}
-		}
+		
 		return false;
 	}
 
-	return HasPendingRenderResourceInitialization() || (bWaitForLODTransition && HasPendingLODTransition());
+	return HasPendingRenderResourceInitialization();
 }
 
 /** Whether there is a pending update and it is locked within an update step. Used to prevent dealocks in SuspendRenderAssetStreaming(). */
@@ -418,7 +408,7 @@ bool UStreamableRenderAsset::IsFullyStreamedIn()
 	return false;
 }
 
-void UStreamableRenderAsset::WaitForPendingInitOrStreaming(bool bWaitForLODTransition, bool bSendCompletionEvents)
+void UStreamableRenderAsset::WaitForPendingInitOrStreaming(ETickStreamingFlags Flags)
 {
 	TRACE_CPUPROFILER_EVENT_SCOPE(UStreamableRenderAsset::WaitForPendingInitOrStreaming);
 
@@ -428,15 +418,15 @@ void UStreamableRenderAsset::WaitForPendingInitOrStreaming(bool bWaitForLODTrans
 	//	which turns off streaming
 	// @todo : Viewport resize should not turn off streaming
 
-	while (HasPendingInitOrStreaming(bWaitForLODTransition))
+	while (HasPendingInitOrStreaming())
 	{
 		// Advance the streaming state.
-		TickStreaming(bSendCompletionEvents);
+		TickStreaming(Flags);
 		// Make sure any render commands are executed, in particular things like InitRHI, or asset updates on the render thread.
 		FlushRenderingCommands();
 
 		// Most of the time, sleeping is not required, so avoid loosing a whole quantum (10ms on W10Pro) unless stricly necessary.
-		if (HasPendingInitOrStreaming(bWaitForLODTransition))
+		if (HasPendingInitOrStreaming())
 		{
 			// try to make sure streaming is enabled before doing horrible busy wait
 			ensure(!IsAssetStreamingSuspended());
@@ -447,17 +437,17 @@ void UStreamableRenderAsset::WaitForPendingInitOrStreaming(bool bWaitForLODTrans
 	}
 }
 
-void UStreamableRenderAsset::WaitForStreaming(bool bWaitForLODTransition, bool bSendCompletionEvents)
+void UStreamableRenderAsset::WaitForStreaming(ETickStreamingFlags Flags)
 {
 	// Complete pending streaming so that the asset can executing new requests if needed.
-	WaitForPendingInitOrStreaming(bWaitForLODTransition, bSendCompletionEvents);
+	WaitForPendingInitOrStreaming(Flags);
 
 	if (IsStreamable())
 	{
 		// Update the streamer state for this asset and execute new requests if needed. For example force loading to all LODs.
 		IStreamingManager::Get().GetRenderAssetStreamingManager().UpdateIndividualRenderAsset(this);
 		// Wait for any action to complete.
-		WaitForPendingInitOrStreaming(bWaitForLODTransition, bSendCompletionEvents);
+		WaitForPendingInitOrStreaming(Flags);
 	}
 }
 
diff --git a/Engine/Source/Runtime/Engine/Private/Streaming/StreamingTexture.cpp b/Engine/Source/Runtime/Engine/Private/Streaming/StreamingTexture.cpp
--- a/Engine/Source/Runtime/Engine/Private/Streaming/StreamingTexture.cpp
+++ b/Engine/Source/Runtime/Engine/Private/Streaming/StreamingTexture.cpp
@@ -269,7 +269,7 @@ FStreamableRenderResourceState FStreamingRenderAsset::UpdateStreamingStatus(bool
 
 	if (RenderAsset)
 	{
-		RenderAsset->TickStreaming(true, DeferredTickCBAssets);
+		RenderAsset->TickStreaming(UStreamableRenderAsset::ETickStreamingFlags::SendCompletionEvents, DeferredTickCBAssets);
 
 		// Call only after UpdateStreamingStatus() since it could update it.
 		ResourceState = RenderAsset->GetStreamableResourceState();
diff --git a/Engine/Source/Runtime/Engine/Private/Texture.cpp b/Engine/Source/Runtime/Engine/Private/Texture.cpp
--- a/Engine/Source/Runtime/Engine/Private/Texture.cpp
+++ b/Engine/Source/Runtime/Engine/Private/Texture.cpp
@@ -1554,11 +1554,6 @@ bool UTexture::HasPendingRenderResourceInitialization() const
 	return GetResource() && !GetResource()->IsInitialized();
 }
 
-bool UTexture::HasPendingLODTransition() const
-{
-	return GetResource() && GetResource()->MipBiasFade.IsFading();
-}
-
 float UTexture::GetLastRenderTimeForStreaming() const
 {
 	float LastRenderTime = -FLT_MAX;
diff --git a/Engine/Source/Runtime/Engine/Private/Texture2D.cpp b/Engine/Source/Runtime/Engine/Private/Texture2D.cpp
--- a/Engine/Source/Runtime/Engine/Private/Texture2D.cpp
+++ b/Engine/Source/Runtime/Engine/Private/Texture2D.cpp
@@ -1926,7 +1926,6 @@ FIntPoint FVirtualTexture2DResource::GetPhysicalTextureSize(uint32 LayerIndex) c
 static void StreamVirtualTextureMips(FVirtualTexture2DResource* VirtualTexture2DResource, FStreamableRenderResourceState& CachedSRRState, int32 NewMipCount)
 {
 	CachedSRRState.NumResidentLODs = CachedSRRState.NumRequestedLODs;
-	CachedSRRState.bHasPendingLODTransitionHint = false;
 
 	ENQUEUE_RENDER_COMMAND(RequestVTTiles)([VirtualTexture2DResource, NewMipCount](FRHICommandListImmediate& RHICmdList)
 	{
diff --git a/Engine/Source/Runtime/Engine/Public/Rendering/StreamableTextureResource.h b/Engine/Source/Runtime/Engine/Public/Rendering/StreamableTextureResource.h
--- a/Engine/Source/Runtime/Engine/Public/Rendering/StreamableTextureResource.h
+++ b/Engine/Source/Runtime/Engine/Public/Rendering/StreamableTextureResource.h
@@ -106,8 +106,6 @@ protected:
 
 	/** The FName of the texture asset */
 	FName TextureName;
-	/** Mip fade settings */
-	EMipFadeSettings MipFadeSetting = MipFade_Normal;
 	/** The asset LOD group */
 	TextureGroup LODGroup = TEXTUREGROUP_World;
 	/** Format of the texture */
diff --git a/Engine/Source/Runtime/Engine/Public/Streaming/StreamableRenderResourceState.h b/Engine/Source/Runtime/Engine/Public/Streaming/StreamableRenderResourceState.h
--- a/Engine/Source/Runtime/Engine/Public/Streaming/StreamableRenderResourceState.h
+++ b/Engine/Source/Runtime/Engine/Public/Streaming/StreamableRenderResourceState.h
@@ -50,9 +50,6 @@ struct FStreamableRenderResourceState
 			/* Whether the render resources have pending InitRHI(). Used in HasPendingInitOrStreaming() to avoid dereferencing the different resource pointers unless required. Once false, expected to stay false. */
 			mutable uint8 bHasPendingInitHint : 1;
 
-			/* Whether there is currently a visual LOD transition, following a streaming operation. Used to prevent accessing the render resource unless necessary. */
-			mutable uint8 bHasPendingLODTransitionHint : 1;
-
 			/* The number of always loaded LODs, which can not be streamed out. */
 			uint8 NumNonStreamingLODs; 
 			/* The number of LODs (streaming or not) that are guarantied to be installed, for example in PAKs. */
@@ -120,7 +117,6 @@ struct FStreamableRenderResourceState
 		if (bSupportsStreaming && InLODCount > NumResidentLODs && NumResidentLODs < MaxNumLODs && NumRequestedLODs == NumResidentLODs)
 		{
 			NumRequestedLODs = (uint8)FMath::Min<int32>(InLODCount, MaxNumLODs);
-			bHasPendingLODTransitionHint = true;
 			return true;
 		}
 		else
@@ -135,7 +131,6 @@ struct FStreamableRenderResourceState
 		if (bSupportsStreaming && InLODCount < NumResidentLODs && NumResidentLODs > NumNonStreamingLODs && NumRequestedLODs == NumResidentLODs)
 		{
 			NumRequestedLODs = (uint8)FMath::Max<int32>(InLODCount, NumNonStreamingLODs);
-			bHasPendingLODTransitionHint = true;
 			return true;
 		}
 		else
diff --git a/Engine/Source/Runtime/Landscape/Private/LandscapeTextureStreamingManager.cpp b/Engine/Source/Runtime/Landscape/Private/LandscapeTextureStreamingManager.cpp
--- a/Engine/Source/Runtime/Landscape/Private/LandscapeTextureStreamingManager.cpp
+++ b/Engine/Source/Runtime/Landscape/Private/LandscapeTextureStreamingManager.cpp
@@ -21,12 +21,11 @@ namespace UE::Landscape
 
 	static bool IsTextureFullyStreamedIn(UTexture* InTexture)
 	{
-		const bool bCheckForLODTransition = true;
 		return InTexture &&
 #if WITH_EDITOR
 			!InTexture->IsDefaultTexture() &&
 #endif // WITH_EDITOR
-			!InTexture->HasPendingInitOrStreaming(bCheckForLODTransition) && InTexture->IsFullyStreamedIn();
+			!InTexture->HasPendingInitOrStreaming() && InTexture->IsFullyStreamedIn();
 	}
 
 	static inline bool EnforceTextureIsFullyStreamedInNow(UTexture* Texture)
@@ -35,8 +34,7 @@ namespace UE::Landscape
 		// in editor, textures can be not compiled yet - we should complete that first
 		Texture->BlockOnAnyAsyncBuild();
 #endif // WITH_EDITOR
-		const bool bWaitForLODTransition = true;
-		Texture->WaitForStreaming(bWaitForLODTransition);
+		Texture->WaitForStreaming();
 		bool bIsFullyStreamedIn = IsTextureFullyStreamedIn(Texture);
 #if WITH_EDITOR
 		// the above should handle ensure textures are fully streamed in.. but just in case it isn't...
diff --git a/Engine/Source/Runtime/RenderCore/Private/RenderCore.cpp b/Engine/Source/Runtime/RenderCore/Private/RenderCore.cpp
--- a/Engine/Source/Runtime/RenderCore/Private/RenderCore.cpp
+++ b/Engine/Source/Runtime/RenderCore/Private/RenderCore.cpp
@@ -234,12 +234,6 @@ static TAutoConsoleVariable<int32> CVarForceLODShadow(
 
 #endif // EXPOSE_FORCE_LOD
 
-/** Whether to pause the global realtime clock for the rendering thread (read and write only on main thread). */
-bool GPauseRenderingRealtimeClock;
-
-/** Global realtime clock for the rendering thread. */
-FRenderTimer GRenderingRealtimeClock;
-
 FInputLatencyTimer GInputLatencyTimer( 2.0f );
 
 //
diff --git a/Engine/Source/Runtime/RenderCore/Private/RenderResource.cpp b/Engine/Source/Runtime/RenderCore/Private/RenderResource.cpp
--- a/Engine/Source/Runtime/RenderCore/Private/RenderResource.cpp
+++ b/Engine/Source/Runtime/RenderCore/Private/RenderResource.cpp
@@ -15,9 +15,6 @@
 #include "RenderCore.h"
 #include "Async/RecursiveMutex.h"
 
-/** Whether to enable mip-level fading or not: +1.0f if enabled, -1.0f if disabled. */
-float GEnableMipLevelFading = 1.0f;
-
 static thread_local FName GRenderResourceScopeName = NAME_None;
 
 class FRenderResourceList
@@ -616,83 +613,3 @@ void FBufferWithRDG::ReleaseRHI()
 	Buffer = nullptr;
 	FRenderResource::ReleaseRHI();
 }
-
-/*=============================================================================
-	FMipBiasFade class
-=============================================================================*/
-
-/** Global mip fading settings, indexed by EMipFadeSettings. */
-FMipFadeSettings GMipFadeSettings[MipFade_NumSettings] =
-{ 
-	FMipFadeSettings(0.3f, 0.1f),	// MipFade_Normal
-	FMipFadeSettings(2.0f, 1.0f)	// MipFade_Slow
-};
-
-/** How "old" a texture must be to be considered a "new texture", in seconds. */
-float GMipLevelFadingAgeThreshold = 0.5f;
-
-/**
- *	Sets up a new interpolation target for the mip-bias.
- *	@param ActualMipCount	Number of mip-levels currently in memory
- *	@param TargetMipCount	Number of mip-levels we're changing to
- *	@param LastRenderTime	Timestamp when it was last rendered (FApp::CurrentTime time space)
- *	@param FadeSetting		Which fade speed settings to use
- */
-void FMipBiasFade::SetNewMipCount( float ActualMipCount, float TargetMipCount, double LastRenderTime, EMipFadeSettings FadeSetting )
-{
-	check( ActualMipCount >=0 && TargetMipCount <= ActualMipCount );
-
-	float TimeSinceLastRendered = float(FApp::GetCurrentTime() - LastRenderTime);
-
-	// Is this a new texture or is this not in-game?
-	if ( TotalMipCount == 0 || TimeSinceLastRendered >= GMipLevelFadingAgeThreshold || GEnableMipLevelFading < 0.0f )
-	{
-		// No fading.
-		TotalMipCount = ActualMipCount;
-		MipCountDelta = 0.0f;
-		MipCountFadingRate = 0.0f;
-		StartTime = GRenderingRealtimeClock.GetCurrentTime();
-		BiasOffset = 0.0f;
-		return;
-	}
-
-	// Calculate the mipcount we're interpolating towards.
-	float CurrentTargetMipCount = TotalMipCount - BiasOffset + MipCountDelta;
-
-	// Is there no change?
-	if ( FMath::IsNearlyEqual(TotalMipCount, ActualMipCount) && FMath::IsNearlyEqual(TargetMipCount, CurrentTargetMipCount) )
-	{
-		return;
-	}
-
-	// Calculate the mip-count at our current interpolation point.
-	float CurrentInterpolatedMipCount = TotalMipCount - CalcMipBias();
-
-	// Clamp it against the available mip-levels.
-	CurrentInterpolatedMipCount = FMath::Clamp<float>(CurrentInterpolatedMipCount, 0, ActualMipCount);
-
-	// Set up a new interpolation from CurrentInterpolatedMipCount to TargetMipCount.
-	StartTime = GRenderingRealtimeClock.GetCurrentTime();
-	TotalMipCount = ActualMipCount;
-	MipCountDelta = TargetMipCount - CurrentInterpolatedMipCount;
-
-	// Don't fade if we're already at the target mip-count.
-	if ( FMath::IsNearlyZero(MipCountDelta) )
-	{
-		MipCountDelta = 0.0f;
-		BiasOffset = 0.0f;
-		MipCountFadingRate = 0.0f;
-	}
-	else
-	{
-		BiasOffset = TotalMipCount - CurrentInterpolatedMipCount;
-		if ( MipCountDelta > 0.0f )
-		{
-			MipCountFadingRate = 1.0f / (GMipFadeSettings[FadeSetting].FadeInSpeed * MipCountDelta);
-		}
-		else
-		{
-			MipCountFadingRate = -1.0f / (GMipFadeSettings[FadeSetting].FadeOutSpeed * MipCountDelta);
-		}
-	}
-}
diff --git a/Engine/Source/Runtime/RenderCore/Public/RenderResource.h b/Engine/Source/Runtime/RenderCore/Public/RenderResource.h
--- a/Engine/Source/Runtime/RenderCore/Public/RenderResource.h
+++ b/Engine/Source/Runtime/RenderCore/Public/RenderResource.h
@@ -256,99 +256,6 @@ extern RENDERCORE_API void EndBatchedRelease();
  */
 extern RENDERCORE_API void ReleaseResourceAndFlush(FRenderResource* Resource);
 
-enum EMipFadeSettings
-{
-	MipFade_Normal = 0,
-	MipFade_Slow,
-
-	MipFade_NumSettings,
-};
-
-/** Mip fade settings, selectable by chosing a different EMipFadeSettings. */
-struct FMipFadeSettings
-{
-	FMipFadeSettings( float InFadeInSpeed, float InFadeOutSpeed )
-		:	FadeInSpeed( InFadeInSpeed )
-		,	FadeOutSpeed( InFadeOutSpeed )
-	{
-	}
-
-	/** How many seconds to fade in one mip-level. */
-	float FadeInSpeed;
-
-	/** How many seconds to fade out one mip-level. */
-	float FadeOutSpeed;
-};
-
-/** Whether to enable mip-level fading or not: +1.0f if enabled, -1.0f if disabled. */
-extern RENDERCORE_API float GEnableMipLevelFading;
-
-/** Global mip fading settings, indexed by EMipFadeSettings. */
-extern RENDERCORE_API FMipFadeSettings GMipFadeSettings[MipFade_NumSettings];
-
-/**
- * Functionality for fading in/out texture mip-levels.
- */
-struct FMipBiasFade
-{
-	/** Default constructor that sets all values to default (no mips). */
-	FMipBiasFade()
-	:	TotalMipCount(0.0f)
-	,	MipCountDelta(0.0f)
-	,	StartTime(0.0f)
-	,	MipCountFadingRate(0.0f)
-	,	BiasOffset(0.0f)
-	{
-	}
-
-	/** Number of mip-levels in the texture. */
-	float	TotalMipCount;
-
-	/** Number of mip-levels to fade (negative if fading out / decreasing the mipcount). */
-	float	MipCountDelta;
-
-	/** Timestamp when the fade was started. */
-	float	StartTime;
-
-	/** Number of seconds to interpolate through all MipCountDelta (inverted). */
-	float	MipCountFadingRate;
-
-	/** Difference between total texture mipcount and the starting mipcount for the fade. */
-	float	BiasOffset;
-
-	/**
-	 *	Sets up a new interpolation target for the mip-bias.
-	 *	@param ActualMipCount	Number of mip-levels currently in memory
-	 *	@param TargetMipCount	Number of mip-levels we're changing to
-	 *	@param LastRenderTime	Timestamp when it was last rendered (FApp::CurrentTime time space)
-	 *	@param FadeSetting		Which fade speed settings to use
-	 */
-	RENDERCORE_API void	SetNewMipCount( float ActualMipCount, float TargetMipCount, double LastRenderTime, EMipFadeSettings FadeSetting );
-
-	/**
-	 *	Calculates the interpolated mip-bias based on the current time.
-	 *	@return				Interpolated mip-bias value
-	 */
-	inline float	CalcMipBias() const
-	{
-		float DeltaTime		= GRenderingRealtimeClock.GetCurrentTime() - StartTime;
-		float TimeFactor	= FMath::Min<float>(DeltaTime * MipCountFadingRate, 1.0f);
-		float MipBias		= BiasOffset - MipCountDelta*TimeFactor;
-		return FMath::FloatSelect(GEnableMipLevelFading, MipBias, 0.0f);
-	}
-
-	/**
-	 *	Checks whether the mip-bias is still interpolating.
-	 *	@return				true if the mip-bias is still interpolating
-	 */
-	inline bool	IsFading( ) const
-	{
-		float DeltaTime = GRenderingRealtimeClock.GetCurrentTime() - StartTime;
-		float TimeFactor = DeltaTime * MipCountFadingRate;
-		return (FMath::Abs<float>(MipCountDelta) > UE_SMALL_NUMBER && TimeFactor < 1.0f);
-	}
-};
-
 /** A textures resource. */
 class FTexture : public FRenderResource
 {
@@ -366,9 +273,6 @@ public:
 	/** The last time the texture has been bound */
 	mutable double		LastRenderTime = -FLT_MAX;
 
-	/** Base values for fading in/out mip-levels. */
-	FMipBiasFade		MipBiasFade;
-
 	/** bGreyScaleFormat indicates the texture is actually in R channel but should be read as Grey (replicate R to RGBA)
 	 *  this is set from CompressionSettings, not PixelFormat
 	 *  this is only used by Editor/Debug shaders, not real game materials, which use SamplerType from MaterialExpressions
diff --git a/Engine/Source/Runtime/RenderCore/Public/RenderTimer.h b/Engine/Source/Runtime/RenderCore/Public/RenderTimer.h
--- a/Engine/Source/Runtime/RenderCore/Public/RenderTimer.h
+++ b/Engine/Source/Runtime/RenderCore/Public/RenderTimer.h
@@ -4,57 +4,6 @@
 
 #include "HAL/PlatformTime.h"
 
-/** Timer helper class. **/
-class FRenderTimer
-{
-public:
-	FRenderTimer()
-		: CurrentDeltaTime(0.0f)
-		, CurrentTime(0.0f)
-	{
-	}
-
-	/**
-	 *	Returns the current time, in seconds.
-	 *	@return Current time, in seconds
-	 */
-	float GetCurrentTime() const
-	{
-		return CurrentTime;
-	}
-
-	/**
-	 *	Returns the current delta time.
-	 *	@return Current delta time (number of seconds that passed between the last two tick)
-	 */
-	float GetCurrentDeltaTime() const
-	{
-		return CurrentDeltaTime;
-	}
-
-	/**
-	 *	Updates the timer.
-	 *	@param DeltaTime	Number of seconds that have passed since the last tick
-	 **/
-	void Tick(float DeltaTime)
-	{
-		CurrentDeltaTime = DeltaTime;
-		CurrentTime += DeltaTime;
-	}
-
-protected:
-	/** Current delta time (number of seconds that passed between the last two tick). */
-	float CurrentDeltaTime;
-	/** Current time, in seconds. */
-	float CurrentTime;
-};
-
-/** Whether to pause the global realtime clock for the rendering thread (read and write only on main thread). */
-extern RENDERCORE_API bool GPauseRenderingRealtimeClock;
-
-/** Global realtime clock for the rendering thread. */
-extern RENDERCORE_API FRenderTimer GRenderingRealtimeClock;
-
 /**
  * Encapsulates a latency timer that measures the time from when mouse input
  * is read on the gamethread until that frame is fully displayed by the GPU.
