[WorldPartition] Force load actor's DataLayer assets before updating actor desc to guarantee desc remains in sync with actor.

#Synced-CL 48520266

diff --git a/Engine/Source/Runtime/Engine/Classes/GameFramework/Actor.h b/Engine/Source/Runtime/Engine/Classes/GameFramework/Actor.h
--- a/Engine/Source/Runtime/Engine/Classes/GameFramework/Actor.h
+++ b/Engine/Source/Runtime/Engine/Classes/GameFramework/Actor.h
@@ -1465,6 +1465,7 @@ public:
 	ENGINE_API bool AddDataLayer(const UDataLayerInstance* DataLayerInstance);
 	ENGINE_API bool RemoveDataLayer(const UDataLayerInstance* DataLayerInstance);
 	ENGINE_API bool CanAddDataLayer(const UDataLayerInstance* InDataLayerInstance, FText* OutReason = nullptr) const;
+	ENGINE_API void ForceLoadDataLayerAssets() const;
 
 	ENGINE_API TArray<const UDataLayerInstance*> RemoveAllDataLayers();
 	ENGINE_API bool SupportsDataLayerType(TSubclassOf<UDataLayerInstance> DataLayerType) const;
diff --git a/Engine/Source/Runtime/Engine/Private/ActorEditor.cpp b/Engine/Source/Runtime/Engine/Private/ActorEditor.cpp
--- a/Engine/Source/Runtime/Engine/Private/ActorEditor.cpp
+++ b/Engine/Source/Runtime/Engine/Private/ActorEditor.cpp
@@ -1796,6 +1796,14 @@ TArray<const UDataLayerAsset*> AActor::ResolveDataLayerAssets(const TArray<TSoft
 	return ResolvedAssets;
 }
 
+void AActor::ForceLoadDataLayerAssets() const
+{
+	for (const TSoftObjectPtr<UDataLayerAsset>& DataLayerAsset : DataLayerAssets)
+	{
+		DataLayerAsset.LoadSynchronous();
+	}
+}
+
 TArray<const UDataLayerAsset*> AActor::GetDataLayerAssets(bool bIncludeExternalDataLayerAsset) const
 {
 	TArray<const UDataLayerAsset*> ResolveDataLayerAsset = ResolveDataLayerAssets(DataLayerAssets);
diff --git a/Engine/Source/Runtime/Engine/Private/WorldPartition/WorldPartitionActorDesc.cpp b/Engine/Source/Runtime/Engine/Private/WorldPartition/WorldPartitionActorDesc.cpp
--- a/Engine/Source/Runtime/Engine/Private/WorldPartition/WorldPartitionActorDesc.cpp
+++ b/Engine/Source/Runtime/Engine/Private/WorldPartition/WorldPartitionActorDesc.cpp
@@ -159,6 +159,8 @@ void FWorldPartitionActorDesc::Init(const AActor* InActor)
 
 		if (UWorldPartition* ActorWorldPartition = FWorldPartitionHelpers::GetWorldPartition(InActor))
 		{
+			// Force load Data Layer Assets (SoftObjectPtrs) to make sure desc remains in sync with the actor
+			InActor->ForceLoadDataLayerAssets();
 			const bool bIncludeExternalDataLayerAsset = false;
 			TArray<const UDataLayerAsset*> DataLayerAssets = InActor->GetDataLayerAssets(bIncludeExternalDataLayerAsset);
 			LocalDataLayerAssetPaths.Reserve(DataLayerAssets.Num());
