Remove UGameViewportClient::SetDropDetail
- This is a legacy function which exists only to set bDropDetail and bAggressiveLOD on the UWorld. The former is only used by legacy cascade UParticleModuleCollision::Update to skip updates if the frame is over budget. The latter is not used at all.

#redirectsaresafe
#Synced-CL 48803475

diff --git a/Engine/Config/BaseEngine.ini b/Engine/Config/BaseEngine.ini
--- a/Engine/Config/BaseEngine.ini
+++ b/Engine/Config/BaseEngine.ini
@@ -321,7 +321,6 @@ NetClientTicksPerSecond=200
 +StatColorMappings=(StatName="Frametime",ColorMap=((In=1.0,Out=(G=255)),(In=25.0,Out=(G=255)),(In=29.0,Out=(R=255,G=255)),(In=33.0,Out=(R=255))))
 +StatColorMappings=(StatName="Streaming fudge factor",ColorMap=((In=0.0,Out=(G=255)),(In=1.0,Out=(G=255)),(In=2.5,Out=(R=255,G=255)),(In=5.0,Out=(R=255)),(In=10.0,Out=(R=255))))
 DisplayGamma=2.2
-MinDesiredFrameRate=35.000000
 ; NetDriver configs
 +NetDriverDefinitions=(DefName="GameNetDriver",DriverClassName="/Script/OnlineSubsystemUtils.IpNetDriver",DriverClassNameFallback="/Script/OnlineSubsystemUtils.IpNetDriver")
 +NetDriverDefinitions=(DefName="BeaconNetDriver",DriverClassName="/Script/OnlineSubsystemUtils.IpNetDriver",DriverClassNameFallback="/Script/OnlineSubsystemUtils.IpNetDriver")
diff --git a/Engine/Plugins/FX/CascadeToNiagaraConverter/Source/CascadeToNiagaraConverter/Private/NiagaraStackGraphUtilitiesAdapterLibrary.cpp b/Engine/Plugins/FX/CascadeToNiagaraConverter/Source/CascadeToNiagaraConverter/Private/NiagaraStackGraphUtilitiesAdapterLibrary.cpp
--- a/Engine/Plugins/FX/CascadeToNiagaraConverter/Source/CascadeToNiagaraConverter/Private/NiagaraStackGraphUtilitiesAdapterLibrary.cpp
+++ b/Engine/Plugins/FX/CascadeToNiagaraConverter/Source/CascadeToNiagaraConverter/Private/NiagaraStackGraphUtilitiesAdapterLibrary.cpp
@@ -956,7 +956,7 @@ void UFXConverterUtilitiesLibrary::GetParticleModuleCollisionProps(
 	bOutOnlyVerticalNormalsDecrementCount = ParticleModule->bOnlyVerticalNormalsDecrementCount;
 	OutVerticalFudgeFactor = ParticleModule->VerticalFudgeFactor;
 	OutDelayAmount = ParticleModule->DelayAmount.Distribution;
-	bOutDropDetail = ParticleModule->bDropDetail;
+	bOutDropDetail = false;
 	bOutCollideOnlyIfVisible = ParticleModule->bCollideOnlyIfVisible;
 	bOutIgnoreSourceActor = ParticleModule->bIgnoreSourceActor;
 	OutMaxCollisionDistance = ParticleModule->MaxCollisionDistance;
diff --git a/Engine/Plugins/VirtualProduction/Avalanche/Source/Avalanche/Private/Framework/AvaGameInstance.cpp b/Engine/Plugins/VirtualProduction/Avalanche/Source/Avalanche/Private/Framework/AvaGameInstance.cpp
--- a/Engine/Plugins/VirtualProduction/Avalanche/Source/Avalanche/Private/Framework/AvaGameInstance.cpp
+++ b/Engine/Plugins/VirtualProduction/Avalanche/Source/Avalanche/Private/Framework/AvaGameInstance.cpp
@@ -257,11 +257,6 @@ void UAvaGameInstance::Tick(float DeltaSeconds)
 		return;
 	}
 	
-	{
-		QUICK_SCOPE_CYCLE_COUNTER(STAT_MotionDesignGame_Tick_SetDropDetail);
-		ViewportClient->SetDropDetail(DeltaSeconds);
-	}
-	
 	/*
 	 * Things Removed Here because they're currently unneeded and will possibly never be
 	 * TickWorldTravel
diff --git a/Engine/Source/Editor/UnrealEd/Private/EditorEngine.cpp b/Engine/Source/Editor/UnrealEd/Private/EditorEngine.cpp
--- a/Engine/Source/Editor/UnrealEd/Private/EditorEngine.cpp
+++ b/Engine/Source/Editor/UnrealEd/Private/EditorEngine.cpp
@@ -2151,13 +2151,6 @@ void UEditorEngine::Tick( float DeltaSeconds, bool bIdleMode )
 					PlayWorld->UpdateLevelStreaming();
 				}
 
-				// Release mouse if the game is paused. The low level input code might ignore the request when e.g. in fullscreen mode.
-				if ( GameViewport != NULL && GameViewport->Viewport != NULL )
-				{
-					// Decide whether to drop high detail because of frame rate
-					GameViewport->SetDropDetail(TickDeltaSeconds);
-				}
-
 				// Update the level.
 				{
 					// So that hierarchical stats work in PIE
diff --git a/Engine/Source/Runtime/Engine/Classes/Engine/Engine.h b/Engine/Source/Runtime/Engine/Classes/Engine/Engine.h
--- a/Engine/Source/Runtime/Engine/Classes/Engine/Engine.h
+++ b/Engine/Source/Runtime/Engine/Classes/Engine/Engine.h
@@ -1823,10 +1823,6 @@ public:
 	UPROPERTY(config)
 	float DisplayGamma;
 
-	/** Minimum desired framerate setting, below this frame rate visual detail may be lowered */
-	UPROPERTY(config, EditAnywhere, Category=Framerate, meta=(UIMin=0, ClampMin=0, EditCondition="!bUseFixedFrameRate"))
-	float MinDesiredFrameRate;
-
 private:
 	/** Default color of selected objects in the level viewport (additive) */
 	UPROPERTY(globalconfig)
diff --git a/Engine/Source/Runtime/Engine/Classes/Engine/GameViewportClient.h b/Engine/Source/Runtime/Engine/Classes/Engine/GameViewportClient.h
--- a/Engine/Source/Runtime/Engine/Classes/Engine/GameViewportClient.h
+++ b/Engine/Source/Runtime/Engine/Classes/Engine/GameViewportClient.h
@@ -352,14 +352,6 @@ public:
 	{
 		 return Window.Pin();
 	}
-	
-	/** 
-	 * Sets bDropDetail and other per-frame detail level flags on the current WorldSettings
-	 *
-	 * @param DeltaSeconds - amount of time passed since last tick
-	 * @see UWorld
-	 */
-	ENGINE_API virtual void SetDropDetail(float DeltaSeconds);
 
 	/**
 	 * Process Console Command
diff --git a/Engine/Source/Runtime/Engine/Classes/Engine/World.h b/Engine/Source/Runtime/Engine/Classes/Engine/World.h
--- a/Engine/Source/Runtime/Engine/Classes/Engine/World.h
+++ b/Engine/Source/Runtime/Engine/Classes/Engine/World.h
@@ -1301,12 +1301,6 @@ public:
 	uint8 bIsNameStableForNetworking:1;
 #endif
 
-	/** frame rate is below DesiredFrameRate, so drop high detail actors */
-	uint8 bDropDetail:1;
-
-	/** frame rate is well below DesiredFrameRate, so make LOD more aggressive */
-	uint8 bAggressiveLOD:1;
-
 	/** That map is default map or not **/
 	uint8 bIsDefaultLevel:1;
 	
diff --git a/Engine/Source/Runtime/Engine/Classes/Particles/Collision/ParticleModuleCollision.h b/Engine/Source/Runtime/Engine/Classes/Particles/Collision/ParticleModuleCollision.h
--- a/Engine/Source/Runtime/Engine/Classes/Particles/Collision/ParticleModuleCollision.h
+++ b/Engine/Source/Runtime/Engine/Classes/Particles/Collision/ParticleModuleCollision.h
@@ -142,10 +142,6 @@ class UParticleModuleCollision : public UParticleModuleCollisionBase
 	UPROPERTY(EditAnywhere, Category=Collision)
 	struct FRawDistributionFloat DelayAmount;
 
-	/**	If true, when the World->bDropDetail flag is set, the module will be ignored. */
-	UPROPERTY(EditAnywhere, Category=Performance)
-	uint32 bDropDetail:1;
-
 	/** If true, Particle collision only if particle system is currently being rendered. */
 	UPROPERTY(EditAnywhere, Category=Performance)
 	uint32 bCollideOnlyIfVisible:1;	
diff --git a/Engine/Source/Runtime/Engine/Private/GameEngine.cpp b/Engine/Source/Runtime/Engine/Private/GameEngine.cpp
--- a/Engine/Source/Runtime/Engine/Private/GameEngine.cpp
+++ b/Engine/Source/Runtime/Engine/Private/GameEngine.cpp
@@ -1786,13 +1786,6 @@ void UGameEngine::Tick( float DeltaSeconds, bool bIdleMode )
 		return;
 	}
 
-	if (GameViewport != NULL)
-	{
-		// Decide whether to drop high detail because of frame rate.
-		QUICK_SCOPE_CYCLE_COUNTER(STAT_UGameEngine_Tick_SetDropDetail);
-		GameViewport->SetDropDetail(DeltaSeconds);
-	}
-
 #if !UE_SERVER
 	// Media module present?
 	static const FName MediaModuleName("Media");
diff --git a/Engine/Source/Runtime/Engine/Private/GameViewportClient.cpp b/Engine/Source/Runtime/Engine/Private/GameViewportClient.cpp
--- a/Engine/Source/Runtime/Engine/Private/GameViewportClient.cpp
+++ b/Engine/Source/Runtime/Engine/Private/GameViewportClient.cpp
@@ -1228,58 +1228,6 @@ TOptional<TSharedRef<SWidget>> UGameViewportClient::MapCursor(FViewport* InViewp
 	return TOptional<TSharedRef<SWidget>>();
 }
 
-void UGameViewportClient::SetDropDetail(float DeltaSeconds)
-{
-	if (GEngine && GetWorld())
-	{
-		float FrameTime = 0.0f;
-		if (FPlatformProperties::SupportsWindowedMode() == false)
-		{
-			FrameTime = FPlatformTime::ToSeconds(FMath::Max3<uint32>(GRenderThreadTime, GGameThreadTime, RHIGetGPUFrameCycles()));
-			// If DeltaSeconds is bigger than 34 ms we can take it into account as we're not VSYNCing in that case.
-			if( DeltaSeconds > 0.034 )
-			{
-				FrameTime = FMath::Max( FrameTime, DeltaSeconds );
-			}
-		}
-		else
-		{
-			FrameTime = DeltaSeconds;
-		}
-		const float FrameRate	= FrameTime > 0 ? 1 / FrameTime : 0;
-
-		// When using FixedFrameRate, FrameRate here becomes FixedFrameRate (even if actual framerate is smaller).
-		const bool bTimeIsManipulated = FApp::IsBenchmarking() || FApp::UseFixedTimeStep() || GEngine->bUseFixedFrameRate;
-		// Drop detail if framerate is below threshold.
-		GetWorld()->bDropDetail		= FrameRate < FMath::Clamp(GEngine->MinDesiredFrameRate, 1.f, 100.f) && !bTimeIsManipulated;
-		GetWorld()->bAggressiveLOD	= FrameRate < FMath::Clamp(GEngine->MinDesiredFrameRate - 5.f, 1.f, 100.f) && !bTimeIsManipulated;
-
-		// this is slick way to be able to do something based on the frametime and whether we are bound by one thing or another
-#if 0
-		// so where we check to see if we are above some threshold and below 150 ms (any thing above that is usually blocking loading of some sort)
-		// also we don't want to do the auto trace when we are blocking on async loading
-		if ((0.070 < FrameTime) && (FrameTime < 0.150) && IsAsyncLoading() == false && GetWorld()->bRequestedBlockOnAsyncLoading == false && (GetWorld()->GetTimeSeconds() > 30.0f))
-		{
-			// now check to see if we have done a trace in the last 30 seconds otherwise we will just trace ourselves to death
-			static float LastTraceTime = -9999.0f;
-			if( (LastTraceTime+30.0f < GetWorld()->GetTimeSeconds()))
-			{
-				LastTraceTime = GetWorld()->GetTimeSeconds();
-				UE_LOG(LogPlayerManagement, Warning, TEXT("Auto Trace initiated!! FrameTime: %f"), FrameTime );
-
-				// do what ever action you want here (e.g. trace <type>, GShouldLogOutAFrameOfMoveActor = true, c.f. LevelTick.cpp for more)
-				//GShouldLogOutAFrameOfMoveActor = true;
-
-#if !WITH_EDITORONLY_DATA
-				UE_LOG(LogPlayerManagement, Warning, TEXT("    GGameThreadTime: %d GRenderThreadTime: %d "), GGameThreadTime, GRenderThreadTime );
-#endif // WITH_EDITORONLY_DATA
-			}
-		}
-#endif // 0
-	}
-}
-
-
 void UGameViewportClient::SetViewportFrame( FViewportFrame* InViewportFrame )
 {
 	ViewportFrame = InViewportFrame;
diff --git a/Engine/Source/Runtime/Engine/Private/Particles/ParticleModules_Collision.cpp b/Engine/Source/Runtime/Engine/Private/Particles/ParticleModules_Collision.cpp
--- a/Engine/Source/Runtime/Engine/Private/Particles/ParticleModules_Collision.cpp
+++ b/Engine/Source/Runtime/Engine/Private/Particles/ParticleModules_Collision.cpp
@@ -60,7 +60,6 @@ UParticleModuleCollision::UParticleModuleCollision(const FObjectInitializer& Obj
 	bApplyPhysics = false;
 	DirScalar = 3.5f;
 	VerticalFudgeFactor = 0.1f;
-	bDropDetail = true;
 	LODDuplicate = false;
 	bPawnsDoNotDecrementCount = true;
 	bCollideOnlyIfVisible = true;
@@ -184,7 +183,7 @@ void UParticleModuleCollision::Update(const FUpdateContext& Context)
 	FParticleEmitterInstance* Owner = &Context.Owner;
 	UParticleSystemComponent* Component = Owner->Component.AsComponent();
 	UWorld* World = Component->GetWorld();
-	if (Owner->ActiveParticles == 0 || (bDropDetail && World && World->bDropDetail))
+	if (Owner->ActiveParticles == 0)
 	{
 		return;
 	}
