[Partial Island Sleep] Expose tunable CVars to project settings
- only for parameters which are supposed to be easily configurable by the user
- ON/OFF,
- Momentum Propagation,
- Post Step Wake Threshold,
- Minimum Motionless,
- Linear & Angular Wake Threshold Multiplier
- only one set of parameters shared for all islands
- all parameters can be set through the CVar (as previously) or also in the project settings.

#Synced-CL 47960832

diff --git a/Engine/Source/Runtime/Engine/Classes/PhysicsEngine/PhysicsSettings.h b/Engine/Source/Runtime/Engine/Classes/PhysicsEngine/PhysicsSettings.h
--- a/Engine/Source/Runtime/Engine/Classes/PhysicsEngine/PhysicsSettings.h
+++ b/Engine/Source/Runtime/Engine/Classes/PhysicsEngine/PhysicsSettings.h
@@ -372,6 +372,49 @@ class UPhysicsSettings : public UPhysicsSettingsCore
 	UPROPERTY(config, EditAnywhere, Category = ChaosPhysics)
 	FChaosPhysicsSettings ChaosSettings;
 
+	/** Allow/disallow partial island sleep for all islands.
+	* We disable the dynamics of a rigid body at rest ("put it to sleep") even if the surrounding bodies in the island are still actively simulated ("the island is awake").
+	* This is a CPU optimization for large simulations with many motionless bodies (e.g. structures stacks and unstructured piles at rest).
+	* Note that partial island sleep is only executed internally if there are enough motionless bodies in the respective island (see Minimum Motionless). */
+	UPROPERTY(config, EditAnywhere, Category = "ChaosPhysics|SleepSettings", meta = (ConsoleVariable = "p.Chaos.Solver.Sleep.PartialIslandSleep", ConfigRestartRequired = false))
+	bool bPartialIslandSleep;
+
+	/** Enable/disable momentum propagation through the constraint graph during partial island waking. 
+	* New collisions, constraints or user-induced actor manipulation generate partial wake events with sleeping bodies.
+	* The linear and angular momentum of an active, dynamic body is defined as its mass and inertia times its linear and angular velocity, respectively.
+	* This parameter defines the decimal percentage propagated through constraints between sleeping bodies. 
+	* 0.3 means that 30% of the incoming momentum is carried over to the next body.
+	* (0.0 to disable, 1.0 to propagate without dissipation). */
+	UPROPERTY(config, EditAnywhere, Category = "ChaosPhysics|SleepSettings", meta = (ClampMin = "0.0", ClampMax = "1.0", EditCondition = bPartialIslandSleep, ConsoleVariable = "p.Chaos.Solver.Sleep.PartialIslandSleep.MomentumPropagation", ConfigRestartRequired = false))
+	float MomentumPropagation;
+
+	/** Enable/disable end of step partial waking caused by velocity solver impulses between any active and sleeping body pair.
+	* The solver impulse is converted to a change in velocity through division by the rigid body mass.
+	* The change in velocity induced by the velocity solver needs to be above this threshold to trigger the wake of the sleep.
+	* Note that only one layer of sleeping bodies can be woken up per step, i.e. propagation through a chain of X bodies will take X steps.
+	* (0.0 to disable). */
+	UPROPERTY(config, EditAnywhere, Category = "ChaosPhysics|SleepSettings", meta = (ClampMin = "0.0", EditCondition = bPartialIslandSleep, ConsoleVariable = "p.Chaos.Solver.Sleep.PartialIslandSleep.PostStepWakeThreshold", ConfigRestartRequired = false))
+	float PostStepWakeThreshold;
+
+	/** The minimum decimal percentage of rigid bodies in an island which need to be motionless to activate partial island sleep internally. 
+	* 0.0: Partial island sleep and wake is always executed internally.
+	* 1.0: All bodies need to be motionless for partial island sleep and wake to be executed internally. */
+	UPROPERTY(config, EditAnywhere, Category = "ChaosPhysics|SleepSettings", meta = (ClampMin = "0.0", ClampMax = "1.0", EditCondition = bPartialIslandSleep, ConsoleVariable = "p.Chaos.Solver.Sleep.PartialIslandSleep.MinimumMotionlessRatio", ConfigRestartRequired = false))
+	float MinimumMotionless;
+
+	/** The linear wake threshold is computed as this multiplier times the body's linear sleep threshold. 
+	* It is used to determine when an active, dynamic body involved in a new constraint or collision will trigger a partial wake.
+	* It also defines the minimum motion momentum (multiplied by the body's mass) to limit momentum propagation throughout the constraint graph. */
+	UPROPERTY(config, EditAnywhere, Category = "ChaosPhysics|SleepSettings", meta = (ClampMin = "1.0", EditCondition = bPartialIslandSleep, ConsoleVariable = "p.Chaos.Solver.Sleep.PartialIslandSleep.LinearWakeThresholdMultiplier", ConfigRestartRequired = false))
+	float LinearWakeThresholdMultiplier;
+
+	/** The angular wake threshold is computed as this multiplier times the body's angular sleep threshold.
+	* It is used to determine when an active, dynamic body involved in a new constraint or collision will trigger a partial wake.
+	* It also defines the minimum motion momentum (multiplied by the body's inertia) to limit momentum propagation throughout the constraint graph. */
+	UPROPERTY(config, EditAnywhere, Category = "ChaosPhysics|SleepSettings", meta = (ClampMin = "1.0", EditCondition = bPartialIslandSleep, ConsoleVariable = "p.Chaos.Solver.Sleep.PartialIslandSleep.AngularWakeThresholdMultiplier", ConfigRestartRequired = false))
+	float AngularWakeThresholdMultiplier;
+
+
 public:
 
 	static UPhysicsSettings* Get() { return CastChecked<UPhysicsSettings>(UPhysicsSettings::StaticClass()->GetDefaultObject()); }
diff --git a/Engine/Source/Runtime/Engine/Private/PhysicsEngine/PhysicsSettings.cpp b/Engine/Source/Runtime/Engine/Private/PhysicsEngine/PhysicsSettings.cpp
--- a/Engine/Source/Runtime/Engine/Private/PhysicsEngine/PhysicsSettings.cpp
+++ b/Engine/Source/Runtime/Engine/Private/PhysicsEngine/PhysicsSettings.cpp
@@ -29,6 +29,12 @@ UPhysicsSettings::UPhysicsSettings(const FObjectInitializer& ObjectInitializer)
 	, InitialAverageFrameRate(1.f / 60.f)
 	, PhysXTreeRebuildRate(10)
 	, MinDeltaVelocityForHitEvents(0.f)
+	, bPartialIslandSleep(false)
+	, MomentumPropagation(0.3f)
+	, PostStepWakeThreshold(20.0f)
+	, MinimumMotionless(0.1f)
+	, LinearWakeThresholdMultiplier(10.0f)
+	, AngularWakeThresholdMultiplier(10.0f)
 {
 }
 
@@ -37,7 +43,12 @@ void UPhysicsSettings::PostInitProperties()
 	Super::PostInitProperties();
 #if WITH_EDITOR
 	LoadSurfaceType();
-#endif
+
+	if (IsTemplate())
+	{
+		ImportConsoleVariableValues();
+	}
+#endif // #if WITH_EDITOR
 
 	if (PhysicsPrediction.bEnablePhysicsResimulation_DEPRECATED)
 	{
@@ -111,6 +122,11 @@ void UPhysicsSettings::PostEditChangeProperty(struct FPropertyChangedEvent& Prop
 {
 	Super::PostEditChangeProperty(PropertyChangedEvent);
 
+	if (PropertyChangedEvent.Property)
+	{
+		ExportValuesToConsoleVariables(PropertyChangedEvent.Property);
+	}
+
 	const FName PropertyName = PropertyChangedEvent.Property ? PropertyChangedEvent.Property->GetFName() : NAME_None;
 	if (PropertyName == GET_MEMBER_NAME_CHECKED(UPhysicsSettings, FrictionCombineMode) || PropertyName == GET_MEMBER_NAME_CHECKED(UPhysicsSettings, RestitutionCombineMode))
 	{
