[WorldPartition] Fixed WorldPartitionRuntimeCellTransformerISM duplicating instances for Blueprint Actors that were not fully transformable.
- WorldPartitionRuntimeCellTransformerISM now also doesn't convert Blueprint Actors with logic (shared code with FastGeo transformer).

diff --git a/Engine/Plugins/Experimental/FastGeoStreaming/Source/FastGeoStreaming/Private/FastGeoWorldPartitionRuntimeCellTransformer.cpp b/Engine/Plugins/Experimental/FastGeoStreaming/Source/FastGeoStreaming/Private/FastGeoWorldPartitionRuntimeCellTransformer.cpp
--- a/Engine/Plugins/Experimental/FastGeoStreaming/Source/FastGeoStreaming/Private/FastGeoWorldPartitionRuntimeCellTransformer.cpp
+++ b/Engine/Plugins/Experimental/FastGeoStreaming/Source/FastGeoStreaming/Private/FastGeoWorldPartitionRuntimeCellTransformer.cpp
@@ -9,8 +9,6 @@
 #include "GameFramework/Actor.h"
 #include "Engine/Level.h"
 #include "Engine/Blueprint.h"
-#include "Engine/BlueprintGeneratedClass.h"
-#include "EdGraph/EdGraph.h"
 #include "UObject/Package.h"
 #include "Modules/ModuleManager.h"
 #include "Physics/Experimental/PhysScene_Chaos.h"
@@ -614,69 +612,6 @@ void UFastGeoWorldPartitionRuntimeCellTransformer::FTransformationStats::DumpSta
 	}
 }
 
-bool UFastGeoWorldPartitionRuntimeCellTransformer::IsBlueprintActorWithLogic(AActor* InActor) const
-{
-	const FName FN_UserConstructionScript(TEXT("UserConstructionScript"));
-
-	UBlueprint* Blueprint = UBlueprint::GetBlueprintFromClass(InActor->GetClass());
-	if (!Blueprint || !Blueprint->GeneratedClass)
-	{
-		return false;
-	}
-
-	check(Blueprint->ParentClass && Blueprint->ParentClass->IsChildOf(AActor::StaticClass()));
-
-	UBlueprintGeneratedClass* BPClass = Cast<UBlueprintGeneratedClass>(Blueprint->GeneratedClass);
-	if (!BPClass)
-	{
-		return false;
-	}
-
-	if (Blueprint->DelegateSignatureGraphs.Num() > 0)
-	{
-		return true;
-	}
-
-	if (Blueprint->ImplementedInterfaces.Num() > 0)
-	{
-		return true;
-	}
-
-	// Check if no extra functions, other than the user construction script (only AActor and subclasses of AActor have)
-	if (Blueprint->FunctionGraphs.Num() > 1)
-	{
-		return true;
-	}
-
-	check(Blueprint->FunctionGraphs.Num() == 0 || Blueprint->FunctionGraphs[0]->GetFName() == FN_UserConstructionScript);
-
-	// Check if the generated class has overridden any functions dynamically
-	for (TFieldIterator<UFunction> It(BPClass, EFieldIteratorFlags::IncludeSuper); It; ++It)
-	{
-		UFunction* Function = *It;
-
-		// Ignore functions from native C++ classes (i.e., inherited but not overridden in BP)
-		if (Function->GetOwnerClass() == BPClass && Function->GetFName() != FN_UserConstructionScript)
-		{
-			return true; // Found an overridden function
-		}
-	}
-
-	// If there is an enabled node in the event graph, the Blueprint is not data only
-	for (UEdGraph* EventGraph : Blueprint->UbergraphPages)
-	{
-		for (UEdGraphNode* GraphNode : EventGraph->Nodes)
-		{
-			if (GraphNode && (GraphNode->GetDesiredEnabledState() != ENodeEnabledState::Disabled))
-			{
-				return true;
-			}
-		}
-	}
-
-	return false;
-}
-
 FFastGeoTransformResult UFastGeoWorldPartitionRuntimeCellTransformer::IsAllowedActorClass(AActor* InActor) const
 {
 	const UFastGeoTransformerSettings* EffectiveSettings = GetSettings();
diff --git a/Engine/Plugins/Experimental/FastGeoStreaming/Source/FastGeoStreaming/Public/FastGeoWorldPartitionRuntimeCellTransformer.h b/Engine/Plugins/Experimental/FastGeoStreaming/Source/FastGeoStreaming/Public/FastGeoWorldPartitionRuntimeCellTransformer.h
--- a/Engine/Plugins/Experimental/FastGeoStreaming/Source/FastGeoStreaming/Public/FastGeoWorldPartitionRuntimeCellTransformer.h
+++ b/Engine/Plugins/Experimental/FastGeoStreaming/Source/FastGeoStreaming/Public/FastGeoWorldPartitionRuntimeCellTransformer.h
@@ -136,7 +136,6 @@ protected:
 	/** Whether component mobility is supported by the transformer. */
 	bool IsComponentMobilitySupported(USceneComponent* InComponent, FString& OutReason) const;
 
-	bool IsBlueprintActorWithLogic(AActor* InActor) const;
 	FFastGeoTransformResult CanTransformActor(AActor* InActor, bool& bOutIsActorFullyTransformable, TArray<UActorComponent*>& OutTransformableComponents) const;
 	FFastGeoTransformResult IsAllowedActorClass(AActor* InActor) const;
 	FFastGeoTransformResult CanTransformComponent(UActorComponent* InComponent) const;
diff --git a/Engine/Source/Runtime/Engine/Private/WorldPartition/WorldPartitionRuntimeCellTransformer.cpp b/Engine/Source/Runtime/Engine/Private/WorldPartition/WorldPartitionRuntimeCellTransformer.cpp
--- a/Engine/Source/Runtime/Engine/Private/WorldPartition/WorldPartitionRuntimeCellTransformer.cpp
+++ b/Engine/Source/Runtime/Engine/Private/WorldPartition/WorldPartitionRuntimeCellTransformer.cpp
@@ -2,6 +2,13 @@
 
 #include "WorldPartition/WorldPartitionRuntimeCellTransformer.h"
 
+#if WITH_EDITOR
+#include "EdGraph/EdGraph.h"
+#include "Engine/Blueprint.h"
+#include "Engine/BlueprintGeneratedClass.h"
+#include "GameFramework/Actor.h"
+#endif
+
 #include UE_INLINE_GENERATED_CPP_BY_NAME(WorldPartitionRuntimeCellTransformer)
 
 UWorldPartitionRuntimeCellTransformer::UWorldPartitionRuntimeCellTransformer(const FObjectInitializer& ObjectInitializer)
@@ -102,6 +109,69 @@ bool UWorldPartitionRuntimeCellTransformer::CanIgnoreComponent(const UActorCompo
 
 	return bCanIgnore;
 }
+
+bool UWorldPartitionRuntimeCellTransformer::IsBlueprintActorWithLogic(AActor* InActor) const
+{
+	const FName FN_UserConstructionScript(TEXT("UserConstructionScript"));
+
+	UBlueprint* Blueprint = UBlueprint::GetBlueprintFromClass(InActor->GetClass());
+	if (!Blueprint || !Blueprint->GeneratedClass)
+	{
+		return false;
+	}
+
+	check(Blueprint->ParentClass && Blueprint->ParentClass->IsChildOf(AActor::StaticClass()));
+
+	UBlueprintGeneratedClass* BPClass = Cast<UBlueprintGeneratedClass>(Blueprint->GeneratedClass);
+	if (!BPClass)
+	{
+		return false;
+	}
+
+	if (Blueprint->DelegateSignatureGraphs.Num() > 0)
+	{
+		return true;
+	}
+
+	if (Blueprint->ImplementedInterfaces.Num() > 0)
+	{
+		return true;
+	}
+
+	// Check if no extra functions, other than the user construction script (only AActor and subclasses of AActor have)
+	if (Blueprint->FunctionGraphs.Num() > 1)
+	{
+		return true;
+	}
+
+	check(Blueprint->FunctionGraphs.Num() == 0 || Blueprint->FunctionGraphs[0]->GetFName() == FN_UserConstructionScript);
+
+	// Check if the generated class has overridden any functions dynamically
+	for (TFieldIterator<UFunction> It(BPClass, EFieldIteratorFlags::IncludeSuper); It; ++It)
+	{
+		UFunction* Function = *It;
+
+		// Ignore functions from native C++ classes (i.e., inherited but not overridden in BP)
+		if (Function->GetOwnerClass() == BPClass && Function->GetFName() != FN_UserConstructionScript)
+		{
+			return true; // Found an overridden function
+		}
+	}
+
+	// If there is an enabled node in the event graph, the Blueprint is not data only
+	for (UEdGraph* EventGraph : Blueprint->UbergraphPages)
+	{
+		for (UEdGraphNode* GraphNode : EventGraph->Nodes)
+		{
+			if (GraphNode && (GraphNode->GetDesiredEnabledState() != ENodeEnabledState::Disabled))
+			{
+				return true;
+			}
+		}
+	}
+
+	return false;
+}
 #endif
 
-const FName UWorldPartitionRuntimeCellTransformer::NAME_CellTransformerIgnoreActor(TEXT("CellTransformer_IgnoreActor"));
+const FName UWorldPartitionRuntimeCellTransformer::NAME_CellTransformerIgnoreActor(TEXT("CellTransformer_IgnoreActor"));
\ No newline at end of file
diff --git a/Engine/Source/Runtime/Engine/Private/WorldPartition/WorldPartitionRuntimeCellTransformerISM.cpp b/Engine/Source/Runtime/Engine/Private/WorldPartition/WorldPartitionRuntimeCellTransformerISM.cpp
--- a/Engine/Source/Runtime/Engine/Private/WorldPartition/WorldPartitionRuntimeCellTransformerISM.cpp
+++ b/Engine/Source/Runtime/Engine/Private/WorldPartition/WorldPartitionRuntimeCellTransformerISM.cpp
@@ -7,6 +7,7 @@
 #include "GameFramework/Actor.h"
 #include "ISMPartition/ISMComponentBatcher.h"
 #include "ISMPartition/ISMComponentDescriptor.h"
+#include "Engine/Blueprint.h"
 #include "Engine/Level.h"
 #include "UObject/Package.h"
 #endif
@@ -66,49 +67,96 @@ void UWorldPartitionRuntimeCellTransformerISM::Transform(ULevel* InLevel)
 	{
 		if (IsValid(Actor) && CanAutoInstanceActor(Actor))
 		{
+			bool bContainsNonTransformable = false;
+			TArray<UStaticMeshComponent*> ActorTransformableComponents;
 			// Gather potential components that can be merged
-			Actor->ForEachComponent<UStaticMeshComponent>(true, [&ISMComponentBatchers, &Actor](UStaticMeshComponent* StaticMeshComponent)
+			Actor->ForEachComponent<UActorComponent>(true, [this, &ISMComponentBatchers, &ActorTransformableComponents, &bContainsNonTransformable, &Actor](UActorComponent* Component)
 			{
-				if (!StaticMeshComponent->IsEditorOnly() && StaticMeshComponent->IsVisible() && (StaticMeshComponent->Mobility == EComponentMobility::Static))
+				UStaticMeshComponent* StaticMeshComponent = Cast<UStaticMeshComponent>(Component);
+				if (StaticMeshComponent && !StaticMeshComponent->IsEditorOnly() && StaticMeshComponent->IsVisible() && (StaticMeshComponent->Mobility == EComponentMobility::Static))
 				{
-					FISMComponentDescriptor ISMComponentDescriptor;
-					ISMComponentDescriptor.InitFrom(StaticMeshComponent);
-					ISMComponentBatchers.FindOrAdd(ISMComponentDescriptor).ActorComponents.FindOrAdd(&Actor).Add(StaticMeshComponent);
+					ActorTransformableComponents.Add(StaticMeshComponent);
+				}
+				else if (!CanIgnoreComponent(Component))
+				{
+					bContainsNonTransformable = true;
 				}
 			});
+
+			// Can't convert partially a Blueprint actor: Rerun CS is called in PIE when registering the component
+			// and is also called when registering components during cook/save of the level.
+			if (UBlueprint::GetBlueprintFromClass(Actor->GetClass()) && ActorTransformableComponents.Num() && bContainsNonTransformable)
+			{
+				continue;
+			}
+
+			for (UStaticMeshComponent* StaticMeshComponent : ActorTransformableComponents)
+			{
+				FISMComponentDescriptor ISMComponentDescriptor;
+				ISMComponentDescriptor.InitFrom(StaticMeshComponent);
+				ISMComponentBatchers.FindOrAdd(ISMComponentDescriptor).ActorComponents.FindOrAdd(&Actor).Add(StaticMeshComponent);
+			}
 		}
 	}
 
-	int32 NumInstancedComponents = 0;
-	for (auto& [ISMComponentDescriptor, ActorComponentBatcherDescriptor] : ISMComponentBatchers)
+	// Keep a set of non-transformable actors
+	TSet<AActor*> NonTransformableActors;
+	for (auto It = ISMComponentBatchers.CreateIterator(); It; ++It)
 	{
-		if ((uint32)ActorComponentBatcherDescriptor.ActorComponents.Num() >= MinNumInstances)
+		const FActorComponentBatcherDescriptor& ActorComponentBatcherDescriptor = It->Value;
+
+		if ((uint32)ActorComponentBatcherDescriptor.ActorComponents.Num() < MinNumInstances)
 		{
 			for (auto& [Actor, Components] : ActorComponentBatcherDescriptor.ActorComponents)
 			{
-				for (UStaticMeshComponent* StaticMeshComponent : Components)
+				// For Blueprint actors, if any of their contributions fail MinNumInstances,
+				// we treat the entire actor as non-transformable to avoid partial conversion.
+				if (UBlueprint::GetBlueprintFromClass((*Actor)->GetClass()))
 				{
-					// Register the component into the batcher
-					StaticMeshComponent->UpdateComponentToWorld();
-					ActorComponentBatcherDescriptor.ISMComponentBatcher.Add(StaticMeshComponent);
+					NonTransformableActors.Add(*Actor);
+				}
+			}
 
-					// Remove the component from the actor
-					(*Actor)->RemoveOwnedComponent(StaticMeshComponent);
-					StaticMeshComponent->MarkAsGarbage();
+			// Remove elements that don't respect the minimum number of instances
+			It.RemoveCurrent();
+		}
+	}
 
-					NumInstancedComponents++;
-				}
+	int32 NumInstancedComponents = 0;
+	for (auto& [ISMComponentDescriptor, ActorComponentBatcherDescriptor] : ISMComponentBatchers)
+	{
+		check((uint32)ActorComponentBatcherDescriptor.ActorComponents.Num() >= MinNumInstances);
+		
+		for (auto& [Actor, Components] : ActorComponentBatcherDescriptor.ActorComponents)
+		{
+			// Skip actors that were marked as non-transformable
+			if (NonTransformableActors.Contains(*Actor))
+			{
+				continue;
+			}
 
-				if (CanRemoveActor(*Actor))
-				{
-					*Actor = nullptr;
-				}
-				else if (USceneComponent* OldRootComponent = (*Actor)->GetRootComponent(); OldRootComponent && !IsValid(OldRootComponent))
-				{
-					USceneComponent* NewRootComponent = NewObject<USceneComponent>(*Actor);
-					NewRootComponent->SetRelativeTransform(OldRootComponent->GetRelativeTransform());
-					(*Actor)->SetRootComponent(NewRootComponent);
-				}
+			for (UStaticMeshComponent* StaticMeshComponent : Components)
+			{
+				// Register the component into the batcher
+				StaticMeshComponent->UpdateComponentToWorld();
+				ActorComponentBatcherDescriptor.ISMComponentBatcher.Add(StaticMeshComponent);
+
+				// Remove the component from the actor
+				(*Actor)->RemoveOwnedComponent(StaticMeshComponent);
+				StaticMeshComponent->MarkAsGarbage();
+
+				NumInstancedComponents++;
+			}
+
+			if (CanRemoveActor(*Actor))
+			{
+				*Actor = nullptr;
+			}
+			else if (USceneComponent* OldRootComponent = (*Actor)->GetRootComponent(); OldRootComponent && !IsValid(OldRootComponent))
+			{
+				USceneComponent* NewRootComponent = NewObject<USceneComponent>(*Actor);
+				NewRootComponent->SetRelativeTransform(OldRootComponent->GetRelativeTransform());
+				(*Actor)->SetRootComponent(NewRootComponent);
 			}
 		}
 	}
@@ -175,6 +223,11 @@ bool UWorldPartitionRuntimeCellTransformerISM::CanAutoInstanceActor(AActor* InAc
 		return false;
 	}
 
+	if (IsBlueprintActorWithLogic(InActor))
+	{
+		return false;
+	}
+
 	UClass* ActorClass = InActor->GetClass();
 	for (TSubclassOf<AActor> DisallowedClass : DisallowedClasses)
 	{
@@ -195,16 +248,21 @@ bool UWorldPartitionRuntimeCellTransformerISM::CanAutoInstanceActor(AActor* InAc
 	return false;
 }
 
+bool UWorldPartitionRuntimeCellTransformerISM::CanIgnoreComponent(const UActorComponent* InComponent) const
+{
+	return InComponent->IsEditorOnly() || Super::CanIgnoreComponent(InComponent);
+}
+
 bool UWorldPartitionRuntimeCellTransformerISM::CanRemoveActor(AActor* InActor) const
 {
 	for (UActorComponent* Component : InActor->GetComponents())
 	{
-		if (!Component->IsEditorOnly() && !CanIgnoreComponent(Component))
-		{
+		if (!CanIgnoreComponent(Component))
+		{		
 			return false;
 		}
 	}
-		
+
 	return true;
 }
 #endif
diff --git a/Engine/Source/Runtime/Engine/Public/WorldPartition/WorldPartitionRuntimeCellTransformer.h b/Engine/Source/Runtime/Engine/Public/WorldPartition/WorldPartitionRuntimeCellTransformer.h
--- a/Engine/Source/Runtime/Engine/Public/WorldPartition/WorldPartitionRuntimeCellTransformer.h
+++ b/Engine/Source/Runtime/Engine/Public/WorldPartition/WorldPartitionRuntimeCellTransformer.h
@@ -58,6 +58,7 @@ public:
 protected:
 #if WITH_EDITOR
 	UE_API virtual bool CanIgnoreComponent(const UActorComponent* InComponent) const;
+	UE_API virtual bool IsBlueprintActorWithLogic(AActor* InActor) const;
 #endif
 	
 protected:
diff --git a/Engine/Source/Runtime/Engine/Public/WorldPartition/WorldPartitionRuntimeCellTransformerISM.h b/Engine/Source/Runtime/Engine/Public/WorldPartition/WorldPartitionRuntimeCellTransformerISM.h
--- a/Engine/Source/Runtime/Engine/Public/WorldPartition/WorldPartitionRuntimeCellTransformerISM.h
+++ b/Engine/Source/Runtime/Engine/Public/WorldPartition/WorldPartitionRuntimeCellTransformerISM.h
@@ -16,6 +16,8 @@ class UWorldPartitionRuntimeCellTransformerISM : public UWorldPartitionRuntimeCe
 public:
 	virtual void Transform(ULevel* InLevel) override;
 
+protected:
+	virtual bool CanIgnoreComponent(const UActorComponent* InComponent) const override;
 private:
 	bool CanAutoInstanceActor(AActor* InActor) const;
 	bool CanRemoveActor(AActor* InActor) const;
