Partial revert of 47037409 for level-specific replication code.

#Synced-CL 48615930

diff --git a/Engine/Source/Runtime/Engine/Classes/GameFramework/UpdateLevelVisibilityLevelInfo.h b/Engine/Source/Runtime/Engine/Classes/GameFramework/UpdateLevelVisibilityLevelInfo.h
--- a/Engine/Source/Runtime/Engine/Classes/GameFramework/UpdateLevelVisibilityLevelInfo.h
+++ b/Engine/Source/Runtime/Engine/Classes/GameFramework/UpdateLevelVisibilityLevelInfo.h
@@ -63,11 +63,11 @@ struct FUpdateLevelVisibilityLevelInfo
 	}
 
 	/**
-	 * @param Container			Container to pull PackageName and FileName from.
+	 * @param Level				Level to pull PackageName and FileName from.
 	 * @param bInIsVisible		Default value for bIsVisible.
 	 * @param bInTryMakeVisible	Whether the level is trying to be made visible or not.
 	 */
-	ENGINE_API FUpdateLevelVisibilityLevelInfo(const class UObject* const Container, const bool bInIsVisible, const bool bInTryMakeVisible = false);
+	ENGINE_API FUpdateLevelVisibilityLevelInfo(const class ULevel* const Level, const bool bInIsVisible, const bool bInTryMakeVisible = false);
 
 	/** The name of the package for the level whose status changed. */
 	UPROPERTY()
diff --git a/Engine/Source/Runtime/Engine/Private/GameFramework/UpdateLevelVisibilityLevelInfo.cpp b/Engine/Source/Runtime/Engine/Private/GameFramework/UpdateLevelVisibilityLevelInfo.cpp
--- a/Engine/Source/Runtime/Engine/Private/GameFramework/UpdateLevelVisibilityLevelInfo.cpp
+++ b/Engine/Source/Runtime/Engine/Private/GameFramework/UpdateLevelVisibilityLevelInfo.cpp
@@ -17,7 +17,7 @@ namespace PlayerControllerCVars
 	);
 }
 
-FUpdateLevelVisibilityLevelInfo::FUpdateLevelVisibilityLevelInfo(const UObject* const Container, const bool bInIsVisible, const bool bInTryMakeVisible)
+FUpdateLevelVisibilityLevelInfo::FUpdateLevelVisibilityLevelInfo(const ULevel* const Level, const bool bInIsVisible, const bool bInTryMakeVisible)
 	: CookPackageName(NAME_None)
 	, bIsVisible(bInIsVisible)
 	, bTryMakeVisible(bInTryMakeVisible)
@@ -26,18 +26,15 @@ FUpdateLevelVisibilityLevelInfo::FUpdateLevelVisibilityLevelInfo(const UObject*
 	// For backward compatibility, bTryMakeVisible was added instead of converting bIsVisible to an enum.
 	// Make sure we don't receive the invalid state (bIsVisible == true) && (bTryMakeVisible == true)
 	check(!bTryMakeVisible || (bIsVisible != bTryMakeVisible));
-	const UPackage* const ContainerPackage = Container->GetOutermost();
-	PackageName = ContainerPackage->GetFName();
+	const UPackage* const LevelPackage = Level->GetOutermost();
+	PackageName = LevelPackage->GetFName();
 #if WITH_EDITOR
-	if (const ULevel* Level = Cast<ULevel>(Container))
-	{
-		CookPackageName = Level->GetCookPackageName();
-	}
+	CookPackageName = Level->GetCookPackageName();
 #endif
 
 	// When packages are duplicated for PIE, they may not have a FileName.
 	// For now, just revert to the old behavior.
-	FName LoadedPathName = ContainerPackage->GetLoadedPath().GetPackageFName();
+	FName LoadedPathName = LevelPackage->GetLoadedPath().GetPackageFName();
 	FileName = (LoadedPathName.IsNone()) ? PackageName : LoadedPathName;
 }
 
