ISKM improvements and fixes.
- Added methods to modify the animation index and transform.
- Added support for previous transforms.
- Modified the instance data manager to use delta updates for skinning data.
- Fixed bAlwaysHasVelocity getting reset to 0 after updating instances due to skinning data changing without the 'dynamic data' flag.

#Synced-CL 48814992

diff --git a/Engine/Plugins/Experimental/FastGeoStreaming/Source/FastGeoStreaming/Internal/FastGeoInstancedSkinnedMeshComponent.h b/Engine/Plugins/Experimental/FastGeoStreaming/Source/FastGeoStreaming/Internal/FastGeoInstancedSkinnedMeshComponent.h
--- a/Engine/Plugins/Experimental/FastGeoStreaming/Source/FastGeoStreaming/Internal/FastGeoInstancedSkinnedMeshComponent.h
+++ b/Engine/Plugins/Experimental/FastGeoStreaming/Source/FastGeoStreaming/Internal/FastGeoInstancedSkinnedMeshComponent.h
@@ -55,6 +55,8 @@ protected:
 	}
 	//~ End FFastGeoSkinnedMeshComponentBase interface
 
+	TConstArrayView<FTransform3f> GetInstancePrevTransforms() const { return {}; }
+
 private:
 	TArray<FSkinnedMeshInstanceData> InstanceData;
 	int32 NumCustomDataFloats = 0;
diff --git a/Engine/Source/Runtime/Engine/Classes/Components/InstancedSkinnedMeshComponent.h b/Engine/Source/Runtime/Engine/Classes/Components/InstancedSkinnedMeshComponent.h
--- a/Engine/Source/Runtime/Engine/Classes/Components/InstancedSkinnedMeshComponent.h
+++ b/Engine/Source/Runtime/Engine/Classes/Components/InstancedSkinnedMeshComponent.h
@@ -74,6 +74,10 @@ protected:
 	UPROPERTY(EditAnywhere, EditFixedSize, SkipSerialization, DisplayName="Custom Data", Category=Instances, AdvancedDisplay, meta=(EditFixedOrder))
 	TArray<float> InstanceCustomData;
 
+	/** Array of previous transforms. Must match the length of InstanceData. */
+	UPROPERTY(Transient)
+	TArray<FTransform3f> PerInstancePrevTransform;
+
 	/** Screen space footprint (in primary view) cutoff which dictates the far distance the instance will play back animation. 
 	 *  Using 0.0 (the default) falls back to a global threshold, a negative value disables the cutoff.
 	 */
@@ -102,6 +106,9 @@ protected:
 	UPROPERTY()
 	bool bIsInstanceDataGPUOnly = false;
 
+	UPROPERTY(EditAnywhere, BlueprintReadOnly, Category = Instances)
+	bool bHasPreviousTransforms = false;
+
 	/** If true, this component will use GPU LOD selection. */
 	UPROPERTY(EditAnywhere, BlueprintReadOnly, Category = Culling)
 	uint8 bUseGpuLodSelection : 1;
@@ -179,8 +186,20 @@ public:
 	UFUNCTION(BlueprintCallable, Category = "Components|InstancedSkinnedMesh")
 	ENGINE_API bool SetCustomDataValue(FPrimitiveInstanceId InstanceId, int32 CustomDataIndex, float CustomDataValue);
 
+	UFUNCTION(BlueprintCallable, Category = "Components|InstancedSkinnedMesh")
+	inline bool SetCustomData(FPrimitiveInstanceId InstanceId, const TArray<float>& CustomDataFloats)
+	{
+		return SetCustomData(InstanceId, TConstArrayView<float>(CustomDataFloats));
+	}
+
 	/** Update all custom data values for specific instance, the size of the array view must match the NumCustomDataFloats. Returns True on success. */
-	ENGINE_API bool SetCustomData(FPrimitiveInstanceId InstanceId, TArrayView<const float> CustomDataFloats);
+	ENGINE_API bool SetCustomData(FPrimitiveInstanceId InstanceId, TConstArrayView<float> CustomDataFloats);
+
+	UFUNCTION(BlueprintCallable, Category = "Components|InstancedSkinnedMesh")
+	inline bool GetCustomData(FPrimitiveInstanceId InstanceId, TArray<float>& CustomDataFloats) const
+	{
+		return GetCustomData(InstanceId, TArrayView<float>(CustomDataFloats));
+	}
 
 	/** Get all custom data values for specific instance, the size of the array view must match the NumCustomDataFloats. Returns True on success. */
 	ENGINE_API bool GetCustomData(FPrimitiveInstanceId InstanceId, TArrayView<float> CustomDataFloats) const;
@@ -192,10 +211,35 @@ public:
 	/** Get the transform for the instance specified. Instance is returned in local space of this component unless bWorldSpace is set. Returns True on success. */
 	UFUNCTION(BlueprintCallable, Category = "Components|InstancedSkinnedMesh")
 	ENGINE_API bool GetInstanceTransform(FPrimitiveInstanceId InstanceId, FTransform& OutInstanceTransform, bool bWorldSpace = false) const;
+	
+	/** Get the previous transform for the instance specified. Instance is returned in local space of this component unless bWorldSpace is set. Returns True on success. */
+	UFUNCTION(BlueprintCallable, Category = "Components|InstancedSkinnedMesh")
+	ENGINE_API bool GetInstancePrevTransform(FPrimitiveInstanceId InstanceId, FTransform& OutInstanceTransform, bool bWorldSpace = false) const;
+
+	/** Update the transform for the instance specified. */
+	UFUNCTION(BlueprintCallable, Category = "Components|InstancedStaticMesh")
+	ENGINE_API bool SetInstanceTransform(FPrimitiveInstanceId InstanceId, const FTransform& InstanceTransform, bool bWorldSpace=false);
+
+	/** Update the previous transform for the instance specified. */
+	UFUNCTION(BlueprintCallable, Category = "Components|InstancedStaticMesh")
+	ENGINE_API bool SetInstancePrevTransform(FPrimitiveInstanceId InstanceId, const FTransform& InstanceTransform, bool bWorldSpace=false);
+	
+	/** Update both the current and previous transforms for the instance specified. */
+	UFUNCTION(BlueprintCallable, Category = "Components|InstancedStaticMesh")
+	ENGINE_API bool SetInstanceTransforms(FPrimitiveInstanceId InstanceId, const FTransform& InstanceTransform, const FTransform& InstancePrevTransform, bool bWorldSpace=false);
 
 	/** Get the animation index for the instance specified. Returns True on success. */
+	UFUNCTION(BlueprintCallable, Category = "Components|InstancedSkinnedMesh")
 	ENGINE_API bool GetInstanceAnimationIndex(FPrimitiveInstanceId InstanceId, int32& OutAnimationIndex) const;
 
+	/** Set the animation index for the instance specified. Returns True on success. */
+	UFUNCTION(BlueprintCallable, Category = "Components|InstancedSkinnedMesh")
+	ENGINE_API bool SetInstanceAnimationIndex(FPrimitiveInstanceId InstanceId, int32 OutAnimationIndex);
+
+	/** Updates an instance on this component. Transform is given in local space of this component unless bWorldSpace is set. */
+	UFUNCTION(BlueprintCallable, Category = "Components|InstancedSkinnedMesh")
+	ENGINE_API bool UpdateInstance(FPrimitiveInstanceId InstanceId, const FTransform& InstanceTransform, int32 AnimationIndex, bool bWorldSpace = false);
+
 	/** Remove the instance specified. Returns True on success. */
 	UFUNCTION(BlueprintCallable, Category = "Components|InstancedSkinnedMesh")
 	ENGINE_API bool RemoveInstance(FPrimitiveInstanceId InstanceId);
@@ -223,12 +267,19 @@ public:
 	ENGINE_API void OptimizeInstanceData(bool bShouldRetainIdMap = false);
 
 	ENGINE_API bool IsEnabled() const;
+
+	UFUNCTION(BlueprintCallable, Category = "Components|InstancedSkinnedMesh")
 	ENGINE_API int32 GetInstanceCount() const;
 
 	int32 GetNumCustomDataFloats() const { return NumCustomDataFloats; }
 	const TArray<float>& GetInstanceCustomData() const { return InstanceCustomData; }
 	const TArray<FSkinnedMeshInstanceData>& GetInstanceData() const { return InstanceData; }
 
+	TConstArrayView<FTransform3f> GetInstancePrevTransforms() const
+	{
+		return PerInstancePrevTransform;
+	}
+
 	int32 GetInstanceCountGPUOnly() const { return NumInstancesGPUOnly; }
 
 	int32 GetMinDrawDistance() const { return InstanceMinDrawDistance; }
diff --git a/Engine/Source/Runtime/Engine/Private/Components/InstancedSkinnedMeshComponent.cpp b/Engine/Source/Runtime/Engine/Private/Components/InstancedSkinnedMeshComponent.cpp
--- a/Engine/Source/Runtime/Engine/Private/Components/InstancedSkinnedMeshComponent.cpp
+++ b/Engine/Source/Runtime/Engine/Private/Components/InstancedSkinnedMeshComponent.cpp
@@ -296,6 +296,10 @@ void UInstancedSkinnedMeshComponent::OptimizeInstanceData(bool bShouldRetainIdMa
 		// Reorder instances according to the remap
 		ReorderArray(InstanceData, IndexRemap);
 		ReorderArray(InstanceCustomData, IndexRemap, NumCustomDataFloats);
+		if (bHasPreviousTransforms && PerInstancePrevTransform.Num() == InstanceData.Num())
+		{
+			ReorderArray(PerInstancePrevTransform, IndexRemap);
+		}
 #if WITH_EDITOR
 		ReorderArray(SelectedInstances, IndexRemap);
 #endif
@@ -772,13 +776,12 @@ TArray<FPrimitiveInstanceId> UInstancedSkinnedMeshComponent::AddInstances(const
 
 bool UInstancedSkinnedMeshComponent::SetCustomDataValue(FPrimitiveInstanceId InstanceId, int32 CustomDataIndex, float CustomDataValue)
 {
-	int32 InstanceIndex = InstanceDataManager.IdToIndex(InstanceId);
-
-	if (!InstanceData.IsValidIndex(InstanceIndex) || CustomDataIndex < 0 || CustomDataIndex >= NumCustomDataFloats)
+	if (!InstanceDataManager.IsValidId(InstanceId) || CustomDataIndex < 0 || CustomDataIndex >= NumCustomDataFloats)
 	{
 		return false;
 	}
 
+	int32 InstanceIndex = InstanceDataManager.IdToIndex(InstanceId);
 	Modify();
 
 	InstanceDataManager.CustomDataChanged(InstanceIndex);
@@ -789,13 +792,12 @@ bool UInstancedSkinnedMeshComponent::SetCustomDataValue(FPrimitiveInstanceId Ins
 
 bool UInstancedSkinnedMeshComponent::SetCustomData(FPrimitiveInstanceId InstanceId, TArrayView<const float> CustomDataFloats)
 {
-	int32 InstanceIndex = InstanceDataManager.IdToIndex(InstanceId);
-
-	if (!InstanceData.IsValidIndex(InstanceIndex) || CustomDataFloats.Num() == 0)
+	if (!InstanceDataManager.IsValidId(InstanceId) || CustomDataFloats.Num() == 0)
 	{
 		return false;
 	}
 
+	int32 InstanceIndex = InstanceDataManager.IdToIndex(InstanceId);
 	Modify();
 
 	const int32 NumToCopy = FMath::Min(CustomDataFloats.Num(), NumCustomDataFloats);
@@ -823,12 +825,12 @@ void UInstancedSkinnedMeshComponent::SetNumCustomDataFloats(int32 InNumCustomDat
 
 bool UInstancedSkinnedMeshComponent::GetCustomData(FPrimitiveInstanceId InstanceId, TArrayView<float> CustomDataFloats) const
 {
-	int32 InstanceIndex = InstanceDataManager.IdToIndex(InstanceId);
-	if (!InstanceData.IsValidIndex(InstanceIndex))
+	if (!InstanceDataManager.IsValidId(InstanceId))
 	{
 		return false;
 	}
 
+	int32 InstanceIndex = InstanceDataManager.IdToIndex(InstanceId);
 	const int32 NumToCopy = FMath::Min(CustomDataFloats.Num(), NumCustomDataFloats);
 	FMemory::Memcpy(CustomDataFloats.GetData(), &InstanceCustomData[InstanceIndex * NumCustomDataFloats], NumToCopy * CustomDataFloats.GetTypeSize());
 	return true;
@@ -836,12 +838,12 @@ bool UInstancedSkinnedMeshComponent::GetCustomData(FPrimitiveInstanceId Instance
 
 bool UInstancedSkinnedMeshComponent::GetInstanceTransform(FPrimitiveInstanceId InstanceId, FTransform& OutInstanceTransform, bool bWorldSpace) const
 {
-	int32 InstanceIndex = InstanceDataManager.IdToIndex(InstanceId);
-	if (!InstanceData.IsValidIndex(InstanceIndex))
+	if (!InstanceDataManager.IsValidId(InstanceId))
 	{
 		return false;
 	}
 
+	int32 InstanceIndex = InstanceDataManager.IdToIndex(InstanceId);
 	const FSkinnedMeshInstanceData& Instance = InstanceData[InstanceIndex];
 
 	OutInstanceTransform = FTransform(Instance.Transform);
@@ -853,23 +855,116 @@ bool UInstancedSkinnedMeshComponent::GetInstanceTransform(FPrimitiveInstanceId I
 	return true;
 }
 
-bool UInstancedSkinnedMeshComponent::GetInstanceAnimationIndex(FPrimitiveInstanceId InstanceId, int32& OutAnimationIndex) const
+bool UInstancedSkinnedMeshComponent::GetInstancePrevTransform(FPrimitiveInstanceId InstanceId, FTransform& OutInstanceTransform, bool bWorldSpace) const
 {
+	if (!InstanceDataManager.IsValidId(InstanceId) || !bHasPreviousTransforms || PerInstancePrevTransform.Num() != InstanceData.Num())
+	{
+		return false;
+	}
+
 	int32 InstanceIndex = InstanceDataManager.IdToIndex(InstanceId);
-	if (!InstanceData.IsValidIndex(InstanceIndex))
+
+	if (bWorldSpace)
+	{
+		OutInstanceTransform = FTransform(PerInstancePrevTransform[InstanceIndex]) * GetComponentTransform();
+	}
+	else
+	{
+		OutInstanceTransform = FTransform(PerInstancePrevTransform[InstanceIndex]);
+	}
+
+	return true;
+}
+
+bool UInstancedSkinnedMeshComponent::SetInstanceTransform(FPrimitiveInstanceId InstanceId, const FTransform& InstanceTransform, bool bWorldSpace)
+{
+	if (!InstanceDataManager.IsValidId(InstanceId))
+	{
+		return false;
+	}
+
+	int32 InstanceIndex = InstanceDataManager.IdToIndex(InstanceId);
+	InstanceData[InstanceIndex].Transform = FTransform3f(bWorldSpace ? InstanceTransform.GetRelativeTransform(GetComponentTransform()) : InstanceTransform);
+	InstanceDataManager.TransformChanged(InstanceIndex);
+	return true;
+}
+
+bool UInstancedSkinnedMeshComponent::SetInstancePrevTransform(FPrimitiveInstanceId InstanceId, const FTransform& InstanceTransform, bool bWorldSpace)
+{
+	if (!InstanceDataManager.IsValidId(InstanceId) || !bHasPreviousTransforms || PerInstancePrevTransform.Num() != InstanceData.Num())
 	{
 		return false;
 	}
 
+	int32 InstanceIndex = InstanceDataManager.IdToIndex(InstanceId);
+	PerInstancePrevTransform[InstanceIndex] = FTransform3f(bWorldSpace ? InstanceTransform.GetRelativeTransform(GetComponentTransform()) : InstanceTransform);
+	InstanceDataManager.TransformChanged(InstanceIndex);
+	return true;
+}
+
+bool UInstancedSkinnedMeshComponent::SetInstanceTransforms(FPrimitiveInstanceId InstanceId, const FTransform& InstanceTransform, const FTransform& InstancePrevTransform, bool bWorldSpace)
+{
+	if (!InstanceDataManager.IsValidId(InstanceId))
+	{
+		return false;
+	}
+
+	int32 InstanceIndex = InstanceDataManager.IdToIndex(InstanceId);
+	InstanceData[InstanceIndex].Transform = FTransform3f(bWorldSpace ? InstanceTransform.GetRelativeTransform(GetComponentTransform()) : InstanceTransform);
+	if (bHasPreviousTransforms && PerInstancePrevTransform.Num() == InstanceData.Num())
+	{
+		PerInstancePrevTransform[InstanceIndex] = FTransform3f(bWorldSpace ? InstancePrevTransform.GetRelativeTransform(GetComponentTransform()) : InstancePrevTransform);
+	}
+	InstanceDataManager.TransformChanged(InstanceIndex);
+	return true;
+}
+
+bool UInstancedSkinnedMeshComponent::GetInstanceAnimationIndex(FPrimitiveInstanceId InstanceId, int32& OutAnimationIndex) const
+{
+	if (!InstanceDataManager.IsValidId(InstanceId))
+	{
+		return false;
+	}
+
+	int32 InstanceIndex = InstanceDataManager.IdToIndex(InstanceId);
 	OutAnimationIndex = InstanceData[InstanceIndex].AnimationIndex;
 	return true;
 }
 
-bool UInstancedSkinnedMeshComponent::RemoveInstance(FPrimitiveInstanceId InstanceId)
+bool UInstancedSkinnedMeshComponent::SetInstanceAnimationIndex(FPrimitiveInstanceId InstanceId, int32 AnimationIndex)
+{
+	if (!InstanceDataManager.IsValidId(InstanceId))
+	{
+		return false;
+	}
+
+	int32 InstanceIndex = InstanceDataManager.IdToIndex(InstanceId);
+	InstanceData[InstanceIndex].AnimationIndex = AnimationIndex;
+	InstanceDataManager.SkinningDataChanged(InstanceIndex);
+	return true;
+}
+
+bool UInstancedSkinnedMeshComponent::UpdateInstance(FPrimitiveInstanceId InstanceId, const FTransform& InstanceTransform, int32 AnimationIndex, bool bWorldSpace)
 {
+	if (!InstanceDataManager.IsValidId(InstanceId))
+	{
+		return false;
+	}
+
 	int32 InstanceIndex = InstanceDataManager.IdToIndex(InstanceId);
-	if (InstanceData.IsValidIndex(InstanceIndex))
+	FSkinnedMeshInstanceData& InstanceDataItem = InstanceData[InstanceIndex];
+	InstanceDataItem.AnimationIndex = AnimationIndex;
+	InstanceDataItem.Transform = FTransform3f(bWorldSpace ? InstanceTransform.GetRelativeTransform(GetComponentTransform()) : InstanceTransform);
+	InstanceDataManager.SkinningDataChanged(InstanceIndex);
+	InstanceDataManager.TransformChanged(InstanceIndex);
+	return true;
+}
+
+bool UInstancedSkinnedMeshComponent::RemoveInstance(FPrimitiveInstanceId InstanceId)
+{
+	if (InstanceDataManager.IsValidId(InstanceId))
 	{
+		int32 InstanceIndex = InstanceDataManager.IdToIndex(InstanceId);
 		Modify();
 		return RemoveInstanceInternal(InstanceIndex, false);
 	}
@@ -894,6 +989,7 @@ void UInstancedSkinnedMeshComponent::ClearInstances()
 	// Clear all the per-instance data
 	InstanceData.Empty();
 	InstanceCustomData.Empty();
+	PerInstancePrevTransform.Empty();
 
 #if WITH_EDITOR
 	SelectedInstances.Empty();
@@ -986,6 +1082,11 @@ FPrimitiveInstanceId UInstancedSkinnedMeshComponent::AddInstanceInternal(int32 I
 	// Add custom data to instance
 	InstanceCustomData.AddZeroed(NumCustomDataFloats);
 
+	if (bHasPreviousTransforms)
+	{
+		PerInstancePrevTransform.Add(LocalTransform);
+	}
+
 #if WITH_EDITOR
 	SelectedInstances.Add(false);
 #endif
@@ -1012,6 +1113,11 @@ bool UInstancedSkinnedMeshComponent::RemoveInstanceInternal(int32 InstanceIndex,
 		InstanceCustomData.RemoveAtSwap(InstanceIndex * NumCustomDataFloats, NumCustomDataFloats, EAllowShrinking::No);
 	}
 
+	if (bHasPreviousTransforms)
+	{
+		PerInstancePrevTransform.RemoveAtSwap(InstanceIndex);
+	}
+
 #if WITH_EDITOR
 	// remove selection flag if array is filled in
 	if (SelectedInstances.IsValidIndex(InstanceIndex))
diff --git a/Engine/Source/Runtime/Engine/Private/InstanceData/InstanceDataManager.cpp b/Engine/Source/Runtime/Engine/Private/InstanceData/InstanceDataManager.cpp
--- a/Engine/Source/Runtime/Engine/Private/InstanceData/InstanceDataManager.cpp
+++ b/Engine/Source/Runtime/Engine/Private/InstanceData/InstanceDataManager.cpp
@@ -143,6 +143,7 @@ public:
 					CondMove(ChangeSet.Flags.bHasPerInstanceCustomData, ProxyData.InstanceCustomData, FromIndex, ToIndex, ChangeSet.NumCustomDataFloats);
 					CondMove(ChangeSet.Flags.bHasPerInstanceRandom, ProxyData.InstanceRandomIDs, FromIndex, ToIndex);
 					CondMove(ChangeSet.Flags.bHasPerInstanceLMSMUVBias, ProxyData.InstanceLightShadowUVBias, FromIndex, ToIndex);
+					CondMove(ChangeSet.Flags.bHasPerInstanceSkinningData, ProxyData.InstanceSkinningData, FromIndex, ToIndex);
 	#if WITH_EDITOR
 					CondMove(ChangeSet.Flags.bHasPerInstanceEditorData, ProxyData.InstanceEditorData, FromIndex, ToIndex);
 	#endif
@@ -447,6 +448,12 @@ void FInstanceDataManager::CustomDataChanged(int32 InstanceIndex)
 	MarkChangeHelper<EChangeFlag::CustomDataChanged>(InstanceIndex);
 }
 
+void FInstanceDataManager::SkinningDataChanged(int32 InstanceIndex)
+{
+	LOG_INST_DATA(TEXT("SkinningDataChanged(IDX: %d)"), InstanceIndex);
+	MarkChangeHelper<EChangeFlag::SkinningDataChanged>(InstanceIndex);
+}
+
 void FInstanceDataManager::BakedLightingDataChanged(int32 InstanceIndex)
 {
 	LOG_INST_DATA(TEXT("BakedLightingDataChanged(IDX: %d)"), InstanceIndex);
diff --git a/Engine/Source/Runtime/Engine/Private/PrimitiveSceneProxy.cpp b/Engine/Source/Runtime/Engine/Private/PrimitiveSceneProxy.cpp
--- a/Engine/Source/Runtime/Engine/Private/PrimitiveSceneProxy.cpp
+++ b/Engine/Source/Runtime/Engine/Private/PrimitiveSceneProxy.cpp
@@ -912,7 +912,7 @@ void FPrimitiveSceneProxy::UpdateInstances_RenderThread(FRHICommandListBase& RHI
 	Bounds = InBounds;
 	LocalBounds = InLocalBounds;
 
-	bAlwaysHasVelocity = CVarVelocityForceOutput.GetValueOnAnyThread() ||  GetInstanceDataHeader().Flags.bHasPerInstanceDynamicData;
+	bAlwaysHasVelocity = CVarVelocityForceOutput.GetValueOnAnyThread() ||  GetInstanceDataHeader().Flags.bHasPerInstanceDynamicData || GetInstanceDataHeader().Flags.bHasPerInstanceSkinningData;
 }
 
 void FPrimitiveSceneProxy::UpdateInstances_RenderThread(FRHICommandListBase& RHICmdList, const FBoxSphereBounds& InBounds, const FBoxSphereBounds& InLocalBounds)
@@ -922,7 +922,7 @@ void FPrimitiveSceneProxy::UpdateInstances_RenderThread(FRHICommandListBase& RHI
 	Bounds = InBounds;
 	LocalBounds = InLocalBounds;
 
-	bAlwaysHasVelocity = CVarVelocityForceOutput.GetValueOnAnyThread() ||  GetInstanceDataHeader().Flags.bHasPerInstanceDynamicData;
+	bAlwaysHasVelocity = CVarVelocityForceOutput.GetValueOnAnyThread() || GetInstanceDataHeader().Flags.bHasPerInstanceDynamicData || GetInstanceDataHeader().Flags.bHasPerInstanceSkinningData;
 }
 
 bool FPrimitiveSceneProxy::WouldSetTransformBeRedundant_AnyThread(const FMatrix& InLocalToWorld, const FBoxSphereBounds& InBounds, const FBoxSphereBounds& InLocalBounds, const FVector& InActorPosition) const
diff --git a/Engine/Source/Runtime/Engine/Public/InstanceData/InstanceDataManager.h b/Engine/Source/Runtime/Engine/Public/InstanceData/InstanceDataManager.h
--- a/Engine/Source/Runtime/Engine/Public/InstanceData/InstanceDataManager.h
+++ b/Engine/Source/Runtime/Engine/Public/InstanceData/InstanceDataManager.h
@@ -99,6 +99,8 @@ public:
 
 	void CustomDataChanged(int32 InstanceIndex);
 
+	void SkinningDataChanged(int32 InstanceIndex);
+
 	void BakedLightingDataChanged(int32 InstanceIndex);
 	void BakedLightingDataChangedAll();
 
diff --git a/Engine/Source/Runtime/Engine/Public/InstanceData/InstanceUpdateChangeSet.h b/Engine/Source/Runtime/Engine/Public/InstanceData/InstanceUpdateChangeSet.h
--- a/Engine/Source/Runtime/Engine/Public/InstanceData/InstanceUpdateChangeSet.h
+++ b/Engine/Source/Runtime/Engine/Public/InstanceData/InstanceUpdateChangeSet.h
@@ -358,9 +358,8 @@ public:
 	auto GetTransformSetup()		{ return GetSetup<FInstanceAttributeTracker::EFlag::TransformChanged>(true, bUpdateAllInstanceTransforms, Transforms); }
 	auto GetPrevTransformSetup()	{ return GetSetup<FInstanceAttributeTracker::EFlag::TransformChanged>(Flags.bHasPerInstanceDynamicData, bUpdateAllInstanceTransforms || ForceFullFlags.bHasPerInstanceDynamicData, PrevTransforms); }
 	auto GetCustomDataSetup()		{ return GetSetup<FInstanceAttributeTracker::EFlag::CustomDataChanged>(Flags.bHasPerInstanceCustomData, ForceFullFlags.bHasPerInstanceCustomData, PerInstanceCustomData, NumCustomDataFloats); }
-	// These use an identity delta, which means they send all or nothing.
-	auto GetLocalBoundsSetup()		{ return GetSetup(Flags.bHasPerInstanceLocalBounds, InstanceLocalBounds); }
-	auto GetSkinningDataSetup()		{ return GetSetup(Flags.bHasPerInstanceSkinningData, InstanceSkinningData); }
+	auto GetLocalBoundsSetup()		{ return GetSetup<FInstanceAttributeTracker::EFlag::SkinningDataChanged>(Flags.bHasPerInstanceLocalBounds, ForceFullFlags.bHasPerInstanceLocalBounds, InstanceLocalBounds); }
+	auto GetSkinningDataSetup()		{ return GetSetup<FInstanceAttributeTracker::EFlag::SkinningDataChanged>(Flags.bHasPerInstanceSkinningData, ForceFullFlags.bHasPerInstanceSkinningData, InstanceSkinningData); }
 
 	// Convenience functions to get a reader / writer for a given array
 	auto GetTransformWriter()		{ return GetTransformSetup().GetWriter(); }
diff --git a/Engine/Source/Runtime/Engine/Public/InstancedSkinnedMeshComponentHelper.h b/Engine/Source/Runtime/Engine/Public/InstancedSkinnedMeshComponentHelper.h
--- a/Engine/Source/Runtime/Engine/Public/InstancedSkinnedMeshComponentHelper.h
+++ b/Engine/Source/Runtime/Engine/Public/InstancedSkinnedMeshComponentHelper.h
@@ -49,10 +49,11 @@ FInstanceDataManagerSourceDataDesc FInstancedSkinnedMeshComponentHelper::GetComp
 	const UTransformProviderData* TransformProvider = InComponent.GetTransformProvider();
 	const bool bForceRefPose = UInstancedSkinnedMeshComponent::ShouldForceRefPose();
 	const bool bValidTransformProvider = !bForceRefPose && TransformProvider != nullptr && TransformProvider->IsEnabled();
+	TConstArrayView<FTransform3f> InstancePrevTransforms = InComponent.GetInstancePrevTransforms();
 
 	Flags.bHasPerInstanceHierarchyOffset = false;
 	Flags.bHasPerInstanceLocalBounds = TransformProvider ? TransformProvider->HasAnimationBounds() : false;
-	Flags.bHasPerInstanceDynamicData = false;
+	Flags.bHasPerInstanceDynamicData = InstancePrevTransforms.Num() > 0 && InstancePrevTransforms.Num() == InComponent.GetInstanceCount();
 	Flags.bHasPerInstanceSkinningData = true;
 
 	Flags.bHasPerInstanceLMSMUVBias = false;//IsStaticLightingAllowed();
@@ -73,6 +74,7 @@ FInstanceDataManagerSourceDataDesc FInstancedSkinnedMeshComponentHelper::GetComp
 	[
 		&InComponent,
 		ComponentTransform,
+		InstancePrevTransforms,
 		TransformProvider = bValidTransformProvider ? TransformProvider : nullptr,
 		MeshBounds = ComponentDesc.MeshBounds,
 		bHasPerInstanceLocalBounds = Flags.bHasPerInstanceLocalBounds
@@ -80,6 +82,7 @@ FInstanceDataManagerSourceDataDesc FInstancedSkinnedMeshComponentHelper::GetComp
 	{
 		// publish data
 		ChangeSet.GetTransformWriter().Gather([&InComponent](int32 InstanceIndex) -> FRenderTransform { return FRenderTransform(InComponent.InstanceData[InstanceIndex].Transform.ToMatrixWithScale()); });
+		ChangeSet.GetPrevTransformWriter().Gather([InstancePrevTransforms](int32 InstanceIndex) -> FRenderTransform { return FRenderTransform(InstancePrevTransforms[InstanceIndex].ToMatrixWithScale()); });
 		ChangeSet.GetCustomDataWriter().Gather(MakeArrayView(InComponent.InstanceCustomData), InComponent.NumCustomDataFloats);
 
 		if (TransformProvider)
@@ -236,4 +239,4 @@ FPrimitiveSceneProxy* FInstancedSkinnedMeshComponentHelper::CreateSceneProxy(con
 	const bool bShouldNaniteSkin = FSkinnedMeshComponentHelper::ShouldNaniteSkin(InComponent);
 	const bool bEnabled = FInstancedSkinnedMeshComponentHelper::IsEnabled(InComponent);
 	return FInstancedSkinnedMeshSceneProxyDesc::CreateSceneProxy(InSceneProxyDesc, InComponent.bHideSkin, bShouldNaniteSkin, bEnabled, MinLODIndex);
-}
\ No newline at end of file
+}
diff --git a/Engine/Source/Runtime/Engine/Public/InstancedStaticMesh/InstanceAttributeTracker.h b/Engine/Source/Runtime/Engine/Public/InstancedStaticMesh/InstanceAttributeTracker.h
--- a/Engine/Source/Runtime/Engine/Public/InstancedStaticMesh/InstanceAttributeTracker.h
+++ b/Engine/Source/Runtime/Engine/Public/InstancedStaticMesh/InstanceAttributeTracker.h
@@ -79,7 +79,8 @@ public:
 	{
 		Added, /** Set when */
 		TransformChanged,
-		CustomDataChanged, 
+		CustomDataChanged,
+		SkinningDataChanged,
 		IndexChanged, /** Implicitly set when calling Remove or RemoveAt as this causes movement to fill holes. */
 		Num,
 	};
@@ -600,6 +601,7 @@ private:
 		const uint32 PrevFlags = GetFlags(Index);
 		if ((PrevFlags & FToBit<EFlag::Added>::Bit) == 0)
 		{
+			DecChangedCount<EFlag::SkinningDataChanged>(PrevFlags);
 			DecChangedCount<EFlag::CustomDataChanged>(PrevFlags);
 			DecChangedCount<EFlag::TransformChanged>(PrevFlags);
 			DecChangedCount<EFlag::IndexChanged>(PrevFlags);
