Don't tick FTickableGameObjects during engine shutdown, this could crash due to objects being destroyed
#Synced-CL 48815206

diff --git a/Engine/Source/Runtime/Engine/Private/Tickable.cpp b/Engine/Source/Runtime/Engine/Private/Tickable.cpp
--- a/Engine/Source/Runtime/Engine/Private/Tickable.cpp
+++ b/Engine/Source/Runtime/Engine/Private/Tickable.cpp
@@ -176,6 +176,8 @@ void FTickableGameObject::TickObjects(UWorld* World, const ELevelTick LevelTickT
 
 	check(IsInGameThread());
 
+	// Don't tick game objects after exit because too many things could be destroyed
+	if (!IsEngineExitRequested())
 	{
 		// It's a long lock but it's ok, the only thing we can block here is the GC worker thread that destroys UObjects
 		UE::TScopeLock LockTickableObjects(Statics.TickableObjectsCritical);
