Added ensures to help detect rare double-deinitialize of Subsystems in a SubsystemCollection

#tests Built and tested on Dedicated Server Cooked Data
#Synced-CL 48686588

diff --git a/Engine/Source/Runtime/Engine/Private/Subsystems/SubsystemCollection.cpp b/Engine/Source/Runtime/Engine/Private/Subsystems/SubsystemCollection.cpp
--- a/Engine/Source/Runtime/Engine/Private/Subsystems/SubsystemCollection.cpp
+++ b/Engine/Source/Runtime/Engine/Private/Subsystems/SubsystemCollection.cpp
@@ -202,6 +202,7 @@ void FSubsystemCollectionBase::RemoveSubsystemsInPackages(TConstArrayView<UPacka
 			
 			for (USubsystem* Subsystem : Iter.Value()->Subsystems)
 			{
+				ensureMsgf(IsValid(Subsystem), TEXT("RemoveSubsystemsInPackages tried to deinit invalid subsystem for SubsystemClass: %s"), *GetPathNameSafe(KeyClass));
 				if (IsValid(Subsystem) && Subsystem->GetClass() == KeyClass)
 				{
 					Subsystem->Deinitialize();
@@ -330,6 +331,7 @@ void FSubsystemCollectionBase::Deinitialize()
 	SubsystemArrayMap.Empty();
 	SubsystemMap.Empty();
 	Outer = nullptr;
+	bDeinitComplete = true;
 }
 
 USubsystem* FSubsystemCollectionBase::InitializeDependency(TSubclassOf<USubsystem> SubsystemClass)
@@ -482,6 +484,7 @@ USubsystem* FSubsystemCollectionBase::AddAndInitializeSubsystem(UClass* Subsyste
 
 USubsystem* FSubsystemCollectionBase::AddAndInitializeValidatedSubsystem(UClass* SubsystemClass)
 {
+	ensureMsgf(!bDeinitComplete, TEXT("FSubsystemCollectionBase::AddAndInitializeValidatedSubsystem was called on an already Deinitialized SubsystemCollection. Class: (%s)"), *GetPathNameSafe(SubsystemClass));
 	checkf(Initialization, TEXT("AddAndInitializeValidatedSubsystem should only be called from an initialization context"));
 	checkf(!SubsystemMap.Contains(SubsystemClass), TEXT("AddAndInitializeValidatedSubsystem called for a class which has already been added"));
 
@@ -512,6 +515,7 @@ USubsystem* FSubsystemCollectionBase::AddAndInitializeValidatedSubsystem(UClass*
 
 void FSubsystemCollectionBase::RemoveAndDeinitializeSubsystem(USubsystem* Subsystem)
 {
+	ensureMsgf(!bDeinitComplete, TEXT("FSubsystemCollectionBase::RemoveAndDeinitializeSubsystem was called on an already Deinitialized SubsystemCollection. Subsystem: (%s)"), *GetPathNameSafe(Subsystem));
 	check(Subsystem);
 	check(!bIterating);
 	USubsystem* SubsystemFound = SubsystemMap.FindAndRemoveChecked(Subsystem->GetClass());
diff --git a/Engine/Source/Runtime/Engine/Public/Subsystems/SubsystemCollection.h b/Engine/Source/Runtime/Engine/Public/Subsystems/SubsystemCollection.h
--- a/Engine/Source/Runtime/Engine/Public/Subsystems/SubsystemCollection.h
+++ b/Engine/Source/Runtime/Engine/Public/Subsystems/SubsystemCollection.h
@@ -147,6 +147,8 @@ private:
 
 	mutable bool bIterating = false; // True if iterating over SubsystemMap
 
+	bool bDeinitComplete = false; // True if the subsystem collection has finished being Deinitialized
+
 	FDelegateHandle ModulesUnloadedHandle;
 
 private:
