[Sound Wave] - LogSlate warnings when expanding timecode that display incorrect symbols
#Synced-CL 48848619

diff --git a/Engine/Source/Runtime/Engine/Private/Audio.cpp b/Engine/Source/Runtime/Engine/Private/Audio.cpp
--- a/Engine/Source/Runtime/Engine/Private/Audio.cpp
+++ b/Engine/Source/Runtime/Engine/Private/Audio.cpp
@@ -13,6 +13,7 @@
 #include "Engine/Engine.h"
 #include "EngineAnalytics.h"
 #include "IAnalyticsProviderET.h"
+#include "Misc/FileHelper.h"
 #include "Misc/Paths.h"
 #include "Sound/AudioOutputTarget.h"
 #include "Sound/SoundCue.h"
@@ -1533,7 +1534,7 @@ struct FRiffBroadcastAudioExtension
 	uint16 MaxTruePeakLevel;		// WORD : «Maximum True Peak Level of the file expressed as dBTP (multiplied by 100) » 
 	uint16 MaxMomentaryLoudness;	// WORD : «Highest value of the Momentary Loudness Level of the file in LUFS (multiplied by 100) » 
 	uint16 MaxShortTermLoudness;	// WORD : «Highest value of the Short-TermLoudness Level of the file in LUFS (multiplied by 100) » 
-	uint16 Reserved[180];			// 180 bytes, reserved for future use, set to “NULL” 
+	uint8 Reserved[180];			// 180 bytes, reserved for future use, set to “NULL” 
 	uint8 CodingHistory[1];			// ASCII : « History coding » (truncated here are don't care about it) 
 };
 
@@ -2013,14 +2014,17 @@ bool FWaveModInfo::ReadWaveInfo( const uint8* WaveData, int32 WaveDataSize, FStr
 		if(NumSamplesSinceMidnight > 0 && FmtChunk && FmtChunk->nSamplesPerSec)
 		{
 			TimecodeInfo = MakePimpl<FSoundWaveTimecodeInfo, EPimplPtrMode::DeepCopy>();
-				
+
 			TimecodeInfo->NumSamplesSinceMidnight   = NumSamplesSinceMidnight;
 			TimecodeInfo->NumSamplesPerSecond		= FmtChunk->nSamplesPerSec;
-			TimecodeInfo->Description				= StringCast<TCHAR>((ANSICHAR*)BextChunk->Description).Get();
-			TimecodeInfo->OriginatorDescription		= StringCast<TCHAR>((ANSICHAR*)BextChunk->Originator).Get();
-			TimecodeInfo->OriginatorDate			= StringCast<TCHAR>((ANSICHAR*)BextChunk->OriginationDate).Get();
-			TimecodeInfo->OriginatorReference		= StringCast<TCHAR>((ANSICHAR*)BextChunk->OriginatorReference).Get();
-			TimecodeInfo->OriginatorTime			= StringCast<TCHAR>((ANSICHAR*)BextChunk->OriginationTime).Get();
+			
+			// BextChunk char arrays have a fixed length and might not be null terminated or padded
+			// Safely retrieve either the first null termination or the safe max length of the char array
+			FFileHelper::BufferToString(TimecodeInfo->Description, BextChunk->Description, sizeof(BextChunk->Description));
+			FFileHelper::BufferToString(TimecodeInfo->OriginatorDescription, BextChunk->Originator, sizeof(BextChunk->Originator));
+			FFileHelper::BufferToString(TimecodeInfo->OriginatorDate, BextChunk->OriginationDate, sizeof(BextChunk->OriginationDate));
+			FFileHelper::BufferToString(TimecodeInfo->OriginatorReference, BextChunk->OriginatorReference, sizeof(BextChunk->OriginatorReference));
+			FFileHelper::BufferToString(TimecodeInfo->OriginatorTime, BextChunk->OriginationTime, sizeof(BextChunk->OriginationTime));
 
 			// Use the sample rate as the timecode rate for now. We'll replace
 			// it with a more accurate rate if possible below.
@@ -2265,7 +2269,7 @@ void SerializeWaveFile(TArray<uint8>& OutWaveFileData, const uint8* InPCMData, c
 	// FieldName: BlockAlign
 	// FieldSize: 2 bytes
 	// FieldValue: NumChannels * BitsPerSample/8
-	int32 BlockAlign = 2;
+	int32 BlockAlign = NumChannels * 2;
 	WriteUInt16ToByteArrayLE(OutWaveFileData, WaveDataByteIndex, BlockAlign);
 
 	// FieldName: BitsPerSample
