[ActorFolders] Make intermediate folder APIs internal; code outside the engine should only use FActorFolders.

#tests Built all sample and main projects
#Synced-CL 48770656

diff --git a/Engine/Source/Editor/UnrealEd/Public/WorldFolders.h b/Engine/Source/Editor/UnrealEd/Public/WorldFolders.h
--- a/Engine/Source/Editor/UnrealEd/Public/WorldFolders.h
+++ b/Engine/Source/Editor/UnrealEd/Public/WorldFolders.h
@@ -15,7 +15,7 @@
 DECLARE_LOG_CATEGORY_EXTERN(LogWorldFolders, Log, All)
 
 USTRUCT()
-struct FActorFolderProps
+struct UE_INTERNAL FActorFolderProps
 {
 	GENERATED_USTRUCT_BODY()
 
@@ -30,7 +30,7 @@ struct FActorFolderProps
 };
 
 USTRUCT()
-struct FActorPlacementFolder
+struct UE_INTERNAL FActorPlacementFolder
 {
 	GENERATED_BODY()
 
@@ -65,9 +65,13 @@ struct FActorPlacementFolder
 	}
 };
 
-/** Per-World Actor Folders UObject (used to support undo/redo reliably) */
+/**
+ * Per-World Actor Folders UObject (used to support undo/redo reliably)
+ * This class is reserved for low-level Engine use.
+ * The public, non-engine API for Folder operations is through the FActorFolders struct.
+ */
 UCLASS(MinimalAPI)
-class UWorldFolders : public UObject
+class UE_INTERNAL UWorldFolders : public UObject
 {
 public:
 	GENERATED_BODY()
diff --git a/Engine/Source/Editor/UnrealEd/Public/WorldFoldersImplementation.h b/Engine/Source/Editor/UnrealEd/Public/WorldFoldersImplementation.h
--- a/Engine/Source/Editor/UnrealEd/Public/WorldFoldersImplementation.h
+++ b/Engine/Source/Editor/UnrealEd/Public/WorldFoldersImplementation.h
@@ -8,7 +8,11 @@ class UWorld;
 class UWorldFolders;
 struct FFolder;
 
-class FWorldFoldersImplementation
+/**
+ * This class is reserved for low-level Engine use, therefore all APIs are marked UE_INTERNAL.
+ * The public, non-engine API for Folder operations is through the FActorFolders struct.
+ */
+class UE_INTERNAL FWorldFoldersImplementation
 {
 public:
 	FWorldFoldersImplementation(UWorldFolders& Owner);
diff --git a/Engine/Source/Editor/UnrealEd/Public/WorldPersistentFolders.h b/Engine/Source/Editor/UnrealEd/Public/WorldPersistentFolders.h
--- a/Engine/Source/Editor/UnrealEd/Public/WorldPersistentFolders.h
+++ b/Engine/Source/Editor/UnrealEd/Public/WorldPersistentFolders.h
@@ -8,8 +8,12 @@
 
 class ULevel;
 
-// Class handling a list of actor folder objects of a world
-class FWorldPersistentFolders : public FWorldFoldersImplementation
+/**
+ * Class handling a list of actor folder objects of a world
+ * This class is reserved for low-level Engine use.
+ * The public, non-engine API for Folder operations is through the FActorFolders struct.
+ */
+class UE_INTERNAL FWorldPersistentFolders : public FWorldFoldersImplementation
 {
 	typedef FWorldFoldersImplementation Super;
 
diff --git a/Engine/Source/Editor/UnrealEd/Public/WorldTransientFolders.h b/Engine/Source/Editor/UnrealEd/Public/WorldTransientFolders.h
--- a/Engine/Source/Editor/UnrealEd/Public/WorldTransientFolders.h
+++ b/Engine/Source/Editor/UnrealEd/Public/WorldTransientFolders.h
@@ -4,8 +4,12 @@
 
 #include "WorldFoldersImplementation.h"
 
-// Class handling a list of transient actor folders of a world
-class FWorldTransientFolders : public FWorldFoldersImplementation
+/**
+ * Class handling a list of transient actor folders of a world
+ * This class is reserved for low-level Engine use, therefore all APIs are marked UE_INTERNAL.
+ * The public, non-engine API for Folder operations is through the FActorFolders struct.
+ */
+class UE_INTERNAL FWorldTransientFolders : public FWorldFoldersImplementation
 {
 	typedef FWorldFoldersImplementation Super;
 
diff --git a/Engine/Source/Runtime/Engine/Private/ActorFolder.cpp b/Engine/Source/Runtime/Engine/Private/ActorFolder.cpp
--- a/Engine/Source/Runtime/Engine/Private/ActorFolder.cpp
+++ b/Engine/Source/Runtime/Engine/Private/ActorFolder.cpp
@@ -69,13 +69,6 @@ namespace ActorFolder
 	static const FName NAME_OuterPackageName("OuterPackageName");
 };
 
-void UActorFolder::GetAssetRegistryTags(TArray<FAssetRegistryTag>& OutTags) const
-{
-	PRAGMA_DISABLE_DEPRECATION_WARNINGS;
-	Super::GetAssetRegistryTags(OutTags);
-	PRAGMA_ENABLE_DEPRECATION_WARNINGS;
-}
-
 void UActorFolder::GetAssetRegistryTags(FAssetRegistryTagsContext Context) const
 {
 	Super::GetAssetRegistryTags(Context);
diff --git a/Engine/Source/Runtime/Engine/Public/ActorFolder.h b/Engine/Source/Runtime/Engine/Public/ActorFolder.h
--- a/Engine/Source/Runtime/Engine/Public/ActorFolder.h
+++ b/Engine/Source/Runtime/Engine/Public/ActorFolder.h
@@ -12,6 +12,11 @@
 class ULevel;
 struct FActorFolderDesc;
 
+/**
+ * This class is reserved for low-level Engine use, with all APIs marked UE_INTERNAL.
+ * The type is visible on public classes however, so itself cannot be made UE_INTERNAL.
+ * If you are looking for the public API for Folder operations, look at FActorFolders.
+ */
 UCLASS(Within = Level, MinimalAPI)
 class UActorFolder : public UObject
 {
@@ -20,58 +25,58 @@ class UActorFolder : public UObject
 #if WITH_EDITOR
 public:
 	// Helper that creates a new UActorFolder
-	static ENGINE_API UActorFolder* Create(ULevel* InOuter, const FString& InFolderLabel, UActorFolder* InParent);
+	UE_INTERNAL static ENGINE_API UActorFolder* Create(ULevel* InOuter, const FString& InFolderLabel, UActorFolder* InParent);
 
-	UE_DEPRECATED(5.4, "This function is deprecated.")
-	static EPackageFlags GetExternalPackageFlags() { return (PKG_EditorOnly | PKG_ContainsMapData | PKG_NewlyCreated); }
-	
 	// Returns actor folder info stored in its package
-	static ENGINE_API FActorFolderDesc GetAssetRegistryInfoFromPackage(FName PackageName);
+	UE_INTERNAL static ENGINE_API FActorFolderDesc GetAssetRegistryInfoFromPackage(FName PackageName);
 	
 	//~ Begin UObject
-	ENGINE_API virtual bool IsAsset() const override;
-	ENGINE_API virtual void GetAssetRegistryTags(FAssetRegistryTagsContext Context) const override;
-	UE_DEPRECATED(5.4, "Implement the version that takes FAssetRegistryTagsContext instead.")
-	ENGINE_API virtual void GetAssetRegistryTags(TArray<FAssetRegistryTag>& OutTags) const override;
-	ENGINE_API virtual void PostLoad() override;
+	UE_INTERNAL ENGINE_API virtual bool IsAsset() const override;
+	UE_INTERNAL ENGINE_API virtual void GetAssetRegistryTags(FAssetRegistryTagsContext Context) const override;
+	UE_INTERNAL ENGINE_API virtual void PostLoad() override;
 	//~ End UObject
 
-	ENGINE_API void SetParent(UActorFolder* InParent);
-	ENGINE_API UActorFolder* GetParent(bool bSkipDeleted = true) const;
+	UE_INTERNAL ENGINE_API void SetParent(UActorFolder* InParent);
+	UE_INTERNAL ENGINE_API UActorFolder* GetParent(bool bSkipDeleted = true) const;
 
-	ENGINE_API void SetLabel(const FString& InFolderLabel);
-	const FString& GetLabel() const { return FolderLabel; }
+	UE_INTERNAL ENGINE_API void SetLabel(const FString& InFolderLabel);
+	UE_INTERNAL const FString& GetLabel() const { return FolderLabel; }
 
 	// Control whether the folder is initially expanded or not
-	ENGINE_API void SetIsInitiallyExpanded(bool bInFolderInitiallyExpanded);
-	bool IsInitiallyExpanded() const { return bFolderInitiallyExpanded; }
+	UE_INTERNAL ENGINE_API void SetIsInitiallyExpanded(bool bInFolderInitiallyExpanded);
+	UE_INTERNAL bool IsInitiallyExpanded() const { return bFolderInitiallyExpanded; }
 
-	const FGuid& GetGuid() const { return FolderGuid; }
-	ENGINE_API FName GetPath() const;
-	ENGINE_API FString GetDisplayName() const;
+	UE_INTERNAL const FGuid& GetGuid() const { return FolderGuid; }
+	UE_INTERNAL ENGINE_API FName GetPath() const;
+	UE_INTERNAL ENGINE_API FString GetDisplayName() const;
 
-	ENGINE_API void MarkAsDeleted();
-	bool IsMarkedAsDeleted() const { return bIsDeleted; }
+	UE_INTERNAL ENGINE_API void MarkAsDeleted();
+	UE_INTERNAL bool IsMarkedAsDeleted() const { return bIsDeleted; }
 	
 	// Checks if folder is valid (if it's not deleted)
-	bool IsValid() const { return !IsMarkedAsDeleted(); }
+	UE_INTERNAL bool IsValid() const
+	{
+		PRAGMA_DISABLE_INTERNAL_WARNINGS
+		return !IsMarkedAsDeleted();
+		PRAGMA_ENABLE_INTERNAL_WARNINGS
+	}
 
 	// Remaps parent folder to the first parent folder not marked as deleted
-	ENGINE_API void Fixup();
+	UE_INTERNAL ENGINE_API void Fixup();
 	// Detects and clears invalid parent folder
-	ENGINE_API void FixupParentFolder();
+	UE_INTERNAL ENGINE_API void FixupParentFolder();
 
-	ENGINE_API FFolder GetFolder() const;
+	UE_INTERNAL ENGINE_API FFolder GetFolder() const;
 
 	/**
 	 * Set the folder packaging mode.
 	* @param bExternal will set the folder packaging mode to external if true, to internal otherwise
 	* @param bShouldDirty should dirty or not the level package
 	*/
-	ENGINE_API void SetPackageExternal(bool bExternal, bool bShouldDirty = true);
+	UE_INTERNAL ENGINE_API void SetPackageExternal(bool bExternal, bool bShouldDirty = true);
 
 private:
-	ENGINE_API FName GetPathInternal(UActorFolder* InSkipFolder) const;
+	UE_INTERNAL ENGINE_API FName GetPathInternal(UActorFolder* InSkipFolder) const;
 #endif
 
 #if WITH_EDITORONLY_DATA
