Remove functions deprecated in 4.13 from the UNavAgentInterface

- GetMoveGoalOffset
- GetMoveGoalReachTest

#rnx
#Synced-CL 48850057

diff --git a/Engine/Source/Runtime/Engine/Classes/AI/Navigation/NavAgentInterface.h b/Engine/Source/Runtime/Engine/Classes/AI/Navigation/NavAgentInterface.h
--- a/Engine/Source/Runtime/Engine/Classes/AI/Navigation/NavAgentInterface.h
+++ b/Engine/Source/Runtime/Engine/Classes/AI/Navigation/NavAgentInterface.h
@@ -51,20 +51,5 @@ class INavAgentInterface
 
 	/** Retrieves PathFollowingAgent for this NavAgent */
 	virtual IPathFollowingAgentInterface* GetPathFollowingAgent() const { return nullptr; }
-
-	//----------------------------------------------------------------------//
-	// DEPRECATED
-	//----------------------------------------------------------------------//
-	UE_DEPRECATED_FORGAME(4.13, "This function is now deprecated and will not be called, please use override with const Actor pointer.")
-	virtual FVector GetMoveGoalOffset(AActor* MovingActor) const
-	{
-		return GetMoveGoalOffset((const AActor*)MovingActor);
-	}
-
-	UE_DEPRECATED_FORGAME(4.13, "This function is now deprecatedand will not be called, please use override with const Actor pointer.")
-	virtual void GetMoveGoalReachTest(AActor* MovingActor, const FVector& MoveOffset, FVector& GoalOffset, float& GoalRadius, float& GoalHalfHeight) const 
-	{
-		GetMoveGoalReachTest((const AActor*)MovingActor, MoveOffset, GoalOffset, GoalRadius, GoalHalfHeight);
-	}
 };
 
