Dstm: Chaos vehicle support 1.1: Make FSkeletalMeshComponentEndPhysicsTick transactional and support migration of partially invalid constraints.

#Synced-CL 48671814

diff --git a/Engine/Source/Runtime/Engine/Private/Components/SkeletalMeshComponent.cpp b/Engine/Source/Runtime/Engine/Private/Components/SkeletalMeshComponent.cpp
--- a/Engine/Source/Runtime/Engine/Private/Components/SkeletalMeshComponent.cpp
+++ b/Engine/Source/Runtime/Engine/Private/Components/SkeletalMeshComponent.cpp
@@ -443,6 +443,7 @@ USkeletalMeshComponent::USkeletalMeshComponent(const FObjectInitializer& ObjectI
 	EndPhysicsTickFunction.TickGroup = TG_EndPhysics;
 	EndPhysicsTickFunction.bCanEverTick = true;
 	EndPhysicsTickFunction.bStartWithTickEnabled = true;
+	EndPhysicsTickFunction.bRunTransactionally = true;
 
 	ClothTickFunction.TickGroup = TG_PrePhysics;
 	ClothTickFunction.EndTickGroup = TG_PostPhysics;
@@ -4008,7 +4009,33 @@ void USkeletalMeshComponent::ApplyMigratedPhysicsState(Chaos::FSerializedDataBuf
 		{
 			check(ConstraintInstance);
 
-			MigratedState.ApplyDataToElement(*ConstraintInstance);
+			// Some other components may invalidates constraints in skeletal mesh to perform their own simulation (e.g. see UChaosWheeledVehicleMovementComponent::FixupSkeletalMesh())
+			// This code ignores invalid constraint data on the receiving server to get around the ensure in Chaos::Private::ExtractJointConstraintDataFromBuffer
+			auto ApplyMigratedData = [&MigratedState, ConstraintInstance](Chaos::ESerializedDataContext ContextToApply)
+				{
+					Chaos::FSerializedDataBufferPtr DataBuffer = MakeUnique<Chaos::FSerializedDataBuffer>();
+
+					MigratedState.ReadElementDataIntoBuffer(*DataBuffer);
+
+					if (!DataBuffer->GetDataAsByteArrayRef().IsEmpty())
+					{
+						ConstraintInstance->ApplySerializedState(MoveTemp(DataBuffer), ContextToApply);
+					}
+					else
+					{
+						UE_LOG(LogSkeletalMesh, Log, TEXT("[%hs]: Constraint %s has empty data when migrated so it is skipped."),__func__, *ConstraintInstance->JointName.ToString());
+					}
+				};
+
+			if (MigratedState.GetHeader().DataContext == Chaos::ESerializedDataContext::Both)
+			{
+				ApplyMigratedData(Chaos::ESerializedDataContext::Internal);
+				ApplyMigratedData(Chaos::ESerializedDataContext::External);
+			}
+			else
+			{
+				ApplyMigratedData(MigratedState.GetHeader().DataContext);
+			}
 		}
 	}
 }
