Enable Dynamic Resolution support for PC RHIs DX12 and Vulkan

#Synced-CL 48819522

diff --git a/Engine/Source/Runtime/D3D12RHI/Private/Windows/WindowsD3D12Device.cpp b/Engine/Source/Runtime/D3D12RHI/Private/Windows/WindowsD3D12Device.cpp
--- a/Engine/Source/Runtime/D3D12RHI/Private/Windows/WindowsD3D12Device.cpp
+++ b/Engine/Source/Runtime/D3D12RHI/Private/Windows/WindowsD3D12Device.cpp
@@ -1892,6 +1892,8 @@ void FD3D12DynamicRHI::Init()
 		GRHIGlobals.RayTracing.SupportsPersistentSBTs = true;
 	}
 
+	GRHISupportsDynamicResolution = true;
+
 	// Indicate that the RHI needs to use the engine's deferred deletion queue.
 	GRHINeedsExtraDeletionLatency = true;
 
diff --git a/Engine/Source/Runtime/Engine/Private/UnrealEngine.cpp b/Engine/Source/Runtime/Engine/Private/UnrealEngine.cpp
--- a/Engine/Source/Runtime/Engine/Private/UnrealEngine.cpp
+++ b/Engine/Source/Runtime/Engine/Private/UnrealEngine.cpp
@@ -488,6 +488,16 @@ static TAutoConsoleVariable<float> CVarSetOverrideFPS(
 	ECVF_Cheat);
 #endif // !UE_BUILD_SHIPPING
 
+#if WITH_EDITOR
+static TAutoConsoleVariable<int> CVarDynamicResAllowInPIE(
+	TEXT("r.DynamicRes.AllowInPIE"),
+	1,
+	TEXT("When Dynamic Resolution is enabled for the project on PC: select if Dynamic Resolution is allowed during PIE in Editor\n")
+	TEXT(" 0: Not Allowed;\n")
+	TEXT(" 1: Allowed (default);\n"),
+	ECVF_RenderThreadSafe | ECVF_Default);
+#endif
+
 static TAutoConsoleVariable<int> CVarDynamicResOperationMode(
 	TEXT("r.DynamicRes.OperationMode"),
 	0,
@@ -13887,6 +13897,8 @@ bool UEngine::ShouldEnableDynamicResolutionState() const
 #if WITH_EDITOR
 	if (GIsEditor && bEnable)
 	{
+		int32 AllowedInPIE = CVarDynamicResAllowInPIE.GetValueOnGameThread();
+
 		int32 bPIEContextCount = 0;
 		for (auto It = WorldList.CreateConstIterator(); It; ++It)
 		{
@@ -13902,7 +13914,7 @@ bool UEngine::ShouldEnableDynamicResolutionState() const
 			}
 		}
 
-		bEnable = (bPIEContextCount == 1);
+		bEnable = (bPIEContextCount == 1) && (AllowedInPIE == 1);
 	}
 #endif
 
diff --git a/Engine/Source/Runtime/VulkanRHI/Private/VulkanGenericPlatform.h b/Engine/Source/Runtime/VulkanRHI/Private/VulkanGenericPlatform.h
--- a/Engine/Source/Runtime/VulkanRHI/Private/VulkanGenericPlatform.h
+++ b/Engine/Source/Runtime/VulkanRHI/Private/VulkanGenericPlatform.h
@@ -97,8 +97,8 @@ public:
 
 	static bool SupportParallelRenderingTasks() { return true; }
 
-	/** The status quo is false, so the default is chosen to not change it. As platforms opt in it may be better to flip the default. */
-	static bool SupportsDynamicResolution() { return false; }
+	/** The status quo is true, Platforms might opt out if GPU timing queries are not supported. */
+	static bool SupportsDynamicResolution() { return true; }
 
 	static bool SupportsVolumeTextureRendering() { return true; }
 
