Add a guid to UWorld's class schema to force world assets to incrementally recook to get past saved ensure message.

#Synced-CL Unknown

diff --git a/Engine/Source/Runtime/Engine/Classes/Engine/World.h b/Engine/Source/Runtime/Engine/Classes/Engine/World.h
--- a/Engine/Source/Runtime/Engine/Classes/Engine/World.h
+++ b/Engine/Source/Runtime/Engine/Classes/Engine/World.h
@@ -3032,6 +3032,7 @@ public:
 	UE_API virtual bool IsReadyForFinishDestroy() override;
 	UE_API virtual void PostLoad() override;
 #if WITH_EDITORONLY_DATA
+	static UE_API void AppendToClassSchema(FAppendToClassSchemaContext& Context);
 	static UE_API void DeclareConstructClasses(TArray<FTopLevelAssetPath>& OutConstructClasses, const UClass* SpecificSubclass);
 #endif
 	UE_API virtual void PreDuplicate(FObjectDuplicationParameters& DupParams) override;
diff --git a/Engine/Source/Runtime/Engine/Private/World.cpp b/Engine/Source/Runtime/Engine/Private/World.cpp
--- a/Engine/Source/Runtime/Engine/Private/World.cpp
+++ b/Engine/Source/Runtime/Engine/Private/World.cpp
@@ -1772,6 +1772,13 @@ void UWorld::PostLoad()
 }
 
 #if WITH_EDITORONLY_DATA
+void UWorld::AppendToClassSchema(FAppendToClassSchemaContext& Context)
+{
+	Super::AppendToClassSchema(Context);
+	constexpr FGuid WorldCookGuid(0x4b58bae1, 0x85e04e88, 0xa8511c76, 0x1d656890);
+	Context.Update(&WorldCookGuid, sizeof(WorldCookGuid));
+}
+
 void UWorld::DeclareConstructClasses(TArray<FTopLevelAssetPath>& OutConstructClasses, const UClass* SpecificSubclass)
 {
 	Super::DeclareConstructClasses(OutConstructClasses, SpecificSubclass);
