Boot hotfix CVar for AudioMixerModuleName

#Synced-CL 48883003

diff --git a/Engine/Source/Runtime/Engine/Private/AudioDeviceManager.cpp b/Engine/Source/Runtime/Engine/Private/AudioDeviceManager.cpp
--- a/Engine/Source/Runtime/Engine/Private/AudioDeviceManager.cpp
+++ b/Engine/Source/Runtime/Engine/Private/AudioDeviceManager.cpp
@@ -9,8 +9,8 @@
 #include "Engine/Engine.h"
 #include "Engine/World.h"
 #include "Features/IModularFeatures.h"
-#include "Misc/App.h"
 #include "GameFramework/GameUserSettings.h"
+#include "Misc/App.h"
 #include "Misc/ConfigCacheIni.h"
 #include "Misc/CoreDelegates.h"
 #include "UObject/UObjectIterator.h"
@@ -68,6 +68,15 @@ FAutoConsoleVariableRef CVarNeverMuteNonRealtimeAudioDevices(
 	TEXT("0: Not Disabled, 1: Disabled"),
 	ECVF_Default);
 
+static FString AudioMixerModuleNameHotfixCVar;
+static FString AudioMixerModuleNameHotfixCVarName = TEXT("au.AudioMixerModuleName");
+FAutoConsoleVariableRef CVarAudioMixerModuleNameCVar(
+	*AudioMixerModuleNameHotfixCVarName,
+	AudioMixerModuleNameHotfixCVar,
+	TEXT("Selects an audio mixer device for Microsoft platforms. \n")
+	TEXT("AudioMixerWasapi or AudioMixerXAudio2"),
+	ECVF_SaveForNextBoot);
+
 static FAutoConsoleCommand GReportAudioDevicesCommand(
 	TEXT("au.ReportAudioDevices"),
 	TEXT("This will log any active audio devices (instances of the audio engine) alive right now."),
@@ -475,15 +484,26 @@ bool FAudioDeviceManager::LoadDefaultAudioDeviceModule()
 {
 	check(!AudioDeviceModule);
 
+	if (!AudioMixerModuleNameHotfixCVar.IsEmpty())
+	{
+		UE_LOG(LogAudio, Display, TEXT("Current AudioMixerModuleName hotfixed to be: %s"), *AudioMixerModuleNameHotfixCVar);
+		AudioMixerModuleName = AudioMixerModuleNameHotfixCVar;
+	}
+
 	bool bForceNonRealtimeRenderer = FParse::Param(FCommandLine::Get(), TEXT("DeterministicAudio"));
-	bool bFoundModuleName = false;
 	
 #if WITH_EDITOR
 	// Check to see if the editor pref has been set. If not, fall back to engine setting.
-	bFoundModuleName = GConfig->GetString(TEXT("/Script/AudioEditor.AudioEditorSettings"), TEXT("AudioMixerModuleName"), AudioMixerModuleName, GEditorSettingsIni);
+	FString TempMixerModuleName;
+	GConfig->GetString(TEXT("/Script/AudioEditor.AudioEditorSettings"), TEXT("AudioMixerModuleName"), TempMixerModuleName, GEditorSettingsIni);
+	if (!TempMixerModuleName.IsEmpty())
+	{
+		AudioMixerModuleName = TempMixerModuleName;
+	}
+	
 #endif //WITH_EDITOR
 
-	if (!bFoundModuleName || AudioMixerModuleName.IsEmpty())
+	if (AudioMixerModuleName.IsEmpty())
 	{
 		// If not using command line switch to use audio mixer, check the game platform engine ini file (e.g. WindowsEngine.ini) which enables it for player
 		GConfig->GetString(TEXT("Audio"), TEXT("AudioMixerModuleName"), AudioMixerModuleName, GEngineIni);
