[NetPhysics] Return INDEX_NONE when there is no last received state frame, and use either frame from received custom state or received target state (RepMovement).

#Synced-CL 48777322

diff --git a/Engine/Source/Runtime/Engine/Private/PhysicsEngine/NetworkPhysicsComponent.cpp b/Engine/Source/Runtime/Engine/Private/PhysicsEngine/NetworkPhysicsComponent.cpp
--- a/Engine/Source/Runtime/Engine/Private/PhysicsEngine/NetworkPhysicsComponent.cpp
+++ b/Engine/Source/Runtime/Engine/Private/PhysicsEngine/NetworkPhysicsComponent.cpp
@@ -2053,6 +2053,7 @@ void FAsyncNetworkPhysicsComponent::ConsumeAsyncInput(const int32 PhysicsStep)
 			// Receive States
 			if (AsyncInput->StateData && AsyncInput->StateData->HasDataInHistory())
 			{
+				LastReceivedStateFrame = FMath::Max(LastReceivedStateFrame, AsyncInput->StateData->GetLatestFrame() - NetworkPhysicsTickOffset);
 				ReceiveHelper(StateHistory.Get(), AsyncInput->StateData.Get(), /*bImportant*/false, bCompareState);
 			}
 
@@ -2073,6 +2074,7 @@ void FAsyncNetworkPhysicsComponent::ConsumeAsyncInput(const int32 PhysicsStep)
 				{
 					continue;
 				}
+				LastReceivedStateFrame = FMath::Max(LastReceivedStateFrame, StateImportant->GetLatestFrame() - NetworkPhysicsTickOffset);
 				ReceiveHelper(StateHistory.Get(), StateImportant.Get(), /*bImportant*/true, bCompareState);
 			}
 
@@ -2812,9 +2814,13 @@ const float FAsyncNetworkPhysicsComponent::GetInputLagFromLastReceivedTargetStat
 		Chaos::FReadPhysicsObjectInterface_Internal Interface = Chaos::FPhysicsObjectInternalInterface::GetRead();
 		if (Chaos::FGeometryParticleHandle* Particle = Interface.GetParticle(PhysicsObject))
 		{
-			const int32 MeasuredInputDelayFrames = RewindData->CurrentFrame() - RewindData->GetLatestReceivedTargetStateFrame(*Particle);
-			const float MeasuredInputDelaySeconds = FMath::Max(MeasuredInputDelayFrames, 0.0f) * BaseSolver->GetAsyncDeltaTime();
-			return MeasuredInputDelaySeconds;
+			const int32 LatestTargetFrame = GetLastReceivedServerFrame();
+			if (LatestTargetFrame >= 0)
+			{
+				const int32 MeasuredInputDelayFrames = RewindData->CurrentFrame() - LatestTargetFrame;
+				const float MeasuredInputDelaySeconds = FMath::Max(MeasuredInputDelayFrames, 0.0f) * BaseSolver->GetAsyncDeltaTime();
+				return MeasuredInputDelaySeconds;
+			}
 		}
 	}
 	return 0.0f;
@@ -3067,11 +3073,24 @@ const int32 FAsyncNetworkPhysicsComponent::SetupRewindData()
 
 int32 FAsyncNetworkPhysicsComponent::GetLastReceivedServerFrame() const
 {
-	if (StateHistory)
+	if (LastReceivedStateFrame > INDEX_NONE)
+	{
+		return LastReceivedStateFrame;
+	}
+	else
 	{
-		return StateHistory->GetLatestFrame();
+		const Chaos::FPhysicsSolverBase* BaseSolver = GetSolver();
+		if (const Chaos::FRewindData* RewindData = BaseSolver ? BaseSolver->GetRewindData() : nullptr)
+		{
+			// Todo: we should not have to get a read lock to get this information bu instead measure the latest lag upon reception
+			Chaos::FReadPhysicsObjectInterface_Internal Interface = Chaos::FPhysicsObjectInternalInterface::GetRead();
+			if (Chaos::FGeometryParticleHandle* Particle = Interface.GetParticle(PhysicsObject))
+			{
+				return RewindData->GetLatestReceivedTargetStateFrame(*Particle);;
+			}
+		}
 	}
-	return -1;
+	return INDEX_NONE;
 }
 
 namespace UE::NetworkPhysicsUtils
diff --git a/Engine/Source/Runtime/Engine/Public/Physics/NetworkPhysicsComponent.h b/Engine/Source/Runtime/Engine/Public/Physics/NetworkPhysicsComponent.h
--- a/Engine/Source/Runtime/Engine/Public/Physics/NetworkPhysicsComponent.h
+++ b/Engine/Source/Runtime/Engine/Public/Physics/NetworkPhysicsComponent.h
@@ -2009,6 +2009,8 @@ private:
 	int32 LastStateSendFrame = INDEX_NONE;
 	int32 NewImportantInputFrame = INT_MAX;
 
+	int32 LastReceivedStateFrame = INDEX_NONE;
+
 	// Component settings
 	TWeakPtr<const FNetworkPhysicsSettingsData> SettingsComponent;
 	static const FNetworkPhysicsSettingsNetworkPhysicsComponent SettingsNetworkPhysicsComponent_Default;
diff --git a/Engine/Source/Runtime/Experimental/Chaos/Private/RewindData.cpp b/Engine/Source/Runtime/Experimental/Chaos/Private/RewindData.cpp
--- a/Engine/Source/Runtime/Experimental/Chaos/Private/RewindData.cpp
+++ b/Engine/Source/Runtime/Experimental/Chaos/Private/RewindData.cpp
@@ -2570,7 +2570,7 @@ int32 FRewindData::GetLatestReceivedTargetStateFrame(Chaos::FGeometryParticleHan
 			return FrameAndPhase.Frame;
 		}
 	}
-	return LatestFrame;
+	return INDEX_NONE;
 }
 
 void FRewindData::RequestResimulation(int32 Frame, Chaos::FGeometryParticleHandle* Particle)
