[NetPhysics] Move custom state recording to post-solve and record it as for the beginning of next frame.
This allows the game thread to then network the state for the beginning of next frame before that frame has started.
We will receive states 1 frame earlier than we do now, effectively reducing the number of frames inside resimulations by 1.

#Synced-CL 48900574

diff --git a/Engine/Source/Runtime/Engine/Private/PhysicsEngine/NetworkPhysicsComponent.cpp b/Engine/Source/Runtime/Engine/Private/PhysicsEngine/NetworkPhysicsComponent.cpp
--- a/Engine/Source/Runtime/Engine/Private/PhysicsEngine/NetworkPhysicsComponent.cpp
+++ b/Engine/Source/Runtime/Engine/Private/PhysicsEngine/NetworkPhysicsComponent.cpp
@@ -134,6 +134,9 @@ namespace PhysicsReplicationCVars
 		bool bApplyPredictiveInterpolationWhenBehindServer = true;
 		static FAutoConsoleVariableRef CVarResimApplyPredictiveInterpolationWhenBehindServer(TEXT("np2.Resim.ApplyPredictiveInterpolationWhenBehindServer"), bApplyPredictiveInterpolationWhenBehindServer, TEXT("When true, switch over to replicating with Predictive Interpolation temporarily, when the client receive target states from the server for frames that have not yet simulated on the client. When false apply the received target via a resimulation when the client has caught up and simulated the corresponding frame."));
 
+		bool bRecordStatePostSolve = false;
+		static FAutoConsoleVariableRef CVarResimRecordStatePostSolve(TEXT("np2.Resim.RecordStatePostSolve"), bRecordStatePostSolve, TEXT("When true, cache custom state in PostSolve_Internal (and mark it for current frame + 1) instead of between ProcessInputs_Internal and OnPreSimulate_Internal. This makes clients receive states 1 frame earlier, reducing number of frames needed to resimulate."));
+
 		int32 bDebugTriggerResimEveryNFrames = 0;
 		static FAutoConsoleVariableRef CVarDebugTriggerResimEveryNFrames(TEXT("np2.Resim.Debug.TriggerResimEveryNFrames"), bDebugTriggerResimEveryNFrames, TEXT("When above 0, trigger a resim at an interval of the set value, as long as we receive a state or input on that frame."));
 	}
@@ -2428,46 +2431,49 @@ void FAsyncNetworkPhysicsComponent::OnPostProcessInputs_Internal(const int32 Phy
 #endif
 	}
 
-	// Cache current state if this is the server or we are comparing predicted states clients
-	const bool bShouldCacheStateHistory = bIsServer
-		|| (ComponentSettings.GetCompareStateToTriggerRewind(bCompareStateToTriggerRewind) && (IsLocallyControlled() || ComponentSettings.GetCompareStateToTriggerRewindIncludeSimProxies(bCompareStateToTriggerRewindIncludeSimProxies)));
-
-	if (StateHistory && StateData && bShouldCacheStateHistory)
+	if (PhysicsReplicationCVars::ResimulationCVars::bRecordStatePostSolve == false)
 	{
-		// Prepare to gather state data
-		FNetworkPhysicsPayload* PhysicsData = StateData.Get();
-		PhysicsData->PrepareFrame(PhysicsStep, bIsServer, GetNetworkPhysicsTickOffset());
+		// Cache current state if this is the server or we are comparing predicted states clients
+		const bool bShouldCacheStateHistory = bIsServer
+			|| (ComponentSettings.GetCompareStateToTriggerRewind(bCompareStateToTriggerRewind) && (IsLocallyControlled() || ComponentSettings.GetCompareStateToTriggerRewindIncludeSimProxies(bCompareStateToTriggerRewindIncludeSimProxies)));
 
-		if (bIsUsingLegacyData)
+		if (StateHistory && StateData && bShouldCacheStateHistory)
 		{
-			if (ensure(ActorComponent.IsValid() && ActorComponent.Get()->IsBeingDestroyed() == false))
+			// Prepare to gather state data
+			FNetworkPhysicsPayload* PhysicsData = StateData.Get();
+			PhysicsData->PrepareFrame(PhysicsStep, bIsServer, GetNetworkPhysicsTickOffset());
+
+			if (bIsUsingLegacyData)
 			{
-				// Gather state data from implementation
-				FNetworkPhysicsData* LegacyData = static_cast<FNetworkPhysicsData*>(PhysicsData);
-				LegacyData->BuildData(ActorComponent.Get());
+				if (ensure(ActorComponent.IsValid() && ActorComponent.Get()->IsBeingDestroyed() == false))
+				{
+					// Gather state data from implementation
+					FNetworkPhysicsData* LegacyData = static_cast<FNetworkPhysicsData*>(PhysicsData);
+					LegacyData->BuildData(ActorComponent.Get());
+				}
 			}
-		}
-		else
-		{
-			if (ensure(ImplementationInterface_Internal))
+			else
 			{
-				// Gather state data from implementation
-				ImplementationInterface_Internal->BuildState(*PhysicsData);
+				if (ensure(ImplementationInterface_Internal))
+				{
+					// Gather state data from implementation
+					ImplementationInterface_Internal->BuildState(*PhysicsData);
+				}
 			}
-		}
 
-		// Record state in history
-		StateHistory->RecordData(PhysicsStep, PhysicsData);
+			// Record state in history
+			StateHistory->RecordData(PhysicsStep, PhysicsData);
 
-#if DEBUG_NETWORK_PHYSICS
-		{
-			UE_LOG(LogChaos, Log, TEXT("		Recording state into history | LocalFrame = %d | ServerFrame = %d | bDataAltered = %d | State: %s ")
-				, PhysicsData->LocalFrame, PhysicsData->ServerFrame, PhysicsData->bDataAltered, *PhysicsData->DebugData());
+	#if DEBUG_NETWORK_PHYSICS
+			{
+				UE_LOG(LogChaos, Log, TEXT("		Recording state into history | LocalFrame = %d | ServerFrame = %d | bDataAltered = %d | State: %s ")
+					, PhysicsData->LocalFrame, PhysicsData->ServerFrame, PhysicsData->bDataAltered, *PhysicsData->DebugData());
+			}
+	#endif
 		}
-#endif
 	}
 
-	if (bIsSolverResim == false)
+	if (bIsSolverResim == false && PhysicsReplicationCVars::ResimulationCVars::bRecordStatePostSolve == false)
 	{
 		// Marshal inputs and states from PT to GT for networking
 		FAsyncNetworkPhysicsComponentOutput& AsyncOutput = GetAsyncOutput_Internal();
@@ -2555,6 +2561,72 @@ void FAsyncNetworkPhysicsComponent::OnPostIntegrate_Internal()
 	}
 }
 
+void FAsyncNetworkPhysicsComponent::OnPostSolve_Internal()
+{
+	Chaos::FPBDRigidsSolver* RigidsSolver = GetRigidSolver();
+	if (!RigidsSolver || !RigidsSolver->GetEvolution())
+	{
+		return;
+	}
+	const bool bIsSolverResim = RigidsSolver->GetEvolution()->IsResimming();
+		const int32 PhysicsStep = RigidsSolver->GetCurrentFrame();
+	const int32 NextPhysicsStep = PhysicsStep + 1;
+
+	const FNetworkPhysicsSettingsNetworkPhysicsComponent& ComponentSettings = GetComponentSettings();
+	const bool bIsServer = IsServer();
+
+	if (PhysicsReplicationCVars::ResimulationCVars::bRecordStatePostSolve == true)
+	{
+		// Cache current state if this is the server or we are comparing predicted states clients
+		const bool bShouldCacheStateHistory = bIsServer
+			|| (ComponentSettings.GetCompareStateToTriggerRewind(bCompareStateToTriggerRewind) && (IsLocallyControlled() || ComponentSettings.GetCompareStateToTriggerRewindIncludeSimProxies(bCompareStateToTriggerRewindIncludeSimProxies)));
+
+		if (StateHistory && StateData && bShouldCacheStateHistory)
+		{
+			// Prepare to gather state data
+			FNetworkPhysicsPayload* PhysicsData = StateData.Get();
+			PhysicsData->PrepareFrame(NextPhysicsStep, bIsServer, GetNetworkPhysicsTickOffset());
+
+			if (bIsUsingLegacyData)
+			{
+				if (ensure(ActorComponent.IsValid() && ActorComponent.Get()->IsBeingDestroyed() == false))
+				{
+					// Gather state data from implementation
+					FNetworkPhysicsData* LegacyData = static_cast<FNetworkPhysicsData*>(PhysicsData);
+					LegacyData->BuildData(ActorComponent.Get());
+				}
+			}
+			else
+			{
+				if (ensure(ImplementationInterface_Internal))
+				{
+					// Gather state data from implementation
+					ImplementationInterface_Internal->BuildState(*PhysicsData);
+				}
+			}
+
+			// Record state in history
+			StateHistory->RecordData(NextPhysicsStep, PhysicsData);
+
+	#if DEBUG_NETWORK_PHYSICS
+			{
+				UE_LOG(LogChaos, Log, TEXT("		Recording state into history | LocalFrame = %d | ServerFrame = %d | bDataAltered = %d | State: %s ")
+					, PhysicsData->LocalFrame, PhysicsData->ServerFrame, PhysicsData->bDataAltered, *PhysicsData->DebugData());
+			}
+	#endif
+		}
+	}
+
+	if (bIsSolverResim == false && PhysicsReplicationCVars::ResimulationCVars::bRecordStatePostSolve == true)
+	{
+		// Marshal inputs and states from PT to GT for networking
+		FAsyncNetworkPhysicsComponentOutput& AsyncOutput = GetAsyncOutput_Internal();
+		SendInputData_Internal(AsyncOutput, PhysicsStep);
+		SendStateData_Internal(AsyncOutput, NextPhysicsStep); // We cache state at the end of the solve which is the state that next frame starts with
+		FinalizeOutputData_Internal();
+	}
+}
+
 void FAsyncNetworkPhysicsComponent::SendInputData_Internal(FAsyncNetworkPhysicsComponentOutput& AsyncOutput, const int32 PhysicsStep)
 {
 	const bool bIsServer = IsServer();
@@ -2810,17 +2882,12 @@ const float FAsyncNetworkPhysicsComponent::GetInputLagFromLastReceivedTargetStat
 	const Chaos::FPhysicsSolverBase* BaseSolver = GetSolver();
 	if (const Chaos::FRewindData* RewindData = BaseSolver ? BaseSolver->GetRewindData() : nullptr)
 	{
-		// Todo: we should not have to get a read lock to get this information bu instead measure the latest lag upon reception
-		Chaos::FReadPhysicsObjectInterface_Internal Interface = Chaos::FPhysicsObjectInternalInterface::GetRead();
-		if (Chaos::FGeometryParticleHandle* Particle = Interface.GetParticle(PhysicsObject))
+		const int32 LatestTargetFrame = GetLastReceivedServerFrame();
+		if (LatestTargetFrame >= 0)
 		{
-			const int32 LatestTargetFrame = GetLastReceivedServerFrame();
-			if (LatestTargetFrame >= 0)
-			{
-				const int32 MeasuredInputDelayFrames = RewindData->CurrentFrame() - LatestTargetFrame;
-				const float MeasuredInputDelaySeconds = FMath::Max(MeasuredInputDelayFrames, 0.0f) * BaseSolver->GetAsyncDeltaTime();
-				return MeasuredInputDelaySeconds;
-			}
+			const int32 MeasuredInputDelayFrames = RewindData->CurrentFrame() - LatestTargetFrame;
+			const float MeasuredInputDelaySeconds = FMath::Max(MeasuredInputDelayFrames, 0.0f) * BaseSolver->GetAsyncDeltaTime();
+			return MeasuredInputDelaySeconds;
 		}
 	}
 	return 0.0f;
diff --git a/Engine/Source/Runtime/Engine/Public/Physics/NetworkPhysicsComponent.h b/Engine/Source/Runtime/Engine/Public/Physics/NetworkPhysicsComponent.h
--- a/Engine/Source/Runtime/Engine/Public/Physics/NetworkPhysicsComponent.h
+++ b/Engine/Source/Runtime/Engine/Public/Physics/NetworkPhysicsComponent.h
@@ -1896,7 +1896,7 @@ struct FAsyncNetworkPhysicsComponentOutput : public Chaos::FSimCallbackOutput
 class FAsyncNetworkPhysicsComponent : public Chaos::TSimCallbackObject<
 	FAsyncNetworkPhysicsComponentInput,
 	FAsyncNetworkPhysicsComponentOutput,
-	Chaos::ESimCallbackOptions::Presimulate |Chaos::ESimCallbackOptions::PostIntegrate | Chaos::ESimCallbackOptions::PhysicsObjectUnregister>
+	Chaos::ESimCallbackOptions::Presimulate | Chaos::ESimCallbackOptions::PostIntegrate | Chaos::ESimCallbackOptions::PostSolve | Chaos::ESimCallbackOptions::PhysicsObjectUnregister>
 {
 	friend UNetworkPhysicsComponent;
 
@@ -1948,6 +1948,8 @@ private:
 
 	virtual void OnPostIntegrate_Internal() override;
 
+	virtual void OnPostSolve_Internal() override;
+
 	virtual void OnPhysicsObjectUnregistered_Internal(Chaos::FConstPhysicsObjectHandle PhysicsObject) override;
 
 	// Delegate ran at the start of FNetworkPhysicsCallback::ProcessInputs_Internal, used to receive and apply inputs and states if needed.
