[NetPhysics] Enable custom state caching post-solve by default in NetworkPhysicsComponent.

#Synced-CL 48901078

diff --git a/Engine/Source/Runtime/Engine/Private/PhysicsEngine/NetworkPhysicsComponent.cpp b/Engine/Source/Runtime/Engine/Private/PhysicsEngine/NetworkPhysicsComponent.cpp
--- a/Engine/Source/Runtime/Engine/Private/PhysicsEngine/NetworkPhysicsComponent.cpp
+++ b/Engine/Source/Runtime/Engine/Private/PhysicsEngine/NetworkPhysicsComponent.cpp
@@ -134,7 +134,7 @@ namespace PhysicsReplicationCVars
 		bool bApplyPredictiveInterpolationWhenBehindServer = true;
 		static FAutoConsoleVariableRef CVarResimApplyPredictiveInterpolationWhenBehindServer(TEXT("np2.Resim.ApplyPredictiveInterpolationWhenBehindServer"), bApplyPredictiveInterpolationWhenBehindServer, TEXT("When true, switch over to replicating with Predictive Interpolation temporarily, when the client receive target states from the server for frames that have not yet simulated on the client. When false apply the received target via a resimulation when the client has caught up and simulated the corresponding frame."));
 
-		bool bRecordStatePostSolve = false;
+		bool bRecordStatePostSolve = true;
 		static FAutoConsoleVariableRef CVarResimRecordStatePostSolve(TEXT("np2.Resim.RecordStatePostSolve"), bRecordStatePostSolve, TEXT("When true, cache custom state in PostSolve_Internal (and mark it for current frame + 1) instead of between ProcessInputs_Internal and OnPreSimulate_Internal. This makes clients receive states 1 frame earlier, reducing number of frames needed to resimulate."));
 
 		int32 bDebugTriggerResimEveryNFrames = 0;
