Re-submit after backout.

Note:
Change can be disabled using r.Material.StripUnusedDefaultTextures.
Cook performance issues have been addressed by re-using existing Cached Expression Data for initial textures.
Shader issues have been solved by the same fix as above.
Cooked Data Commandlet issues have been fixed by handling editor Targets which do not render and thus have no material compilation output.

Save runtime Texture memory usage by stripping unused Default Textures at Cook time.
The dependency to the array which contains all the default textures is removed and encapsulated behind getters.
The original ReferencedTextures is now renamed EditorReferencedDefaultTextures and moved to Editor Only Data and a new member is added to the Material Interface which is modified to hold only the textures which are used by that Material, at runtime.

Material Instances without static parameters have an empty array and get their textures from their parents.
Everything else is built to support this by leveraging encapsulation.

StripUnusedDefaultTextures() contains the functionality to remove the unused textures.
This is done at cook time, based on the Shader Compilation output for the platform for which the cook is being done.

GetReferencedTextures() in FMaterial is now deprecated to allow for that better encapsulation of the actual Textures array and specific functions were added for each specific request.

GetDefaultTextureIdx(), GetDefaultTextureAtIdx(), GetDefaultTextures() have been added to the UMaterialInterface for this.
The GetMat variants of these methods have been added to the FMaterial, which in the end call the UMaterialInterface methods.

GetDefaultTextures and GetMatDefaultTextures has been added as a variant which replicates old behavior by returning an array with all default textures. This is used only in editor behaviors,

Common properties like CachedMaterialResourcesForCooking have been moved from UMaterial and UMaterialInstance to UMaterialInterface.

Common functionalities between UMaterial and UMaterialInstance that could be refactored out to clean out the code  have been moved to UMaterialInterface:
GetMaterialResources, GetTextureExpressionValues, BuildQualityLevelRange, BuildShaderPlatformRange.

Move Material Instance default texture lookup to happen after overrides and only if an override was not found.

[RN] medium | Rendering Architecture

#Synced-CL 48863451

diff --git a/Engine/Source/Developer/MaterialBaking/Private/ExportMaterialProxy.h b/Engine/Source/Developer/MaterialBaking/Private/ExportMaterialProxy.h
--- a/Engine/Source/Developer/MaterialBaking/Private/ExportMaterialProxy.h
+++ b/Engine/Source/Developer/MaterialBaking/Private/ExportMaterialProxy.h
@@ -210,11 +210,11 @@ struct FExportMaterialCompiler : public FProxyMaterialCompiler
 	virtual EMaterialCompilerType GetCompilerType() const override { return EMaterialCompilerType::MaterialProxy; }
 };
 
-class FExportMaterialProxy : public FMaterial, public FMaterialRenderProxy
+class FExportMaterialProxy : public FEditorMaterialResourceBase, public FMaterialRenderProxy
 {
 public:
 	FExportMaterialProxy(UMaterialInterface* InMaterialInterface, EMaterialProperty InPropertyToCompile, const FString& InCustomOutputToCompile = TEXT(""), bool bInSynchronousCompilation = true, bool bTangentSpaceNormal = false, EBlendMode ProxyBlendMode = BLEND_Opaque, bool bAllowPixelDepthOffset = true)
-		: FMaterial()
+		: FEditorMaterialResourceBase()
 		, FMaterialRenderProxy(GetPathNameSafe(InMaterialInterface->GetMaterial()))
 		, MaterialInterface(InMaterialInterface)
 		, PropertyToCompile(InPropertyToCompile)
@@ -225,7 +225,7 @@ public:
 	{
 		SetQualityLevelProperties(GMaxRHIShaderPlatform);
 		Material = InMaterialInterface->GetMaterial();
-		ReferencedTextures = InMaterialInterface->GetReferencedTextures();
+		ReferencedTextures = InMaterialInterface->GetDefaultTextures();
 		ReferencedTextureCollections = InMaterialInterface->GetReferencedTextureCollections();
 
 		const FMaterialResource* Resource = InMaterialInterface->GetMaterialResource(GMaxRHIShaderPlatform);
@@ -300,16 +300,6 @@ public:
 		return bCorrectVertexFactory && bPCPlatform && bCorrectFrequency;
 	}
 
-	virtual TArrayView<const TObjectPtr<UObject>> GetReferencedTextures() const override
-	{
-		return ReferencedTextures;
-	}
-
-	virtual TConstArrayView<TObjectPtr<UTextureCollection>> GetReferencedTextureCollections() const override
-	{
-		return ReferencedTextureCollections;
-	}
-
 	virtual void GetStaticParameterSet(FStaticParameterSet& OutSet) const override
 	{
 		if (const FMaterialResource* Resource = MaterialInterface->GetMaterialResource(GMaxRHIShaderPlatform))
@@ -726,8 +716,6 @@ private:
 	/** The material interface for this proxy */
 	UMaterialInterface* MaterialInterface;
 	UMaterial* Material;
-	TArray<TObjectPtr<UObject>> ReferencedTextures;
-	TArray<TObjectPtr<UTextureCollection>> ReferencedTextureCollections;
 	/** The property to compile for rendering the sample */
 	EMaterialProperty PropertyToCompile;
 	/** Stores which exported attribute this proxy is compiling for. */
diff --git a/Engine/Source/Developer/MaterialUtilities/Private/MaterialUtilities.cpp b/Engine/Source/Developer/MaterialUtilities/Private/MaterialUtilities.cpp
--- a/Engine/Source/Developer/MaterialUtilities/Private/MaterialUtilities.cpp
+++ b/Engine/Source/Developer/MaterialUtilities/Private/MaterialUtilities.cpp
@@ -453,26 +453,27 @@ struct FExportMaterialCompiler : public FProxyMaterialCompiler
 	virtual EMaterialCompilerType GetCompilerType() const override { return EMaterialCompilerType::MaterialProxy; }
 };
 
-
-class FExportMaterialProxy : public FMaterial, public FMaterialRenderProxy
+class FExportMaterialProxy : public FEditorMaterialResourceBase, public FMaterialRenderProxy
 {
 public:
 	FExportMaterialProxy()
-		: FMaterial()
+		: FEditorMaterialResourceBase()
 		, FMaterialRenderProxy(TEXT("FExportMaterialProxy"))
 	{
 		SetQualityLevelProperties(GMaxRHIShaderPlatform);
 	}
 
 	FExportMaterialProxy(UMaterialInterface* InMaterialInterface, EMaterialProperty InPropertyToCompile)
-		: FMaterial()
+		: FEditorMaterialResourceBase()
 		, FMaterialRenderProxy(GetPathNameSafe(InMaterialInterface->GetMaterial()))
 		, MaterialInterface(InMaterialInterface)
 		, PropertyToCompile(InPropertyToCompile)
 	{
 		SetQualityLevelProperties(GMaxRHIShaderPlatform);
 		Material = InMaterialInterface->GetMaterial();
-		ReferencedTextures = InMaterialInterface->GetReferencedTextures();
+		ReferencedTextures = InMaterialInterface->GetDefaultTextures();
+		ReferencedTextureCollections = MaterialInterface->GetReferencedTextureCollections();
+
 		FPlatformMisc::CreateGuid(Id);
 
 		FMaterialResource* Resource = InMaterialInterface->GetMaterialResource(GMaxRHIShaderPlatform);
@@ -534,11 +535,6 @@ public:
 		return true;
 	}
 
-	virtual TArrayView<const TObjectPtr<UObject>> GetReferencedTextures() const override
-	{
-		return ReferencedTextures;
-	}
-
 	////////////////
 	// FMaterialRenderProxy interface.
 	virtual const FMaterial* GetMaterialNoFallback(ERHIFeatureLevel::Type InFeatureLevel) const override
@@ -796,7 +792,6 @@ private:
 	/** The material interface for this proxy */
 	UMaterialInterface* MaterialInterface;
 	UMaterial* Material;	
-	TArray<TObjectPtr<UObject>> ReferencedTextures;
 	/** The property to compile for rendering the sample */
 	EMaterialProperty PropertyToCompile;
 	FGuid Id;
diff --git a/Engine/Source/Editor/MaterialEditor/Private/MaterialEditor.cpp b/Engine/Source/Editor/MaterialEditor/Private/MaterialEditor.cpp
--- a/Engine/Source/Editor/MaterialEditor/Private/MaterialEditor.cpp
+++ b/Engine/Source/Editor/MaterialEditor/Private/MaterialEditor.cpp
@@ -248,7 +248,7 @@ static bool MarkSubstrateAutoConversionSaved(UMaterial* Material)
 ///////////////////////////
 
 FMatExpressionPreview::FMatExpressionPreview()
-	: FMaterial()
+	: FEditorMaterialResourceBase()
 	, FMaterialRenderProxy(TEXT("FMatExpressionPreview"))
 	, UnrelatedNodesOpacity(1.0f)
 {
@@ -258,7 +258,7 @@ FMatExpressionPreview::FMatExpressionPreview()
 }
 
 FMatExpressionPreview::FMatExpressionPreview(UMaterialExpression* InExpression)
-	: FMaterial()
+	: FEditorMaterialResourceBase()
 	, FMaterialRenderProxy(GetPathNameSafe(InExpression->Material))
 	, UnrelatedNodesOpacity(1.0f)
 	, Expression(InExpression)
@@ -270,7 +270,7 @@ FMatExpressionPreview::FMatExpressionPreview(UMaterialExpression* InExpression)
 	SetQualityLevelProperties(GMaxRHIShaderPlatform);
 
 	UMaterial* BaseMaterial = InExpression->Material;
-	ReferencedTextures = InExpression->Material->GetReferencedTextures();
+	ReferencedTextures = InExpression->Material->GetDefaultTextures();
 	ReferencedTextureCollections = InExpression->Material->GetReferencedTextureCollections();
 }
 
@@ -282,10 +282,6 @@ void FMatExpressionPreview::AddReferencedObjects(FReferenceCollector& Collector)
 {
 	Collector.AddReferencedObjects(ReferencedTextures);
 	Collector.AddReferencedObjects(ReferencedTextureCollections);
-	if (CachedExpressionData)
-	{
-		CachedExpressionData->AddReferencedObjects(Collector);
-	}
 }
 
 bool FMatExpressionPreview::ShouldCache(const FShaderType* ShaderType, const FVertexFactoryType* VertexFactoryType) const
@@ -421,30 +417,6 @@ bool FMatExpressionPreview::IsUsingNewHLSLGenerator() const
 	return false;
 }
 
-TArrayView<const TObjectPtr<UObject>> FMatExpressionPreview::GetReferencedTextures() const
-{
-	if (CachedExpressionData)
-	{
-		// Path for new HLSL translator
-		return MakeArrayView(CachedExpressionData->ReferencedTextures);
-	}
-
-	// Legacy path
-	return MakeArrayView(ReferencedTextures);
-}
-
-TConstArrayView<TObjectPtr<UTextureCollection>> FMatExpressionPreview::GetReferencedTextureCollections() const
-{
-	if (CachedExpressionData)
-	{
-		// Path for new HLSL translator
-		return MakeArrayView(CachedExpressionData->ReferencedTextureCollections);
-	}
-
-	// Legacy path
-	return MakeArrayView(ReferencedTextureCollections);
-}
-
 bool FMatExpressionPreview::CheckInValidStateForCompilation(FMaterialCompiler* Compiler) const
 {
 	return Expression.IsValid() && Expression->Material && Expression->Material->CheckInValidStateForCompilation(Compiler);
diff --git a/Engine/Source/Editor/MaterialEditor/Private/MaterialEditor.h b/Engine/Source/Editor/MaterialEditor/Private/MaterialEditor.h
--- a/Engine/Source/Editor/MaterialEditor/Private/MaterialEditor.h
+++ b/Engine/Source/Editor/MaterialEditor/Private/MaterialEditor.h
@@ -52,7 +52,7 @@ typedef TSet<class UObject*> FGraphPanelSelectionSet;
 /**
  * Class for rendering previews of material expressions in the material editor's linked object viewport.
  */
-class FMatExpressionPreview : public FMaterial, public FMaterialRenderProxy
+class FMatExpressionPreview : public FEditorMaterialResourceBase, public FMaterialRenderProxy
 {
 public:
 	FMatExpressionPreview();
@@ -84,10 +84,7 @@ public:
 	 * @return true if the shader should be compiled
 	 */
 	virtual bool ShouldCache(const FShaderType* ShaderType, const FVertexFactoryType* VertexFactoryType) const override;
-
-	virtual TArrayView<const TObjectPtr<UObject>> GetReferencedTextures() const override;
-	virtual TConstArrayView<TObjectPtr<UTextureCollection>> GetReferencedTextureCollections() const override;
-
+	
 	////////////////
 	// FMaterialRenderProxy interface.
 	virtual const FMaterial* GetMaterialNoFallback(ERHIFeatureLevel::Type InFeatureLevel) const override
@@ -178,10 +175,7 @@ public:
 	float UnrelatedNodesOpacity;
 
 private:
-	TUniquePtr<FMaterialCachedExpressionData> CachedExpressionData;
 	TWeakObjectPtr<UMaterialExpression> Expression;
-	TArray<TObjectPtr<UObject>> ReferencedTextures;
-	TArray<TObjectPtr<UTextureCollection>> ReferencedTextureCollections;
 	FGuid Id;
 };
 
diff --git a/Engine/Source/Editor/UnrealEd/Private/Lightmass/LightmassRender.cpp b/Engine/Source/Editor/UnrealEd/Private/Lightmass/LightmassRender.cpp
--- a/Engine/Source/Editor/UnrealEd/Private/Lightmass/LightmassRender.cpp
+++ b/Engine/Source/Editor/UnrealEd/Private/Lightmass/LightmassRender.cpp
@@ -241,10 +241,10 @@ struct FLightmassMaterialCompiler : public FProxyMaterialCompiler
 /**
  * Class for rendering previews of material expressions in the material editor's linked object viewport.
  */
-class FLightmassMaterialProxy : public FMaterial, public FMaterialRenderProxy
+class FLightmassMaterialProxy : public FEditorMaterialResourceBase, public FMaterialRenderProxy
 {
 public:
-	FLightmassMaterialProxy(): FMaterial(), FMaterialRenderProxy(TEXT("FLightmassMaterialProxy"))
+	FLightmassMaterialProxy(): FEditorMaterialResourceBase(), FMaterialRenderProxy(TEXT("FLightmassMaterialProxy"))
 	{
 		SetQualityLevelProperties(GMaxRHIShaderPlatform);
 	}
@@ -260,7 +260,7 @@ public:
 			PropertyToCompile = InPropertyToCompile;
 			Usage = InUsage;
 
-			ReferencedTextures = MaterialInterface->GetReferencedTextures();
+			ReferencedTextures = MaterialInterface->GetDefaultTextures();
 			ReferencedTextureCollections = MaterialInterface->GetReferencedTextureCollections();
 
 			FMaterialResource* Resource = InMaterialInterface->GetMaterialResource(GMaxRHIShaderPlatform);
@@ -287,16 +287,6 @@ public:
 		}
 	}
 
-	virtual TArrayView<const TObjectPtr<UObject>> GetReferencedTextures() const override
-	{
-		return ReferencedTextures;
-	}
-
-	virtual TConstArrayView<TObjectPtr<UTextureCollection>> GetReferencedTextureCollections() const override
-	{
-		return ReferencedTextureCollections;
-	}
-
 	/**
 	 * Should the shader for this material with the given platform, shader type and vertex 
 	 * factory type combination be compiled
diff --git a/Engine/Source/Runtime/Engine/Private/Materials/HLSLMaterialTranslator.cpp b/Engine/Source/Runtime/Engine/Private/Materials/HLSLMaterialTranslator.cpp
--- a/Engine/Source/Runtime/Engine/Private/Materials/HLSLMaterialTranslator.cpp
+++ b/Engine/Source/Runtime/Engine/Private/Materials/HLSLMaterialTranslator.cpp
@@ -7952,7 +7952,7 @@ int32 FHLSLMaterialTranslator::TextureSample(
 		const FMaterialUniformExpressionTexture* TextureExpression = Expression ? Expression->GetTextureUniformExpression() : nullptr;
 		if (ensure(TextureExpression))
 		{
-			Texture = Cast<UTexture>(Material->GetReferencedTextures()[TextureExpression->GetTextureIndex()]);
+			Texture = Cast<UTexture>(Material->GetMatDefaultTextureAtIdx(TextureExpression->GetTextureIndex()));
 		}
 	}
 
@@ -9229,13 +9229,13 @@ int32 FHLSLMaterialTranslator::Switch(int32 SwitchValueInput, int32 DefaultInput
 int32 FHLSLMaterialTranslator::Texture(UTexture* InTexture, int32& TextureReferenceIndex, EMaterialSamplerType SamplerType, ESamplerSourceMode SamplerSource, ETextureMipValueMode MipValueMode)
 {
 	EMaterialValueType ShaderType = InTexture->GetMaterialType();
-	TextureReferenceIndex = Material->GetReferencedTextures().Find(InTexture);
+	TextureReferenceIndex = Material->GetMatDefaultTextureIdx(InTexture);
 
 #if DO_CHECK
 	// UE-3518: Additional pre-assert logging to help determine the cause of this failure.
 	if (TextureReferenceIndex == INDEX_NONE)
 	{
-		const auto ReferencedTextures = Material->GetReferencedTextures();
+		const TArray<TObjectPtr<UObject>>& ReferencedTextures = Material->GetMatDefaultTextures();
 		UE_LOG(LogMaterial, Error, TEXT("Compiler->Texture() failed to find texture '%s' in referenced list of size '%i':"), *InTexture->GetName(), ReferencedTextures.Num());
 		for (int32 i = 0; i < ReferencedTextures.Num(); ++i)
 		{
@@ -9268,7 +9268,7 @@ int32 FHLSLMaterialTranslator::TextureParameter(FName ParameterName, UTexture* I
 	}
 
 	EMaterialValueType ShaderType = DefaultValue->GetMaterialType();
-	TextureReferenceIndex = Material->GetReferencedTextures().Find(DefaultValue);
+	TextureReferenceIndex = Material->GetMatDefaultTextureIdx(DefaultValue);
 	if (TextureReferenceIndex == INDEX_NONE)
 	{
 		FString TextureName;
@@ -9379,7 +9379,7 @@ int32 FHLSLMaterialTranslator::VirtualTexture(URuntimeVirtualTexture* InTexture,
 		return INDEX_NONE;
 	}
 
-	TextureReferenceIndex = Material->GetReferencedTextures().Find(InTexture);
+	TextureReferenceIndex = Material->GetMatDefaultTextureIdx(InTexture);
 	checkf(TextureReferenceIndex != INDEX_NONE, TEXT("Material expression called Compiler->VirtualTexture() without implementing UMaterialExpression::GetReferencedTexture properly"));
 
 	FAddUniformExpressionScope Scope(this);
@@ -9402,7 +9402,7 @@ int32 FHLSLMaterialTranslator::VirtualTextureParameter(FName ParameterName, URun
 		DefaultValue = Meta.Value.RuntimeVirtualTexture;
 	}
 
-	TextureReferenceIndex = Material->GetReferencedTextures().Find(DefaultValue);
+	TextureReferenceIndex = Material->GetMatDefaultTextureIdx(DefaultValue);
 	checkf(TextureReferenceIndex != INDEX_NONE, TEXT("Material expression called Compiler->VirtualTexture() without implementing UMaterialExpression::GetReferencedTexture properly"));
 
 	FMaterialParameterInfo ParameterInfo = GetParameterAssociationInfo();
@@ -9547,7 +9547,7 @@ int32 FHLSLMaterialTranslator::ExternalTexture(UTexture* InTexture, int32& Textu
 		return NonPixelShaderExpressionError();
 	}
 
-	TextureReferenceIndex = Material->GetReferencedTextures().Find(InTexture);
+	TextureReferenceIndex = Material->GetMatDefaultTextureIdx(InTexture);
 	checkf(TextureReferenceIndex != INDEX_NONE, TEXT("Material expression called Compiler->ExternalTexture() without implementing UMaterialExpression::GetReferencedTexture properly"));
 
 	FAddUniformExpressionScope Scope(this);
@@ -9563,7 +9563,7 @@ int32 FHLSLMaterialTranslator::ExternalTextureParameter(FName ParameterName, UTe
 		return NonPixelShaderExpressionError();
 	}
 
-	TextureReferenceIndex = Material->GetReferencedTextures().Find(DefaultValue);
+	TextureReferenceIndex = Material->GetMatDefaultTextureIdx(DefaultValue);
 	checkf(TextureReferenceIndex != INDEX_NONE, TEXT("Material expression called Compiler->ExternalTextureParameter() without implementing UMaterialExpression::GetReferencedTexture properly"));
 	FAddUniformExpressionScope Scope(this);
 	return AddUniformExpression(Scope, new FMaterialUniformExpressionExternalTextureParameter(ParameterName, TextureReferenceIndex), MCT_TextureExternal, TEXT(""));
@@ -9592,7 +9592,7 @@ int32 FHLSLMaterialTranslator::ExternalTextureCoordinateOffset(const FGuid& Exte
 
 UObject* FHLSLMaterialTranslator::GetReferencedTexture(int32 Index)
 {
-	return Material->GetReferencedTextures()[Index];
+	return Material->GetMatDefaultTextureAtIdx(Index);
 }
 
 UTextureCollection* FHLSLMaterialTranslator::GetReferencedTextureCollection(int32 Index)
@@ -17971,7 +17971,7 @@ void FHLSLMaterialTranslator::PrepareMaterialSourceStringParameters()
 
 int32 FHLSLMaterialTranslator::SparseVolumeTexture(USparseVolumeTexture* Texture, int32& TextureReferenceIndex, EMaterialSamplerType SamplerType)
 {
-	TextureReferenceIndex = Material->GetReferencedTextures().Find(Texture);
+	TextureReferenceIndex = Material->GetMatDefaultTextureIdx(Texture);
 	checkf(TextureReferenceIndex != INDEX_NONE, TEXT("Material expression called Compiler->SparseVolumeTexture() without implementing UMaterialExpression::GetReferencedTexture properly"));
 
 	FAddUniformExpressionScope Scope(this);
@@ -17989,7 +17989,7 @@ int32 FHLSLMaterialTranslator::SparseVolumeTextureParameter(FName ParameterName,
 		DefaultTexture = Meta.Value.SparseVolumeTexture;
 	}
 
-	TextureReferenceIndex = Material->GetReferencedTextures().Find(DefaultTexture);
+	TextureReferenceIndex = Material->GetMatDefaultTextureIdx(DefaultTexture);
 	checkf(TextureReferenceIndex != INDEX_NONE, TEXT("Material expression called Compiler->SparseVolumeTextureParameter() without implementing UMaterialExpression::GetReferencedTexture properly"));
 
 	FMaterialParameterInfo ParameterInfo = GetParameterAssociationInfo();
@@ -18057,7 +18057,7 @@ int32 FHLSLMaterialTranslator::SparseVolumeTextureSamplePageTable(int32 SparseVo
 	int32 SVTReferenceIndex = UniformTextureExpressions[(uint32)EMaterialTextureParameterType::SparseVolume].Find(TextureUniformExpression);
 	check(UniformTextureExpressions[(uint32)EMaterialTextureParameterType::SparseVolume].IsValidIndex(SVTReferenceIndex));
 
-	const USparseVolumeTexture* SVTexture = Cast<USparseVolumeTexture>(Material->GetReferencedTextures()[TextureUniformExpression->GetTextureIndex()]);
+	const USparseVolumeTexture* SVTexture = Cast<USparseVolumeTexture>(Material->GetMatDefaultTextureAtIdx(TextureUniformExpression->GetTextureIndex()));
 
 	// StaticAddress mode at time of compile
 	// Similar to VirtualTextures, this may not be 100% correct, if SamplerSource is set to 'SSM_FromTextureAsset', as texture parameter may change the address mode in a derived instance
diff --git a/Engine/Source/Runtime/Engine/Private/Materials/Material.cpp b/Engine/Source/Runtime/Engine/Private/Materials/Material.cpp
--- a/Engine/Source/Runtime/Engine/Private/Materials/Material.cpp
+++ b/Engine/Source/Runtime/Engine/Private/Materials/Material.cpp
@@ -1303,113 +1303,13 @@ void UMaterial::GetUsedTextures(TArray<UTexture*>& OutTextures, TOptional<EMater
 
 	if (!FPlatformProperties::IsServerOnly())
 	{
-		FInt32Range QualityLevelRange(0, EMaterialQualityLevel::Num - 1);
-		if (InSpecificQualityLevel.IsSet())
-		{
-			EMaterialQualityLevel::Type QualityLevel = InSpecificQualityLevel.GetValue();
-			ensure(QualityLevel != EMaterialQualityLevel::Num);
-			if (QualityLevel == EMaterialQualityLevel::Num)
-			{
-				QualityLevel = GetCurrentMaterialQualityLevelChecked();
-			}
-			QualityLevelRange = FInt32Range(QualityLevel, QualityLevel);
-		}
-
-#if WITH_EDITOR
-		FInt32Range ShaderPlatformRange(0, EShaderPlatform::SP_NumPlatforms - 1);
-		if (InSpecificShaderPlatform.IsSet())
-		{
-			EShaderPlatform ShaderPlatform = InSpecificShaderPlatform.GetValue();
-			ensure(ShaderPlatform != EShaderPlatform::SP_NumPlatforms);
-			if (ShaderPlatform == EShaderPlatform::SP_NumPlatforms)
-			{
-				ShaderPlatform = GetMaxShaderPlatformChecked();
-			}
-			ShaderPlatformRange = FInt32Range(ShaderPlatform, ShaderPlatform);
-		}
-#else
-		FInt32Range ShaderPlatformRange(GMaxRHIShaderPlatform, GMaxRHIShaderPlatform);
-#endif
-
-		TArray<const FMaterialResource*, TInlineAllocator<4>> MatchedResources;
-		// Parse all relevant quality and feature levels.
-		for (int32 QualityLevelIndex = QualityLevelRange.GetLowerBoundValue(); QualityLevelIndex <= QualityLevelRange.GetUpperBoundValue(); ++QualityLevelIndex)
-		{
-			for (int32 ShaderPlatformIndex = ShaderPlatformRange.GetLowerBoundValue(); ShaderPlatformIndex <= ShaderPlatformRange.GetUpperBoundValue(); ++ShaderPlatformIndex)
-			{
-				if (!FDataDrivenShaderPlatformInfo::IsValid((EShaderPlatform)ShaderPlatformIndex))
-				{
-					continue;
-				}
-
-				const FMaterialResource* CurrentResource = FindMaterialResource(MaterialResources, (EShaderPlatform)ShaderPlatformIndex, (EMaterialQualityLevel::Type)QualityLevelIndex, true);
-				if (CurrentResource)
-				{
-					MatchedResources.AddUnique(CurrentResource);
-				}
-			}
-		}
+		const FInt32Range QualityLevelRange = BuildQualityLevelRange(InSpecificQualityLevel);
+		const FInt32Range ShaderPlatformRange = BuildShaderPlatformRange(InSpecificShaderPlatform);
 
+		TArray<const FMaterialResource*, TInlineAllocator<4>> MatchedResources = GetMaterialResources(this, QualityLevelRange, ShaderPlatformRange);
 		for (const FMaterialResource* CurrentResource : MatchedResources)
 		{
-			for (int32 TypeIndex = 0; TypeIndex < NumMaterialTextureParameterTypes; TypeIndex++)
-			{
-				// Iterate over each of the material's texture expressions.
-				for (const FMaterialTextureParameterInfo& Parameter : CurrentResource->GetUniformTextureExpressions((EMaterialTextureParameterType)TypeIndex))
-				{
-					UTexture* Texture = NULL;
-					Parameter.GetGameThreadTextureValue(this, *CurrentResource, Texture);
-					if (Texture)
-					{
-						OutTextures.AddUnique(Texture);
-					}
-				}
-			}
-
-			for (const FMaterialTextureCollectionParameterInfo& Parameter : CurrentResource->GetUniformTextureCollectionExpressions())
-			{
-				UTextureCollection* TextureCollection = nullptr;
-				Parameter.GetGameThreadTextureCollectionValue(this, *CurrentResource, TextureCollection);
-				if (TextureCollection)
-				{
-					for (UTexture* Texture : TextureCollection->Textures)
-					{
-						if (Texture)
-						{
-							OutTextures.AddUnique(Texture);
-						}
-					}
-				}
-			}
-
-#if WITH_EDITOR
-/*
-			// Also look for any scalar parameters that are acting as lookups for an atlas texture, and store the atlas texture
-			const TArrayView<const FMaterialNumericParameterInfo> AtlasExpressions[1] =
-			{
-				CurrentResource->GetUniformNumericParameterExpressions()
-			};
-
-			for (int32 TypeIndex = 0; TypeIndex < UE_ARRAY_COUNT(AtlasExpressions); TypeIndex++)
-			{
-				// Iterate over each of the material's texture expressions.
-				for (const FMaterialNumericParameterInfo& Parameter : AtlasExpressions[TypeIndex])
-				{
-					if (Parameter.ParameterType == EMaterialParameterType::Scalar)
-					{
-						bool bIsUsedAsAtlasPosition;
-						TSoftObjectPtr<UCurveLinearColor> Curve;
-						TSoftObjectPtr<UCurveLinearColorAtlas> Atlas;
-						IsScalarParameterUsedAsAtlasPosition(Parameter.ParameterInfo, bIsUsedAsAtlasPosition, Curve, Atlas);
-						if (Atlas)
-						{
-							OutTextures.AddUnique(Atlas.Get());
-						}
-					}
-				}
-			}
-*/
-#endif // WITH_EDITOR
+			GetTextureExpressionValues(CurrentResource, OutTextures, nullptr, true, false);
 		}
 	}
 }
@@ -1426,36 +1326,7 @@ void UMaterial::GetUsedTexturesAndIndices(TArray<UTexture*>& OutTextures, TArray
 		const FMaterialResource* CurrentResource = FindMaterialResource(MaterialResources, InShaderPlatform, QualityLevel, true);
 		if (CurrentResource)
 		{
-			TArrayView<const FMaterialTextureParameterInfo> ExpressionsByType[NumMaterialTextureParameterTypes];
-			uint32 NumTextures = 0u;
-			for (uint32 TypeIndex = 0u; TypeIndex < NumMaterialTextureParameterTypes; ++TypeIndex)
-			{
-				ExpressionsByType[TypeIndex] = CurrentResource->GetUniformTextureExpressions((EMaterialTextureParameterType)TypeIndex);
-				NumTextures += ExpressionsByType[TypeIndex].Num();
-			}
-
-			// Try to prevent resizing since this would be expensive.
-			OutIndices.Empty(NumTextures);
-
-			for (int32 TypeIndex = 0; TypeIndex < UE_ARRAY_COUNT(ExpressionsByType); TypeIndex++)
-			{
-				// Iterate over each of the material's texture expressions.
-				for (const FMaterialTextureParameterInfo& Parameter : ExpressionsByType[TypeIndex])
-				{
-					UTexture* Texture = NULL;
-					Parameter.GetGameThreadTextureValue(this, *CurrentResource, Texture);
-
-					if (Texture)
-					{
-						int32 InsertIndex = OutTextures.AddUnique(Texture);
-						if (InsertIndex >= OutIndices.Num())
-						{
-							OutIndices.AddDefaulted(InsertIndex - OutIndices.Num() + 1);
-						}
-						OutIndices[InsertIndex].Add(Parameter.TextureIndex);
-					}
-				}
-			}
+			GetTextureExpressionValues(CurrentResource, OutTextures, &OutIndices, false, false);
 		}
 	}
 }
@@ -2393,8 +2264,11 @@ void UMaterial::UpdateCachedExpressionData()
 
 	CachedExpressionData.Reset(LocalCachedExpressionData);
 	EditorOnlyData->CachedExpressionData = LocalCachedExpressionData->EditorOnlyData;
+	ReferencedDefaultTextures = CachedExpressionData->EditorOnlyData->EditorReferencedDefaultTextures;
 
+#if WITH_EDITOR
 	FObjectCacheEventSink::NotifyReferencedTextureChanged_Concurrent(this);
+#endif
 }
 #endif // WITH_EDITOR
 
@@ -3074,6 +2948,14 @@ void UMaterial::Serialize(FArchive& Ar)
 	Ar.UsingCustomVersion(FUE5MainStreamObjectVersion::GUID);
 	Ar.UsingCustomVersion(FUE5ReleaseStreamObjectVersion::GUID);
 
+#if WITH_EDITOR
+	if (Ar.IsSaving() && Ar.IsCooking() && ShouldStripUnusedDefaultTextures())
+	{
+		constexpr bool bKeepProjectPrivateTextures = true;
+		StripUnusedDefaultTextures(bKeepProjectPrivateTextures);
+	}
+#endif
+
 	Super::Serialize(Ar);
 
 	if (Ar.UEVer() >= VER_UE4_PURGED_FMATERIAL_COMPILE_OUTPUTS)
@@ -3249,6 +3131,14 @@ void UMaterial::Serialize(FArchive& Ar)
 			SubsurfaceProfiles.Add(SubsurfaceProfile);
 		}
 	}
+
+#if WITH_EDITOR
+	if (Ar.IsSaving() && Ar.IsCooking() && ShouldStripUnusedDefaultTextures())
+	{
+		// Undo textures removal for Multi Platform Cook
+		ReferencedDefaultTextures = CachedExpressionData->EditorOnlyData->EditorReferencedDefaultTextures;
+	}
+#endif
 }
 
 void UMaterial::PostDuplicate(bool bDuplicateForPIE)
@@ -4773,7 +4663,7 @@ void UMaterial::PostLoad()
 		{
 			// Before caching shader resources we have to make sure all referenced textures have been post loaded
 			// as we depend on their resources being valid.
-			for (UObject* Texture : CachedExpressionData->ReferencedTextures)
+			for (UObject* Texture : ReferencedDefaultTextures)
 			{
 				if (Texture)
 				{
@@ -4853,6 +4743,22 @@ void UMaterial::DumpDebugInfo(FOutputDevice& OutputDevice) const
 #endif
 }
 
+TObjectPtr<UObject> UMaterial::GetDefaultTextureAtIdx(const int32 InIdx, TMicRecursionGuard RecursionGuard)
+{
+	if (!ReferencedDefaultTextures.IsValidIndex(InIdx))
+	{
+		static bool bWarnedOnce = false;
+		if (!bWarnedOnce)
+		{
+			UE_LOG(LogMaterial, Warning, TEXT("%s: Requesting an invalid TextureIndex! (%u / %u)"), *GetName(), InIdx, ReferencedDefaultTextures.Num());
+			bWarnedOnce = true;
+		}
+		return nullptr;
+	}
+
+	return ReferencedDefaultTextures[InIdx];
+}
+
 void UMaterial::SaveShaderStableKeys(const class ITargetPlatform* TP)
 {
 #if WITH_EDITOR
@@ -7290,6 +7196,32 @@ void UMaterial::ForceRecompileForRendering(EMaterialShaderPrecompileMode Compile
 	CacheResourceShadersForRendering(false,  CompileMode);
 }
 
+#if WITH_EDITOR
+bool UMaterial::TryGetUsedTexturesInCookCache(TArray<TObjectPtr<UObject>>& OutTextures)
+{
+	OutTextures.Empty();
+	const FMaterialCachedExpressionData& UsedCachedData = GetCachedExpressionData();
+	TArray<const FMaterialResource*, TInlineAllocator<4>> MatchedResources = GetCookedCacheMaterialResources();
+
+	for (const FMaterialResource* CurrentResource : MatchedResources)
+	{
+		const FUniformExpressionSet& CurrentSet = CurrentResource->GetUniformExpressions();
+
+		for (int32 TypeIdx = 0; TypeIdx < NumMaterialTextureParameterTypes; ++TypeIdx)
+		{
+			for (int32 TextureIdx = 0; TextureIdx < CurrentSet.GetNumTextures((EMaterialTextureParameterType)TypeIdx); ++TextureIdx)
+			{
+				const FMaterialTextureParameterInfo& Parameter = CurrentSet.GetTextureParameter((EMaterialTextureParameterType)TypeIdx, TextureIdx);
+				OutTextures.AddUnique(UsedCachedData.EditorOnlyData->EditorReferencedDefaultTextures[Parameter.TextureIndex]);
+			}
+		}
+	}
+
+	const bool bValidCompilationOutput = MatchedResources.Num() > 0;
+	return bValidCompilationOutput;
+}
+#endif
+
 bool UMaterial::CheckInValidStateForCompilation(class FMaterialCompiler* Compiler) const
 {
 	bool bSuccess = true;
diff --git a/Engine/Source/Runtime/Engine/Private/Materials/MaterialCachedData.cpp b/Engine/Source/Runtime/Engine/Private/Materials/MaterialCachedData.cpp
--- a/Engine/Source/Runtime/Engine/Private/Materials/MaterialCachedData.cpp
+++ b/Engine/Source/Runtime/Engine/Private/Materials/MaterialCachedData.cpp
@@ -109,7 +109,9 @@ FMaterialCachedExpressionData::FMaterialCachedExpressionData()
 
 void FMaterialCachedExpressionData::AddReferencedObjects(FReferenceCollector& Collector)
 {
-	Collector.AddStableReferenceArray(&ReferencedTextures);
+#if WITH_EDITORONLY_DATA
+ 	Collector.AddStableReferenceArray(&EditorOnlyData->EditorReferencedDefaultTextures);
+#endif
 	Collector.AddStableReferenceArray(&ReferencedTextureCollections);
 	Collector.AddStableReferenceArray(&GrassTypes);
 	Collector.AddStableReferenceArray(&MaterialCacheTags);
@@ -556,7 +558,7 @@ void FMaterialCachedExpressionData::UpdateForExpressions(const FMaterialCachedEx
 
 		if (ReferencedTexture)
 		{
-			ReferencedTextures.AddUnique(ReferencedTexture);
+			EditorOnlyData->EditorReferencedDefaultTextures.AddUnique(ReferencedTexture);
 		}
 		else if (ReferencedTextureCollection)
 		{
@@ -572,7 +574,7 @@ void FMaterialCachedExpressionData::UpdateForExpressions(const FMaterialCachedEx
 			const UMaterialExpression::ReferencedTextureArray ExpressionReferencedTextures = Expression->GetReferencedTextures();
 			for (UObject* ExpressionReferencedTexture : ExpressionReferencedTextures)
 			{
-				ReferencedTextures.AddUnique(ExpressionReferencedTexture);
+				EditorOnlyData->EditorReferencedDefaultTextures.AddUnique(ExpressionReferencedTexture);
 			}
 		}
 
@@ -1133,4 +1135,11 @@ void FMaterialCachedExpressionData::GetAllGlobalParameterInfoOfType(EMaterialPar
 	}
 }
 
+#if WITH_EDITOR
+int32 FMaterialCachedExpressionData::GetDefaultTextureIdx(UObject* InTexture) const
+{
+	return EditorOnlyData->EditorReferencedDefaultTextures.Find(TObjectPtr<UObject>(InTexture));
+}
+#endif
+
 #undef LOCTEXT_NAMESPACE
diff --git a/Engine/Source/Runtime/Engine/Private/Materials/MaterialIREmitter.cpp b/Engine/Source/Runtime/Engine/Private/Materials/MaterialIREmitter.cpp
--- a/Engine/Source/Runtime/Engine/Private/Materials/MaterialIREmitter.cpp
+++ b/Engine/Source/Runtime/Engine/Private/Materials/MaterialIREmitter.cpp
@@ -1260,7 +1260,7 @@ FValueRef FEmitter::PreshaderParameter(FType Type, EPreshader2Opcode Opcode, FVa
 	FPreshaderParameter Prototype = MakePrototype<FPreshaderParameter>(Type);
 	Prototype.SourceParameter = SourceParameter;
 	Prototype.Opcode = Opcode;
-	Prototype.TextureIndex = Material->GetReferencedTextures().Find(SourceParameterTexture);
+	Prototype.TextureIndex = Material->GetDefaultTextureIdx(SourceParameterTexture);
 	Prototype.Payload = MoveTemp(Payload);
 	return EmitPrototype(*this, Prototype);
 }
diff --git a/Engine/Source/Runtime/Engine/Private/Materials/MaterialIRValueAnalyzer.cpp b/Engine/Source/Runtime/Engine/Private/Materials/MaterialIRValueAnalyzer.cpp
--- a/Engine/Source/Runtime/Engine/Private/Materials/MaterialIRValueAnalyzer.cpp
+++ b/Engine/Source/Runtime/Engine/Private/Materials/MaterialIRValueAnalyzer.cpp
@@ -181,7 +181,7 @@ static int32 RegisterTextureParameter(FMaterialIRValueAnalyzer& Analyzer, UObjec
 	const FMaterialTextureParameterInfo TextureParameterInfo
 	{
 		.ParameterInfo = ParameterInfo,
-		.TextureIndex = Analyzer.Material->GetReferencedTextures().Find(Texture),
+		.TextureIndex = Analyzer.Material->GetDefaultTextureIdx(Texture),
 		.SamplerSource = SSM_FromTextureAsset,
 		.VirtualTextureLayerIndex = static_cast<uint8>(VTLayerIndex == INDEX_NONE ? UINT8_MAX : VTLayerIndex),
 	};
diff --git a/Engine/Source/Runtime/Engine/Private/Materials/MaterialInstance.cpp b/Engine/Source/Runtime/Engine/Private/Materials/MaterialInstance.cpp
--- a/Engine/Source/Runtime/Engine/Private/Materials/MaterialInstance.cpp
+++ b/Engine/Source/Runtime/Engine/Private/Materials/MaterialInstance.cpp
@@ -1186,11 +1186,6 @@ bool UMaterialInstance::GetParameterValue(EMaterialParameterType Type, const FMe
 	GetMaterialInheritanceChain(InstanceChain);
 
 	bool bResult = false;
-	if (EnumHasAnyFlags(Flags, EMaterialGetParameterValueFlags::CheckNonOverrides))
-	{
-		bResult = InstanceChain.GetCachedExpressionData().GetParameterValue(Type, ParameterInfo, OutResult);
-	}
-
 	const bool bCheckInstanceOverrides = EnumHasAnyFlags(Flags, EMaterialGetParameterValueFlags::CheckInstanceOverrides);
 	FMemoryImageMaterialParameterInfo CurrentParameterInfo = ParameterInfo;
 	bool bHasValidParameter = true;
@@ -1222,6 +1217,12 @@ bool UMaterialInstance::GetParameterValue(EMaterialParameterType Type, const FMe
 		ParentIndex++;
 	}
 
+	// Look through the default values if we have not found an override
+	if (!bResult && EnumHasAnyFlags(Flags, EMaterialGetParameterValueFlags::CheckNonOverrides))
+	{
+		bResult = InstanceChain.GetCachedExpressionData().GetParameterValue(Type, ParameterInfo, OutResult);
+	}
+
 	return bResult;
 }
 
@@ -1323,67 +1324,9 @@ int32 UMaterialInstance::GetBlendablePriority(const UMaterial* Base) const
 	}
 }
 
-void UMaterialInstance::GetTextureExpressionValues(const FMaterialResource* MaterialResource, TArray<UTexture*>& OutTextures, TArray<TArray<int32>>* OutIndices, bool bIncludeTextureCollections) const
+void UMaterialInstance::GetTextureExpressionValues(const FMaterialResource* MaterialResource, TArray<UTexture*>& OutTextures, TArray<TArray<int32>>* OutIndices, bool bIncludeTextureCollections, bool bAllowOverrides) const
 {
-	check(MaterialResource);
-	const FUniformExpressionSet& UniformExpressions = MaterialResource->GetUniformExpressions();
-
-	if (OutIndices) // Try to prevent resizing since this would be expensive.
-	{
-		uint32 NumTextures = 0u;
-		for (uint32 TypeIndex = 0u; TypeIndex < NumMaterialTextureParameterTypes; ++TypeIndex)
-		{
-			NumTextures += UniformExpressions.GetNumTextures((EMaterialTextureParameterType)TypeIndex);
-		}
-		OutIndices->Empty(NumTextures);
-	}
-
-	for(int32 TypeIndex = 0;TypeIndex < NumMaterialTextureParameterTypes;TypeIndex++)
-	{
-		// Iterate over each of the material's texture expressions.
-		for(int32 TextureIndex = 0; TextureIndex < UniformExpressions.GetNumTextures((EMaterialTextureParameterType)TypeIndex); TextureIndex++)
-		{
-			// Evaluate the expression in terms of this material instance.
-			UTexture* Texture = NULL;
-			UniformExpressions.GetGameThreadTextureValue((EMaterialTextureParameterType)TypeIndex, TextureIndex, this,*MaterialResource,Texture, true);
-			
-			if (Texture)
-			{
-				const int32 InsertIndex = OutTextures.AddUnique(Texture);
-				if (OutIndices)
-				{
-					const FMaterialTextureParameterInfo& Parameter = UniformExpressions.GetTextureParameter((EMaterialTextureParameterType)TypeIndex, TextureIndex);
-					if (InsertIndex >= OutIndices->Num())
-					{
-						OutIndices->AddDefaulted(InsertIndex - OutIndices->Num() + 1);
-					}
-					(*OutIndices)[InsertIndex].Add(Parameter.TextureIndex);
-				}
-			}
-		}
-	}
-
-	if (bIncludeTextureCollections)
-	{
-		checkf(OutIndices == nullptr, TEXT("Texture Collections don't work with Texture Indices."));
-
-		for (int32 TextureCollectionIndex = 0; TextureCollectionIndex < UniformExpressions.GetNumTextureCollections(); TextureCollectionIndex++)
-		{
-			UTextureCollection* TextureCollection = nullptr;
-			UniformExpressions.GetGameThreadTextureCollectionValue(TextureCollectionIndex, this, *MaterialResource, TextureCollection);
-
-			if (TextureCollection)
-			{
-				for (UTexture* Texture : TextureCollection->Textures)
-				{
-					if (Texture)
-					{
-						OutTextures.AddUnique(Texture);
-					}
-				}
-			}
-		}
-	}
+	Super::GetTextureExpressionValues(MaterialResource, OutTextures, OutIndices, bIncludeTextureCollections, bAllowOverrides);
 
 #if WITH_EDITOR
 	for (int32 i = 0; i < TransientTextureParameterOverrides.Num(); ++i)
@@ -1418,33 +1361,8 @@ void UMaterialInstance::GetUsedTextures(TArray<UTexture*>& OutTextures, TOptiona
 	// Do not care if we're running dedicated server
 	if (!FPlatformProperties::IsServerOnly())
 	{
-		FInt32Range QualityLevelRange(0, EMaterialQualityLevel::Num - 1);
-		if (InSpecificQualityLevel.IsSet())
-		{
-			EMaterialQualityLevel::Type QualityLevel = InSpecificQualityLevel.GetValue();
-			ensure(QualityLevel != EMaterialQualityLevel::Num);
-			if (QualityLevel == EMaterialQualityLevel::Num)
-			{
-				QualityLevel = GetCurrentMaterialQualityLevelChecked();
-			}
-			QualityLevelRange = FInt32Range(QualityLevel, QualityLevel);
-		}
-
-#if WITH_EDITOR
-		FInt32Range ShaderPlatformRange(0, EShaderPlatform::SP_NumPlatforms - 1);
-		if (InSpecificShaderPlatform.IsSet())
-		{
-			EShaderPlatform ShaderPlatform = InSpecificShaderPlatform.GetValue();
-			ensure(ShaderPlatform != EShaderPlatform::SP_NumPlatforms);
-			if (ShaderPlatform == EShaderPlatform::SP_NumPlatforms)
-			{
-				ShaderPlatform = GetMaxShaderPlatformChecked();
-			}
-			ShaderPlatformRange = FInt32Range(ShaderPlatform, ShaderPlatform);
-		}
-#else
-		FInt32Range ShaderPlatformRange(GMaxRHIShaderPlatform, GMaxRHIShaderPlatform);
-#endif
+		const FInt32Range QualityLevelRange = BuildQualityLevelRange(InSpecificQualityLevel);
+		const FInt32Range ShaderPlatformRange = BuildShaderPlatformRange(InSpecificShaderPlatform);
 
 		const UMaterial* BaseMaterial = GetMaterial();
 		const UMaterialInstance* MaterialInstanceToUse = this;
@@ -1460,28 +1378,13 @@ void UMaterialInstance::GetUsedTextures(TArray<UTexture*>& OutTextures, TOptiona
 			// Use the uniform expressions from the lowest material instance with static parameters in the chain, if one exists
 			const UMaterialInterface* MaterialToUse = (MaterialInstanceToUse && MaterialInstanceToUse->bHasStaticPermutationResource) ? (const UMaterialInterface*)MaterialInstanceToUse : (const UMaterialInterface*)BaseMaterial;
 
-			TArray<const FMaterialResource*, TInlineAllocator<4>> MatchedResources;
-			// Parse all relevant quality and feature levels.
-			for (int32 QualityLevelIndex = QualityLevelRange.GetLowerBoundValue(); QualityLevelIndex <= QualityLevelRange.GetUpperBoundValue(); ++QualityLevelIndex)
-			{
-				for (int32 ShaderPlatformIndex = ShaderPlatformRange.GetLowerBoundValue(); ShaderPlatformIndex <= ShaderPlatformRange.GetUpperBoundValue(); ++ShaderPlatformIndex)
-				{
-					if (!FDataDrivenShaderPlatformInfo::IsValid((EShaderPlatform)ShaderPlatformIndex))
-					{
-						continue;
-					}
+			TArray<const FMaterialResource*, TInlineAllocator<4>> MatchedResources = GetMaterialResources(MaterialToUse, QualityLevelRange, ShaderPlatformRange);
 
-					const FMaterialResource* MaterialResource = MaterialToUse->GetMaterialResource((EShaderPlatform)ShaderPlatformIndex, (EMaterialQualityLevel::Type)QualityLevelIndex);
-					if (MaterialResource)
-					{
-						MatchedResources.AddUnique(MaterialResource);
-					}
-				}
-			}
-
-			for (const FMaterialResource* MaterialResource : MatchedResources)
+			for (const FMaterialResource* CurrentResource : MatchedResources)
 			{
-				GetTextureExpressionValues(MaterialResource, OutTextures, nullptr, true);
+				constexpr bool bIncludeTextureCollections = true;
+				constexpr bool bAllowOverrides = true;
+				GetTextureExpressionValues(CurrentResource, OutTextures, nullptr, bIncludeTextureCollections, bAllowOverrides);
 			}
 		}
 		else
@@ -1549,7 +1452,10 @@ void UMaterialInstance::LogMaterialsAndTextures(FOutputDevice& Ar, int32 Indent)
 				if (MaterialResource->HasValidGameThreadShaderMap())
 				{
 					TArray<UTexture*> Textures;
-					GetTextureExpressionValues(MaterialResource, Textures);
+					constexpr bool bIncludeTextureCollections = false;
+					constexpr bool bAllowOverrides = true;
+					GetTextureExpressionValues(MaterialResource, Textures, nullptr, bIncludeTextureCollections, bAllowOverrides);
+
 					for (UTexture* Texture : Textures)
 					{
 						if (Texture)
@@ -1697,7 +1603,9 @@ void UMaterialInstance::GetUsedTexturesAndIndices(TArray<UTexture*>& OutTextures
 			const FMaterialResource* CurrentResource = FindMaterialResource(MaterialInstanceToUse->StaticPermutationMaterialResources, InShaderPlatform, QualityLevel, true);
 			if (CurrentResource)
 			{
-				GetTextureExpressionValues(CurrentResource, OutTextures, &OutIndices);
+				constexpr bool bIncludeTextureCollections = false;
+				constexpr bool bAllowOverrides = true;
+				GetTextureExpressionValues(CurrentResource, OutTextures, &OutIndices, bIncludeTextureCollections, bAllowOverrides);
 			}
 		}
 		else // Use the uniform expressions from the base material
@@ -1708,7 +1616,9 @@ void UMaterialInstance::GetUsedTexturesAndIndices(TArray<UTexture*>& OutTextures
 				const FMaterialResource* MaterialResource = Material->GetMaterialResource(InShaderPlatform, QualityLevel);
 				if( MaterialResource )
 				{
-					GetTextureExpressionValues(MaterialResource, OutTextures, &OutIndices);
+					constexpr bool bIncludeTextureCollections = false;
+					constexpr bool bAllowOverrides = true;
+					GetTextureExpressionValues(MaterialResource, OutTextures, &OutIndices, bIncludeTextureCollections, bAllowOverrides);
 				}
 			}
 			else // If the material instance has no material, use the default material.
@@ -3287,6 +3197,19 @@ void UMaterialInstance::Serialize(FArchive& Ar)
 #endif
 
 #if WITH_EDITOR
+	if (Ar.IsSaving() && Ar.IsCooking())
+	{
+		if (ShouldStripUnusedDefaultTextures())
+		{
+			constexpr bool bKeepProjectPrivateTextures = false;
+			StripUnusedDefaultTextures(bKeepProjectPrivateTextures);
+		}
+		else
+		{
+			// All textures will be kept in Base materials
+			ReferencedDefaultTextures.Empty();
+		}
+	}
 	// Do not serialize the overrides
 	ResetAllTextureParameterOverrides();
 #endif
@@ -3499,6 +3422,14 @@ void UMaterialInstance::Serialize(FArchive& Ar)
 	{
 		ValidateTextureOverrides(GMaxRHIShaderPlatform);
 	}
+
+#if WITH_EDITOR
+	// Undo textures removal for Multi Platform Cook
+	if (Ar.IsSaving() && Ar.IsCooking() && ShouldStripUnusedDefaultTextures())
+	{
+		UpdateReferencedDefaultTextures();
+	}
+#endif
 }
 
 void UMaterialInstance::PostDuplicate(bool bDuplicateForPIE)
@@ -3670,14 +3601,11 @@ void UMaterialInstance::PostLoad()
 		}
 	}
 
-	if (CachedExpressionData != nullptr)
+	for (UObject* Texture : ReferencedDefaultTextures)
 	{
-		for (UObject* Texture : CachedExpressionData->ReferencedTextures)
+		if (Texture)
 		{
-			if (Texture)
-			{
-				Texture->ConditionalPostLoad();
-			}
+			Texture->ConditionalPostLoad();
 		}
 	}
 
@@ -4690,7 +4618,7 @@ TSoftObjectPtr<UObject> UMaterialInstance::GetEnumerationObject(int32 Index) con
 
 void UMaterialInstance::GetReferencedTexturesAndOverrides(TSet<const UTexture*>& InOutTextures) const
 {
-	for (UObject* UsedObject : GetCachedExpressionData().ReferencedTextures)
+	for (UObject* UsedObject : GetDefaultTextures())
 	{
 		if (const UTexture* UsedTexture = Cast<UTexture>(UsedObject))
 		{
@@ -4718,6 +4646,20 @@ void UMaterialInstance::UpdateCachedData()
 	// Overridden for MIC/MID
 }
 
+void UMaterialInstance::UpdateReferencedDefaultTextures()
+{
+	const FStaticParameterSet& LocalStaticParameters = GetStaticParameters();
+	if (bHasStaticPermutationResource || LocalStaticParameters.bHasMaterialLayers)
+	{
+		const FMaterialCachedExpressionData& UsedCachedData = GetCachedExpressionData();
+		ReferencedDefaultTextures = UsedCachedData.EditorOnlyData->EditorReferencedDefaultTextures;
+	}
+	else
+	{
+		ReferencedDefaultTextures.Empty();
+	}
+}
+
 void UMaterialInstance::UpdateStaticPermutation(const FStaticParameterSet& NewParameters, FMaterialUpdateContext* MaterialUpdateContext)
 {
 	UpdateStaticPermutation(NewParameters, BasePropertyOverrides, false, MaterialUpdateContext);
@@ -5454,6 +5396,48 @@ int32 UMaterialInstance::CompilePropertyEx( class FMaterialCompiler* Compiler, c
 	return Parent ? Parent->CompilePropertyEx(Compiler, AttributeID) : INDEX_NONE;
 }
 
+bool UMaterialInstance::TryGetUsedTexturesInCookCache(TArray<TObjectPtr<UObject>>& OutTextures)
+{
+	OutTextures.Empty();
+	const UMaterial* BaseMaterial = GetMaterial();
+
+	if (!BaseMaterial || BaseMaterial->IsDefaultMaterial())
+	{
+		return false;
+	}
+
+	const FMaterialCachedExpressionData& UsedCachedData = GetCachedExpressionData();
+	TArray<const FMaterialResource*, TInlineAllocator<4>> MatchedResources = GetCookedCacheMaterialResources();
+
+	for (const FMaterialResource* CurrentResource : MatchedResources)
+	{
+		const FUniformExpressionSet& CurrentSet = CurrentResource->GetUniformExpressions();
+		for (int32 TypeIndex = 0; TypeIndex < NumMaterialTextureParameterTypes; ++TypeIndex)
+		{
+			for (int32 TextureIndex = 0; TextureIndex < CurrentSet.GetNumTextures((EMaterialTextureParameterType)TypeIndex); ++TextureIndex)
+			{
+				const FMaterialTextureParameterInfo& Parameter = CurrentSet.GetTextureParameter((EMaterialTextureParameterType)TypeIndex, TextureIndex);
+				FMaterialParameterMetadata SourceValue;
+				bool bIsOverriden = false;
+
+				// Do not add overriden default textures
+				for (int32 ParamTypeIdx = (int32)EMaterialParameterType::Texture; (ParamTypeIdx <= (int32)EMaterialParameterType::SparseVolumeTexture) && !bIsOverriden; ++ParamTypeIdx)
+				{
+					bIsOverriden |= GetParameterValue((EMaterialParameterType)ParamTypeIdx, Parameter.ParameterInfo, SourceValue, EMaterialGetParameterValueFlags::CheckInstanceOverrides);
+				}
+
+				if (!bIsOverriden)
+				{
+					OutTextures.AddUnique(UsedCachedData.EditorOnlyData->EditorReferencedDefaultTextures[Parameter.TextureIndex]);
+				}
+			}
+		}
+	}
+
+	const bool bValidCompilationOutput = MatchedResources.Num() > 0;
+	return bValidCompilationOutput;
+}
+
 const FStaticParameterSetEditorOnlyData& UMaterialInstance::GetEditorOnlyStaticParameters() const
 {
 	const UMaterialInstanceEditorOnlyData* EditorOnly = GetEditorOnlyData();
@@ -5714,6 +5698,25 @@ void UMaterialInstance::DumpDebugInfo(FOutputDevice& OutputDevice) const
 	}
 }
 
+TObjectPtr<UObject> UMaterialInstance::GetDefaultTextureAtIdx(const int32 InIdx, TMicRecursionGuard RecursionGuard)
+{
+	// We also look through the Parents in case there's a texture that has been forced to 
+	// remain there because it is Project Private (PKG_NotExternallyReferenceable)
+
+	if (ReferencedDefaultTextures.IsValidIndex(InIdx) && ReferencedDefaultTextures[InIdx] != nullptr)
+	{
+		return ReferencedDefaultTextures[InIdx];
+	}
+
+	if (Parent && !RecursionGuard.Contains(this))
+	{
+		RecursionGuard.Set(this);
+		return Parent->GetDefaultTextureAtIdx(InIdx, RecursionGuard);
+	}
+
+	return nullptr;
+}
+
 void UMaterialInstance::SaveShaderStableKeys(const class ITargetPlatform* TP)
 {
 #if WITH_EDITOR
@@ -6058,4 +6061,4 @@ bool UMaterialInstance::GetTextureCollectionParameterValue(const FHashedMaterial
 {
 	bool bResult = Super::GetTextureCollectionParameterValue(ParameterInfo, OutValue, bOveriddenOnly);
 	return bResult;
-}
\ No newline at end of file
+}
diff --git a/Engine/Source/Runtime/Engine/Private/Materials/MaterialInstanceConstant.cpp b/Engine/Source/Runtime/Engine/Private/Materials/MaterialInstanceConstant.cpp
--- a/Engine/Source/Runtime/Engine/Private/Materials/MaterialInstanceConstant.cpp
+++ b/Engine/Source/Runtime/Engine/Private/Materials/MaterialInstanceConstant.cpp
@@ -238,7 +238,11 @@ void UMaterialInstanceConstant::UpdateCachedData()
 			EditorOnlyData->CachedExpressionData = CachedExpressionData->EditorOnlyData;
 		}
 
+		UpdateReferencedDefaultTextures();
+
+#if WITH_EDITOR
 		FObjectCacheEventSink::NotifyReferencedTextureChanged_Concurrent(this);
+#endif
 	}
 }
 
diff --git a/Engine/Source/Runtime/Engine/Private/Materials/MaterialInterface.cpp b/Engine/Source/Runtime/Engine/Private/Materials/MaterialInterface.cpp
--- a/Engine/Source/Runtime/Engine/Private/Materials/MaterialInterface.cpp
+++ b/Engine/Source/Runtime/Engine/Private/Materials/MaterialInterface.cpp
@@ -50,13 +50,17 @@
 #include "UObject/ObjectEditorOptionalSupport.h"
 #include "UObject/ICookInfo.h"
 #include "Async/UniqueLock.h"
+#include "Engine/TextureCollection.h"
 
 #if WITH_EDITOR
 #include "ObjectCacheEventSink.h"
+#include "Misc/PackageName.h"
 #endif
 
 #include UE_INLINE_GENERATED_CPP_BY_NAME(MaterialInterface)
 
+DECLARE_CYCLE_STAT(TEXT("MaterialInterface StripUnusedDefaultTextures"), STAT_MaterialInterface_StripUnusedDefaultTextures, STATGROUP_Shaders);
+
 #define LOCTEXT_NAMESPACE "MaterialInterface"
 
 /** 
@@ -72,6 +76,14 @@ UEnum* UMaterialInterface::SamplerTypeEnum = nullptr;
 
 //////////////////////////////////////////////////////////////////////////
 
+#if WITH_EDITOR
+static TAutoConsoleVariable<bool> CVarStripUnusedDefaultTextures(
+	TEXT("r.Material.StripUnusedDefaultTextures"),
+	true,
+	TEXT("If enabled, unused Default textures will be discarded at cook time.\n"),
+	ECVF_ReadOnly);
+#endif
+
 bool IsHairStrandsGeometrySupported(const EShaderPlatform Platform)
 {
 	check(Platform != SP_NumPlatforms);
@@ -292,6 +304,25 @@ const FMaterialCachedExpressionData& UMaterialInterface::GetCachedExpressionData
 	return LocalData ? *LocalData : FMaterialCachedExpressionData::EmptyData;
 }
 
+#if WITH_EDITOR
+int32 UMaterialInterface::GetDefaultTextureIdx(UObject* InTexture) const
+{
+	const FMaterialCachedExpressionData& UsedCachedData = GetCachedExpressionData();
+	return UsedCachedData.GetDefaultTextureIdx(InTexture);
+}
+
+bool UMaterialInterface::ContainsDefaultTexture(UObject* InTexture) const
+{
+	return GetDefaultTextureIdx(InTexture) != INDEX_NONE;
+}
+
+const TArray<TObjectPtr<UObject>>& UMaterialInterface::GetDefaultTextures() const
+{
+	const FMaterialCachedExpressionData& UsedCachedData = GetCachedExpressionData();
+	return UsedCachedData.EditorOnlyData->EditorReferencedDefaultTextures;
+}
+#endif // WITH_EDITOR
+
 bool UMaterialInterface::IsUsingNewHLSLGenerator() const
 {
 	static const auto CVar = IConsoleManager::Get().FindTConsoleVariableDataBool(TEXT("r.Material.Translator.EnableNew"));
@@ -368,7 +399,7 @@ void UMaterialInterface::GetQualityLevelUsage(TArray<bool, TInlineAllocator<EMat
 
 TArrayView<const TObjectPtr<UObject>> UMaterialInterface::GetReferencedTextures() const
 {
-	return GetCachedExpressionData().ReferencedTextures;
+	return ReferencedDefaultTextures;
 }
 
 TConstArrayView<TObjectPtr<UTextureCollection>> UMaterialInterface::GetReferencedTextureCollections() const
@@ -379,7 +410,8 @@ TConstArrayView<TObjectPtr<UTextureCollection>> UMaterialInterface::GetReference
 #if WITH_EDITOR
 void UMaterialInterface::GetReferencedTexturesAndOverrides(TSet<const UTexture*>& InOutTextures) const
 {
-	for (UObject* UsedObject : GetCachedExpressionData().ReferencedTextures)
+	const TArray<TObjectPtr<UObject>> Textures = GetDefaultTextures();
+	for (UObject* UsedObject : Textures)
 	{
 		if (const UTexture* UsedTexture = Cast<UTexture>(UsedObject))
 		{
@@ -401,6 +433,96 @@ void UMaterialInterface::GetUsedTexturesAndIndices(TArray<UTexture*>& OutTexture
 	OutIndices.AddDefaulted(OutTextures.Num());
 }
 
+void UMaterialInterface::GetTextureExpressionValues(const FMaterialResource* MaterialResource, TArray<UTexture*>& OutTextures, TArray<TArray<int32>>* OutIndices, bool bIncludeTextureCollections, bool bAllowOverrides) const
+{
+	check(MaterialResource);
+	const FUniformExpressionSet& UniformExpressions = MaterialResource->GetUniformExpressions();
+
+	// Pre-allocate array
+	if (OutIndices) 
+	{
+		uint32 NumTextures = 0u;
+		for (uint32 TypeIndex = 0u; TypeIndex < NumMaterialTextureParameterTypes; ++TypeIndex)
+		{
+			NumTextures += UniformExpressions.GetNumTextures((EMaterialTextureParameterType)TypeIndex);
+		}
+		OutIndices->Empty(NumTextures);
+	}
+
+	for (int32 TypeIndex = 0; TypeIndex < NumMaterialTextureParameterTypes; ++TypeIndex)
+	{
+		// Iterate over each of the material's texture expressions.
+		for (int32 TextureIndex = 0; TextureIndex < UniformExpressions.GetNumTextures((EMaterialTextureParameterType)TypeIndex); TextureIndex++)
+		{
+			// Evaluate the expression in terms of this material instance.
+			UTexture* Texture = nullptr;
+			UniformExpressions.GetGameThreadTextureValue((EMaterialTextureParameterType)TypeIndex, TextureIndex, this, *MaterialResource, Texture, bAllowOverrides);
+
+			if (Texture)
+			{
+				const int32 InsertIndex = OutTextures.AddUnique(Texture);
+
+				if (OutIndices)
+				{
+					if (InsertIndex >= OutIndices->Num())
+					{
+						OutIndices->AddDefaulted(InsertIndex - OutIndices->Num() + 1);
+					}
+
+					const FMaterialTextureParameterInfo& Parameter = UniformExpressions.GetTextureParameter((EMaterialTextureParameterType)TypeIndex, TextureIndex);
+					(*OutIndices)[InsertIndex].Add(Parameter.TextureIndex);
+				}
+			}
+		}
+	}
+
+	if (bIncludeTextureCollections)
+	{
+		checkf(OutIndices == nullptr, TEXT("Texture Collections don't work with Texture Indices."));
+
+		for (int32 TextureCollectionIndex = 0; TextureCollectionIndex < UniformExpressions.GetNumTextureCollections(); TextureCollectionIndex++)
+		{
+			UTextureCollection* TextureCollection = nullptr;
+			UniformExpressions.GetGameThreadTextureCollectionValue(TextureCollectionIndex, this, *MaterialResource, TextureCollection);
+
+			if (TextureCollection)
+			{
+				for (UTexture* Texture : TextureCollection->Textures)
+				{
+					if (Texture)
+					{
+						OutTextures.AddUnique(Texture);
+					}
+				}
+			}
+		}
+	}
+}
+
+TArray<const FMaterialResource*, TInlineAllocator<4>> UMaterialInterface::GetMaterialResources(const UMaterialInterface* InMaterial, const FInt32Range InQualityLevelRange, const FInt32Range InShaderPlatformRange)
+{
+	TArray<const FMaterialResource*, TInlineAllocator<4>> MatchedResources;
+	// Parse all relevant quality and feature levels.
+	for (int32 QualityLevelIndex = InQualityLevelRange.GetLowerBoundValue(); QualityLevelIndex <= InQualityLevelRange.GetUpperBoundValue(); ++QualityLevelIndex)
+	{
+		for (int32 ShaderPlatformIndex = InShaderPlatformRange.GetLowerBoundValue(); ShaderPlatformIndex <= InShaderPlatformRange.GetUpperBoundValue(); ++ShaderPlatformIndex)
+		{
+			if (!FDataDrivenShaderPlatformInfo::IsValid((EShaderPlatform)ShaderPlatformIndex))
+			{
+				continue;
+			}
+
+			const FMaterialResource* CurrentResource = InMaterial->GetMaterialResource((EShaderPlatform)ShaderPlatformIndex, (EMaterialQualityLevel::Type)QualityLevelIndex);
+			if (CurrentResource)
+			{
+				MatchedResources.AddUnique(CurrentResource);
+			}
+		}
+	}
+
+	return MatchedResources;
+}
+
 #if WITH_EDITORONLY_DATA
 bool UMaterialInterface::GetStaticSwitchParameterValue(const FHashedMaterialParameterInfo& ParameterInfo, bool& OutValue, FGuid& OutExpressionGuid, bool bOveriddenOnly /*= false*/) const
 {
@@ -917,6 +1039,10 @@ void UMaterialInterface::AppendToClassSchema(FAppendToClassSchemaContext& Contex
 	Super::AppendToClassSchema(Context);
 
 	UE::MaterialInterface::Private::HashMaterialStaticClassDependenciesForCook(Context);
+
+	const bool bStripUnusedDefaultTextures = ShouldStripUnusedDefaultTextures();
+	const uint8 Byte = bStripUnusedDefaultTextures ? 1 : 0;
+	Context.Update(&Byte, sizeof(Byte));
 }
 #endif
 
@@ -1962,6 +2088,120 @@ void UMaterialInterface::UpdateMaterialRenderProxy(FMaterialRenderProxy& Proxy)
 	}
 }
 
+#if WITH_EDITOR
+bool UMaterialInterface::ShouldStripUnusedDefaultTextures()
+{
+	return CVarStripUnusedDefaultTextures.GetValueOnGameThread();
+}
+
+TArray<const FMaterialResource*, TInlineAllocator<4>> UMaterialInterface::GetCookedCacheMaterialResources() const
+{
+	TArray<const FMaterialResource*, TInlineAllocator<4>> CookedResources;
+	for (const auto& KVP : CachedMaterialResourcesForCooking)
+	{
+		const TArray<FMaterialResourceForCooking>& PlatformResources = KVP.Value;
+
+		for (const FMaterialResourceForCooking& ResForCooking : PlatformResources)
+		{
+			if (ResForCooking.Resource)
+			{
+				CookedResources.Add(ResForCooking.Resource);
+			}
+		}
+	}
+
+	return CookedResources;
+}
+
+void UMaterialInterface::StripUnusedDefaultTextures(const bool bInKeepProjectPrivateTextures)
+{
+	SCOPE_CYCLE_COUNTER(STAT_MaterialInterface_StripUnusedDefaultTextures);
+
+	if (!CVarStripUnusedDefaultTextures.GetValueOnGameThread())
+	{
+		static bool bWarnedOnce = false;
+		if (!bWarnedOnce)
+		{
+			UE_LOG(LogMaterial, Warning, TEXT("Calling StripUnusedDefaultTextures even though CVarStripUnusedDefaultTextures is false."));
+			bWarnedOnce = true;
+		}
+
+		return;
+	}
+
+	TArray<TObjectPtr<UObject>> UsedTextures;
+	const bool bValidUsedTextures = TryGetUsedTexturesInCookCache(UsedTextures);
+
+	if (!bValidUsedTextures)
+	{
+		if (!bInKeepProjectPrivateTextures)
+		{
+			// Textures will be kept in Parent material
+			ReferencedDefaultTextures.Empty();
+		}
+
+		return;
+	}
+
+	for (int32 Idx = 0; Idx < ReferencedDefaultTextures.Num(); ++Idx)
+	{
+		const TObjectPtr<UObject> CurrentTexture = ReferencedDefaultTextures[Idx];
+
+		if (!CurrentTexture)
+		{
+			continue;
+		}
+
+		// Check for project private textures (PKC_NotExternallyReferenceable)
+		const UPackage* TexturePackage = CurrentTexture->GetPackage();
+		const UPackage* OwnPackage = GetPackage();
+
+		if (!TexturePackage || !OwnPackage)
+		{
+			continue;
+		}
+
+		const bool bUsedTexture = UsedTextures.Contains(CurrentTexture);
+		bool bRemoveTexture = !bUsedTexture;
+		const bool bIsPackagePrivate = TexturePackage->GetAssetAccessSpecifier() == EAssetAccessSpecifier::Private;
+
+		if (!bInKeepProjectPrivateTextures)
+		{
+			// We are in a Material Instance, check for private package and mount points
+			if (bIsPackagePrivate)
+			{
+				const FName OwnMountPointName = FPackageName::GetPackageMountPoint(FNameBuilder(OwnPackage->GetFName()).ToView());
+				const FName TextureMountPointName = FPackageName::GetPackageMountPoint(FNameBuilder(TexturePackage->GetFName()).ToView());
+
+				if (!OwnMountPointName.IsNone() && !TextureMountPointName.IsNone())
+				{
+					if (OwnMountPointName != TextureMountPointName)
+					{
+						// Private package, different mount points -> remove texture
+						bRemoveTexture = true;
+					}
+				}
+				else
+				{
+					// Private package and None mount points -> remove texture
+					bRemoveTexture = true;
+				}
+			}
+		}
+		else
+		{
+			// We are in a Material, we only remove if the package is not private
+			bRemoveTexture &= !bIsPackagePrivate;
+		}
+
+		if (bRemoveTexture)
+		{
+			ReferencedDefaultTextures[Idx] = nullptr;
+		}
+	}
+}
+#endif // WITH_EDITOR
+
 bool FMaterialTextureInfo::IsValid(bool bCheckTextureIndex) const
 { 
 #if WITH_EDITORONLY_DATA
@@ -2272,6 +2512,48 @@ void UMaterialInterface::FilterOutPlatformShadingModels(const EShaderPlatform In
 	}
 }
 
+FInt32Range UMaterialInterface::BuildQualityLevelRange(TOptional<EMaterialQualityLevel::Type>  InQualityLevel)
+{
+	FInt32Range QualityLevelRange(0, EMaterialQualityLevel::Num - 1);
+
+	if (InQualityLevel.IsSet())
+	{
+		EMaterialQualityLevel::Type QualityLevel = InQualityLevel.GetValue();
+		ensure(QualityLevel != EMaterialQualityLevel::Num);
+
+		if (QualityLevel == EMaterialQualityLevel::Num)
+		{
+			QualityLevel = GetCurrentMaterialQualityLevelChecked();
+		}
+		QualityLevelRange = FInt32Range(QualityLevel, QualityLevel);
+	}
+
+	return QualityLevelRange;
+}
+
+FInt32Range UMaterialInterface::BuildShaderPlatformRange(TOptional<EShaderPlatform> InShaderPlatform)
+{
+#if WITH_EDITOR
+	FInt32Range ShaderPlatformRange(0, EShaderPlatform::SP_NumPlatforms - 1);
+
+	if (InShaderPlatform.IsSet())
+	{
+		EShaderPlatform ShaderPlatform = InShaderPlatform.GetValue();
+		ensure(ShaderPlatform != EShaderPlatform::SP_NumPlatforms);
+
+		if (ShaderPlatform == EShaderPlatform::SP_NumPlatforms)
+		{
+			ShaderPlatform = GetMaxShaderPlatformChecked();
+		}
+		ShaderPlatformRange = FInt32Range(ShaderPlatform, ShaderPlatform);
+	}
+#else
+	FInt32Range ShaderPlatformRange(GMaxRHIShaderPlatform, GMaxRHIShaderPlatform);
+#endif
+
+	return ShaderPlatformRange;
+}
+
 #if WITH_EDITORONLY_DATA
 
 
diff --git a/Engine/Source/Runtime/Engine/Private/Materials/MaterialShared.cpp b/Engine/Source/Runtime/Engine/Private/Materials/MaterialShared.cpp
--- a/Engine/Source/Runtime/Engine/Private/Materials/MaterialShared.cpp
+++ b/Engine/Source/Runtime/Engine/Private/Materials/MaterialShared.cpp
@@ -1569,6 +1569,65 @@ void FMaterialResource::LegacySerialize(FArchive& Ar)
 	}
 }
 
+TObjectPtr<UObject> FMaterialResource::GetMatDefaultTextureAtIdx(const int32 InIdx) const
+{
+	if (MaterialInstance)
+	{
+		TObjectPtr<UObject> Texture =  MaterialInstance->GetDefaultTextureAtIdx(InIdx);
+		if (Texture)
+		{
+			return Texture;
+		}
+	}
+
+	if (Material)
+	{
+		return Material->GetDefaultTextureAtIdx(InIdx);
+	}
+
+	return UMaterial::GetDefaultMaterial(MD_Surface)->GetDefaultTextureAtIdx(InIdx);
+}
+
+#if WITH_EDITOR
+int32 FMaterialResource::GetMatDefaultTextureIdx(TObjectPtr<UObject> InTexture) const
+{
+	if (MaterialInstance)
+	{
+		const int32 Idx = MaterialInstance->GetDefaultTextureIdx(InTexture);
+		if (Idx != INDEX_NONE)
+		{
+			return Idx;
+		}
+	}
+
+	if (Material)
+	{
+		return Material->GetDefaultTextureIdx(InTexture);
+	}
+
+	return UMaterial::GetDefaultMaterial(MD_Surface)->GetDefaultTextureIdx(InTexture);
+}
+
+const TArray<TObjectPtr<UObject>>& FMaterialResource::GetMatDefaultTextures() const
+{
+	if (MaterialInstance)
+	{
+		const TArray<TObjectPtr<UObject>>& Textures = MaterialInstance->GetDefaultTextures();
+		if (Textures.Num() > 0)
+		{
+			return Textures;
+		}
+	}
+
+	if (Material)
+	{
+		return Material->GetDefaultTextures();
+	}
+
+	return UMaterial::GetDefaultMaterial(MD_Surface)->GetDefaultTextures();
+}
+#endif // WITH_EDITOR
+
 TArrayView<const TObjectPtr<UObject>> FMaterialResource::GetReferencedTextures() const
 {
 	if (MaterialInstance)
@@ -4539,7 +4598,7 @@ void FMaterial::GetReferencedTexturesHash(FSHAHash& OutHash) const
 {
 	FSHA1 HashState;
 
-	const TArrayView<const TObjectPtr<UObject>> ReferencedTextures = GetReferencedTextures();
+	const TArray<TObjectPtr<UObject>>& ReferencedTextures = GetMatDefaultTextures();
 	// Hash the names of the uniform expression textures to capture changes in their order or values resulting from material compiler code changes
 	for (int32 TextureIndex = 0; TextureIndex < ReferencedTextures.Num(); TextureIndex++)
 	{
@@ -6620,6 +6679,36 @@ bool IsDevelopmentFeatureEnabled(const FName& FeatureName, EShaderPlatform Platf
 
 }  // namespace UE::MaterialTranslatorUtils
 
+const TArray<TObjectPtr<UObject>>& FEditorMaterialResourceBase::GetMatDefaultTextures() const
+{
+	return ReferencedTextures;
+}
+
+int32 FEditorMaterialResourceBase::GetMatDefaultTextureIdx(TObjectPtr<UObject> InTexture) const
+{
+	return ReferencedTextures.Find(InTexture);
+}
+
+TObjectPtr<UObject> FEditorMaterialResourceBase::GetMatDefaultTextureAtIdx(const int32 InIdx) const
+{
+	if (!ensureAlwaysMsgf(ReferencedTextures.IsValidIndex(InIdx), TEXT("Invalid texture Idx %d requested from Array with %d Elements."), InIdx, ReferencedTextures.Num()))
+	{
+		return nullptr;
+	}
+
+	return ReferencedTextures[InIdx];
+}
+
+TArrayView<const TObjectPtr<UObject>> FEditorMaterialResourceBase::GetReferencedTextures() const
+{
+	return ReferencedTextures;
+}
+
+TConstArrayView<TObjectPtr<UTextureCollection>> FEditorMaterialResourceBase::GetReferencedTextureCollections() const
+{
+	return MakeArrayView(ReferencedTextureCollections);
+}
+
 #endif // WITH_EDITOR
 
 #undef LOCTEXT_NAMESPACE
diff --git a/Engine/Source/Runtime/Engine/Private/Materials/MaterialUniformExpressions.h b/Engine/Source/Runtime/Engine/Private/Materials/MaterialUniformExpressions.h
--- a/Engine/Source/Runtime/Engine/Private/Materials/MaterialUniformExpressions.h
+++ b/Engine/Source/Runtime/Engine/Private/Materials/MaterialUniformExpressions.h
@@ -452,22 +452,7 @@ private:
 template<typename TextureType>
 static TextureType* GetIndexedTexture(const FMaterial& Material, int32 TextureIndex)
 {
-	UObject* IndexedTexture = nullptr;
-	const TArrayView<const TObjectPtr<UObject>> ReferencedTextures = Material.GetReferencedTextures();
-	if (ReferencedTextures.IsValidIndex(TextureIndex))
-	{
-		IndexedTexture = ReferencedTextures[TextureIndex];
-	}
-
-	if (IndexedTexture == nullptr)
-	{
-		static bool bWarnedOnce = false;
-		if (!bWarnedOnce)
-		{
-			UE_LOG(LogMaterial, Warning, TEXT("%s: Requesting an invalid TextureIndex! (%u / %u)"), *Material.GetFriendlyName(), TextureIndex, ReferencedTextures.Num());
-			bWarnedOnce = true;
-		}
-	}
+	UObject* IndexedTexture = Material.GetMatDefaultTextureAtIdx(TextureIndex);
 
 	// Can return nullptr if TextureType doesn't match type of indexed texture
 	return Cast<TextureType>(IndexedTexture);
diff --git a/Engine/Source/Runtime/Engine/Private/SparseVolumeTexture/SparseVolumeTexture.cpp b/Engine/Source/Runtime/Engine/Private/SparseVolumeTexture/SparseVolumeTexture.cpp
--- a/Engine/Source/Runtime/Engine/Private/SparseVolumeTexture/SparseVolumeTexture.cpp
+++ b/Engine/Source/Runtime/Engine/Private/SparseVolumeTexture/SparseVolumeTexture.cpp
@@ -1043,7 +1043,7 @@ void USparseVolumeTexture::NotifyMaterials(const ENotifyMaterialsEffectOnShaders
 		for (TObjectIterator<UMaterialInterface> It; It; ++It)
 		{
 			UMaterialInterface* MaterialInterface = *It;
-			if (!FPlatformProperties::IsServerOnly() && MaterialInterface->GetReferencedTextures().Contains(this))
+			if (!FPlatformProperties::IsServerOnly() && (MaterialInterface->ContainsDefaultTexture(this)))
 			{
 				UpdateContext.AddMaterialInterface(MaterialInterface);
 				// This is a bit tricky. We want to make sure all materials using this texture are
diff --git a/Engine/Source/Runtime/Engine/Public/MaterialCachedData.h b/Engine/Source/Runtime/Engine/Public/MaterialCachedData.h
--- a/Engine/Source/Runtime/Engine/Public/MaterialCachedData.h
+++ b/Engine/Source/Runtime/Engine/Public/MaterialCachedData.h
@@ -199,6 +199,10 @@ struct FMaterialCachedExpressionEditorOnlyData
 
 	UPROPERTY()
 	TArray<FName> ShaderTags;
+
+	/** All the default textures present in the nodes of the material this represents. */
+	UPROPERTY()
+	TArray<TObjectPtr<UObject>> EditorReferencedDefaultTextures;
 };
 
 USTRUCT()
@@ -268,6 +272,11 @@ struct FMaterialCachedExpressionData
 	void GetAllGlobalParametersOfType(EMaterialParameterType Type, TMap<FMaterialParameterInfo, FMaterialParameterMetadata>& OutParameters) const;
 	void GetAllGlobalParameterInfoOfType(EMaterialParameterType Type, TArray<FMaterialParameterInfo>& OutParameterInfo, TArray<FGuid>& OutParameterIds) const;
 
+#if WITH_EDITOR
+	/** @return The index at which the given texture sits in the default textures array or INDEX_NONE if not present. */
+	int32 GetDefaultTextureIdx(UObject* InTexture) const;
+#endif
+
 	ENGINE_API void AddReferencedObjects(FReferenceCollector& Collector);
 
 	/** Returns an array of the guids of functions used, with the call hierarchy flattened. */
@@ -341,10 +350,6 @@ struct FMaterialCachedExpressionData
 	UPROPERTY()
 	TArray<TSoftObjectPtr<USparseVolumeTexture>> SparseVolumeTextureValues;
 
-	/** Array of all texture referenced by this material */
-	UPROPERTY()
-	TArray<TObjectPtr<UObject>> ReferencedTextures;
-
 	UPROPERTY()
 	TArray<TObjectPtr<UTextureCollection>> ReferencedTextureCollections;
 
diff --git a/Engine/Source/Runtime/Engine/Public/MaterialShared.h b/Engine/Source/Runtime/Engine/Public/MaterialShared.h
--- a/Engine/Source/Runtime/Engine/Public/MaterialShared.h
+++ b/Engine/Source/Runtime/Engine/Public/MaterialShared.h
@@ -2737,6 +2737,13 @@ public:
 	ENGINE_API virtual void AddReferencedObjects(FReferenceCollector& Collector);
 
 	virtual TArrayView<const TObjectPtr<UObject>> GetReferencedTextures() const = 0;
+	virtual TObjectPtr<UObject> GetMatDefaultTextureAtIdx(const int32 inIdx) const = 0;
+
+#if WITH_EDITOR
+	virtual int32 GetMatDefaultTextureIdx(TObjectPtr<UObject> InTexture) const = 0;
+	virtual const TArray<TObjectPtr<UObject>>& GetMatDefaultTextures() const = 0;
+#endif
+
 	virtual TConstArrayView<TObjectPtr<UTextureCollection>> GetReferencedTextureCollections() const = 0;
 
 	/**
@@ -3325,6 +3332,12 @@ public:
 	ENGINE_API virtual FName GetAssetPath() const override;
 	ENGINE_API virtual bool IsUsingNewHLSLGenerator() const override;
 	ENGINE_API virtual bool CheckInValidStateForCompilation(FMaterialCompiler* Compiler) const override;
+
+	/** @return The index at which the given texture sits in the default textures array for this material. */
+	ENGINE_API virtual int32 GetMatDefaultTextureIdx(TObjectPtr<UObject> InTexture) const override;
+
+	/** @return All default textures present in this material. */
+	ENGINE_API virtual const TArray<TObjectPtr<UObject>>& GetMatDefaultTextures() const override;
 #endif // WITH_EDITOR
 
 	ENGINE_API virtual FString GetFullPath() const override;
@@ -3333,6 +3346,15 @@ public:
 
 	ENGINE_API virtual void LegacySerialize(FArchive& Ar) override;
 
+	/**
+	 * Returns the relevant texture for the given Idx.
+	 * In editor, it looks through all default textures. 
+	 * In non-editor, it has access to default textures based on cook output, if ShouldStripUnusedDefaultTextures returns true.
+	 * @return The texture present at the given Idx or nullptr if not available.
+	 */
+	ENGINE_API virtual TObjectPtr<UObject> GetMatDefaultTextureAtIdx(const int32 InIdx) const override;
+
+	UE_DEPRECATED(5.8, "Please use the GetMatDefaultTextures for editor or GetMatDefaultTextureAtIdx for non-editor.")
 	ENGINE_API virtual TArrayView<const TObjectPtr<UObject>> GetReferencedTextures() const override;
 	ENGINE_API virtual TConstArrayView<TObjectPtr<UTextureCollection>> GetReferencedTextureCollections() const override;
 
@@ -3588,6 +3610,19 @@ struct FMaterialResourceForCooking
 	TRefCountPtr<FMaterialResource> Resource;
 	EShaderPlatform Platform;
 };
+
+class FEditorMaterialResourceBase : public FMaterial
+{
+	ENGINE_API virtual const TArray<TObjectPtr<UObject>>& GetMatDefaultTextures() const override;
+	ENGINE_API virtual int32 GetMatDefaultTextureIdx(TObjectPtr<UObject> InTexture) const override;
+	ENGINE_API virtual TObjectPtr<UObject> GetMatDefaultTextureAtIdx(const int32 InIdx) const override;
+
+	ENGINE_API virtual TArrayView<const TObjectPtr<UObject>> GetReferencedTextures() const override;
+	ENGINE_API virtual TConstArrayView<TObjectPtr<UTextureCollection>> GetReferencedTextureCollections() const override;
+protected:
+	TArray<TObjectPtr<UObject>> ReferencedTextures;
+	TArray<TObjectPtr<UTextureCollection>> ReferencedTextureCollections;
+};
 #endif
 
 inline bool ShouldIncludeMaterialInDefaultOpaquePass(const FMaterial& Material)
diff --git a/Engine/Source/Runtime/Engine/Public/Materials/Material.h b/Engine/Source/Runtime/Engine/Public/Materials/Material.h
--- a/Engine/Source/Runtime/Engine/Public/Materials/Material.h
+++ b/Engine/Source/Runtime/Engine/Public/Materials/Material.h
@@ -1217,10 +1217,7 @@ private:
 	 * These are always valid and non-null, but only the entries affected by CacheResourceShadersForRendering are actually valid for rendering.
 	 */
 	TArray<FMaterialResource*> MaterialResources;
-#if WITH_EDITOR
-	/** Material resources being cached for cooking. */
-	TMap<const ITargetPlatform*, TArray<FMaterialResourceForCooking>> CachedMaterialResourcesForCooking;
-#endif
+
 	/** Flag used to guarantee that the RT is finished using various resources in this UMaterial before cleanup. */
 	FThreadSafeBool ReleasedByRT;
 
@@ -1355,6 +1352,9 @@ public:
 	ENGINE_API virtual int32 CompilePropertyEx( class FMaterialCompiler* Compiler, const FGuid& AttributeID ) override;
 	ENGINE_API virtual bool ShouldForcePlanePreview() override;
 	ENGINE_API virtual void ForceRecompileForRendering(EMaterialShaderPrecompileMode CompileMode = EMaterialShaderPrecompileMode::Default) override;
+
+	/* returns whether valid compilation output exists. Outputs all default textures used by the material, based on the compilation output. */
+	ENGINE_API virtual bool TryGetUsedTexturesInCookCache(TArray<TObjectPtr<UObject>>& OutTextures) override;
 #endif // WITH_EDITOR
 	//~ End UMaterialInterface Interface.
 
@@ -1930,6 +1930,13 @@ protected:
 #endif
 
 public:
+	/**
+	 * Looks through the ReferencedDefaultTextures array for the texture present at the given Idx.
+	 * At editor time, it looks through all default textures. 
+	 * At runtime, it has access to default textures based on cook output, if ShouldStripUnusedDefaultTextures returns true.
+	 * @return The default texture that is relevant to the given Idx or nullptr if not present.
+	 */
+	ENGINE_API virtual TObjectPtr<UObject> GetDefaultTextureAtIdx(const int32 InIdx, TMicRecursionGuard RecursionGuard = TMicRecursionGuard()) override;
 	ENGINE_API virtual void DumpDebugInfo(FOutputDevice& OutputDevice) const override;
 	void SaveShaderStableKeys(const class ITargetPlatform* TP);
 	ENGINE_API virtual void SaveShaderStableKeysInner(const class ITargetPlatform* TP, const struct FStableShaderKeyAndValue& SaveKeyVal) override;
diff --git a/Engine/Source/Runtime/Engine/Public/Materials/MaterialInstance.h b/Engine/Source/Runtime/Engine/Public/Materials/MaterialInstance.h
--- a/Engine/Source/Runtime/Engine/Public/Materials/MaterialInstance.h
+++ b/Engine/Source/Runtime/Engine/Public/Materials/MaterialInstance.h
@@ -875,10 +875,6 @@ private:
 	 * For example the material instance needs to support being rendered at different quality levels and feature levels within the same process.
 	 */
 	TArray<FMaterialResource*> StaticPermutationMaterialResources;
-#if WITH_EDITOR
-	/** Material resources being cached for cooking. */
-	TMap<const ITargetPlatform*, TArray<FMaterialResourceForCooking>> CachedMaterialResourcesForCooking;
-#endif
 
 #if WITH_EDITOR
 	struct FTextureParameterOverride
@@ -1006,6 +1002,10 @@ public:
 #if WITH_EDITOR
 	/** Allows material properties to be compiled with the option of being overridden by the material attributes input. */
 	ENGINE_API virtual int32 CompilePropertyEx(class FMaterialCompiler* Compiler, const FGuid& AttributeID) override;
+
+	/* returns whether valid compilation output exists. Outputs all default textures used by the material, based on the compilation output. */
+	ENGINE_API virtual bool TryGetUsedTexturesInCookCache(TArray<TObjectPtr<UObject>>& OutTextures) override;
+
 #endif // WITH_EDITOR
 	//~ End UMaterialInterface Interface.
 
@@ -1040,6 +1040,10 @@ public:
 	 * Can also trigger recompile based on new set of FMaterialInstanceBasePropertyOverrides 
 	 */
 	ENGINE_API void UpdateStaticPermutation(const FStaticParameterSet& NewParameters, FMaterialInstanceBasePropertyOverrides& NewBasePropertyOverrides, const bool bForceStaticPermutationUpdate = false, FMaterialUpdateContext* MaterialUpdateContext = nullptr);
+
+	/** Updates all default textures of this material instance, based on permutations and Parent Material. */
+	ENGINE_API void UpdateReferencedDefaultTextures();
+
 	/**
 	* Sets new static parameter overrides on the instance and recompiles the static permutation resources if needed.
 	* Can be passed either a minimal parameter set (overridden parameters only) or the entire set generated by GetStaticParameterValues().
@@ -1157,6 +1161,14 @@ public:
 	 */
 	ENGINE_API virtual void GetLightingGuidChain(bool bIncludeTextures, TArray<FGuid>& OutGuids) const override;
 
+	/**
+	 * Looks through this Instance and all its parents for the relevant texture.
+	 * At editor time, it looks through all default textures. 
+	 * At runtime, it has access to default textures based on cook output, if ShouldStripUnusedDefaultTextures returns true.
+	 * @return The default texture that is relevant to the given Idx or nullptr if not present.
+	 */
+	ENGINE_API virtual TObjectPtr<UObject> GetDefaultTextureAtIdx(const int32 InIdx, TMicRecursionGuard RecursionGuard = TMicRecursionGuard()) override;
+
 	ENGINE_API virtual void DumpDebugInfo(FOutputDevice& OutputDevice) const override;
 	void SaveShaderStableKeys(const class ITargetPlatform* TP);
 	ENGINE_API virtual void SaveShaderStableKeysInner(const class ITargetPlatform* TP, const struct FStableShaderKeyAndValue& SaveKeyVal) override;
@@ -1204,7 +1216,8 @@ protected:
 
 	ENGINE_API bool SetParentInternal(class UMaterialInterface* NewParent, bool RecacheShaders);
 
-	void GetTextureExpressionValues(const FMaterialResource* MaterialResource, TArray<UTexture*>& OutTextures, TArray<TArray<int32>>* OutIndices = nullptr, bool bIncludeTextureCollections = false) const;
+	/** Returns the textures which are used by the Uniform Expressions of the given Material Resource, along with any existing transient overrides. */
+	ENGINE_API virtual void GetTextureExpressionValues(const FMaterialResource* MaterialResource, TArray<UTexture*>& OutTextures, TArray<TArray<int32>>* OutIndices = nullptr, bool bIncludeTextureCollections = false, bool bAllowOverrides = false) const override;
 	void GetTextureCollectionExpressionValues(const FMaterialResource* MaterialResource, TArray<UTextureCollection*>& OutTextureCollections) const;
 
 #if WITH_EDITOR
diff --git a/Engine/Source/Runtime/Engine/Public/Materials/MaterialInterface.h b/Engine/Source/Runtime/Engine/Public/Materials/MaterialInterface.h
--- a/Engine/Source/Runtime/Engine/Public/Materials/MaterialInterface.h
+++ b/Engine/Source/Runtime/Engine/Public/Materials/MaterialInterface.h
@@ -54,7 +54,6 @@ class UMaterialInstance;
 struct FDebugShaderTypeInfo;
 struct FMaterialParameterInfo;
 struct FMaterialResourceLocOnDisk;
-class FMaterialCachedData;
 struct FMaterialCachedExpressionData;
 struct FMaterialCachedExpressionEditorOnlyData;
 #if WITH_EDITOR
@@ -397,6 +396,10 @@ protected:
 	UPROPERTY()
 	TArray<FMaterialTextureInfo> TextureStreamingData;
 
+	/** Default Textures. At runtime, based on cook output, if ShouldStripUnusedDefaultTextures returns true. */
+	UPROPERTY()
+	TArray<TObjectPtr<UObject>> ReferencedDefaultTextures;
+
 	/** Array of user data stored with the asset */
 	UPROPERTY(EditAnywhere, AdvancedDisplay, Instanced, Category = Material)
 	TArray<TObjectPtr<UAssetUserData>> AssetUserData;
@@ -458,6 +461,12 @@ public:
 	UPROPERTY(EditAnywhere, Instanced, Category = ImportSettings)
 	TObjectPtr<class UAssetImportData> AssetImportData;
 	
+protected:
+#if WITH_EDITOR
+		/** Material resources being cached for cooking. */
+		TMap<const ITargetPlatform*, TArray<FMaterialResourceForCooking>> CachedMaterialResourcesForCooking;
+#endif
+
 private:
 	/** Unique ID for this material, used for caching during distributed lighting */
 	UPROPERTY()
@@ -534,6 +543,17 @@ public:
 
 	ENGINE_API virtual const FMaterialCachedExpressionData& GetCachedExpressionData(TMicRecursionGuard RecursionGuard = TMicRecursionGuard()) const;
 
+#if WITH_EDITOR
+	/** @return The Index at which the given texture is present in the editor only Default Textures array. */
+	ENGINE_API int32 GetDefaultTextureIdx(UObject* InTexture) const;
+
+	/** @return Whether the given texture is present in the editor only Default Textures. */
+	ENGINE_API bool ContainsDefaultTexture(UObject* InTexture) const;
+
+	/** @return The array of default textures present in editor. */
+	ENGINE_API const TArray<TObjectPtr<UObject>>& GetDefaultTextures() const;
+#endif
+
 	ENGINE_API bool IsUsingNewHLSLGenerator() const;
 	ENGINE_API bool IsUsingNewTranslatorPrototype() const;
 
@@ -619,6 +639,9 @@ public:
 	ENGINE_API virtual void GetUsedTexturesAndIndices(TArray<UTexture*>& OutTextures, TArray< TArray<int32> >& OutIndices, EMaterialQualityLevel::Type QualityLevel, ERHIFeatureLevel::Type FeatureLevel) const final;
 	ENGINE_API virtual void GetUsedTexturesAndIndices(TArray<UTexture*>& OutTextures, TArray< TArray<int32> >& OutIndices, EMaterialQualityLevel::Type QualityLevel, EShaderPlatform InShaderPlatform) const;
 
+	/** Returns the textures which are used by the Uniform Expressions of the given Material Resource. */
+	ENGINE_API virtual void GetTextureExpressionValues(const FMaterialResource* MaterialResource, TArray<UTexture*>& OutTextures, TArray<TArray<int32>>* OutIndices = nullptr, bool bIncludeTextureCollections = false, bool bAllowOverrides = false) const;
+
 	/**
 	 * Override a specific texture (transient)
 	 *
@@ -687,6 +710,12 @@ public:
 	virtual const FMaterialResource* GetMaterialResource(ERHIFeatureLevel::Type InFeatureLevel, EMaterialQualityLevel::Type QualityLevel = EMaterialQualityLevel::Num) const final { return NULL; }
 	virtual const FMaterialResource* GetMaterialResource(EShaderPlatform InShaderPlatform, EMaterialQualityLevel::Type QualityLevel = EMaterialQualityLevel::Num) const { return NULL; }
 
+	/**
+	 * Get the Material Resources created by the given Material, restricted to the given Quality Range and Shader Platform Range
+	 * @return - The appropriate FMaterialResources, if any exist
+	 */
+	static TArray<const FMaterialResource*, TInlineAllocator<4>> GetMaterialResources(const UMaterialInterface* InMaterial, const FInt32Range InQualityLevelRange, const FInt32Range InShaderPlatformRange);
+
 	/**
 	 * Get the material layers stack
 	 * @return - material layers stack, or nullptr if material has no layers
@@ -820,9 +849,11 @@ public:
 	ENGINE_API virtual void GetReferencedTexturesAndOverrides(TSet<const UTexture*>& InOutTextures) const;
 #endif // WITH_EDITOR
 
-	/** Get textures referenced by expressions, including nested functions. */
 	ENGINE_API TArrayView<const TObjectPtr<UObject>> GetReferencedTextures() const;
 
+	/** @return The default texture relevant to the given Idx. At runtime, if ShouldStripUnusedDefaultTextures returns true, these are present based the compilation output. */
+	virtual TObjectPtr<UObject> GetDefaultTextureAtIdx(const int32 InIdx, TMicRecursionGuard RecursionGuard = TMicRecursionGuard()) PURE_VIRTUAL(UMaterialInterface::GetDefaultTextureAtIdx, return nullptr;);
+
 	ENGINE_API TConstArrayView<TObjectPtr<UTextureCollection>> GetReferencedTextureCollections() const;
 
 	virtual void SaveShaderStableKeysInner(const class ITargetPlatform* TP, const struct FStableShaderKeyAndValue& SaveKeyVal)
@@ -1354,6 +1385,12 @@ public:
 
 	/** Filter out ShadingModels field to a shader platform settings */
 	ENGINE_API static void FilterOutPlatformShadingModels(EShaderPlatform Platform, FMaterialShadingModelField& ShadingModels);
+
+	/** Helper function for building a Range from a given Quality Level */
+	static FInt32Range BuildQualityLevelRange(TOptional<EMaterialQualityLevel::Type> InQualityLevel);
+
+	/** Helper function for building a Range from a given Shader Platform */
+	static FInt32Range BuildShaderPlatformRange(TOptional<EShaderPlatform> InShaderPlatform);
 	
 #if WITH_EDITOR
 	TUniquePtr<FMaterialInsights> MaterialInsight;
@@ -1378,6 +1415,20 @@ protected:
 	/** Bitfield to track which usage flags have triggered warnings so that we only warn once per usage. */
 	uint32 UsageFlagWarnings = 0;
 
+#if WITH_EDITOR
+	/* returns whether unused default textures should be removed during the cooking process. */
+	static bool ShouldStripUnusedDefaultTextures();
+
+	/* returns whether valid compilation output exists. Outputs all default textures used by the material, based on the compilation output. */
+	virtual bool TryGetUsedTexturesInCookCache(TArray<TObjectPtr<UObject>>& OutTextures) PURE_VIRTUAL(UMaterialInterface::TryGetUsedTexturesInCookCache, return false;);
+	
+	/** @return All FMaterialResources created during the cooking process. */
+	TArray<const FMaterialResource*, TInlineAllocator<4>> GetCookedCacheMaterialResources() const;
+
+	/** Removes all unused default textures based on the Cook Shader Compilation output. */
+	void StripUnusedDefaultTextures(const bool bInKeepProjectPrivateTextures);
+#endif
+
 private:
 	/** Bitfield to track game thread dispatch of usage flag checks. */
 	mutable std::atomic<uint32> UsageFlagCacheMiss { 0u };
