[shaders] reset output object when retrying distributed jobs locally; prevents attempting to write to an already finalized shadercode object in the compile step

#Synced-CL 48769550

diff --git a/Engine/Source/Runtime/Engine/Private/ShaderCompiler/ShaderCompilerDistributed.cpp b/Engine/Source/Runtime/Engine/Private/ShaderCompiler/ShaderCompilerDistributed.cpp
--- a/Engine/Source/Runtime/Engine/Private/ShaderCompiler/ShaderCompilerDistributed.cpp
+++ b/Engine/Source/Runtime/Engine/Private/ShaderCompiler/ShaderCompilerDistributed.cpp
@@ -610,6 +610,13 @@ int32 FShaderCompileDistributedThreadRunnable_Interface::CompilingLoop()
 					// Dumping debug info in case the compile job crashes the cooker
 					FShaderCompileInternalUtilities::DumpDebugInfo(*Job);
 
+					Job->ForEachSingleShaderJob(
+						[](FShaderCompileJob& SingleJob)
+						{
+							SingleJob.Output = FShaderCompilerOutput();
+							SingleJob.bFinalized = false;
+						});
+
 					FShaderCompileUtilities::ExecuteShaderCompileJob(*Job);
 				}
 			}
