Iris
* Moved AActor::OnSubobjectDestroyFromReplication trigger from the NetSubObjectFactory to the NetActorFactory

#rnx
#Synced-CL 48743431

diff --git a/Engine/Source/Runtime/Engine/Private/Net/Iris/ReplicationSystem/NetActorFactory.cpp b/Engine/Source/Runtime/Engine/Private/Net/Iris/ReplicationSystem/NetActorFactory.cpp
--- a/Engine/Source/Runtime/Engine/Private/Net/Iris/ReplicationSystem/NetActorFactory.cpp
+++ b/Engine/Source/Runtime/Engine/Private/Net/Iris/ReplicationSystem/NetActorFactory.cpp
@@ -489,6 +489,23 @@ void UNetActorFactory::SubObjectCreatedFromReplication(UE::Net::FNetRefHandle Ro
 	}
 }
 
+void UNetActorFactory::SubObjectDetachedFromReplication(const FDestroyedContext& Context)
+{
+	if (!Context.DestroyedInstance)
+	{
+		return;
+	}
+
+	// Find the subobject owner and trigger the callback. 
+	if (AActor* ActorOwner = Cast<AActor>(Context.RootObject))
+	{
+		if (ensureMsgf(IsValid(ActorOwner) && !ActorOwner->IsUnreachable(), TEXT("Destroyed subobject: %s has an invalid owner: %s"), *GetNameSafe(Context.DestroyedInstance), *GetPathNameSafe(Context.RootObject)))
+		{
+			ActorOwner->OnSubobjectDestroyFromReplication(Context.DestroyedInstance);
+		}
+	}
+}
+
 void UNetActorFactory::DetachedFromReplication(const FDestroyedContext& Context)
 {
 	UEngineReplicationBridge* ActorBridge = CastChecked<UEngineReplicationBridge>(Bridge);
diff --git a/Engine/Source/Runtime/Engine/Private/Net/Iris/ReplicationSystem/NetSubObjectFactory.cpp b/Engine/Source/Runtime/Engine/Private/Net/Iris/ReplicationSystem/NetSubObjectFactory.cpp
--- a/Engine/Source/Runtime/Engine/Private/Net/Iris/ReplicationSystem/NetSubObjectFactory.cpp
+++ b/Engine/Source/Runtime/Engine/Private/Net/Iris/ReplicationSystem/NetSubObjectFactory.cpp
@@ -277,15 +277,7 @@ void UNetSubObjectFactory::DetachedFromReplication(const FDestroyedContext& Cont
 		return;
 	}
 
-	// Find the subobject owner and trigger the callback. 
-    // TODO: This should be moved to NetFactorFactory::SubObjectDetachedFromReplication
-	if (AActor* ActorOwner = Cast<AActor>(Context.RootObject))
-	{
-		if (ensureMsgf(IsValid(ActorOwner) && !ActorOwner->IsUnreachable(), TEXT("UEngineReplicationBridge::DestroyInstanceFromRemote Destroyed subobject: %s has an invalid owner: %s"), *GetNameSafe(Context.DestroyedInstance), *GetPathNameSafe(Context.RootObject)))
-		{
-			ActorOwner->OnSubobjectDestroyFromReplication(Context.DestroyedInstance);
-		}
-	}
+	
 
 	Context.DestroyedInstance->PreDestroyFromReplication();
 	Context.DestroyedInstance->MarkAsGarbage();
diff --git a/Engine/Source/Runtime/Engine/Public/Net/Iris/ReplicationSystem/NetActorFactory.h b/Engine/Source/Runtime/Engine/Public/Net/Iris/ReplicationSystem/NetActorFactory.h
--- a/Engine/Source/Runtime/Engine/Public/Net/Iris/ReplicationSystem/NetActorFactory.h
+++ b/Engine/Source/Runtime/Engine/Public/Net/Iris/ReplicationSystem/NetActorFactory.h
@@ -187,6 +187,7 @@ public:
 	ENGINE_API virtual void PostInit(const FPostInitContext& Context) override;
 
 	ENGINE_API virtual void SubObjectCreatedFromReplication(UE::Net::FNetRefHandle RootObject, UE::Net::FNetRefHandle SubObjectCreated) override;
+	ENGINE_API virtual void SubObjectDetachedFromReplication(const FDestroyedContext& Context) override;
 
 	ENGINE_API virtual void DetachedFromReplication(const FDestroyedContext& Context) override;
 
