Add NetDriver null check inside GetReplicationSystem to fix potential null crash

#tests Built and tested against Dedicated Server Cooked Data
#rnx

diff --git a/Engine/Source/Runtime/Engine/Private/Net/Iris/ReplicationSystem/ReplicationSystemUtil.cpp b/Engine/Source/Runtime/Engine/Private/Net/Iris/ReplicationSystem/ReplicationSystemUtil.cpp
--- a/Engine/Source/Runtime/Engine/Private/Net/Iris/ReplicationSystem/ReplicationSystemUtil.cpp
+++ b/Engine/Source/Runtime/Engine/Private/Net/Iris/ReplicationSystem/ReplicationSystemUtil.cpp
@@ -238,7 +238,10 @@ UReplicationSystem* FReplicationSystemUtil::GetReplicationSystem(const AActor* A
 
 	if (const UNetConnection* NetConnection = Actor->GetNetConnection())
 	{
-		return NetConnection->Driver->GetReplicationSystem();
+		if (NetConnection->Driver)
+		{
+			return NetConnection->Driver->GetReplicationSystem();
+		}
 	}
 
 	UNetDriver* NetDriver = Actor->GetNetDriver();
