Add cvar to control if the rewind/resim is being executed before or after the advance.

#Synced-CL 48908224

diff --git a/Engine/Source/Runtime/Engine/Private/PhysicsEngine/NetworkPhysicsComponent.cpp b/Engine/Source/Runtime/Engine/Private/PhysicsEngine/NetworkPhysicsComponent.cpp
--- a/Engine/Source/Runtime/Engine/Private/PhysicsEngine/NetworkPhysicsComponent.cpp
+++ b/Engine/Source/Runtime/Engine/Private/PhysicsEngine/NetworkPhysicsComponent.cpp
@@ -142,6 +142,11 @@ namespace PhysicsReplicationCVars
 	}
 }
 
+namespace Chaos
+{
+	extern CHAOS_API int32 RewindBeforeAdvance;
+}
+
 bool FNetworkPhysicsRewindDataProxy::NetSerializeBase(FArchive& Ar, class UPackageMap* Map, bool& bOutSuccess, TUniqueFunction<TUniquePtr<Chaos::FBaseRewindHistory> ()> CreateHistoryFunction , TUniqueFunction<FNetworkPhysicsData* (const int32)> GetDeltaSourceData)
 {
 	bDeltaSerializationIssue = false;
@@ -388,6 +393,11 @@ void FNetworkPhysicsCallback::PostResimStep_Internal(int32 PhysicsStep)
 
 }
 
+void FNetworkPhysicsCallback::AddResimulationRequest_Internal(const int32 PhysicsStep)
+{
+	AddResimulationRequestInternal.Broadcast(PhysicsStep);
+}
+
 int32 FNetworkPhysicsCallback::TriggerRewindIfNeeded_Internal(int32 LatestStepCompleted)
 {
 	int32 ResimFrame = INDEX_NONE;
@@ -1817,6 +1827,7 @@ void FAsyncNetworkPhysicsComponent::OnInitialize_Internal()
 			{
 				DelegateOnPreProcessInputs_Internal = SolverCallback->PreProcessInputsInternal.AddRaw(this, &FAsyncNetworkPhysicsComponent::OnPreProcessInputs_Internal);
 				DelegateOnPostProcessInputs_Internal = SolverCallback->PostProcessInputsInternal.AddRaw(this, &FAsyncNetworkPhysicsComponent::OnPostProcessInputs_Internal);
+				DelegateOnAddResimulationRequest_Internal = SolverCallback->AddResimulationRequestInternal.AddRaw(this, &FAsyncNetworkPhysicsComponent::OnAddResimulationRequest_Internal);
 			}
 		}
 		else
@@ -1838,6 +1849,9 @@ void FAsyncNetworkPhysicsComponent::OnUninitialize_Internal()
 
 			SolverCallback->PostProcessInputsInternal.Remove(DelegateOnPostProcessInputs_Internal);
 			DelegateOnPostProcessInputs_Internal.Reset();
+
+			SolverCallback->AddResimulationRequestInternal.Remove(DelegateOnAddResimulationRequest_Internal);
+			DelegateOnAddResimulationRequest_Internal.Reset();
 		}
 	}
 
@@ -1858,7 +1872,7 @@ const FNetworkPhysicsSettingsNetworkPhysicsComponent& FAsyncNetworkPhysicsCompon
 	return SettingsComponent.IsValid() ? SettingsComponent.Pin()->NetworkPhysicsComponentSettings : SettingsNetworkPhysicsComponent_Default;
 };
 
-void FAsyncNetworkPhysicsComponent::ConsumeAsyncInput(const int32 PhysicsStep)
+void FAsyncNetworkPhysicsComponent::ConsumeAsyncInput_Internal(const int32 PhysicsStep, const bool bTriggerResim)
 {
 	if (const FAsyncNetworkPhysicsComponentInput* AsyncInput = GetConsumerInput_Internal())
 	{
@@ -1981,36 +1995,36 @@ void FAsyncNetworkPhysicsComponent::ConsumeAsyncInput(const int32 PhysicsStep)
 
 			/** Receive data helper */
 			auto ReceiveHelper = [&](Chaos::FBaseRewindHistory* History, Chaos::FBaseRewindHistory* ReceiveData, const bool bImportant, const bool bCompareData)
-			{
-				const bool bCompareDataForRewind = bCompareData && !bIsServer;
-				const int32 TryInjectAtFrame = bIsServer ? PhysicsStep : 0;
-				const int32 ResimFrame = History->ReceiveNewData(*ReceiveData, (bIsServer ? 0 : NetworkPhysicsTickOffset), bCompareDataForRewind, bImportant, TryInjectAtFrame);
-				if (bCompareDataForRewind)
 				{
-					TriggerResimulation(ResimFrame);
-				}
+					const bool bCompareDataForRewind = bCompareData && !bIsServer && bTriggerResim;
+					const int32 TryInjectAtFrame = bIsServer ? PhysicsStep : 0;
+					const int32 ResimFrame = History->ReceiveNewData(*ReceiveData, (bIsServer ? 0 : NetworkPhysicsTickOffset), bCompareDataForRewind, bImportant, TryInjectAtFrame);
+					if (bCompareDataForRewind)
+					{
+						TriggerResimulation(ResimFrame);
+					}
 
 #if !UE_BUILD_SHIPPING
-				if (PhysicsReplicationCVars::ResimulationCVars::bDebugTriggerResimEveryNFrames > 0 && !bIsServer && IsLocallyControlled())
-				{
-					int32 LatestReceivedFrame = ReceiveData->GetLatestFrame() - NetworkPhysicsTickOffset;
-					if (PhysicsStep % PhysicsReplicationCVars::ResimulationCVars::bDebugTriggerResimEveryNFrames == 0)
+					if (PhysicsReplicationCVars::ResimulationCVars::bDebugTriggerResimEveryNFrames > 0 && !bIsServer && IsLocallyControlled() && bTriggerResim)
 					{
-						TriggerResimulation(LatestReceivedFrame);
+						int32 LatestReceivedFrame = ReceiveData->GetLatestFrame() - NetworkPhysicsTickOffset;
+						if (PhysicsStep % PhysicsReplicationCVars::ResimulationCVars::bDebugTriggerResimEveryNFrames == 0)
+						{
+							TriggerResimulation(LatestReceivedFrame);
+						}
 					}
-				}
 #endif // !UE_BUILD_SHIPPING
 
 #if DEBUG_NETWORK_PHYSICS
-				{
-					FString NetRoleString = bIsServer ? FString("SERVER") : (IsLocallyControlled() ? FString("AUTONO") : FString("PROXY "));
-					ReceiveData->DebugData(FString::Printf(TEXT("%s | PT | RECEIVE DATA | LatestFrame: %d | bImportant: %d | Name: %s"), *NetRoleString, ReceiveData->GetLatestFrame(), bImportant, *GetActorName()));
-				}
+					{
+						FString NetRoleString = bIsServer ? FString("SERVER") : (IsLocallyControlled() ? FString("AUTONO") : FString("PROXY "));
+						ReceiveData->DebugData(FString::Printf(TEXT("%s | PT | RECEIVE DATA | LatestFrame: %d | bImportant: %d | Name: %s"), *NetRoleString, ReceiveData->GetLatestFrame(), bImportant, *GetActorName()));
+					}
 #endif
 
-				// Reset the received data after having consumed it
-				ReceiveData->ResetFast();
-			};
+					// Reset the received data after having consumed it
+					ReceiveData->ResetFast();
+				};
 
 			const bool bCompareInput = ComponentSettings.GetCompareInputToTriggerRewind(bCompareInputToTriggerRewind) && IsLocallyControlled();
 			const bool bCompareState = ComponentSettings.GetCompareStateToTriggerRewind(bCompareStateToTriggerRewind) && (IsLocallyControlled() || ComponentSettings.GetCompareStateToTriggerRewindIncludeSimProxies(bCompareStateToTriggerRewindIncludeSimProxies));
@@ -2036,9 +2050,9 @@ void FAsyncNetworkPhysicsComponent::ConsumeAsyncInput(const int32 PhysicsStep)
 						}
 					}
 				}
-				
+
 				// If setting is true, request resimulation for sim-proxies when receiving inputs that are newer than latest already received input
-				if (!bIsServer && ComponentSettings.GetTriggerResimOnInputReceive() && !IsLocallyControlled())
+				if (!bIsServer && ComponentSettings.GetTriggerResimOnInputReceive() && !IsLocallyControlled() && bTriggerResim)
 				{
 					const int32 LatestReceiveFrame = AsyncInput->InputData.Get()->GetLatestFrame() - NetworkPhysicsTickOffset;
 					const int32 NextInputFrame = InputHistory.Get()->GetLatestFrame() + 1;
@@ -2118,6 +2132,14 @@ FAsyncNetworkPhysicsComponentOutput& FAsyncNetworkPhysicsComponent::GetAsyncOutp
 	return AsyncOutput;
 }
 
+void FAsyncNetworkPhysicsComponent::OnAddResimulationRequest_Internal(const int32 PhysicsStep)
+{
+	if(Chaos::RewindBeforeAdvance != 0)
+	{ 
+		ConsumeAsyncInput_Internal(PhysicsStep, true);
+	}
+}
+
 void FAsyncNetworkPhysicsComponent::OnPreProcessInputs_Internal(const int32 PhysicsStep)
 {
 	QUICK_SCOPE_CYCLE_COUNTER(AsyncNetworkPhysicsComponent_OnPreProcessInputs_Internal);
@@ -2129,10 +2151,10 @@ void FAsyncNetworkPhysicsComponent::OnPreProcessInputs_Internal(const int32 Phys
 		bIsSolverResim = Evolution->IsResimming();
 		bIsFirstResimFrame = Evolution->IsResetting();
 	}
-
+	
 	if (!bIsSolverResim)
 	{
-		ConsumeAsyncInput(PhysicsStep);
+		ConsumeAsyncInput_Internal(PhysicsStep, Chaos::RewindBeforeAdvance == 0);
 	}
 
 	const FNetworkPhysicsSettingsNetworkPhysicsComponent& ComponentSettings = GetComponentSettings();
@@ -2149,6 +2171,7 @@ void FAsyncNetworkPhysicsComponent::OnPreProcessInputs_Internal(const int32 Phys
 	{
 		// Apply replicated state on clients if we are resimulating or on simulated proxies if setting is enabled
 		const bool bApplySimProxyState = ComponentSettings.GetApplySimProxyStateAtRuntime() && !bIsServer && !IsLocallyControlled();
+
 		if ((bApplySimProxyState || bIsSolverResim) && StateHistory && StateData)
 		{
 			FNetworkPhysicsPayload* PhysicsData = StateData.Get();
@@ -2191,6 +2214,7 @@ void FAsyncNetworkPhysicsComponent::OnPreProcessInputs_Internal(const int32 Phys
 		// Apply replicated inputs on server and simulated proxies if setting is enabled, and on local player if we are resimulating
 		const bool bApplySimProxyInput = ComponentSettings.GetApplySimProxyInputAtRuntime() && !bIsServer && !IsLocallyControlled();
 		const bool bApplyServerInput = bIsServer && !IsLocallyControlled();
+
 		if ((bApplyServerInput || bApplySimProxyInput || bIsSolverResim) && InputHistory && InputData)
 		{
 			FNetworkPhysicsPayload* PhysicsData = InputData.Get();
diff --git a/Engine/Source/Runtime/Engine/Private/PhysicsEngine/PhysicsReplication.cpp b/Engine/Source/Runtime/Engine/Private/PhysicsEngine/PhysicsReplication.cpp
--- a/Engine/Source/Runtime/Engine/Private/PhysicsEngine/PhysicsReplication.cpp
+++ b/Engine/Source/Runtime/Engine/Private/PhysicsEngine/PhysicsReplication.cpp
@@ -292,6 +292,11 @@ namespace PhysicsReplicationCVars
 
 }
 
+namespace Chaos
+{
+	extern CHAOS_API int32 RewindBeforeAdvance;
+}
+
 
 FPhysicsReplication::FPhysicsReplication(FPhysScene* InPhysicsScene)
 	: PhysScene(InPhysicsScene)
@@ -924,6 +929,68 @@ void FPhysicsReplicationAsync::OnPostInitialize_Internal()
 	RigidsSolver.SetPhysicsReplication_Internal(this);
 }
 
+void FPhysicsReplicationAsync::AddResimulationRequest_Internal(const float DeltaSeconds)
+{
+	if (FPhysicsReplication::ShouldSkipPhysicsReplication())
+	{
+		return;
+	}
+
+	if (const FPhysicsReplicationAsyncInput* AsyncInput = GetConsumerInput_Internal())
+	{
+		Chaos::FPBDRigidsSolver* RigidsSolver = static_cast<Chaos::FPBDRigidsSolver*>(GetSolver());
+		check(RigidsSolver);
+
+		Chaos::FReadPhysicsObjectInterface_Internal Interface = Chaos::FPhysicsObjectInternalInterface::GetRead();
+		for(const FPhysicsRepAsyncInputData& InputData : AsyncInput->InputData)
+		{
+			if( InputData.TargetState.Flags != ERigidBodyFlags::None)
+			{
+				Chaos::FRewindData* RewindData = RigidsSolver->GetRewindData();
+				if (!RewindData && InputData.RepMode == EPhysicsReplicationMode::Resimulation)
+				{
+					// We don't have rewind data but an actor is set to replicate using resimulation; we need to enable rewind capture.
+					if (ensure(Chaos::FPBDRigidsSolver::IsNetworkPhysicsPredictionEnabled() && RigidsSolver->IsUsingFixedDt()))
+					{
+						RigidsSolver->EnableRewindCapture();
+						RewindData = RigidsSolver->GetRewindData();
+					}
+				}
+				UpdateRewindDataTarget(InputData);
+
+				if(InputData.RepMode == EPhysicsReplicationMode::Resimulation)
+				{ 
+					if(RewindData)
+					{ 
+						if(Chaos::FGeometryParticleHandle* GeometryHandle = Interface.GetParticle(InputData.PhysicsObject))
+						{
+							if(Chaos::FPBDRigidParticleHandle* RigidHandle = GeometryHandle->CastToRigidParticle())
+							{
+								FReplicatedPhysicsTargetAsync Target;
+
+								// Update target data required to trigger resim from input
+								Target.ServerFrame = InputData.ServerFrame;
+								Target.TargetState = InputData.TargetState;
+								Target.FrameOffset = InputData.FrameOffset.IsSet() ? *InputData.FrameOffset : 0;
+
+								if (IsTargetValidForResim(Target))
+								{
+									const int32 LocalFrame = Target.ServerFrame - Target.FrameOffset;
+									if (ShouldResimulateParticle(RigidHandle, Target, DeltaSeconds) && LocalFrame > RewindData->GetBlockedResimFrame())
+									{
+										// Request resimulation
+										RewindData->RequestResimulation(LocalFrame, GeometryHandle);
+									}
+								}
+							}
+						}
+					}
+				}
+			}
+		}
+	}
+}
+
 void FPhysicsReplicationAsync::OnPreSimulate_Internal()
 {
 	if (FPhysicsReplication::ShouldSkipPhysicsReplication())
@@ -970,16 +1037,19 @@ void FPhysicsReplicationAsync::OnPreSimulate_Internal()
 				continue;
 			}
 
-			if (!bRewindDataExist && Input.RepMode == EPhysicsReplicationMode::Resimulation)
-			{
-				// We don't have rewind data but an actor is set to replicate using resimulation; we need to enable rewind capture.
-				if (ensure(Chaos::FPBDRigidsSolver::IsNetworkPhysicsPredictionEnabled() && RigidsSolver->IsUsingFixedDt()))
+			if(Chaos::RewindBeforeAdvance == 0)
+			{ 
+				if (!bRewindDataExist && Input.RepMode == EPhysicsReplicationMode::Resimulation)
 				{
-					RigidsSolver->EnableRewindCapture();
+					// We don't have rewind data but an actor is set to replicate using resimulation; we need to enable rewind capture.
+					if (ensure(Chaos::FPBDRigidsSolver::IsNetworkPhysicsPredictionEnabled() && RigidsSolver->IsUsingFixedDt()))
+					{
+						RigidsSolver->EnableRewindCapture();
+					}
 				}
-			}
 
-			UpdateRewindDataTarget(Input);
+				UpdateRewindDataTarget(Input);
+			}
 			UpdateAsyncTarget(Input, RigidsSolver);
 			
 			DebugDrawReplicationMode(Input);
@@ -2257,246 +2327,281 @@ void FPhysicsReplicationAsync::ExtrapolateTarget(FReplicatedPhysicsTargetAsync&
 	Target.TargetState.Quaternion = TargetRotExtrapDelta * Target.TargetState.Quaternion;
 }
 
-/** Compare states and trigger resimulation if needed */
-bool FPhysicsReplicationAsync::ResimulationReplication(Chaos::FPBDRigidParticleHandle* Handle, FReplicatedPhysicsTargetAsync& Target, const float DeltaSeconds)
+bool FPhysicsReplicationAsync::IsTargetValidForResim(const FReplicatedPhysicsTargetAsync& Target) const
 {
-	Chaos::FPBDRigidsSolver* RigidsSolver = static_cast<Chaos::FPBDRigidsSolver*>(GetSolver());
+	const Chaos::FPBDRigidsSolver* RigidsSolver = static_cast<const Chaos::FPBDRigidsSolver*>(GetSolver());
 	if (RigidsSolver == nullptr)
 	{
-		return true;
+		return false;
 	}
 
-	Chaos::FRewindData* RewindData = RigidsSolver->GetRewindData();
+	const Chaos::FRewindData* RewindData = RigidsSolver->GetRewindData();
 	if (RewindData == nullptr)
 	{
-		return true;
+		return false;
 	}
 
 	if (Target.ServerFrame <= 0)
 	{
-		return true;
+		return false;
 	}
 
 	const int32 LocalFrame = Target.ServerFrame - Target.FrameOffset;
 
 	if (!RewindData->IsFrameWithinRewindHistory(LocalFrame))
 	{
-		return true;
+		return false;
 	}
+	return true;
+}
 
-	const bool bShouldSleep = (Target.TargetState.Flags & ERigidBodyFlags::Sleeping) != 0;
-	bool bClearTarget = true;
-
-	static constexpr Chaos::FFrameAndPhase::EParticleHistoryPhase RewindPhase = Chaos::FFrameAndPhase::EParticleHistoryPhase::PrePushData;
-	
-	// Get state from locally cached history for frame corresponding to received data
-	const Chaos::FGeometryParticleState PastState = RewindData->GetPastStateAtFrame(*Handle, LocalFrame, RewindPhase);
-
-	// Check which comparisons to perform to trigger resimulation from
-	const bool bCompareX = Chaos::FPhysicsSolverBase::GetResimulationErrorPositionThresholdEnabled() || SettingsCurrent.ResimulationSettings.bOverrideResimulationErrorPositionThreshold;
-	const bool bCompareR = Chaos::FPhysicsSolverBase::GetResimulationErrorRotationThresholdEnabled() || SettingsCurrent.ResimulationSettings.bOverrideResimulationErrorRotationThreshold;
-	const bool bCompareV = Chaos::FPhysicsSolverBase::GetResimulationErrorLinearVelocityThresholdEnabled() || SettingsCurrent.ResimulationSettings.bOverrideResimulationErrorLinearVelocityThreshold;
-	const bool bCompareW = Chaos::FPhysicsSolverBase::GetResimulationErrorAngularVelocityThresholdEnabled() || SettingsCurrent.ResimulationSettings.bOverrideResimulationErrorAngularVelocityThreshold;
+bool FPhysicsReplicationAsync::ShouldResimulateParticle(const Chaos::FPBDRigidParticleHandle* Handle, const FReplicatedPhysicsTargetAsync& Target, const float DeltaSeconds) const
+{
 	bool bShouldTriggerResim = false;
+	if (IsTargetValidForResim(Target))
+	{
+		const Chaos::FRewindData* RewindData = static_cast<const Chaos::FPBDRigidsSolver*>(GetSolver())->GetRewindData();
+		const int32 LocalFrame = Target.ServerFrame - Target.FrameOffset;
+		static constexpr Chaos::FFrameAndPhase::EParticleHistoryPhase RewindPhase = Chaos::FFrameAndPhase::EParticleHistoryPhase::PrePushData;
 
-#if !(UE_BUILD_SHIPPING || UE_BUILD_TEST)
-	// Debugging
-	FColor DebugColor = FColor::Black;
-	bool bResimV = false;
-	bool bResimW = false;
-#endif
+		// Get state from locally cached history for frame corresponding to received data
+		const Chaos::FGeometryParticleState PastState = RewindData->GetPastStateAtFrame(*Handle, LocalFrame, RewindPhase);
 
-	// Check for positional discrepancy in Distance between client and server
-	if (bCompareX)
-	{
-		const float ResimPositionErrorThreshold = SettingsCurrent.ResimulationSettings.GetResimulationErrorPositionThreshold(Chaos::FPhysicsSolverBase::GetResimulationErrorPositionThreshold());
-		bShouldTriggerResim = Chaos::FRewindData::CheckVectorThreshold(Target.TargetState.Position, PastState.GetX(), ResimPositionErrorThreshold);
-#if !(UE_BUILD_SHIPPING || UE_BUILD_TEST)
-		if (bShouldTriggerResim)
+		// Check which comparisons to perform to trigger resimulation from
+		const bool bCompareX = Chaos::FPhysicsSolverBase::GetResimulationErrorPositionThresholdEnabled() || SettingsCurrent.ResimulationSettings.bOverrideResimulationErrorPositionThreshold;
+		const bool bCompareR = Chaos::FPhysicsSolverBase::GetResimulationErrorRotationThresholdEnabled() || SettingsCurrent.ResimulationSettings.bOverrideResimulationErrorRotationThreshold;
+		const bool bCompareV = Chaos::FPhysicsSolverBase::GetResimulationErrorLinearVelocityThresholdEnabled() || SettingsCurrent.ResimulationSettings.bOverrideResimulationErrorLinearVelocityThreshold;
+		const bool bCompareW = Chaos::FPhysicsSolverBase::GetResimulationErrorAngularVelocityThresholdEnabled() || SettingsCurrent.ResimulationSettings.bOverrideResimulationErrorAngularVelocityThreshold;
+
+	#if !(UE_BUILD_SHIPPING || UE_BUILD_TEST)
+		// Debugging
+		FColor DebugColor = FColor::Black;
+		bool bResimV = false;
+		bool bResimW = false;
+	#endif
+
+		// Check for positional discrepancy in Distance between client and server
+		if (bCompareX)
 		{
-			DebugColor = FColor::Orange;
+			const float ResimPositionErrorThreshold = SettingsCurrent.ResimulationSettings.GetResimulationErrorPositionThreshold(Chaos::FPhysicsSolverBase::GetResimulationErrorPositionThreshold());
+			bShouldTriggerResim = Chaos::FRewindData::CheckVectorThreshold(Target.TargetState.Position, PastState.GetX(), ResimPositionErrorThreshold);
+	#if !(UE_BUILD_SHIPPING || UE_BUILD_TEST)
+			if (bShouldTriggerResim)
+			{
+				DebugColor = FColor::Orange;
+			}
+	#endif
 		}
-#endif
-	}
 
-	// Check for linear velocity discrepancy in Distance / s between client and server
-	if (!bShouldTriggerResim && bCompareV)
-	{
-		const float ResimLinVelocityErrorThreshold = SettingsCurrent.ResimulationSettings.GetResimulationErrorLinearVelocityThreshold(Chaos::FPhysicsSolverBase::GetResimulationErrorLinearVelocityThreshold());
-		bShouldTriggerResim = Chaos::FRewindData::CheckVectorThreshold(Target.TargetState.LinVel, PastState.GetV(), ResimLinVelocityErrorThreshold);
-#if !(UE_BUILD_SHIPPING || UE_BUILD_TEST)
-		if (bShouldTriggerResim)
+		// Check for linear velocity discrepancy in Distance / s between client and server
+		if (!bShouldTriggerResim && bCompareV)
 		{
-			bResimV = true;
+			const float ResimLinVelocityErrorThreshold = SettingsCurrent.ResimulationSettings.GetResimulationErrorLinearVelocityThreshold(Chaos::FPhysicsSolverBase::GetResimulationErrorLinearVelocityThreshold());
+			bShouldTriggerResim = Chaos::FRewindData::CheckVectorThreshold(Target.TargetState.LinVel, PastState.GetV(), ResimLinVelocityErrorThreshold);
+	#if !(UE_BUILD_SHIPPING || UE_BUILD_TEST)
+			if (bShouldTriggerResim)
+			{
+				bResimV = true;
+			}
+	#endif
 		}
-#endif
-	}
-
-	// Check for angular velocity discrepancy in Degrees / s between client and server
-	if (!bShouldTriggerResim && bCompareW)
-	{
-		const float ResimAngVelocityErrorThreshold = SettingsCurrent.ResimulationSettings.GetResimulationErrorAngularVelocityThreshold(Chaos::FPhysicsSolverBase::GetResimulationErrorAngularVelocityThreshold());
-		bShouldTriggerResim = Chaos::FRewindData::CheckVectorThreshold(Target.TargetState.AngVel, FMath::RadiansToDegrees(PastState.GetW()), ResimAngVelocityErrorThreshold);
-#if !(UE_BUILD_SHIPPING || UE_BUILD_TEST)
-		 if (bShouldTriggerResim)
-		 {
-			 bResimW = true;
-		 }
-#endif
-	}
 
-	// Check for rotational discrepancy in Degrees between client and server
-	if (!bShouldTriggerResim && bCompareR)
-	{
-		const float ResimRotationErrorThreshold = SettingsCurrent.ResimulationSettings.GetResimulationErrorRotationThreshold(Chaos::FPhysicsSolverBase::GetResimulationErrorRotationThreshold());
-		bShouldTriggerResim = Chaos::FRewindData::CheckQuaternionThreshold(Target.TargetState.Quaternion, PastState.GetR(), ResimRotationErrorThreshold);
-#if !(UE_BUILD_SHIPPING || UE_BUILD_TEST)
-		if (bShouldTriggerResim)
+		// Check for angular velocity discrepancy in Degrees / s between client and server
+		if (!bShouldTriggerResim && bCompareW)
 		{
-			DebugColor = FColor::Magenta;
+			const float ResimAngVelocityErrorThreshold = SettingsCurrent.ResimulationSettings.GetResimulationErrorAngularVelocityThreshold(Chaos::FPhysicsSolverBase::GetResimulationErrorAngularVelocityThreshold());
+			bShouldTriggerResim = Chaos::FRewindData::CheckVectorThreshold(Target.TargetState.AngVel, FMath::RadiansToDegrees(PastState.GetW()), ResimAngVelocityErrorThreshold);
+	#if !(UE_BUILD_SHIPPING || UE_BUILD_TEST)
+			if (bShouldTriggerResim)
+			{
+				bResimW = true;
+			}
+	#endif
 		}
-#endif
-	}
 
-#if !(UE_BUILD_SHIPPING || UE_BUILD_TEST)
+		// Check for rotational discrepancy in Degrees between client and server
+		if (!bShouldTriggerResim && bCompareR)
+		{
+			const float ResimRotationErrorThreshold = SettingsCurrent.ResimulationSettings.GetResimulationErrorRotationThreshold(Chaos::FPhysicsSolverBase::GetResimulationErrorRotationThreshold());
+			bShouldTriggerResim = Chaos::FRewindData::CheckQuaternionThreshold(Target.TargetState.Quaternion, PastState.GetR(), ResimRotationErrorThreshold);
+	#if !(UE_BUILD_SHIPPING || UE_BUILD_TEST)
+			if (bShouldTriggerResim)
+			{
+				DebugColor = FColor::Magenta;
+			}
+	#endif
+		}
 
-	if (Chaos::FPhysicsSolverBase::CanDebugNetworkPhysicsPrediction())
-	{
-		UE_LOG(LogPhysics, Log, TEXT("Apply Rigid body state at local frame %d with offset = %d"), LocalFrame, Target.FrameOffset);
-		UE_LOG(LogPhysics, Log, TEXT("Should Trigger Resim = %s | Server Frame = %d | Client Frame = %d"), (bShouldTriggerResim ? TEXT("True") : TEXT("False")), Target.ServerFrame, LocalFrame);
-		UE_LOG(LogPhysics, Log, TEXT("Particle Target Position = %s | Current Position = %s"), *Target.TargetState.Position.ToString(), *PastState.GetX().ToString());
-		UE_LOG(LogPhysics, Log, TEXT("Particle Target Velocity = %s | Current Velocity = %s"), *Target.TargetState.LinVel.ToString(), *PastState.GetV().ToString());
-		UE_LOG(LogPhysics, Log, TEXT("Particle Target Quaternion = %s | Current Quaternion = %s"), *Target.TargetState.Quaternion.ToString(), *PastState.GetR().ToString());
-		UE_LOG(LogPhysics, Log, TEXT("Particle Target Omega = %s | Current Omega= %s"), *Target.TargetState.AngVel.ToString(), *PastState.GetW().ToString());
-	}
+	#if !(UE_BUILD_SHIPPING || UE_BUILD_TEST)
 
-	if (PhysicsReplicationCVars::ResimulationCVars::bDrawDebug || RenderInterpolationCVars::bRenderInterpDebugDrawResimTrigger)
-	{
-		if (bShouldTriggerResim)
+		if (Chaos::FPhysicsSolverBase::CanDebugNetworkPhysicsPrediction())
 		{
-			FVector Box = RenderInterpolationCVars::bRenderInterpDebugDrawResimTrigger ? FVector(6, 3, 2) : FVector(40, 20, 10);
-			Box *= RenderInterpolationCVars::RenderInterpDebugDrawResimBoxScale;
-			const float DrawThickness = (RenderInterpolationCVars::bRenderInterpDebugDrawResimTrigger ? 0.5f : 1.5f) * RenderInterpolationCVars::RenderInterpDebugDrawResimBoxScale;
+			UE_LOG(LogPhysics, Log, TEXT("Apply Rigid body state at local frame %d with offset = %d"), LocalFrame, Target.FrameOffset);
+			UE_LOG(LogPhysics, Log, TEXT("Should Trigger Resim = %s | Server Frame = %d | Client Frame = %d"), (bShouldTriggerResim ? TEXT("True") : TEXT("False")), Target.ServerFrame, LocalFrame);
+			UE_LOG(LogPhysics, Log, TEXT("Particle Target Position = %s | Current Position = %s"), *Target.TargetState.Position.ToString(), *PastState.GetX().ToString());
+			UE_LOG(LogPhysics, Log, TEXT("Particle Target Velocity = %s | Current Velocity = %s"), *Target.TargetState.LinVel.ToString(), *PastState.GetV().ToString());
+			UE_LOG(LogPhysics, Log, TEXT("Particle Target Quaternion = %s | Current Quaternion = %s"), *Target.TargetState.Quaternion.ToString(), *PastState.GetR().ToString());
+			UE_LOG(LogPhysics, Log, TEXT("Particle Target Omega = %s | Current Omega= %s"), *Target.TargetState.AngVel.ToString(), *PastState.GetW().ToString());
+		}
 
-			if (RenderInterpolationCVars::bRenderInterpDebugDrawResimTrigger) // Resim debug draw extension for render interpolation 
+		if (PhysicsReplicationCVars::ResimulationCVars::bDrawDebug || RenderInterpolationCVars::bRenderInterpDebugDrawResimTrigger)
+		{
+			if (bShouldTriggerResim)
 			{
-				Chaos::FDebugDrawQueue::GetInstance().DrawDebugBox(PastState.GetX(), Box, PastState.GetR(), FColor::White, false, PhysicsReplicationCVars::DebugDrawLifeTime, 0, DrawThickness);
-				Chaos::FDebugDrawQueue::GetInstance().DrawDebugBox(Target.TargetState.Position, Box, Target.TargetState.Quaternion, DebugColor, false, PhysicsReplicationCVars::DebugDrawLifeTime, 0, DrawThickness);
-
-				Chaos::FDebugDrawQueue::GetInstance().DrawDebugDirectionalArrow(Handle->GetX(), PastState.GetX(), 5.0f, FColor::White, false, PhysicsReplicationCVars::DebugDrawLifeTime, 0, DrawThickness);
-				Chaos::FDebugDrawQueue::GetInstance().DrawDebugDirectionalArrow(PastState.GetX(), Target.TargetState.Position, 5.0f, FColor::Black, false, PhysicsReplicationCVars::DebugDrawLifeTime, 0, DrawThickness);
+				FVector Box = RenderInterpolationCVars::bRenderInterpDebugDrawResimTrigger ? FVector(6, 3, 2) : FVector(40, 20, 10);
+				Box *= RenderInterpolationCVars::RenderInterpDebugDrawResimBoxScale;
+				const float DrawThickness = (RenderInterpolationCVars::bRenderInterpDebugDrawResimTrigger ? 0.5f : 1.5f) * RenderInterpolationCVars::RenderInterpDebugDrawResimBoxScale;
 
-				if (bResimV)
-				{
-					const FVector DiffV = Target.TargetState.LinVel - PastState.GetV();
-					Chaos::FDebugDrawQueue::GetInstance().DrawDebugDirectionalArrow(Target.TargetState.Position, Target.TargetState.Position + DiffV, 5.0f, FColor::Orange, false, PhysicsReplicationCVars::DebugDrawLifeTime, 0, DrawThickness);
-				}
-				if (bResimW)
+				if (RenderInterpolationCVars::bRenderInterpDebugDrawResimTrigger) // Resim debug draw extension for render interpolation 
 				{
-					const FVector DiffW = Target.TargetState.AngVel - FMath::RadiansToDegrees(PastState.GetW());
-					Chaos::FDebugDrawQueue::GetInstance().DrawDebugDirectionalArrow(Target.TargetState.Position + DiffW, Target.TargetState.Position, 5.0f, FColor::Magenta, false, PhysicsReplicationCVars::DebugDrawLifeTime, 0, DrawThickness);
-				}
-			}
-			else // Resim trigger debug draw
-			{
-				Chaos::FDebugDrawQueue::GetInstance().DrawDebugBox(Handle->GetX(), Box, PastState.GetR(), FColor::White, false, PhysicsReplicationCVars::DebugDrawLifeTime, 0, DrawThickness);
-				Chaos::FDebugDrawQueue::GetInstance().DrawDebugBox(Handle->GetX() + (Target.TargetState.Position - PastState.GetX()), Box, Target.TargetState.Quaternion, DebugColor, false, PhysicsReplicationCVars::DebugDrawLifeTime, 0, DrawThickness);
+					Chaos::FDebugDrawQueue::GetInstance().DrawDebugBox(PastState.GetX(), Box, PastState.GetR(), FColor::White, false, PhysicsReplicationCVars::DebugDrawLifeTime, 0, DrawThickness);
+					Chaos::FDebugDrawQueue::GetInstance().DrawDebugBox(Target.TargetState.Position, Box, Target.TargetState.Quaternion, DebugColor, false, PhysicsReplicationCVars::DebugDrawLifeTime, 0, DrawThickness);
 
-				if (bResimV)
-				{
-					const FVector DiffV = Target.TargetState.LinVel - PastState.GetV();
-					Chaos::FDebugDrawQueue::GetInstance().DrawDebugDirectionalArrow(Handle->GetX(), Handle->GetX() + DiffV, 5.0f, FColor::Orange, false, PhysicsReplicationCVars::DebugDrawLifeTime, 0, DrawThickness);
+					Chaos::FDebugDrawQueue::GetInstance().DrawDebugDirectionalArrow(Handle->GetX(), PastState.GetX(), 5.0f, FColor::White, false, PhysicsReplicationCVars::DebugDrawLifeTime, 0, DrawThickness);
+					Chaos::FDebugDrawQueue::GetInstance().DrawDebugDirectionalArrow(PastState.GetX(), Target.TargetState.Position, 5.0f, FColor::Black, false, PhysicsReplicationCVars::DebugDrawLifeTime, 0, DrawThickness);
+
+					if (bResimV)
+					{
+						const FVector DiffV = Target.TargetState.LinVel - PastState.GetV();
+						Chaos::FDebugDrawQueue::GetInstance().DrawDebugDirectionalArrow(Target.TargetState.Position, Target.TargetState.Position + DiffV, 5.0f, FColor::Orange, false, PhysicsReplicationCVars::DebugDrawLifeTime, 0, DrawThickness);
+					}
+					if (bResimW)
+					{
+						const FVector DiffW = Target.TargetState.AngVel - FMath::RadiansToDegrees(PastState.GetW());
+						Chaos::FDebugDrawQueue::GetInstance().DrawDebugDirectionalArrow(Target.TargetState.Position + DiffW, Target.TargetState.Position, 5.0f, FColor::Magenta, false, PhysicsReplicationCVars::DebugDrawLifeTime, 0, DrawThickness);
+					}
 				}
-				if (bResimW)
+				else // Resim trigger debug draw
 				{
-					const FVector DiffW = Target.TargetState.AngVel - FMath::RadiansToDegrees(PastState.GetW());
-					Chaos::FDebugDrawQueue::GetInstance().DrawDebugDirectionalArrow(Handle->GetX() + DiffW, Handle->GetX(), 5.0f, FColor::Magenta, false, PhysicsReplicationCVars::DebugDrawLifeTime, 0, DrawThickness);
+					Chaos::FDebugDrawQueue::GetInstance().DrawDebugBox(Handle->GetX(), Box, PastState.GetR(), FColor::White, false, PhysicsReplicationCVars::DebugDrawLifeTime, 0, DrawThickness);
+					Chaos::FDebugDrawQueue::GetInstance().DrawDebugBox(Handle->GetX() + (Target.TargetState.Position - PastState.GetX()), Box, Target.TargetState.Quaternion, DebugColor, false, PhysicsReplicationCVars::DebugDrawLifeTime, 0, DrawThickness);
+
+					if (bResimV)
+					{
+						const FVector DiffV = Target.TargetState.LinVel - PastState.GetV();
+						Chaos::FDebugDrawQueue::GetInstance().DrawDebugDirectionalArrow(Handle->GetX(), Handle->GetX() + DiffV, 5.0f, FColor::Orange, false, PhysicsReplicationCVars::DebugDrawLifeTime, 0, DrawThickness);
+					}
+					if (bResimW)
+					{
+						const FVector DiffW = Target.TargetState.AngVel - FMath::RadiansToDegrees(PastState.GetW());
+						Chaos::FDebugDrawQueue::GetInstance().DrawDebugDirectionalArrow(Handle->GetX() + DiffW, Handle->GetX(), 5.0f, FColor::Magenta, false, PhysicsReplicationCVars::DebugDrawLifeTime, 0, DrawThickness);
+					}
 				}
 			}
 		}
+	#endif
 	}
-#endif
 
-	// Wake up if is sleeping and should not sleep
-	if (Handle->IsSleeping() && !bShouldSleep)
-	{
-		RigidsSolver->GetEvolution()->SetParticleObjectState(Handle, Chaos::EObjectStateType::Dynamic);
-	}
+	return bShouldTriggerResim;
+}
 
-	if (bShouldTriggerResim && Target.TickCount == 0 && LocalFrame > RewindData->GetBlockedResimFrame())
+/** Compare states and trigger resimulation if needed */
+bool FPhysicsReplicationAsync::ResimulationReplication(Chaos::FPBDRigidParticleHandle* Handle, FReplicatedPhysicsTargetAsync& Target, const float DeltaSeconds)
+{
+	if(!IsTargetValidForResim(Target) || !Handle)
 	{
-		// Request resimulation
-		RewindData->RequestResimulation(LocalFrame, Handle);
+		return true;
 	}
-	else if (SettingsCurrent.ResimulationSettings.GetRuntimeCorrectionEnabled())
-	{
-		const int32 NumPredictedFrames = RigidsSolver->GetCurrentFrame() - LocalFrame - Target.TickCount;
 
-		if (Target.TickCount <= NumPredictedFrames && NumPredictedFrames > 0)
-		{
-			const FVector ErrorOffset = (Target.TargetState.Position - PastState.GetX());
+	bool bClearTarget = true;
+	if(Chaos::FPBDRigidsSolver* RigidsSolver = static_cast<Chaos::FPBDRigidsSolver*>(GetSolver()))
+	{ 
+		if(Chaos::FRewindData* RewindData = RigidsSolver->GetRewindData())
+		{ 
+			if(RigidsSolver->GetEvolution())
+			{ 
+				const int32 LocalFrame = Target.ServerFrame - Target.FrameOffset;
+				static constexpr Chaos::FFrameAndPhase::EParticleHistoryPhase RewindPhase = Chaos::FFrameAndPhase::EParticleHistoryPhase::PrePushData;
+
+				// Get state from locally cached history for frame corresponding to received data
+				const Chaos::FGeometryParticleState PastState = RewindData->GetPastStateAtFrame(*Handle, LocalFrame, RewindPhase);
+
+				const bool bShouldSleep = (Target.TargetState.Flags & ERigidBodyFlags::Sleeping) != 0;
+				const bool bShouldTriggerResim = (Chaos::RewindBeforeAdvance == 0) ? ShouldResimulateParticle(Handle, Target, DeltaSeconds) : false;
+
+				// Wake up if is sleeping and should not sleep
+				if (Handle->IsSleeping() && !bShouldSleep)
+				{
+					RigidsSolver->GetEvolution()->SetParticleObjectState(Handle, Chaos::EObjectStateType::Dynamic);
+				}
 
-			// Positional Correction
-			const float CorrectionAmountX = SettingsCurrent.ResimulationSettings.GetPosStabilityMultiplier() / NumPredictedFrames;
-			const FVector PosDiffCorrection = ErrorOffset * CorrectionAmountX; // Same result as (ErrorOffset / NumPredictedFrames) * PosStabilityMultiplier
-			const FVector CorrectedX = Handle->GetX() + PosDiffCorrection;
+				if (bShouldTriggerResim && Target.TickCount == 0 && LocalFrame > RewindData->GetBlockedResimFrame())
+				{
+					// Request resimulation
+					RewindData->RequestResimulation(LocalFrame, Handle);
+				}
+				else if (SettingsCurrent.ResimulationSettings.GetRuntimeCorrectionEnabled())
+				{
+					const int32 NumPredictedFrames = RigidsSolver->GetCurrentFrame() - LocalFrame - Target.TickCount;
 
-			// Rotational Correction
-			const float CorrectionAmountR = SettingsCurrent.ResimulationSettings.GetRotStabilityMultiplier() / NumPredictedFrames;
-			const FQuat DeltaQuat = PastState.GetR().Inverse() * Target.TargetState.Quaternion;
-			const FQuat TargetCorrectionR = Handle->GetR() * DeltaQuat;
-			const FQuat CorrectedR = FQuat::Slerp(Handle->GetR(), TargetCorrectionR, CorrectionAmountR);
+					if (Target.TickCount <= NumPredictedFrames && NumPredictedFrames > 0)
+					{
+						const FVector ErrorOffset = (Target.TargetState.Position - PastState.GetX());
 
-			if (SettingsCurrent.ResimulationSettings.GetRuntimeVelocityCorrectionEnabled())
-			{
-				// Linear Velocity Correction
-				const FVector LinVelDiff = Target.TargetState.LinVel - PastState.GetV(); // Velocity vector that the server covers but the client doesn't
-				const float CorrectionAmountV = SettingsCurrent.ResimulationSettings.GetVelStabilityMultiplier() / NumPredictedFrames;
-				const FVector VelCorrection = LinVelDiff * CorrectionAmountV; // Same result as (LinVelDiff / NumPredictedFrames) * VelStabilityMultiplier
-				const FVector CorrectedV = Handle->GetV() + VelCorrection;
-
-				// Angular Velocity Correction
-				const FVector AngVelDiff = FMath::DegreesToRadians(Target.TargetState.AngVel) - PastState.GetW(); // Angular velocity vector that the server covers but the client doesn't
-				const float CorrectionAmountW = SettingsCurrent.ResimulationSettings.GetAngVelStabilityMultiplier() / NumPredictedFrames;
-				const FVector AngVelCorrection = AngVelDiff * CorrectionAmountW; // Same result as (AngVelDiff / NumPredictedFrames) * VelStabilityMultiplier
-				const FVector CorrectedW = Handle->GetW() + AngVelCorrection;
+						// Positional Correction
+						const float CorrectionAmountX = SettingsCurrent.ResimulationSettings.GetPosStabilityMultiplier() / NumPredictedFrames;
+						const FVector PosDiffCorrection = ErrorOffset * CorrectionAmountX; // Same result as (ErrorOffset / NumPredictedFrames) * PosStabilityMultiplier
+						const FVector CorrectedX = Handle->GetX() + PosDiffCorrection;
+
+						// Rotational Correction
+						const float CorrectionAmountR = SettingsCurrent.ResimulationSettings.GetRotStabilityMultiplier() / NumPredictedFrames;
+						const FQuat DeltaQuat = PastState.GetR().Inverse() * Target.TargetState.Quaternion;
+						const FQuat TargetCorrectionR = Handle->GetR() * DeltaQuat;
+						const FQuat CorrectedR = FQuat::Slerp(Handle->GetR(), TargetCorrectionR, CorrectionAmountR);
+
+						if (SettingsCurrent.ResimulationSettings.GetRuntimeVelocityCorrectionEnabled())
+						{
+							// Linear Velocity Correction
+							const FVector LinVelDiff = Target.TargetState.LinVel - PastState.GetV(); // Velocity vector that the server covers but the client doesn't
+							const float CorrectionAmountV = SettingsCurrent.ResimulationSettings.GetVelStabilityMultiplier() / NumPredictedFrames;
+							const FVector VelCorrection = LinVelDiff * CorrectionAmountV; // Same result as (LinVelDiff / NumPredictedFrames) * VelStabilityMultiplier
+							const FVector CorrectedV = Handle->GetV() + VelCorrection;
+
+							// Angular Velocity Correction
+							const FVector AngVelDiff = FMath::DegreesToRadians(Target.TargetState.AngVel) - PastState.GetW(); // Angular velocity vector that the server covers but the client doesn't
+							const float CorrectionAmountW = SettingsCurrent.ResimulationSettings.GetAngVelStabilityMultiplier() / NumPredictedFrames;
+							const FVector AngVelCorrection = AngVelDiff * CorrectionAmountW; // Same result as (AngVelDiff / NumPredictedFrames) * VelStabilityMultiplier
+							const FVector CorrectedW = Handle->GetW() + AngVelCorrection;
 					
-				// Apply correction to velocities
-				Handle->SetV(CorrectedV);
-				Handle->SetW(CorrectedW);
-			}
+							// Apply correction to velocities
+							Handle->SetV(CorrectedV);
+							Handle->SetW(CorrectedW);
+						}
 
-#if !(UE_BUILD_SHIPPING || UE_BUILD_TEST)
-			if (PhysicsReplicationCVars::ResimulationCVars::bDrawDebug)
-			{
-				Chaos::FDebugDrawQueue::GetInstance().DrawDebugDirectionalArrow(Handle->GetX(), CorrectedX, 5.0f, FColor::MakeRandomSeededColor(LocalFrame), false, PhysicsReplicationCVars::DebugDrawLifeTime, 0, 0.5f);
-			}
-#endif
-			// Apply correction to position and rotation
-			RigidsSolver->GetEvolution()->ApplyParticleTransformCorrection(Handle, CorrectedX, CorrectedR, SettingsCurrent.ResimulationSettings.GetRuntimeCorrectConnectedBodies(), /*bInRecalculateFrictionOnConnectedBodies*/true, ReplicatedParticleIDs);
-		}
+			#if !(UE_BUILD_SHIPPING || UE_BUILD_TEST)
+						if (PhysicsReplicationCVars::ResimulationCVars::bDrawDebug)
+						{
+							Chaos::FDebugDrawQueue::GetInstance().DrawDebugDirectionalArrow(Handle->GetX(), CorrectedX, 5.0f, FColor::MakeRandomSeededColor(LocalFrame), false, PhysicsReplicationCVars::DebugDrawLifeTime, 0, 0.5f);
+						}
+			#endif
+						// Apply correction to position and rotation
+						RigidsSolver->GetEvolution()->ApplyParticleTransformCorrection(Handle, CorrectedX, CorrectedR, SettingsCurrent.ResimulationSettings.GetRuntimeCorrectConnectedBodies(), /*bInRecalculateFrictionOnConnectedBodies*/true, ReplicatedParticleIDs);
+					}
 
-		// Keep target for NumPredictedFrames time to perform runtime corrections with until a new target is received
-		bClearTarget = Target.TickCount >= NumPredictedFrames;
-	}
+					// Keep target for NumPredictedFrames time to perform runtime corrections with until a new target is received
+					bClearTarget = Target.TickCount >= NumPredictedFrames;
+				}
 
-	// Set sleep state if we are about to clear the target from memory and the target is set to sleep
-	if (bClearTarget && bShouldSleep)
-	{
-		// Snap object into correct state, it should already be at that state or very close to it
-		RigidsSolver->GetEvolution()->ApplyParticleTransformCorrection(Handle, Target.TargetState.Position, Target.TargetState.Quaternion, /*bApplyToConnectedBodies*/true, /*bInRecalculateFrictionOnConnectedBodies*/true, ReplicatedParticleIDs);
+				// Set sleep state if we are about to clear the target from memory and the target is set to sleep
+				if (bClearTarget && bShouldSleep)
+				{
+					// Snap object into correct state, it should already be at that state or very close to it
+					RigidsSolver->GetEvolution()->ApplyParticleTransformCorrection(Handle, Target.TargetState.Position, Target.TargetState.Quaternion, /*bApplyToConnectedBodies*/true, /*bInRecalculateFrictionOnConnectedBodies*/true, ReplicatedParticleIDs);
 
-		RigidsSolver->GetEvolution()->SetParticleObjectState(Handle, Chaos::EObjectStateType::Sleeping);
-		if (PhysicsReplicationCVars::PredictiveInterpolationCVars::bSleepConnectedBodies)
-		{
-			RigidsSolver->GetEvolution()->ApplySleepOnConnectedParticles(Handle);
+					RigidsSolver->GetEvolution()->SetParticleObjectState(Handle, Chaos::EObjectStateType::Sleeping);
+					if (PhysicsReplicationCVars::PredictiveInterpolationCVars::bSleepConnectedBodies)
+					{
+						RigidsSolver->GetEvolution()->ApplySleepOnConnectedParticles(Handle);
+					}
+				}
+				else if (Target.IsWaiting()) 
+				{
+					// Don't clear the target if we are waiting for a specific target frame and not sleeping
+					bClearTarget = false;
+				}
+			}
 		}
 	}
-	else if (Target.IsWaiting()) 
-	{
-		// Don't clear the target if we are waiting for a specific target frame and not sleeping
-		bClearTarget = false;
-	}
 
 	return bClearTarget;
 }
diff --git a/Engine/Source/Runtime/Engine/Public/Physics/NetworkPhysicsComponent.h b/Engine/Source/Runtime/Engine/Public/Physics/NetworkPhysicsComponent.h
--- a/Engine/Source/Runtime/Engine/Public/Physics/NetworkPhysicsComponent.h
+++ b/Engine/Source/Runtime/Engine/Public/Physics/NetworkPhysicsComponent.h
@@ -21,10 +21,12 @@
 #endif
 
 class FAsyncNetworkPhysicsComponent;
+struct FAsyncNetworkPhysicsComponentInput;
 
 DECLARE_MULTICAST_DELEGATE_OneParam(FOnPreProcessInputsInternal, const int32);
 DECLARE_MULTICAST_DELEGATE_OneParam(FOnPostProcessInputsInternal, const int32);
 DECLARE_MULTICAST_DELEGATE_TwoParams(FOnInjectInputsExternal, const int32 /* PhysicsStep */, const int32 /* NumSteps */);
+DECLARE_MULTICAST_DELEGATE_OneParam(FAddResimulationRequestInternal, const int32 /* PhysicsStep */);
 
 /** Templated data history, holding a data buffer */
 template<typename DataType, bool bLegacyData = false>
@@ -684,6 +686,8 @@ struct FNetworkPhysicsCallback : public Chaos::IRewindCallback
 	// Delegate on the internal inputs process
 	FOnPreProcessInputsInternal PreProcessInputsInternal;
 	FOnPostProcessInputsInternal PostProcessInputsInternal;
+	FAddResimulationRequestInternal AddResimulationRequestInternal;
+
 	// Bind to this for additional processing on the GT during InjectInputs_External()
 	FOnInjectInputsExternal InjectInputsExternal;
 
@@ -707,6 +711,9 @@ struct FNetworkPhysicsCallback : public Chaos::IRewindCallback
 		RewindableCallbackObjects.Remove(SimCallbackObject);
 	}
 
+	/** Add potential resimulation requests to the rewind data */
+	virtual void AddResimulationRequest_Internal(const int32 PhysicsStep) override;
+
 	// World owning that callback
 	UWorld* World = nullptr;
 
@@ -1958,8 +1965,11 @@ private:
 	// Delegate ran at the end of FNetworkPhysicsCallback::ProcessInputs_Internal, used to record and send inputs and states if needed.
 	void OnPostProcessInputs_Internal(const int32 PhysicsStep);
 
+	// Delegate ran to potentially trigger resimulation if needed
+	void OnAddResimulationRequest_Internal(const int32 PhysicsStep);
+
 	// Consume data from async input
-	void ConsumeAsyncInput(const int32 PhysicsStep);
+	void ConsumeAsyncInput_Internal(const int32 PhysicsStep, const bool bTriggerResim);
 
 	/** Get the rigid solver */
 	Chaos::FPBDRigidsSolver* GetRigidSolver();
@@ -2082,6 +2092,7 @@ private:
 
 	FDelegateHandle DelegateOnPreProcessInputs_Internal;
 	FDelegateHandle DelegateOnPostProcessInputs_Internal;
+	FDelegateHandle DelegateOnAddResimulationRequest_Internal;
 };
 
 namespace UE::NetworkPhysicsUtils
diff --git a/Engine/Source/Runtime/Engine/Public/PhysicsReplication.h b/Engine/Source/Runtime/Engine/Public/PhysicsReplication.h
--- a/Engine/Source/Runtime/Engine/Public/PhysicsReplication.h
+++ b/Engine/Source/Runtime/Engine/Public/PhysicsReplication.h
@@ -208,7 +208,10 @@ class FPhysicsReplicationAsync : public IPhysicsReplicationAsync,
 	virtual bool ResimulationReplication(Chaos::FPBDRigidParticleHandle* Handle, FReplicatedPhysicsTargetAsync& Target, const float DeltaSeconds);
 
 public:
+	//~ Begin IPhysicsReplicationAsync interface
 	virtual void RegisterSettings(Chaos::FConstPhysicsObjectHandle PhysicsObject, TWeakPtr<const FNetworkPhysicsSettingsData> InSettings) override;
+	virtual void AddResimulationRequest_Internal(const float DeltaSeconds) override;
+	//~ End IPhysicsReplicationAsync interface
 
 private:
 	float LatencyOneWay;
@@ -239,6 +242,12 @@ private:
 	void ApplyPhysicsReplicationLOD(Chaos::FConstPhysicsObjectHandle PhysicsObjectHandle, FReplicatedPhysicsTargetAsync& Target, const uint32 LODFlags);
 	void DebugDrawReplicationMode(const FPhysicsRepAsyncInputData& Input);
 
+	/** Check if we have to resimulate the particle or not */
+	bool ShouldResimulateParticle(const Chaos::FPBDRigidParticleHandle* Handle, const FReplicatedPhysicsTargetAsync& Target, const float DeltaSeconds) const;
+
+	/** Check if a target is valid for resimulation or not */
+	bool IsTargetValidForResim(const FReplicatedPhysicsTargetAsync& Target) const;
+
 	/** Static function to extrapolate a target for N ticks using X DeltaSeconds */
 	static void ExtrapolateTarget(FReplicatedPhysicsTargetAsync& Target, const int32 ExtrapolateFrames, const float DeltaSeconds);
 	/** Static function to extrapolate a target for N Seconds */
diff --git a/Engine/Source/Runtime/Engine/Public/PhysicsReplicationInterface.h b/Engine/Source/Runtime/Engine/Public/PhysicsReplicationInterface.h
--- a/Engine/Source/Runtime/Engine/Public/PhysicsReplicationInterface.h
+++ b/Engine/Source/Runtime/Engine/Public/PhysicsReplicationInterface.h
@@ -31,4 +31,7 @@ public:
 	virtual void RegisterSettings(Chaos::FConstPhysicsObjectHandle PhysicsObject, FNetworkPhysicsSettingsAsync InSettings) { }
 
 	virtual void RegisterSettings(Chaos::FConstPhysicsObjectHandle PhysicsObject, TWeakPtr<const FNetworkPhysicsSettingsData> InSettings) = 0;
+
+	/** Add potential resimulation reques from the physics replication */
+	virtual void AddResimulationRequest_Internal( const float DeltaSeconds) = 0;
 };
diff --git a/Engine/Source/Runtime/Experimental/Chaos/Private/Chaos/Framework/PhysicsSolverBase.cpp b/Engine/Source/Runtime/Experimental/Chaos/Private/Chaos/Framework/PhysicsSolverBase.cpp
--- a/Engine/Source/Runtime/Experimental/Chaos/Private/Chaos/Framework/PhysicsSolverBase.cpp
+++ b/Engine/Source/Runtime/Experimental/Chaos/Private/Chaos/Framework/PhysicsSolverBase.cpp
@@ -28,6 +28,10 @@ CSV_DEFINE_CATEGORY(ChaosPhysicsSolver, true);
 
 namespace Chaos
 {	
+	CHAOS_API int32 RewindBeforeAdvance = 1;
+	FAutoConsoleVariableRef CVarReorderRewindResimulation(TEXT("p.Resim.RewindBeforeAdvance"), RewindBeforeAdvance, TEXT("Whether to apply the rewind before or after the advance"));
+
+
 	extern int GSingleThreadedPhysics;
 	void FPhysicsSolverBase::ChangeBufferMode(EMultiBufferMode InBufferMode)
 	{
@@ -95,6 +99,11 @@ namespace Chaos
 		return PhysicsRunsOnGT == 0 ? CPrio_FPhysicsTickTask.Get() : ENamedThreads::GameThread;
 	}
 
+	ENamedThreads::Type FPhysicsSolverRewindTask::GetDesiredThread()
+	{
+		return PhysicsRunsOnGT == 0 ? CPrio_FPhysicsTickTask.Get() : ENamedThreads::GameThread;
+	}
+
 	void FPhysicsSolverProcessPushDataTask::ProcessPushData()
 	{
 		using namespace Chaos;
@@ -136,6 +145,31 @@ namespace Chaos
 		
 	}
 
+	void FPhysicsSolverRewindTask::RewindSolver()
+	{
+		LLM_SCOPE(ELLMTag::ChaosUpdate);
+		SCOPE_CYCLE_COUNTER(STAT_ChaosTick);
+		CSV_SCOPED_TIMING_STAT_EXCLUSIVE(Physics);
+		CSV_SCOPED_TIMING_STAT(PhysicsVerbose, StepSolver_RewindData);
+
+		CVD_SCOPE_CONTEXT(Solver.GetChaosVDContextData());
+
+#if PHYSICS_THREAD_CONTEXT
+		FPhysicsThreadContextScope Scope(/*IsPhysicsThreadContext=*/true);
+#endif
+		if(PushData && !Solver.IsResimming())
+		{ 
+			if(RewindBeforeAdvance != 0)
+			{
+				Solver.AddResimulationRequest_Internal(*PushData);
+			}
+
+			{
+				SCOPE_CYCLE_COUNTER(STAT_ConditionalApplyRewind);
+				Solver.ConditionalApplyRewind_Internal();
+			}
+		}
+	}
 
 	FPhysicsSolverAdvanceTask::FPhysicsSolverAdvanceTask(FPhysicsSolverBase& InSolver, FPushPhysicsData* InPushData)
 		: Solver(InSolver)
@@ -173,20 +207,14 @@ namespace Chaos
 		// StepFraction: how much of the remaining time this step represents while substepping, used to interpolate kinematic targets
 		// E.g., for 4 steps this will be: 1/4, 1/3, 1/2, 1
 		const FReal PseudoFraction = PushData->bSolverSubstepped ? (FReal)1 / (FReal)(PushData->IntervalNumSteps - PushData->IntervalStep) : 1.0f;
+		const FSubStepInfo SubStepInfo = { PseudoFraction, PushData->IntervalStep, PushData->IntervalNumSteps, PushData->bSolverSubstepped };
 
-		Solver.AdvanceSolverBy(FSubStepInfo{ PseudoFraction, PushData->IntervalStep, PushData->IntervalNumSteps, PushData->bSolverSubstepped });
-
+		Solver.AdvanceSolverBy(SubStepInfo);
 		{
 			SCOPE_CYCLE_COUNTER(STAT_ResetMarshallingData);
 			Solver.GetMarshallingManager().FreeDataToHistory_Internal(PushData);	//cannot use push data after this point
 			PushData = nullptr;
 		}
-
-		{
-			SCOPE_CYCLE_COUNTER(STAT_ConditionalApplyRewind);
-			Solver.ConditionalApplyRewind_Internal();
-		}
-
 		Solver.NumPendingSolverAdvanceTasks--;
 	}
 
@@ -582,16 +610,38 @@ namespace Chaos
 					Prereqs.Add(PendingTasks);
 				}
 
-				PendingTasks = TGraphTask<FPhysicsSolverProcessPushDataTask>::CreateTask(&Prereqs).ConstructAndDispatchWhenReady(*this, PushData);
-				Prereqs.Add(PendingTasks);
+				if(RewindBeforeAdvance == 0)
+				{ 
+					PendingTasks = TGraphTask<FPhysicsSolverProcessPushDataTask>::CreateTask(&Prereqs).ConstructAndDispatchWhenReady(*this, PushData);
+					Prereqs.Add(PendingTasks);
 
-				if (bSolverHasFrozenGameThreadCallbacks)
-				{
-					PendingTasks = TGraphTask<FPhysicsSolverFrozenGTPreSimCallbacks>::CreateTask(&Prereqs).ConstructAndDispatchWhenReady(*this);
+					if (bSolverHasFrozenGameThreadCallbacks)
+					{
+						PendingTasks = TGraphTask<FPhysicsSolverFrozenGTPreSimCallbacks>::CreateTask(&Prereqs).ConstructAndDispatchWhenReady(*this);
+						Prereqs.Add(PendingTasks);
+					}
+
+					PendingTasks = TGraphTask<FPhysicsSolverAdvanceTask>::CreateTask(&Prereqs).ConstructAndDispatchWhenReady(*this, PushData);
 					Prereqs.Add(PendingTasks);
+
+					PendingTasks = TGraphTask<FPhysicsSolverRewindTask>::CreateTask(&Prereqs).ConstructAndDispatchWhenReady(*this, PushData);
 				}
+				else
+				{
+					PendingTasks = TGraphTask<FPhysicsSolverRewindTask>::CreateTask(&Prereqs).ConstructAndDispatchWhenReady(*this, PushData);
+					Prereqs.Add(PendingTasks);
+
+					PendingTasks = TGraphTask<FPhysicsSolverProcessPushDataTask>::CreateTask(&Prereqs).ConstructAndDispatchWhenReady(*this, PushData);
+					Prereqs.Add(PendingTasks);
 
-				PendingTasks = TGraphTask<FPhysicsSolverAdvanceTask>::CreateTask(&Prereqs).ConstructAndDispatchWhenReady(*this, PushData);
+					if (bSolverHasFrozenGameThreadCallbacks)
+					{
+						PendingTasks = TGraphTask<FPhysicsSolverFrozenGTPreSimCallbacks>::CreateTask(&Prereqs).ConstructAndDispatchWhenReady(*this);
+						Prereqs.Add(PendingTasks);
+					}
+
+					PendingTasks = TGraphTask<FPhysicsSolverAdvanceTask>::CreateTask(&Prereqs).ConstructAndDispatchWhenReady(*this, PushData);
+				}
 
 				if (IsUsingAsyncResults() == false)
 				{
@@ -616,14 +666,35 @@ namespace Chaos
 
 	void FAllSolverTasks::AdvanceSolver()
 	{
-		ProcessPushData.ProcessPushData();
-		GTPreSimCallbacks.GTPreSimCallbacks();
-		AdvanceTask.AdvanceSolver();
+		if (RewindBeforeAdvance == 0)
+		{
+			ProcessPushData.ProcessPushData();
+			GTPreSimCallbacks.GTPreSimCallbacks();
+			AdvanceTask.AdvanceSolver();
+			RewindTask.RewindSolver();
+		}
+		else
+		{
+			RewindTask.RewindSolver();
+			ProcessPushData.ProcessPushData();
+			GTPreSimCallbacks.GTPreSimCallbacks();
+			AdvanceTask.AdvanceSolver();
+		}
 	}
 
 	void FSolverTasksPTOnly::AdvanceSolver()
 	{
-		ProcessPushData.ProcessPushData();
-		AdvanceTask.AdvanceSolver();
+		if (RewindBeforeAdvance == 0)
+		{
+			ProcessPushData.ProcessPushData();
+			AdvanceTask.AdvanceSolver();
+			RewindTask.RewindSolver();
+		}
+		else
+		{
+			RewindTask.RewindSolver();
+			ProcessPushData.ProcessPushData();
+			AdvanceTask.AdvanceSolver();
+		}
 	}
 }
diff --git a/Engine/Source/Runtime/Experimental/Chaos/Private/PBDRigidsSolver.cpp b/Engine/Source/Runtime/Experimental/Chaos/Private/PBDRigidsSolver.cpp
--- a/Engine/Source/Runtime/Experimental/Chaos/Private/PBDRigidsSolver.cpp
+++ b/Engine/Source/Runtime/Experimental/Chaos/Private/PBDRigidsSolver.cpp
@@ -31,9 +31,10 @@
 #include "Chaos/AsyncInitBodyHelper.h"
 
 #include "Chaos/DebugDraw/DebugDrawParticle.h"
-#include "ChaosDebugDraw/ChaosDDContext.h"
+#include "ChaosDebugDraw/ChaosDDContext.h" 
 #include "ChaosDebugDraw/ChaosDDScene.h"
 #include "ChaosDebugDraw/ChaosDDTimeline.h"
+#include "PhysicsReplication.h"
 
 #include "ProfilingDebugging/CountersTrace.h"
 #include "ProfilingDebugging/CsvProfiler.h"
@@ -1883,6 +1884,36 @@ namespace Chaos
 			Callback->HasOption(ESimCallbackOptions::PostIntegrate);
 	}
 
+	void FPBDRigidsSolver::AddResimulationRequest_Internal(FPushPhysicsData& PushData)
+	{
+		if(!MRewindData)
+		{
+			return;
+		}
+
+		// Set all the input on the matching callbacks for future use (PhysicsReplication, RewindCallbacks...)
+		for (FSimCallbackInputAndObject& InputAndCallbackObj : PushData.SimCallbackInputs)
+		{
+			if(InputAndCallbackObj.CallbackObject)
+			{
+				InputAndCallbackObj.CallbackObject->SetCurrentInput_Internal(InputAndCallbackObj.Input);
+			}
+		}
+
+		// Add resimulation request from physics state replication
+		if (IPhysicsReplicationAsync* ReplicationCallback = GetPhysicsReplication_Internal())
+		{
+			ReplicationCallback->AddResimulationRequest_Internal((float)MLastDt);
+		}
+
+		// Add resimulation request from rewind callbacks inputs
+		if (!IsShuttingDown() && !bGameThreadFrozen && MRewindCallback)
+		{
+			QUICK_SCOPE_CYCLE_COUNTER(Chaos_RewindCallback_ProcessInputs_Internal);
+			MRewindCallback->AddResimulationRequest_Internal(PushData.InternalStep);
+		}
+	}
+
 	void FPBDRigidsSolver::ProcessPushedData_Internal(FPushPhysicsData& PushData)
 	{
 		QUICK_SCOPE_CYCLE_COUNTER(ChaosPushData);
@@ -2155,7 +2186,6 @@ namespace Chaos
 #if DEBUG_REWIND_DATA
 					UE_LOG(LogChaos, Warning, TEXT("CLIENT | PT | ConditionalApplyRewind_Internal | PERFORMING RESIMULATION | Resim From Frame = %d | Num Steps = %d | To Current Frame: %d"), ResimStep, NumResimSteps, CurrentFrame);
 #endif
-
 					SetIsResimming(true);
 					CurrentFrame = ResimStep;
 
@@ -2191,6 +2221,7 @@ namespace Chaos
 								}
 
 								FPushPhysicsData* PushData = RecordedPushData[LastStep - Step];	//push data is sorted as latest first
+		
 								if (bFirst)
 								{
 									SCOPED_NAMED_EVENT(ChaosRewindAndResim_ResetCollisions, FColor::Red);
diff --git a/Engine/Source/Runtime/Experimental/Chaos/Public/Chaos/Framework/PhysicsSolverBase.h b/Engine/Source/Runtime/Experimental/Chaos/Public/Chaos/Framework/PhysicsSolverBase.h
--- a/Engine/Source/Runtime/Experimental/Chaos/Public/Chaos/Framework/PhysicsSolverBase.h
+++ b/Engine/Source/Runtime/Experimental/Chaos/Public/Chaos/Framework/PhysicsSolverBase.h
@@ -57,6 +57,7 @@ namespace Chaos
 	extern CHAOS_API int32 UseAsyncInterpolation;
 	extern CHAOS_API int32 ForceDisableAsyncPhysics;
 	extern CHAOS_API FRealSingle AsyncInterpolationMultiplier;
+	extern CHAOS_API int32 RewindBeforeAdvance;
 
 	struct FSubStepInfo
 	{
@@ -174,13 +175,38 @@ namespace Chaos
 #endif
 	};
 
+
+	/**
+	 * Task responsible for rewinding the solver if required by networked physics
+	 */
+	class FPhysicsSolverRewindTask
+	{
+	public:
+
+		FPhysicsSolverRewindTask(FPhysicsSolverBase& InSolver, FPushPhysicsData* InPushData) : Solver(InSolver), PushData(InPushData) {}
+
+		TStatId GetStatId() const { RETURN_QUICK_DECLARE_CYCLE_STAT(FPhysicsSolverRewindTask, STATGROUP_TaskGraphTasks); }
+		static CHAOS_API ENamedThreads::Type GetDesiredThread();
+		static ESubsequentsMode::Type GetSubsequentsMode() { return ESubsequentsMode::TrackSubsequents; }
+		void DoTask(ENamedThreads::Type CurrentThread, const FGraphEventRef& MyCompletionGraphEvent) { RewindSolver(); }
+		void RewindSolver();
+
+	private:
+
+		FPhysicsSolverBase& Solver;
+		FPushPhysicsData* PushData;
+	};
+
 	// Container for all steps required to fully update a solver
-	struct FAllSolverTasks
+	class FAllSolverTasks
 	{
+	public:
+
 		FAllSolverTasks(FPhysicsSolverBase& InSolver, FPushPhysicsData* PushData)
 			: ProcessPushData(InSolver, PushData)
 			, GTPreSimCallbacks(InSolver)
 			, AdvanceTask(InSolver, PushData)
+			, RewindTask(InSolver, PushData)
 			, Solver(InSolver)
 		{
 		}
@@ -188,16 +214,20 @@ namespace Chaos
 		FPhysicsSolverProcessPushDataTask ProcessPushData;
 		FPhysicsSolverFrozenGTPreSimCallbacks GTPreSimCallbacks;
 		FPhysicsSolverAdvanceTask AdvanceTask;
+		FPhysicsSolverRewindTask RewindTask;
 
 		CHAOS_API void AdvanceSolver();
 
+	private:
+
 		FPhysicsSolverBase& Solver;
 	};
 
 	// Container for all physics-thread steps required to update a solver
 	// This leaves out the game thread callbacks for situations that require and update only in a physics-thread context
-	struct FSolverTasksPTOnly
+	class FSolverTasksPTOnly
 	{
+	public:
 		FSolverTasksPTOnly() = delete;
 		FSolverTasksPTOnly(const FSolverTasksPTOnly&) = delete;
 		FSolverTasksPTOnly(FSolverTasksPTOnly&&) = delete;
@@ -207,13 +237,17 @@ namespace Chaos
 		FSolverTasksPTOnly(FPhysicsSolverBase& InSolver, FPushPhysicsData* InPushData)
 			: ProcessPushData(InSolver, InPushData)
 			, AdvanceTask(InSolver, InPushData)
+			, RewindTask(InSolver, InPushData)
 			, Solver(InSolver)
 		{}
+	
+		CHAOS_API void AdvanceSolver();
+
+	private:
 
 		FPhysicsSolverProcessPushDataTask ProcessPushData;
 		FPhysicsSolverAdvanceTask AdvanceTask;
-
-		CHAOS_API void AdvanceSolver();
+		FPhysicsSolverRewindTask RewindTask;
 
 		FPhysicsSolverBase& Solver;
 	};
@@ -601,6 +635,7 @@ namespace Chaos
 		
 		CHAOS_API void DisableAsyncMode();
 		
+		virtual void AddResimulationRequest_Internal(FPushPhysicsData& PushData){}
 		virtual void ConditionalApplyRewind_Internal(){}
 		virtual bool IsResimming() const {return false;}
 
@@ -975,6 +1010,7 @@ namespace Chaos
 		friend FChaosSolversModule;
 		friend FPhysicsSolverProcessPushDataTask;
 		friend FPhysicsSolverAdvanceTask;
+		friend FPhysicsSolverRewindTask;
 
 		template<ELockType>
 		friend struct TSolverSimMaterialScope;
diff --git a/Engine/Source/Runtime/Experimental/Chaos/Public/PBDRigidsSolver.h b/Engine/Source/Runtime/Experimental/Chaos/Public/PBDRigidsSolver.h
--- a/Engine/Source/Runtime/Experimental/Chaos/Public/PBDRigidsSolver.h
+++ b/Engine/Source/Runtime/Experimental/Chaos/Public/PBDRigidsSolver.h
@@ -462,6 +462,7 @@ namespace Chaos
 		CHAOS_API virtual void ProcessPushedData_Internal(FPushPhysicsData& PushData) override;
 		CHAOS_API void DestroyPendingProxies_Internal();
 
+		CHAOS_API virtual void AddResimulationRequest_Internal(FPushPhysicsData& PushData) override;
 		CHAOS_API virtual void ConditionalApplyRewind_Internal() override;
 
 		/** Check if we are resimming or not */
diff --git a/Engine/Source/Runtime/Experimental/Chaos/Public/RewindData.h b/Engine/Source/Runtime/Experimental/Chaos/Public/RewindData.h
--- a/Engine/Source/Runtime/Experimental/Chaos/Public/RewindData.h
+++ b/Engine/Source/Runtime/Experimental/Chaos/Public/RewindData.h
@@ -1380,7 +1380,6 @@ public:
 	{
 	}
 
-
 	REWIND_PARTICLE_STATIC_PROPERTY(ParticlePositionRotation, GetX)
 	REWIND_PARTICLE_STATIC_PROPERTY(ParticlePositionRotation, GetR)
 
@@ -1950,7 +1949,7 @@ public:
 	void EndResim();
 
 	/** Check if a frame number is within the currently allowed rewind span */
-	bool IsFrameWithinRewindHistory(int32 Frame) { return Frame < LatestFrame && Frame >= GetEarliestFrame_Internal(); }
+	bool IsFrameWithinRewindHistory(int32 Frame) const { return Frame < LatestFrame && Frame >= GetEarliestFrame_Internal(); }
 
 	/** Request a resimulation by setting a requested frame to rewind to
 	* @param Particle can be used to mark a specific particle as the one triggering the resimulation */
@@ -2261,6 +2260,7 @@ private:
 	// Used to keep track of objects we "freeze" during resim (by making them kinematic and marking them FrozenDuringResim)
 	// and their object state prior to the freezing so it can be set back to its original value
 	TArray<TPair<FParticleSafeRef, EObjectStateType>> FrozenParticles;
+
 	// Resim bubbles start as the swept AABBs of particles needing resim, and grow to encompass nearby kinematic and dynamic particles involved in the resim
 	// Their formation happens in FindParticlesOverlappingResimBubbles()
 	TArray<Chaos::FAABB3> ResimBubbles;
@@ -2315,6 +2315,9 @@ public:
 	*/
 	virtual int32 TriggerRewindIfNeeded_Internal(int32 LatestStepCompleted) { return INDEX_NONE; }
 
+	/** Add potential resimulation requests to the rewind data */
+	virtual void AddResimulationRequest_Internal(const int32 PhysicsStep) {}
+
 	/** Called before each rewind step. This is to give user code the opportunity to trigger other code before each rewind step
 	*   Usually to simulate external systems that ran in lock step with the physics sim
 	*/
