Add more info to assertion

#Synced-CL 48918445

diff --git a/Engine/Source/Runtime/Engine/Private/SkeletalMesh.cpp b/Engine/Source/Runtime/Engine/Private/SkeletalMesh.cpp
--- a/Engine/Source/Runtime/Engine/Private/SkeletalMesh.cpp
+++ b/Engine/Source/Runtime/Engine/Private/SkeletalMesh.cpp
@@ -1117,7 +1117,8 @@ void USkeletalMesh::InitResources()
 				}
 				const FSkeletalMeshLODModel& ImportLODModel = GetImportedModel()->LODModels[LODIndex];
 				FSkeletalMeshLODRenderData& RenderLODModel = GetSkeletalMeshRenderData()->LODRenderData[LODIndex];
-				check(ImportLODModel.Sections.Num() == RenderLODModel.RenderSections.Num());
+				checkf(ImportLODModel.Sections.Num() == RenderLODModel.RenderSections.Num(), TEXT("Mismatch between import and render LOD model section count for %s LOD %d. %d vs %d"),
+					*GetPathName(), LODIndex, ImportLODModel.Sections.Num(), RenderLODModel.RenderSections.Num());
 				for (int32 SectionIndex = 0; SectionIndex < ImportLODModel.Sections.Num(); ++SectionIndex)
 				{
 					const FSkelMeshSection& ImportSection = ImportLODModel.Sections[SectionIndex];
