Functional Spline Tests: Refactor to add support for serialization tests

#virtualized
#Synced-CL 48816665

diff --git a/Engine/Source/Runtime/Engine/Classes/Curves/Spline.h b/Engine/Source/Runtime/Engine/Classes/Curves/Spline.h
--- a/Engine/Source/Runtime/Engine/Classes/Curves/Spline.h
+++ b/Engine/Source/Runtime/Engine/Classes/Curves/Spline.h
@@ -98,10 +98,9 @@ struct FSpline
 	/* Misc Interface */
 	
 	ENGINE_API bool operator==(const FSpline& Other) const;
-	bool operator!=(const FSpline& Other) const
-	{
-		return !(*this == Other);
-	}
+	ENGINE_API bool operator!=(const FSpline& Other) const;
+
+	ENGINE_API bool IsNearlyEqual(const FSpline& Other, const float RelativeTolerance, const float AbsoluteTolerance) const;
 	
 	friend FArchive& operator<<(FArchive& Ar, FSpline& Spline)
 	{
diff --git a/Engine/Source/Runtime/Engine/Private/Curves/Spline.cpp b/Engine/Source/Runtime/Engine/Private/Curves/Spline.cpp
--- a/Engine/Source/Runtime/Engine/Private/Curves/Spline.cpp
+++ b/Engine/Source/Runtime/Engine/Private/Curves/Spline.cpp
@@ -1416,6 +1416,128 @@ bool FSpline::operator==(const FSpline& Other) const
 	return false;
 }
 
+bool FSpline::operator!=(const FSpline& Other) const
+{
+	return !(*this == Other);
+}
+
+bool FSpline::IsNearlyEqual(const FSpline& Other, const float RelativeTolerance, const float AbsoluteTolerance) const
+{
+	if (!ensureAlways(IsNew() && Other.IsNew()))
+	{
+		return false;
+	}
+
+	const int32 NumPoints = NewData->GetSpline().GetNumPoints();
+	if (Other.NewData->GetSpline().GetNumPoints() != NumPoints)
+	{
+		return false;
+	}
+
+	auto ScalarNearlyEqual = [AbsoluteTolerance, RelativeTolerance](const float& A, const float& B)
+	{
+		const float Delta = FMath::Abs(A - B);
+		const float Tolerance = FMath::Max(AbsoluteTolerance, RelativeTolerance * FMath::Max(FMath::Abs(A), FMath::Abs(B)));
+		return Delta <= Tolerance;
+	};
+
+	auto VectorNearlyEqual = [AbsoluteTolerance, RelativeTolerance](const FVector& A, const FVector& B)
+	{
+		const float MagDelta = (A - B).Length();
+		const float MagTolerance = FMath::Max(AbsoluteTolerance, RelativeTolerance * FMath::Max(A.Length(), B.Length()));
+		return MagDelta <= MagTolerance;
+	};
+
+	auto QuatComponentsNearlyEqual = [AbsoluteTolerance, RelativeTolerance](const FQuat& A, const FQuat& B)
+	{
+		const float DX = FMath::Abs(A.X - B.X);
+		const float TX = FMath::Max(AbsoluteTolerance, RelativeTolerance * FMath::Max(FMath::Abs(A.X), FMath::Abs(B.X)));
+		if (DX > TX) return false;
+
+		const float DY = FMath::Abs(A.Y - B.Y);
+		const float TY = FMath::Max(AbsoluteTolerance, RelativeTolerance * FMath::Max(FMath::Abs(A.Y), FMath::Abs(B.Y)));
+		if (DY > TY) return false;
+
+		const float DZ = FMath::Abs(A.Z - B.Z);
+		const float TZ = FMath::Max(AbsoluteTolerance, RelativeTolerance * FMath::Max(FMath::Abs(A.Z), FMath::Abs(B.Z)));
+		if (DZ > TZ) return false;
+
+		const float DW = FMath::Abs(A.W - B.W);
+		const float TW = FMath::Max(AbsoluteTolerance, RelativeTolerance * FMath::Max(FMath::Abs(A.W), FMath::Abs(B.W)));
+		return DW <= TW;
+	};
+
+	for (int32 PointIdx = 0; PointIdx < NumPoints; ++PointIdx)
+	{
+		if (!VectorNearlyEqual(NewData->GetLocation(PointIdx), Other.NewData->GetLocation(PointIdx)))
+		{
+			return false;
+		}
+
+		if (!VectorNearlyEqual(NewData->GetInTangent(PointIdx), Other.NewData->GetInTangent(PointIdx)))
+		{
+			return false;
+		}
+
+		if (!VectorNearlyEqual(NewData->GetOutTangent(PointIdx), Other.NewData->GetOutTangent(PointIdx)))
+		{
+			return false;
+		}
+
+		const FQuat Rotation = NewData->GetRotation(PointIdx);
+		const FQuat OtherRotation = Other.NewData->GetRotation(PointIdx);
+		const FQuat OtherRotationNeg = FQuat(-OtherRotation.X, -OtherRotation.Y, -OtherRotation.Z, -OtherRotation.W);
+		if (!QuatComponentsNearlyEqual(Rotation, OtherRotation) && !QuatComponentsNearlyEqual(Rotation, OtherRotationNeg))
+		{
+			return false;
+		}
+
+		if (!VectorNearlyEqual(NewData->GetScale(PointIdx), Other.NewData->GetScale(PointIdx)))
+		{
+			return false;
+		}
+
+		if (!ScalarNearlyEqual(NewData->GetParameterAtIndex(PointIdx), Other.NewData->GetParameterAtIndex(PointIdx)))
+		{
+			return false;
+		}
+	}
+
+	const TArray<FName> OtherFloatChannelNames = Other.GetFloatPropertyChannels();
+	for (const FName& FloatChannelName : GetFloatPropertyChannels())
+	{
+		if (!OtherFloatChannelNames.Contains(FloatChannelName))
+		{
+			return false;
+		}
+
+		const int32 NumFloatChannelPoints = NumAttributeValues<float>(FloatChannelName);
+		if (Other.NumAttributeValues<float>(FloatChannelName) != NumFloatChannelPoints)
+		{
+			return false;
+		}
+
+		for (int32 FloatChannelIdx = 0; FloatChannelIdx < NumFloatChannelPoints; ++FloatChannelIdx)
+		{
+			if (!ScalarNearlyEqual(
+				GetAttributeValue<float>(FloatChannelIdx, FloatChannelName),
+				Other.GetAttributeValue<float>(FloatChannelIdx, FloatChannelName)))
+			{
+				return false;
+			}
+
+			if (!ScalarNearlyEqual(
+				GetAttributeParameter<float>(FloatChannelIdx, FloatChannelName),
+				Other.GetAttributeParameter<float>(FloatChannelIdx, FloatChannelName)))
+			{
+				return false;
+			}
+		}
+	}
+
+	return true;
+}
+
 bool FSpline::Serialize(FArchive& Ar)
 {
 	Ar.UsingCustomVersion(FFortniteMainBranchObjectVersion::GUID);
