Caching max size from an access that isn't triggering a race condition

#Synced-CL 48915086

diff --git a/Engine/Source/Runtime/Engine/Private/UnrealEngine.cpp b/Engine/Source/Runtime/Engine/Private/UnrealEngine.cpp
--- a/Engine/Source/Runtime/Engine/Private/UnrealEngine.cpp
+++ b/Engine/Source/Runtime/Engine/Private/UnrealEngine.cpp
@@ -13257,7 +13257,8 @@ void DrawStatsHUD( UWorld* World, FViewport* Viewport, FCanvas* Canvas, UCanvas*
 
 	float DPIScale = Canvas->GetDPIScale();
 	
-	FIntPoint TextureSize = Canvas->GetRenderTarget()->GetSizeXY();
+	const FIntPoint TextureSize = Canvas->GetRenderTarget()->GetSizeXY();
+	const int32 MaxY = TextureSize.Y;
 
 	//@todo joeg: Move this stuff to a function, make safe to use on consoles by
 	// respecting the various safe zones, and make it compile out.
@@ -13492,7 +13493,7 @@ void DrawStatsHUD( UWorld* World, FViewport* Viewport, FCanvas* Canvas, UCanvas*
 	// draw debug properties
 #if !UE_BUILD_SHIPPING
 
-	auto DrawDebugPropertiesForWorld = [Canvas, CanvasObject, &DebugProperties](UWorld* DebugWorld, int32& X, int32& Y)
+	auto DrawDebugPropertiesForWorld = [&](UWorld* DebugWorld, int32& X, int32& Y)
 	{
 		// construct a list of objects relevant to "getall" type elements, so that we only have to do the object iterator once
 		// we do the iterator each frame so that new objects will show up immediately
@@ -13551,7 +13552,7 @@ void DrawStatsHUD( UWorld* World, FViewport* Viewport, FCanvas* Canvas, UCanvas*
 				}
 			}
 		}
-		int32 MaxY = int32(Canvas->GetRenderTarget()->GetSizeXY().Y);
+
 		for (int32 i = 0; i < DebugProperties.Num() && Y < MaxY; i++)
 		{
 			// we removed entries with invalid Obj above so no need to check for that here
