PR #14098: fix for exporting struct array properties as json

#Synced-CL 48800164

diff --git a/Engine/Source/Runtime/Engine/Private/DataTableJSON.cpp b/Engine/Source/Runtime/Engine/Private/DataTableJSON.cpp
--- a/Engine/Source/Runtime/Engine/Private/DataTableJSON.cpp
+++ b/Engine/Source/Runtime/Engine/Private/DataTableJSON.cpp
@@ -1,4 +1,4 @@
-// Copyright Epic Games, Inc. All Rights Reserved.
+﻿// Copyright Epic Games, Inc. All Rights Reserved.
 
 #include "DataTableJSON.h"
 #include "Engine/DataTable.h"
@@ -420,7 +420,8 @@ bool TDataTableExporterJSON<CharType>::WriteStructUsingVisitor(TNotNull<const US
 	using namespace UE::DataTableJSON;
 
 	// Keep a stack of identifier overrides since when entering arrays we don't want struct to output their identifier
-	// We always push override for structs and arrays so an empty override list or an unset value on top of the stack
+	// We always push override for structs (when UseJsonObjectsForStructs is set in the export flags) and arrays 
+	// so an empty override list or an unset value on top of the stack
 	// indicates that no override is effective and the property export name should be used.
 	TArray<TOptional<FString>> IdentifierOverrides;
 
@@ -476,10 +477,9 @@ bool TDataTableExporterJSON<CharType>::WriteStructUsingVisitor(TNotNull<const US
 					{
 						JsonWriter->WriteObjectStart(Identifier);
 					}
+					// Push a "no override" entry when entering a struct to export it as json object. If UseJsonObjectsForStructs is not used the struct is saved as string value.
+					IdentifierOverrides.AddDefaulted();
 				}
-
-				// Push a "no override" entry when entering a struct
-				IdentifierOverrides.AddDefaulted();
 			}
 		};
 
@@ -514,9 +514,10 @@ bool TDataTableExporterJSON<CharType>::WriteStructUsingVisitor(TNotNull<const US
 			{
 				if (EnumHasAnyFlags(DTExportFlags, EDataTableExportFlags::UseJsonObjectsForStructs))
 				{
+					// Only close object and pop override if saving structs as json objects. If not they are saved as string values so no need to close it
 					JsonWriter->WriteObjectEnd();
+					IdentifierOverrides.Pop();
 				}
-				IdentifierOverrides.Pop();
 			}
 
 			// Process static array last to properly close an array of structs
