Fix for invalid cached nav bounds infecting other cached bounds.

#Synced-CL 48922636

diff --git a/Engine/Source/Runtime/Engine/Classes/Components/InstancedStaticMeshComponent.h b/Engine/Source/Runtime/Engine/Classes/Components/InstancedStaticMeshComponent.h
--- a/Engine/Source/Runtime/Engine/Classes/Components/InstancedStaticMeshComponent.h
+++ b/Engine/Source/Runtime/Engine/Classes/Components/InstancedStaticMeshComponent.h
@@ -627,13 +627,9 @@ public:
 	void OnPostPopulatePerInstanceData() { OnPostLoadPerInstanceData(); }
 
 	/** Request to navigation system to update only part of navmesh occupied by specified instance. */
-	ENGINE_API virtual void PartialNavigationUpdate(int32 InstanceIdx, bool bUpdateCachedNavigationBounds);
-	UE_DEPRECATED(5.8, "Use the overload with the flag indicating whether the bounds need to be updated.")
 	ENGINE_API virtual void PartialNavigationUpdate(int32 InstanceIdx) final;
 
 	/** Request to navigation system to update only part of navmesh occupied specified instances transforms. */
-	ENGINE_API virtual void PartialNavigationUpdates(TConstArrayView<FTransform> InstanceTransforms, bool bUpdateCachedNavigationBounds);
-	UE_DEPRECATED(5.8, "Use the overload with the flag indicating whether the bounds need to be updated.")
 	ENGINE_API virtual void PartialNavigationUpdates(TConstArrayView<FTransform> InstanceTransforms) final;
 
 	ENGINE_API TSharedPtr<FISMCInstanceDataSceneProxy, ESPMode::ThreadSafe> GetOrCreateInstanceDataSceneProxy();
@@ -695,6 +691,7 @@ protected:
 
 	void InvalidateCachedBounds();
 
+	UE_DEPRECATED(5.8, "Navigation bounds are now lazily updated and no longer require explicit calls to update them.")
 	ENGINE_API void CalcAndCacheNavigationBounds();
 
 	/** Creates body instances for all instances owned by this component. */
diff --git a/Engine/Source/Runtime/Engine/Private/InstancedStaticMesh.cpp b/Engine/Source/Runtime/Engine/Private/InstancedStaticMesh.cpp
--- a/Engine/Source/Runtime/Engine/Private/InstancedStaticMesh.cpp
+++ b/Engine/Source/Runtime/Engine/Private/InstancedStaticMesh.cpp
@@ -3193,7 +3193,7 @@ void UInstancedStaticMeshComponent::UpdateBounds()
 
 	if (bCanEverAffectNavigation && (IsRegistered() || IsPreRegistering()))
 	{
-		CalcAndCacheNavigationBounds();
+		GetOrCalcBounds(GetComponentTransform(), EBoundsCacheType::Navigation);
 	}
 }
 
@@ -3646,6 +3646,8 @@ void UInstancedStaticMeshComponent::PreAllocateInstancesMemory(int32 AddedInstan
 
 int32 UInstancedStaticMeshComponent::AddInstanceInternal(int32 InstanceIndex, FInstancedStaticMeshInstanceData* InNewInstanceData, const FTransform& InstanceTransform, bool bWorldSpace)
 {
+	InvalidateCachedBounds();
+
 	FInstancedStaticMeshInstanceData* NewInstanceData = InNewInstanceData;
 
 	PrimitiveInstanceDataManager.Add(InstanceIndex, NewInstanceData != nullptr);
@@ -3677,9 +3679,6 @@ int32 UInstancedStaticMeshComponent::AddInstanceInternal(int32 InstanceIndex, FI
 	// Update navigation relevancy
 	bNavigationRelevant = IsNavigationRelevant();
 
-	// Indicates if the cached bounds were updated.
-	bool bNavigationCachedBoundsUpdated = false;
-	
 	// Perform navigation update if the component is relevant to navigation and registered.
 	// Otherwise this will be handled in OnRegister (e.g. Editor manipulations).
 	if (bNavigationRelevant && IsRegistered())
@@ -3688,28 +3687,17 @@ int32 UInstancedStaticMeshComponent::AddInstanceInternal(int32 InstanceIndex, FI
 		// since it was skipped because IsNavigationRelevant() requires at least one instance.
 		if (GetInstanceCount() == 1)
 		{
-			CalcAndCacheNavigationBounds();
-			bNavigationCachedBoundsUpdated = true;
-
 			FNavigationSystem::RegisterComponent(*this);
 		}
 
 		if (SupportsPartialNavigationUpdate())
 		{
-			PartialNavigationUpdate(InstanceIndex,  /*bUpdateCachedNavigationBounds*/ !bNavigationCachedBoundsUpdated);
+			PartialNavigationUpdate(InstanceIndex);
 		}
 		else
 		{
-			if (!bNavigationCachedBoundsUpdated)
-			{
-				CalcAndCacheNavigationBounds();
-			}
-
 			FullNavigationUpdate();
 		}
-
-		// At this point navigation bound were calculated and cached
-		bNavigationCachedBoundsUpdated = true;
 	}
 
 	if (FInstancedStaticMeshDelegates::OnInstanceIndexUpdated.IsBound())
@@ -3717,15 +3705,6 @@ int32 UInstancedStaticMeshComponent::AddInstanceInternal(int32 InstanceIndex, FI
 		FInstancedStaticMeshDelegates::FInstanceIndexUpdateData IndexUpdate{ FInstancedStaticMeshDelegates::EInstanceIndexUpdateType::Added, InstanceIndex };
 		FInstancedStaticMeshDelegates::OnInstanceIndexUpdated.Broadcast(this, MakeArrayView(&IndexUpdate, 1));
 	}
-
-	const FBoundsCacheElement NavigationCacheElement = CachedBounds[(int32)EBoundsCacheType::Navigation];
-	InvalidateCachedBounds();
-
-	// Restore Navigation CachedBounds (the invalidation is not needed if it was updated).
-	if (bNavigationCachedBoundsUpdated)
-	{
-		CachedBounds[(int32)EBoundsCacheType::Navigation] = NavigationCacheElement;
-	}
 	
 	return InstanceIndex;
 }
@@ -3740,6 +3719,8 @@ TArray<int32> UInstancedStaticMeshComponent::AddInstancesInternal(TConstArrayVie
 	TRACE_CPUPROFILER_EVENT_SCOPE(UInstancedStaticMeshComponent::AddInstancesInternal);
 	CSV_SCOPED_TIMING_STAT_EXCLUSIVE(UInstancedStaticMeshComponent_AddInstancesInternal);
 
+	InvalidateCachedBounds();
+
 	const int32 Count = InstanceTransforms.Num();
 
 	TArray<int32> NewInstanceIndices;
@@ -3765,9 +3746,6 @@ TArray<int32> UInstancedStaticMeshComponent::AddInstancesInternal(TConstArrayVie
 	// and IsNavigationRelevant() takes the number of instance into account which can be 0 here.
 	bNavigationRelevant = NumInstances > 0 && Super::IsNavigationRelevant();
 
-	// Indicates if the cached bounds were updated.
-	bool bNavigationCachedBoundsUpdated = false;
-
 	// Navigation update required if explicitly requested and if the component is relevant to navigation and registered.
 	// Otherwise this will be handled in OnRegister (e.g. Editor manipulations).
 	const bool bNavigationUpdateRequired = bUpdateNavigation && bNavigationRelevant && IsRegistered();
@@ -3822,37 +3800,17 @@ TArray<int32> UInstancedStaticMeshComponent::AddInstancesInternal(TConstArrayVie
 	{
 		if (bNeedsNavigationRegistration)
 		{
-			CalcAndCacheNavigationBounds();
-			bNavigationCachedBoundsUpdated = true;
-
 			FNavigationSystem::RegisterComponent(*this);
 		}
 
 		if (bPartialNavigationUpdateRequired)
 		{
-			PartialNavigationUpdates(NavigationUpdateTransforms,  /*bUpdateCachedNavigationBounds*/ !bNavigationCachedBoundsUpdated);
+			PartialNavigationUpdates(NavigationUpdateTransforms);
 		}
 		else
 		{
-			if (!bNavigationCachedBoundsUpdated)
-			{
-				CalcAndCacheNavigationBounds();
-			}
-
 			FullNavigationUpdate();
 		}
-
-		// At this point navigation bound were calculated and cached
-		bNavigationCachedBoundsUpdated = true;
-	}
-
-	const FBoundsCacheElement NavigationCacheElement = CachedBounds[(int32)EBoundsCacheType::Navigation];
-	InvalidateCachedBounds();
-
-	// Restore Navigation CachedBounds (the invalidation is not needed if it was updated).
-	if (bNavigationCachedBoundsUpdated)
-	{
-		CachedBounds[(int32)EBoundsCacheType::Navigation] = NavigationCacheElement;
 	}
 
 	return NewInstanceIndices;
@@ -4082,7 +4040,7 @@ bool UInstancedStaticMeshComponent::RemoveInstanceInternal(int32 InstanceIndex,
 		// the instance's transform to dirty the area and full update will dirty the whole covered by the navigation bounds.
 		if (bNavigationUpdateRequired && SupportsPartialNavigationUpdate())
 		{
-			PartialNavigationUpdate(InstanceIndex, /*bUpdateCachedNavigationBounds*/ true );
+			PartialNavigationUpdate(InstanceIndex);
 		}
 
 		if (bUseRemoveAtSwap)
@@ -4347,7 +4305,7 @@ void UInstancedStaticMeshComponent::UpdateComponentTransform(EUpdateTransformFla
 		// Only handle partial updates since modifying the transform is handled by base class (i.e. SceneComponent).
 		if (bDoPartialNavigationUpdate)
 		{
-			PartialNavigationUpdates(NavigationUpdateTransforms, /*bUpdateCachedNavigationBounds*/ true );
+			PartialNavigationUpdates(NavigationUpdateTransforms);
 		}
 	}
 
@@ -4455,7 +4413,7 @@ bool UInstancedStaticMeshComponent::UpdateInstanceTransform(const int32 Instance
 				{
 					PreviousTransform * GetComponentTransform(),
 					(bWorldSpace ? NewInstanceTransform : NewInstanceTransform * GetComponentTransform())
-				}, /*bUpdateCachedNavigationBounds*/ true );
+				} );
 		}
 		else
 		{
@@ -4541,7 +4499,7 @@ bool UInstancedStaticMeshComponent::BatchUpdateInstancesTransforms(int32 StartIn
 	{
 		if (bPartialNavigationUpdateRequired)
 		{
-			PartialNavigationUpdates(NavigationUpdateTransforms, /*bUpdateCachedNavigationBounds*/ true);
+			PartialNavigationUpdates(NavigationUpdateTransforms);
 		}
 		else
 		{
@@ -4628,7 +4586,7 @@ bool UInstancedStaticMeshComponent::BatchUpdateInstancesTransformsInternal(int32
 	{
 		if (bPartialNavigationUpdateRequired)
 		{
-			PartialNavigationUpdates(NavigationUpdateTransforms, /*bUpdateCachedNavigationBounds*/ true);
+			PartialNavigationUpdates(NavigationUpdateTransforms);
 		}
 		else
 		{
@@ -4699,7 +4657,7 @@ bool UInstancedStaticMeshComponent::BatchUpdateInstancesTransform(int32 StartIns
 	{
 		if (bPartialNavigationUpdateRequired)
 		{
-			PartialNavigationUpdates(NavigationUpdateTransforms, /*bUpdateCachedNavigationBounds*/ true);
+			PartialNavigationUpdates(NavigationUpdateTransforms);
 		}
 		else
 		{
@@ -4765,7 +4723,7 @@ bool UInstancedStaticMeshComponent::BatchUpdateInstancesData(int32 StartInstance
 	{
 		if (bPartialNavigationUpdateRequired)
 		{
-			PartialNavigationUpdates(NavigationUpdateTransforms, /*bUpdateCachedNavigationBounds*/ true);
+			PartialNavigationUpdates(NavigationUpdateTransforms);
 		}
 		else
 		{
@@ -5483,10 +5441,10 @@ void UInstancedStaticMeshComponent::PartialNavigateUpdateForCurrentInstances()
 		NavigationUpdateTransforms.Add(FTransform(PerInstanceSMData[i].Transform) * GetComponentTransform());
 	}
 
-	PartialNavigationUpdates(NavigationUpdateTransforms, /*bUpdateCachedNavigationBounds*/ true);
+	PartialNavigationUpdates(NavigationUpdateTransforms);
 }
 
-void UInstancedStaticMeshComponent::PartialNavigationUpdate(int32 InstanceIdx, const bool bUpdateCachedNavigationBounds)
+void UInstancedStaticMeshComponent::PartialNavigationUpdate(int32 InstanceIdx)
 {
 	TRACE_CPUPROFILER_EVENT_SCOPE(UInstancedStaticMeshComponent::PartialNavigationUpdate);
 	if (!IsNavigationRelevant())
@@ -5497,16 +5455,11 @@ void UInstancedStaticMeshComponent::PartialNavigationUpdate(int32 InstanceIdx, c
 	FTransform InstanceTransform;
 	if (GetInstanceTransform(InstanceIdx, InstanceTransform, /*bWorldSpace*/true))
 	{
-		PartialNavigationUpdates({InstanceTransform}, bUpdateCachedNavigationBounds);
+		PartialNavigationUpdates({InstanceTransform});
 	}
 }
 
-void UInstancedStaticMeshComponent::PartialNavigationUpdate(int32 InstanceIdx)
-{
-	PartialNavigationUpdate(InstanceIdx, /*bUpdateCachedNavigationBounds*/ true);
-}
-
-void UInstancedStaticMeshComponent::PartialNavigationUpdates(const TConstArrayView<FTransform> InstanceTransforms, const bool bUpdateCachedNavigationBounds)
+void UInstancedStaticMeshComponent::PartialNavigationUpdates(const TConstArrayView<FTransform> InstanceTransforms)
 {
 	if (!IsNavigationRelevant() || InstanceTransforms.IsEmpty())
 	{
@@ -5516,16 +5469,10 @@ void UInstancedStaticMeshComponent::PartialNavigationUpdates(const TConstArrayVi
 	const FBox InstanceBounds = GetInstanceNavigationBounds();
 	if (InstanceBounds.IsValid)
 	{
-		if( bUpdateCachedNavigationBounds )
+		if (CVarISMInvalidateCachedNavigationBoundsOnPartialNavigationUpdates.GetValueOnAnyThread() != 0)
 		{
-			if (CVarISMInvalidateCachedNavigationBoundsOnPartialNavigationUpdates.GetValueOnAnyThread() != 0)
-			{
-				// Make sure to invalidate cached navigation bounds that might have been computed from a different number of instances.
-				CachedBounds[static_cast<int32>(EBoundsCacheType::Navigation)].Invalidate();
-			}
-
-			// Update cached navigation bounds
-			CalcAndCacheNavigationBounds();
+			// Make sure to invalidate cached navigation bounds that might have been computed from a different number of instances.
+			CachedBounds[static_cast<int32>(EBoundsCacheType::Navigation)].Invalidate();
 		}
 
 		// Dirty areas around the modified instances
@@ -5539,12 +5486,6 @@ void UInstancedStaticMeshComponent::PartialNavigationUpdates(const TConstArrayVi
 	}
 }
 
-void UInstancedStaticMeshComponent::PartialNavigationUpdates(const TConstArrayView<FTransform> InstanceTransforms)
-{
-	PartialNavigationUpdates(InstanceTransforms, /*bUpdateCachedNavigationBounds*/ true);
-}
-
-
 bool UInstancedStaticMeshComponent::DoCustomNavigableGeometryExport(FNavigableGeometryExport& GeomExport) const
 {
 	return FInstancedStaticMeshComponentHelper::DoCustomNavigableGeometryExport(*this, GeomExport, FNavDataPerInstanceTransformDelegate::CreateUObject(this, &UInstancedStaticMeshComponent::GetNavigationPerInstanceTransforms));
@@ -5699,9 +5640,7 @@ void UInstancedStaticMeshComponent::GetNavigationData(FNavigationRelevantData& D
 
 FBox UInstancedStaticMeshComponent::GetNavigationBounds() const
 {
-	// Navigation bounds must be calculated before being requested
-	constexpr bool bCheckValid = true;
-	return CachedBounds[(int)EBoundsCacheType::Navigation].GetValue(bCheckValid).GetBox();
+	return GetOrCalcBounds(GetComponentTransform(), EBoundsCacheType::Navigation).GetBox();
 }
 
 bool UInstancedStaticMeshComponent::IsNavigationRelevant() const
