Chaos, Networked Physics: Renaming some functions involved in lag adaptive scaling of input decay. Better behavior above the reference lag, to avoid completely decaying input as before but still maintaining similar movement behavior when the lag is higher.

#Synced-CL 48814712

diff --git a/Engine/Source/Runtime/Engine/Private/PhysicsEngine/NetworkPhysicsComponent.cpp b/Engine/Source/Runtime/Engine/Private/PhysicsEngine/NetworkPhysicsComponent.cpp
--- a/Engine/Source/Runtime/Engine/Private/PhysicsEngine/NetworkPhysicsComponent.cpp
+++ b/Engine/Source/Runtime/Engine/Private/PhysicsEngine/NetworkPhysicsComponent.cpp
@@ -96,10 +96,10 @@ namespace PhysicsReplicationCVars
 		static FAutoConsoleVariableRef CVarApplyInputDecayOverSetTime(TEXT("np2.Resim.ApplyInputDecayOverSetTime"), bApplyInputDecayOverSetTime, TEXT("When true, apply the Input Decay Curve over a set amount of time instead of over the start of input prediction and end of resim which is variable each resimulation"));
 		float InputDecaySetTime = 0.15f;
 		static FAutoConsoleVariableRef CVarInputDecaySetTime(TEXT("np2.Resim.InputDecaySetTime"), InputDecaySetTime, TEXT("Applied when np2.Resim.ApplyInputDecayOverSetTime is true, read there for more info. Set time to apply Input Decay Curve over while predicting inputs during resimulation"));
-		bool bEnableLagScaledInputDecay = false;
-		static FAutoConsoleVariableRef CVarEnableLagScaledInputDecay(TEXT("np2.Resim.EnableLagScaledInputDecay"), bEnableLagScaledInputDecay, TEXT("If true, scales input decay as a proportion of measured input lag compared to the reference input lag (np2.Resim.InputDecayReferenceLagMs)"));
+		bool bEnableLagScalingInputDecay = false;
+		static FAutoConsoleVariableRef CVarEnableLagScaledInputDecay(TEXT("np2.Resim.EnableLagScalingInputDecay"), bEnableLagScalingInputDecay, TEXT("If true, scales input decay as a proportion of measured input lag compared to the reference input lag (np2.Resim.InputDecayReferenceLagMs)"));
 		float InputDecayReferenceLagMs = 100.0f;
-		static FAutoConsoleVariableRef CVarInputDecayReferenceLagMs(TEXT("np2.Resim.InputDecayReferenceLagMs"), InputDecayReferenceLagMs, TEXT("The reference input lag in milliseconds used to scale input decay. Only applies if np2.Resim.EnableLagScaledInputDecay is true"));
+		static FAutoConsoleVariableRef CVarInputDecayReferenceLagMs(TEXT("np2.Resim.InputDecayReferenceLagMs"), InputDecayReferenceLagMs, TEXT("The reference input lag in milliseconds used to scale input decay. Only applies if np2.Resim.EnableLagScalingInputDecay is true"));
 		bool bApplyInputDecaySimProxyInputAtRuntime = false;
 		static FAutoConsoleVariableRef CVarApplyInputDecaySimProxyInputAtRuntime(TEXT("np2.Resim.ApplyInputDecaySimProxyInputAtRuntime"), bApplyInputDecaySimProxyInputAtRuntime, TEXT("When true, apply input decay on inputs applied on simulated proxies at runtime (outside of resimulation)"));
 		float InputDecaySimProxyInputAtRuntime = 0.25f;
@@ -2850,13 +2850,34 @@ const float FAsyncNetworkPhysicsComponent::GetCurrentInputDecay() const
 	float InputDecay = InputDecayCurve.GetRichCurveConst()->Eval(PredictionAlpha);
 
 	// If enabled, scale input decay as a function of measured lag
-	if (ComponentSettings.GetEnableLerpInputDecayWithLag())
+	if (ComponentSettings.GetEnableLagScalingInputDecay())
 	{
-		const float MeasuredInputDelaySeconds = GetInputLagFromLastReceivedTargetState();
-		const float ReferenceInputDelaySeconds = ComponentSettings.GetInputDecayReferenceLagMs() * 0.001f;
-		// We do not clamp this lerp value to 1.0; if the measured delay is greater than the reference one, we apply a greater decay
-		const float LerpValue = (ReferenceInputDelaySeconds >= UE_SMALL_NUMBER) ? (MeasuredInputDelaySeconds / ReferenceInputDelaySeconds) : 1.0f;
-		InputDecay *= LerpValue;
+		const float ReferenceLagSeconds = ComponentSettings.GetInputDecayReferenceLagMs() * 0.001f;
+		if (ReferenceLagSeconds >= UE_SMALL_NUMBER) // Setting ReferenceLagSeconds to 0 means don't scale with lag
+		{
+			const float ForwardPredictionTimeSeconds = GetForwardPredictionTime();
+			if (ForwardPredictionTimeSeconds < UE_SMALL_NUMBER)
+			{
+				InputDecay = 0.0f; // no lag -> no decay
+			}
+			else if (ForwardPredictionTimeSeconds <= ReferenceLagSeconds)
+			{
+				InputDecay *= ForwardPredictionTimeSeconds / ReferenceLagSeconds;
+			}
+			else
+			{
+				// The formula below is designed to maintain the behavior tuned at ReferenceLagSeconds when we get higher lag
+				// We interpret getting "similar behavior at higher lag" as maintaining the same amount of travel for a given input as that we get at the reference lag.
+				// For instance, if the lag doubles relative to the reference lag, we want to travel half the distance per second for the same input, because we will be dead reckoning
+				// for twice as long with a given input before being corrected. If d(x) is the function that gives us the decay, x being the lag
+				// the modified input i(x) for that lag will be given by i(x) = I * (1 - d(x)), I being the unmodified received input.
+				// We have d(L) = D, a value of decay D that feels good at the reference lag L
+				// The equal distance at different lag can be written
+				// (1 - d(x)) x = (1 - d(L)) L
+				// So d(x) = 1 - L/x * (1-D)
+				InputDecay = 1.0f - (ReferenceLagSeconds / ForwardPredictionTimeSeconds) * (1.0f - InputDecay);
+			}
+		}
 	}
 
 	return FMath::Clamp(InputDecay, 0.0f, 1.0f);
@@ -2868,24 +2889,53 @@ const float FAsyncNetworkPhysicsComponent::GetCurrentSimProxyInputDecayAtRuntime
 	float SimProxyRuntimeInputDecay = ComponentSettings.GetInputDecaySimProxyInputAtRuntime();
 	if (ComponentSettings.GetEnableLagScalingSimProxyRuntimeInputDecay())
 	{
-		const float MeasuredInputDelaySeconds = GetInputLagFromLastReceivedTargetState();
-		const float ReferenceInputDelaySeconds = ComponentSettings.GetInputDecayReferenceLagMs() * 0.001f;
-		// We do not clamp this lerp value to 1.0; if the measured delay is greater than the reference one, we apply a greater decay
-		float LerpValue = (ReferenceInputDelaySeconds >= UE_SMALL_NUMBER) ? (MeasuredInputDelaySeconds / ReferenceInputDelaySeconds) : 1.0f;
-		SimProxyRuntimeInputDecay *= LerpValue;
+		const float ReferenceLagSeconds = ComponentSettings.GetInputDecayReferenceLagMs() * 0.001f;
+		if (ReferenceLagSeconds >= UE_SMALL_NUMBER) // Setting ReferenceLagSeconds to 0 means don't scale with lag
+		{
+			const float ForwardPredictionTimeSeconds = GetForwardPredictionTime();
+			if (ForwardPredictionTimeSeconds < UE_SMALL_NUMBER)
+			{
+				SimProxyRuntimeInputDecay = 0.0f; // no lag -> no decay
+			}
+			else if (ForwardPredictionTimeSeconds <= ReferenceLagSeconds)
+			{
+				SimProxyRuntimeInputDecay *= ForwardPredictionTimeSeconds / ReferenceLagSeconds;
+			}
+			else
+			{
+				 // This formula is explained in GetCurrentInputDecay()
+				SimProxyRuntimeInputDecay = 1.0f - (ReferenceLagSeconds / ForwardPredictionTimeSeconds) * (1.0f - SimProxyRuntimeInputDecay);
+			}
+		}
 	}
 	return FMath::Clamp(SimProxyRuntimeInputDecay, 0.0f, 1.0f);
 }
 
-const float FAsyncNetworkPhysicsComponent::GetInputLagFromLastReceivedTargetState() const
+const int32 FAsyncNetworkPhysicsComponent::GetLatestReceivedTargetStateFrame() const
+{
+	const Chaos::FPhysicsSolverBase* BaseSolver = GetSolver();
+	if (const Chaos::FRewindData* RewindData = BaseSolver ? BaseSolver->GetRewindData() : nullptr)
+	{
+		// Todo: we should not have to get a read lock to get this information but instead measure the latest lag upon reception
+		Chaos::FReadPhysicsObjectInterface_Internal Interface = Chaos::FPhysicsObjectInternalInterface::GetRead();
+		if (Chaos::FGeometryParticleHandle* Particle = Interface.GetParticle(PhysicsObject))
+		{
+			return RewindData->GetLatestReceivedTargetStateFrame(*Particle);
+		}
+	}
+
+	return INDEX_NONE;
+}
+
+const float FAsyncNetworkPhysicsComponent::GetForwardPredictionTime() const
 {
 	const Chaos::FPhysicsSolverBase* BaseSolver = GetSolver();
 	if (const Chaos::FRewindData* RewindData = BaseSolver ? BaseSolver->GetRewindData() : nullptr)
 	{
-		const int32 LatestTargetFrame = GetLastReceivedServerFrame();
-		if (LatestTargetFrame >= 0)
+		int32 LatestReceivedTargetStateFrame = GetLatestReceivedTargetStateFrame();
+		if (LatestReceivedTargetStateFrame != INDEX_NONE)
 		{
-			const int32 MeasuredInputDelayFrames = RewindData->CurrentFrame() - LatestTargetFrame;
+			const int32 MeasuredInputDelayFrames = RewindData->GetLatestFrame() - LatestReceivedTargetStateFrame;
 			const float MeasuredInputDelaySeconds = FMath::Max(MeasuredInputDelayFrames, 0.0f) * BaseSolver->GetAsyncDeltaTime();
 			return MeasuredInputDelaySeconds;
 		}
diff --git a/Engine/Source/Runtime/Engine/Public/Physics/NetworkPhysicsComponent.h b/Engine/Source/Runtime/Engine/Public/Physics/NetworkPhysicsComponent.h
--- a/Engine/Source/Runtime/Engine/Public/Physics/NetworkPhysicsComponent.h
+++ b/Engine/Source/Runtime/Engine/Public/Physics/NetworkPhysicsComponent.h
@@ -1937,6 +1937,20 @@ public:
 	// Get the frame number of the last server update
 	ENGINE_API int32 GetLastReceivedServerFrame() const;
 
+	/** Returns the time in seconds between the frame of the latest received state and the latest frame we have reached.
+	  * This can be used as a quantity that is stable over the course of a resimulation (from rewind to last frame of resim) and grows with the amount of time being resimulated
+	  * We use this to indicate how much more we should decay input for instance when we have high RTT compared to some baseline decay
+	  */
+	ENGINE_API const float GetForwardPredictionTime() const;
+
+	/** Returns the current amount of input decay during resimulation as a magnitude from 0.0 to 1.0. Returns 0 if not currently resimulating. */
+	ENGINE_API const float GetCurrentInputDecay() const;
+
+	/** Returns the current amount of input decay for sim proxies during simulation as a magnitude from 0.0 to 1.0. */
+	ENGINE_API const float GetCurrentSimProxyInputDecayAtRuntime() const;
+
+	ENGINE_API const int32 GetLatestReceivedTargetStateFrame() const;
+
 private:
 	// Initialize, bind delegates etc.
 	void OnInitialize_Internal();
@@ -1977,15 +1991,6 @@ private:
 	UE_DEPRECATED(5.8, "Use function that doesn't take a parameter, the decay curve will be based on the current frame number compared with the latest input frame number available")
 	const float GetCurrentInputDecay(const FNetworkPhysicsPayload* PhysicsData) { return GetCurrentInputDecay(); }
 
-	/** Returns the current amount of input decay during resimulation as a magnitude from 0.0 to 1.0. Returns 0 if not currently resimulating. */
-	const float GetCurrentInputDecay() const;
-
-	/** Returns the current amount of input decay for sim proxies during simulation as a magnitude from 0.0 to 1.0. */
-	const float GetCurrentSimProxyInputDecayAtRuntime() const;
-
-	/** Returns the input lag (in seconds) as measured based on the difference between the frame associated with the last received target and the frame at which it was received */
-	const float GetInputLagFromLastReceivedTargetState() const;
-
 	/** Populate input data in AsyncOutput to send over the network */
 	void SendInputData_Internal(FAsyncNetworkPhysicsComponentOutput& AsyncOutput, const int32 PhysicsStep);
 
diff --git a/Engine/Source/Runtime/Engine/Public/Physics/NetworkPhysicsSettingsComponent.h b/Engine/Source/Runtime/Engine/Public/Physics/NetworkPhysicsSettingsComponent.h
--- a/Engine/Source/Runtime/Engine/Public/Physics/NetworkPhysicsSettingsComponent.h
+++ b/Engine/Source/Runtime/Engine/Public/Physics/NetworkPhysicsSettingsComponent.h
@@ -83,7 +83,7 @@ namespace PhysicsReplicationCVars
 		extern float InputDecaySimProxyInputAtRuntime;
 		extern bool bApplyInputDecayOverSetTime;
 		extern float InputDecaySetTime;
-		extern bool bEnableLagScaledInputDecay;
+		extern bool bEnableLagScalingInputDecay;
 		extern float InputDecayReferenceLagMs;
 	}
 }
@@ -516,7 +516,7 @@ struct FNetworkPhysicsSettingsNetworkPhysicsComponent
 		, bOverrideEnableLagScalingSimProxyRuntimeInputDecay(0)
 		, bOverrideApplyInputDecayOverSetTime(0)
 		, bOverrideInputDecaySetTime(0)
-		, bOverrideEnableLagScaledInputDecay(0)
+		, bOverrideEnableLagScalingInputDecay(0)
 		, bOverrideInputDecayReferenceLagMs(0)
 	{};
 
@@ -574,7 +574,7 @@ struct FNetworkPhysicsSettingsNetworkPhysicsComponent
 	UPROPERTY(config, EditDefaultsOnly, Category = "Overrides", Meta = (InlineEditConditionToggle))
 	uint32 bOverrideInputDecaySetTime : 1;
 	UPROPERTY(config, EditDefaultsOnly, Category = "Overrides", Meta = (InlineEditConditionToggle))
-	uint32 bOverrideEnableLagScaledInputDecay : 1;
+	uint32 bOverrideEnableLagScalingInputDecay : 1;
 	UPROPERTY(config, EditDefaultsOnly, Category = "Overrides", Meta = (InlineEditConditionToggle))
 	uint32 bOverrideInputDecayReferenceLagMs : 1;
 
@@ -668,7 +668,7 @@ struct FNetworkPhysicsSettingsNetworkPhysicsComponent
 	float InputDecaySimProxyInputAtRuntime = PhysicsReplicationCVars::ResimulationCVars::InputDecaySimProxyInputAtRuntime;
 	const float GetInputDecaySimProxyInputAtRuntime() const { return bOverrideInputDecaySimProxyInputAtRuntime ? InputDecaySimProxyInputAtRuntime : PhysicsReplicationCVars::ResimulationCVars::InputDecaySimProxyInputAtRuntime; }
 
-	// Overrides CVar: np2.Resim.EnableLagScalingSimProxyRuntimeInputDecay -- When true, scales the input decay applied on simulated proxies at runtime (outside of resimulation) as a function of measured lag
+	// Overrides CVar: np2.Resim.EnableLagScalingSimProxyRuntimeInputDecay -- When true, scales the input decay applied on simulated proxies at runtime (outside of resimulation) as a function of measured lag (forward prediction time)
 	UPROPERTY(config, EditDefaultsOnly, Category = "Overrides", Meta = (EditCondition = "bOverrideEnableLagScalingSimProxyRuntimeInputDecay"))
 	bool bEnableLagScalingSimProxyRuntimeInputDecay = PhysicsReplicationCVars::ResimulationCVars::bEnableLagScalingSimProxyRuntimeInputDecay;
 	const bool GetEnableLagScalingSimProxyRuntimeInputDecay() const { return bOverrideEnableLagScalingSimProxyRuntimeInputDecay ? bEnableLagScalingSimProxyRuntimeInputDecay : PhysicsReplicationCVars::ResimulationCVars::bEnableLagScalingSimProxyRuntimeInputDecay; }
@@ -690,12 +690,16 @@ struct FNetworkPhysicsSettingsNetworkPhysicsComponent
 	FRuntimeFloatCurve InputDecayCurve;
 	const FRuntimeFloatCurve& GetInputDecayCurve() const { return InputDecayCurve; }
 
-	// Overrides CVar: np2.Resim.EnableLagScaledInputDecay -- When true, scales input decay as a function of measured input lag (the greater the lag, the greater the decay). The decay will be exactly the tuned InputDecay when the measured lag is equal to the reference input lag (InputDecayReferenceLagMs)
-	UPROPERTY(config, EditDefaultsOnly, Category = "Overrides", Meta = (EditCondition = "bOverrideEnableLagScaledInputDecay"))
-	bool bEnableLagScaledInputDecay = PhysicsReplicationCVars::ResimulationCVars::bEnableLagScaledInputDecay;
-	const bool GetEnableLerpInputDecayWithLag() const { return bOverrideEnableLagScaledInputDecay ? bEnableLagScaledInputDecay : PhysicsReplicationCVars::ResimulationCVars::bEnableLagScaledInputDecay; }
+	/* Overrides CVar: np2.Resim.EnableLagScalingInputDecay -- When true, scales input decay as a function of measured lag (i.e. forward prediction time). The greater the lag, the greater the decay. 
+	 * The decay will be exactly the tuned InputDecay when the measured lag is equal to the reference lag (InputDecayReferenceLagMs)
+	 **/
+	UPROPERTY(config, EditDefaultsOnly, Category = "Overrides", Meta = (EditCondition = "bOverrideEnableLagScalingInputDecay"))
+	bool bEnableLagScalingInputDecay = PhysicsReplicationCVars::ResimulationCVars::bEnableLagScalingInputDecay;
+	const bool GetEnableLagScalingInputDecay() const { return bOverrideEnableLagScalingInputDecay ? bEnableLagScalingInputDecay : PhysicsReplicationCVars::ResimulationCVars::bEnableLagScalingInputDecay; }
 
-	// Overrides CVar: np2.Resim.InputDecayReferenceLagMs -- Applied when np2.Resim.EnableLagScaledInputDecay is true. Reference lag at which the decay will be exactly InputDecay. Linearly interpolated below and beyond, reaching 0 decay at 0 lag.
+	/** Overrides CVar: np2.Resim.InputDecayReferenceLagMs -- Applied when np2.Resim.EnableLagScalingInputDecay is true. Reference lag at which the decay will be exactly InputDecay. Linearly interpolated below, reaching 0 decay at 0 lag.
+	*   Tries to maintain behavior above that lag, by decaying input more to travel the same distance over greater time for the same input.
+	*/
 	UPROPERTY(config, EditDefaultsOnly, Category = "Overrides", Meta = (EditCondition = "bOverrideInputDecayReferenceLagMs"))
 	float InputDecayReferenceLagMs = PhysicsReplicationCVars::ResimulationCVars::InputDecayReferenceLagMs;
 	const float GetInputDecayReferenceLagMs() const { return bOverrideInputDecayReferenceLagMs ? InputDecayReferenceLagMs : PhysicsReplicationCVars::ResimulationCVars::InputDecayReferenceLagMs; }
