Automated testcase to PIE and observe BP_ST_ExploPawn in the level and check if the debug text appears and disappears after the delay

#rnx

#Synced-CL 47732431

diff --git a/Engine/Source/Runtime/Engine/Private/HUD.cpp b/Engine/Source/Runtime/Engine/Private/HUD.cpp
--- a/Engine/Source/Runtime/Engine/Private/HUD.cpp
+++ b/Engine/Source/Runtime/Engine/Private/HUD.cpp
@@ -745,6 +745,12 @@ void AHUD::DrawDebugTextList()
 				DebugCanvas->DrawItem(TextItem, FVector2D(FMath::CeilToFloat(ScreenLoc.X), FMath::CeilToFloat(ScreenLoc.Y)));
 			}
 
+			UE_LOG(LogHUD, VeryVerbose, TEXT("[%f] %s: \"%s\" Loc: [%s] ")
+				, GetWorld()->GetTimeSeconds()
+				, *GetNameSafe(DebugTextList[Idx].SrcActor)
+				, *DebugTextList[Idx].DebugText
+				, *WorldTextLoc.ToString());
+
 			// do this at the end so even small durations get at least one frame
 			if (DebugTextList[Idx].TimeRemaining != -1.f)
 			{
