World Partition:
- When moving external actors into streaming levels, make sure to clear references to invalid objects.
- Backout 48413819 since it is no more required with the above change.
- Bonus: this pass will also report unreferenced objects moved from external actor packages into the streaming levels.

#Synced-CL 48911619

diff --git a/Engine/Source/Runtime/Engine/Private/WorldPartition/WorldPartitionLevelHelper.cpp b/Engine/Source/Runtime/Engine/Private/WorldPartition/WorldPartitionLevelHelper.cpp
--- a/Engine/Source/Runtime/Engine/Private/WorldPartition/WorldPartitionLevelHelper.cpp
+++ b/Engine/Source/Runtime/Engine/Private/WorldPartition/WorldPartitionLevelHelper.cpp
@@ -229,6 +229,7 @@ void FWorldPartitionLevelHelper::MoveExternalActorsToLevel(const TArray<FWorldPa
 	// Move all actors to Cell level
 	TSet<AActor*> LoadedActors;
 	TSet<UPackage*> ModifiedPackages;
+	TSet<UObject*> MovedObjects;
 
 	for (const FWorldPartitionRuntimeCellObjectMapping& PackageObjectMapping : InChildPackages)
 	{
@@ -278,12 +279,6 @@ void FWorldPartitionLevelHelper::MoveExternalActorsToLevel(const TArray<FWorldPa
 			}
 			check(Actor->GetPackage() == LevelPackage);
 
-			if (PackageObjectMapping.bIsEditorOnly)
-			{
-				Actor->SetFlags(RF_Transient);
-				UE_LOG(LogWorldPartition, Log, TEXT("Marked actor %s transient as it was referenced by an editor-only context"), *Actor->GetPathName());
-			}
-
 			// Process objects found in the source actor package
 			if (bIsActoPackageExternal)
 			{
@@ -305,6 +300,7 @@ void FWorldPartitionLevelHelper::MoveExternalActorsToLevel(const TArray<FWorldPa
 							if (InLevel != Object->GetOuter())
 							{
 								Object->Rename(nullptr, InLevel, REN_NonTransactional | REN_DoNotDirty | REN_DontCreateRedirectors);
+								MovedObjects.Add(Object);
 							}
 							else if (NestedActor && !InLevel->Actors.Contains(NestedActor))
 							{
@@ -328,12 +324,7 @@ void FWorldPartitionLevelHelper::MoveExternalActorsToLevel(const TArray<FWorldPa
 							}
 
 							Object->Rename(*NewName.ToString(), LevelPackage, REN_NonTransactional | REN_DoNotDirty | REN_DontCreateRedirectors);
-						}
-
-						if (PackageObjectMapping.bIsEditorOnly)
-						{
-							Object->SetFlags(RF_Transient);
-							UE_LOG(LogWorldPartition, Log, TEXT("Marked actor object %s transient as it was referenced by an editor-only context"), *Object->GetPathName());
+							MovedObjects.Add(Object);
 						}
 					}
 				}
@@ -368,6 +359,90 @@ void FWorldPartitionLevelHelper::MoveExternalActorsToLevel(const TArray<FWorldPa
 	}
 
 	OutModifiedPackages.Append(ModifiedPackages.Array());
+	
+	// Perform validations and fixups to make sure the level is saveable
+	{
+		TArray<UObject*> ObjectStack { InLevel };
+		TSet<UObject*> VisitedObjects { InLevel };
+
+		for (AActor* Actor : InLevel->Actors)
+		{
+			if (IsValid(Actor))
+			{
+				ObjectStack.Add(Actor);
+
+				for (UActorComponent* Component : Actor->GetComponents())
+				{
+					if (IsValid(Component))
+					{
+						ObjectStack.Add(Component);
+					}
+				}
+			}
+		}
+
+		while (ObjectStack.Num())
+		{
+			UObject* Object = ObjectStack.Pop();
+
+			Object->GetClass()->Visit(Object, 
+				[&ObjectStack, &VisitedObjects, &MovedObjects, InLevel](const FPropertyVisitorContext& Context) -> EPropertyVisitorControlFlow
+				{
+					const FPropertyVisitorPath& Path = Context.Path;
+					const FPropertyVisitorData& Data = Context.Data;
+					const FProperty* Property = Path.Top().Property;
+
+					if (const FObjectPropertyBase* ObjectProperty = CastField<FObjectPropertyBase>(Property))
+					{
+						if (UObject* PropertyObject = ObjectProperty->GetObjectPropertyValue(Data.PropertyData))
+						{
+							bool bWasAlreadyInSet;
+							VisitedObjects.Add(PropertyObject, &bWasAlreadyInSet);
+
+							if (bWasAlreadyInSet)
+							{
+								return EPropertyVisitorControlFlow::StepOver;
+							}
+
+							if (!ObjectProperty->HasAnyPropertyFlags(CPF_Transient | CPF_EditorOnly))
+							{
+								if (!PropertyObject->HasAnyFlags(RF_Transient) && !PropertyObject->IsEditorOnly())
+								{
+									if (MovedObjects.Remove(PropertyObject))
+									{
+										return EPropertyVisitorControlFlow::StepInto;
+									}
+								}
+							}
+
+							if (!IsValidChecked(PropertyObject))
+							{
+								UE_LOG(LogWorldPartition, Log, TEXT("Reachable invalid object '%s' from '%s' will be nulled out"), *GetNameSafe(PropertyObject), *Path.ToString());
+								ObjectProperty->SetObjectPropertyValue(Data.PropertyData, nullptr);
+								return EPropertyVisitorControlFlow::StepOver;
+							}
+
+							if (!PropertyObject->IsIn(InLevel))
+							{
+								return EPropertyVisitorControlFlow::StepOver;
+							}
+						}
+					}
+
+					return EPropertyVisitorControlFlow::StepInto;
+				}, FPropertyVisitorContext::EScope::ObjectRefs);
+		}
+
+		if (MovedObjects.Num())
+		{
+			UE_LOG(LogWorldPartition, Log, TEXT("Unreferenced objects from '%s'"), *InLevel->GetPackage()->GetName());
+
+			for (UObject* Object : MovedObjects)
+			{
+				UE_LOG(LogWorldPartition, Log, TEXT("\t'%s'"), *Object->GetName());
+			}
+		}
+	}
 }
 
 void FWorldPartitionLevelHelper::RemapLevelSoftObjectPaths(ULevel* InLevel, UWorldPartition* InWorldPartition)
diff --git a/Engine/Source/Runtime/Engine/Private/WorldPartition/WorldPartitionRuntimeLevelStreamingCell.cpp b/Engine/Source/Runtime/Engine/Private/WorldPartition/WorldPartitionRuntimeLevelStreamingCell.cpp
--- a/Engine/Source/Runtime/Engine/Private/WorldPartition/WorldPartitionRuntimeLevelStreamingCell.cpp
+++ b/Engine/Source/Runtime/Engine/Private/WorldPartition/WorldPartitionRuntimeLevelStreamingCell.cpp
@@ -328,43 +328,46 @@ void UWorldPartitionRuntimeLevelStreamingCell::AddActorToCell(const FStreamingGe
 		ParentContainerInstance = CurrentContainerInstance->GetParentContainerInstance();
 	}
 
-	for (const FGuid& EditorReferenceGuid : ActorDescView.GetEditorReferences())
+	if (!IsRunningCookCommandlet())
 	{
-		FName ReferencePackage;
-		FName ReferencePath;
-		FTopLevelAssetPath ReferenceBaseClass;
-		FTopLevelAssetPath ReferenceNativeClass;
-
-		// Special case where the actor descriptor view has an invalid references, use invalid reference information as the actor guid isn't necessarily in the container instance.
-		if (const FStreamingGenerationActorDescView::FInvalidReference* InvalidReference = ActorDescView.GetInvalidReference(EditorReferenceGuid))
-		{
-			ReferencePackage = InvalidReference->ActorPackage;
-			ReferencePath = *InvalidReference->ActorSoftPath.ToString();
-			ReferenceBaseClass = InvalidReference->BaseClass;
-			ReferenceNativeClass = InvalidReference->NativeClass;
-		}
-		else
+		for (const FGuid& EditorReferenceGuid : ActorDescView.GetEditorReferences())
 		{
-			const FWorldPartitionActorDescInstance& ReferenceActorDesc = ContainerInstance->GetActorDescInstanceChecked(EditorReferenceGuid);
-			ReferencePackage = ReferenceActorDesc.GetActorPackage();
-			ReferencePath = *ReferenceActorDesc.GetActorSoftPath().ToString();
-			ReferenceBaseClass = ReferenceActorDesc.GetBaseClass();
-			ReferenceNativeClass = ReferenceActorDesc.GetNativeClass();
+			FName ReferencePackage;
+			FName ReferencePath;
+			FTopLevelAssetPath ReferenceBaseClass;
+			FTopLevelAssetPath ReferenceNativeClass;
+
+			// Special case where the actor descriptor view has an invalid references, use invalid reference information as the actor guid isn't necessarily in the container instance.
+			if (const FStreamingGenerationActorDescView::FInvalidReference* InvalidReference = ActorDescView.GetInvalidReference(EditorReferenceGuid))
+			{
+				ReferencePackage = InvalidReference->ActorPackage;
+				ReferencePath = *InvalidReference->ActorSoftPath.ToString();
+				ReferenceBaseClass = InvalidReference->BaseClass;
+				ReferenceNativeClass = InvalidReference->NativeClass;
+			}
+			else
+			{
+				const FWorldPartitionActorDescInstance& ReferenceActorDesc = ContainerInstance->GetActorDescInstanceChecked(EditorReferenceGuid);
+				ReferencePackage = ReferenceActorDesc.GetActorPackage();
+				ReferencePath = *ReferenceActorDesc.GetActorSoftPath().ToString();
+				ReferenceBaseClass = ReferenceActorDesc.GetBaseClass();
+				ReferenceNativeClass = ReferenceActorDesc.GetNativeClass();
+			}
+
+			Packages.Emplace(
+				ReferencePackage,
+				ReferencePath,
+				ReferenceBaseClass,
+				ReferenceNativeClass,
+				ContainerID,
+				ContainerTransform,
+				FTransform::Identity,
+				ContainerPackage,
+				GetWorld()->GetPackage()->GetFName(),
+				ContainerID.GetActorGuid(EditorReferenceGuid),
+				true
+			);
 		}
-		
-		Packages.Emplace(
-			ReferencePackage,
-			ReferencePath,
-			ReferenceBaseClass,
-			ReferenceNativeClass,
-			ContainerID,
-			ContainerTransform,
-			FTransform::Identity,
-			ContainerPackage,
-			GetWorld()->GetPackage()->GetFName(),
-			ContainerID.GetActorGuid(EditorReferenceGuid),
-			true
-		);
 	}
 
 	FWorldPartitionRuntimeCellObjectMapping ActorMapping(
