Added UActorComponent::GetPhysicsScene() to return physics scene (required for future development).

#Synced-CL 48932408

diff --git a/Engine/Source/Runtime/Engine/Classes/Components/ActorComponent.h b/Engine/Source/Runtime/Engine/Classes/Components/ActorComponent.h
--- a/Engine/Source/Runtime/Engine/Classes/Components/ActorComponent.h
+++ b/Engine/Source/Runtime/Engine/Classes/Components/ActorComponent.h
@@ -1180,6 +1180,9 @@ public:
 	/** Returns the rendering scene associated with this component */
 	ENGINE_API class FSceneInterface* GetScene() const;
 
+	/** Returns the physics scene associated with this component */
+	ENGINE_API FPhysScene* GetPhysicsScene() const;
+
 	/** Return the ULevel that this Component is part of. */
 	ENGINE_API ULevel* GetComponentLevel() const;
 
diff --git a/Engine/Source/Runtime/Engine/Private/Components/ActorComponent.cpp b/Engine/Source/Runtime/Engine/Private/Components/ActorComponent.cpp
--- a/Engine/Source/Runtime/Engine/Private/Components/ActorComponent.cpp
+++ b/Engine/Source/Runtime/Engine/Private/Components/ActorComponent.cpp
@@ -1108,6 +1108,11 @@ FSceneInterface* UActorComponent::GetScene() const
 	return (WorldPrivate ? WorldPrivate->Scene : NULL);
 }
 
+FPhysScene* UActorComponent::GetPhysicsScene() const
+{
+	return WorldPrivate ? WorldPrivate->GetPhysicsScene() : nullptr;
+}
+
 ULevel* UActorComponent::GetComponentLevel() const
 {
 	// For model components Level is outer object
@@ -1485,7 +1490,7 @@ bool UActorComponent::ShouldAsyncCreatePhysicsState(UWorld* WorldContext) const
 
 bool UActorComponent::ShouldAsyncDestroyPhysicsState() const
 {
-	return bPhysicsStateCreated && WorldPrivate && WorldPrivate->GetPhysicsScene() && AllowsAsyncPhysicsStateDestruction();
+	return bPhysicsStateCreated && WorldPrivate && GetPhysicsScene() && AllowsAsyncPhysicsStateDestruction();
 }
 
 bool UActorComponent::ShouldIncrementalPreRegister(UWorld* WorldContext) const
@@ -1519,7 +1524,7 @@ void UActorComponent::OnPreRegister()
 
 	if (ShouldAsyncCreatePhysicsState(WorldPrivate))
 	{
-		FPhysScene_Chaos* PhysScene = WorldPrivate->GetPhysicsScene();
+		FPhysScene_Chaos* PhysScene = GetPhysicsScene();
 		check(PhysScene);
 		check(!IsAsyncCreatePhysicsStateRunning());
 		check(!IsAsyncDestroyPhysicsStateRunning());
@@ -1804,7 +1809,7 @@ void UActorComponent::RegisterAllComponentTickFunctions(bool bRegister)
 
 void UActorComponent::RegisterAsyncPhysicsTickEnabled(bool bRegister)
 {
-	if (FPhysScene_Chaos* Scene = static_cast<FPhysScene_Chaos*>(WorldPrivate->GetPhysicsScene()))
+	if (FPhysScene_Chaos* Scene = static_cast<FPhysScene_Chaos*>(GetPhysicsScene()))
 	{
 		if (bRegister)
 		{
@@ -1830,7 +1835,7 @@ void UActorComponent::SetAsyncPhysicsTickEnabled(bool bEnable)
 
 void UActorComponent::DeferRemoveAsyncPhysicsTick()
 {
-	if (FPhysScene_Chaos* Scene = static_cast<FPhysScene_Chaos*>(WorldPrivate->GetPhysicsScene()))
+	if (FPhysScene_Chaos* Scene = static_cast<FPhysScene_Chaos*>(GetPhysicsScene()))
 	{
 		// Set 0 for the step so that it gets run immediately on the next async tick.
 		Scene->EnqueueAsyncPhysicsCommand(0, this,
@@ -2074,7 +2079,7 @@ void UActorComponent::OnPreUnregister()
 	}
 
 	check(ShouldAsyncDestroyPhysicsState());
-	FPhysScene_Chaos* PhysScene = WorldPrivate->GetPhysicsScene();
+	FPhysScene_Chaos* PhysScene = GetPhysicsScene();
 	check(PhysScene);
 	check(!IsAsyncCreatePhysicsStateRunning());
 	check(!IsAsyncDestroyPhysicsStateRunning());
@@ -2273,7 +2278,7 @@ void UActorComponent::OnCreatePhysicsState()
 	TRACE_CPUPROFILER_EVENT_SCOPE(UActorComponent::OnCreatePhysicsState);
 	check(IsRegistered() || IsPreRegistering());
 	check(ShouldCreatePhysicsState());
-	check(WorldPrivate->GetPhysicsScene());
+	check(GetPhysicsScene());
 	check(!bPhysicsStateCreated);
 	bPhysicsStateCreated = true;
 }
@@ -2359,7 +2364,7 @@ void UActorComponent::CreatePhysicsState(bool bAllowDeferral)
 	SCOPE_CYCLE_COUNTER(STAT_ComponentCreatePhysicsState);
 	TRACE_CPUPROFILER_EVENT_SCOPE(UActorComponent::CreatePhysicsState);
 
-	if (!bPhysicsStateCreated && WorldPrivate->GetPhysicsScene() && ShouldCreatePhysicsState())
+	if (!bPhysicsStateCreated && GetPhysicsScene() && ShouldCreatePhysicsState())
 	{
 		UPrimitiveComponent* Primitive = Cast<UPrimitiveComponent>(this);
 
@@ -2384,7 +2389,7 @@ void UActorComponent::CreatePhysicsState(bool bAllowDeferral)
 
 		if (ShouldDefer)
 		{
-			WorldPrivate->GetPhysicsScene()->DeferPhysicsStateCreation(Primitive);
+			GetPhysicsScene()->DeferPhysicsStateCreation(Primitive);
 		}
 		else
 		{
@@ -2402,7 +2407,7 @@ void UActorComponent::CreatePhysicsState(bool bAllowDeferral)
 	{
 		check(bPhysicsStateCreated);
 		bCallAsyncPhysicsStateCreatedDelegate = false;
-		if (ensure(WorldPrivate->GetPhysicsScene()))
+		if (ensure(GetPhysicsScene()))
 		{
 			GlobalCreatePhysicsDelegate.Broadcast(this);
 		}
