Added FScopedDeferAsyncPhysicsStateJobs that allows to temporarily prevent new added jobs from being processed by FPhysScene_AsyncPhysicsStateJobQueue async task (required for future development).

#Synced-CL 48932408

diff --git a/Engine/Source/Runtime/Engine/Private/PhysicsEngine/Experimental/PhysScene_AsyncPhysicsStateJobQueue.cpp b/Engine/Source/Runtime/Engine/Private/PhysicsEngine/Experimental/PhysScene_AsyncPhysicsStateJobQueue.cpp
--- a/Engine/Source/Runtime/Engine/Private/PhysicsEngine/Experimental/PhysScene_AsyncPhysicsStateJobQueue.cpp
+++ b/Engine/Source/Runtime/Engine/Private/PhysicsEngine/Experimental/PhysScene_AsyncPhysicsStateJobQueue.cpp
@@ -25,8 +25,27 @@ namespace Chaos
 	);
 }
 
-FPhysScene_AsyncPhysicsStateJobQueue::FPhysScene_AsyncPhysicsStateJobQueue(FPhysScene* InPhysicScene)
-	: PhysScene(InPhysicScene)
+FPhysScene_AsyncPhysicsStateJobQueue::FScopedDeferAsyncPhysicsStateJobs::FScopedDeferAsyncPhysicsStateJobs(FPhysScene* InPhysScene)
+: PhysScene(InPhysScene)
+{
+	if (PhysScene)
+	{
+		check(IsInGameThread());
+		PhysScene->GetOrCreateAsyncPhysicsStateJobQueue().BeginDeferringNewJobs();
+	}
+}
+
+FPhysScene_AsyncPhysicsStateJobQueue::FScopedDeferAsyncPhysicsStateJobs::~FScopedDeferAsyncPhysicsStateJobs()
+{
+	if (PhysScene)
+	{
+		check(IsInGameThread());
+		PhysScene->GetOrCreateAsyncPhysicsStateJobQueue().EndDeferringNewJobs();
+	}
+}
+
+FPhysScene_AsyncPhysicsStateJobQueue::FPhysScene_AsyncPhysicsStateJobQueue(FPhysScene* InPhysScene)
+	: PhysScene(InPhysScene)
 {
 	PhysScene->GetOwningWorld()->OnAllLevelsChanged().AddRaw(this, &FPhysScene_AsyncPhysicsStateJobQueue::OnUpdateLevelStreamingDone);
 }
@@ -35,6 +54,7 @@ FPhysScene_AsyncPhysicsStateJobQueue::~FPhysScene_AsyncPhysicsStateJobQueue()
 {
 	PhysScene->GetOwningWorld()->OnAllLevelsChanged().RemoveAll(this);
 
+	check(DeferNewJobsScopeCount == 0);
 	if (!IsCompleted())
 	{
 		Tick(true);
@@ -48,8 +68,16 @@ void FPhysScene_AsyncPhysicsStateJobQueue::AddJob(const FJob& Job)
 	check(IsInGameThread());
 
 	Job.OnPreExecute_GameThread();
-	FWriteScopeLock Lock(JobsLock);
-	JobsToExecute.Add(Job);
+
+	if (DeferNewJobsScopeCount > 0)
+	{
+		DeferredJobs.Add(Job);
+	}
+	else
+	{
+		FWriteScopeLock Lock(JobsLock);
+		JobsToExecute.Add(Job);
+	}
 }
 
 void FPhysScene_AsyncPhysicsStateJobQueue::RemoveJob(const FJob& Job)
@@ -79,6 +107,10 @@ void FPhysScene_AsyncPhysicsStateJobQueue::RemoveJob(const FJob& Job)
 		{
 			bFoundJob = true;
 		}
+		else if (DeferredJobs.Remove(Job))
+		{
+			bFoundJob = true;
+		}
 	}
 
 	if (bFoundJob)
@@ -106,10 +138,32 @@ void FPhysScene_AsyncPhysicsStateJobQueue::RemoveJob(const FJob& Job)
 	}
 }
 
+void FPhysScene_AsyncPhysicsStateJobQueue::BeginDeferringNewJobs()
+{
+	check(IsInGameThread());
+	++DeferNewJobsScopeCount;
+}
+
+void FPhysScene_AsyncPhysicsStateJobQueue::EndDeferringNewJobs()
+{
+	check(IsInGameThread());
+	check(DeferNewJobsScopeCount > 0);
+	--DeferNewJobsScopeCount;
+
+	if (DeferNewJobsScopeCount == 0 && DeferredJobs.Num() > 0)
+	{
+		FWriteScopeLock Lock(JobsLock);
+		JobsToExecute.Append(MoveTemp(DeferredJobs));
+	}
+}
+
 bool FPhysScene_AsyncPhysicsStateJobQueue::IsCompleted() const
 {
 	FReadScopeLock Lock(JobsLock);
-	return AsyncJobTask.IsCompleted() && JobsToExecute.IsEmpty() && !ExecutingJob.IsSet() && CompletedJobs.IsEmpty();
+	const bool bNoAsyncJob = AsyncJobTask.IsCompleted() && JobsToExecute.IsEmpty() && !ExecutingJob.IsSet() && CompletedJobs.IsEmpty();
+	const bool bNoDeferredJob = (DeferNewJobsScopeCount == 0) && DeferredJobs.IsEmpty();
+
+	return bNoAsyncJob && bNoDeferredJob;
 }
 
 void FPhysScene_AsyncPhysicsStateJobQueue::OnUpdateLevelStreamingDone()
diff --git a/Engine/Source/Runtime/Engine/Private/PhysicsEngine/Experimental/PhysScene_Chaos.cpp b/Engine/Source/Runtime/Engine/Private/PhysicsEngine/Experimental/PhysScene_Chaos.cpp
--- a/Engine/Source/Runtime/Engine/Private/PhysicsEngine/Experimental/PhysScene_Chaos.cpp
+++ b/Engine/Source/Runtime/Engine/Private/PhysicsEngine/Experimental/PhysScene_Chaos.cpp
@@ -1854,6 +1854,17 @@ bool FPhysScene_Chaos::SupportsAsyncPhysicsStateDestruction()
 	return SupportsAsyncPhysicsStateCreation();
 }
 
+FPhysScene_AsyncPhysicsStateJobQueue& FPhysScene_Chaos::GetOrCreateAsyncPhysicsStateJobQueue()
+{
+	check(IsInGameThread());
+
+	if (!AsyncPhysicsStateJobQueue.IsValid())
+	{
+		AsyncPhysicsStateJobQueue = MakeUnique<FPhysScene_AsyncPhysicsStateJobQueue>(this);
+	}
+	return *AsyncPhysicsStateJobQueue.Get();
+}
+
 bool FPhysScene_Chaos::PushAsyncCreatePhysicsStateJob(IAsyncPhysicsStateProcessor* Processor)
 {
 	if (SupportsAsyncPhysicsStateCreation())
@@ -1862,11 +1873,7 @@ bool FPhysScene_Chaos::PushAsyncCreatePhysicsStateJob(IAsyncPhysicsStateProcesso
 		check(!Processor->IsAsyncPhysicsStateCreated());
 		check(Processor->AllowsAsyncPhysicsStateCreation());
 
-		if (!AsyncPhysicsStateJobQueue.IsValid())
-		{
-			AsyncPhysicsStateJobQueue = MakeUnique<FPhysScene_AsyncPhysicsStateJobQueue>(this);
-		}
-		AsyncPhysicsStateJobQueue->AddJob({ Processor, FPhysScene_AsyncPhysicsStateJobQueue::EJobType::CreatePhysicsState });
+		GetOrCreateAsyncPhysicsStateJobQueue().AddJob({ Processor, FPhysScene_AsyncPhysicsStateJobQueue::EJobType::CreatePhysicsState });
 		return true;
 	}
 
@@ -1881,13 +1888,10 @@ bool FPhysScene_Chaos::PushAsyncDestroyPhysicsStateJob(IAsyncPhysicsStateProcess
 		check(Processor->IsAsyncPhysicsStateCreated());
 		check(Processor->AllowsAsyncPhysicsStateDestruction());
 
-		if (!AsyncPhysicsStateJobQueue.IsValid())
-		{
-			AsyncPhysicsStateJobQueue = MakeUnique<FPhysScene_AsyncPhysicsStateJobQueue>(this);
-		}
-		AsyncPhysicsStateJobQueue->AddJob({ Processor, FPhysScene_AsyncPhysicsStateJobQueue::EJobType::DestroyPhysicsState });
+		GetOrCreateAsyncPhysicsStateJobQueue().AddJob({ Processor, FPhysScene_AsyncPhysicsStateJobQueue::EJobType::DestroyPhysicsState });
+		return true;
 	}
-	return true;
+	return false;
 }
 
 void FPhysScene_Chaos::RemoveAsyncPhysicsStateJobs(IAsyncPhysicsStateProcessor* Processor)
diff --git a/Engine/Source/Runtime/Engine/Public/Physics/Experimental/PhysScene_AsyncPhysicsStateJobQueue.h b/Engine/Source/Runtime/Engine/Public/Physics/Experimental/PhysScene_AsyncPhysicsStateJobQueue.h
--- a/Engine/Source/Runtime/Engine/Public/Physics/Experimental/PhysScene_AsyncPhysicsStateJobQueue.h
+++ b/Engine/Source/Runtime/Engine/Public/Physics/Experimental/PhysScene_AsyncPhysicsStateJobQueue.h
@@ -11,7 +11,7 @@
 
 struct FPhysScene_AsyncPhysicsStateJobQueue
 {
-	FPhysScene_AsyncPhysicsStateJobQueue(FPhysScene* InPhysicScene);
+	FPhysScene_AsyncPhysicsStateJobQueue(FPhysScene* InPhysScene);
 	~FPhysScene_AsyncPhysicsStateJobQueue();
 
 	enum class EJobType
@@ -50,6 +50,18 @@ struct FPhysScene_AsyncPhysicsStateJobQueue
 		EJobType Type;
 	};
 
+	struct FScopedDeferAsyncPhysicsStateJobs
+	{
+		explicit FScopedDeferAsyncPhysicsStateJobs(FPhysScene* InPhysScene);
+		~FScopedDeferAsyncPhysicsStateJobs();
+		
+		FScopedDeferAsyncPhysicsStateJobs(const FScopedDeferAsyncPhysicsStateJobs&) = delete;
+		FScopedDeferAsyncPhysicsStateJobs& operator=(const FScopedDeferAsyncPhysicsStateJobs&) = delete;
+
+	private:
+		FPhysScene* PhysScene = nullptr;
+	};
+
 	void Tick(bool bWaitForCompletion = false);
 	void AddJob(const FJob& Job);
 	void RemoveJob(const FJob& Job);
@@ -57,6 +69,9 @@ struct FPhysScene_AsyncPhysicsStateJobQueue
 
 private:
 
+	void BeginDeferringNewJobs();
+	void EndDeferringNewJobs();
+
 	void LaunchAsyncJobTask();
 	void ExecuteJobsAsync(double TimeBudgetSeconds);
 	void OnUpdateLevelStreamingDone();
@@ -75,6 +90,13 @@ private:
 	TOptional<FJob> ExecutingJob;
 	TArray<FJob> CompletedJobs;
 
+	// When greater than zero, new jobs are added to DeferredJobs until counter goes back to 0 
+	// This allows to temporarily pause adding to JobsToExecute.
+	int32 DeferNewJobsScopeCount = 0;
+	TArray<FJob> DeferredJobs;
+
 	std::atomic<bool> bIsBlocking = false; // Notifies to the async task that we are block waiting for the task to complete
 	std::atomic<int32> GameThreadEpoch = 0; // Epoch updated by the GameThread at every frame and used by the async task to reset UsedAsyncTaskTimeBudgetSec
+
+	friend struct FScopedDeferAsyncPhysicsStateJobs;
 };
\ No newline at end of file
diff --git a/Engine/Source/Runtime/Engine/Public/Physics/Experimental/PhysScene_Chaos.h b/Engine/Source/Runtime/Engine/Public/Physics/Experimental/PhysScene_Chaos.h
--- a/Engine/Source/Runtime/Engine/Public/Physics/Experimental/PhysScene_Chaos.h
+++ b/Engine/Source/Runtime/Engine/Public/Physics/Experimental/PhysScene_Chaos.h
@@ -277,6 +277,7 @@ public:
 
 	ENGINE_API static bool SupportsAsyncPhysicsStateCreation();
 	ENGINE_API static bool SupportsAsyncPhysicsStateDestruction();
+	ENGINE_API FPhysScene_AsyncPhysicsStateJobQueue& GetOrCreateAsyncPhysicsStateJobQueue();
 	ENGINE_API bool PushAsyncCreatePhysicsStateJob(IAsyncPhysicsStateProcessor* Processor);
 	ENGINE_API bool PushAsyncDestroyPhysicsStateJob(IAsyncPhysicsStateProcessor* Processor);
 	ENGINE_API void RemoveAsyncPhysicsStateJobs(IAsyncPhysicsStateProcessor* Processor);
