- Make LastRenderTime / LastRenderTimeOnScreen relaxed atomic to fix TSan issue
- Add Get/Set accessors and deprecate using the raw member variables

#Synced-CL 48918725

diff --git a/Engine/Plugins/Experimental/FastGeoStreaming/Source/FastGeoStreaming/Private/FastGeoPrimitiveComponent.cpp b/Engine/Plugins/Experimental/FastGeoStreaming/Source/FastGeoStreaming/Private/FastGeoPrimitiveComponent.cpp
--- a/Engine/Plugins/Experimental/FastGeoStreaming/Source/FastGeoStreaming/Private/FastGeoPrimitiveComponent.cpp
+++ b/Engine/Plugins/Experimental/FastGeoStreaming/Source/FastGeoStreaming/Private/FastGeoPrimitiveComponent.cpp
@@ -578,7 +578,7 @@ FMatrix FFastGeoPrimitiveComponent::GetRenderMatrix() const
 
 float FFastGeoPrimitiveComponent::GetLastRenderTimeOnScreen() const
 {
-	return PrimitiveSceneData.LastRenderTimeOnScreen;
+	return PrimitiveSceneData.GetLastRenderTimeOnScreen();
 }
 
 void FFastGeoPrimitiveComponent::SetCollisionEnabled(bool bInCollisionEnabled)
diff --git a/Engine/Source/Runtime/Engine/Private/Components/PrimitiveComponent.cpp b/Engine/Source/Runtime/Engine/Private/Components/PrimitiveComponent.cpp
--- a/Engine/Source/Runtime/Engine/Private/Components/PrimitiveComponent.cpp
+++ b/Engine/Source/Runtime/Engine/Private/Components/PrimitiveComponent.cpp
@@ -691,7 +691,7 @@ void UPrimitiveComponent::OnRegister()
 #endif
 
 	// Update our Owner's LastRenderTime
-	SetLastRenderTime(SceneData.LastRenderTime);
+	SetLastRenderTime(SceneData.GetLastRenderTime());
 }
 
 
@@ -5151,7 +5151,7 @@ bool UPrimitiveComponent::WasRecentlyRendered(float Tolerance /*= 0.2*/) const
 
 void UPrimitiveComponent::SetLastRenderTime(float InLastRenderTime)
 {
-	SceneData.LastRenderTime = InLastRenderTime;
+	SceneData.SetLastRenderTime(InLastRenderTime);
 	if (AActor* Owner = GetOwner())
 	{
 		if (InLastRenderTime > Owner->GetLastRenderTime())
@@ -5167,7 +5167,7 @@ float UPrimitiveComponent::GetLastRenderTime() const
 	{
 		return GetWorld()->GetTimeSeconds();
 	}
-	return SceneData.LastRenderTime;
+	return SceneData.GetLastRenderTime();
 }
 
 float UPrimitiveComponent::GetLastRenderTimeOnScreen() const
@@ -5176,7 +5176,7 @@ float UPrimitiveComponent::GetLastRenderTimeOnScreen() const
 	{
 		return GetWorld()->GetTimeSeconds();
 	}
-	return SceneData.LastRenderTimeOnScreen;
+	return SceneData.GetLastRenderTimeOnScreen();
 }
 
 #if MESH_DRAW_COMMAND_STATS
diff --git a/Engine/Source/Runtime/Engine/Private/PrimitiveSceneInfoData.cpp b/Engine/Source/Runtime/Engine/Private/PrimitiveSceneInfoData.cpp
--- a/Engine/Source/Runtime/Engine/Private/PrimitiveSceneInfoData.cpp
+++ b/Engine/Source/Runtime/Engine/Private/PrimitiveSceneInfoData.cpp
@@ -9,11 +9,11 @@ FThreadSafeCounter FPrimitiveSceneInfoData::NextPrimitiveId;
 
 void FPrimitiveSceneInfoData::SetLastRenderTime(float InLastRenderTime, bool bUpdateLastRenderTimeOnScreen) const
 {
-	LastRenderTime = InLastRenderTime;
+	SetLastRenderTime(InLastRenderTime);
 
 	if (bUpdateLastRenderTimeOnScreen)
 	{
-		LastRenderTimeOnScreen = InLastRenderTime;
+		SetLastRenderTimeOnScreen(InLastRenderTime);
 	}
 
 	if (OwnerLastRenderTimePtr)
diff --git a/Engine/Source/Runtime/Engine/Public/PrimitiveSceneInfoData.h b/Engine/Source/Runtime/Engine/Public/PrimitiveSceneInfoData.h
--- a/Engine/Source/Runtime/Engine/Public/PrimitiveSceneInfoData.h
+++ b/Engine/Source/Runtime/Engine/Public/PrimitiveSceneInfoData.h
@@ -22,10 +22,12 @@ struct FPrimitiveSceneInfoData
 	 * from the render thread, which is up to a frame behind the game thread, so you should allow this time to
 	 * be at least a frame behind the game thread's world time before you consider the actor non-visible.
 	 */
-	mutable float LastRenderTime = -1000.0f;
+	UE_DEPRECATED(5.8, "Use GetLastRenderTime/SetLastRenderTime instead.")
+	mutable std::atomic<float> LastRenderTime = -1000.0f;
 
 	/** Same as LastRenderTime but only updated if the component is on screen. Used by the texture streamer. */
-	mutable float LastRenderTimeOnScreen = -1000.0f;
+	UE_DEPRECATED(5.8, "Use GetLastRenderTimeOnScreen/SetLastRenderTimeOnScreen instead.")
+	mutable std::atomic<float> LastRenderTimeOnScreen = -1000.0f;
 
 	/**
 	* Incremented by the main thread before being attached to the scene, decremented
@@ -72,6 +74,34 @@ struct FPrimitiveSceneInfoData
 
 	ENGINE_API void SetLastRenderTime(float InLastRenderTime, bool bUpdateLastRenderTimeOnScreen) const;
 
+	float GetLastRenderTime() const
+	{
+		PRAGMA_DISABLE_DEPRECATION_WARNINGS
+		return LastRenderTime.load(std::memory_order_relaxed);
+		PRAGMA_ENABLE_DEPRECATION_WARNINGS
+	}
+
+	void SetLastRenderTime(float InTime) const
+	{
+		PRAGMA_DISABLE_DEPRECATION_WARNINGS
+		LastRenderTime.store(InTime, std::memory_order_relaxed);
+		PRAGMA_ENABLE_DEPRECATION_WARNINGS
+	}
+
+	float GetLastRenderTimeOnScreen() const
+	{
+		PRAGMA_DISABLE_DEPRECATION_WARNINGS
+		return LastRenderTimeOnScreen.load(std::memory_order_relaxed);
+		PRAGMA_ENABLE_DEPRECATION_WARNINGS
+	}
+
+	void SetLastRenderTimeOnScreen(float InTime) const
+	{
+		PRAGMA_DISABLE_DEPRECATION_WARNINGS
+		LastRenderTimeOnScreen.store(InTime, std::memory_order_relaxed);
+		PRAGMA_ENABLE_DEPRECATION_WARNINGS
+	}
+
 protected:
 
 	/** Next id to be used by a component. */
diff --git a/Engine/Source/Runtime/Renderer/Private/RendererScene.cpp b/Engine/Source/Runtime/Renderer/Private/RendererScene.cpp
--- a/Engine/Source/Runtime/Renderer/Private/RendererScene.cpp
+++ b/Engine/Source/Runtime/Renderer/Private/RendererScene.cpp
@@ -1681,7 +1681,7 @@ void FScene::UpdatePrimitiveTransformInternal(T* Primitive)
 										FSimpleStreamableAssetManager::GetStreamingScaleFactor(SceneProxy, UpdateParams.LocalToWorld),
 										SceneProxy->GetMinDrawDistance(),
 										SceneProxy->GetMaxDrawDistance(),
-										SceneProxy->PrimitiveSceneInfo->GetSceneData()->LastRenderTimeOnScreen,
+										SceneProxy->PrimitiveSceneInfo->GetSceneData()->GetLastRenderTimeOnScreen(),
 										SceneProxy->IsForceMipStreaming()
 									}
 								);
@@ -1742,7 +1742,7 @@ void FScene::UpdatePrimitiveTransforms(IPrimitiveTransformUpdater* InUpdater)
 							FSimpleStreamableAssetManager::GetStreamingScaleFactor(SceneProxy, UpdateParams.LocalToWorld),
 							SceneProxy->GetMinDrawDistance(),
 							SceneProxy->GetMaxDrawDistance(),
-							SceneProxy->PrimitiveSceneInfo->GetSceneData()->LastRenderTimeOnScreen,
+							SceneProxy->PrimitiveSceneInfo->GetSceneData()->GetLastRenderTimeOnScreen(),
 							SceneProxy->IsForceMipStreaming()
 						}
 					);
