Fix memory corruption in cooked builds caused by unexpected texture mip size by storing and validating the actual size when streaming in, which previously only happened on android.

#Synced-CL 48821263

diff --git a/Engine/Source/Runtime/Engine/Private/Streaming/Texture2DStreamIn.cpp b/Engine/Source/Runtime/Engine/Private/Streaming/Texture2DStreamIn.cpp
--- a/Engine/Source/Runtime/Engine/Private/Streaming/Texture2DStreamIn.cpp
+++ b/Engine/Source/Runtime/Engine/Private/Streaming/Texture2DStreamIn.cpp
@@ -45,8 +45,7 @@ void FTexture2DStreamIn::DoAllocateNewMips(const FContext& Context)
 
 			check(MipData[MipIndex].Data == nullptr);
 			MipData[MipIndex].Data = FMemory::Malloc(MipSize);
-			// would be nice to store Size here !
-			//MipData[MipIndex].Size = MipSize;
+			MipData[MipIndex].Size = MipSize;
 			MipData[MipIndex].Pitch = 0; // 0 means tight packed
 		}
 	}
@@ -60,8 +59,8 @@ void FTexture2DStreamIn::DoFreeNewMips(const FContext& Context)
 		{
 			FMemory::Free(MipData[MipIndex].Data);
 			MipData[MipIndex].Data = nullptr;
+			MipData[MipIndex].Size = 0;
 			MipData[MipIndex].Pitch = -1;
-			//MipData[MipIndex].Size
 		}
 	}
 }
@@ -79,10 +78,10 @@ void FTexture2DStreamIn::DoLockNewMips(FRHICommandListImmediate& RHICmdList, con
 		{
 			check(MipData[MipIndex].Data == nullptr);
 			uint32 DestPitch = -1;
-			MipData[MipIndex].Data = RHICmdList.LockTexture2D(IntermediateTextureRHI, MipIndex - MipOffset, RLM_WriteOnly, DestPitch, false, CVarFlushRHIThreadOnSTreamingTextureLocks.GetValueOnAnyThread() > 0);
+			uint64 Size = ~0ULL;
+			MipData[MipIndex].Data = RHICmdList.LockTexture2D(IntermediateTextureRHI, MipIndex - MipOffset, RLM_WriteOnly, DestPitch, false, CVarFlushRHIThreadOnSTreamingTextureLocks.GetValueOnAnyThread() > 0, &Size);
+			MipData[MipIndex].Size = Size;
 			MipData[MipIndex].Pitch = DestPitch;
-			// note: should store Size but LockTexture2D doesn't tell us size
-			//MipData[MipIndex].Size
 
 			UE_LOG(LogTextureUpload,Verbose,TEXT("FTexture2DStreamIn::DoLockNewMips( : Lock Mip %d Pitch=%d"),MipIndex,DestPitch);
 		}
@@ -105,6 +104,7 @@ void FTexture2DStreamIn::DoUnlockNewMips(FRHICommandListImmediate& RHICmdList, c
 			{
 				RHICmdList.UnlockTexture2D(IntermediateTextureRHI, MipIndex - MipOffset, false, CVarFlushRHIThreadOnSTreamingTextureLocks.GetValueOnAnyThread() > 0 );
 				MipData[MipIndex].Data = nullptr;
+				MipData[MipIndex].Size = 0;
 				MipData[MipIndex].Pitch = -1;
 			}
 		}
@@ -180,3 +180,16 @@ void FTexture2DStreamIn::DoAsyncCreateWithNewMips(const FContext& Context)
 		}
 	}
 }
+
+void FTexture2DStreamIn::ValidateMipBulkDataSize(const UTexture2D& Texture, int32 MipSizeX, int32 MipSizeY, int32 MipIndex, int64 ActualMipSize, int64& BulkDataSize)
+{
+	if (BulkDataSize != ActualMipSize)
+	{
+#if !UE_BUILD_SHIPPING
+		UE_LOG(LogTexture, Warning, TEXT("Mip (%d) %dx%d has an unexpected size %lld, expected size %lld. %s, Pixel format %s"),
+			MipIndex, MipSizeX, MipSizeY, BulkDataSize, ActualMipSize, *Texture.GetFullName(), GPixelFormats[Texture.GetPixelFormat()].Name);
+#endif
+		// Make sure we don't overrun buffer allocated for this mip
+		BulkDataSize = FMath::Min(BulkDataSize, ActualMipSize);
+	}
+}
\ No newline at end of file
diff --git a/Engine/Source/Runtime/Engine/Private/Streaming/Texture2DStreamIn.h b/Engine/Source/Runtime/Engine/Private/Streaming/Texture2DStreamIn.h
--- a/Engine/Source/Runtime/Engine/Private/Streaming/Texture2DStreamIn.h
+++ b/Engine/Source/Runtime/Engine/Private/Streaming/Texture2DStreamIn.h
@@ -23,8 +23,8 @@ protected:
 	struct FStreamMipData
 	{
 		void * Data = nullptr;
+		uint64 Size = 0;
 		uint32 Pitch = 0;
-		// note: Size would be nice but LockTexture doesn't return size
 	};
 	TArray<FStreamMipData, TInlineAllocator<MAX_TEXTURE_MIP_COUNT> > MipData;
 	void * InitialMipDataForAsyncCreate[MAX_TEXTURE_MIP_COUNT];
@@ -48,4 +48,7 @@ protected:
 
 	// Async create the texture to the requested size.
 	void DoAsyncCreateWithNewMips(const FContext& Context);
+
+	// Make sure the destination mip budffer is large enough to store the bulk data, warn and adjust if not
+	void ValidateMipBulkDataSize(const UTexture2D& Texture, int32 MipSizeX, int32 MipSizeY, int32 MipIndex, int64 ActualMipSize, int64& BulkDataSize);
 };
diff --git a/Engine/Source/Runtime/Engine/Private/Streaming/Texture2DStreamIn_DDC.cpp b/Engine/Source/Runtime/Engine/Private/Streaming/Texture2DStreamIn_DDC.cpp
--- a/Engine/Source/Runtime/Engine/Private/Streaming/Texture2DStreamIn_DDC.cpp
+++ b/Engine/Source/Runtime/Engine/Private/Streaming/Texture2DStreamIn_DDC.cpp
@@ -223,7 +223,7 @@ void FTexture2DStreamIn_DDC::DoLoadNewMipsFromDDC(const FContext& Context)
 			{
 				// Pitch ignored
 				// to copy and respect Pitch, use CopyTextureData2D
-				// MipData[] should have size but doesn't
+				check(MipData[MipIndex].Size >= Status.Buffer.GetSize());
 				FMemory::Memcpy(MipData[MipIndex].Data, Status.Buffer.GetData(), Status.Buffer.GetSize());
 			}
 			else
diff --git a/Engine/Source/Runtime/Engine/Private/Streaming/Texture2DStreamIn_DerivedData.cpp b/Engine/Source/Runtime/Engine/Private/Streaming/Texture2DStreamIn_DerivedData.cpp
--- a/Engine/Source/Runtime/Engine/Private/Streaming/Texture2DStreamIn_DerivedData.cpp
+++ b/Engine/Source/Runtime/Engine/Private/Streaming/Texture2DStreamIn_DerivedData.cpp
@@ -33,9 +33,11 @@ void FTexture2DStreamIn_DerivedData::DoBeginIoRequests(const FContext& Context)
 		if (MipMap.DerivedData.HasData())
 		{
 			uint32 ExpectedPitch = 0;
-			const uint32 MipSize = FTexture2DResource::CalculateTightPackedMipSize(MipMap.SizeX, MipMap.SizeY, Context.Resource->GetPixelFormat(), ExpectedPitch);
+			int64 MipSize = FTexture2DResource::CalculateTightPackedMipSize(MipMap.SizeX, MipMap.SizeY, Context.Resource->GetPixelFormat(), ExpectedPitch);
+
+			// Validate buffer size for the mip, so we don't overrun it on streaming
+			ValidateMipBulkDataSize(*Context.Texture, MipMap.SizeX, MipMap.SizeY, MipIndex, MipData[MipIndex].Size, MipSize);
 
-			// @@!! MipData[] should have size
 			FDerivedDataIoOptions Options(MakeMemoryView(MipData[MipIndex].Data, MipSize));
 			Batch.Read(MipMap.DerivedData, Options);
 
diff --git a/Engine/Source/Runtime/Engine/Private/Streaming/Texture2DStreamIn_IO.cpp b/Engine/Source/Runtime/Engine/Private/Streaming/Texture2DStreamIn_IO.cpp
--- a/Engine/Source/Runtime/Engine/Private/Streaming/Texture2DStreamIn_IO.cpp
+++ b/Engine/Source/Runtime/Engine/Private/Streaming/Texture2DStreamIn_IO.cpp
@@ -14,14 +14,10 @@ Texture2DStreamIn.cpp: Stream in helper for 2D textures using texture streaming
 #include "Streaming/Texture2DUpdate.h"
 #include "ProfilingDebugging/CsvProfiler.h"
 #include "ProfilingDebugging/IoStoreTrace.h"
+#include "EngineLogs.h"
 
 CSV_DECLARE_CATEGORY_EXTERN(TextureStreaming);
 
-#if PLATFORM_ANDROID
-#include "EngineLogs.h"
-#include "RenderUtils.h"
-#endif
-
 #if !UE_BUILD_SHIPPING
 float GTextureStreamInAbandonCallbackIOLatency = 0.0f;
 static FAutoConsoleVariableRef CVarTextureStreamInFastCancelCallbackIOLatency(
@@ -42,21 +38,17 @@ FTexture2DStreamIn_IO::~FTexture2DStreamIn_IO()
 {
 }
 
-static void ValidateMipBulkDataSize(const UTexture2D& Texture, int32 MipSizeX, int32 MipSizeY, int32 MipIndex, int64& BulkDataSize)
+static void ValidateMipBulkDataSize(const UTexture2D& Texture, int32 MipSizeX, int32 MipSizeY, int32 MipIndex, int64 ActualMipSize, int64& BulkDataSize)
 {
-	// why is this not done on all platforms?
-#if PLATFORM_ANDROID
-	const int64 ExpectedMipSize = CalcTextureMipMapSize((uint32)MipSizeX, (uint32)MipSizeY, Texture.GetPixelFormat(), 0);
-	if (BulkDataSize != ExpectedMipSize)
+	if (BulkDataSize != ActualMipSize)
 	{
 #if !UE_BUILD_SHIPPING
 		UE_LOG(LogTexture, Warning, TEXT("Mip (%d) %dx%d has an unexpected size %lld, expected size %lld. %s, Pixel format %s"), 
-			MipIndex, MipSizeX, MipSizeY, BulkDataSize, ExpectedMipSize, *Texture.GetFullName(), GPixelFormats[Texture.GetPixelFormat()].Name);
+			MipIndex, MipSizeX, MipSizeY, BulkDataSize, ActualMipSize, *Texture.GetFullName(), GPixelFormats[Texture.GetPixelFormat()].Name);
 #endif
 		// Make sure we don't overrun buffer allocated for this mip
-		BulkDataSize = FMath::Min(BulkDataSize, ExpectedMipSize);
+		BulkDataSize = FMath::Min(BulkDataSize, ActualMipSize);
 	}
-#endif // PLATFORM_ANDROID
 }
 
 void FTexture2DStreamIn_IO::SetIORequests(const FContext& Context)
@@ -73,9 +65,7 @@ void FTexture2DStreamIn_IO::SetIORequests(const FContext& Context)
 		if (BulkDataSize > 0)
 		{
 			// Validate buffer size for the mip, so we don't overrun it on streaming
-			// note: MipData[] should have size
-			// ValidateMipBulkDataSize only does anything on Android
-			ValidateMipBulkDataSize(*Context.Texture, MipMap.SizeX, MipMap.SizeY, MipIndex, BulkDataSize);
+			ValidateMipBulkDataSize(*Context.Texture, MipMap.SizeX, MipMap.SizeY, MipIndex, MipData[MipIndex].Size, BulkDataSize);
 			
 			// reads directly into MipData[] , doesn't respect Pitch
 			// we do get a completion callback at AsyncFileCallBack
