Comment fix to my last checkin
#Synced-CL 48964788

diff --git a/Engine/Source/Runtime/Engine/Private/Streaming/Texture2DStreamIn.h b/Engine/Source/Runtime/Engine/Private/Streaming/Texture2DStreamIn.h
--- a/Engine/Source/Runtime/Engine/Private/Streaming/Texture2DStreamIn.h
+++ b/Engine/Source/Runtime/Engine/Private/Streaming/Texture2DStreamIn.h
@@ -49,6 +49,6 @@ protected:
 	// Async create the texture to the requested size.
 	void DoAsyncCreateWithNewMips(const FContext& Context);
 
-	// Make sure the destination mip budffer is large enough to store the bulk data, warn and adjust if not
+	// Make sure the destination mip buffer is large enough to store the bulk data, warn and adjust if not.
 	void ValidateMipBulkDataSize(const UTexture2D& Texture, int32 MipSizeX, int32 MipSizeY, int32 MipIndex, int64 ActualMipSize, int64& BulkDataSize);
 };
