Fixed anim bank cooking on mobile. Added automatic submit for transform provider when instance data is invalidated.

#Synced-CL 48933410

diff --git a/Engine/Source/Runtime/Engine/Classes/Animation/AnimBank.h b/Engine/Source/Runtime/Engine/Classes/Animation/AnimBank.h
--- a/Engine/Source/Runtime/Engine/Classes/Animation/AnimBank.h
+++ b/Engine/Source/Runtime/Engine/Classes/Animation/AnimBank.h
@@ -842,7 +842,7 @@ public:
 public:
 	// Copies items from this anim bank to create a new anim bank instance.
 	UFUNCTION(BlueprintCallable, Category = "TransformProviders|AnimBank")
-	static ENGINE_API UAnimBankDataInstance* Create(UAnimBankData* AnimBank, UInstancedSkinnedMeshComponent* Owner);
+	static ENGINE_API UAnimBankDataInstance* CreateAnimBankDataInstance(UAnimBankData* AnimBank, UInstancedSkinnedMeshComponent* Owner);
 
 	// Allocates a new track with auto play data, or -1 if the track allocation failed.
 	UFUNCTION(BlueprintCallable, Category = "TransformProviders|AnimBank")
@@ -866,7 +866,7 @@ public:
 
 	// Submits all pending changes to rendering.
 	UFUNCTION(BlueprintCallable, Category = "TransformProviders|AnimBank")
-	ENGINE_API void SubmitChanges();
+	ENGINE_API virtual void SubmitChanges() override;
 
 private:
 	void InitRenderData();
@@ -886,6 +886,7 @@ private:
 	FAnimBankTrackPool Tracks;
 	TArray<FSequence> Sequences;
 	TUniquePtr<FAnimBankRenderData> RenderData{ new FAnimBankRenderData };
+	UE::FMutex SubmitChangesMutex;
 
 	friend class UAnimBankData;
 };
diff --git a/Engine/Source/Runtime/Engine/Classes/Animation/AnimTrackPool.h b/Engine/Source/Runtime/Engine/Classes/Animation/AnimTrackPool.h
--- a/Engine/Source/Runtime/Engine/Classes/Animation/AnimTrackPool.h
+++ b/Engine/Source/Runtime/Engine/Classes/Animation/AnimTrackPool.h
@@ -94,6 +94,7 @@ public:
 		if (IsActiveIndex(Index))
 		{
 			SetDirty(Index);
+			Datas[Index] = {};
 			ActiveBits[Index] = false;
 			NumActiveTracks--;
 			return true;
@@ -126,15 +127,35 @@ public:
 		return Patch;
 	}
 
-	void Patch(FPatch&& Patch)
+	struct FDefaultAssign
 	{
-		Datas.SetNumUninitialized(Patch.NumTracks);
+		void operator()(DataType& Out, const DataType& In) const
+		{
+			Out = In;
+		}
+	};
+
+	struct FDefaultDeactivate
+	{
+		void operator()(DataType&) const
+		{
+		}
+	};
+
+	template <typename AssignFunctionType = FDefaultAssign, typename DeactivateFunctionType = FDefaultDeactivate>
+	void Patch(
+		FPatch&& Patch,
+		AssignFunctionType AssignFunction = FDefaultAssign{},
+		DeactivateFunctionType DeactivateFunction = FDefaultDeactivate{})
+	{
+		Datas.SetNum(Patch.NumTracks);
 		ActiveBits.SetNum(Patch.NumTracks, false);
 		DirtyBits.SetNum(Patch.NumTracks, false);
 		NumActiveTracks = Patch.NumActiveTracks;
 
 		for (typename FPatch::FDeactivateOp Op : Patch.DeactivateOps)
 		{
+			DeactivateFunction(Datas[Op.Index]);
 			ActiveBits[Op.Index] = false;
 			SetDirty(Op.Index);
 		}
@@ -142,7 +163,7 @@ public:
 		for (const typename FPatch::FActivateOp& Op : Patch.ActivateOps)
 		{
 			ActiveBits[Op.Index] = true;
-			Datas[Op.Index] = Op.Data;
+			AssignFunction(Datas[Op.Index], Op.Data);
 			SetDirty(Op.Index);
 		}
 	}
@@ -297,6 +318,15 @@ public:
 		}
 	}
 
+	void Empty()
+	{
+		Datas.Empty();
+		ActiveBits.Empty();
+		DirtyBits.Empty();
+		NumDirtyTracks = 0;
+		NumActiveTracks = 0;
+	}
+
 protected:
 	TArray<DataType> Datas;
 	TBitArray<> ActiveBits;
diff --git a/Engine/Source/Runtime/Engine/Classes/Components/InstancedSkinnedMeshComponent.h b/Engine/Source/Runtime/Engine/Classes/Components/InstancedSkinnedMeshComponent.h
--- a/Engine/Source/Runtime/Engine/Classes/Components/InstancedSkinnedMeshComponent.h
+++ b/Engine/Source/Runtime/Engine/Classes/Components/InstancedSkinnedMeshComponent.h
@@ -271,6 +271,9 @@ public:
 	UFUNCTION(BlueprintCallable, Category = "Components|InstancedSkinnedMesh")
 	ENGINE_API int32 GetInstanceCount() const;
 
+	UFUNCTION(BlueprintCallable, Category = "Components|InstancedSkinnedMesh")
+	ENGINE_API FPrimitiveInstanceId GetInstanceId(int32 InstanceIndex) const;
+
 	int32 GetNumCustomDataFloats() const { return NumCustomDataFloats; }
 	const TArray<float>& GetInstanceCustomData() const { return InstanceCustomData; }
 	const TArray<FSkinnedMeshInstanceData>& GetInstanceData() const { return InstanceData; }
diff --git a/Engine/Source/Runtime/Engine/Private/Animation/AnimBank.cpp b/Engine/Source/Runtime/Engine/Private/Animation/AnimBank.cpp
--- a/Engine/Source/Runtime/Engine/Private/Animation/AnimBank.cpp
+++ b/Engine/Source/Runtime/Engine/Private/Animation/AnimBank.cpp
@@ -757,7 +757,7 @@ bool UAnimBank::IsReadyForFinishDestroy()
 
 bool UAnimBank::NeedsLoadForTargetPlatform(const ITargetPlatform* TargetPlatform) const
 {
-	return DoesTargetPlatformSupportNanite(TargetPlatform);
+	return true;
 }
 
 void UAnimBank::GetResourceSizeEx(FResourceSizeEx& CumulativeResourceSize)
@@ -1007,11 +1007,6 @@ void UAnimBank::Reschedule(FQueuedThreadPool* InThreadPool, EQueuedWorkPriority
 
 FIoHash UAnimBank::CreateDerivedDataKeyHash(const ITargetPlatform* TargetPlatform)
 {
-	if (!DoesTargetPlatformSupportNanite(TargetPlatform))
-	{
-		return FIoHash::Zero;
-	}
-
 	FMemoryHasherBlake3 Writer;
 	FGuid AnimBankVersionGuid(0x51842261, 0xDBAF4D8C, 0xB14ABA3D, 0xC9A914B5);
 	Writer << AnimBankVersionGuid;
@@ -1250,7 +1245,7 @@ void FAnimBankTrackPackedData::Wrap(float SequenceLength, bool bLooping)
 
 void FAnimBankTrackPackedData::Serialize(FArchive& Ar)
 {
-	constexpr auto MakeBitMask = [] (uint32 Index) -> uint32 { return ~((1 << Index) - 1); };
+	constexpr auto MakeBitMask = [] (uint32 Bits) -> uint32 { return (Bits >= 32) ? 0xffffffffu : ((1u << Bits) - 1u); };
 
 	constexpr uint32 ItemIndexBits = 20;
 	constexpr uint32 ItemIndexMask = MakeBitMask(ItemIndexBits);
@@ -1563,9 +1558,9 @@ float UAnimBankData::GetSequencePlayLength(int32 ItemIndex) const
 
 ///////////////////////////////////////////////////////////////////////////////////////////////////
 
-UAnimBankDataInstance* UAnimBankDataInstance::Create(UAnimBankData* AnimBank, UInstancedSkinnedMeshComponent* Owner)
+UAnimBankDataInstance* UAnimBankDataInstance::CreateAnimBankDataInstance(UAnimBankData* AnimBank, UInstancedSkinnedMeshComponent* Owner)
 {
-	if (IsValid(Owner))
+	if (IsValid(AnimBank) && IsValid(Owner))
 	{
 		UAnimBankDataInstance* AnimBankInstance = NewObject<UAnimBankDataInstance>(Owner);
 		AnimBankInstance->AnimBankItems = AnimBank->AnimBankItems;
@@ -1631,6 +1626,9 @@ void UAnimBankDataInstance::SubmitChanges()
 {
 	check(RenderData);
 
+	// SubmitChanges is called during the ISKM instance data update. Guard against multiple concurrent updates.
+	UE::TScopeLock Lock(SubmitChangesMutex);
+
 	if (Tracks.GetNumDirtyTracks() > 0)
 	{
 		ENQUEUE_RENDER_COMMAND(AnimBankDataInstanceSubmit)([&RenderData = *RenderData, Patch = Tracks.Finalize()] (FRHICommandList&) mutable
diff --git a/Engine/Source/Runtime/Engine/Private/Components/InstancedSkinnedMeshComponent.cpp b/Engine/Source/Runtime/Engine/Private/Components/InstancedSkinnedMeshComponent.cpp
--- a/Engine/Source/Runtime/Engine/Private/Components/InstancedSkinnedMeshComponent.cpp
+++ b/Engine/Source/Runtime/Engine/Private/Components/InstancedSkinnedMeshComponent.cpp
@@ -257,6 +257,15 @@ int32 UInstancedSkinnedMeshComponent::GetInstanceCount() const
 	return bIsInstanceDataGPUOnly ? NumInstancesGPUOnly : InstanceData.Num();
 }
 
+FPrimitiveInstanceId UInstancedSkinnedMeshComponent::GetInstanceId(int32 InstanceIndex) const
+{
+	if (!bIsInstanceDataGPUOnly && InstanceIndex >= 0 && InstanceIndex < InstanceDataManager.GetMaxInstanceIndex())
+	{
+		return InstanceDataManager.IndexToId(InstanceIndex);
+	}
+	return FPrimitiveInstanceId();
+}
+
 UTransformProviderData* UInstancedSkinnedMeshComponent::GetTransformProvider() const
 {
 	return TransformProvider.Get();
@@ -383,6 +392,11 @@ void UInstancedSkinnedMeshComponent::SendRenderInstanceData_Concurrent()
 		return;
 	}
 
+	if (TransformProvider)
+	{
+		TransformProvider->SubmitChanges();
+	}
+
 	// If the primitive isn't hidden update its instances.
 	const bool bDetailModeAllowsRendering = true;//DetailMode <= GetCachedScalabilityCVars().DetailMode;
 	// The proxy may not be created, this can happen when a SM is async loading for example.
diff --git a/Engine/Source/Runtime/Engine/Public/Animation/TransformProviderData.h b/Engine/Source/Runtime/Engine/Public/Animation/TransformProviderData.h
--- a/Engine/Source/Runtime/Engine/Public/Animation/TransformProviderData.h
+++ b/Engine/Source/Runtime/Engine/Public/Animation/TransformProviderData.h
@@ -77,6 +77,10 @@ public:
 		return nullptr;
 	}
 
+	virtual void SubmitChanges()
+	{
+	}
+
 #if WITH_EDITOR
 	void BeginCacheForCookedPlatformData(const ITargetPlatform* TargetPlatform)
 	{
