Fix for crash when calling SetPhysicsAsset during blueprint recompilation
#Synced-CL 48933814

diff --git a/Engine/Source/Runtime/Engine/Private/Components/SkeletalMeshComponent.cpp b/Engine/Source/Runtime/Engine/Private/Components/SkeletalMeshComponent.cpp
--- a/Engine/Source/Runtime/Engine/Private/Components/SkeletalMeshComponent.cpp
+++ b/Engine/Source/Runtime/Engine/Private/Components/SkeletalMeshComponent.cpp
@@ -2923,6 +2923,16 @@ void USkeletalMeshComponent::DispatchParallelEvaluationTasks(FActorComponentTick
 
 void USkeletalMeshComponent::DoParallelEvaluationTasks_OnGameThread()
 {
+#if WITH_EDITOR
+	// If we are re-instancing, we will not have run an update before since 
+	// bNeedsUpdate == false and UAnimInstance::UpdateAnimation gets skipped
+	if (GIsReinstancing)
+	{
+		return;
+	}
+#endif
+
+
 	SwapEvaluationContextBuffers();
 
 	ParallelAnimationEvaluation();
