Add LinearColorRamp object ( ported over from Dataflow similar struct )

#Synced-CL 48963730

diff --git a/Engine/Source/Editor/DetailCustomizations/Private/DetailCustomizations.cpp b/Engine/Source/Editor/DetailCustomizations/Private/DetailCustomizations.cpp
--- a/Engine/Source/Editor/DetailCustomizations/Private/DetailCustomizations.cpp
+++ b/Engine/Source/Editor/DetailCustomizations/Private/DetailCustomizations.cpp
@@ -43,6 +43,7 @@
 #include "CustomPrimitiveDataCustomization.h"
 #include "Customizations/ColorStructCustomization.h"
 #include "Customizations/CurveTableCustomization.h"
+#include "Customizations/LinearColorRampCustomization.h"
 #include "Customizations/MathStructCustomizations.h"
 #include "Customizations/MathStructProxyCustomizations.h"
 #include "Customizations/SlateBrushCustomization.h"
@@ -254,6 +255,7 @@ void FDetailCustomizationsModule::RegisterPropertyTypeCustomizations()
 	RegisterCustomPropertyTypeLayout("DataTableRowHandle", FOnGetPropertyTypeCustomizationInstance::CreateStatic(&FDataTableCustomizationLayout::MakeInstance));
 	RegisterCustomPropertyTypeLayout("DataTableCategoryHandle", FOnGetPropertyTypeCustomizationInstance::CreateStatic(&FDataTableCategoryCustomizationLayout::MakeInstance));
 	RegisterCustomPropertyTypeLayout("CurveTableRowHandle", FOnGetPropertyTypeCustomizationInstance::CreateStatic(&FCurveTableCustomizationLayout::MakeInstance));
+	RegisterCustomPropertyTypeLayout("DataflowColorRamp", FOnGetPropertyTypeCustomizationInstance::CreateStatic(&FLinearColorRampCustomization::MakeInstance));
 	REGISTER_UIMINMAX_CUSTOMIZATION(NAME_Vector, &FVectorStructCustomization::MakeInstance);
 	REGISTER_UIMINMAX_CUSTOMIZATION(NAME_Vector3f, &FVectorStructCustomization::MakeInstance);
 	REGISTER_UIMINMAX_CUSTOMIZATION(NAME_Vector3d, &FVectorStructCustomization::MakeInstance);
diff --git a/Engine/Source/Editor/DetailCustomizations/Private/LinearColorRampCustomization.cpp b/Engine/Source/Editor/DetailCustomizations/Private/LinearColorRampCustomization.cpp
new file mode 100644
--- /dev/null
+++ b/Engine/Source/Editor/DetailCustomizations/Private/LinearColorRampCustomization.cpp
@@ -0,0 +1,43 @@
+// Copyright Epic Games, Inc. All Rights Reserved.
+
+#include "Customizations/LinearColorRampCustomization.h"
+#include "DetailLayoutBuilder.h"
+#include "DetailCategoryBuilder.h"
+#include "SColorGradientEditor.h"
+#include "Curves/LinearColorRamp.h"
+#include "DetailWidgetRow.h"
+
+TSharedRef<IPropertyTypeCustomization> FLinearColorRampCustomization::MakeInstance()
+{
+	return MakeShareable(new FLinearColorRampCustomization);
+}
+
+void FLinearColorRampCustomization::CustomizeHeader(TSharedRef<IPropertyHandle> PropertyHandle, FDetailWidgetRow& HeaderRow, IPropertyTypeCustomizationUtils& CustomizationUtils)
+{
+	if (const FStructProperty* const StructProperty = CastField<FStructProperty>(PropertyHandle->GetProperty()))
+	{
+		if (StructProperty->Struct == FLinearColorRamp::StaticStruct())
+		{
+			void* Data = nullptr;
+
+			if (PropertyHandle->GetValueData(Data) == FPropertyAccess::Success)
+			{
+				FLinearColorRamp* ColorRamp = reinterpret_cast<FLinearColorRamp*>(Data);
+
+				TSharedPtr<SColorGradientEditor> GradientEditor;
+
+				HeaderRow
+					[
+						SAssignNew(GradientEditor, SColorGradientEditor)
+							.ViewMinInput(0.0f)
+							.ViewMaxInput(1.0f)
+							.ClampStopsToViewRange(true)
+					];
+
+				GradientEditor->SetCurveOwner(ColorRamp);
+			}
+		}
+	}
+}
+
+
diff --git a/Engine/Source/Editor/DetailCustomizations/Public/Customizations/LinearColorRampCustomization.h b/Engine/Source/Editor/DetailCustomizations/Public/Customizations/LinearColorRampCustomization.h
new file mode 100644
--- /dev/null
+++ b/Engine/Source/Editor/DetailCustomizations/Public/Customizations/LinearColorRampCustomization.h
@@ -0,0 +1,18 @@
+// Copyright Epic Games, Inc. All Rights Reserved.
+
+#pragma once
+
+#include "IPropertyTypeCustomization.h"
+
+class FLinearColorRampCustomization : public IPropertyTypeCustomization
+{
+public:
+
+	/** Makes a new instance of this detail layout class for a specific detail view requesting it */
+	DETAILCUSTOMIZATIONS_API static TSharedRef<IPropertyTypeCustomization> MakeInstance();
+
+	//~ Begin - IPropertyTypeCustomization interface
+	DETAILCUSTOMIZATIONS_API virtual void CustomizeHeader(TSharedRef<IPropertyHandle> PropertyHandle, FDetailWidgetRow& HeaderRow, IPropertyTypeCustomizationUtils& CustomizationUtils) override;
+	virtual void CustomizeChildren(TSharedRef<IPropertyHandle> /*StructPropertyHandle*/, IDetailChildrenBuilder& /*ChildBuilder*/, IPropertyTypeCustomizationUtils& /*CustomizationUtils*/) override {}
+	//~ End - IPropertyTypeCustomization interface
+};
diff --git a/Engine/Source/Runtime/Engine/Classes/Curves/LinearColorRamp.h b/Engine/Source/Runtime/Engine/Classes/Curves/LinearColorRamp.h
new file mode 100644
--- /dev/null
+++ b/Engine/Source/Runtime/Engine/Classes/Curves/LinearColorRamp.h
@@ -0,0 +1,49 @@
+// Copyright Epic Games, Inc. All Rights Reserved.
+
+#pragma once 
+
+#include "Curves/RichCurve.h"
+#include "Curves/CurveOwnerInterface.h"
+
+#include "LinearColorRamp.generated.h"
+
+#define UE_API ENGINE_API
+
+USTRUCT()
+struct FLinearColorRamp 
+#if CPP
+	: public FCurveOwnerInterface
+#endif
+{
+	GENERATED_USTRUCT_BODY()
+
+public:
+	UE_API void SetColorAtTime(float Time, const FLinearColor& Color, bool bOnlyRBG);
+	UE_API bool IsEmpty() const;
+	UE_API void Reset();
+	UE_API void SetFrom(const FLinearColorRamp& Other);
+
+	DECLARE_MULTICAST_DELEGATE_OneParam(FOnColorCurveChanged, TArray<FRichCurve*> /* ChangedCurves */);
+	FOnColorCurveChanged OnColorCurveChangedDelegate;
+
+	//~ Begin FCurveOwnerInterface interface
+	UE_API virtual TArray<FRichCurveEditInfoConst> GetCurves() const override;
+	UE_API virtual void GetCurves(TAdderReserverRef<FRichCurveEditInfoConst> Curves) const override;
+	UE_API virtual TArray<FRichCurveEditInfo> GetCurves() override;
+	UE_API virtual void ModifyOwner() override;
+	UE_API virtual TArray<const UObject*> GetOwners() const override;
+	UE_API virtual void MakeTransactional() override;
+	UE_API virtual void OnCurveChanged(const TArray<FRichCurveEditInfo>& ChangedCurveEditInfos) override;
+	UE_API virtual bool IsLinearColorCurve() const override;
+	UE_API virtual FLinearColor GetLinearColorValue(float InTime) const override;
+	UE_API virtual bool HasAnyAlphaKeys() const override;
+	UE_API virtual bool IsValidCurve(FRichCurveEditInfo CurveInfo) override;
+	UE_API virtual FLinearColor GetCurveColor(FRichCurveEditInfo CurveInfo) const override;
+	//~ End FCurveOwnerInterface interface
+
+private:
+	UPROPERTY()
+	FRichCurve RichCurves[4];
+};
+
+#undef UE_API
diff --git a/Engine/Source/Runtime/Engine/Private/Curves/LinearColorRamp.cpp b/Engine/Source/Runtime/Engine/Private/Curves/LinearColorRamp.cpp
new file mode 100644
--- /dev/null
+++ b/Engine/Source/Runtime/Engine/Private/Curves/LinearColorRamp.cpp
@@ -0,0 +1,157 @@
+// Copyright Epic Games, Inc. All Rights Reserved.
+
+#include "Curves/LinearColorRamp.h"
+
+#include UE_INLINE_GENERATED_CPP_BY_NAME(LinearColorRamp)
+
+namespace UE::LinearColorMap::Private
+{
+	static const FName RedCurveName("R");
+	static const FName GreenCurveName("G");
+	static const FName BlueCurveName("B");
+	static const FName AlphaCurveName("A");
+}
+
+////////////////////////////////////////////////////////////////////////////////
+
+void FLinearColorRamp::SetFrom(const FLinearColorRamp& Other)
+{
+	RichCurves[0] = Other.RichCurves[0];
+	RichCurves[1] = Other.RichCurves[1];
+	RichCurves[2] = Other.RichCurves[2];
+	RichCurves[3] = Other.RichCurves[3];
+	OnCurveChanged(GetCurves());
+}
+
+
+bool FLinearColorRamp::IsEmpty() const
+{
+	return (RichCurves[0].IsEmpty()
+		&& RichCurves[1].IsEmpty()
+		&& RichCurves[2].IsEmpty()
+		&& RichCurves[3].IsEmpty());
+}
+
+void FLinearColorRamp::Reset()
+{
+	RichCurves[0].Reset();
+	RichCurves[1].Reset();
+	RichCurves[2].Reset();
+	RichCurves[3].Reset();
+}
+
+void FLinearColorRamp::SetColorAtTime(float Time, const FLinearColor& Color, bool bOnlyRBG)
+{
+	RichCurves[0].AddKey(Time, Color.R);
+	RichCurves[1].AddKey(Time, Color.G);
+	RichCurves[2].AddKey(Time, Color.B);
+	if (!bOnlyRBG)
+	{
+		RichCurves[3].AddKey(Time, Color.A);
+	}
+}
+
+TArray<FRichCurveEditInfoConst> FLinearColorRamp::GetCurves() const
+{
+	using namespace UE::LinearColorMap;
+
+	TArray<FRichCurveEditInfoConst> Curves;
+	Curves.Add(FRichCurveEditInfoConst(&RichCurves[0], Private::RedCurveName));
+	Curves.Add(FRichCurveEditInfoConst(&RichCurves[1], Private::GreenCurveName));
+	Curves.Add(FRichCurveEditInfoConst(&RichCurves[2], Private::BlueCurveName));
+	Curves.Add(FRichCurveEditInfoConst(&RichCurves[3], Private::AlphaCurveName));
+	return Curves;
+}
+
+void FLinearColorRamp::GetCurves(TAdderReserverRef<FRichCurveEditInfoConst> Curves) const
+{
+	using namespace UE::LinearColorMap;
+
+	Curves.Add(FRichCurveEditInfoConst(&RichCurves[0], Private::RedCurveName));
+	Curves.Add(FRichCurveEditInfoConst(&RichCurves[1], Private::GreenCurveName));
+	Curves.Add(FRichCurveEditInfoConst(&RichCurves[2], Private::BlueCurveName));
+	Curves.Add(FRichCurveEditInfoConst(&RichCurves[3], Private::AlphaCurveName));
+}
+
+TArray<FRichCurveEditInfo> FLinearColorRamp::GetCurves()
+{
+	using namespace UE::LinearColorMap;
+
+	TArray<FRichCurveEditInfo> Curves;
+	Curves.Add(FRichCurveEditInfo(&RichCurves[0], Private::RedCurveName));
+	Curves.Add(FRichCurveEditInfo(&RichCurves[1], Private::GreenCurveName));
+	Curves.Add(FRichCurveEditInfo(&RichCurves[2], Private::BlueCurveName));
+	Curves.Add(FRichCurveEditInfo(&RichCurves[3], Private::AlphaCurveName));
+	return Curves;
+}
+
+void FLinearColorRamp::ModifyOwner()
+{
+}
+
+TArray<const UObject*> FLinearColorRamp::GetOwners() const
+{
+	static TArray<const UObject*> EmptyObjectArray;
+	return EmptyObjectArray;
+}
+
+void FLinearColorRamp::MakeTransactional()
+{
+}
+
+void FLinearColorRamp::OnCurveChanged(const TArray<FRichCurveEditInfo>& ChangedCurveEditInfos)
+{
+	TArray<FRichCurve*> ChangedCurves;
+	for (const FRichCurveEditInfo& ChangedCurveEditInfo : ChangedCurveEditInfos)
+	{
+		ChangedCurves.Add((FRichCurve*)ChangedCurveEditInfo.CurveToEdit);
+	}
+	if (ChangedCurves.Num() > 0)
+	{
+		OnColorCurveChangedDelegate.Broadcast(ChangedCurves);
+	}
+}
+
+bool FLinearColorRamp::IsLinearColorCurve() const
+{
+	return true;
+}
+
+FLinearColor FLinearColorRamp::GetLinearColorValue(float InTime) const
+{
+	return FLinearColor(
+		RichCurves[0].Eval(InTime),
+		RichCurves[1].Eval(InTime),
+		RichCurves[2].Eval(InTime),
+		RichCurves[3].Eval(InTime));
+}
+
+bool FLinearColorRamp::HasAnyAlphaKeys() const
+{
+	return RichCurves[3].GetNumKeys() != 0;
+}
+
+bool FLinearColorRamp::IsValidCurve(FRichCurveEditInfo CurveInfo)
+{
+	return (CurveInfo.CurveToEdit == &RichCurves[0]
+		|| CurveInfo.CurveToEdit == &RichCurves[1]
+		|| CurveInfo.CurveToEdit == &RichCurves[2]
+		|| CurveInfo.CurveToEdit == &RichCurves[3]);
+}
+
+FLinearColor FLinearColorRamp::GetCurveColor(FRichCurveEditInfo CurveInfo) const
+{
+	if (CurveInfo.CurveToEdit == &RichCurves[0])
+	{
+		return FLinearColor::Red;
+	}
+	else if (CurveInfo.CurveToEdit == &RichCurves[1])
+	{
+		return FLinearColor::Green;
+	}
+	else if (CurveInfo.CurveToEdit == &RichCurves[2])
+	{
+		return FLinearColor::Blue;
+	}
+	return FLinearColor::White;
+}
