Fixed animbank serialization to use Serialize function instead of the bulk serialize.

#Synced-CL 48968555

diff --git a/Engine/Source/Runtime/Engine/Classes/Animation/AnimBank.h b/Engine/Source/Runtime/Engine/Classes/Animation/AnimBank.h
--- a/Engine/Source/Runtime/Engine/Classes/Animation/AnimBank.h
+++ b/Engine/Source/Runtime/Engine/Classes/Animation/AnimBank.h
@@ -717,7 +717,7 @@ private:
 	uint32 LoopMode             : 2;
 	uint32 bAutoPlay            : 1;
 	uint32 bHasPreviousPosition : 1;
-	uint32 Unused               : 8;
+	uint32 Unused               : 8 = 0;
 	float Position;
 	float PreviousPosition;
 	float PlayRate;
@@ -725,7 +725,11 @@ private:
 	friend class FAnimBankTransformProvider;
 };
 
-template<> struct TCanBulkSerialize<FAnimBankTrackPackedData> { enum { Value = true }; };
+inline FArchive& operator <<(FArchive& Ar, FAnimBankTrackPackedData& Data)
+{
+	Data.Serialize(Ar);
+	return Ar;
+}
 
 using FAnimBankTrackPool = TAnimTrackPool<FAnimBankTrackPackedData>;
 
