Fix compilation issue for FLinearColorRamp

#Synced-CL 48963730

diff --git a/Engine/Source/Runtime/Engine/Private/Curves/LinearColorRamp.cpp b/Engine/Source/Runtime/Engine/Private/Curves/LinearColorRamp.cpp
--- a/Engine/Source/Runtime/Engine/Private/Curves/LinearColorRamp.cpp
+++ b/Engine/Source/Runtime/Engine/Private/Curves/LinearColorRamp.cpp
@@ -53,35 +53,29 @@ void FLinearColorRamp::SetColorAtTime(float Time, const FLinearColor& Color, boo
 
 TArray<FRichCurveEditInfoConst> FLinearColorRamp::GetCurves() const
 {
-	using namespace UE::LinearColorMap;
-
 	TArray<FRichCurveEditInfoConst> Curves;
-	Curves.Add(FRichCurveEditInfoConst(&RichCurves[0], Private::RedCurveName));
-	Curves.Add(FRichCurveEditInfoConst(&RichCurves[1], Private::GreenCurveName));
-	Curves.Add(FRichCurveEditInfoConst(&RichCurves[2], Private::BlueCurveName));
-	Curves.Add(FRichCurveEditInfoConst(&RichCurves[3], Private::AlphaCurveName));
+	Curves.Add(FRichCurveEditInfoConst(&RichCurves[0], UE::LinearColorMap::Private::RedCurveName));
+	Curves.Add(FRichCurveEditInfoConst(&RichCurves[1], UE::LinearColorMap::Private::GreenCurveName));
+	Curves.Add(FRichCurveEditInfoConst(&RichCurves[2], UE::LinearColorMap::Private::BlueCurveName));
+	Curves.Add(FRichCurveEditInfoConst(&RichCurves[3], UE::LinearColorMap::Private::AlphaCurveName));
 	return Curves;
 }
 
 void FLinearColorRamp::GetCurves(TAdderReserverRef<FRichCurveEditInfoConst> Curves) const
 {
-	using namespace UE::LinearColorMap;
-
-	Curves.Add(FRichCurveEditInfoConst(&RichCurves[0], Private::RedCurveName));
-	Curves.Add(FRichCurveEditInfoConst(&RichCurves[1], Private::GreenCurveName));
-	Curves.Add(FRichCurveEditInfoConst(&RichCurves[2], Private::BlueCurveName));
-	Curves.Add(FRichCurveEditInfoConst(&RichCurves[3], Private::AlphaCurveName));
+	Curves.Add(FRichCurveEditInfoConst(&RichCurves[0],UE::LinearColorMap::Private::RedCurveName));
+	Curves.Add(FRichCurveEditInfoConst(&RichCurves[1],UE::LinearColorMap::Private::GreenCurveName));
+	Curves.Add(FRichCurveEditInfoConst(&RichCurves[2],UE::LinearColorMap::Private::BlueCurveName));
+	Curves.Add(FRichCurveEditInfoConst(&RichCurves[3],UE::LinearColorMap::Private::AlphaCurveName));
 }
 
 TArray<FRichCurveEditInfo> FLinearColorRamp::GetCurves()
 {
-	using namespace UE::LinearColorMap;
-
 	TArray<FRichCurveEditInfo> Curves;
-	Curves.Add(FRichCurveEditInfo(&RichCurves[0], Private::RedCurveName));
-	Curves.Add(FRichCurveEditInfo(&RichCurves[1], Private::GreenCurveName));
-	Curves.Add(FRichCurveEditInfo(&RichCurves[2], Private::BlueCurveName));
-	Curves.Add(FRichCurveEditInfo(&RichCurves[3], Private::AlphaCurveName));
+	Curves.Add(FRichCurveEditInfo(&RichCurves[0], UE::LinearColorMap::Private::RedCurveName));
+	Curves.Add(FRichCurveEditInfo(&RichCurves[1], UE::LinearColorMap::Private::GreenCurveName));
+	Curves.Add(FRichCurveEditInfo(&RichCurves[2], UE::LinearColorMap::Private::BlueCurveName));
+	Curves.Add(FRichCurveEditInfo(&RichCurves[3], UE::LinearColorMap::Private::AlphaCurveName));
 	return Curves;
 }
 
