Iris
* Cleanup deprecated BeginReplication/StopReplication functions in Actor and ActorComponent

#rnx
#Synced-CL 48942789

diff --git a/Engine/Source/Runtime/Engine/Classes/Components/ActorComponent.h b/Engine/Source/Runtime/Engine/Classes/Components/ActorComponent.h
--- a/Engine/Source/Runtime/Engine/Classes/Components/ActorComponent.h
+++ b/Engine/Source/Runtime/Engine/Classes/Components/ActorComponent.h
@@ -762,17 +762,14 @@ public:
 
 	/** Called when we want to start replicating this component, should not be called explicitly.*/
 	UE_DEPRECATED(5.7, "This is being deprecated as part of iris beta. Override OnReplicationStartedForIris instead.")
-	ENGINE_API virtual void BeginReplication();
+	virtual void BeginReplication() {}
 
 	/** Tell component to end replication, should not be called explicitly. */
 	UE_DEPRECATED(5.7, "This is being deprecated as part of iris beta. Override OnStopReplicationForIris instead.")
-	ENGINE_API virtual void EndReplication();
+	virtual void EndReplication() {}
 
 private:
 
-	//ENGINE_API static void StartReplicatingForIris(UActorComponent* ActorComponent);
-	//ENGINE_API static void StopReplicatingForIris(UActorComponent* ActorComponent);
-
 	/** If WorldPrivate isn't set this will determine the world from outers */
 	ENGINE_API UWorld* GetWorld_Uncached() const;
 
diff --git a/Engine/Source/Runtime/Engine/Classes/GameFramework/Actor.h b/Engine/Source/Runtime/Engine/Classes/GameFramework/Actor.h
--- a/Engine/Source/Runtime/Engine/Classes/GameFramework/Actor.h
+++ b/Engine/Source/Runtime/Engine/Classes/GameFramework/Actor.h
@@ -3501,20 +3501,20 @@ protected:
 	* Called for all Actors set to replicate during BeginPlay, It will also be called if SetReplicates(true) is called and the object is not already replicating
 	*/
 	UE_DEPRECATED(5.7, "This is being deprecated as part of iris beta. Use OnReplicationStartedForIris or FillReplicationParams overrides instead.")
-	ENGINE_API virtual void BeginReplication();
+	virtual void BeginReplication() {}
 
 	/**
 	* Called when we want to end replication for this actor, typically called from EndPlay() for actors that should be replicated during their lifespan
 	*/
 	UE_DEPRECATED(5.7, "This is being deprecated as part of iris beta. Override OnStopReplicationForIris instead.")
-	ENGINE_API virtual void EndReplication(EEndPlayReason::Type EndPlayReason);
+	virtual void EndReplication(EEndPlayReason::Type EndPlayReason) {}
 
 	/**
 	* Helper to BeginReplication passing on additional parameters to the ReplicationSystem, typically called from code overriding normal BeginReplication()
 	* @param Params Additional parameters we want to pass on
 	*/
 	UE_DEPRECATED(5.7, "This is being deprecated as part of iris beta. Do not call directly override FillReplicationParams instead to provide custom parameters.")
-	ENGINE_API void BeginReplication(const FActorReplicationParams& Params);
+	void BeginReplication(const FActorReplicationParams& Params) {}
 
 	/**
 	 * Pushes the owning NetConnection for the actor and all of its children to the replication system.
diff --git a/Engine/Source/Runtime/Engine/Private/ActorReplication.cpp b/Engine/Source/Runtime/Engine/Private/ActorReplication.cpp
--- a/Engine/Source/Runtime/Engine/Private/ActorReplication.cpp
+++ b/Engine/Source/Runtime/Engine/Private/ActorReplication.cpp
@@ -1055,18 +1055,6 @@ void AActor::OnStopReplicationForIris(const FOnStopReplicationParams&)
 {
 }
 
-void AActor::BeginReplication(const FActorReplicationParams& Params)
-{
-}
-
-void AActor::BeginReplication()
-{
-}
-
-void AActor::EndReplication(EEndPlayReason::Type EndPlayReason)
-{
-}
-
 void AActor::UpdateOwningNetConnection()
 {
 	using namespace UE::Net;
diff --git a/Engine/Source/Runtime/Engine/Private/Components/ActorComponent.cpp b/Engine/Source/Runtime/Engine/Private/Components/ActorComponent.cpp
--- a/Engine/Source/Runtime/Engine/Private/Components/ActorComponent.cpp
+++ b/Engine/Source/Runtime/Engine/Private/Components/ActorComponent.cpp
@@ -3083,15 +3083,6 @@ void UActorComponent::OnStopReplicationForIris(const FOnStopReplicationParams&)
 {
 }
 
-void UActorComponent::BeginReplication()
-{
-}
-
-void UActorComponent::EndReplication()
-{
-}
-
-
 bool UActorComponent::IsEditableWhenInherited() const
 {
 	bool bCanEdit = bEditableWhenInherited;
diff --git a/Engine/Source/Runtime/Engine/Private/Net/Iris/ReplicationSystem/ReplicationSystemUtil.cpp b/Engine/Source/Runtime/Engine/Private/Net/Iris/ReplicationSystem/ReplicationSystemUtil.cpp
--- a/Engine/Source/Runtime/Engine/Private/Net/Iris/ReplicationSystem/ReplicationSystemUtil.cpp
+++ b/Engine/Source/Runtime/Engine/Private/Net/Iris/ReplicationSystem/ReplicationSystemUtil.cpp
@@ -324,15 +324,7 @@ void FReplicationSystemUtil::StartReplicatingActor(AActor* Actor)
 	if (bCallReplicationStarted)
 	{
 		AActor::FOnReplicationStartedParams Params;
-
 		Actor->OnReplicationStartedForIris(Params);	
-
-PRAGMA_DISABLE_DEPRECATION_WARNINGS
-		// TODO: Remove as soon as everything is submitted.
-		// Call old callback
-		Actor->BeginReplication();
-PRAGMA_ENABLE_DEPRECATION_WARNINGS
-
 	}
 }
 
@@ -347,12 +339,6 @@ void FReplicationSystemUtil::StopReplicatingActor(AActor* Actor, const FStopRepl
 	AActor::FOnStopReplicationParams OnStopReplicationParams = { .StopReplicatingParams = StopReplicatingParams };
 	Actor->OnStopReplicationForIris(OnStopReplicationParams);
 
-PRAGMA_DISABLE_DEPRECATION_WARNINGS
-	// TODO: Remove as soon as everything is submitted.
-	// Call old callback
-	Actor->EndReplication(StopReplicatingParams.EndPlayReason);
-PRAGMA_ENABLE_DEPRECATION_WARNINGS
-
 	// If the call is coming from for example destroying an actor then a formerly associated NetRefHandle will no longer be valid.
 	// The bridge itself will verify that the actor is replicated by it so there's no reason to check that either.
 	ForEachReplicationSystem([Actor, StopReplicatingParams](UReplicationSystem* ReplicationSystem)
@@ -401,12 +387,6 @@ void FReplicationSystemUtil::StartReplicatingActorComponent(const AActor* Actor,
 	{
 		UActorComponent::FOnReplicationStartedParams Params;
 		ActorComp->OnReplicationStartedForIris(Params);
-
-PRAGMA_DISABLE_DEPRECATION_WARNINGS
-		// TODO: Remove as soon as everything is submitted.
-		// Call old callback
-		ActorComp->BeginReplication();
-PRAGMA_ENABLE_DEPRECATION_WARNINGS
 	}
 }
 
@@ -501,12 +481,6 @@ void FReplicationSystemUtil::StopReplicatingActorComponent(UActorComponent* Acto
 	UActorComponent::FOnStopReplicationParams OnStopReplicationParams;
 	ActorComponent->OnStopReplicationForIris(OnStopReplicationParams);
 
-PRAGMA_DISABLE_DEPRECATION_WARNINGS
-	// TODO: Remove as soon as everything is submitted.
-	// Call old callback
-	ActorComponent->EndReplication();
-PRAGMA_ENABLE_DEPRECATION_WARNINGS
-
 	ForEachReplicationSystem([ActorComponent](UReplicationSystem* ReplicationSystem)
 	{
 		if (UEngineReplicationBridge* Bridge = Cast<UEngineReplicationBridge>(ReplicationSystem->GetReplicationBridge()))
