[ITF Gizmos] Fix an issue where the viewport rectangle was not taken into account when using unbound cursor drag for viewport interactions.

#rnx
#Synced-CL 48950337

diff --git a/Engine/Source/Editor/UnrealEd/Private/EditorViewportClient.cpp b/Engine/Source/Editor/UnrealEd/Private/EditorViewportClient.cpp
--- a/Engine/Source/Editor/UnrealEd/Private/EditorViewportClient.cpp
+++ b/Engine/Source/Editor/UnrealEd/Private/EditorViewportClient.cpp
@@ -290,10 +290,7 @@ public:
 FViewportCursorLocation::FViewportCursorLocation(const FSceneView* View, FEditorViewportClient* InViewportClient, int32 X, int32 Y, const bool bInClampToBounds)
 	: Origin(ForceInit), Direction(ForceInit), CursorPos(X, Y)
 {
-	const FVector4 ScreenPos =
-		bInClampToBounds
-		? View->CursorToScreen(static_cast<float>(X), static_cast<float>(Y), 0)
-		: View->PixelToScreen(static_cast<float>(X), static_cast<float>(Y), 0);
+	const FVector4 ScreenPos = View->CursorToScreen(static_cast<float>(X), static_cast<float>(Y), 0.0f, bInClampToBounds);
 
 	const FMatrix InvViewMatrix = View->ViewMatrices.GetInvViewMatrix();
 	const FMatrix InvProjMatrix = View->ViewMatrices.GetInvProjectionMatrix();
diff --git a/Engine/Source/Runtime/Engine/Private/SceneView.cpp b/Engine/Source/Runtime/Engine/Private/SceneView.cpp
--- a/Engine/Source/Runtime/Engine/Private/SceneView.cpp
+++ b/Engine/Source/Runtime/Engine/Private/SceneView.cpp
@@ -1371,14 +1371,19 @@ FVector4 FSceneView::PixelToScreen(float InX,float InY,float Z) const
 }
 
 // Similar to PixelToScreen, but handles cases where the viewport isn't the whole render target, such as a mouse cursor location in
-// render target space.  If outside the bounds of the viewport, will clamp to the nearest edge, so a valid result is returned.  The
-// caller can check the cursor against UnscaledViewRect if they want to reject out of bounds cursor locations.
-FVector4 FSceneView::CursorToScreen(float InX, float InY, float Z) const
+// render target space.  If desired and outside the bounds of the viewport, will clamp to the nearest edge, so a valid result is returned.
+FVector4 FSceneView::CursorToScreen(float InX, float InY, float Z, bool bClamped /* = true */) const
 {
-	float ClampedX = FMath::Clamp(InX - UnscaledViewRect.Min.X, 0.0f, UnscaledViewRect.Width());
-	float ClampedY = FMath::Clamp(InY - UnscaledViewRect.Min.Y, 0.0f, UnscaledViewRect.Height());
+	float X = InX - UnscaledViewRect.Min.X;
+	float Y = InY - UnscaledViewRect.Min.Y; 
 
-	return PixelToScreen(ClampedX, ClampedY, Z);
+	if (bClamped)
+	{
+		X = FMath::Clamp(X, 0.0f, UnscaledViewRect.Width());
+		Y = FMath::Clamp(Y, 0.0f, UnscaledViewRect.Height());
+	}
+	
+	return PixelToScreen(X, Y, Z);
 }
 
 /** Transforms a point from the view's world-space into pixel coordinates relative to the view's X,Y. */
diff --git a/Engine/Source/Runtime/Engine/Public/SceneView.h b/Engine/Source/Runtime/Engine/Public/SceneView.h
--- a/Engine/Source/Runtime/Engine/Public/SceneView.h
+++ b/Engine/Source/Runtime/Engine/Public/SceneView.h
@@ -1905,7 +1905,7 @@ public:
 	ENGINE_API FVector4 PixelToScreen(float X,float Y,float Z) const;
 
 	/** Transforms a cursor location in render target pixel coordinates into the view's screen-space, taking into account the viewport rectangle. */
-	ENGINE_API FVector4 CursorToScreen(float X, float Y, float Z) const;
+	ENGINE_API FVector4 CursorToScreen(float X, float Y, float Z, bool bClamped = true) const;
 
 	/** Transforms a point from the view's world-space into pixel coordinates relative to the view's X,Y (left, top). */
 	ENGINE_API bool WorldToPixel(const FVector& WorldPoint,FVector2D& OutPixelLocation) const;
