UWorld::CleanupActors() now also skips processing making invisible levels as the code that spreads this work across several frames relies on the actor count not changing as it keeps an index into the array (when s.LevelStreamingRouteActorEndPlayForRemoveFromWorldGranularity > 0).

#Synced-CL 48932408

diff --git a/Engine/Source/Runtime/Engine/Private/LevelTick.cpp b/Engine/Source/Runtime/Engine/Private/LevelTick.cpp
--- a/Engine/Source/Runtime/Engine/Private/LevelTick.cpp
+++ b/Engine/Source/Runtime/Engine/Private/LevelTick.cpp
@@ -2044,10 +2044,10 @@ void UWorld::CleanupActors()
 	// actors shouldn't be deleted during gameplay.
 	for (ULevel* Level : Levels)
 	{
-		// Don't compact actors array for levels that are currently in the process of being made visible as the
-		// code that spreads this work across several frames relies on the actor count not changing as it keeps
+		// Don't compact actors array for levels that are currently in the process of being made visible or invisible
+		// as the code that spreads this work across several frames relies on the actor count not changing as it keeps
 		// an index into the array.
-		if( ensure(Level != nullptr) && !IsLevelMakingVisible(Level))
+		if( ensure(Level != nullptr) && !IsLevelMakingVisible(Level) && !IsLevelMakingInvisible(Level))
 		{
 			// Actor 0 (world info) and 1 (default brush) are special and should never be removed from the actor array even if NULL
 			const int32 FirstDynamicIndex = 2;
