Added anim runtime transform provider to ISKM, allowing UAF to provide transforms from the CPU.

#Synced-CL 48995677

diff --git a/Engine/Shaders/Shared/SkinningDefinitions.h b/Engine/Shaders/Shared/SkinningDefinitions.h
--- a/Engine/Shaders/Shared/SkinningDefinitions.h
+++ b/Engine/Shaders/Shared/SkinningDefinitions.h
@@ -125,6 +125,7 @@ struct FCompressedBoneTransform
 #define ANIM_BANK_GPU_TRANSFORM_PROVIDER_GUID 0xA5C0027A, 0x8F884C7C, 0x9312F138, 0x71A9300F
 #define ANIM_BANK_CPU_TRANSFORM_PROVIDER_GUID 0xE7D6173D, 0x246F431A, 0x912D384E, 0x156C0D2C
 #define ANIM_RUNTIME_TRANSFORM_PROVIDER_GUID 0xF1508490, 0xFCC24BB9, 0xA9F277B3, 0x1AF766F0
+#define ANIM_MESH_OBJECT_TRANSFORM_PROVIDER_GUID 0x1073C0C5, 0xB981F6FD, 0xBF41B1C2, 0xC4BA0C50
 
 inline void StoreCompressedBoneTransform(FCompressedBoneTransform* CompressedTransform, const FMatrix44f& Transform)
 {
diff --git a/Engine/Source/Runtime/Engine/Classes/Animation/AnimRuntimeTransformProviderData.h b/Engine/Source/Runtime/Engine/Classes/Animation/AnimRuntimeTransformProviderData.h
new file mode 100644
--- /dev/null
+++ b/Engine/Source/Runtime/Engine/Classes/Animation/AnimRuntimeTransformProviderData.h
@@ -0,0 +1,213 @@
+// Copyright Epic Games, Inc. All Rights Reserved.
+
+#pragma once
+
+#include "Animation/TransformProviderData.h"
+#include "Experimental/Containers/RobinHoodHashTable.h"
+#include "AnimTrackPool.h"
+#include "SkinningDefinitions.h"
+#include "SkeletalRenderPublic.h"
+#include "Containers/IntrusiveDoubleLinkedList.h"
+#include "AnimRuntimeTransformProviderData.generated.h"
+
+class UAnimRuntimeTransformProviderData;
+class FAnimRuntimeTransformProviderProxy;
+class IAnimRuntimeTransformProvider;
+
+class FAnimRuntimeTrackTransformData final : public TIntrusiveDoubleLinkedListNode<FAnimRuntimeTrackTransformData>
+{
+public:
+	~FAnimRuntimeTrackTransformData()
+	{
+		check(!IsInList());
+	}
+
+	FAnimRuntimeTrackTransformData(int32 NumTransforms)
+	{
+		Transforms.SetNumUninitialized(NumTransforms);
+	}
+
+	TArrayView<FCompressedBoneTransform> GetTransforms()
+	{
+		return Transforms;
+	}
+
+	TConstArrayView<FCompressedBoneTransform> GetTransforms() const
+	{
+		return Transforms;
+	}
+
+	int32 Num() const
+	{
+		return Transforms.Num();
+	}
+
+	const FCompressedBoneTransform& operator[] (int32 Index) const
+	{
+		return Transforms[Index];
+	}
+
+	FCompressedBoneTransform& operator[] (int32 Index)
+	{
+		return Transforms[Index];
+	}
+
+	void SetRevisionNumber(uint32 InRevisionNumber)
+	{
+		RevisionNumber = InRevisionNumber;
+	}
+
+	uint32 GetRevisionNumber() const
+	{
+		return RevisionNumber;
+	}
+
+private:
+	TArray<FCompressedBoneTransform> Transforms;
+	uint32 RevisionNumber = 0;
+
+	friend UAnimRuntimeTransformProviderData;
+};
+
+struct FAnimRuntimeTrackData
+{
+	FAnimRuntimeTrackTransformData* Current = nullptr;
+	FAnimRuntimeTrackTransformData* Previous = nullptr;
+};
+
+inline FArchive& operator <<(FArchive& Ar, FAnimRuntimeTrackData& Data)
+{
+	// Serialization is a no-op. Defaults to nullptr.
+	return Ar;
+}
+
+using FAnimRuntimeTrackPool = TAnimTrackPool<FAnimRuntimeTrackData>;
+
+class FAnimRuntimeRenderData
+{
+public:
+	ENGINE_API ~FAnimRuntimeRenderData();
+
+private:
+	void Patch(FAnimRuntimeTrackPool::FPatch&& Patch);
+	void InsertProxy(FInstancedSkinningSceneExtensionProxy* Proxy);
+	void RemoveProxy(FInstancedSkinningSceneExtensionProxy* Proxy);
+
+	FAnimRuntimeTrackPool Tracks;
+	Experimental::TRobinHoodHashSet<FInstancedSkinningSceneExtensionProxy*> Proxies;
+
+	UE::FMutex FreeListMutex;
+	TIntrusiveDoubleLinkedList<FAnimRuntimeTrackTransformData> FreeList;
+
+	friend FAnimRuntimeTransformProviderProxy;
+	friend UAnimRuntimeTransformProviderData;
+};
+
+struct FAnimRuntimeTrackAllocation
+{
+	bool IsValid() const
+	{
+		return Index != INDEX_NONE && Current != nullptr;
+	}
+
+	int32 Index = INDEX_NONE;
+	FAnimRuntimeTrackTransformData* Current = nullptr;
+	FAnimRuntimeTrackTransformData* Previous = nullptr;
+};
+
+class FAnimRuntimeTransformProviderProxy : public FTransformProviderRenderProxy
+{
+	FAnimRuntimeTransformProviderProxy(FInstancedSkinningSceneExtensionProxy* InSceneExtensionProxy, FAnimRuntimeRenderData& InRenderData)
+		: SceneExtensionProxy(InSceneExtensionProxy)
+		, RenderData(InRenderData)
+		, Tracks(InRenderData.Tracks)
+	{}
+
+public:
+	virtual void CreateRenderThreadResources(FSceneInterface& Scene, FRHICommandListBase& RHICmdList) override;
+	virtual void DestroyRenderThreadResources() override;
+
+	FAnimRuntimeTrackPool& GetTracks()
+	{
+		return Tracks;
+	}
+
+	const FInstancedSkinningSceneExtensionProxy* GetSceneExtensionProxy() const
+	{
+		return SceneExtensionProxy;
+	}
+
+private:
+	FInstancedSkinningSceneExtensionProxy* SceneExtensionProxy = nullptr;
+	FAnimRuntimeRenderData& RenderData;
+	FAnimRuntimeTrackPool& Tracks;
+	IAnimRuntimeTransformProvider* AnimRuntimeTransformProvider = nullptr;
+
+	friend class UAnimRuntimeTransformProviderData;
+};
+
+class IAnimRuntimeTransformProvider
+{
+public:
+	virtual ~IAnimRuntimeTransformProvider() = default;
+
+	virtual void RegisterProxy(FAnimRuntimeTransformProviderProxy* Proxy) {}
+	virtual void UnregisterProxy(FAnimRuntimeTransformProviderProxy* Proxy) {}
+};
+
+UCLASS(config=Engine, hidecategories=Object, MinimalAPI, BlueprintType, Within=InstancedSkinnedMeshComponent)
+class UAnimRuntimeTransformProviderData : public UTransformProviderData
+{
+	GENERATED_BODY()
+
+public:
+	virtual bool IsEnabled() const override;
+	virtual const FGuid& GetTransformProviderID() const override;
+	virtual uint32 GetUniqueAnimationCount() const override;
+	virtual bool UsesSkeletonBatching() const override;
+	virtual bool HasAnimationBounds() const override;
+	virtual bool GetAnimationBounds(uint32 AnimationIndex, FRenderBounds& OutBounds) const override;
+	virtual uint32 GetSkinningDataOffset(int32 InstanceIndex, const FTransform& ComponentTransform, const FSkinnedMeshInstanceData& InstanceData) const override;
+	virtual void BeginDestroy() override;
+	virtual void Serialize(FArchive& Ar) override;
+
+	virtual FTransformProviderRenderProxy* CreateRenderProxy(FInstancedSkinningSceneExtensionProxy* ExtensionProxy) const override;
+
+public:
+	UFUNCTION(BlueprintCallable, Category = "TransformProviders|AnimBank")
+	static ENGINE_API UAnimRuntimeTransformProviderData* CreateAnimRuntimeTransformProviderData(UInstancedSkinnedMeshComponent* Owner);
+
+	// Allocates a new track and returns a set of current and (optionally) previous transforms. The data should be filled prior to calling SubmitChanges.
+	UFUNCTION(BlueprintCallable, Category = "TransformProviders|AnimRuntime")
+	ENGINE_API int32 AllocateTrack();
+
+	// Allocates a new track and returns a set of current and (optionally) previous transforms. The data should be filled prior to calling SubmitChanges.
+	ENGINE_API FAnimRuntimeTrackAllocation AllocateTrack(EPreviousBoneTransformUpdateMode UpdateMode);
+	
+	// Updates an existing track and returns a set of current and (optionally) previous transforms. The data should be filled prior to calling SubmitChanges.
+	ENGINE_API FAnimRuntimeTrackAllocation UpdateTrack(int32 TrackIndex, EPreviousBoneTransformUpdateMode UpdateMode);
+
+	// Deallocates a track. Returns whether the deallocation succeeded.
+	UFUNCTION(BlueprintCallable, Category = "TransformProviders|AnimRuntime")
+	ENGINE_API bool DeallocateTrack(int32 TrackIndex);
+
+	// Submits all pending changes to rendering.
+	UFUNCTION(BlueprintCallable, Category = "TransformProviders|AnimRuntime")
+	ENGINE_API virtual void SubmitChanges() override;
+
+	UFUNCTION(BlueprintCallable, Category = "TransformProviders|AnimRuntime")
+	inline bool IsValidIndex(int32 TrackIndex) const
+	{
+		return Tracks.IsValidIndex(TrackIndex);
+	}
+
+private:
+	FAnimRuntimeTrackData AcquireTrackData(EPreviousBoneTransformUpdateMode UpdateMode);
+
+	int32 NumTransforms = 0;
+	int32 NumPoolInstances = 0;
+	FAnimRuntimeTrackPool Tracks;
+	TUniquePtr<FAnimRuntimeRenderData> RenderData{ new FAnimRuntimeRenderData };
+	TIntrusiveDoubleLinkedList<FAnimRuntimeTrackTransformData> FreeList;
+	UE::FMutex SubmitChangesMutex;
+};
\ No newline at end of file
diff --git a/Engine/Source/Runtime/Engine/Private/Animation/AnimRuntimeTransformProviderData.cpp b/Engine/Source/Runtime/Engine/Private/Animation/AnimRuntimeTransformProviderData.cpp
new file mode 100644
--- /dev/null
+++ b/Engine/Source/Runtime/Engine/Private/Animation/AnimRuntimeTransformProviderData.cpp
@@ -0,0 +1,336 @@
+// Copyright Epic Games, Inc. All Rights Reserved.
+
+#include "Animation/AnimRuntimeTransformProviderData.h"
+#include "SkinningSceneExtensionProxy.h"
+#include "Components/InstancedSkinnedMeshComponent.h"
+#include "SceneInterface.h"
+#include "UObject/AnimObjectVersion.h"
+
+///////////////////////////////////////////////////////////////////////////////////////////////////
+
+inline void TryAddTail(TIntrusiveDoubleLinkedList<FAnimRuntimeTrackTransformData>& List, FAnimRuntimeTrackTransformData* Data)
+{
+	if (Data)
+	{
+		List.AddTail(Data);
+	}
+}
+
+FAnimRuntimeRenderData::~FAnimRuntimeRenderData()
+{
+	check(Proxies.Num() == 0);
+	check(FreeList.IsEmpty());
+}
+
+void FAnimRuntimeRenderData::InsertProxy(FInstancedSkinningSceneExtensionProxy* Proxy)
+{
+	Proxies.FindOrAdd(Proxy);
+}
+
+void FAnimRuntimeRenderData::RemoveProxy(FInstancedSkinningSceneExtensionProxy* Proxy)
+{
+	Proxies.Remove(Proxy);
+}
+
+void FAnimRuntimeRenderData::Patch(FAnimRuntimeTrackPool::FPatch&& Patch)
+{
+	if (Patch.GetNumTracks() != Tracks.GetNumTracks())
+	{
+		for (FInstancedSkinningSceneExtensionProxy* Proxy : Proxies)
+		{
+			Proxy->SetUniqueAnimationCount(Patch.GetNumTracks());
+		}
+	}
+
+	UE::TScopeLock Lock(FreeListMutex);
+
+	Tracks.Patch(
+		MoveTemp(Patch),
+		// Assignment
+		[&] (FAnimRuntimeTrackData& Dst, const FAnimRuntimeTrackData& Src)
+		{
+			// Retire the previous transform data if it's not a duplicate of the current.
+			if (Dst.Current != Dst.Previous)
+			{
+				TryAddTail(FreeList, Dst.Previous);
+			}
+
+			// Previous data is provided, free everything from the destination.
+			if (Src.Previous)
+			{
+				TryAddTail(FreeList, Dst.Current);
+
+				Dst.Previous = Src.Previous;
+				Dst.Current  = Src.Current;
+			}
+			// Move current to previous on the destination and assign the new current.
+			else
+			{
+				Dst.Previous = Dst.Current;
+				Dst.Current  = Src.Current;
+			}
+		},
+		// Deactivate
+		[&] (FAnimRuntimeTrackData& Dst)
+		{
+			if (Dst.Previous != Dst.Current)
+			{
+				TryAddTail(FreeList, Dst.Previous);
+			}
+			TryAddTail(FreeList, Dst.Current);
+			Dst = {};
+		}
+	);
+}
+
+///////////////////////////////////////////////////////////////////////////////////////////////////
+
+void FAnimRuntimeTransformProviderProxy::CreateRenderThreadResources(FSceneInterface& Scene, FRHICommandListBase& RHICmdList)
+{
+	RenderData.InsertProxy(SceneExtensionProxy);
+
+	AnimRuntimeTransformProvider = Scene.GetAnimRuntimeTransformProvider();
+
+	if (AnimRuntimeTransformProvider)
+	{
+		AnimRuntimeTransformProvider->RegisterProxy(this);
+	}
+}
+
+void FAnimRuntimeTransformProviderProxy::DestroyRenderThreadResources()
+{
+	RenderData.RemoveProxy(SceneExtensionProxy);
+
+	if (AnimRuntimeTransformProvider)
+	{
+		AnimRuntimeTransformProvider->UnregisterProxy(this);
+		AnimRuntimeTransformProvider = nullptr;
+	}
+}
+
+///////////////////////////////////////////////////////////////////////////////////////////////////
+
+bool UAnimRuntimeTransformProviderData::IsEnabled() const
+{
+	return Tracks.GetNumTracks() != 0 && Super::IsEnabled();
+}
+
+const FGuid& UAnimRuntimeTransformProviderData::GetTransformProviderID() const
+{
+	static FGuid GUID(ANIM_RUNTIME_TRANSFORM_PROVIDER_GUID);
+	return GUID;
+}
+
+uint32 UAnimRuntimeTransformProviderData::GetUniqueAnimationCount() const
+{
+	return Tracks.GetNumTracks();
+}
+
+bool UAnimRuntimeTransformProviderData::UsesSkeletonBatching() const
+{
+	return false;
+}
+
+bool UAnimRuntimeTransformProviderData::HasAnimationBounds() const
+{
+	return false;
+}
+
+bool UAnimRuntimeTransformProviderData::GetAnimationBounds(uint32 AnimationIndex, FRenderBounds& OutBounds) const
+{
+	return false;
+}
+
+uint32 UAnimRuntimeTransformProviderData::GetSkinningDataOffset(int32 InstanceIndex, const FTransform& ComponentTransform, const FSkinnedMeshInstanceData& InstanceData) const
+{
+	const uint32 AnimationIndex = InstanceData.AnimationIndex;
+	if (AnimationIndex < uint32(Tracks.GetNumTracks()))
+	{
+		return AnimationIndex * 2u;
+	}
+
+	return 0u;
+}
+
+void UAnimRuntimeTransformProviderData::BeginDestroy()
+{
+	TRACE_CPUPROFILER_EVENT_SCOPE(UAnimRuntimeTransformProviderData::BeginDestroy);
+	check(RenderData);
+
+	// Push any pending changes so that the render thread can clean up everything.
+	if (Tracks.GetNumDirtyTracks() > 0)
+	{
+		SubmitChanges();
+	}
+
+	ENQUEUE_RENDER_COMMAND(AnimRuntimeTransformProviderData_Release)([RenderData = MoveTemp(RenderData), FreeList = MoveTemp(FreeList), NumPoolInstances = NumPoolInstances](FRHICommandList&) mutable
+	{
+		FreeList.AddTail(MoveTemp(RenderData->FreeList));
+		RenderData->FreeList.Reset();
+
+		// Gather all active data and free it.
+		RenderData->Tracks.EnumerateActiveDirtyTracks([&] (int32 Index)
+		{
+			FAnimRuntimeTrackData& Data = RenderData->Tracks.GetData(Index);
+			if (Data.Current != Data.Previous)
+			{
+				TryAddTail(FreeList, Data.Previous);
+			}
+			TryAddTail(FreeList, Data.Current);
+			Data = {};
+		});
+
+		int32 NumFreed = 0;
+		while (FAnimRuntimeTrackTransformData* Data = FreeList.PopTail())
+		{
+			delete Data;
+			NumFreed++;
+		}
+		check(NumFreed == NumPoolInstances);
+	});
+
+	FreeList.Reset();
+	Tracks.Empty();
+	Super::BeginDestroy();
+}
+
+void UAnimRuntimeTransformProviderData::Serialize(FArchive& Ar)
+{
+	Ar.UsingCustomVersion(FAnimObjectVersion::GUID);
+	Super::Serialize(Ar);
+	Tracks.Serialize(Ar);
+	Ar << NumTransforms;
+
+	if (Ar.IsLoading())
+	{
+		RenderData->Tracks = Tracks;
+	}
+}
+
+FTransformProviderRenderProxy* UAnimRuntimeTransformProviderData::CreateRenderProxy(FInstancedSkinningSceneExtensionProxy* ExtensionProxy) const
+{
+	check(RenderData);
+	return new FAnimRuntimeTransformProviderProxy(ExtensionProxy, *RenderData);
+}
+
+UAnimRuntimeTransformProviderData* UAnimRuntimeTransformProviderData::CreateAnimRuntimeTransformProviderData(UInstancedSkinnedMeshComponent* Owner)
+{
+	if (IsValid(Owner))
+	{
+		if (USkinnedAsset* SkinnedAsset = Owner->GetSkinnedAsset())
+		{
+			UAnimRuntimeTransformProviderData* Instance = NewObject<UAnimRuntimeTransformProviderData>(Owner);
+			Instance->NumTransforms = SkinnedAsset->GetRefBasesInvMatrix().Num();
+			Instance->bEnabled = true;
+			return Instance;
+		}
+	}
+
+	return nullptr;
+}
+
+void UAnimRuntimeTransformProviderData::SubmitChanges()
+{
+	check(RenderData);
+
+	// Attempt to reclaim the free list from the last patch operation on the render thread. Include the render command as well
+	// since SubmitChanges is called during the ISKM instance data update. Guard against multiple concurrent updates.
+	UE::TScopeLock Lock(RenderData->FreeListMutex);
+	FreeList.AddTail(MoveTemp(RenderData->FreeList));
+
+	if (Tracks.GetNumDirtyTracks() > 0)
+	{
+		ENQUEUE_RENDER_COMMAND(AnimRuntimeTransformProviderDataSubmit)([&RenderData = *RenderData, Patch = Tracks.Finalize()](FRHICommandList&) mutable
+		{
+			RenderData.Patch(MoveTemp(Patch));
+		});
+	}
+}
+
+FAnimRuntimeTrackData UAnimRuntimeTransformProviderData::AcquireTrackData(EPreviousBoneTransformUpdateMode UpdateMode)
+{
+	const auto AcquireData = [&]
+	{
+		if (!FreeList.IsEmpty())
+		{
+			return FreeList.PopTail();
+		}
+		else
+		{
+			NumPoolInstances++;
+			return new FAnimRuntimeTrackTransformData(NumTransforms);
+		}
+	};
+
+	FAnimRuntimeTrackData TrackData;
+	TrackData.Current = AcquireData();
+	if (UpdateMode == EPreviousBoneTransformUpdateMode::UpdatePrevious)
+	{
+		TrackData.Previous = AcquireData();
+	}
+	else if (UpdateMode == EPreviousBoneTransformUpdateMode::DuplicateCurrentToPrevious)
+	{
+		TrackData.Previous = TrackData.Current;
+	}
+	return TrackData;
+}
+
+int32 UAnimRuntimeTransformProviderData::AllocateTrack()
+{
+	return Tracks.AllocateTrack(FAnimRuntimeTrackData{});
+}
+
+FAnimRuntimeTrackAllocation UAnimRuntimeTransformProviderData::AllocateTrack(EPreviousBoneTransformUpdateMode UpdateMode)
+{
+	FAnimRuntimeTrackData TrackData = AcquireTrackData(UpdateMode);
+
+	FAnimRuntimeTrackAllocation Allocation;
+	Allocation.Index    = Tracks.AllocateTrack(TrackData);
+	Allocation.Current  = TrackData.Current;
+	if (UpdateMode != EPreviousBoneTransformUpdateMode::DuplicateCurrentToPrevious)
+	{
+		Allocation.Previous = TrackData.Previous;
+	}
+	return Allocation;
+}
+
+FAnimRuntimeTrackAllocation UAnimRuntimeTransformProviderData::UpdateTrack(int32 TrackIndex, EPreviousBoneTransformUpdateMode UpdateMode)
+{
+	FAnimRuntimeTrackAllocation Allocation;
+
+	if (!Tracks.IsActiveIndex(TrackIndex))
+	{
+		return Allocation;
+	}
+
+	FAnimRuntimeTrackData& TrackData = Tracks.GetData(TrackIndex);
+
+	// Track was already dirty. We are updating multiple times without submitting changes.
+	if (Tracks.IsDirtyIndex(TrackIndex))
+	{
+		if (TrackData.Previous != TrackData.Current)
+		{
+			TryAddTail(FreeList, TrackData.Previous);
+		}
+		TryAddTail(FreeList, TrackData.Current);
+	}
+
+	TrackData = AcquireTrackData(UpdateMode);
+	const bool bSuccess = Tracks.UpdateTrack(TrackIndex, TrackData);
+	check(bSuccess);
+
+	Allocation.Index    = TrackIndex;
+	Allocation.Current  = TrackData.Current;
+	if (UpdateMode != EPreviousBoneTransformUpdateMode::DuplicateCurrentToPrevious)
+	{
+		Allocation.Previous = TrackData.Previous;
+	}
+	return Allocation;
+}
+
+bool UAnimRuntimeTransformProviderData::DeallocateTrack(int32 TrackIndex)
+{
+	return Tracks.DeallocateTrack(TrackIndex);
+}
+
+///////////////////////////////////////////////////////////////////////////////
\ No newline at end of file
diff --git a/Engine/Source/Runtime/Engine/Private/SkinningSceneExtensionProxy.cpp b/Engine/Source/Runtime/Engine/Private/SkinningSceneExtensionProxy.cpp
--- a/Engine/Source/Runtime/Engine/Private/SkinningSceneExtensionProxy.cpp
+++ b/Engine/Source/Runtime/Engine/Private/SkinningSceneExtensionProxy.cpp
@@ -134,8 +134,8 @@ void FSkinningSceneExtensionProxy::InitForBoneMapStorage()
 
 const FGuid& FSkinningSceneExtensionProxy::GetTransformProviderId() const
 {
-	static FGuid AnimRuntimeId(ANIM_RUNTIME_TRANSFORM_PROVIDER_GUID);
-	return AnimRuntimeId;
+	static FGuid MeshObjectProviderId(ANIM_MESH_OBJECT_TRANSFORM_PROVIDER_GUID);
+	return MeshObjectProviderId;
 }
 
 FInstancedSkinningSceneExtensionProxy::FInstancedSkinningSceneExtensionProxy(const UTransformProviderData* TransformProvider, const FSkeletalMeshObject* InMeshObject, const USkinnedAsset* InSkinnedAsset, bool bAllowScaling)
diff --git a/Engine/Source/Runtime/Engine/Public/SceneInterface.h b/Engine/Source/Runtime/Engine/Public/SceneInterface.h
--- a/Engine/Source/Runtime/Engine/Public/SceneInterface.h
+++ b/Engine/Source/Runtime/Engine/Public/SceneInterface.h
@@ -15,6 +15,7 @@
 
 class AWorldSettings;
 class IAnimBankTransformProvider;
+class IAnimRuntimeTransformProvider;
 class FArchive;
 class FFloat16Color;
 class FInstanceCullingManager;
@@ -432,6 +433,7 @@ public:
 	virtual void RemoveSparseVolumeTextureViewer(FSparseVolumeTextureViewerSceneProxy* SVTV) = 0;
 
 	virtual IAnimBankTransformProvider* GetAnimBankTransformProvider() = 0;
+	virtual IAnimRuntimeTransformProvider* GetAnimRuntimeTransformProvider() = 0;
 
 	/**
 	 * Adds the unique volumetric cloud component to the scene
diff --git a/Engine/Source/Runtime/Renderer/Private/RendererScene.cpp b/Engine/Source/Runtime/Renderer/Private/RendererScene.cpp
--- a/Engine/Source/Runtime/Renderer/Private/RendererScene.cpp
+++ b/Engine/Source/Runtime/Renderer/Private/RendererScene.cpp
@@ -102,6 +102,7 @@
 #include "ObjectCacheContext.h"
 #include "Animation/AnimBank.h"
 #include "Skinning/AnimBankTransformProvider.h"
+#include "Skinning/AnimRuntimeTransformProvider.h"
 #include "Skinning/SkinningTransformProvider.h"
 #include "PathTracing.h"
 
@@ -2736,6 +2737,11 @@ IAnimBankTransformProvider* FScene::GetAnimBankTransformProvider()
 	return SceneExtensions.GetExtensionPtr<FAnimBankTransformProvider>();
 }
 
+IAnimRuntimeTransformProvider* FScene::GetAnimRuntimeTransformProvider()
+{
+	return SceneExtensions.GetExtensionPtr<FAnimRuntimeTransformProvider>();
+}
+
 void FScene::GetLightIESAtlasSlot(const FLightSceneProxy* Proxy, FLightRenderParameters* Out)
 {
 	if (Proxy)
@@ -6777,6 +6783,7 @@ public:
 	virtual void RemoveSparseVolumeTextureViewer(FSparseVolumeTextureViewerSceneProxy* SVTV) override {}
 
 	virtual IAnimBankTransformProvider* GetAnimBankTransformProvider() { return nullptr; }
+	virtual IAnimRuntimeTransformProvider* GetAnimRuntimeTransformProvider() { return nullptr; }
 
 	virtual void AddHairStrands(FHairStrandsInstance* Proxy) override {}
 	virtual void RemoveHairStrands(FHairStrandsInstance* Proxy) override {}
diff --git a/Engine/Source/Runtime/Renderer/Private/ScenePrivate.h b/Engine/Source/Runtime/Renderer/Private/ScenePrivate.h
--- a/Engine/Source/Runtime/Renderer/Private/ScenePrivate.h
+++ b/Engine/Source/Runtime/Renderer/Private/ScenePrivate.h
@@ -3508,6 +3508,7 @@ public:
 	virtual void RemoveSparseVolumeTextureViewer(FSparseVolumeTextureViewerSceneProxy* SVTV) override;
 
 	virtual IAnimBankTransformProvider* GetAnimBankTransformProvider() override;
+	virtual IAnimRuntimeTransformProvider* GetAnimRuntimeTransformProvider() override;
 
 	virtual void SetPhysicsField(class FPhysicsFieldSceneProxy* PhysicsFieldSceneProxy) override;
 	virtual void ResetPhysicsField() override;
diff --git a/Engine/Source/Runtime/Renderer/Private/Skinning/AnimRuntimeTransformProvider.cpp b/Engine/Source/Runtime/Renderer/Private/Skinning/AnimRuntimeTransformProvider.cpp
new file mode 100644
--- /dev/null
+++ b/Engine/Source/Runtime/Renderer/Private/Skinning/AnimRuntimeTransformProvider.cpp
@@ -0,0 +1,159 @@
+// Copyright Epic Games, Inc. All Rights Reserved.
+
+#include "AnimRuntimeTransformProvider.h"
+#include "ScenePrivate.h"
+
+IMPLEMENT_SCENE_EXTENSION(FAnimRuntimeTransformProvider);
+
+static FGuid AnimRuntimeProviderId(ANIM_RUNTIME_TRANSFORM_PROVIDER_GUID);
+
+FAnimRuntimeTransformProvider::FAnimRuntimeTransformProvider(FScene& InScene)
+	: ISceneExtension(InScene)
+{}
+
+bool FAnimRuntimeTransformProvider::ShouldCreateExtension(FScene& InScene)
+{
+	return IsGPUSkinSceneExtensionEnabled() || (NaniteSkinnedMeshesSupported() && DoesRuntimeSupportNanite(GetFeatureLevelShaderPlatform(InScene.GetFeatureLevel()), true, true));
+}
+
+void FAnimRuntimeTransformProvider::InitExtension(FScene& InScene)
+{
+	if (auto TransformProvider = InScene.GetExtensionPtr<FSkinningTransformProvider>())
+	{
+		TransformProvider->RegisterProvider(
+			AnimRuntimeProviderId,
+			FSkinningTransformProvider::FOnProvideTransforms::CreateRaw(this, &FAnimRuntimeTransformProvider::ProvideAnimRuntimeTransforms),
+			false /* Use skeleton batching */
+		);
+	}
+}
+
+void FAnimRuntimeTransformProvider::RegisterProxy(FAnimRuntimeTransformProviderProxy* Proxy)
+{
+	ProxyDataMap.Emplace(Proxy->GetSceneExtensionProxy(), Proxy);
+}
+
+void FAnimRuntimeTransformProvider::UnregisterProxy(FAnimRuntimeTransformProviderProxy* Proxy)
+{
+	ProxyDataMap.Remove(Proxy->GetSceneExtensionProxy());
+}
+
+void FAnimRuntimeTransformProvider::ProvideAnimRuntimeTransforms(FSkinningTransformProvider::FProviderContext& Context)
+{
+	TRACE_CPUPROFILER_EVENT_SCOPE(FAnimRuntimeTransformProvider::ProvideAnimRuntimeTransforms);
+	RDG_EVENT_SCOPE(Context.GraphBuilder, "ProvideAnimRuntimeTransforms");
+
+	uint32 GlobalTransformCount = 0;
+
+	for (const FSkinningTransformProvider::FProviderIndirection Indirection : Context.Indirections)
+	{
+		const FInstancedSkinningSceneExtensionProxy* Proxy = static_cast<FInstancedSkinningSceneExtensionProxy*>(Context.Proxies[Indirection.Index]);
+		FAnimRuntimeTransformProviderProxy* AnimRuntimeProxy = static_cast<FAnimRuntimeTransformProviderProxy*>(Proxy->GetTransformProviderProxy());
+
+		const uint32 TransformCount = Proxy->GetMaxBoneTransformCount();
+		const uint32 ActiveAnimationCount = AnimRuntimeProxy->GetTracks().GetNumActiveTracks();
+		GlobalTransformCount += ActiveAnimationCount * TransformCount * 2; // Current and Previous
+	}
+
+	if (GlobalTransformCount == 0)
+	{
+		return;
+	}
+
+	FRDGBuilder& GraphBuilder = Context.GraphBuilder;
+	FRDGAsyncScatterUploadBuffer& TransformUploadBuffer = *GraphBuilder.AllocObject<FRDGAsyncScatterUploadBuffer>();
+
+	Context.ScatterUploadBuilder.AddPass(GraphBuilder, TransformUploadBuffer, Context.TransformBuffer, GlobalTransformCount, sizeof(FCompressedBoneTransform), TEXT("Skinning.AnimRuntimeTransforms"),
+		[this, Indirections = Context.Indirections, Proxies = Context.Proxies, &ViewFamily = Context.ViewFamily, GlobalTransformCount] (FRDGScatterUploader& ScatterUploader)
+	{
+		TRACE_CPUPROFILER_EVENT_SCOPE(FAnimRuntimeTransformProvider::ProvideAnimRuntimeTransformsTask);
+
+		for (const FSkinningTransformProvider::FProviderIndirection Indirection : Indirections)
+		{
+			const FInstancedSkinningSceneExtensionProxy* Proxy = static_cast<FInstancedSkinningSceneExtensionProxy*>(Proxies[Indirection.Index]);
+			FProxyData& ProxyData = ProxyDataMap.FindChecked(Proxy);
+			FAnimRuntimeTrackPool& Tracks = ProxyData.Tracks;
+			ProxyData.TrackHistories.SetNum(Tracks.GetNumTracks());
+
+			const EBoneTransformStorageMode StorageMode = Proxy->GetBoneTransformStorageMode();
+			const uint32 MaxTransformCount = Proxy->GetMaxBoneTransformCount();
+			const uint32 UniqueAnimationCount = Proxy->GetUniqueAnimationCount();
+			const uint32 MaxTotalTransformCount = MaxTransformCount * 2u; // Current and Previous
+
+			const uint32 CurrentTransformIndex  = Indirection.CurrentTransformOffset  / sizeof(FCompressedBoneTransform);
+			const uint32 PreviousTransformIndex = Indirection.PreviousTransformOffset / sizeof(FCompressedBoneTransform);
+			const uint32 CurrentTransformSlot   = CurrentTransformIndex < PreviousTransformIndex ? 0 : 1;
+
+			Tracks.EnumerateActiveTracks([&] (int32 TrackIndex)
+			{
+				FAnimRuntimeTrackData& Data = Tracks.GetData(TrackIndex);
+
+				bool bNeedsCurrentUpdate, bNeedsPreviousUpdate;
+				ProxyData.TrackHistories[TrackIndex].Update(Data, CurrentTransformSlot, bNeedsCurrentUpdate, bNeedsPreviousUpdate);
+
+				// Dirty previous transforms means the whole allocation is dirty. Rebuild everything.
+				const bool bDirtyTransformAllocation = EnumHasAnyFlags(Indirection.DirtyBoneTransforms, EDirtyBoneTransforms::Previous);
+				bNeedsCurrentUpdate  |= bDirtyTransformAllocation;
+				bNeedsPreviousUpdate |= bDirtyTransformAllocation;
+
+				const auto FillTransforms = [&] (TArrayView<FCompressedBoneTransform> DstTransforms, const FAnimRuntimeTrackTransformData* SrcTransformData)
+				{
+					if (SrcTransformData)
+					{
+						TConstArrayView<FCompressedBoneTransform> SrcTransforms = SrcTransformData->GetTransforms();
+
+						if (StorageMode == EBoneTransformStorageMode::BoneMap)
+						{
+							uint32 DstTransformIndex = 0;
+
+							for (uint32 SrcTransformIndex : Proxy->GetBoneHierarchy())
+							{
+								DstTransforms[DstTransformIndex] = SrcTransforms[SrcTransformIndex];
+								DstTransformIndex++;
+							}
+						}
+						else
+						{
+							for (int32 DstTransformIndex = 0; DstTransformIndex < DstTransforms.Num(); ++DstTransformIndex)
+							{
+								DstTransforms[DstTransformIndex] = SrcTransforms[DstTransformIndex];
+							}
+						}
+					}
+					else
+					{
+						if (StorageMode == EBoneTransformStorageMode::BoneMap)
+						{
+							uint32 DstTransformIndex = 0;
+
+							for (uint32 SrcTransformIndex : Proxy->GetBoneHierarchy())
+							{
+								SetCompressedBoneTransformIdentity(DstTransforms[DstTransformIndex]);
+								DstTransformIndex++;
+							}
+						}
+						else
+						{
+							for (int32 DstTransformIndex = 0; DstTransformIndex < DstTransforms.Num(); ++DstTransformIndex)
+							{
+								SetCompressedBoneTransformIdentity(DstTransforms[DstTransformIndex]);
+							}
+						}
+					}
+				};
+
+				if (bNeedsCurrentUpdate)
+				{
+					TArrayView<FCompressedBoneTransform> DstCurrentTransforms = ScatterUploader.Add_GetRef<FCompressedBoneTransform>(CurrentTransformIndex + TrackIndex * MaxTotalTransformCount, MaxTransformCount);
+					FillTransforms(DstCurrentTransforms, Data.Current);
+				}
+
+				if (bNeedsPreviousUpdate)
+				{
+					TArrayView<FCompressedBoneTransform> DstPreviousTransforms = ScatterUploader.Add_GetRef<FCompressedBoneTransform>(PreviousTransformIndex + TrackIndex * MaxTotalTransformCount, MaxTransformCount);
+					FillTransforms(DstPreviousTransforms, Data.Previous ? Data.Previous : Data.Current);
+				}
+			});
+		}
+	});
+}
\ No newline at end of file
diff --git a/Engine/Source/Runtime/Renderer/Private/Skinning/AnimRuntimeTransformProvider.h b/Engine/Source/Runtime/Renderer/Private/Skinning/AnimRuntimeTransformProvider.h
new file mode 100644
--- /dev/null
+++ b/Engine/Source/Runtime/Renderer/Private/Skinning/AnimRuntimeTransformProvider.h
@@ -0,0 +1,63 @@
+// Copyright Epic Games, Inc. All Rights Reserved.
+
+#pragma once
+
+#include "Animation/AnimRuntimeTransformProviderData.h"
+#include "SkinningTransformProvider.h"
+#include "SkinningSceneExtensionProxy.h"
+#include "UnifiedBuffer.h"
+#include "Math/IntVector.h"
+
+class FAnimRuntimeTransformProvider
+	: public ISceneExtension
+	, public IAnimRuntimeTransformProvider
+{
+	DECLARE_SCENE_EXTENSION(RENDERER_API, FAnimRuntimeTransformProvider);
+
+public:
+	FAnimRuntimeTransformProvider(FScene& InScene);
+
+	static bool ShouldCreateExtension(FScene& InScene);
+
+	virtual void InitExtension(FScene& InScene) override;
+
+	void RegisterProxy(FAnimRuntimeTransformProviderProxy* Proxy) override;
+	void UnregisterProxy(FAnimRuntimeTransformProviderProxy* Proxy) override;
+
+private:
+	void ProvideAnimRuntimeTransforms(FSkinningTransformProvider::FProviderContext& Context);
+
+	struct FTrackHistory
+	{
+		void Update(const FAnimRuntimeTrackData& Data, int32 InCurrentSlot, bool& bUpdateCurrent, bool& bUpdatePrevious)
+		{
+			bUpdateCurrent  = Data.Current  && Slots[ InCurrentSlot] != Data.Current->GetRevisionNumber();
+			bUpdatePrevious = Data.Previous && Slots[!InCurrentSlot] != Data.Previous->GetRevisionNumber();
+
+			if (bUpdateCurrent)
+			{
+				Slots[InCurrentSlot] = Data.Current->GetRevisionNumber();
+			}
+
+			if (bUpdatePrevious)
+			{
+				Slots[!InCurrentSlot] = Data.Previous->GetRevisionNumber();
+			}
+		}
+
+		FUintVector2 Slots = FUintVector2::NoneValue;
+	};
+
+	struct FProxyData
+	{
+		FProxyData(FAnimRuntimeTransformProviderProxy* Proxy)
+			: Tracks(Proxy->GetTracks())
+		{}
+
+		FAnimRuntimeTrackPool& Tracks;
+		TArray<FTrackHistory> TrackHistories;
+	};
+
+	TMap<const FSkinningSceneExtensionProxy*, FProxyData> ProxyDataMap;
+	FRDGAsyncScatterUploadBuffer SkinningTransformUploadBuffer;
+};
\ No newline at end of file
diff --git a/Engine/Source/Runtime/Renderer/Private/Skinning/SkinningSceneExtension.cpp b/Engine/Source/Runtime/Renderer/Private/Skinning/SkinningSceneExtension.cpp
--- a/Engine/Source/Runtime/Renderer/Private/Skinning/SkinningSceneExtension.cpp
+++ b/Engine/Source/Runtime/Renderer/Private/Skinning/SkinningSceneExtension.cpp
@@ -161,7 +161,7 @@ private:
 IMPLEMENT_GLOBAL_SHADER(FRefPoseTransformProviderCS, "/Engine/Private/Skinning/TransformProviders.usf", "RefPoseProviderCS", SF_Compute);
 
 static FGuid RefPoseProviderId(REF_POSE_TRANSFORM_PROVIDER_GUID);
-static FGuid AnimRuntimeProviderId(ANIM_RUNTIME_TRANSFORM_PROVIDER_GUID);
+static FGuid MeshObjectProviderId(ANIM_MESH_OBJECT_TRANSFORM_PROVIDER_GUID);
 
 static void GetDefaultSkinningParameters(FSkinningSceneParameters& OutParameters, FRDGBuilder& GraphBuilder)
 {
@@ -203,8 +203,8 @@ void FSkinningSceneExtension::InitExtension(FScene& InScene)
 		);
 
 		TransformProvider->RegisterProvider(
-			GetAnimRuntimeProviderId(),
-			FSkinningTransformProvider::FOnProvideTransforms::CreateStatic(&FSkinningSceneExtension::ProvideAnimRuntimeTransforms),
+			GetMeshObjectProviderId(),
+			FSkinningTransformProvider::FOnProvideTransforms::CreateStatic(&FSkinningSceneExtension::ProvideMeshObjectTransforms),
 			false /* Use skeleton batching */
 		);
 
@@ -1589,9 +1589,9 @@ const FSkinningTransformProvider::FProviderId& FSkinningSceneExtension::GetRefPo
 	return RefPoseProviderId;
 }
 
-const FSkinningTransformProvider::FProviderId& FSkinningSceneExtension::GetAnimRuntimeProviderId()
+const FSkinningTransformProvider::FProviderId& FSkinningSceneExtension::GetMeshObjectProviderId()
 {
-	return AnimRuntimeProviderId;
+	return MeshObjectProviderId;
 }
 
 void FSkinningSceneExtension::ProvideRefPoseTransforms(FSkinningTransformProvider::FProviderContext& Context)
@@ -1674,15 +1674,10 @@ void FSkinningSceneExtension::ProvideRefPoseTransforms(FSkinningTransformProvide
 	);
 }
 
-BEGIN_SHADER_PARAMETER_STRUCT(FCopyBufferParameters, )
-	RDG_BUFFER_ACCESS(SrcBuffer, ERHIAccess::CopySrc)
-	RDG_BUFFER_ACCESS(DstBuffer, ERHIAccess::CopyDest)
-END_SHADER_PARAMETER_STRUCT()
-
-void FSkinningSceneExtension::ProvideAnimRuntimeTransforms(FSkinningTransformProvider::FProviderContext& Context)
+void FSkinningSceneExtension::ProvideMeshObjectTransforms(FSkinningTransformProvider::FProviderContext& Context)
 {
-	TRACE_CPUPROFILER_EVENT_SCOPE(FSkinningSceneExtension::ProvideAnimRuntimeTransforms);
-	RDG_EVENT_SCOPE(Context.GraphBuilder, "ProvideAnimRuntimeTransforms");
+	TRACE_CPUPROFILER_EVENT_SCOPE(FSkinningSceneExtension::ProvideMeshObjectTransforms);
+	RDG_EVENT_SCOPE(Context.GraphBuilder, "ProvideMeshObjectTransforms");
 
 	uint32 GlobalTransformCount = 0;
 
@@ -1704,7 +1699,7 @@ void FSkinningSceneExtension::ProvideAnimRuntimeTransforms(FSkinningTransformPro
 	Context.ScatterUploadBuilder.AddPass(GraphBuilder, TransformUploadBuffer, Context.TransformBuffer, GlobalTransformCount, sizeof(FCompressedBoneTransform), TEXT("Skinning.AnimTransforms"),
 		[Indirections = Context.Indirections, Proxies = Context.Proxies, &ViewFamily = Context.ViewFamily, GlobalTransformCount] (FRDGScatterUploader& ScatterUploader)
 	{
-		TRACE_CPUPROFILER_EVENT_SCOPE(FSkinningSceneExtension::ProvideAnimRuntimeTransformsTask);
+		TRACE_CPUPROFILER_EVENT_SCOPE(FSkinningSceneExtension::ProvideMeshObjectTransformsTask);
 
 		int64 NumBonesTotal = 0;
 
diff --git a/Engine/Source/Runtime/Renderer/Private/Skinning/SkinningSceneExtension.h b/Engine/Source/Runtime/Renderer/Private/Skinning/SkinningSceneExtension.h
--- a/Engine/Source/Runtime/Renderer/Private/Skinning/SkinningSceneExtension.h
+++ b/Engine/Source/Runtime/Renderer/Private/Skinning/SkinningSceneExtension.h
@@ -86,7 +86,7 @@ public:
 	RENDERER_API void GetSkinnedPrimitives(TArray<FPrimitiveSceneInfo*>& OutPrimitives) const;
 
 	RENDERER_API static const FSkinningTransformProvider::FProviderId& GetRefPoseProviderId();
-	RENDERER_API static const FSkinningTransformProvider::FProviderId& GetAnimRuntimeProviderId();
+	RENDERER_API static const FSkinningTransformProvider::FProviderId& GetMeshObjectProviderId();
 
 private:
 	enum ETask : uint32
@@ -243,5 +243,5 @@ private:
 
 public:
 	RENDERER_API static void ProvideRefPoseTransforms(FSkinningTransformProvider::FProviderContext& Context);
-	RENDERER_API static void ProvideAnimRuntimeTransforms(FSkinningTransformProvider::FProviderContext& Context);
+	RENDERER_API static void ProvideMeshObjectTransforms(FSkinningTransformProvider::FProviderContext& Context);
 };
\ No newline at end of file
diff --git a/Engine/Source/Runtime/Renderer/Private/Skinning/SkinningTransformProvider.cpp b/Engine/Source/Runtime/Renderer/Private/Skinning/SkinningTransformProvider.cpp
--- a/Engine/Source/Runtime/Renderer/Private/Skinning/SkinningTransformProvider.cpp
+++ b/Engine/Source/Runtime/Renderer/Private/Skinning/SkinningTransformProvider.cpp
@@ -71,10 +71,10 @@ const FSkinningTransformProvider::FProviderId& GetRefPoseProviderId()
 	return FSkinningSceneExtension::GetRefPoseProviderId();
 }
 
-const FSkinningTransformProvider::FProviderId& GetAnimRuntimeProviderId()
+const FSkinningTransformProvider::FProviderId& GetMeshObjectProviderId()
 {
 	// TODO: Temp until skinning scene extension is refactored into a public API outside of Nanite
-	return FSkinningSceneExtension::GetAnimRuntimeProviderId();
+	return FSkinningSceneExtension::GetMeshObjectProviderId();
 }
 
 FRDGBufferDesc GetCompressedBoneTransformBufferDesc(int32 NumTransforms)
diff --git a/Engine/Source/Runtime/Renderer/Private/Skinning/SkinningTransformProvider.h b/Engine/Source/Runtime/Renderer/Private/Skinning/SkinningTransformProvider.h
--- a/Engine/Source/Runtime/Renderer/Private/Skinning/SkinningTransformProvider.h
+++ b/Engine/Source/Runtime/Renderer/Private/Skinning/SkinningTransformProvider.h
@@ -197,7 +197,7 @@ private:
 };
 
 RENDERER_API const FSkinningTransformProvider::FProviderId& GetRefPoseProviderId();
-RENDERER_API const FSkinningTransformProvider::FProviderId& GetAnimRuntimeProviderId();
+RENDERER_API const FSkinningTransformProvider::FProviderId& GetMeshObjectProviderId();
 
 RENDERER_API FRDGBufferDesc GetCompressedBoneTransformBufferDesc(int32 NumTransforms);
 RENDERER_API FRDGBufferRef GetDefaultCompressedBoneTransformBuffer(FRDGBuilder& GraphBuilder);
