Full renaming of FMaterialInstanceResource to FMaterialInstanceRenderProxy, and it's associated Resource variable to RenderProxy to ensure the usage is clearer compared to FMaterialResource.

#Synced-CL 49030741

diff --git a/Engine/Config/BaseEngine.ini b/Engine/Config/BaseEngine.ini
--- a/Engine/Config/BaseEngine.ini
+++ b/Engine/Config/BaseEngine.ini
@@ -1169,6 +1169,7 @@ ClassName=/Script/OnlineSubsystemUtils.OnlineEngineInterfaceImpl
 +PropertyRedirects=(OldName="MaterialInstanceBasePropertyOverrides.MaxWorldPositionOffsetDistance", NewName="MaterialInstanceBasePropertyOverrides.MaxWorldPositionOffsetDisplacement")
 +PropertyRedirects=(OldName="MaterialInstanceBasePropertyOverrides.bOverride_MaxWorldPositionOffsetDistance", NewName="MaterialInstanceBasePropertyOverrides.bOverride_MaxWorldPositionOffsetDisplacement")
 +PropertyRedirects=(OldName="Material.MaxWorldPositionOffsetDistance", NewName="Material.MaxWorldPositionOffsetDisplacement")
++PropertyRedirects=(OldName="MaterialInstance.Resource", NewName="MaterialInstance.RenderProxy")
 
 +PropertyRedirects=(OldName="/Script/Engine.SoundAttenuationSettings.OmniRadius", NewName="NonSpatializedRadiusStart")
 
diff --git a/Engine/Source/Runtime/Engine/Private/Materials/MaterialInstance.cpp b/Engine/Source/Runtime/Engine/Private/Materials/MaterialInstance.cpp
--- a/Engine/Source/Runtime/Engine/Private/Materials/MaterialInstance.cpp
+++ b/Engine/Source/Runtime/Engine/Private/Materials/MaterialInstance.cpp
@@ -103,12 +103,12 @@ void UMaterialInstance::FinishCacheUniformExpressions() const
 	bCachingUniformExpressions = false;
 }
 
-void FMaterialInstanceResource::StartCacheUniformExpressions() const
+void FMaterialInstanceRenderProxy::StartCacheUniformExpressions() const
 {
 	Owner->StartCacheUniformExpressions();
 }
 
-void FMaterialInstanceResource::FinishCacheUniformExpressions() const
+void FMaterialInstanceRenderProxy::FinishCacheUniformExpressions() const
 {
 	Owner->FinishCacheUniformExpressions();
 }
@@ -119,9 +119,9 @@ void FMaterialInstanceResource::FinishCacheUniformExpressions() const
  */
 void CacheMaterialInstanceUniformExpressions(const UMaterialInstance* MaterialInstance, bool bRecreateUniformBuffer)
 {
-	if (MaterialInstance->Resource)
+	if (MaterialInstance->GetInstanceRenderProxy())
 	{
-		MaterialInstance->Resource->CacheUniformExpressions_GameThread(bRecreateUniformBuffer);
+		MaterialInstance->GetInstanceRenderProxy()->CacheUniformExpressions_GameThread(bRecreateUniformBuffer);
 	}
 }
 
@@ -168,7 +168,7 @@ FFontParameterValue::ValueType FFontParameterValue::GetValue(const FFontParamete
 	return Value;
 }
 
-FMaterialInstanceResource::FMaterialInstanceResource(UMaterialInstance* InOwner)
+FMaterialInstanceRenderProxy::FMaterialInstanceRenderProxy(UMaterialInstance* InOwner)
 	: FMaterialRenderProxy(InOwner->GetName())
 	, Parent(NULL)
 	, Owner(InOwner)
@@ -176,13 +176,13 @@ FMaterialInstanceResource::FMaterialInstanceResource(UMaterialInstance* InOwner)
 {
 }
 
-void FMaterialInstanceResource::GameThread_Destroy()
+void FMaterialInstanceRenderProxy::GameThread_Destroy()
 {
-	FMaterialInstanceResource* Resource = this;
+	FMaterialInstanceRenderProxy* RenderProxy = this;
 	ENQUEUE_RENDER_COMMAND(FDestroyMaterialInstanceResourceCommand)(
-		[Resource](FRHICommandList& RHICmdList)
+		[RenderProxy](FRHICommandList& RHICmdList)
 		{
-			delete Resource;
+			delete RenderProxy;
 		});
 }
 
@@ -231,7 +231,7 @@ const FMaterial& FMaterialInstanceResource::GetMaterialWithFallback(ERHIFeatureL
 }
 #endif // 0
 
-const FMaterialRenderProxy* FMaterialInstanceResource::GetFallback(ERHIFeatureLevel::Type InFeatureLevel) const
+const FMaterialRenderProxy* FMaterialInstanceRenderProxy::GetFallback(ERHIFeatureLevel::Type InFeatureLevel) const
 {
 	if (Parent)
 	{
@@ -259,7 +259,7 @@ const FMaterialRenderProxy* FMaterialInstanceResource::GetFallback(ERHIFeatureLe
 	return FallbackMaterial->GetRenderProxy();
 }
 
-const FMaterial* FMaterialInstanceResource::GetMaterialNoFallback(ERHIFeatureLevel::Type InFeatureLevel) const
+const FMaterial* FMaterialInstanceRenderProxy::GetMaterialNoFallback(ERHIFeatureLevel::Type InFeatureLevel) const
 {
 	checkSlow(IsInParallelRenderingThread());
 
@@ -286,12 +286,12 @@ const FMaterial* FMaterialInstanceResource::GetMaterialNoFallback(ERHIFeatureLev
 	return nullptr;
 }
 
-UMaterialInterface* FMaterialInstanceResource::GetMaterialInterface() const
+UMaterialInterface* FMaterialInstanceRenderProxy::GetMaterialInterface() const
 {
 	return Owner;
 }
 
-bool FMaterialInstanceResource::GetParameterValue(EMaterialParameterType Type, const FHashedMaterialParameterInfo& ParameterInfo, FMaterialParameterValue& OutValue, const FMaterialRenderContext& Context) const
+bool FMaterialInstanceRenderProxy::GetParameterValue(EMaterialParameterType Type, const FHashedMaterialParameterInfo& ParameterInfo, FMaterialParameterValue& OutValue, const FMaterialRenderContext& Context) const
 {
 	checkSlow(IsInParallelRenderingThread());
 
@@ -397,7 +397,7 @@ bool FMaterialInstanceResource::GetParameterValue(EMaterialParameterType Type, c
 	return bResult;
 }
 
-bool FMaterialInstanceResource::GetUserSceneTextureOverride(FName& InOutName) const
+bool FMaterialInstanceRenderProxy::GetUserSceneTextureOverride(FName& InOutName) const
 {
 	checkSlow(IsInParallelRenderingThread());
 
@@ -423,7 +423,7 @@ bool FMaterialInstanceResource::GetUserSceneTextureOverride(FName& InOutName) co
 	}
 }
 
-EBlendableLocation FMaterialInstanceResource::GetBlendableLocation(const FMaterial* Base) const
+EBlendableLocation FMaterialInstanceRenderProxy::GetBlendableLocation(const FMaterial* Base) const
 {
 	check(Base);
 	checkSlow(IsInParallelRenderingThread());
@@ -449,7 +449,7 @@ EBlendableLocation FMaterialInstanceResource::GetBlendableLocation(const FMateri
 	}
 }
 
-int32 FMaterialInstanceResource::GetBlendablePriority(const FMaterial* Base) const
+int32 FMaterialInstanceRenderProxy::GetBlendablePriority(const FMaterial* Base) const
 {
 	check(Base);
 	checkSlow(IsInParallelRenderingThread());
@@ -470,13 +470,13 @@ int32 FMaterialInstanceResource::GetBlendablePriority(const FMaterial* Base) con
 
 void UMaterialInstance::PropagateDataToMaterialProxy()
 {
-	if (Resource)
+	if (RenderProxy)
 	{
-		UpdateMaterialRenderProxy(*Resource);
+		UpdateMaterialRenderProxy(*RenderProxy);
 	}
 }
 
-void FMaterialInstanceResource::GameThread_SetParent(UMaterialInterface* ParentMaterialInterface)
+void FMaterialInstanceRenderProxy::GameThread_SetParent(UMaterialInterface* ParentMaterialInterface)
 {
 	// @todo loadtimes: this is no longer valid because of the ParallelFor calling AddPrimitive in UnrealEngine.cpp
 	// check(IsInGameThread() || IsAsyncLoading());
@@ -489,12 +489,12 @@ void FMaterialInstanceResource::GameThread_SetParent(UMaterialInterface* ParentM
 
 		// Set the rendering thread's parent and instance pointers.
 		check(ParentMaterialInterface != NULL);
-		FMaterialInstanceResource* Resource = this;
+		FMaterialInstanceRenderProxy* RenderProxy = this;
 		ENQUEUE_RENDER_COMMAND(InitMaterialInstanceResource)(
-			[Resource, ParentMaterialInterface](FRHICommandListImmediate& RHICmdList)
+			[RenderProxy, ParentMaterialInterface](FRHICommandListImmediate& RHICmdList)
 			{
-				Resource->Parent = ParentMaterialInterface;
-				Resource->InvalidateUniformExpressionCache(true);
+				RenderProxy->Parent = ParentMaterialInterface;
+				RenderProxy->InvalidateUniformExpressionCache(true);
 			});
 
 		if (OldParent)
@@ -505,12 +505,12 @@ void FMaterialInstanceResource::GameThread_SetParent(UMaterialInterface* ParentM
 	}
 }
 
-void FMaterialInstanceResource::GameThread_UpdateCachedData(const FMaterialInstanceCachedData& CachedData)
+void FMaterialInstanceRenderProxy::GameThread_UpdateCachedData(const FMaterialInstanceCachedData& CachedData)
 {
 	ENQUEUE_RENDER_COMMAND(MaterialInstanceResource_UpdateCachedData)(
-		[Resource = this, ParentLayerIndexRemap = CachedData.ParentLayerIndexRemap](FRHICommandListImmediate& RHICmdList) mutable
+		[RenderProxy = this, ParentLayerIndexRemap = CachedData.ParentLayerIndexRemap](FRHICommandListImmediate& RHICmdList) mutable
 		{
-			Resource->ParentLayerIndexRemap = MoveTemp(ParentLayerIndexRemap);
+			RenderProxy->ParentLayerIndexRemap = MoveTemp(ParentLayerIndexRemap);
 		});
 }
 
@@ -531,7 +531,7 @@ static bool SortMaterialInstanceParametersPredicate(
 	return GetTypeHashLegacy(Left.Info) < GetTypeHashLegacy(Right.Info);
 }
 
-void FMaterialInstanceResource::InitMIParameters(FMaterialInstanceParameterSet& ParameterSet)
+void FMaterialInstanceRenderProxy::InitMIParameters(FMaterialInstanceParameterSet& ParameterSet)
 {
 	InvalidateUniformExpressionCache(false);
 
@@ -570,7 +570,7 @@ void FMaterialInstanceResource::InitMIParameters(FMaterialInstanceParameterSet&
 	ParameterCollectionParameterArray.HashAddAllItems();
 }
 
-void FMaterialInstanceResource::EvaluateParameterCollections(FRHICommandListBase& RHICmdList, FUniformExpressionCache& OutUniformExpressionCache, const FMaterialRenderContext& Context, FUniformExpressionCacheAsyncUpdater* Updater) const
+void FMaterialInstanceRenderProxy::EvaluateParameterCollections(FRHICommandListBase& RHICmdList, FUniformExpressionCache& OutUniformExpressionCache, const FMaterialRenderContext& Context, FUniformExpressionCacheAsyncUpdater* Updater) const
 {
 	FMaterialRenderProxy::EvaluateParameterCollections(RHICmdList, OutUniformExpressionCache, Context, Updater);
 
@@ -618,14 +618,14 @@ void GameThread_UpdateMIParameter(const UMaterialInstance* Instance, const Param
 				GetRendererModule().InvalidatePathTracedOutput(PathTracing::EInvalidateReason::UpdateMaterialParameter);
 			}
 		}
-		FMaterialInstanceResource* Resource = Instance->Resource;
+		FMaterialInstanceRenderProxy* RenderProxy = Instance->GetInstanceRenderProxy();
 		const FMaterialParameterInfo& ParameterInfo = Parameter.ParameterInfo;
 		typename ParameterType::ValueType Value = ParameterType::GetValue(Parameter);
 		ENQUEUE_RENDER_COMMAND(SetMIParameterValue)(
-			[Resource, ParameterInfo, Value](FRHICommandListImmediate& RHICmdList)
+			[RenderProxy, ParameterInfo, Value](FRHICommandListImmediate& RHICmdList)
 			{
-				Resource->RenderThread_UpdateParameter(ParameterInfo, Value);
-				Resource->CacheUniformExpressions(RHICmdList, false);
+				RenderProxy->RenderThread_UpdateParameter(ParameterInfo, Value);
+				RenderProxy->CacheUniformExpressions(RHICmdList, false);
 			});
 	}
 }
@@ -903,7 +903,7 @@ void UMaterialInstance::PostInitProperties()
 
 	if(!HasAnyFlags(RF_ClassDefaultObject) && FApp::CanEverRenderOrProduceRenderData())
 	{
-		Resource = new FMaterialInstanceResource(this);
+		RenderProxy = new FMaterialInstanceRenderProxy(this);
 		bResourceCreated = true;
 	}
 }
@@ -918,7 +918,7 @@ void GameThread_InitMIParameters(const UMaterialInstance& Instance)
 		return;
 	}
 
-	FMaterialInstanceResource* Resource = Instance.Resource;
+	FMaterialInstanceRenderProxy* RenderProxy = Instance.GetInstanceRenderProxy();
 	FMaterialInstanceParameterSet ParameterSet;
 
 	// Scalar parameters
@@ -1018,9 +1018,9 @@ void GameThread_InitMIParameters(const UMaterialInstance& Instance)
 	ParameterSet.PostProcessBlendableOverrides.BlendablePriorityOverride = Instance.BlendablePriorityOverride;
 
 	ENQUEUE_RENDER_COMMAND(InitMIParameters)(
-		[Resource, Parameters = MoveTemp(ParameterSet)](FRHICommandListImmediate& RHICmdList) mutable
+		[RenderProxy, Parameters = MoveTemp(ParameterSet)](FRHICommandListImmediate& RHICmdList) mutable
 		{
-			Resource->InitMIParameters(Parameters);
+			RenderProxy->InitMIParameters(Parameters);
 		});
 }
 
@@ -1055,10 +1055,10 @@ void UMaterialInstance::InitResources()
 
 	// TODO - should merge all of render commands sent to initialize resource into a single command
 	// Set the material instance's parent on its resources.
-	if (Resource != nullptr)
+	if (RenderProxy != nullptr)
 	{
-		Resource->GameThread_SetParent(SafeParent);
-		Resource->GameThread_UpdateCachedData(GetCachedInstanceData());
+		RenderProxy->GameThread_SetParent(SafeParent);
+		RenderProxy->GameThread_UpdateCachedData(GetCachedInstanceData());
 	}
 
 	GameThread_InitMIParameters(*this);
@@ -2125,7 +2125,7 @@ const FMaterialResource* UMaterialInstance::GetMaterialResource(EShaderPlatform
 
 FMaterialRenderProxy* UMaterialInstance::GetRenderProxy() const
 {
-	return Resource;
+	return GetInstanceRenderProxy();
 }
 
 UPhysicalMaterial* UMaterialInstance::GetPhysicalMaterial() const
@@ -3801,15 +3801,15 @@ void UMaterialInstance::BeginDestroy()
 
 	Super::BeginDestroy();
 
-	if (Resource || ResourcesToDestroy.Num() > 0)
+	if (RenderProxy || ResourcesToDestroy.Num() > 0)
 	{
 		ENQUEUE_RENDER_COMMAND(BeginDestroyCommand)(
 		[ResourcesToDestroy = MoveTemp(ResourcesToDestroy), this](FRHICommandListImmediate& RHICmdList) mutable
 		{
-			if (Resource)
+			if (RenderProxy)
 			{
-				Resource->MarkForGarbageCollection();
-				Resource->ReleaseResource();
+				RenderProxy->MarkForGarbageCollection();
+				RenderProxy->ReleaseResource();
 			}
 
 			for (FMaterialResource* CurrentResource : ResourcesToDestroy)
@@ -3823,7 +3823,7 @@ void UMaterialInstance::BeginDestroy()
 			// And remove from deferred uniform expression cache queue if it's in that
 			if (bCachingUniformExpressions)
 			{
-				Resource->CancelCacheUniformExpressions();
+				RenderProxy->CancelCacheUniformExpressions();
 				bCachingUniformExpressions = false;
 			}
 
@@ -3844,10 +3844,10 @@ void UMaterialInstance::FinishDestroy()
 {
 	if(!HasAnyFlags(RF_ClassDefaultObject))
 	{
-		if (Resource)
+		if (RenderProxy)
 		{
-			Resource->GameThread_Destroy();
-			Resource = nullptr;
+			RenderProxy->GameThread_Destroy();
+			RenderProxy = nullptr;
 		}
 	}
 
@@ -4518,13 +4518,13 @@ void UMaterialInstance::ClearParameterValuesInternal(EMaterialInstanceClearParam
 	}
 
 
-	if (Resource && bUpdateResource)
+	if (RenderProxy && bUpdateResource)
 	{
-		FMaterialInstanceResource* InResource = Resource;
+		FMaterialInstanceRenderProxy* InRenderProxy = RenderProxy;
 		ENQUEUE_RENDER_COMMAND(FClearMIParametersCommand)(
-			[InResource](FRHICommandList& RHICmdList)
+			[InRenderProxy](FRHICommandList& RHICmdList)
 			{
-				InResource->RenderThread_ClearParameters();
+				InRenderProxy->RenderThread_ClearParameters();
 			});
 	}
 
@@ -4745,7 +4745,7 @@ void UMaterialInstance::PostEditChangeProperty(FPropertyChangedEvent& PropertyCh
 			// Brute force all flush virtual textures if this material writes to any runtime virtual texture.
 			if (WritesToRuntimeVirtualTexture())
 			{
-				ENQUEUE_RENDER_COMMAND(FlushVTCommand)([ResourcePtr = Resource](FRHICommandListImmediate& RHICmdList)
+				ENQUEUE_RENDER_COMMAND(FlushVTCommand)([](FRHICommandListImmediate& RHICmdList)
 				{
 					GetRendererModule().FlushVirtualTextureCache();
 				});
@@ -4924,9 +4924,9 @@ void UMaterialInstance::GetResourceSizeEx(FResourceSizeEx& CumulativeResourceSiz
 		}
 	}
 
-	if (Resource)
+	if (RenderProxy)
 	{
-		CumulativeResourceSize.AddDedicatedSystemMemoryBytes(sizeof(FMaterialInstanceResource));
+		CumulativeResourceSize.AddDedicatedSystemMemoryBytes(sizeof(FMaterialInstanceRenderProxy));
 		CumulativeResourceSize.AddDedicatedSystemMemoryBytes(ScalarParameterValues.Num() * sizeof(THashedMaterialParameterMap<float>::TNamedParameter));
 		CumulativeResourceSize.AddDedicatedSystemMemoryBytes(VectorParameterValues.Num() * sizeof(THashedMaterialParameterMap<FLinearColor>::TNamedParameter));
 		CumulativeResourceSize.AddDedicatedSystemMemoryBytes(DoubleVectorParameterValues.Num() * sizeof(THashedMaterialParameterMap<FVector4d>::TNamedParameter));
@@ -5985,14 +5985,14 @@ void UMaterialInstance::OverrideTextureParameterValue(const UTexture* InTextureT
 		// Make sure to update ourself since the parent changed
 		if (OldParentParameterOverrideCount != ParentInstance->TransientTextureParameterOverrides.Num())
 		{
-			FMaterialInstanceResource* LocalResource = Resource;
-			if (LocalResource)
+			FMaterialInstanceRenderProxy* LocalRenderProxy = RenderProxy;
+			if (LocalRenderProxy)
 			{
 				ENQUEUE_RENDER_COMMAND(RefreshMIParameterValue)
 					(
-						[LocalResource](FRHICommandListImmediate& RHICmdList)
+						[LocalRenderProxy](FRHICommandListImmediate& RHICmdList)
 						{
-							LocalResource->CacheUniformExpressions(RHICmdList, false);
+							LocalRenderProxy->CacheUniformExpressions(RHICmdList, false);
 						}
 				);
 			}
diff --git a/Engine/Source/Runtime/Engine/Private/Materials/MaterialInstanceConstant.cpp b/Engine/Source/Runtime/Engine/Private/Materials/MaterialInstanceConstant.cpp
--- a/Engine/Source/Runtime/Engine/Private/Materials/MaterialInstanceConstant.cpp
+++ b/Engine/Source/Runtime/Engine/Private/Materials/MaterialInstanceConstant.cpp
@@ -205,9 +205,9 @@ void UMaterialInstanceConstant::UpdateCachedData()
 		FMaterialLayersFunctions ParentLayers;
 		const bool bParentHasLayers = Parent && Parent->GetMaterialLayers(ParentLayers);
 		CachedData->InitializeForConstant(bHasLayers ? &Layers : nullptr, bParentHasLayers ? &ParentLayers : nullptr);
-		if (Resource)
+		if (RenderProxy)
 		{
-			Resource->GameThread_UpdateCachedData(*CachedData);
+			RenderProxy->GameThread_UpdateCachedData(*CachedData);
 		}
 	}
 
diff --git a/Engine/Source/Runtime/Engine/Private/Materials/MaterialInstanceDynamic.cpp b/Engine/Source/Runtime/Engine/Private/Materials/MaterialInstanceDynamic.cpp
--- a/Engine/Source/Runtime/Engine/Private/Materials/MaterialInstanceDynamic.cpp
+++ b/Engine/Source/Runtime/Engine/Private/Materials/MaterialInstanceDynamic.cpp
@@ -46,9 +46,9 @@ void UMaterialInstanceDynamic::UpdateCachedDataDynamic()
 	}
 	CachedData->InitializeForDynamic(bParentHasLayers ? &ParentLayers : nullptr);
 
-	if (Resource)
+	if (RenderProxy)
 	{
-		Resource->GameThread_UpdateCachedData(*CachedData);
+		RenderProxy->GameThread_UpdateCachedData(*CachedData);
 	}
 
 #if WITH_EDITOR
diff --git a/Engine/Source/Runtime/Engine/Private/Materials/MaterialInstanceSupport.h b/Engine/Source/Runtime/Engine/Private/Materials/MaterialInstanceSupport.h
--- a/Engine/Source/Runtime/Engine/Private/Materials/MaterialInstanceSupport.h
+++ b/Engine/Source/Runtime/Engine/Private/Materials/MaterialInstanceSupport.h
@@ -204,11 +204,11 @@ public:
 /**
 * The resource used to render a UMaterialInstance.
 */
-class FMaterialInstanceResource: public FMaterialRenderProxy
+class FMaterialInstanceRenderProxy: public FMaterialRenderProxy
 {
 public:
 	/** Initialization constructor. */
-	FMaterialInstanceResource(UMaterialInstance* InOwner);
+	FMaterialInstanceRenderProxy(UMaterialInstance* InOwner);
 
 	/**
 	 * Called from the game thread to destroy the material instance on the rendering thread.
@@ -426,24 +426,24 @@ private:
 	FPostProcessBlendableOverrides PostProcessBlendableOverrides;
 };
 
-template <> FORCEINLINE THashedMaterialParameterMap<bool>& FMaterialInstanceResource::GetValueArray() { return StaticSwitchParameterArray; }
-template <> FORCEINLINE THashedMaterialParameterMap<float>& FMaterialInstanceResource::GetValueArray() { return ScalarParameterArray; }
-template <> FORCEINLINE THashedMaterialParameterMap<FLinearColor>& FMaterialInstanceResource::GetValueArray() { return VectorParameterArray; }
-template <> FORCEINLINE THashedMaterialParameterMap<FVector4d>& FMaterialInstanceResource::GetValueArray() { return DoubleVectorParameterArray; }
-template <> FORCEINLINE THashedMaterialParameterMap<const UTexture*>& FMaterialInstanceResource::GetValueArray() { return TextureParameterArray; }
-template <> FORCEINLINE THashedMaterialParameterMap<const UTextureCollection*>& FMaterialInstanceResource::GetValueArray() { return TextureCollectionParameterArray; }
-template <> FORCEINLINE THashedMaterialParameterMap<const URuntimeVirtualTexture*>& FMaterialInstanceResource::GetValueArray() { return RuntimeVirtualTextureParameterArray; }
-template <> FORCEINLINE THashedMaterialParameterMap<const USparseVolumeTexture*>& FMaterialInstanceResource::GetValueArray() { return SparseVolumeTextureParameterArray; }
-template <> FORCEINLINE THashedMaterialParameterMap<const UMaterialParameterCollection*>& FMaterialInstanceResource::GetValueArray() { return ParameterCollectionParameterArray; }
-template <> FORCEINLINE const THashedMaterialParameterMap<bool>& FMaterialInstanceResource::GetValueArray() const { return StaticSwitchParameterArray; }
-template <> FORCEINLINE const THashedMaterialParameterMap<float>& FMaterialInstanceResource::GetValueArray() const { return ScalarParameterArray; }
-template <> FORCEINLINE const THashedMaterialParameterMap<FLinearColor>& FMaterialInstanceResource::GetValueArray() const { return VectorParameterArray; }
-template <> FORCEINLINE const THashedMaterialParameterMap<FVector4d>& FMaterialInstanceResource::GetValueArray() const { return DoubleVectorParameterArray; }
-template <> FORCEINLINE const THashedMaterialParameterMap<const UTexture*>& FMaterialInstanceResource::GetValueArray() const { return TextureParameterArray; }
-template <> FORCEINLINE const THashedMaterialParameterMap<const UTextureCollection*>& FMaterialInstanceResource::GetValueArray() const { return TextureCollectionParameterArray; }
-template <> FORCEINLINE const THashedMaterialParameterMap<const URuntimeVirtualTexture*>& FMaterialInstanceResource::GetValueArray() const { return RuntimeVirtualTextureParameterArray; }
-template <> FORCEINLINE const THashedMaterialParameterMap<const USparseVolumeTexture*>& FMaterialInstanceResource::GetValueArray() const { return SparseVolumeTextureParameterArray; }
-template <> FORCEINLINE const THashedMaterialParameterMap<const UMaterialParameterCollection*>& FMaterialInstanceResource::GetValueArray() const { return ParameterCollectionParameterArray; }
+template <> FORCEINLINE THashedMaterialParameterMap<bool>& FMaterialInstanceRenderProxy::GetValueArray() { return StaticSwitchParameterArray; }
+template <> FORCEINLINE THashedMaterialParameterMap<float>& FMaterialInstanceRenderProxy::GetValueArray() { return ScalarParameterArray; }
+template <> FORCEINLINE THashedMaterialParameterMap<FLinearColor>& FMaterialInstanceRenderProxy::GetValueArray() { return VectorParameterArray; }
+template <> FORCEINLINE THashedMaterialParameterMap<FVector4d>& FMaterialInstanceRenderProxy::GetValueArray() { return DoubleVectorParameterArray; }
+template <> FORCEINLINE THashedMaterialParameterMap<const UTexture*>& FMaterialInstanceRenderProxy::GetValueArray() { return TextureParameterArray; }
+template <> FORCEINLINE THashedMaterialParameterMap<const UTextureCollection*>& FMaterialInstanceRenderProxy::GetValueArray() { return TextureCollectionParameterArray; }
+template <> FORCEINLINE THashedMaterialParameterMap<const URuntimeVirtualTexture*>& FMaterialInstanceRenderProxy::GetValueArray() { return RuntimeVirtualTextureParameterArray; }
+template <> FORCEINLINE THashedMaterialParameterMap<const USparseVolumeTexture*>& FMaterialInstanceRenderProxy::GetValueArray() { return SparseVolumeTextureParameterArray; }
+template <> FORCEINLINE THashedMaterialParameterMap<const UMaterialParameterCollection*>& FMaterialInstanceRenderProxy::GetValueArray() { return ParameterCollectionParameterArray; }
+template <> FORCEINLINE const THashedMaterialParameterMap<bool>& FMaterialInstanceRenderProxy::GetValueArray() const { return StaticSwitchParameterArray; }
+template <> FORCEINLINE const THashedMaterialParameterMap<float>& FMaterialInstanceRenderProxy::GetValueArray() const { return ScalarParameterArray; }
+template <> FORCEINLINE const THashedMaterialParameterMap<FLinearColor>& FMaterialInstanceRenderProxy::GetValueArray() const { return VectorParameterArray; }
+template <> FORCEINLINE const THashedMaterialParameterMap<FVector4d>& FMaterialInstanceRenderProxy::GetValueArray() const { return DoubleVectorParameterArray; }
+template <> FORCEINLINE const THashedMaterialParameterMap<const UTexture*>& FMaterialInstanceRenderProxy::GetValueArray() const { return TextureParameterArray; }
+template <> FORCEINLINE const THashedMaterialParameterMap<const UTextureCollection*>& FMaterialInstanceRenderProxy::GetValueArray() const { return TextureCollectionParameterArray; }
+template <> FORCEINLINE const THashedMaterialParameterMap<const URuntimeVirtualTexture*>& FMaterialInstanceRenderProxy::GetValueArray() const { return RuntimeVirtualTextureParameterArray; }
+template <> FORCEINLINE const THashedMaterialParameterMap<const USparseVolumeTexture*>& FMaterialInstanceRenderProxy::GetValueArray() const { return SparseVolumeTextureParameterArray; }
+template <> FORCEINLINE const THashedMaterialParameterMap<const UMaterialParameterCollection*>& FMaterialInstanceRenderProxy::GetValueArray() const { return ParameterCollectionParameterArray; }
 
 struct FMaterialInstanceParameterSet
 {
diff --git a/Engine/Source/Runtime/Engine/Public/Materials/MaterialInstance.h b/Engine/Source/Runtime/Engine/Public/Materials/MaterialInstance.h
--- a/Engine/Source/Runtime/Engine/Public/Materials/MaterialInstance.h
+++ b/Engine/Source/Runtime/Engine/Public/Materials/MaterialInstance.h
@@ -824,10 +824,8 @@ public:
 	bool bDisallowStaticParameterPermutations;
 #endif
 
-	/** 
-	 * FMaterialRenderProxy derivative that represent this material instance to the renderer, when the renderer needs to fetch parameter values. 
-	 */
-	class FMaterialInstanceResource* Resource;
+	UE_DEPRECATED(5.8, "Exposed public Instance Render Proxy derivative has been moved to protected. Please use either GetRenderProxy or GetInstanceRenderProxy.")
+	class FMaterialInstanceRenderProxy* Resource;
 
 #if WITH_EDITOR
 	ENGINE_API virtual void OnCookEvent(UE::Cook::ECookEvent CookEvent, UE::Cook::FCookEventContext& CookContext) override;
@@ -861,6 +859,11 @@ private:
 #endif // WITH_EDITORONLY_DATA
 
 protected:
+	/**
+	 * FMaterialRenderProxy derivative that represent this material instance to the renderer, when the renderer needs to fetch parameter values.
+	 */
+	class FMaterialInstanceRenderProxy* RenderProxy;
+
 	TUniquePtr<FMaterialInstanceCachedData> CachedData;
 
 private:
@@ -925,7 +928,10 @@ public:
 	virtual ENGINE_API bool IsDependent(UMaterialInterface* TestDependency) override;
 	virtual ENGINE_API bool IsDependent_Concurrent(UMaterialInterface* TestDependency, TMicRecursionGuard RecursionGuard) override;
 	virtual ENGINE_API void GetDependencies(TSet<UMaterialInterface*>& Dependencies) override;
+
 	virtual ENGINE_API FMaterialRenderProxy* GetRenderProxy() const override;
+	inline FMaterialInstanceRenderProxy* GetInstanceRenderProxy() const { return RenderProxy; }
+
 	virtual ENGINE_API UPhysicalMaterial* GetPhysicalMaterial() const override;
 	virtual ENGINE_API UPhysicalMaterialMask* GetPhysicalMaterialMask() const override;
 	virtual ENGINE_API UPhysicalMaterial* GetPhysicalMaterialFromMap(int32 Index) const override;
@@ -1281,7 +1287,7 @@ protected:
 	void PropagateDataToMaterialProxy();
 
 	/** Allow resource to access private members. */
-	friend class FMaterialInstanceResource;
+	friend class FMaterialInstanceRenderProxy;
 	/** Editor-only access to private members. */
 	friend class UMaterialEditingLibrary;
 	/** Class that knows how to update MI's */
