Fixed issues (asserts and invalid behavior) when loading a non-spatially loaded level instance containing property overrides.

#Synced-CL 48932408

diff --git a/Engine/Source/Runtime/Engine/Private/WorldPartition/WorldPartitionLevelStreamingDynamic.cpp b/Engine/Source/Runtime/Engine/Private/WorldPartition/WorldPartitionLevelStreamingDynamic.cpp
--- a/Engine/Source/Runtime/Engine/Private/WorldPartition/WorldPartitionLevelStreamingDynamic.cpp
+++ b/Engine/Source/Runtime/Engine/Private/WorldPartition/WorldPartitionLevelStreamingDynamic.cpp
@@ -332,10 +332,23 @@ bool UWorldPartitionLevelStreamingDynamic::RequestLevel(UWorld* InPersistentWorl
 						UE_LOG(LogStreaming, Display, TEXT("UWorldPartitionLevelStreamingDynamic::RequestLevel(%s) is flushing async loading"), *WorldAssetPackageName);
 					}
 
-					// Finish all async loading. Since we will clear our requests upon completion of all loads, 
-					// we take a copy so FlushAsyncLoading won't touch an invalidated array view
-					TArray<int32> LocalRequestIDs = AsyncRequestIDs;
-					FlushAsyncLoading(LocalRequestIDs);
+					// Finish all async loading
+					// Note that FlushAsyncLoading will invoke completion callbacks, which may themselves trigger new async loading 
+					// requests (FWorldPartitionLevelHelper::LoadActorsWithPropertyOverridesInternal is a good example).
+					// Loop until no new requests are found in AsyncRequestIDs.
+					TSet<int32> PrevRequestIDs;
+					TSet<int32> CurrRequestIDs;
+					TSet<int32> DiffRequestIDs;
+					do
+					{
+						CurrRequestIDs = TSet<int32>(AsyncRequestIDs);
+						DiffRequestIDs = CurrRequestIDs.Difference(PrevRequestIDs);
+						if (DiffRequestIDs.Num() > 0)
+						{
+							FlushAsyncLoading(DiffRequestIDs.Array());
+							PrevRequestIDs = MoveTemp(CurrRequestIDs);
+						}
+					} while (DiffRequestIDs.Num() > 0);
 				}
 				else
 				{
