Fixed oob access in weighted random sampler.
From Github PR
#14093

#Synced-CL 48842817

diff --git a/Engine/Source/Runtime/Engine/Public/WeightedRandomSampler.h b/Engine/Source/Runtime/Engine/Public/WeightedRandomSampler.h
--- a/Engine/Source/Runtime/Engine/Public/WeightedRandomSampler.h
+++ b/Engine/Source/Runtime/Engine/Public/WeightedRandomSampler.h
@@ -35,11 +35,16 @@ public:
 
 	/**
 	Takes two random values (0...1) and returns the corresponding element index.
+	Returns INDEX_NONE if there are no entries.
 	*/
 	inline int32 GetEntryIndex(float R0, float R1)const
 	{
-		int32 Idx = (int32)(R0 * float(Prob.Num()));
-		return R1 < Prob[Idx] ? Idx : Alias[Idx];
+		if(Prob.Num() > 0)
+		{			
+			int32 Idx = FMath::Clamp((int32)(R0 * float(Prob.Num())), 0, Prob.Num() - 1);
+			return R1 < Prob[Idx] ? Idx : Alias[Idx];
+		}
+		return INDEX_NONE;
 	}
 
 	inline float GetTotalWeight()const { return TotalWeight; }
