Entity Streaming:
- Streaming component now has an optional enumeration to override the streaming policy, which can be set to 'atomic' or 'granular' to control how children are streamed for the entity.
- Added the option to enable or disable streaming support on the level entity, through the scene outliner world item right-click menu. This essentially add or remove a default streaming_component on the level entity.
- Fixed LHGrid graph tree processor to avoid merging root nodes with child nodes, since they are part of the level.
- Always create the entity streaming subsystem when Entity.Streaming.Enable is true.

#Synced-CL 48998606

diff --git a/Engine/Source/Runtime/Engine/Classes/Interfaces/Interface_AssetUserData.h b/Engine/Source/Runtime/Engine/Classes/Interfaces/Interface_AssetUserData.h
--- a/Engine/Source/Runtime/Engine/Classes/Interfaces/Interface_AssetUserData.h
+++ b/Engine/Source/Runtime/Engine/Classes/Interfaces/Interface_AssetUserData.h
@@ -69,6 +69,12 @@ class IInterface_AssetUserData
 		return CastChecked<T>(GetAssetUserDataOfClass(T::StaticClass()));
 	}
 
+	template<typename T>
+	bool HasAssetUserData()
+	{
+		return HasAssetUserDataOfClass(T::StaticClass());
+	}
+
 	ENGINE_API virtual void RemoveUserDataOfClass(TSubclassOf<UAssetUserData> InUserDataClass);
 
 };
