Don't overide the static mesh usage flag in editor.
This allows us to immediately see the visual result of turning the usage off which we might want to do to save on shader permutations.
The existing (readonly) cvar "r.Material.IgnoreLVFUsageMode" can still be used to force static mesh usage on.
#Synced-CL 49028617

diff --git a/Engine/Source/Runtime/Engine/Private/LocalVertexFactory.cpp b/Engine/Source/Runtime/Engine/Private/LocalVertexFactory.cpp
--- a/Engine/Source/Runtime/Engine/Private/LocalVertexFactory.cpp
+++ b/Engine/Source/Runtime/Engine/Private/LocalVertexFactory.cpp
@@ -255,8 +255,8 @@ void FLocalVertexFactoryShaderParameters::GetElementShaderBindings(
 }
 
 // CVarIgnoreLVFUsageMode settings:
-//     0: ignored only in editor (Default)
-//     1: ignored in editor and on PC >= SM5
+//     0: never ignore (Default)
+//     1: ignored on PC >= SM5
 //     2: ignored always
 TAutoConsoleVariable<int32> CVarIgnoreLVFUsageMode(
 	TEXT("r.Material.IgnoreLVFUsageMode"), 0,
@@ -308,9 +308,8 @@ bool DoesPlatformUseInstancedStaticMeshVertexFactory(const FStaticShaderPlatform
  */
 bool FLocalVertexFactory::ShouldCompilePermutation(const FVertexFactoryShaderPermutationParameters& Parameters)
 {
-	// Compile local vertex factory permutations while in the editor, always. Important for editor features like preview 
-	// Additionally, allow ignoring via CVar to change what scenarios respect the bIsUsedWithStaticMesh flag
-	if (EnumHasAnyFlags(Parameters.Flags, EShaderPermutationFlags::HasEditorOnlyData) || ShouldIgnore(Parameters.Platform))
+	// Allow config forcing of compilation without StaticMesh usage being set.
+	if (ShouldIgnore(Parameters.Platform))
 	{
 		return true;
 	}
@@ -321,7 +320,13 @@ bool FLocalVertexFactory::ShouldCompilePermutation(const FVertexFactoryShaderPer
 		return true;
 	}
 
-	// Otherwise compile when the Material has selected that it can be used with StaticMesh
+	// We use the local vertex factory in instanced mode for materials that select InstancedStaticMesh usage when that is supported.
+	if (Parameters.MaterialParameters.bIsUsedWithInstancedStaticMeshes && DoesPlatformUseInstancedStaticMeshVertexFactory(Parameters.Platform))
+	{
+		return true;
+	}
+
+	// Otherwise compile when the material has selected StaticMesh usage.
 	return Parameters.MaterialParameters.bIsUsedWithStaticMesh == 1;
 }
 
