[UserDefinedEnum]
- Move thread-safe postload code into serialize so the upgrades are ready faster and atomic inside LinkerLoad::Preload
- This fixes a race condition where postload could modify the enum while something else was trying to read a value from the enum

#Synced-CL 48817740

diff --git a/Engine/Source/Runtime/Engine/Private/UserDefinedEnum.cpp b/Engine/Source/Runtime/Engine/Private/UserDefinedEnum.cpp
--- a/Engine/Source/Runtime/Engine/Private/UserDefinedEnum.cpp
+++ b/Engine/Source/Runtime/Engine/Private/UserDefinedEnum.cpp
@@ -29,14 +29,28 @@ void UUserDefinedEnum::Serialize(FArchive& Ar)
 	Ar.UsingCustomVersion(FEditorObjectVersion::GUID);
 
 #if WITH_EDITOR
-	if (Ar.IsLoading() && Ar.IsPersistent())
+	if (Ar.IsLoading())
 	{
-		Names.MakeValuesSequentialFromZero();
+		if (Ar.IsPersistent())
+		{
+			Names.MakeValuesSequentialFromZero();
+
+			if (Ar.CustomVer(FEditorObjectVersion::GUID) < FEditorObjectVersion::StableUserDefinedEnumDisplayNames)
+			{
+				// Make sure DisplayNameMap is empty so we perform the meta-data upgrade in PostLoad
+				DisplayNameMap.Reset();
+			}
+		}
 
-		if (Ar.CustomVer(FEditorObjectVersion::GUID) < FEditorObjectVersion::StableUserDefinedEnumDisplayNames)
+		// Apply any transformations right away so that the enum is consistent at the end of its serialization during loading.
+		if (!Ar.IsObjectReferenceCollector() && !Ar.IsCountingMemory() && !Ar.IsTransacting())
 		{
-			// Make sure DisplayNameMap is empty so we perform the meta-data upgrade in PostLoad
-			DisplayNameMap.Reset();
+			FEnumEditorUtils::UpdateAfterPathChanged(this);
+			if (NumEnums() > 1 && DisplayNameMap.Num() == 0) // >1 because User Defined Enums always have a "MAX" entry
+			{
+				FEnumEditorUtils::UpgradeDisplayNamesFromMetaData(this);
+			}
+			FEnumEditorUtils::EnsureAllDisplayNamesExist(this);
 		}
 	}
 #endif // WITH_EDITOR
@@ -87,13 +101,6 @@ void UUserDefinedEnum::PostLoad()
 {
 	Super::PostLoad();
 
-	FEnumEditorUtils::UpdateAfterPathChanged(this);
-	if (NumEnums() > 1 && DisplayNameMap.Num() == 0) // >1 because User Defined Enums always have a "MAX" entry
-	{
-		FEnumEditorUtils::UpgradeDisplayNamesFromMetaData(this);
-	}
-	FEnumEditorUtils::EnsureAllDisplayNamesExist(this);
-
 	// Apply the transactional flag to user defined enums that were not created with it
 	SetFlags(RF_Transactional);
 }
