Collision Filter Refactor part 5.5.1: Split of part of the final 64 bit upgrade. This change does:
1. Updates the legacy Get GetQuery/SimData to return a copy instead of const reference (needed as the data won't be a member anymore)
2. Updates some hardcoded spots referencing 32 channels or Channel18 so they're easier to update to 64 later
3. Updates the API to a few classes in CollisionQueryParams to be 64bit.
#Synced-CL 49038256

diff --git a/Engine/Plugins/ChaosVD/Source/ChaosVD/Private/DetailsCustomizations/ChaosVDQueryDataWrappersCustomizationDetails.cpp b/Engine/Plugins/ChaosVD/Source/ChaosVD/Private/DetailsCustomizations/ChaosVDQueryDataWrappersCustomizationDetails.cpp
--- a/Engine/Plugins/ChaosVD/Source/ChaosVD/Private/DetailsCustomizations/ChaosVDQueryDataWrappersCustomizationDetails.cpp
+++ b/Engine/Plugins/ChaosVD/Source/ChaosVD/Private/DetailsCustomizations/ChaosVDQueryDataWrappersCustomizationDetails.cpp
@@ -391,7 +391,7 @@ void FChaosVDCollisionObjectParamsCustomization::CustomizeChildren(TSharedRef<IP
 		.NameContent()
 		[
 			SNew(STextBlock)
-			.Text(FText::FromString(CachedCollisionChannelInfos->CustomChannelsNames[ChannelIndex].DisplayName))
+			.Text(FText::FromString(ChannelInfo.DisplayName))
 			.Font(IDetailLayoutBuilder::GetDetailFont())
 		]
 		.ValueContent()
diff --git a/Engine/Plugins/Experimental/ControlRigPhysics/Source/ControlRigPhysics/Private/RigPhysicsSimulationWorldObjects.cpp b/Engine/Plugins/Experimental/ControlRigPhysics/Source/ControlRigPhysics/Private/RigPhysicsSimulationWorldObjects.cpp
--- a/Engine/Plugins/Experimental/ControlRigPhysics/Source/ControlRigPhysics/Private/RigPhysicsSimulationWorldObjects.cpp
+++ b/Engine/Plugins/Experimental/ControlRigPhysics/Source/ControlRigPhysics/Private/RigPhysicsSimulationWorldObjects.cpp
@@ -159,7 +159,7 @@ void FRigPhysicsSimulation::UpdateWorldObjectsPostPhysics(
 	// The actual overlap needs to be on the game thread, so we will run it in an async task.
 	// However, note that between now and then, we/our simulation etc might get destroyed so we must
 	// protect against that. Even the mutex guarding against access to the simulation may be released.
-	if (ObjectQueryParams.ObjectTypesToQuery)
+	if (ObjectQueryParams.GetObjectTypesToQuery())
 	{
 		AsyncTask(ENamedThreads::GameThread,
 			[WeakSimulationMutex, WeakWorld, WeakWorldObjects, OverlapAABox, OverlapAABoxCenter, QueryParams, 
diff --git a/Engine/Source/Editor/DetailCustomizations/Private/CollisionProfileDetails.cpp b/Engine/Source/Editor/DetailCustomizations/Private/CollisionProfileDetails.cpp
--- a/Engine/Source/Editor/DetailCustomizations/Private/CollisionProfileDetails.cpp
+++ b/Engine/Source/Editor/DetailCustomizations/Private/CollisionProfileDetails.cpp
@@ -59,8 +59,11 @@ class SWidget;
 DECLARE_DELEGATE_RetVal_OneParam(bool, FOnValidateChannel, const FCustomChannelSetup* )
 DECLARE_DELEGATE_RetVal_TwoParams(bool, FOnValidateProfile, const FCollisionResponseTemplate*, int32)
 
-#define MAX_CUSTOMCOLLISION_CHANNEL (ECC_GameTraceChannel18-ECC_GameTraceChannel1+1)
-#define MAX_COLLISION_CHANNEL		32
+constexpr ECollisionChannel MinCustomCollisionChannel = ECC_GameTraceChannel1;
+constexpr ECollisionChannel MaxCustomCollisionChannel = ECC_GameTraceChannel18;
+constexpr uint32 MaxNumCustomCollisionChannels = MaxCustomCollisionChannel - MinCustomCollisionChannel + 1;
+constexpr uint32 MaxNumCollisionChannels = ECC_OverlapAll_Deprecated;
+static_assert(MaxNumCollisionChannels == 32);
 
 #define COLLIISION_COLUMN_WIDTH		50.f
 
@@ -142,7 +145,7 @@ void SChannelEditDialog::Construct(const FArguments& InArgs)
 		ChannelSetup.bTraceType = InArgs._bTraceType;
 	}
 
-	check(ChannelSetup.Channel >= ECC_GameTraceChannel1 && ChannelSetup.Channel <= ECC_GameTraceChannel18);
+	check(ChannelSetup.Channel >= MinCustomCollisionChannel && ChannelSetup.Channel <= MaxCustomCollisionChannel);
 
 	OnValidateChannel = InArgs._OnValidateChannel;
 	WidgetWindow = InArgs._WidgetWindow;
@@ -1140,7 +1143,7 @@ void SProfileEditDialog::AddCollisionChannel(TArray<FCollisionChannelInfo>	Valid
 
 void SProfileEditDialog::OnCollisionChannelChanged(ECheckBoxState InNewValue, int32 ValidIndex, ECollisionResponse InCollisionResponse)
 {
-	if(ValidIndex >= 0 && ValidIndex < MAX_COLLISION_CHANNEL)
+	if(ValidIndex >= 0 && ValidIndex < MaxNumCollisionChannels)
 	{
 		ProfileTemplate.ResponseToChannels.EnumArray[ValidIndex] = static_cast<uint8>(InCollisionResponse);
 	}
@@ -1150,7 +1153,7 @@ ECheckBoxState SProfileEditDialog::IsCollisionChannelChecked(int32 ValidIndex, E
 {
 	TArray<uint8> CollisionResponses;
 
-	if( ValidIndex >= 0 && ValidIndex < MAX_COLLISION_CHANNEL )
+	if( ValidIndex >= 0 && ValidIndex < MaxNumCollisionChannels )
 	{
 		if(ProfileTemplate.ResponseToChannels.EnumArray[ValidIndex] == static_cast<uint8>(InCollisionResponse))
 		{
@@ -1163,7 +1166,7 @@ ECheckBoxState SProfileEditDialog::IsCollisionChannelChecked(int32 ValidIndex, E
 
 void SProfileEditDialog::OnAllCollisionChannelChanged(ECheckBoxState InNewValue, ECollisionResponse InCollisionResponse)
 {
-	for(int32 Index=0; Index<MAX_COLLISION_CHANNEL; ++Index)
+	for(int32 Index=0; Index<MaxNumCollisionChannels; ++Index)
 	{
 		ProfileTemplate.ResponseToChannels.EnumArray[Index] = static_cast<uint8>(InCollisionResponse);
 	}
@@ -1171,7 +1174,7 @@ void SProfileEditDialog::OnAllCollisionChannelChanged(ECheckBoxState InNewValue,
 
 ECheckBoxState SProfileEditDialog::IsAllCollisionChannelChecked(ECollisionResponse InCollisionResponse) const
 {
-	for(int32 Index=0; Index<MAX_COLLISION_CHANNEL; ++Index)
+	for(int32 Index=0; Index<MaxNumCollisionChannels; ++Index)
 	{
 		if(ProfileTemplate.ResponseToChannels.EnumArray[Index] != static_cast<uint8>(InCollisionResponse))
 		{
@@ -1756,7 +1759,7 @@ void FCollisionProfileDetails::CommitProfileChange(int32 ProfileIndex, FCollisio
 			
 		};
 		
-		for (int32 Index = 0; Index<MAX_COLLISION_CHANNEL; ++Index)
+		for (int32 Index = 0; Index<MaxNumCollisionChannels; ++Index)
 		{
 			if (NewProfile.ResponseToChannels.EnumArray[Index] != SavedProfile.ResponseToChannels.EnumArray[Index])
 			{
@@ -1839,7 +1842,7 @@ void FCollisionProfileDetails::RefreshChannelList(bool bTraceType)
 		for(auto Iter = CollisionProfile->DefaultChannelResponses.CreateIterator(); Iter; ++Iter)
 		{
 			// only display game channels
-			if(Iter->Channel >= ECC_GameTraceChannel1 && Iter->bTraceType)
+			if(Iter->Channel >= MinCustomCollisionChannel && Iter->bTraceType)
 			{
 				TraceChannelList.Add(MakeShareable(new FChannelListItem(MakeShareable(new FCustomChannelSetup(*Iter)))));
 			}
@@ -1852,7 +1855,7 @@ void FCollisionProfileDetails::RefreshChannelList(bool bTraceType)
 		for(auto Iter = CollisionProfile->DefaultChannelResponses.CreateIterator(); Iter; ++Iter)
 		{
 			// only display game channels
-			if(Iter->Channel >= ECC_GameTraceChannel1 && !Iter->bTraceType)
+			if(Iter->Channel >= MinCustomCollisionChannel && !Iter->bTraceType)
 			{
 				ObjectChannelList.Add(MakeShareable(new FChannelListItem(MakeShareable(new FCustomChannelSetup(*Iter)))));
 			}
@@ -1922,10 +1925,10 @@ FCustomChannelSetup * FCollisionProfileDetails::FindFromChannel(ECollisionChanne
 
 ECollisionChannel	FCollisionProfileDetails::FindAvailableChannel() const
 {
-	if(CollisionProfile->DefaultChannelResponses.Num() < MAX_CUSTOMCOLLISION_CHANNEL)
+	if(CollisionProfile->DefaultChannelResponses.Num() < MaxNumCustomCollisionChannels)
 	{
 		// this is very inefficient
-		for(int32 ChannelIndex = ECC_GameTraceChannel1; ChannelIndex < ECC_GameTraceChannel1 + MAX_CUSTOMCOLLISION_CHANNEL; ++ChannelIndex)
+		for(int32 ChannelIndex = MinCustomCollisionChannel; ChannelIndex < MinCustomCollisionChannel + MaxNumCustomCollisionChannels; ++ChannelIndex)
 		{
 			if( FindFromChannel((ECollisionChannel)ChannelIndex) == NULL)
 			{
@@ -1972,7 +1975,7 @@ bool	FCollisionProfileDetails::IsValidProfileSetup(const FCollisionResponseTempl
 }
 bool	FCollisionProfileDetails::IsNewChannelAvailable() const
 {
-	return (CollisionProfile && CollisionProfile->DefaultChannelResponses.Num() < MAX_CUSTOMCOLLISION_CHANNEL);
+	return (CollisionProfile && CollisionProfile->DefaultChannelResponses.Num() < MaxNumCustomCollisionChannels);
 }
 
 FReply FCollisionProfileDetails::OnNewChannel(bool bTraceType)
@@ -1980,7 +1983,7 @@ FReply FCollisionProfileDetails::OnNewChannel(bool bTraceType)
 	// find empty channel and see if we can add it. 
 	ECollisionChannel NewChannel = FindAvailableChannel();
 
-	if (ensure(NewChannel >= ECC_GameTraceChannel1 && NewChannel <= ECC_GameTraceChannel18)) //-V547
+	if (ensure(NewChannel >= MinCustomCollisionChannel && NewChannel <= MaxCustomCollisionChannel)) //-V547
 	{
 		// Create modal window for modification
 		TSharedRef<SWindow> WidgetWindow = SNew(SWindow)
diff --git a/Engine/Source/Runtime/Engine/Private/Collision/CollisionProfile.cpp b/Engine/Source/Runtime/Engine/Private/Collision/CollisionProfile.cpp
--- a/Engine/Source/Runtime/Engine/Private/Collision/CollisionProfile.cpp
+++ b/Engine/Source/Runtime/Engine/Private/Collision/CollisionProfile.cpp
@@ -10,8 +10,7 @@
 
 DEFINE_LOG_CATEGORY_STATIC(LogCollisionProfile, Warning, All)
 
-#define MIN_CUSTOMIZABLE_COLLISIONCHANNEL	ECC_GameTraceChannel1
-#define MAX_CUSTOMIZABLE_COLLISIONCHANNEL	ECC_GameTraceChannel18
+constexpr uint32 MaxNumCollisionChannels = 32;
 #define IS_VALID_COLLISIONCHANNEL(x) ((x)> ECC_Destructible && (x) < ECC_OverlapAll_Deprecated)
 
 // do not chnage this name. This value is serialize to other object, if you change, it will mess up serialization, you'll need to fix up name by versioning
@@ -714,7 +713,7 @@ void UCollisionProfile::SaveCustomResponses(FCollisionResponseTemplate& Template
 	const FCollisionResponseContainer& DefaultContainer = FCollisionResponseContainer::GetDefaultResponseContainer();
 
 	Template.CustomResponses.Empty();
-	for(int32 Index=0; Index<32; ++Index)
+	for(int32 Index=0; Index<MaxNumCollisionChannels; ++Index)
 	{
 		//Save responses different from default
 		if(Template.ResponseToChannels.EnumArray[Index] != DefaultContainer.EnumArray[Index])
diff --git a/Engine/Source/Runtime/Engine/Private/Collision/PhysicsFiltering.cpp b/Engine/Source/Runtime/Engine/Private/Collision/PhysicsFiltering.cpp
--- a/Engine/Source/Runtime/Engine/Private/Collision/PhysicsFiltering.cpp
+++ b/Engine/Source/Runtime/Engine/Private/Collision/PhysicsFiltering.cpp
@@ -7,16 +7,16 @@ void ExtractResponses(const FCollisionResponseContainer& InResponses, uint64& Bl
 {
 	BlockingBits = 0;
 	TouchingBits = 0;
-	for (int32 i = 0; i < UE_ARRAY_COUNT(InResponses.EnumArray); i++)
+	constexpr uint64 Count = UE_ARRAY_COUNT(InResponses.EnumArray);
+	for (uint64 i = 0; i < Count; ++i)
 	{
+		const uint64 ChannelBit = (1LLU << i);
 		if (InResponses.EnumArray[i] == ECR_Block)
 		{
-			const uint32 ChannelBit = (1LLU << i);
 			BlockingBits |= ChannelBit;
 		}
 		else if (InResponses.EnumArray[i] == ECR_Overlap)
 		{
-			const uint32 ChannelBit = (1LLU << i);
 			TouchingBits |= ChannelBit;
 		}
 	}
diff --git a/Engine/Source/Runtime/Engine/Private/PhysicsEngine/PhysicsInterfaceUtils.cpp b/Engine/Source/Runtime/Engine/Private/PhysicsEngine/PhysicsInterfaceUtils.cpp
--- a/Engine/Source/Runtime/Engine/Private/PhysicsEngine/PhysicsInterfaceUtils.cpp
+++ b/Engine/Source/Runtime/Engine/Private/PhysicsEngine/PhysicsInterfaceUtils.cpp
@@ -17,7 +17,7 @@ Chaos::Filter::FQueryFilterData CreateChaosQueryFilterData(const uint8 MyChannel
 	if (ObjectParam.IsValid())
 	{
 		FPhysicsObjectQueryFilterBuilder Builder;
-		Builder.SetObjectTypes(ObjectParam.GetQueryBitfield());
+		Builder.SetObjectTypes(ObjectParam.GetQueryBitfield64());
 		Builder.SetMultiQuery(bMultitrace);
 		Builder.SetMaskFilter(ObjectParam.IgnoreMask);
 		Builder.SetFlags(FlagsToSet, true);
diff --git a/Engine/Source/Runtime/Engine/Public/CollisionQueryParams.h b/Engine/Source/Runtime/Engine/Public/CollisionQueryParams.h
--- a/Engine/Source/Runtime/Engine/Public/CollisionQueryParams.h
+++ b/Engine/Source/Runtime/Engine/Public/CollisionQueryParams.h
@@ -347,79 +347,104 @@ struct FCollisionResponseParams
 struct FCollisionQueryFlag
 {
 private:
-	int32 AllObjectQueryFlag;
-	int32 AllStaticObjectQueryFlag;
+	static constexpr uint32 MaxCollisionChannelCount = ECC_MAX - 1;
+	int64 AllObjectQueryFlag;
+	int64 AllStaticObjectQueryFlag;
+
+	static int64 ToBitMask(ECollisionChannel QueryChannel)
+	{
+		const int32 Result = ECC_TO_BITFIELD(QueryChannel);
+		return Result;
+	}
 
 	FCollisionQueryFlag()
 	{
-		AllObjectQueryFlag = 0xFFFFFFFF;
-		AllStaticObjectQueryFlag = ECC_TO_BITFIELD(ECC_WorldStatic);
+		AllObjectQueryFlag = std::numeric_limits<uint64>::max();
+		AllStaticObjectQueryFlag = ToBitMask(ECC_WorldStatic);
 	}
 
 public:
 	static ENGINE_API FCollisionQueryFlag& Get();
 
+	UE_DEPRECATED(5.8, "Use GetAllObjectsQueryFlag64")
 	int32 GetAllObjectsQueryFlag()
+	{
+		return (int32)GetAllObjectsQueryFlag64();
+	}
+
+	int64 GetAllObjectsQueryFlag64() const
 	{
 		// this doesn't really verify trace queries to come this way
 		return AllObjectQueryFlag;
 	}
-
+	
+	UE_DEPRECATED(5.8, "Use GetAllStaticObjectsQueryFlag64")
 	int32 GetAllStaticObjectsQueryFlag()
 	{
-		return AllStaticObjectQueryFlag;
+		return (int32)GetAllStaticObjectsQueryFlag64();
 	}
 
+	int64 GetAllStaticObjectsQueryFlag64() const
+	{
+		return AllStaticObjectQueryFlag;
+	}
+	
+	UE_DEPRECATED(5.8, "Use GetAllDynamicObjectsQueryFlag64")
 	int32 GetAllDynamicObjectsQueryFlag()
+	{
+		return (int32)GetAllDynamicObjectsQueryFlag64();
+	}
+
+	int64 GetAllDynamicObjectsQueryFlag64() const
 	{
 		return (AllObjectQueryFlag & ~AllStaticObjectQueryFlag);
 	}
 
 	void AddToAllObjectsQueryFlag(ECollisionChannel NewChannel)
 	{
-		if (ensure ((int32)NewChannel < 32))
+		if (ensure ((int32)NewChannel < MaxCollisionChannelCount))
 		{
-			SetAllObjectsQueryFlag (AllObjectQueryFlag |= ECC_TO_BITFIELD(NewChannel));
+			SetAllObjectsQueryFlag (AllObjectQueryFlag |= ToBitMask(NewChannel));
 		}
 	}
 
 	void AddToAllStaticObjectsQueryFlag(ECollisionChannel NewChannel)
 	{
-		if (ensure ((int32)NewChannel < 32))
+		if (ensure ((int32)NewChannel < MaxCollisionChannelCount))
 		{
-			SetAllStaticObjectsQueryFlag (AllStaticObjectQueryFlag |= ECC_TO_BITFIELD(NewChannel));
+			SetAllStaticObjectsQueryFlag (AllStaticObjectQueryFlag |= ToBitMask(NewChannel));
 		}
 	}
 
 	void RemoveFromAllObjectsQueryFlag(ECollisionChannel NewChannel)
 	{
-		if (ensure ((int32)NewChannel < 32))
+		if (ensure ((int32)NewChannel < MaxCollisionChannelCount))
 		{
-			SetAllObjectsQueryFlag (AllObjectQueryFlag &= ~ECC_TO_BITFIELD(NewChannel));
+			SetAllObjectsQueryFlag (AllObjectQueryFlag &= ~ToBitMask(NewChannel));
 		}
 	}
 
 	void RemoveFromAllStaticObjectsQueryFlag(ECollisionChannel NewChannel)
 	{
-		if (ensure ((int32)NewChannel < 32))
+		if (ensure ((int32)NewChannel < MaxCollisionChannelCount))
 		{
-			SetAllStaticObjectsQueryFlag (AllObjectQueryFlag &= ~ECC_TO_BITFIELD(NewChannel));
+			SetAllStaticObjectsQueryFlag (AllStaticObjectQueryFlag &= ~ToBitMask(NewChannel));
 		}
 	}
 
-	void SetAllObjectsQueryFlag(int32 NewQueryFlag)
+	void SetAllObjectsQueryFlag(int64 NewQueryFlag)
 	{
 		// if all object query has changed, make sure to apply to static object query
 		AllObjectQueryFlag = NewQueryFlag;
 		AllStaticObjectQueryFlag = AllObjectQueryFlag & AllStaticObjectQueryFlag;
 	}
 
-	void SetAllStaticObjectsQueryFlag(int32 NewQueryFlag)
+	void SetAllStaticObjectsQueryFlag(int64 NewQueryFlag)
 	{
 		AllStaticObjectQueryFlag = NewQueryFlag;
 	}
 
-	void SetAllDynamicObjectsQueryFlag(int32 NewQueryFlag)
+	void SetAllDynamicObjectsQueryFlag(int64 NewQueryFlag)
 	{
 		AllStaticObjectQueryFlag = AllObjectQueryFlag & ~NewQueryFlag;
 	}
@@ -428,6 +453,23 @@ public:
 /** Structure that contains list of object types the query is intersted in.  */
 struct FCollisionObjectQueryParams
 {
+private:
+	int64 ObjectTypesToQuery64 = 0;
+
+	void UpdateDeprecatedObjectTypesToQuery()
+	{
+		PRAGMA_DISABLE_DEPRECATION_WARNINGS
+		ObjectTypesToQuery = (int32)ObjectTypesToQuery64;
+		PRAGMA_ENABLE_DEPRECATION_WARNINGS
+	}
+
+	static int64 ToBitMask(ECollisionChannel QueryChannel)
+	{
+		const int32 Result = ECC_TO_BITFIELD(QueryChannel);
+		return Result;
+	}
+public:
+
 	enum InitType
 	{
 		AllObjects,
@@ -435,27 +477,47 @@ struct FCollisionObjectQueryParams
 		AllDynamicObjects
 	};
 
+	int64 GetObjectTypesToQuery() const
+	{
+		return ObjectTypesToQuery64;
+	}
+
+	void SetObjectTypesToQuery(int64 InObjectTypesToQuery)
+	{
+		ObjectTypesToQuery64 = InObjectTypesToQuery;
+		UpdateDeprecatedObjectTypesToQuery();
+	}
+
 	/** Set of object type queries that it is interested in **/
-	int32 ObjectTypesToQuery;
+	UE_DEPRECATED(5.8, "Use Get/Set ObjectTypesToQuery instead.")
+	int32 ObjectTypesToQuery = 0;
 
 	/** Extra filtering done during object query. See declaration for filtering logic */
 	FMaskFilter IgnoreMask;
 
+	PRAGMA_DISABLE_DEPRECATION_WARNINGS
+	FCollisionObjectQueryParams(const FCollisionObjectQueryParams&) = default;
+	FCollisionObjectQueryParams(FCollisionObjectQueryParams&&) = default;
+	FCollisionObjectQueryParams& operator=(const FCollisionObjectQueryParams&) = default;
+	FCollisionObjectQueryParams& operator=(FCollisionObjectQueryParams&&) = default;
+	PRAGMA_ENABLE_DEPRECATION_WARNINGS
+
 	FCollisionObjectQueryParams()
-		: ObjectTypesToQuery(0)
+		: ObjectTypesToQuery64(0)
 		, IgnoreMask(0)
 	{
 	}
 
 	FCollisionObjectQueryParams(ECollisionChannel QueryChannel)
 	{
-		ObjectTypesToQuery = ECC_TO_BITFIELD(QueryChannel);
+		ObjectTypesToQuery64 = ToBitMask(QueryChannel);
 		IgnoreMask = 0;
+		UpdateDeprecatedObjectTypesToQuery();
 	}
 
 	FCollisionObjectQueryParams(const TArray<TEnumAsByte<EObjectTypeQuery> > & ObjectTypes)
 	{
-		ObjectTypesToQuery = 0;
+		ObjectTypesToQuery64 = 0;
 
 		for (auto Iter = ObjectTypes.CreateConstIterator(); Iter; ++Iter)
 		{
@@ -470,63 +532,73 @@ struct FCollisionObjectQueryParams
 		switch (QueryType)
 		{
 		case AllObjects:
-			ObjectTypesToQuery = FCollisionQueryFlag::Get().GetAllObjectsQueryFlag();
+			ObjectTypesToQuery64 = FCollisionQueryFlag::Get().GetAllObjectsQueryFlag64();
 			break;
 		case AllStaticObjects:
-			ObjectTypesToQuery = FCollisionQueryFlag::Get().GetAllStaticObjectsQueryFlag();
+			ObjectTypesToQuery64 = FCollisionQueryFlag::Get().GetAllStaticObjectsQueryFlag64();
 			break;
 		case AllDynamicObjects:
-			ObjectTypesToQuery = FCollisionQueryFlag::Get().GetAllDynamicObjectsQueryFlag();
+			ObjectTypesToQuery64 = FCollisionQueryFlag::Get().GetAllDynamicObjectsQueryFlag64();
 			break;
 		}
 
 		IgnoreMask = 0;
-		
+		UpdateDeprecatedObjectTypesToQuery();
 	};
 
 	// to do this, use ECC_TO_BITFIELD to convert to bit field
 	// i.e. FCollisionObjectQueryParams ( ECC_TO_BITFIELD(ECC_WorldStatic) | ECC_TO_BITFIELD(ECC_WorldDynamic) )
-	FCollisionObjectQueryParams(int32 InObjectTypesToQuery)
+	FCollisionObjectQueryParams(int64 InObjectTypesToQuery)
 	{
-		ObjectTypesToQuery = InObjectTypesToQuery;
+		ObjectTypesToQuery64 = InObjectTypesToQuery;
+		UpdateDeprecatedObjectTypesToQuery();
 		IgnoreMask = 0;
 		DoVerify();
 	}
 
 	void AddObjectTypesToQuery(ECollisionChannel QueryChannel)
 	{
-		ObjectTypesToQuery |= ECC_TO_BITFIELD(QueryChannel);
+		ObjectTypesToQuery64 |= ToBitMask(QueryChannel);
+		UpdateDeprecatedObjectTypesToQuery();
 		DoVerify();
 	}
 
 	void RemoveObjectTypesToQuery(ECollisionChannel QueryChannel)
 	{
-		ObjectTypesToQuery &= ~ECC_TO_BITFIELD(QueryChannel);
+		ObjectTypesToQuery64 &= ~ToBitMask(QueryChannel);
+		UpdateDeprecatedObjectTypesToQuery();
 		DoVerify();
 	}
 
+	UE_DEPRECATED(5.8, "Use GetQueryBitfield64")
 	int32 GetQueryBitfield() const
+	{
+		return (int32)GetQueryBitfield64();
+	}
+
+	int64 GetQueryBitfield64() const
 	{
 		checkSlow(IsValid());
 
-		return ObjectTypesToQuery;
+		return ObjectTypesToQuery64;
 	}
 
 	bool IsValid() const
 	{ 
-		return (ObjectTypesToQuery != 0); 
+		return (ObjectTypesToQuery64 != 0); 
 	}
 
 	static bool IsValidObjectQuery(ECollisionChannel QueryChannel) 
 	{
+		const uint64 QueryChannelBit = ToBitMask(QueryChannel);
 		// return true if this belong to object query type
-		return (ECC_TO_BITFIELD(QueryChannel) & FCollisionQueryFlag::Get().GetAllObjectsQueryFlag()) != 0;
+		return (QueryChannelBit & FCollisionQueryFlag::Get().GetAllObjectsQueryFlag64()) != 0;
 	}
 
 	void DoVerify() const
 	{
 		// you shouldn't use Trace Types to be Object Type Query parameter. This is not technical limitation, but verification process. 
-		checkSlow((ObjectTypesToQuery & FCollisionQueryFlag::Get().GetAllObjectsQueryFlag()) == ObjectTypesToQuery);
+		checkSlow((ObjectTypesToQuery64 & FCollisionQueryFlag::Get().GetAllObjectsQueryFlag64()) == ObjectTypesToQuery64);
 	}
 
 	/** Internal. */
diff --git a/Engine/Source/Runtime/Experimental/Chaos/Private/Chaos/ParticleDirtyFlags.cpp b/Engine/Source/Runtime/Experimental/Chaos/Private/Chaos/ParticleDirtyFlags.cpp
--- a/Engine/Source/Runtime/Experimental/Chaos/Private/Chaos/ParticleDirtyFlags.cpp
+++ b/Engine/Source/Runtime/Experimental/Chaos/Private/Chaos/ParticleDirtyFlags.cpp
@@ -4,7 +4,7 @@
 
 namespace Chaos
 {
-	const FCollisionFilterData& FCollisionData::GetQueryData() const
+	FCollisionFilterData FCollisionData::GetQueryData() const
 	{
 		PRAGMA_DISABLE_DEPRECATION_WARNINGS
 		return QueryData;
@@ -19,7 +19,7 @@ namespace Chaos
 		PRAGMA_ENABLE_DEPRECATION_WARNINGS
 	}
 
-	const FCollisionFilterData& FCollisionData::GetSimData() const
+	FCollisionFilterData FCollisionData::GetSimData() const
 	{
 		PRAGMA_DISABLE_DEPRECATION_WARNINGS
 		return SimData;
diff --git a/Engine/Source/Runtime/Experimental/Chaos/Public/Chaos/ParticleDirtyFlags.h b/Engine/Source/Runtime/Experimental/Chaos/Public/Chaos/ParticleDirtyFlags.h
--- a/Engine/Source/Runtime/Experimental/Chaos/Public/Chaos/ParticleDirtyFlags.h
+++ b/Engine/Source/Runtime/Experimental/Chaos/Public/Chaos/ParticleDirtyFlags.h
@@ -748,9 +748,9 @@ struct FCollisionData
 	FCollisionData& operator=(FCollisionData&& Rhs) = default;
 	PRAGMA_ENABLE_DEPRECATION_WARNINGS
 
-	UE_INTERNAL CHAOS_API const FCollisionFilterData& GetQueryData() const;
+	UE_INTERNAL CHAOS_API FCollisionFilterData GetQueryData() const;
 	UE_INTERNAL CHAOS_API void SetQueryData(const FCollisionFilterData& InQueryData);
-	UE_INTERNAL CHAOS_API const FCollisionFilterData& GetSimData() const;
+	UE_INTERNAL CHAOS_API FCollisionFilterData GetSimData() const;
 	UE_INTERNAL CHAOS_API void SetSimData(const FCollisionFilterData& InSimData);
 
 	CHAOS_API Chaos::Filter::FShapeFilterData GetShapeFilterData() const;
diff --git a/Engine/Source/Runtime/Experimental/Chaos/Public/Chaos/ShapeInstance.h b/Engine/Source/Runtime/Experimental/Chaos/Public/Chaos/ShapeInstance.h
--- a/Engine/Source/Runtime/Experimental/Chaos/Public/Chaos/ShapeInstance.h
+++ b/Engine/Source/Runtime/Experimental/Chaos/Public/Chaos/ShapeInstance.h
@@ -88,10 +88,10 @@ namespace Chaos
 		void* GetUserData() const;
 		void SetUserData(void* InUserData);
 
-		UE_DEPRECATED(5.8, "Use GetCombinedShapeFilterData instead")  const FCollisionFilterData& GetQueryData() const;
+		UE_DEPRECATED(5.8, "Use GetCombinedShapeFilterData instead")  FCollisionFilterData GetQueryData() const;
 		UE_DEPRECATED(5.8, "Use SetCombinedShapeFilterData instead")  void SetQueryData(const FCollisionFilterData& InQueryData);
 
-		UE_DEPRECATED(5.8, "Use GetCombinedShapeFilterData instead") const FCollisionFilterData& GetSimData() const;
+		UE_DEPRECATED(5.8, "Use GetCombinedShapeFilterData instead") FCollisionFilterData GetSimData() const;
 		UE_DEPRECATED(5.8, "Use SetCombinedShapeFilterData instead") void SetSimData(const FCollisionFilterData& InSimData);
 
 		CHAOS_API const Filter::FShapeFilterData GetShapeFilterData() const;
@@ -289,13 +289,13 @@ namespace Chaos
 		}
 
 		PRAGMA_DISABLE_DEPRECATION_WARNINGS
-		UE_DEPRECATED(5.8, "Use GetCombinedShapeFilterData instead") const FCollisionFilterData& GetQueryData() const { return CollisionData.Read().GetQueryData(); }
+		UE_DEPRECATED(5.8, "Use GetCombinedShapeFilterData instead") FCollisionFilterData GetQueryData() const { return CollisionData.Read().GetQueryData(); }
 		UE_DEPRECATED(5.8, "Use SetCombinedShapeFilterData instead") void SetQueryData(const FCollisionFilterData& InQueryData)
 		{
 			CollisionData.Modify(true, DirtyFlags, Proxy, ShapeIdx, [InQueryData](FCollisionData& Data) { Data.SetQueryData(InQueryData); });
 		}
 
-		UE_DEPRECATED(5.8, "Use GetCombinedShapeFilterData instead") const FCollisionFilterData& GetSimData() const { return CollisionData.Read().GetSimData(); }
+		UE_DEPRECATED(5.8, "Use GetCombinedShapeFilterData instead") FCollisionFilterData GetSimData() const { return CollisionData.Read().GetSimData(); }
 		UE_DEPRECATED(5.8, "Use SetCombinedShapeFilterData instead") void SetSimData(const FCollisionFilterData& InSimData)
 		{
 			CollisionData.Modify(true, DirtyFlags, Proxy, ShapeIdx, [InSimData](FCollisionData& Data) { Data.SetSimData(InSimData); });
@@ -576,10 +576,10 @@ namespace Chaos
 		void SetUserData(void* InUserData) { CollisionData.UserData = InUserData; }
 
 		PRAGMA_DISABLE_DEPRECATION_WARNINGS
-		UE_DEPRECATED(5.8, "Use GetCombinedShapeFilterData instead") const FCollisionFilterData& GetQueryData() const { return CollisionData.GetQueryData(); }
+		UE_DEPRECATED(5.8, "Use GetCombinedShapeFilterData instead") FCollisionFilterData GetQueryData() const { return CollisionData.GetQueryData(); }
 		UE_DEPRECATED(5.8, "Use SetCombinedShapeFilterData instead") void SetQueryData(const FCollisionFilterData& InQueryData) { CollisionData.SetQueryData(InQueryData); }
 
-		UE_DEPRECATED(5.8, "Use GetCombinedShapeFilterData instead") const FCollisionFilterData& GetSimData() const { return CollisionData.GetSimData(); }
+		UE_DEPRECATED(5.8, "Use GetCombinedShapeFilterData instead") FCollisionFilterData GetSimData() const { return CollisionData.GetSimData(); }
 		UE_DEPRECATED(5.8, "Use SetCombinedShapeFilterData instead") void SetSimData(const FCollisionFilterData& InSimData) { CollisionData.SetSimData(InSimData); }
 		PRAGMA_ENABLE_DEPRECATION_WARNINGS
 
@@ -1024,10 +1024,10 @@ namespace Chaos
 		DownCast([InUserData](auto& ShapeInstance) { ShapeInstance.SetUserData(InUserData); });
 	}
 
-	inline const FCollisionFilterData& FPerShapeData::GetQueryData() const
+	inline FCollisionFilterData FPerShapeData::GetQueryData() const
 	{
 		PRAGMA_DISABLE_DEPRECATION_WARNINGS
-		return DownCast([](const auto& ShapeInstance) -> const auto& { return ShapeInstance.GetQueryData(); });
+		return DownCast([](const auto& ShapeInstance) -> FCollisionFilterData { return ShapeInstance.GetQueryData(); });
 		PRAGMA_ENABLE_DEPRECATION_WARNINGS
 	}
 
@@ -1038,10 +1038,10 @@ namespace Chaos
 		PRAGMA_ENABLE_DEPRECATION_WARNINGS
 	}
 
-	inline const FCollisionFilterData& FPerShapeData::GetSimData() const
+	inline FCollisionFilterData FPerShapeData::GetSimData() const
 	{
 		PRAGMA_DISABLE_DEPRECATION_WARNINGS
-		return DownCast([](const auto& ShapeInstance) -> const auto& { return ShapeInstance.GetSimData(); });
+		return DownCast([](const auto& ShapeInstance) -> FCollisionFilterData { return ShapeInstance.GetSimData(); });
 		PRAGMA_ENABLE_DEPRECATION_WARNINGS
 	}
 
diff --git a/Engine/Source/Runtime/Experimental/Chaos/Public/RewindData.h b/Engine/Source/Runtime/Experimental/Chaos/Public/RewindData.h
--- a/Engine/Source/Runtime/Experimental/Chaos/Public/RewindData.h
+++ b/Engine/Source/Runtime/Experimental/Chaos/Public/RewindData.h
@@ -1017,7 +1017,7 @@ struct FPerShapeDataStateBase
 
 	template <typename TParticle>
 	UE_DEPRECATED(5.8, "This is unused and will be removed. Use GetCombinedShapeFilterData instead")
-	static const FCollisionFilterData& GetQueryData(const FPerShapeDataStateBase* State, const TParticle& Particle, int32 ShapeIdx)
+	static FCollisionFilterData GetQueryData(const FPerShapeDataStateBase* State, const TParticle& Particle, int32 ShapeIdx)
 	{
 		PRAGMA_DISABLE_DEPRECATION_WARNINGS
 		return State && State->CollisionData.IsSet() ? State->CollisionData.Read().GetQueryData() : Particle.ShapesArray()[ShapeIdx]->GetQueryData();
@@ -1044,7 +1044,7 @@ public:
 
 	PRAGMA_DISABLE_DEPRECATION_WARNINGS
 	UE_DEPRECATED(5.8, "This is unused and will be removed. Use GetCombinedShapeFilterData instead")
-	const FCollisionFilterData& GetQueryData() const { return FPerShapeDataStateBase::GetQueryData(State, Particle, ShapeIdx); }
+	FCollisionFilterData GetQueryData() const { return FPerShapeDataStateBase::GetQueryData(State, Particle, ShapeIdx); }
 	PRAGMA_ENABLE_DEPRECATION_WARNINGS
 private:
 	const FPerShapeDataStateBase* State;
diff --git a/Engine/Source/Runtime/Experimental/ChaosSolverEngine/Private/Chaos/ChaosVDEngineEditorBridge.cpp b/Engine/Source/Runtime/Experimental/ChaosSolverEngine/Private/Chaos/ChaosVDEngineEditorBridge.cpp
--- a/Engine/Source/Runtime/Experimental/ChaosSolverEngine/Private/Chaos/ChaosVDEngineEditorBridge.cpp
+++ b/Engine/Source/Runtime/Experimental/ChaosSolverEngine/Private/Chaos/ChaosVDEngineEditorBridge.cpp
@@ -162,8 +162,7 @@ void FChaosVDEngineEditorBridge::SerializeCollisionChannelsNames()
 
 	if (UCollisionProfile* CollisionProfileData = UCollisionProfile::Get())
 	{
-		constexpr int32 MaxSupportedChannels = 32;
-		for (int32 ChannelIndex = 0; ChannelIndex < MaxSupportedChannels; ++ChannelIndex)
+		for (int32 ChannelIndex = 0; ChannelIndex < FChaosVDCollisionChannelsInfoContainer::MaxSupportedChannels; ++ChannelIndex)
 		{
 			FChaosVDCollisionChannelInfo Info;
 			Info.DisplayName = CollisionProfileData->ReturnChannelNameFromContainerIndex(ChannelIndex).ToString();
diff --git a/Engine/Source/Runtime/Experimental/ChaosVisualDebugger/Public/DataWrappers/ChaosVDCollisionDataWrappers.h b/Engine/Source/Runtime/Experimental/ChaosVisualDebugger/Public/DataWrappers/ChaosVDCollisionDataWrappers.h
--- a/Engine/Source/Runtime/Experimental/ChaosVisualDebugger/Public/DataWrappers/ChaosVDCollisionDataWrappers.h
+++ b/Engine/Source/Runtime/Experimental/ChaosVisualDebugger/Public/DataWrappers/ChaosVDCollisionDataWrappers.h
@@ -496,6 +496,8 @@ USTRUCT()
 struct FChaosVDCollisionChannelsInfoContainer
 {
 	GENERATED_BODY()
+	
+	UE_INTERNAL static constexpr int32 MaxSupportedChannels = 32;
 
 	CHAOSVDRUNTIME_API static FStringView WrapperTypeName;
 
diff --git a/Engine/Source/Runtime/Experimental/ChaosVisualDebugger/Public/DataWrappers/ChaosVDQueryDataWrappers.h b/Engine/Source/Runtime/Experimental/ChaosVisualDebugger/Public/DataWrappers/ChaosVDQueryDataWrappers.h
--- a/Engine/Source/Runtime/Experimental/ChaosVisualDebugger/Public/DataWrappers/ChaosVDQueryDataWrappers.h
+++ b/Engine/Source/Runtime/Experimental/ChaosVisualDebugger/Public/DataWrappers/ChaosVDQueryDataWrappers.h
@@ -71,7 +71,7 @@ struct FChaosVDCollisionObjectQueryParams : public FChaosVDWrapperDataBase
 	template <typename TOther>
 	void CopyFrom(const TOther& Other)
 	{
-		ObjectTypesToQuery = Other.ObjectTypesToQuery;
+		ObjectTypesToQuery = Other.GetObjectTypesToQuery();
 		IgnoreMask = Other.IgnoreMask;
 		bHasValidData = true;
 	}
