Added support for async physics creation/destruction of SkeletalMeshComponent (can be enabled using p.SkeletalMeshComponent.EnableAsyncInitBody).

#Synced-CL 49040969

diff --git a/Engine/Source/Runtime/Engine/Classes/Components/SkeletalMeshComponent.h b/Engine/Source/Runtime/Engine/Classes/Components/SkeletalMeshComponent.h
--- a/Engine/Source/Runtime/Engine/Classes/Components/SkeletalMeshComponent.h
+++ b/Engine/Source/Runtime/Engine/Classes/Components/SkeletalMeshComponent.h
@@ -1668,6 +1668,15 @@ protected:
 
 private:
 
+	/** Called by OnRegister/OnPreRegister before calling Super. */
+	void InitializeBeforeRegister();
+
+	/** Called by OnRegister/OnPreRegister after calling Super. */
+	void InitializeOnRegister();
+
+	/** Shared code between synchronous and asynchronous phyiscs destruction. */
+	void DestroyPhysicsStateInternal();
+
 	/** Wrapper that calls our constraint broken delegate */
 	ENGINE_API void OnConstraintBrokenWrapper(int32 ConstraintIndex);
 
@@ -1834,6 +1843,7 @@ public:
 	ENGINE_API virtual bool HasValidPhysicsState() const override;
 
 protected:
+	ENGINE_API virtual void OnPreRegister() override;
 	ENGINE_API virtual void OnRegister() override;
 	ENGINE_API virtual void OnUnregister() override;
 	ENGINE_API virtual bool ShouldCreatePhysicsState() const override;
@@ -1842,6 +1852,14 @@ protected:
 	ENGINE_API virtual void SendRenderDynamicData_Concurrent() override;
 	ENGINE_API virtual void RegisterComponentTickFunctions(bool bRegister) override;
 public:
+	ENGINE_API virtual bool AllowsAsyncPhysicsStateCreation() const override;
+	ENGINE_API virtual bool AllowsAsyncPhysicsStateDestruction() const override;
+	ENGINE_API virtual void OnAsyncCreatePhysicsStateBegin_GameThread() override;
+	ENGINE_API virtual bool OnAsyncCreatePhysicsState(const UE::FTimeout& TimeOut) override;
+	ENGINE_API virtual void OnAsyncCreatePhysicsStateEnd_GameThread() override;
+	ENGINE_API virtual void OnAsyncDestroyPhysicsStateBegin_GameThread() override;
+	ENGINE_API virtual bool OnAsyncDestroyPhysicsState(const UE::FTimeout& Timeout) override;
+
 	ENGINE_API virtual void InitializeComponent() override;
 	ENGINE_API virtual void TickComponent(float DeltaTime, enum ELevelTick TickType, FActorComponentTickFunction *ThisTickFunction) override;
 	ENGINE_API virtual void BeginPlay() override;
@@ -2060,6 +2078,7 @@ public:
 	ENGINE_API virtual void UpdateVisualizeLODString(FString& DebugString) override;
 	ENGINE_API virtual void RefreshBoneTransforms( FActorComponentTickFunction* TickFunction = NULL ) override;
 protected:
+	ENGINE_API virtual bool AreBoneTransformsValid() const override;
 	ENGINE_API virtual void DispatchParallelTickPose( FActorComponentTickFunction* TickFunction ) override;
 public:
 	ENGINE_API virtual void TickPose(float DeltaTime, bool bNeedsValidRootMotion) override;
@@ -2820,7 +2839,6 @@ private:
 	 */
 	ENGINE_API bool AreRequiredCurvesUpToDate() const;
 
-
 private:
 	// Can't rely on time value, because those may be affected by dilation and whether or not
 	// the game is paused.
@@ -2829,6 +2847,11 @@ private:
 	UPROPERTY(Transient)
 	uint32 LastPoseTickFrame; 
 
+	/** Payload used by asynchronous destruction of physics state (see OnAsyncDestroyPhysicsState). */
+	TArray<FBodyInstance*> AsyncDestroyPhysicsStatePayload;
+	TArray<FConstraintInstance*> AsyncDestroyConstraintsPayload;
+	FPhysicsAggregateHandle AggregatePayload;
+
 public:
 	ENGINE_API void ConditionallyDispatchQueuedAnimEvents();
 
@@ -2840,6 +2863,9 @@ public:
 	ENGINE_API virtual Chaos::FPhysicsObject* GetPhysicsObjectByName(const FName& Name) const override;
 	ENGINE_API virtual TArray<Chaos::FPhysicsObject*> GetAllPhysicsObjects() const override;
 	//~ End IPhysicsComponent Interface.
+
+private:
+	bool bBoneTransformsInitialized = false;
 };
 
 #if WITH_EDITOR
diff --git a/Engine/Source/Runtime/Engine/Classes/Components/SkinnedMeshComponent.h b/Engine/Source/Runtime/Engine/Classes/Components/SkinnedMeshComponent.h
--- a/Engine/Source/Runtime/Engine/Classes/Components/SkinnedMeshComponent.h
+++ b/Engine/Source/Runtime/Engine/Classes/Components/SkinnedMeshComponent.h
@@ -1273,6 +1273,7 @@ public:
 protected:
 	//~ Begin UActorComponent Interface
 	ENGINE_API virtual void InitializeComponent() override;
+	ENGINE_API virtual void OnPreRegister() override;
 	ENGINE_API virtual void OnRegister() override;
 	ENGINE_API virtual void OnUnregister() override;
 	ENGINE_API virtual void BeginPlay() override;
@@ -1572,6 +1573,14 @@ protected:
 	/** Traverse the attachment tree and compute the highest streaming LOD index clamp among components with SyncAttachParentLOD enabled */
 	uint8 ComputeMaxFirstLODIdxInAttachmentTree() const;
 
+private:
+
+	/** Called by OnRegister/OnPreRegister before calling Super. */
+	void InitializeBeforeRegister();
+
+	/** Called by OnRegister/OnPreRegister after calling Super. */
+	void InitializeOnRegister();
+
 public:
 	/**
 	 * Tick Pose, this function ticks and do whatever it needs to do in this frame, should be called before RefreshBoneTransforms
@@ -1880,6 +1889,13 @@ public:
 	/** @return true if this component has any attach children that are attached via sockets */
 	ENGINE_API bool HasAnyAttachedSockets() const;
 
+protected:
+	/**
+	 * Returns whether the bone transforms are valid.
+	 */
+	ENGINE_API virtual bool AreBoneTransformsValid() const;
+
+public:
 	/** 
 	 * Get Bone Matrix from index
 	 *
diff --git a/Engine/Source/Runtime/Engine/Private/Components/PrimitiveComponent.cpp b/Engine/Source/Runtime/Engine/Private/Components/PrimitiveComponent.cpp
--- a/Engine/Source/Runtime/Engine/Private/Components/PrimitiveComponent.cpp
+++ b/Engine/Source/Runtime/Engine/Private/Components/PrimitiveComponent.cpp
@@ -923,8 +923,11 @@ void UPrimitiveComponent::OnCreatePhysicsState()
 			// Create the body.
 			BodyInstance.InitBody(BodySetup, BodyTransform, this, GetWorld()->GetPhysicsScene());		
 #if UE_ENABLE_DEBUG_DRAWING
-			SendRenderDebugPhysics();
-#endif // !(UE_BUILD_SHIPPING || UE_BUILD_TEST)
+			if (!IsAsyncCreatePhysicsStateRunning())
+			{
+				SendRenderDebugPhysics();
+			}
+#endif // UE_ENABLE_DEBUG_DRAWING
 
 #if WITH_EDITOR
 			// Make sure we have a valid body instance here. As we do not keep BIs with no collision shapes at all,
@@ -965,6 +968,10 @@ void UPrimitiveComponent::OnAsyncCreatePhysicsStateEnd_GameThread()
 {
 	check(IsPhysicsStateCreated());
 
+#if UE_ENABLE_DEBUG_DRAWING
+	SendRenderDebugPhysics();
+#endif
+
 	OnComponentPhysicsStateChanged.Broadcast(this, EComponentPhysicsStateChange::Created);
 
 	Super::OnAsyncCreatePhysicsStateEnd_GameThread();
diff --git a/Engine/Source/Runtime/Engine/Private/Components/SkeletalMeshComponent.cpp b/Engine/Source/Runtime/Engine/Private/Components/SkeletalMeshComponent.cpp
--- a/Engine/Source/Runtime/Engine/Private/Components/SkeletalMeshComponent.cpp
+++ b/Engine/Source/Runtime/Engine/Private/Components/SkeletalMeshComponent.cpp
@@ -62,6 +62,7 @@
 #include "Chaos/Serialization/SerializedDataBuffer.h"
 #include "Chaos/Serialization/SerializedMultiPhysicsState.h"
 #include "PhysicsProxy/SingleParticlePhysicsProxy.h"
+#include "Physics/Experimental/PhysScene_AsyncPhysicsStateJobQueue.h"
 #include "UObject/UObjectMigrationContext.h"
 
 #include UE_INLINE_GENERATED_CPP_BY_NAME(SkeletalMeshComponent)
@@ -870,7 +871,7 @@ void USkeletalMeshComponent::FinalizeAttributeEvaluationResults(const FBoneConta
 
 bool USkeletalMeshComponent::NeedToSpawnAnimScriptInstance() const
 {
-	if (!IsRegistered())
+	if (!IsRegistered() && !IsPreRegistering())
 	{
 		return false;
 	}
@@ -932,13 +933,22 @@ bool USkeletalMeshComponent::NeedToSpawnPostPhysicsInstance(bool bForceReinit) c
 bool USkeletalMeshComponent::IsAnimBlueprintInstanced() const
 {
 	return (AnimScriptInstance && AnimScriptInstance->GetClass() == AnimClass);
+}	
+
+bool USkeletalMeshComponent::AreBoneTransformsValid() const
+{
+	return Super::AreBoneTransformsValid() || bBoneTransformsInitialized;
 }
 
-void USkeletalMeshComponent::OnRegister()
+void USkeletalMeshComponent::InitializeBeforeRegister()
 {
-	UpdateHasValidBodies();	//Make sure this is done before we call into the Super which will trigger OnCreatePhysicsState
+	// Make sure this is done before we call into the Super which will trigger OnCreatePhysicsState
+	UpdateHasValidBodies();
+}
 
-	Super::OnRegister();
+void USkeletalMeshComponent::InitializeOnRegister()
+{
+	checkf(!bBoneTransformsInitialized, TEXT("BoneTransforms should not be initialized before skeletal registration setup."));
 
 	// Ensure we have an empty list of linked instances on registration. Ready for the initialization below 
 	// to correctly populate that list.
@@ -947,7 +957,7 @@ void USkeletalMeshComponent::OnRegister()
 	// Ensure the primary tick function is set to tick off the game thread if we are not using internal animation
 	PrimaryComponentTick.bRunOnAnyThread = !bEnableAnimation;
 
-	if(bEnableAnimation)
+	if (bEnableAnimation)
 	{
 		// We force an initialization here because we're in one of two cases.
 		// 1) First register, no spawned instance, need to initialize
@@ -962,6 +972,30 @@ void USkeletalMeshComponent::OnRegister()
 		ResetToRefPose();
 	}
 
+	// Set to true regardless 
+	// This handles cases where we don't really update the transforms 
+	// - SkeletalMeshAsset is invalid 
+	// - LeaderPoseComponent is valid
+	// - AnimScriptInstance is valid && InitializeAnimScriptInstance returns false
+	bBoneTransformsInitialized = true;
+}
+
+void USkeletalMeshComponent::OnRegister()
+{
+	const bool bHasBeenPreRegistered = IsPreRegistered();
+
+	if (!bHasBeenPreRegistered)
+	{
+		InitializeBeforeRegister();
+	}
+
+	Super::OnRegister();
+
+	if (!bHasBeenPreRegistered)
+	{
+		InitializeOnRegister();
+	}
+
 	if (bRenderStatic || (VisibilityBasedAnimTickOption == EVisibilityBasedAnimTickOption::OnlyTickPoseWhenRendered && !FApp::CanEverRender()))
 	{
 		SetComponentTickEnabled(false);
@@ -971,6 +1005,26 @@ void USkeletalMeshComponent::OnRegister()
 	RecreateClothingActors();
 }
 
+void USkeletalMeshComponent::OnPreRegister()
+{
+	TRACE_CPUPROFILER_EVENT_SCOPE(USkeletalMeshComponent::OnPreRegister);
+
+	// Temporarily pause the AsyncPhysicsStateJobQueue
+	//
+	// This is necessary because calling Super::OnPreRegister will schedule an async physics job
+	// Pausing new jobs to be added will prevent USkeletalMeshComponent::OnAsyncCreatePhysicsState from being called 
+	// until USkeletalMeshComponent::OnPreRegister completes and initializes BoneSpace transforms.
+	FPhysScene* PhysScene = GetPhysicsScene();
+	check(PhysScene);
+	FPhysScene_AsyncPhysicsStateJobQueue::FScopedDeferAsyncPhysicsStateJobs DeferJobs(PhysScene);
+
+	InitializeBeforeRegister();
+
+	Super::OnPreRegister();
+
+	InitializeOnRegister();
+}
+
 void USkeletalMeshComponent::OnUnregister()
 {
 	const bool bBlockOnTask = true; // wait on evaluation task so we complete any work before this component goes away
@@ -1020,6 +1074,7 @@ void USkeletalMeshComponent::OnUnregister()
 	{
 		SharedRequiredBones->Reset();
 	}
+	bBoneTransformsInitialized = false;
 
 	Super::OnUnregister();
 }
@@ -1037,7 +1092,7 @@ void USkeletalMeshComponent::InitAnim(bool bForceReinit)
 	
 	// a lot of places just call InitAnim without checking Mesh, so 
 	// I'm moving the check here
-	if ( GetSkeletalMeshAsset() != nullptr && IsRegistered() )
+	if (GetSkeletalMeshAsset() != nullptr && (IsRegistered() || IsPreRegistering()))
 	{
 		//clear cache UID since we don't know if skeleton changed
 		CachedAnimCurveUidVersion = 0;
@@ -1096,6 +1151,8 @@ void USkeletalMeshComponent::InitAnim(bool bForceReinit)
 					ResetToRefPose();
 				}
 
+				bBoneTransformsInitialized = true;
+
 				if (bInitializedAnimInstance)
 				{
 					// Allow blueprints to respond to the event in editor
@@ -1139,7 +1196,7 @@ bool USkeletalMeshComponent::InitializeAnimScriptInstance(bool bForceReinit, boo
 	bool bInitializedPostInstance = false;
 
 	USkeletalMesh* SkelMesh = GetSkeletalMeshAsset();
-	if (SkelMesh && IsRegistered())
+	if (SkelMesh && (IsRegistered() || IsPreRegistering()))
 	{
 		// We may be doing parallel evaluation on the current anim instance
 		// Calling this here with true will block this init till that thread completes
@@ -3023,7 +3080,7 @@ void USkeletalMeshComponent::PostAnimEvaluation(FAnimationEvaluationContext& Eva
 		PostProcessAnimInstance->PostUpdateAnimation();
 	}
 
-	if (!IsRegistered()) // Notify/Event has caused us to go away so cannot carry on from here
+	if (!IsRegistered() && !IsPreRegistering()) // Notify/Event has caused us to go away so cannot carry on from here
 	{
 		return;
 	}
@@ -3396,10 +3453,11 @@ void USkeletalMeshComponent::SetSkeletalMesh(USkeletalMesh* InSkelMesh, bool bRe
 		{
 			InitAnim(bReinitPose);
 		}
-		else if(IsRegistered())
+		else if(IsRegistered() || IsPreRegistering())
 		{
 			RecalcRequiredBones(GetPredictedLODLevel());
 			ResetToRefPose();
+			bBoneTransformsInitialized = true;
 		}
 
 		RecreateClothingActors();
diff --git a/Engine/Source/Runtime/Engine/Private/Components/SkinnedMeshComponent.cpp b/Engine/Source/Runtime/Engine/Private/Components/SkinnedMeshComponent.cpp
--- a/Engine/Source/Runtime/Engine/Private/Components/SkinnedMeshComponent.cpp
+++ b/Engine/Source/Runtime/Engine/Private/Components/SkinnedMeshComponent.cpp
@@ -797,33 +797,32 @@ void USkinnedMeshComponent::InitializeComponent()
 	Super::InitializeComponent();
 }
 
-void USkinnedMeshComponent::OnRegister()
+void USkinnedMeshComponent::InitializeBeforeRegister()
 {
-	LLM_SCOPE(ELLMTag::SkeletalMesh);
-
 	NotifyIfSkinnedAssetChanged();
 
 	// The reason this happens before register
 	// is so that any transform update (or children transform update)
 	// won't result in any issues of accessing SpaceBases
 	// This isn't really ideal solution because these transform won't have 
-	// any valid data yet. 
+	// any valid data yet.
 
 	AnimUpdateRateParams = FAnimUpdateRateManager::GetUpdateRateParameters(this);
 
 	if (LeaderPoseComponent.IsValid())
 	{
-		// we have to make sure it updates the leader pose
+		// We have to make sure it updates the leader pose
 		SetLeaderPoseComponent(LeaderPoseComponent.Get(), true);
 	}
 	else
 	{
 		AllocateTransformData();
 	}
+}
 
-	Super::OnRegister();
-
-	if(FSceneInterface* Scene = GetScene())
+void USkinnedMeshComponent::InitializeOnRegister()
+{
+	if (FSceneInterface* Scene = GetScene())
 	{
 		CachedSceneFeatureLevel = Scene->GetFeatureLevel();
 	}
@@ -838,22 +837,53 @@ void USkinnedMeshComponent::OnRegister()
 	// When we are in editor, ensure that the initial setup is done at LOD 0 so that non-ticking components dont
 	// get unexpected behavior when transitioning LODs post-construction script
 	const int32 OldForcedLOD = GetForcedLOD();
-	if(GIsEditor)
+	if (GIsEditor)
 	{
 		SetForcedLOD(1);
 	}
 #endif
-	
+
 	UpdateLODStatus();
 
 #if WITH_EDITOR
-	if(GIsEditor)
+	if (GIsEditor)
 	{
 		SetForcedLOD(OldForcedLOD);
 	}
 #endif
 
 	InvalidateCachedBounds();
+}
+
+void USkinnedMeshComponent::OnPreRegister()
+{
+	// For more information about the logic in this method, see USkinnedMeshComponent::OnRegister comments.
+	TRACE_CPUPROFILER_EVENT_SCOPE(USkinnedMeshComponent::OnPreRegister);
+
+	// Super schedules an async physics job and TransformData must be allocated before that job starts running
+	InitializeBeforeRegister();
+	
+	Super::OnPreRegister();
+
+	InitializeOnRegister();
+}
+
+void USkinnedMeshComponent::OnRegister()
+{
+	LLM_SCOPE(ELLMTag::SkeletalMesh);
+
+	const bool bHasBeenPreRegistered = IsPreRegistered();
+	if (!bHasBeenPreRegistered)
+	{
+		InitializeBeforeRegister();
+	}
+
+	Super::OnRegister();
+
+	if (!bHasBeenPreRegistered)
+	{
+		InitializeOnRegister();
+	}
 
 	CreateMeshDeformerInstances(GetActiveMeshDeformers());
 
@@ -2180,9 +2210,14 @@ void USkinnedMeshComponent::GetPreSkinnedLocalBounds(FBoxSphereBounds& OutBounds
 	}
 }
 
+bool USkinnedMeshComponent::AreBoneTransformsValid() const
+{
+	return IsRegistered();
+}
+
 FMatrix USkinnedMeshComponent::GetBoneMatrix(int32 BoneIdx) const
 {
-	if ( !IsRegistered() )
+	if (!AreBoneTransformsValid())
 	{
 		// if not registered, we don't have SpaceBases yet. 
 		// also GetComponentTransform() isn't set yet (They're set from relativetranslation, relativerotation, relativescale)
@@ -2232,7 +2267,7 @@ FMatrix USkinnedMeshComponent::GetBoneMatrix(int32 BoneIdx) const
 
 FTransform USkinnedMeshComponent::GetBoneTransform(int32 BoneIdx) const
 {
-	if (!IsRegistered())
+	if (!AreBoneTransformsValid())
 	{
 		// if not registered, we don't have SpaceBases yet. 
 		// also GetComponentTransform() isn't set yet (They're set from relativelocation, relativerotation, relativescale)
@@ -2248,7 +2283,7 @@ FTransform USkinnedMeshComponent::GetBoneTransform(int32 BoneIdx, const FTransfo
 	const USkinnedMeshComponent* const LeaderPoseComponentInst = LeaderPoseComponent.Get();
 	if(LeaderPoseComponentInst)
 	{
-		if (!LeaderPoseComponentInst->IsRegistered())
+		if (!LeaderPoseComponentInst->AreBoneTransformsValid())
 		{
 			// We aren't going to get anything valid from the leader pose if it
 			// isn't valid so for now return identity
@@ -2622,7 +2657,7 @@ void USkinnedMeshComponent::SetSkinnedAssetAndUpdate(USkinnedAsset* InSkinnedAss
 		}
 
 		// Don't init anim state if not registered
-		if (IsRegistered())
+		if (IsRegistered() || IsPreRegistering())
 		{
 			AllocateTransformData();
 			UpdateLeaderBoneMap();
@@ -2642,7 +2677,7 @@ void USkinnedMeshComponent::SetSkinnedAssetAndUpdate(USkinnedAsset* InSkinnedAss
 	// Update external weight array sizes.
 	RefreshExternalMorphTargetWeights();
 	
-	if (IsRegistered())
+	if (IsRegistered() || IsPreRegistering())
 	{
 		// We do this after the FRenderStateRecreator has gone as
 		// UpdateLODStatus needs a valid MeshObject
diff --git a/Engine/Source/Runtime/Engine/Private/PhysicsEngine/ConstraintInstance.cpp b/Engine/Source/Runtime/Engine/Private/PhysicsEngine/ConstraintInstance.cpp
--- a/Engine/Source/Runtime/Engine/Private/PhysicsEngine/ConstraintInstance.cpp
+++ b/Engine/Source/Runtime/Engine/Private/PhysicsEngine/ConstraintInstance.cpp
@@ -462,7 +462,14 @@ bool GetActorRefs(FBodyInstance* Body1, FBodyInstance* Body2, FPhysicsActorHandl
 
 	// Do not create joint unless you have two actors
 	// Do not create joint unless one of the actors is dynamic
-	if((!FPhysicsInterface::IsValid(ActorRef1) || !FPhysicsInterface::IsRigidBody(ActorRef1)) && (!FPhysicsInterface::IsValid(ActorRef2) || !FPhysicsInterface::IsRigidBody(ActorRef2)))
+	bool bActor1InvalidRigidBody = false;
+	bool bActor2InvalidRigidBody = false;
+	FPhysicsCommand::ExecuteRead(ActorRef1, ActorRef2, [&bActor1InvalidRigidBody, &bActor2InvalidRigidBody, ActorRef1, ActorRef2](const FPhysicsActorHandle& ActorA, const FPhysicsActorHandle& ActorB)
+	{
+		bActor1InvalidRigidBody = !FPhysicsInterface::IsValid(ActorRef1) || !FPhysicsInterface::IsRigidBody(ActorRef1);
+		bActor2InvalidRigidBody = !FPhysicsInterface::IsValid(ActorRef2) || !FPhysicsInterface::IsRigidBody(ActorRef2);
+	});
+	if(bActor1InvalidRigidBody && bActor2InvalidRigidBody)
 	{
 #if !(UE_BUILD_SHIPPING || UE_BUILD_TEST)
 		FMessageLog("PIE").Warning()
diff --git a/Engine/Source/Runtime/Engine/Private/SkeletalMeshComponentPhysics.cpp b/Engine/Source/Runtime/Engine/Private/SkeletalMeshComponentPhysics.cpp
--- a/Engine/Source/Runtime/Engine/Private/SkeletalMeshComponentPhysics.cpp
+++ b/Engine/Source/Runtime/Engine/Private/SkeletalMeshComponentPhysics.cpp
@@ -66,6 +66,9 @@ static TAutoConsoleVariable<int32> CVarEnableKinematicDeferralPrePhysicsConditio
 
 static TAutoConsoleVariable<bool> CVarUpdateRigidBodyScaling(TEXT("p.UpdateRigidBodyScaling"), true, TEXT("Updates the rigid body scaling for simulation when the actor scaling is changed \n Default: true."));
 
+static bool GEnableSkeletalMeshComponentAsyncInitBody = false;
+static FAutoConsoleVariableRef CVarEnableSkeletalMeshComponentAsyncInitBody(TEXT("p.SkeletalMeshComponent.EnableAsyncInitBody"), GEnableSkeletalMeshComponentAsyncInitBody, TEXT("Enables SkeletalMeshComponent's body instances to be initialized outside of game thread (default is false)."), ECVF_ReadOnly);
+
 // Used as a default return value for invalid cloth data access
 static const TMap<int32, FClothSimulData> SEmptyClothSimulationData;
 
@@ -890,6 +893,8 @@ void USkeletalMeshComponent::InstantiatePhysicsAsset(const UPhysicsAsset& PhysAs
 {
 	auto BoneTMCallable = [this](int32 BoneIndex)
 	{
+		const bool bExpectedBoneTransformsValidity = IsPreRegistering() || IsPreRegistered() || IsRegistered();
+		checkf(bExpectedBoneTransformsValidity == AreBoneTransformsValid(), TEXT("Component bone transforms are inconsistent with the component's registration state."));
 		return GetBoneTransform(BoneIndex);
 	};
 
@@ -1030,7 +1035,7 @@ void USkeletalMeshComponent::InstantiatePhysicsAssetBodies_Internal(const UPhysi
 			else
 			{
 				BodyInst->DOFMode = EDOFMode::None;
-				if(PhysScene != nullptr && !CVarEnableRagdollPhysics.GetValueOnGameThread())	//We only limit creation of the global physx scenes and not assets related to immedate mode
+				if(PhysScene != nullptr && !CVarEnableRagdollPhysics.GetValueOnAnyThread())	//We only limit creation of the global physx scenes and not assets related to immedate mode
 				{
 					continue;
 				}
@@ -1143,51 +1148,200 @@ void USkeletalMeshComponent::InstantiatePhysicsAsset_Internal(const UPhysicsAsse
 	}
 }
 
-
 void USkeletalMeshComponent::TermArticulated()
 {
-	ResetRootBodyIndex();
+	const bool bIsGameThread = IsInGameThread();
+	const bool bIsAsyncDestroyPhysicsStateRunning = IsAsyncDestroyPhysicsStateRunning();
+	const bool bPrepareAsyncPayload = bIsAsyncDestroyPhysicsStateRunning && bIsGameThread;
+	const bool bUseAsyncPayload = bIsAsyncDestroyPhysicsStateRunning && !bIsGameThread;
+	check(!bPrepareAsyncPayload || !bUseAsyncPayload);
 
-	TermCollisionRelationships();
-
-	FPhysicsCommand::ExecuteWrite(this, [&]()
+	if (bIsGameThread)
 	{
-	// We shut down the physics for each body and constraint here. 
-	// The actual UObjects will get GC'd
+		ResetRootBodyIndex();
 
-	for(int32 i=0; i<Constraints.Num(); i++)
-	{
-		check( Constraints[i] );
-		Constraints[i]->TermConstraint();
-		delete Constraints[i];
+		TermCollisionRelationships();
 	}
+	
+	auto TermArticulatedInternal = [this, bPrepareAsyncPayload, bUseAsyncPayload]()
+	{
+		TArray<FBodyInstance*>& InOutBodies = bUseAsyncPayload ? AsyncDestroyPhysicsStatePayload : Bodies;
+		TArray<FConstraintInstance*>& InOutConstraints = bUseAsyncPayload ? AsyncDestroyConstraintsPayload : Constraints;
+		FPhysicsAggregateHandle& InOutAggregate = bUseAsyncPayload ? AggregatePayload : Aggregate;
 
-	Constraints.Empty();
+		// We shut down the physics for each body and constraint here. The actual UObjects will get GC'd.
+		for (int32 i = 0; i < InOutBodies.Num(); i++)
+		{
+			if (ensure(InOutBodies[i]))
+			{
+				// There seems to be scenarios where there is a body that still has an active weld parent.
+				if (!ensure(InOutBodies[i]->WeldParent == nullptr))
+				{
+					InOutBodies[i]->WeldParent->UnWeld(InOutBodies[i]);
+				}
 
-	for(int32 i=0; i<Bodies.Num(); i++)
-	{
-		if (ensure(Bodies[i]))
+				// Asynchronous destruction is done in OnAsyncDestroyPhysicsState
+				if (!bPrepareAsyncPayload)
+				{
+					InOutBodies[i]->TermBody();
+					delete InOutBodies[i];
+				}
+			}
+		}
+
+		if (!bPrepareAsyncPayload)
 		{
-			// There seems to be scenarios where there is a body that still has an active weld parent.
-			if (!ensure(Bodies[i]->WeldParent == nullptr))
+			for (int32 i = 0; i < InOutConstraints.Num(); i++)
+			{
+				check(InOutConstraints[i]);
+				InOutConstraints[i]->TermConstraint();
+				delete InOutConstraints[i];
+			}
+
+			InOutConstraints.Empty();
+			InOutBodies.Empty();
+
+			// releasing Aggregate, it shouldn't contain any Bodies now, because they are released above
+			if (InOutAggregate.IsValid())
 			{
-				Bodies[i]->WeldParent->UnWeld(Bodies[i]);
+				check(FPhysicsInterface::GetNumActorsInAggregate(InOutAggregate) == 0);
+				FPhysicsInterface::ReleaseAggregate(InOutAggregate);
 			}
-			Bodies[i]->TermBody();
-			delete Bodies[i];
 		}
+		else
+		{
+			// Prepare payload for asynchronous destruction
+			check(AsyncDestroyConstraintsPayload.IsEmpty());
+			AsyncDestroyConstraintsPayload = MoveTemp(Constraints);
+
+			check(AsyncDestroyPhysicsStatePayload.IsEmpty());
+			AsyncDestroyPhysicsStatePayload = MoveTemp(Bodies);
+
+			check(!AggregatePayload.IsValid());
+			AggregatePayload = MoveTemp(Aggregate);
+		}
+	};
+
+	if (bPrepareAsyncPayload)
+	{
+		// No need to lock scene when preparing asynchronous destruction payloads
+		TermArticulatedInternal();
 	}
-	
-	Bodies.Empty();
+	else
+	{
+		FPhysicsCommand::ExecuteWrite(this, [&TermArticulatedInternal]()
+		{
+			TermArticulatedInternal();
+		});
+	}
+}
+
+bool USkeletalMeshComponent::AllowsAsyncPhysicsStateCreation() const
+{
+	check(FPhysScene::SupportsAsyncPhysicsStateCreation());
+
+	// Note: AutoWeld is disabled by default for PhysicsAsset's BodyInstances
+	return GEnableSkeletalMeshComponentAsyncInitBody && (!BodyInstance.bAutoWeld || !Cast<UPrimitiveComponent>(GetAttachParent()));
+}
+
+bool USkeletalMeshComponent::AllowsAsyncPhysicsStateDestruction() const
+{
+	check(FPhysScene::SupportsAsyncPhysicsStateDestruction());
+	return AllowsAsyncPhysicsStateCreation();
+}
 
-	// releasing Aggregate, it shouldn't contain any Bodies now, because they are released above
-		if(Aggregate.IsValid())
+void USkeletalMeshComponent::OnAsyncCreatePhysicsStateBegin_GameThread()
+{
+	TRACE_CPUPROFILER_EVENT_SCOPE(USkeletalMeshComponent::OnAsyncCreatePhysicsStateBegin_GameThread);
+
+	if (bEnablePerPolyCollision == false)
 	{
-			check(FPhysicsInterface::GetNumActorsInAggregate(Aggregate) == 0);
-			FPhysicsInterface::ReleaseAggregate(Aggregate);
+		USceneComponent::OnAsyncCreatePhysicsStateBegin_GameThread();
+	}
+	else
+	{
+		CreateBodySetup();
+		BodySetup->CreatePhysicsMeshes();
+		Super::OnAsyncCreatePhysicsStateBegin_GameThread();
 	}
+}
 
-	});
+bool USkeletalMeshComponent::OnAsyncCreatePhysicsState(const UE::FTimeout& TimeOut)
+{
+	// Don't call Super::OnAsyncCreatePhysicsState as we don't want UActorComponent::OnAsyncCreatePhysicsState
+	// to call OnCreatePhysicsState (which would call USkeletalMeshComponent::OnCreatePhysicsState).
+	TRACE_CPUPROFILER_EVENT_SCOPE(USkeletalMeshComponent::OnAsyncCreatePhysicsState);
+	check(AllowsAsyncPhysicsStateCreation());
+
+	check(AreBoneTransformsValid());
+
+	if (bEnablePerPolyCollision == false)
+	{
+		InitArticulated(GetWorld()->GetPhysicsScene());
+
+		// Don't call Super (skip PrimitiveComponent)
+		USceneComponent::OnCreatePhysicsState();
+	}
+	else
+	{
+		// If we're doing per poly we'll use the body instance of the primitive component
+		Super::OnCreatePhysicsState();
+	}
+
+	return true;
+}
+
+
+void USkeletalMeshComponent::OnAsyncCreatePhysicsStateEnd_GameThread()
+{
+	TRACE_CPUPROFILER_EVENT_SCOPE(USkeletalMeshComponent::OnAsyncCreatePhysicsStateEnd_GameThread);
+
+	if (bEnablePerPolyCollision == false)
+	{
+		// Don't call Super (skip PrimitiveComponent)
+		USceneComponent::OnAsyncCreatePhysicsStateEnd_GameThread();
+#if UE_ENABLE_DEBUG_DRAWING
+		SendRenderDebugPhysics();
+#endif
+	}
+	else
+	{
+		Super::OnAsyncCreatePhysicsStateEnd_GameThread();
+	}
+
+	// Notify physics created
+	OnSkelMeshPhysicsCreated.Broadcast();
+}
+
+void USkeletalMeshComponent::OnAsyncDestroyPhysicsStateBegin_GameThread()
+{
+	TRACE_CPUPROFILER_EVENT_SCOPE(USkeletalMeshComponent::OnAsyncDestroyPhysicsStateBegin_GameThread);
+
+	check(IsAsyncDestroyPhysicsStateRunning());
+
+	DestroyPhysicsStateInternal();
+
+	Super::OnAsyncDestroyPhysicsStateBegin_GameThread();
+}
+
+bool USkeletalMeshComponent::OnAsyncDestroyPhysicsState(const UE::FTimeout& Timeout)
+{
+	TRACE_CPUPROFILER_EVENT_SCOPE(USkeletalMeshComponent::OnAsyncDestroyPhysicsState);
+
+	if (bEnablePerPolyCollision == false)
+	{
+		check(Bodies.IsEmpty());
+		check(Constraints.IsEmpty());
+		check(!Aggregate.IsValid());
+
+		TermArticulated();
+
+		check(AsyncDestroyPhysicsStatePayload.IsEmpty());
+		check(AsyncDestroyConstraintsPayload.IsEmpty());
+		check(!AggregatePayload.IsValid());
+	}
+
+	return Super::OnAsyncDestroyPhysicsState(Timeout);
 }
 
 void USkeletalMeshComponent::TermBodiesBelow(FName ParentBoneName)
@@ -1777,12 +1931,14 @@ bool USkeletalMeshComponent::HasValidPhysicsState() const
 
 void USkeletalMeshComponent::OnCreatePhysicsState()
 {
+	TRACE_CPUPROFILER_EVENT_SCOPE(USkeletalMeshComponent::OnCreatePhysicsState);
+
 	// Init physics
 	if (bEnablePerPolyCollision == false)
 	{
 		InitArticulated(GetWorld()->GetPhysicsScene());
 		USceneComponent::OnCreatePhysicsState(); // Need to route CreatePhysicsState, skip PrimitiveComponent
-#if !(UE_BUILD_SHIPPING || UE_BUILD_TEST)
+#if UE_ENABLE_DEBUG_DRAWING
 		SendRenderDebugPhysics();
 #endif
 	}
@@ -1800,18 +1956,25 @@ void USkeletalMeshComponent::OnCreatePhysicsState()
 
 void USkeletalMeshComponent::OnDestroyPhysicsState()
 {
+	TRACE_CPUPROFILER_EVENT_SCOPE(USkeletalMeshComponent::OnDestroyPhysicsState);
+
+	DestroyPhysicsStateInternal();
+
+	Super::OnDestroyPhysicsState();
+}
+
+void USkeletalMeshComponent::DestroyPhysicsStateInternal()
+{
+	check(IsInGameThread());
+
 	if (bEnablePerPolyCollision == false)
 	{
 		UnWeldFromParent();
 		UnWeldChildren();
 		TermArticulated();
 	}
-
-	Super::OnDestroyPhysicsState();
 }
 
-
-
 #if 0 && !(UE_BUILD_SHIPPING || UE_BUILD_TEST)
 #define DEBUGBROKENCONSTRAINTUPDATE(x) { ##x }
 #else
@@ -2320,7 +2483,7 @@ void USkeletalMeshComponent::SetPhysicsAsset(UPhysicsAsset* InPhysicsAsset, bool
 		// Indicate that 'required bones' array will need to be recalculated.
 		bRequiredBonesUpToDate = false;
 
-#if !(UE_BUILD_SHIPPING || UE_BUILD_TEST)
+#if UE_ENABLE_DEBUG_DRAWING
 		SendRenderDebugPhysics();
 #endif
 	}
