Optimize FSplineMeshSceneProxyDesc::CalculateShaderParams to skip calculating the MeshDeformScaleMinMax if the mesh isn't nanite enabled.

MeshDeformScaleMinMax is only used for nanite LOD calculations and computing it is very expensive.

diff --git a/Engine/Source/Runtime/Engine/Private/SplineMeshSceneProxyDesc.cpp b/Engine/Source/Runtime/Engine/Private/SplineMeshSceneProxyDesc.cpp
--- a/Engine/Source/Runtime/Engine/Private/SplineMeshSceneProxyDesc.cpp
+++ b/Engine/Source/Runtime/Engine/Private/SplineMeshSceneProxyDesc.cpp
@@ -15,6 +15,13 @@ static TAutoConsoleVariable<bool> CVarSplineMeshFitToSourceMeshBounds(
 	TEXT("mesh bounds of all LODs. This prevents gaps that might occur due bounds being expanded by lower LODs or bounds extensions.")
 );
 
+static TAutoConsoleVariable<bool> CVarSkipUpdatingMeshDeformScaleForNonNaniteMeshes(
+	TEXT("r.SplineMesh.SkipUpdatingMeshDeformScaleForNonNaniteMeshes"),
+	true,
+	TEXT("When true, FSplineMeshShaderParams::MeshDeformScaleMinMax will only be computed if the static mesh is nanite enabled. Otherwise it's defaulted to (1,1).")
+	TEXT("This speeds up CalculateShaderParams() for non-nanite meshes considerably.")
+);
+
 static FVector3f SplineEvalPos(const FVector3f& StartPos, const FVector3f& StartTangent, const FVector3f& EndPos, const FVector3f& EndTangent, float A)
 {
 	const float A2 = A * A;
@@ -80,6 +87,8 @@ void FSplineMeshSceneProxyDesc::InitializeFrom(const USplineMeshComponent* InCom
 	{
 		SourceMeshBounds = StaticMesh->GetBounds(); // legacy behavior
 
+		bNaniteEnabled = StaticMesh->HasValidNaniteData();
+
 		if (CVarSplineMeshFitToSourceMeshBounds.GetValueOnAnyThread())
 		{
 			if (const FStaticMeshRenderData* RenderData = StaticMesh->GetRenderData())
@@ -149,40 +158,47 @@ FSplineMeshShaderParams FSplineMeshSceneProxyDesc::CalculateShaderParams() const
 			Output.SplineDistToTexelOffset = (SplineBoundaryMin - MeshMinZ) * InvMeshZLen * MeshTexelLen;
 		}
 
-		// Iteratively solve for an approximation of spline length
-		float SplineLength = 0.0f;
+		if (bNaniteEnabled || !CVarSkipUpdatingMeshDeformScaleForNonNaniteMeshes.GetValueOnAnyThread())
 		{
-			static const uint32 NumIterations = 63; // 64 sampled points
-			static const float IterStep = 1.0f / float(NumIterations);
-			float A = 0.0f;
-			FVector3f PrevPoint = SplineEvalPos(SplineParams, A);
-			for (uint32 i = 0; i < NumIterations; ++i)
+			// Iteratively solve for an approximation of spline length
+			float SplineLength = 0.0f;
 			{
-				FVector3f Point = SplineEvalPos(SplineParams, A);
-				SplineLength += (Point - PrevPoint).Length();
-				PrevPoint = Point;
-				A += IterStep;
+				static const uint32 NumIterations = 63; // 64 sampled points
+				static const float IterStep = 1.0f / float(NumIterations);
+				float A = 0.0f;
+				FVector3f PrevPoint = SplineEvalPos(SplineParams, A);
+				for (uint32 i = 0; i < NumIterations; ++i)
+				{
+					FVector3f Point = SplineEvalPos(SplineParams, A);
+					SplineLength += (Point - PrevPoint).Length();
+					PrevPoint = Point;
+					A += IterStep;
+				}
 			}
-		}
 
-		// Calculate an approximation of how much the mesh gets scaled in each local axis as a result of spline
-		// deformation and take the smallest of the axes. This is important for LOD selection of Nanite spline
-		// meshes.
-		{
-			// Estimate length added due to twisting as well
-			const float XYRadius = BoundsXYRadius * FMath::Max(Output.StartScale.GetAbsMax(), Output.EndScale.GetAbsMax());
-			const float TwistRadians = FMath::Abs(Output.StartRoll - Output.EndRoll);
-			SplineLength += TwistRadians * XYRadius;
-
-			// Take the mid-point scale in X/Y to balance out LOD selection in case either of them are extreme.
-			auto AvgAbs = [](float A, float B) { return (FMath::Abs(A) + FMath::Abs(B)) * 0.5f; };
-			const FVector3f DeformScale = FVector3f(
-				SplineLength * Output.MeshZScale,
-				AvgAbs(Output.StartScale.X, Output.EndScale.X),
-				AvgAbs(Output.StartScale.Y, Output.EndScale.Y)
-			);
+			// Calculate an approximation of how much the mesh gets scaled in each local axis as a result of spline
+			// deformation and take the smallest of the axes. This is important for LOD selection of Nanite spline
+			// meshes.
+			{
+				// Estimate length added due to twisting as well
+				const float XYRadius = BoundsXYRadius * FMath::Max(Output.StartScale.GetAbsMax(), Output.EndScale.GetAbsMax());
+				const float TwistRadians = FMath::Abs(Output.StartRoll - Output.EndRoll);
+				SplineLength += TwistRadians * XYRadius;
+
+				// Take the mid-point scale in X/Y to balance out LOD selection in case either of them are extreme.
+				auto AvgAbs = [](float A, float B) { return (FMath::Abs(A) + FMath::Abs(B)) * 0.5f; };
+				const FVector3f DeformScale = FVector3f(
+					SplineLength * Output.MeshZScale,
+					AvgAbs(Output.StartScale.X, Output.EndScale.X),
+					AvgAbs(Output.StartScale.Y, Output.EndScale.Y)
+				);
 			
-			Output.MeshDeformScaleMinMax = FVector2f(DeformScale.GetMin(), DeformScale.GetMax());
+				Output.MeshDeformScaleMinMax = FVector2f(DeformScale.GetMin(), DeformScale.GetMax());
+			}
+		}
+		else
+		{
+			Output.MeshDeformScaleMinMax = FVector2f::One();
 		}
 	}
 	return Output;
diff --git a/Engine/Source/Runtime/Engine/Public/SplineMeshSceneProxyDesc.h b/Engine/Source/Runtime/Engine/Public/SplineMeshSceneProxyDesc.h
--- a/Engine/Source/Runtime/Engine/Public/SplineMeshSceneProxyDesc.h
+++ b/Engine/Source/Runtime/Engine/Public/SplineMeshSceneProxyDesc.h
@@ -19,6 +19,7 @@ struct FSplineMeshSceneProxyDesc
 	float SplineBoundaryMax = 0.0f;
 	TEnumAsByte<ESplineMeshAxis::Type> ForwardAxis = ESplineMeshAxis::X;
 	bool bSmoothInterpRollScale = false;
+	bool bNaniteEnabled = false;
 	FBoxSphereBounds SourceMeshBounds { ForceInit };
 
 	FSplineMeshParams GetSplineParams() const { return SplineParams; }
