Removed skin weight buffer from skin cache.

diff --git a/Engine/Source/Runtime/Engine/Private/GPUSkinCache.cpp b/Engine/Source/Runtime/Engine/Private/GPUSkinCache.cpp
--- a/Engine/Source/Runtime/Engine/Private/GPUSkinCache.cpp
+++ b/Engine/Source/Runtime/Engine/Private/GPUSkinCache.cpp
@@ -707,18 +707,17 @@ public:
 
 	void UpdateSkinWeightBuffer()
 	{
-		const FSkinWeightVertexBuffer* WeightBuffer = GPUSkin->GetSkinWeightVertexBuffer(LOD);
-		bUse16BitBoneIndex = WeightBuffer->Use16BitBoneIndex();
-		bUse16BitBoneWeight = WeightBuffer->Use16BitBoneWeight();
-		InputWeightIndexSize = WeightBuffer->GetBoneIndexByteSize() | (WeightBuffer->GetBoneWeightByteSize() << 8);
-		InputWeightStride = WeightBuffer->GetConstantInfluencesVertexStride();
-		InputWeightStreamSRV = WeightBuffer->GetDataVertexBuffer()->GetSRV();
-		InputWeightLookupStreamSRV = WeightBuffer->GetLookupVertexBuffer()->GetSRV();
-
-		if (WeightBuffer->GetBoneInfluenceType() == GPUSkinBoneInfluenceType::DefaultBoneInfluence)
-		{
-			int32 MaxBoneInfluences = WeightBuffer->GetMaxBoneInfluences();
-			BoneInfluenceType = MaxBoneInfluences > MAX_INFLUENCES_PER_STREAM ? EGPUSkinBoneInfluenceType::Extra : EGPUSkinBoneInfluenceType::Default;
+		FSkinWeightVertexBufferInfo Info = GPUSkin->GetSkinWeightVertexBuffer_RenderThread(LOD);
+		bUse16BitBoneIndex         = Info.bUse16BitBoneIndex;
+		bUse16BitBoneWeight        = Info.bUse16BitBoneWeight;
+		InputWeightIndexSize       = Info.InputWeightIndexSize;
+		InputWeightStride          = Info.InputWeightStride;
+		InputWeightStreamSRV       = Info.InputWeightStreamSRV;
+		InputWeightLookupStreamSRV = Info.InputWeightLookupStreamSRV;
+
+		if (Info.BoneInfluenceType == GPUSkinBoneInfluenceType::DefaultBoneInfluence)
+		{
+			BoneInfluenceType = Info.MaxBoneInfluences > MAX_INFLUENCES_PER_STREAM ? EGPUSkinBoneInfluenceType::Extra : EGPUSkinBoneInfluenceType::Default;
 		}
 		else
 		{
diff --git a/Engine/Source/Runtime/Engine/Private/SkeletalRenderGPUSkin.cpp b/Engine/Source/Runtime/Engine/Private/SkeletalRenderGPUSkin.cpp
--- a/Engine/Source/Runtime/Engine/Private/SkeletalRenderGPUSkin.cpp
+++ b/Engine/Source/Runtime/Engine/Private/SkeletalRenderGPUSkin.cpp
@@ -1771,6 +1771,25 @@ void FSkeletalMeshObjectGPUSkin::FSkeletalMeshObjectLOD::UpdateMorphVertexBuffer
 	}
 }
 
+inline FSkinWeightVertexBufferInfo GetSkinWeightVertexBufferInfo(const FSkinWeightVertexBuffer* Buffer)
+{
+	FSkinWeightVertexBufferInfo Info;
+
+	if (Buffer)
+	{
+		Info.InputWeightStreamSRV       = Buffer->GetDataVertexBuffer()->GetSRV();
+		Info.InputWeightLookupStreamSRV = Buffer->GetLookupVertexBuffer()->GetSRV();
+		Info.InputWeightIndexSize       = Buffer->GetBoneIndexByteSize() | (Buffer->GetBoneWeightByteSize() << 8);
+		Info.InputWeightStride          = Buffer->GetConstantInfluencesVertexStride();
+		Info.MaxBoneInfluences          = Buffer->GetMaxBoneInfluences();
+		Info.BoneInfluenceType          = Buffer->GetBoneInfluenceType();
+		Info.bUse16BitBoneIndex         = Buffer->Use16BitBoneIndex();
+		Info.bUse16BitBoneWeight        = Buffer->Use16BitBoneWeight();
+	}
+
+	return Info;
+}
+
 void FSkeletalMeshObjectGPUSkin::FSkeletalMeshObjectLOD::UpdateSkinWeights(const FSkelMeshComponentLODInfo* CompLODInfo, EBoneTransformStorageMode InBoneTransformStorageMode)
 {
 	QUICK_SCOPE_CYCLE_COUNTER(STAT_FSkeletalMeshObjectLOD_UpdateSkinWeights);
@@ -1794,7 +1813,7 @@ void FSkeletalMeshObjectGPUSkin::FSkeletalMeshObjectLOD::UpdateSkinWeights(const
 			ENQUEUE_RENDER_COMMAND(UpdateSkinWeightsGPUSkin)(UE::RenderCommandPipe::SkeletalMesh,
 				[NewMeshObjectWeightBuffer, VertexBuffers, Self](FRHICommandList& RHICmdList)
 			{
-				Self->MeshObjectWeightBuffer_RenderThread = NewMeshObjectWeightBuffer;
+				Self->MeshObjectWeightBuffer_RenderThread = GetSkinWeightVertexBufferInfo(VertexBuffers.SkinWeightVertexBuffer);
 				Self->GPUSkinVertexFactories.UpdateVertexFactoryData(VertexBuffers);
 			});
 		}
@@ -1994,9 +2013,11 @@ FGPUBaseSkinVertexFactory const* FSkeletalMeshObjectGPUSkin::GetBaseSkinVertexFa
 	return LODs[LODIndex].GPUSkinVertexFactories.VertexFactories[ChunkIdx].Get();
 }
 
-const FSkinWeightVertexBuffer* FSkeletalMeshObjectGPUSkin::GetSkinWeightVertexBuffer(int32 LODIndex) const
+FSkinWeightVertexBufferInfo FSkeletalMeshObjectGPUSkin::GetSkinWeightVertexBuffer_RenderThread(int32 LODIndex) const
 {
-	checkSlow(LODs.IsValidIndex(LODIndex));
+	check(LODs.IsValidIndex(LODIndex));
+	check(IsInRenderingThread());
+
 	return LODs[LODIndex].MeshObjectWeightBuffer_RenderThread;
 }
 
@@ -2560,7 +2581,7 @@ void FSkeletalMeshObjectGPUSkin::FSkeletalMeshObjectLOD::InitResources(
 	ENQUEUE_RENDER_COMMAND(FSkeletalMeshObjectLOD_InitResources)(UE::RenderCommandPipe::SkeletalMesh,
 		[this, InGPUSkinCache, VertexBuffers, &LODData, InFeatureLevel, VertexAttributeMask](FRHICommandList& RHICmdList)
 	{
-		MeshObjectWeightBuffer_RenderThread = VertexBuffers.SkinWeightVertexBuffer;
+		MeshObjectWeightBuffer_RenderThread = GetSkinWeightVertexBufferInfo(VertexBuffers.SkinWeightVertexBuffer);
 		GPUSkinVertexFactories.InitVertexFactories(RHICmdList, VertexBuffers, LODData.RenderSections, InFeatureLevel, VertexAttributeMask, GPUSkinTechnique);
 
 		if (DynamicBoundsStartOffset == INDEX_NONE && GPUSkinTechnique == ESkeletalMeshGPUSkinTechnique::GPUSkinCache)
diff --git a/Engine/Source/Runtime/Engine/Private/SkeletalRenderGPUSkin.h b/Engine/Source/Runtime/Engine/Private/SkeletalRenderGPUSkin.h
--- a/Engine/Source/Runtime/Engine/Private/SkeletalRenderGPUSkin.h
+++ b/Engine/Source/Runtime/Engine/Private/SkeletalRenderGPUSkin.h
@@ -470,7 +470,7 @@ public:
 	/** Check if a given morph set is active or not. If so, we will process it. */
 	bool IsExternalMorphSetActive(int32 MorphSetID, const FExternalMorphSet& MorphSet) const;
 
-	const FSkinWeightVertexBuffer* GetSkinWeightVertexBuffer(int32 LODIndex) const;
+	FSkinWeightVertexBufferInfo GetSkinWeightVertexBuffer_RenderThread(int32 LODIndex) const override;
 
 	/** 
 	 * Get the skin vertex factory for direct skinning. 
@@ -629,7 +629,7 @@ private:
 
 		/** Skin weight buffer to use, could be from asset or component override */
 		const FSkinWeightVertexBuffer* MeshObjectWeightBuffer = nullptr;
-		const FSkinWeightVertexBuffer* MeshObjectWeightBuffer_RenderThread = nullptr;
+		FSkinWeightVertexBufferInfo MeshObjectWeightBuffer_RenderThread;
 
 		/** Color buffer to user, could be from asset or component override */
 		FColorVertexBuffer* MeshObjectColorBuffer = nullptr;
diff --git a/Engine/Source/Runtime/Engine/Public/SkeletalRenderPublic.h b/Engine/Source/Runtime/Engine/Public/SkeletalRenderPublic.h
--- a/Engine/Source/Runtime/Engine/Public/SkeletalRenderPublic.h
+++ b/Engine/Source/Runtime/Engine/Public/SkeletalRenderPublic.h
@@ -18,6 +18,7 @@
 #include "RenderMath.h"
 #include "Async/Mutex.h"
 #include "Matrix3x4.h"
+#include "GPUSkinVertexFactory.h"
 
 class FPrimitiveDrawInterface;
 class FVertexFactory;
@@ -37,6 +38,18 @@ namespace Nanite
 	struct FMaterialAudit;
 }
 
+struct FSkinWeightVertexBufferInfo
+{
+	FShaderResourceViewRHIRef InputWeightStreamSRV;
+	FShaderResourceViewRHIRef InputWeightLookupStreamSRV;
+	uint32 InputWeightIndexSize = 0;
+	uint32 InputWeightStride = 0;
+	uint32 MaxBoneInfluences = 0;
+	GPUSkinBoneInfluenceType BoneInfluenceType = DefaultBoneInfluence;
+	uint8 bUse16BitBoneIndex  : 1 = 0;
+	uint8 bUse16BitBoneWeight : 1 = 0;
+};
+
 namespace UE::SkeletalRender::Settings
 {
 	// Returns the maximum value allowed for morph targets blend weights, configured at RenderSettings
@@ -468,8 +481,15 @@ public:
 	static FColorVertexBuffer* GetColorVertexBuffer(FSkeletalMeshLODRenderData& LODData, const FSkelMeshComponentLODInfo* CompLODInfo);
 
 	/** Get the weight buffer for specific LOD index. Only needed for GPU skin cache */
+	UE_DEPRECATED(5.8, "Use GetSkinWeightVertexBuffer_RenderThread instead")
 	virtual const FSkinWeightVertexBuffer* GetSkinWeightVertexBuffer(int32 LODIndex) const { return nullptr; }
 
+	/** Get the weight buffer for specific LOD index. Only needed for GPU skin cache */
+	virtual FSkinWeightVertexBufferInfo GetSkinWeightVertexBuffer_RenderThread(int32 LODIndex) const
+	{
+		return FSkinWeightVertexBufferInfo();
+	}
+
 	/** Get the offset in the dynamic mesh bounds buffer for the current LOD. Only needed for GPU skin cache. */
 	virtual int32 GetDynamicBoundsStartOffset(int32 LODIndex) const { return -1; }
 
