Fix flipped RoundTrip vs RoundTripExclFrame ping analytics.

Trivial changes, outside of current working branch:
#Synced-CL -1

diff --git a/Engine/Source/Runtime/Engine/Private/NetConnection.cpp b/Engine/Source/Runtime/Engine/Private/NetConnection.cpp
--- a/Engine/Source/Runtime/Engine/Private/NetConnection.cpp
+++ b/Engine/Source/Runtime/Engine/Private/NetConnection.cpp
@@ -3079,8 +3079,8 @@ bool UNetConnection::ReadPacketInfo(FBitReader& Reader, bool bHasPacketInfoPaylo
 
 		if (NetPing.IsValid())
 		{
-			NetPing->UpdatePing(EPingType::RoundTrip, CurrentTime, PacketPing);
-			NetPing->UpdatePing(EPingType::RoundTripExclFrame, CurrentTime, FMath::Max(RTT, 0.0));
+			NetPing->UpdatePing(EPingType::RoundTrip, CurrentTime, FMath::Max(RTT, 0.0));
+			NetPing->UpdatePing(EPingType::RoundTripExclFrame, CurrentTime, PacketPing);
 		}
 
 		if (NetworkCongestionControl.IsSet())
