Fixed crash due to IsInRenderingThread check firing in rendering task.

diff --git a/Engine/Source/Runtime/Engine/Private/SkeletalRenderGPUSkin.cpp b/Engine/Source/Runtime/Engine/Private/SkeletalRenderGPUSkin.cpp
--- a/Engine/Source/Runtime/Engine/Private/SkeletalRenderGPUSkin.cpp
+++ b/Engine/Source/Runtime/Engine/Private/SkeletalRenderGPUSkin.cpp
@@ -2016,7 +2016,6 @@ FGPUBaseSkinVertexFactory const* FSkeletalMeshObjectGPUSkin::GetBaseSkinVertexFa
 FSkinWeightVertexBufferInfo FSkeletalMeshObjectGPUSkin::GetSkinWeightVertexBuffer_RenderThread(int32 LODIndex) const
 {
 	check(LODs.IsValidIndex(LODIndex));
-	check(IsInRenderingThread());
 
 	return LODs[LODIndex].MeshObjectWeightBuffer_RenderThread;
 }
