Minor optimisation in ULevel::AddLoadedActors, don't loop over loaded actors if components aren't registered yet for the current level.

#Synced-CL 48998606

diff --git a/Engine/Source/Runtime/Engine/Private/Level.cpp b/Engine/Source/Runtime/Engine/Private/Level.cpp
--- a/Engine/Source/Runtime/Engine/Private/Level.cpp
+++ b/Engine/Source/Runtime/Engine/Private/Level.cpp
@@ -1019,7 +1019,7 @@ void ULevel::AddLoadedActors(const TArray<AActor*>& ActorList, const FTransform*
 		}
 	}
 
-	FScopedSlowTask SlowTask(ActorsQueue.Num() * 3, LOCTEXT("RegisteringActors", "Registering actors..."));
+	FScopedSlowTask SlowTask(ActorsQueue.Num() * 2 + (bAreComponentsCurrentlyRegistered ? ActorsQueue.Num() : 0), LOCTEXT("RegisteringActors", "Registering actors..."));
 	SlowTask.MakeDialogDelayed(1.0f);
 
 	OnLoadedActorAddedToLevelPreEvent.Broadcast(ActorsQueue);
@@ -1048,18 +1048,18 @@ void ULevel::AddLoadedActors(const TArray<AActor*>& ActorList, const FTransform*
 	}
 
 	// Rerun construction scripts
-	for (AActor* Actor : ActorsQueue)
+	if (bAreComponentsCurrentlyRegistered)
 	{
-		// RegisterAllComponents/RerunConstructionScripts can destroy child actors
-		if (IsValid(Actor))
+		for (AActor* Actor : ActorsQueue)
 		{
-			if (bAreComponentsCurrentlyRegistered)
+			// RegisterAllComponents/RerunConstructionScripts can destroy child actors
+			if (IsValid(Actor))
 			{
 				Actor->RerunConstructionScripts();
 			}
-		}
 
-		SlowTask.EnterProgressFrame(1);
+			SlowTask.EnterProgressFrame(1);
+		}
 	}
 
 	using namespace UE::Editor::DataStorage;
