- Refactor all PSO Precache collection tweaking into a single struct which is more overviewable and should be easier to tweak per platform (for example on low end PC machines use different settings)
- All cvars have moved to this file instead of being scattered all around the engine (cvars are still applied on top of default or ini settings if values are >= 0)
- Add different PSO precache settings for low end PC with <= 6 cores by default to not precache slate and particles at system level (but at proxy creation level only)

#Synced-CL 49066199

diff --git a/Engine/Plugins/FX/Niagara/Source/Niagara/Private/NiagaraComponent.cpp b/Engine/Plugins/FX/Niagara/Source/Niagara/Private/NiagaraComponent.cpp
--- a/Engine/Plugins/FX/Niagara/Source/Niagara/Private/NiagaraComponent.cpp
+++ b/Engine/Plugins/FX/Niagara/Source/Niagara/Private/NiagaraComponent.cpp
@@ -33,6 +33,7 @@
 #include "UObject/UObjectThreadContext.h"
 #include "PrimitiveUniformShaderParametersBuilder.h"
 #include "NiagaraActor.h"
+#include "PSOPrecacheSettings.h"
 
 #if WITH_EDITOR
 #include "INiagaraEditorOnlyDataUtlities.h"
@@ -54,8 +55,6 @@ DEFINE_LOG_CATEGORY(LogNiagara);
 
 DEFINE_RENDER_COMMAND_PIPE(NiagaraDynamicData, ERenderCommandPipeFlags::None);
 
-extern int32 GNiagaraPSOPrecachingTime;
-
 namespace NiagaraComponentPrivate
 {
 	static bool GParallelGDMEEnabled = true;
@@ -1231,7 +1230,7 @@ bool UNiagaraComponent::InitializeSystem()
 			}
 		}
 
-		if (GNiagaraPSOPrecachingTime > 1)
+		if (UPSOPrecacheSettingsManager::GetSettings().ParticlePrecachingTime > 1)
 		{
 			PrecachePSOs();
 		}
@@ -1746,7 +1745,7 @@ void UNiagaraComponent::PostSystemTick_GameThread()
 
 	if (bRecachePSOs)
 	{
-		if (GNiagaraPSOPrecachingTime > 1)
+		if (UPSOPrecacheSettingsManager::GetSettings().ParticlePrecachingTime > 1)
 		{
 			PrecachePSOs();
 		}
@@ -2427,7 +2426,7 @@ FPrimitiveSceneProxy* UNiagaraComponent::CreateSceneProxy()
 	if (Asset != nullptr)
 	{
 		// If we are relying on proxy creation PSOs only then call precache PSOs here because asset doesn't have collected PSO precache data
-		if (GNiagaraPSOPrecachingTime == 1)
+		if (UPSOPrecacheSettingsManager::GetSettings().ParticlePrecachingTime == 1)
 		{
 			PrecachePSOs();
 		}
diff --git a/Engine/Plugins/FX/Niagara/Source/Niagara/Private/NiagaraSystem.cpp b/Engine/Plugins/FX/Niagara/Source/Niagara/Private/NiagaraSystem.cpp
--- a/Engine/Plugins/FX/Niagara/Source/Niagara/Private/NiagaraSystem.cpp
+++ b/Engine/Plugins/FX/Niagara/Source/Niagara/Private/NiagaraSystem.cpp
@@ -60,6 +60,7 @@
 #include "ShaderDiagnostics.h"
 #include "NiagaraDataChannel.h"
 #include "UObject/UObjectIterator.h"
+#include "PSOPrecacheSettings.h"
 
 #include UE_INLINE_GENERATED_CPP_BY_NAME(NiagaraSystem)
 
@@ -193,17 +194,6 @@ static FAutoConsoleVariableRef CVarNiagaraCompileWaitLoggingTerminationCap(
 	ECVF_Default
 );
 
-int32 GNiagaraPSOPrecachingTime = 3;
-FAutoConsoleVariableRef CVarNiagaraPSOPrecachingTime(
-	TEXT("r.PSOPrecache.NiagaraPrecachingTime"),
-	GNiagaraPSOPrecachingTime,
-	TEXT("Controls when PSO precaching happens for Niagara systems:\n")
-	TEXT("	0: no precaching\n")
-	TEXT("	1: precaching at component proxy creation time only\n")
-	TEXT("	2: precaching at component loading time\n")
-	TEXT("	3: precaching at asset loading time and component time (default)"),
-	ECVF_Default);
-
 #if WITH_EDITORONLY_DATA
 static int GNiagaraOnDemandCompileEnabled = 2;
 static FAutoConsoleVariableRef CVarNiagaraOnDemandCompileEnabled(
@@ -1593,7 +1583,7 @@ void UNiagaraSystem::PostLoad()
 	}
 #endif // WITH_EDITORONLY_DATA
 
-	if (GNiagaraPSOPrecachingTime >= 3)
+	if (UPSOPrecacheSettingsManager::GetSettings().ParticlePrecachingTime >= 3)
 	{
 		PrecachePSOs();
 	}
diff --git a/Engine/Source/Runtime/Engine/Private/Components/PrimitiveComponent.cpp b/Engine/Source/Runtime/Engine/Private/Components/PrimitiveComponent.cpp
--- a/Engine/Source/Runtime/Engine/Private/Components/PrimitiveComponent.cpp
+++ b/Engine/Source/Runtime/Engine/Private/Components/PrimitiveComponent.cpp
@@ -52,6 +52,7 @@
 #include "MeshUVChannelInfo.h"
 #include "PrimitiveSceneDesc.h"
 #include "PSOPrecacheMaterial.h"
+#include "PSOPrecacheSettings.h"
 #include "MaterialCachedData.h"
 #include "MaterialShared.h"
 #include "MarkActorRenderStateDirtyTask.h"
@@ -4974,8 +4975,7 @@ void UPrimitiveComponent::SetRenderCustomDepth(bool bValue)
 		bRenderCustomDepth = bValue;
 
 		// Precache the PSOs again when custom depth mode is set to 1 (only precache custom depth PSOs when bRenderCustomDepth is explicitly set)
-		static const auto CVar = IConsoleManager::Get().FindConsoleVariable(TEXT("r.PSOPrecache.CustomDepth"));
-		if (CVar && CVar->GetInt() == 1)
+		if (UPSOPrecacheSettingsManager::GetSettings().CustomDepthMode == 1)
 		{
 			PrecachePSOs();
 		}
diff --git a/Engine/Source/Runtime/Engine/Private/PSOPrecache.cpp b/Engine/Source/Runtime/Engine/Private/PSOPrecache.cpp
--- a/Engine/Source/Runtime/Engine/Private/PSOPrecache.cpp
+++ b/Engine/Source/Runtime/Engine/Private/PSOPrecache.cpp
@@ -14,14 +14,6 @@
 #include "PSOPrecacheValidation.h"
 #include "PSOPrecacheComponent.h"
 
-static TAutoConsoleVariable<int32> CVarPrecacheGlobalComputeShaders(
-	TEXT("r.PSOPrecache.GlobalShaders"),
-	0,
-	TEXT("Precache global shaders during startup (disable(0) - only compute shaders(1) - all global shaders(2).\n") 
-	TEXT("Note: r.PSOPrecache.GlobalShaders == 2 is only supported when r.PSOPrecache.Mode == 1 (Preload shaders)."),
-	ECVF_ReadOnly
-);
-
 int32 GPSOPrecacheComponents = 1;
 static FAutoConsoleVariableRef CVarPSOPrecacheComponents(
 	TEXT("r.PSOPrecache.Components"),
diff --git a/Engine/Source/Runtime/Engine/Private/PSOPrecacheSettings.cpp b/Engine/Source/Runtime/Engine/Private/PSOPrecacheSettings.cpp
new file mode 100644
--- /dev/null
+++ b/Engine/Source/Runtime/Engine/Private/PSOPrecacheSettings.cpp
@@ -0,0 +1,216 @@
+// Copyright Epic Games, Inc. All Rights Reserved.
+
+/*=============================================================================
+	PSOPrecacheSettings.cpp: 
+=============================================================================*/
+
+#include "PSOPrecacheSettings.h"
+#include "HAL/IConsoleManager.h"
+#include "RHIShaderPlatform.h"
+#include "RHIStrings.h"
+
+static int32 GPSOPrecacheGlobalShaders = -1;
+static FAutoConsoleVariableRef CVarPSOPrecacheGlobalShaders(
+	TEXT("r.PSOPrecache.GlobalShaders"),
+	GPSOPrecacheGlobalShaders,
+	TEXT("Precache global shaders during startup (disable(0) - global compute shaders and registered global graphis PSOs(1) - all global shaders(2).\n") 
+	TEXT("Note: r.PSOPrecache.GlobalShaders == 2 is only supported when r.PSOPrecache.Mode == 1 (Preload shaders)."),
+	ECVF_ReadOnly
+);
+
+static int32 GPSOPrecacheLightMapPolicyMode = -1;
+static FAutoConsoleVariableRef CVarPSOPrecacheLightMapPolicyMode(
+	TEXT("r.PSOPrecache.LightMapPolicyMode"),
+	GPSOPrecacheLightMapPolicyMode,
+	TEXT("Defines which light map policies should be checked during PSO precaching of the base pass.\n") \
+	TEXT("-1: Use engine defaults.\n") \
+	TEXT(" 0: All possible LMP will be checked.\n") \
+	TEXT(" 1: Only LMP_NO_LIGHTMAP will be precached.\n"),
+	ECVF_ReadOnly
+);
+
+static int32 GPSOPrecacheTranslucencyAllPass = -1;
+static FAutoConsoleVariableRef CVarPSOPrecacheTranslucencyAllPass(
+	TEXT("r.PSOPrecache.TranslucencyAllPass"),
+	GPSOPrecacheTranslucencyAllPass,
+	TEXT("Precache PSOs for TranslucencyAll pass.\n") \
+	TEXT("-1: Use engine defaults.\n") \
+	TEXT(" 0: No PSOs are compiled for this pass.\n") \
+	TEXT(" 1: PSOs are compiled for all primitives which render to a translucency pass.\n"),
+	ECVF_ReadOnly
+);
+
+static int32 GPSOPrecacheTranslucencyStandardBeforeDOF = -1;
+static FAutoConsoleVariableRef CVarPSOPrecacheTranslucencyStandardBeforeDOF(
+	TEXT("r.PSOPrecache.TranslucencyStandardBeforeDOF"),
+	GPSOPrecacheTranslucencyStandardBeforeDOF,
+	TEXT("Precache permutations for EMeshPass::TranslucencyStandard even when marked as AfterDOF - can be needed when AutoBeforeDOFTranslucencyBoundary is set on the view:\n") \
+	TEXT("-1: Use engine defaults.\n") \
+	TEXT(" 0: No PSOs are compiled for this use case.\n") \
+	TEXT(" 1: PSOs are compiled for all primitives which render to a translucency standard and are marked for AfterDOF.\n"),
+	ECVF_ReadOnly
+);
+
+static int32 GPSOPrecacheAlphaColorChannel = -1;
+static FAutoConsoleVariableRef CVarPSOPrecacheAlphaColorChannel(
+	TEXT("r.PSOPrecache.PrecacheAlphaColorChannel"),
+	GPSOPrecacheAlphaColorChannel,
+	TEXT("Also Precache PSOs with scene color alpha channel enabled. Planar reflections and scene captures use this for compositing into a different scene later (-1 use engine default)") \
+	TEXT("-1: Use engine defaults.\n"),
+	ECVF_ReadOnly
+);
+
+static int32 GPSOPrecacheTranslucencyUnderWater = -1;
+static FAutoConsoleVariableRef CVarPSOPrecacheTranslucencyUnderWater(
+	TEXT("r.PSOPrecache.PrecacheTranslucencyUnderWater"),
+	GPSOPrecacheTranslucencyUnderWater,
+	TEXT("Precache translucency under water pass permutations with different render target formats (-1 use engine default)"),
+	ECVF_ReadOnly
+);
+
+static int32 GPSOPrecachePostProcessingMaterial = -1;
+static FAutoConsoleVariableRef CVarPSOPrecachePostProcessingMaterial(
+	TEXT("r.PSOPrecache.PostProcessingMaterial"), 
+	GPSOPrecachePostProcessingMaterial, 
+	TEXT("Precache all possible required PSOs for loaded PostProcessing Materials (-1 use engine default)"),
+	ECVF_ReadOnly);
+
+static int32 GPSOPrecacheTranslucencyLightingVolumeMaterial = -1;
+static FAutoConsoleVariableRef CVarPSOPrecacheTranslucencyLightingVolumeMaterial(
+	TEXT("r.PSOPrecache.TranslucencyLightingVolumeMaterial"), 
+	GPSOPrecacheTranslucencyLightingVolumeMaterial, 
+	TEXT("Precache all possible required Translucency Lighting Volume PSOs for loaded LightMaterials (-1 use engine default)"),
+	ECVF_ReadOnly);
+
+static int32 GPSOPrecacheDitheredLODFadingOutMaskPass = -1;
+static FAutoConsoleVariableRef CVarPSOPrecacheDitheredLODFadingOutMaskPass(
+	TEXT("r.PSOPrecache.DitheredLODFadingOutMaskPass"),
+	GPSOPrecacheDitheredLODFadingOutMaskPass,
+	TEXT("Precache PSOs for DitheredLODFadingOutMaskPass.\n") \
+	TEXT("-1: Use engine defaults.\n") \
+	TEXT(" 0: No PSOs are compiled for this pass.\n") \
+	TEXT(" 1: PSOs are compiled for all primitives which render to depth pass.\n"),
+	ECVF_ReadOnly
+);
+
+static int32 GPSOPrecacheProjectedShadows = -1;
+static FAutoConsoleVariableRef CVarPSOPrecacheProjectedShadows(
+	TEXT("r.PSOPrecache.ProjectedShadows"),
+	GPSOPrecacheProjectedShadows,
+	TEXT("Also Precache PSOs with for projected shadows.") \
+	TEXT("-1: Use engine defaults.\n") \
+	TEXT(" 0: No PSOs are compiled for projected shadows.\n") \
+	TEXT(" 1: PSOs are compiled for all primitives which render to depth pass when VSMs are disabled).\n"),
+	ECVF_ReadOnly
+);
+
+static int32 GPSOPrecacheCustomDepth = -1;
+static FAutoConsoleVariableRef CVarPSOPrecacheCustomDepth(
+	TEXT("r.PSOPrecache.CustomDepth"),
+	GPSOPrecacheCustomDepth,
+	TEXT("Also Precache PSOs with for custom depth pass.") \
+	TEXT("-1: Use engine defaults.\n") \
+	TEXT(" 0: No PSOs are compiled for this pass.\n") \
+	TEXT(" 1: PSOs are compiled for all primitives which explicitly request custom depth rendering (PSOs are precached again on custom depth toggle on component).\n") \
+	TEXT(" 2: PSOs are compiled for all primitives which also request regular depth rendering.\n"),
+	ECVF_ReadOnly
+);
+
+static int32 GPSOPrecacheDeferredDecalMaterials = -1;
+static FAutoConsoleVariableRef CVarPSOPrecacheDeferredDecalMaterials(
+	TEXT("r.PSOPrecache.DeferredDecalMaterials"), 
+	GPSOPrecacheDeferredDecalMaterials, 
+	TEXT("Precache all possible required PSOs for loaded DeferredDecalPass Materials (-1 use engine default)"),
+	ECVF_ReadOnly);
+
+static int32 GPSOPrecacheParticlePrecachingTime = -1;
+static FAutoConsoleVariableRef CVarPSOPrecachingParticlePrecachingTime(
+	TEXT("r.PSOPrecache.ParticlePrecachingTime"),
+	GPSOPrecacheParticlePrecachingTime,
+	TEXT("Controls when PSO precaching happens for particles systems:\n") \
+	TEXT(" -1: Use engine defaults.\n")
+	TEXT("	0: no precaching\n")
+	TEXT("	1: precaching at component proxy creation time only\n")
+	TEXT("	2: precaching at component loading time\n")
+	TEXT("	3: precaching at asset loading time and component time"),
+	ECVF_Default);
+
+static int32 GPSOPrecacheWaterInfoDepthPass = -1;
+static FAutoConsoleVariableRef CVarPSOPrecacheWaterInfoDepthPass(
+	TEXT("r.PSOPrecache.WaterInfoDepthPass"),
+	GPSOPrecacheWaterInfoDepthPass,
+	TEXT("Precaches PSOs for WaterInfoDepthPass (-1 use engine default)\n"),
+	ECVF_ReadOnly
+);
+
+static int32 GPSOPrecacheSlateMaterials = -1;
+static FAutoConsoleVariableRef CVarPSOPrecacheSlateMaterials(
+	TEXT("r.PSOPrecache.SlateMaterials"), 
+	GPSOPrecacheSlateMaterials, 
+	TEXT("Precache all possible required PSOs for loaded Slate Materials (-1 use engine default)"),
+	ECVF_ReadOnly);
+
+
+static void ApplyPSOPrecacheSettingsCVarOverrides(FPSOPrecacheSettings& PSOPrecacheSettings)
+{
+#define SET_CVAR_VALUE(Value, MemberName) \
+	if (Value >= 0) \
+	{ \
+		PSOPrecacheSettings.MemberName = Value; \
+	}
+#define SET_CVAR_VALUE_BOOL(Value, MemberName) \
+	if (Value >= 0) \
+	{ \
+		PSOPrecacheSettings.MemberName = Value > 0; \
+	}
+
+	// Reverse handling of cvar value to keep old naming working
+	if (GPSOPrecacheLightMapPolicyMode >= 0)
+	{
+		PSOPrecacheSettings.bBasePassApplyLightMapPolicies = GPSOPrecacheLightMapPolicyMode == 0;
+	}
+
+	SET_CVAR_VALUE(GPSOPrecacheGlobalShaders, GlobalShaderMode);
+	SET_CVAR_VALUE_BOOL(GPSOPrecacheTranslucencyAllPass, bTranslucencyAll);
+	SET_CVAR_VALUE_BOOL(GPSOPrecacheTranslucencyStandardBeforeDOF, bTranslucencyStandardBeforeDOF);
+	SET_CVAR_VALUE_BOOL(GPSOPrecacheAlphaColorChannel, bBasePassColorWithAlphaChannel);
+	SET_CVAR_VALUE_BOOL(GPSOPrecacheTranslucencyUnderWater, bTranslucencyUnderWater);
+	SET_CVAR_VALUE_BOOL(GPSOPrecachePostProcessingMaterial, bPostProcessingMaterial);
+	SET_CVAR_VALUE_BOOL(GPSOPrecacheTranslucencyLightingVolumeMaterial, bTranslucencyLightingVolume);
+	SET_CVAR_VALUE_BOOL(GPSOPrecacheDitheredLODFadingOutMaskPass, bDitheredLODFadingOutMask);
+	SET_CVAR_VALUE_BOOL(GPSOPrecacheProjectedShadows, bDepthProjectedShadowsState);
+	SET_CVAR_VALUE(GPSOPrecacheCustomDepth, CustomDepthMode);
+	SET_CVAR_VALUE_BOOL(GPSOPrecacheDeferredDecalMaterials, bDeferredDecals);
+	SET_CVAR_VALUE(GPSOPrecacheParticlePrecachingTime, ParticlePrecachingTime);
+	SET_CVAR_VALUE_BOOL(GPSOPrecacheWaterInfoDepthPass, bWaterInfoDepth);
+	SET_CVAR_VALUE_BOOL(GPSOPrecacheSlateMaterials, bSlateMaterials);
+
+#undef SET_CVAR_VALUE
+}
+
+const FPSOPrecacheSettings& UPSOPrecacheSettingsManager::GetSettings()
+{
+	static FPSOPrecacheSettings PSOPrecacheSettings = []()
+	{
+		// Use default value is no specific setting can be found
+		FPSOPrecacheSettings Settings;
+
+		const UPSOPrecacheSettingsManager& Manager = *GetDefault<UPSOPrecacheSettingsManager>();
+
+		// Try and find settings for current shader platform and lowest matching CPU thread count
+		FString ShaderPlatformName = LexToString(GMaxRHIShaderPlatform, true);
+		int32 NumThreads = FPlatformMisc::NumberOfCoresIncludingHyperthreads();
+		for (const FPSOPrecacheSettings& IterSettings : Manager.Settings)
+		{
+			if (IterSettings.ShaderPlatform == ShaderPlatformName && NumThreads <= IterSettings.MaxCPUThreadCount && IterSettings.MaxCPUThreadCount < Settings.MaxCPUThreadCount)
+			{
+				Settings = IterSettings;
+			}
+		}
+
+		ApplyPSOPrecacheSettingsCVarOverrides(Settings);
+
+		return Settings;
+	}();
+	return PSOPrecacheSettings;
+}
diff --git a/Engine/Source/Runtime/Engine/Private/ParticleSystemSceneProxy.cpp b/Engine/Source/Runtime/Engine/Private/ParticleSystemSceneProxy.cpp
--- a/Engine/Source/Runtime/Engine/Private/ParticleSystemSceneProxy.cpp
+++ b/Engine/Source/Runtime/Engine/Private/ParticleSystemSceneProxy.cpp
@@ -12,6 +12,7 @@
 #include "SceneManagement.h"
 #include "SceneView.h"
 #include "UnrealEngine.h"
+#include "PSOPrecacheSettings.h"
 
 #if WITH_EDITOR
 #include "ParticleEmitterInstances.h"
@@ -32,9 +33,6 @@ bool GTrackParticleRenderingStats = false;
 /** Whether to do LOD calculation on GameThread in game */
 extern bool GbEnableGameThreadLODCalculation;
 
-/** When to precache Cascade systems' PSOs */
-extern int32 GCascadePSOPrecachingTime;
-
 ///////////////////////////////////////////////////////////////////////////////
 //	FParticleSystemSceneProxyDesc
 ///////////////////////////////////////////////////////////////////////////////
@@ -536,7 +534,7 @@ FPrimitiveSceneProxy* UParticleSystemComponent::CreateSceneProxy()
 #if UE_WITH_PSO_PRECACHING
 		if (!PSOPrecacheComponentData.IsPSOPrecacheCalled())
 		{
-			if (GCascadePSOPrecachingTime == 3)
+			if (UPSOPrecacheSettingsManager::GetSettings().ParticlePrecachingTime == 1)
 			{
 				Template->PrecachePSOs();
 			}
diff --git a/Engine/Source/Runtime/Engine/Private/Particles/ParticleSystem.cpp b/Engine/Source/Runtime/Engine/Private/Particles/ParticleSystem.cpp
--- a/Engine/Source/Runtime/Engine/Private/Particles/ParticleSystem.cpp
+++ b/Engine/Source/Runtime/Engine/Private/Particles/ParticleSystem.cpp
@@ -19,16 +19,12 @@
 #include "UObject/AssetRegistryTagsContext.h"
 #include "UObject/ObjectSaveContext.h"
 #include "UObject/UObjectIterator.h"
+#include "PSOPrecacheSettings.h"
 
 #include UE_INLINE_GENERATED_CPP_BY_NAME(ParticleSystem)
 
 ////////////////////////////////////////////////////////////////////////////////////////////////////
 
-/** When to precache Cascade systems' PSOs */
-extern int32 GCascadePSOPrecachingTime;
-
-////////////////////////////////////////////////////////////////////////////////////////////////////
-
 UFXSystemAsset::UFXSystemAsset(const FObjectInitializer& ObjectInitializer)
 	: Super(ObjectInitializer)
 {
@@ -472,7 +468,7 @@ void UParticleSystem::PostLoad()
 		}
 	}
 
-	if (GCascadePSOPrecachingTime == 1)
+	if (UPSOPrecacheSettingsManager::GetSettings().ParticlePrecachingTime >= 3)
 	{
 		PrecachePSOs();
 	}
diff --git a/Engine/Source/Runtime/Engine/Private/Particles/ParticleSystemComponent.cpp b/Engine/Source/Runtime/Engine/Private/Particles/ParticleSystemComponent.cpp
--- a/Engine/Source/Runtime/Engine/Private/Particles/ParticleSystemComponent.cpp
+++ b/Engine/Source/Runtime/Engine/Private/Particles/ParticleSystemComponent.cpp
@@ -41,6 +41,7 @@
 #include "UObject/FrameworkObjectVersion.h"
 #include "UObject/UObjectIterator.h"
 #include "UnrealEngine.h"
+#include "PSOPrecacheSettings.h"
 
 #include UE_INLINE_GENERATED_CPP_BY_NAME(ParticleSystemComponent)
 
@@ -121,17 +122,6 @@ FAutoConsoleVariableRef CVarFXSkipZeroDeltaTime(
 	TEXT("This fixes issue like PSA_Velocity aligned sprites, but could cause issues with things that rely on accurate velocities (i.e. TSR)."),
 	ECVF_Default);
 
-int32 GCascadePSOPrecachingTime = 1;
-FAutoConsoleVariableRef CVarCascadePSOPrecachingTime(
-	TEXT("r.PSOPrecache.CascadePrecachingTime"),
-	GCascadePSOPrecachingTime,
-	TEXT("Controls when PSO precaching happens for Cascade systems:\n")
-	TEXT("	0: no precaching\n")
-	TEXT("	1: precaching at asset loading time (default)\n")
-	TEXT("	2: precaching at component loading time\n")
-	TEXT("	3: precaching at component proxy creation time"),
-	ECVF_Default);
-
 /** Whether to allow particle systems to perform work. */
 ENGINE_API bool GIsAllowingParticles = true;
 
@@ -826,7 +816,7 @@ void UParticleSystemComponent::PostLoad()
 		PrimaryComponentTick.bStartWithTickEnabled = false;
 	}
 
-	if (Template && GCascadePSOPrecachingTime == 2)
+	if (Template && UPSOPrecacheSettingsManager::GetSettings().ParticlePrecachingTime > 1)
 	{
 		Template->ConditionalPostLoad();
 		Template->PrecachePSOs();
diff --git a/Engine/Source/Runtime/Engine/Private/ShaderCompiler/ShaderCompiler.cpp b/Engine/Source/Runtime/Engine/Private/ShaderCompiler/ShaderCompiler.cpp
--- a/Engine/Source/Runtime/Engine/Private/ShaderCompiler/ShaderCompiler.cpp
+++ b/Engine/Source/Runtime/Engine/Private/ShaderCompiler/ShaderCompiler.cpp
@@ -58,6 +58,7 @@
 
 #include "SceneTexturesConfig.h"
 #include "PSOPrecacheMaterial.h"
+#include "PSOPrecacheSettings.h"
 
 #if WITH_EDITOR
 #include "Serialization/ArchiveSavePackageDataBuffer.h"
@@ -5025,8 +5026,8 @@ void VerifyGlobalShaders(EShaderPlatform Platform, bool bLoadedFromCacheFile, co
 
 void PrecacheComputePipelineStatesForGlobalShaders(ERHIFeatureLevel::Type FeatureLevel, const ITargetPlatform* TargetPlatform)
 {
-	static IConsoleVariable* PrecacheGlobalShadersCVar = IConsoleManager::Get().FindConsoleVariable(TEXT("r.PSOPrecache.GlobalShaders"));
-	if (PrecacheGlobalShadersCVar == nullptr || PrecacheGlobalShadersCVar->GetInt() == 0)
+	int32 PrecacheGlobalShaders = UPSOPrecacheSettingsManager::GetSettings().GlobalShaderMode;
+	if (PrecacheGlobalShaders == 0)
 	{
 		return;
 	}
@@ -5045,8 +5046,6 @@ void PrecacheComputePipelineStatesForGlobalShaders(ERHIFeatureLevel::Type Featur
 	EShaderPlatform ShaderPlatform = GetFeatureLevelShaderPlatform(FeatureLevel);
 	FGlobalShaderMap* GlobalShaderMap = GetGlobalShaderMap(ShaderPlatform);
 	
-	int32 PrecacheGlobalShaders = PrecacheGlobalShadersCVar->GetInt();
-
 	// some RHIs (OpenGL) can only create shaders on the Render thread. Queue the creation instead of doing it here.
 	TArray<TShaderRef<FShader>> ComputeShadersToPrecache;
 	for (TLinkedList<FShaderType*>::TIterator ShaderTypeIt(FShaderType::GetTypeList()); ShaderTypeIt; ShaderTypeIt.Next())
diff --git a/Engine/Source/Runtime/Engine/Public/PSOPrecacheSettings.h b/Engine/Source/Runtime/Engine/Public/PSOPrecacheSettings.h
new file mode 100644
--- /dev/null
+++ b/Engine/Source/Runtime/Engine/Public/PSOPrecacheSettings.h
@@ -0,0 +1,123 @@
+// Copyright Epic Games, Inc. All Rights Reserved.
+
+/*=============================================================================
+	PSOPrecacheSettings.h
+=============================================================================*/
+
+#pragma once
+
+#include "Engine/EngineTypes.h"
+#include "PSOPrecacheFwd.h"
+
+#include "PSOPrecacheSettings.generated.h"
+
+/**
+ * Settings used to configure what should be precached at runtime
+*/
+USTRUCT()
+struct FPSOPrecacheSettings
+{
+	GENERATED_BODY()
+
+	FPSOPrecacheSettings() {}
+
+	/** Shader platform for which these settings should be applied */
+	UPROPERTY()
+	FString ShaderPlatform;
+
+	/** Maximum number of CPU threads for which the settings should be applied */
+	UPROPERTY()
+	int32 MaxCPUThreadCount = INT32_MAX;
+
+	/**
+	 * Precache global shaders 
+	 *	0: no precaching
+	 *	1: All global compute and subset of global graphics PSOs for which PSO collector is defined (default)
+	 *	2: All global compute and global graphics shaders (only available when r.PSOPrecache.Mode == 1 (Preload shaders))
+	 */
+	UPROPERTY()
+	int GlobalShaderMode = 1;
+
+	/** Precache slate materials (default true) */
+	UPROPERTY()
+	bool bSlateMaterials = true;
+			
+	/** Precache all possible light map policies for base pass shaders, otherwise only LMP_NO_LIGHTMAP (default false) */
+	UPROPERTY()
+	bool bBasePassApplyLightMapPolicies = false;
+
+	/** Precache base pass permutation with alpha channel for planar reflections and scene capture passes (default false) */
+	UPROPERTY()	
+	bool bBasePassColorWithAlphaChannel = false;
+	
+	/** Precache permutations for EMeshPass::TranslucencyAll (default false) */
+	UPROPERTY()	
+	bool bTranslucencyAll = false;
+	
+	/** Precache permutations for EMeshPass::TranslucencyStandard even when marked as AfterDOF - can be needed when AutoBeforeDOFTranslucencyBoundary is set on the view (default false) */
+	UPROPERTY()	
+	bool bTranslucencyStandardBeforeDOF = false;
+	
+	/** Precache translucency under water pass permutations with different render target formats (default false) */
+	UPROPERTY()	
+	bool bTranslucencyUnderWater = false;
+	
+	/**
+	 * Controls when PSO precaching happens for Niagara/Cascade particles systems:
+	 *	0: no precaching
+	 *	1: precaching at component proxy creation time only
+	 *	2: precaching at component loading time
+	 *	3: precaching at asset loading time and component time (default)
+	 */
+	UPROPERTY()	
+	int32 ParticlePrecachingTime = 3;
+	
+	/** Precache deferred decal pass (default true) **/
+	UPROPERTY()
+	bool bDeferredDecals = true;
+	
+	// Precache depth pass permutations with different depth state and render targets used for projected shadows (default false) **/
+	UPROPERTY()
+	bool bDepthProjectedShadowsState = false;
+	
+	/** 
+	 * Precache custom depth PSO permutations
+	 * 	0: No PSOs are compiled
+	 *  1: PSOs are compiled for all primitives which explicitly request custom depth rendering (PSOs are precached again on custom depth toggle on component) (default)
+	 *  2: PSOs are compiled for all primitives which also request regular depth rendering
+	 **/	
+	UPROPERTY()
+	int32 CustomDepthMode = 1;
+	
+	/** Precache permutations for EMeshPass::WaterInfoTextureDepthPass (default false) **/
+	UPROPERTY()
+	bool bWaterInfoDepth = false;
+	
+	/** Precache permutations for EMeshPass::DitheredLODFadingOutMaskPass (default false) **/
+	UPROPERTY()
+	bool bDitheredLODFadingOutMask = false;
+	
+	/** Precache permutations for PostProcess material types (default true) **/
+	UPROPERTY()
+	bool bPostProcessingMaterial = true;
+	
+	/** Precache permutations for LightType material types which render to translucency light volume textures (default true) **/
+	UPROPERTY()
+	bool bTranslucencyLightingVolume = true;	
+};
+
+
+UCLASS(MinimalAPI, defaultconfig, Config = Engine)
+class UPSOPrecacheSettingsManager : public UObject
+{
+	GENERATED_BODY()
+
+public:
+
+	static ENGINE_API const FPSOPrecacheSettings& GetSettings();
+
+private:
+	
+	UPROPERTY(globalconfig)
+	TArray<FPSOPrecacheSettings> Settings;
+};
\ No newline at end of file
diff --git a/Engine/Source/Runtime/Renderer/Private/BasePassRendering.cpp b/Engine/Source/Runtime/Renderer/Private/BasePassRendering.cpp
--- a/Engine/Source/Runtime/Renderer/Private/BasePassRendering.cpp
+++ b/Engine/Source/Runtime/Renderer/Private/BasePassRendering.cpp
@@ -31,6 +31,7 @@
 #include "ShaderCompiler.h"
 #include "GBufferInfo.h"
 #include "LocalFogVolumeRendering.h"
+#include "PSOPrecacheSettings.h"
 
 #include "BasePassRendering.inl"
 
@@ -92,38 +93,6 @@ static TAutoConsoleVariable<int32> CVarClearGBufferDBeforeBasePass(
 	TEXT("Whether to clear GBuffer D before basepass"),
 	ECVF_RenderThreadSafe);
 
-static TAutoConsoleVariable<int32> CVarPSOPrecacheLightMapPolicyMode(
-	TEXT("r.PSOPrecache.LightMapPolicyMode"),
-	1,
-	TEXT("Defines which light map policies should be checked during PSO precaching of the base pass.\n") \
-	TEXT(" 0: All possible LMP will be checked.\n") \
-	TEXT(" 1: Only LMP_NO_LIGHTMAP will be precached (default).\n"),
-	ECVF_ReadOnly
-);
-
-static TAutoConsoleVariable<int32> CVarPSOPrecacheTranslucencyAllPass(
-	TEXT("r.PSOPrecache.TranslucencyAllPass"),
-	0,
-	TEXT("Precache PSOs for TranslucencyAll pass.\n") \
-	TEXT(" 0: No PSOs are compiled for this pass (default).\n") \
-	TEXT(" 1: PSOs are compiled for all primitives which render to a translucency pass.\n"),
-	ECVF_ReadOnly
-);
-
-static TAutoConsoleVariable<int32> CVarPSOPrecacheAlphaColorChannel(
-	TEXT("r.PSOPrecache.PrecacheAlphaColorChannel"),
-	0,
-	TEXT("Also Precache PSOs with scene color alpha channel enabled. Planar reflections and scene captures use this for compositing into a different scene later."),
-	ECVF_ReadOnly
-);
-
-static TAutoConsoleVariable<int32> CVarPSOPrecacheNotSeparateTranslucency(
-	TEXT("r.PSOPrecache.PrecacheNotSeparateTranslucency"),
-	0,
-	TEXT("TODO."),
-	ECVF_ReadOnly
-);
-
 static TAutoConsoleVariable<bool> CVarAllow128BitBPPSCompilation(
 	TEXT("r.128BitBPPSCompilation.Allow"),
 	true,
@@ -1846,7 +1815,7 @@ void FBasePassMeshProcessor::CollectPSOInitializersForLMPolicy(
 				ERenderTargetLoadAction::ELoad, FExclusiveDepthStencil::DepthRead_StencilWrite, RenderTargetsInfo);
 		}
 
-		if (!bRenderInSeparateTranslucency || CVarPSOPrecacheNotSeparateTranslucency->GetInt() > 0)
+		if (!bRenderInSeparateTranslucency || UPSOPrecacheSettingsManager::GetSettings().bTranslucencyUnderWater)
 		{
 			AddGraphicsPipelineStateInitializer(
 				VertexFactoryData,
@@ -2107,7 +2076,7 @@ bool FBasePassMeshProcessor::ShouldDraw(const FMaterial& Material)
 			{
 			case ETranslucencyPass::TPT_TranslucencyStandard:
 				bShouldDraw = !Material.IsTranslucencyAfterDOFEnabled() && !Material.IsTranslucencyAfterMotionBlurEnabled();
-				if (AutoBeforeDOFTranslucencyBoundary > 0.0f || !Scene) //Always include in PSO precache even if AutoBeforeDOFTranslucencyBoundary == 0
+				if (AutoBeforeDOFTranslucencyBoundary > 0.0f || UPSOPrecacheSettingsManager::GetSettings().bTranslucencyStandardBeforeDOF)
 				{
 					bShouldDraw = bShouldDraw || (Material.IsTranslucencyAfterDOFEnabled() && BlendMode != BLEND_ColoredTransmittanceOnly);
 					bShouldDraw = bShouldDraw || (Material.IsTranslucencyAfterDOFEnabled() && (Material.IsDualBlendingEnabled() || IsModulateBlendMode(BlendMode)));
@@ -2161,7 +2130,7 @@ bool FBasePassMeshProcessor::TryAddMeshBatch(const FMeshBatch& RESTRICT MeshBatc
 	const ERasterizerFillMode MeshFillMode = ComputeMeshFillMode(Material, OverrideSettings);
 	const ERasterizerCullMode MeshCullMode = ComputeMeshCullMode(Material, OverrideSettings);
 	const bool bIsValidTranslucentNonDeferredDecal = PrimitiveSceneProxy && bIsTranslucent && !Material.IsDeferredDecal();
-	
+
 	bool bShouldDraw = false;
 	if (AutoBeforeDOFTranslucencyBoundary > 0.0f && bIsValidTranslucentNonDeferredDecal)
 	{
@@ -2468,15 +2437,12 @@ TArray<ELightMapPolicyType, TInlineAllocator<2>> FBasePassMeshProcessor::GetUnif
 	const bool bAllowIndirectLightingCache = true;// Scene&& Scene->PrecomputedLightVolumes.Num() > 0;
 	const bool bUseVolumetricLightmap = true;// Scene&& Scene->VolumetricLightmapSceneData.HasData();
 
-	static const auto CVarSupportLightMapPolicyMode = IConsoleManager::Get().FindTConsoleVariableDataInt(TEXT("r.PSOPrecache.LightMapPolicyMode"));
-	const bool bNoLightMapOnlyMode = (CVarSupportLightMapPolicyMode && CVarSupportLightMapPolicyMode->GetValueOnAnyThread() == 1);
-
 	TArray<ELightMapPolicyType, TInlineAllocator<2>> UniformLightMapPolicyTypes;
 
 	// always add the fallback no lightmap mode
 	UniformLightMapPolicyTypes.Add(LMP_NO_LIGHTMAP);
 
-	if (!bNoLightMapOnlyMode)
+	if (UPSOPrecacheSettingsManager::GetSettings().bBasePassApplyLightMapPolicies)
 	{
 		// Simplified version of GetUniformLightMapPolicyType with worst case to collect all types which could be enabled at runtime	
 
@@ -2544,7 +2510,7 @@ void FBasePassMeshProcessor::CollectPSOInitializers(const FSceneTexturesConfig&
 	}
 
 	// PSO precaching enabled for TranslucencyAll
-	if (MeshPassType == EMeshPass::TranslucencyAll && CVarPSOPrecacheTranslucencyAllPass.GetValueOnAnyThread() == 0)
+	if (MeshPassType == EMeshPass::TranslucencyAll && !UPSOPrecacheSettingsManager::GetSettings().bTranslucencyAll)
 	{
 		return;
 	}
@@ -2588,10 +2554,7 @@ void FBasePassMeshProcessor::CollectPSOInitializersForSkyLight(
 	const bool bIsTranslucent = IsTranslucentBlendMode(Material);
 	const bool bIsLitMaterial = ShadingModels.IsLit();
 	
-	static const auto CVarSupportLightMapPolicyMode = IConsoleManager::Get().FindTConsoleVariableDataInt(TEXT("r.PSOPrecache.LightMapPolicyMode"));
-	const bool bNoLightMapOnlyMode = (CVarSupportLightMapPolicyMode && CVarSupportLightMapPolicyMode->GetValueOnAnyThread() == 1);
-
-	if (!bNoLightMapOnlyMode && bIsLitMaterial && bIsTranslucent)
+	if (UPSOPrecacheSettingsManager::GetSettings().bBasePassApplyLightMapPolicies && bIsLitMaterial && bIsTranslucent)
 	{
 		const bool bAllowStaticLighting = IsStaticLightingAllowed();
 
diff --git a/Engine/Source/Runtime/Renderer/Private/BasePassRendering.h b/Engine/Source/Runtime/Renderer/Private/BasePassRendering.h
--- a/Engine/Source/Runtime/Renderer/Private/BasePassRendering.h
+++ b/Engine/Source/Runtime/Renderer/Private/BasePassRendering.h
@@ -38,6 +38,7 @@
 #include "RenderUtils.h"
 #include "SceneTexturesConfig.h"
 #include "HeterogeneousVolumes/HeterogeneousVolumes.h"
+#include "PSOPrecacheSettings.h"
 
 class FScene;
 
@@ -1013,8 +1014,7 @@ public:
 			PSOInitializers);
 
 		// Planar reflections and scene captures use scene color alpha to keep track of where content has been rendered, for compositing into a different scene later
-		static TConsoleVariableData<int32>* CVar = IConsoleManager::Get().FindTConsoleVariableDataInt(TEXT("r.PSOPrecache.PrecacheAlphaColorChannel"));
-		if (bPrecacheAlphaColorChannel && CVar && CVar->GetValueOnAnyThread() > 0)
+		if (bPrecacheAlphaColorChannel && UPSOPrecacheSettingsManager::GetSettings().bBasePassColorWithAlphaChannel)
 		{
 			FGraphicsPipelineRenderTargetsInfo AlphaColorRenderTargetsInfo = RenderTargetsInfo;
 
diff --git a/Engine/Source/Runtime/Renderer/Private/CompositionLighting/PostProcessMeshDecals.cpp b/Engine/Source/Runtime/Renderer/Private/CompositionLighting/PostProcessMeshDecals.cpp
--- a/Engine/Source/Runtime/Renderer/Private/CompositionLighting/PostProcessMeshDecals.cpp
+++ b/Engine/Source/Runtime/Renderer/Private/CompositionLighting/PostProcessMeshDecals.cpp
@@ -22,6 +22,7 @@
 #include "MeshPassProcessor.inl"
 #include "SimpleMeshDrawCommandPass.h"
 #include "ComponentRecreateRenderStateContext.h"
+#include "PSOPrecacheSettings.h"
 
 static TAutoConsoleVariable<int32> CVarParallelMeshDecal(
 	TEXT("r.ParallelMeshDecal"),
@@ -30,21 +31,6 @@ static TAutoConsoleVariable<int32> CVarParallelMeshDecal(
 	ECVF_RenderThreadSafe
 );
 
-static int32 GDeferredDecalPassMaterialPSOPrecache = 1;
-static FAutoConsoleVariableRef CVarGDeferredDecalPassMaterialPSOPrecache(
-	TEXT("r.PSOPrecache.DeferredDecalPassMaterial"), 
-	GDeferredDecalPassMaterialPSOPrecache, 
-	TEXT("Precache all possible required PSOs for loaded DeferredDecalPass Materials."),
-	ECVF_ReadOnly);
-
-static int32 GDeferredDecalMeshMaterialPSOPrecache = 1;
-static FAutoConsoleVariableRef CVarGDeferredDecalMeshMaterialPSOPrecache(
-	TEXT("r.PSOPrecache.DeferredDecalMeshMaterial"), 
-	GDeferredDecalMeshMaterialPSOPrecache, 
-	TEXT("Precache all possible required PSOs for loaded DeferredDecalMesh Materials."),
-	ECVF_ReadOnly);
-
-
 static constexpr uint32 GetMeshDecalStencilRef()
 {
 	return GET_STENCIL_BIT_MASK(RECEIVE_DECAL, 1);
@@ -466,7 +452,7 @@ void FMeshDecalMeshProcessor::CollectPSOInitializers(
 	const FPSOPrecacheParams& PreCacheParams, 
 	TArray<FPSOPrecacheData>& PSOInitializers)
 {
-	if (!Material.IsDeferredDecal())
+	if (!Material.IsDeferredDecal() || !UPSOPrecacheSettingsManager::GetSettings().bDeferredDecals)
 	{
 		return;
 	}
@@ -491,17 +477,11 @@ void FMeshDecalMeshProcessor::CollectPSOInitializers(
 			continue;
 		}
 		
-		if (GDeferredDecalPassMaterialPSOPrecache)
-		{
-			// Collect decal pass PSOs
-			CollectDeferredDecalPassPSOInitializers(PSOCollectorIndex, FeatureLevel, SceneTexturesConfig, Material, LocalDecalRenderStage, PSOInitializers);
-		}
-
-		if (GDeferredDecalMeshMaterialPSOPrecache)
-		{
-			// Collect decal mesh PSOs
-			CollectDeferredDecalMeshPSOInitializers(SceneTexturesConfig, VertexFactoryData, PreCacheParams, Material, DecalBlendDesc, LocalDecalRenderStage, PSOInitializers);
-		}
+		// Collect decal pass PSOs
+		CollectDeferredDecalPassPSOInitializers(PSOCollectorIndex, FeatureLevel, SceneTexturesConfig, Material, LocalDecalRenderStage, PSOInitializers);
+		
+		// Collect decal mesh PSOs
+		CollectDeferredDecalMeshPSOInitializers(SceneTexturesConfig, VertexFactoryData, PreCacheParams, Material, DecalBlendDesc, LocalDecalRenderStage, PSOInitializers);
 	}
 }
 
diff --git a/Engine/Source/Runtime/Renderer/Private/CustomDepthRendering.cpp b/Engine/Source/Runtime/Renderer/Private/CustomDepthRendering.cpp
--- a/Engine/Source/Runtime/Renderer/Private/CustomDepthRendering.cpp
+++ b/Engine/Source/Runtime/Renderer/Private/CustomDepthRendering.cpp
@@ -11,6 +11,7 @@
 #include "MeshPassProcessor.inl"
 #include "UnrealEngine.h"
 #include "Nanite/NaniteComposition.h"
+#include "PSOPrecacheSettings.h"
 
 static TAutoConsoleVariable<int32> CVarCustomDepthOrder(
 	TEXT("r.CustomDepth.Order"),
@@ -32,16 +33,6 @@ static TAutoConsoleVariable<bool> CVarCustomDepthEnableFastClear(
 	TEXT("Enable HTile on the custom depth buffer (default:false).\n"),
 	ECVF_RenderThreadSafe);
 
-static TAutoConsoleVariable<int32> CVarPSOPrecacheCustomDepth(
-	TEXT("r.PSOPrecache.CustomDepth"),
-	1,
-	TEXT("Also Precache PSOs with for custom depth pass.")  \
-	TEXT(" 0: No PSOs are compiled for this pass.\n") \
-	TEXT(" 1: PSOs are compiled for all primitives which explicitly request custom depth rendering (PSOs are precached again on custom depth toggle on component) (default).\n") \
-	TEXT(" 2: PSOs are compiled for all primitives which also request regular depth rendering.\n"),
-	ECVF_ReadOnly
-);
-
 DECLARE_DWORD_COUNTER_STAT(TEXT("Nanite Custom Depth Instances"), STAT_NaniteCustomDepthInstances, STATGROUP_Nanite);
 
 DECLARE_GPU_DRAWCALL_STAT_NAMED(CustomDepth, TEXT("Custom Depth"));
@@ -664,8 +655,8 @@ bool FCustomDepthPassMeshProcessor::Process(
 
 void FCustomDepthPassMeshProcessor::CollectPSOInitializers(const FSceneTexturesConfig& SceneTexturesConfig, const FMaterial& Material, const FPSOPrecacheVertexFactoryData& VertexFactoryData, const FPSOPrecacheParams& PreCacheParams, TArray<FPSOPrecacheData>& PSOInitializers)
 {
-	int32 CustomDepthPrecacheMode = CVarPSOPrecacheCustomDepth.GetValueOnAnyThread();
-	if (CustomDepthPrecacheMode == 0)
+	int32 CustomDepthPrecacheMode = UPSOPrecacheSettingsManager::GetSettings().CustomDepthMode;
+	if (CustomDepthPrecacheMode <= 0)
 	{
 		return;
 	}
diff --git a/Engine/Source/Runtime/Renderer/Private/DepthRendering.cpp b/Engine/Source/Runtime/Renderer/Private/DepthRendering.cpp
--- a/Engine/Source/Runtime/Renderer/Private/DepthRendering.cpp
+++ b/Engine/Source/Runtime/Renderer/Private/DepthRendering.cpp
@@ -37,6 +37,7 @@
 #include "SimpleMeshDrawCommandPass.h"
 #include "UnrealEngine.h"
 #include "DepthCopy.h"
+#include "PSOPrecacheSettings.h"
 
 static TAutoConsoleVariable<int32> CVarParallelPrePass(
 	TEXT("r.ParallelPrePass"),
@@ -69,24 +70,6 @@ static TAutoConsoleVariable<int32> CVarStencilForLODDither(
 	TEXT("Forces a full prepass when enabled."),
 	ECVF_RenderThreadSafe | ECVF_ReadOnly);
 
-static TAutoConsoleVariable<int32> CVarPSOPrecacheDitheredLODFadingOutMaskPass(
-	TEXT("r.PSOPrecache.DitheredLODFadingOutMaskPass"),
-	0,
-	TEXT("Precache PSOs for DitheredLODFadingOutMaskPass.\n") \
-	TEXT(" 0: No PSOs are compiled for this pass (default).\n") \
-	TEXT(" 1: PSOs are compiled for all primitives which render to depth pass.\n"),
-	ECVF_ReadOnly
-);
-
-static TAutoConsoleVariable<int32> CVarPSOPrecacheProjectedShadows(
-	TEXT("r.PSOPrecache.ProjectedShadows"),
-	0,
-	TEXT("Also Precache PSOs with for projected shadows.")  \
-	TEXT(" 0: No PSOs are compiled for projected shadows.\n") \
-	TEXT(" 1: PSOs are compiled for all primitives which render to depth pass when VSMs are disabled (default).\n"),
-	ECVF_ReadOnly
-);
-
 extern bool IsHMDHiddenAreaMaskActive();
 
 FDepthPassInfo GetDepthPassInfo(const FScene* Scene)
@@ -920,7 +903,7 @@ void FDepthPassMeshProcessor::CollectPSOInitializersInternal(
 	// Also cache with project shadow depth stencil state (see FProjectedShadowInfo::SetupMeshDrawCommandsForProjectionStenciling)
 	EShaderPlatform ShaderPlatform = GetFeatureLevelShaderPlatform(SceneTexturesConfig.FeatureLevel);
 	bool bUseVirtualShadowMaps = UseVirtualShadowMaps(ShaderPlatform, SceneTexturesConfig.FeatureLevel);
-	if (CVarPSOPrecacheProjectedShadows.GetValueOnAnyThread() > 0 && !bUseVirtualShadowMaps)
+	if (UPSOPrecacheSettingsManager::GetSettings().bDepthProjectedShadowsState && !bUseVirtualShadowMaps)
 	{
 		// Set stencil to one.
 		DrawRenderState.SetDepthStencilState(
@@ -1120,7 +1103,7 @@ void FDepthPassMeshProcessor::CollectPSOInitializers(const FSceneTexturesConfig&
 	}
 
 	// PSO precaching enabled for DitheredLODFadingOutMaskPass
-	if (MeshPassType == EMeshPass::DitheredLODFadingOutMaskPass && CVarPSOPrecacheDitheredLODFadingOutMaskPass.GetValueOnAnyThread() == 0)
+	if (MeshPassType == EMeshPass::DitheredLODFadingOutMaskPass && !UPSOPrecacheSettingsManager::GetSettings().bDitheredLODFadingOutMask)
 	{
 		return;
 	}
diff --git a/Engine/Source/Runtime/Renderer/Private/MobileBasePass.cpp b/Engine/Source/Runtime/Renderer/Private/MobileBasePass.cpp
--- a/Engine/Source/Runtime/Renderer/Private/MobileBasePass.cpp
+++ b/Engine/Source/Runtime/Renderer/Private/MobileBasePass.cpp
@@ -1018,7 +1018,6 @@ void FMobileBasePassMeshProcessor::CollectPSOInitializers(const FSceneTexturesCo
 		
 	if (bTranslucentBasePass)
 	{
-		static IConsoleVariable* PSOPrecacheTranslucencyAllPass = IConsoleManager::Get().FindConsoleVariable(TEXT("r.PSOPrecache.TranslucencyAllPass"));
 		static IConsoleVariable* CVarSeparateTranslucency = IConsoleManager::Get().FindConsoleVariable(TEXT("r.SeparateTranslucency"));
 		if (CVarSeparateTranslucency->GetInt() == 0)
 		{
@@ -1028,7 +1027,7 @@ void FMobileBasePassMeshProcessor::CollectPSOInitializers(const FSceneTexturesCo
 				return;
 			}
 		}
-		else if (MeshPassType == EMeshPass::TranslucencyAll && PSOPrecacheTranslucencyAllPass->GetInt() == 0)
+		else if (MeshPassType == EMeshPass::TranslucencyAll && !UPSOPrecacheSettingsManager::GetSettings().bTranslucencyAll)
 		{
 			// PSO precaching is disabled for TranslucencyAll while SeparateTranslucency is active
 			return;
diff --git a/Engine/Source/Runtime/Renderer/Private/PostProcess/PostProcessMaterial.cpp b/Engine/Source/Runtime/Renderer/Private/PostProcess/PostProcessMaterial.cpp
--- a/Engine/Source/Runtime/Renderer/Private/PostProcess/PostProcessMaterial.cpp
+++ b/Engine/Source/Runtime/Renderer/Private/PostProcess/PostProcessMaterial.cpp
@@ -31,6 +31,7 @@
 #include "SingleLayerWaterRendering.h"
 #include "Engine/NeuralProfile.h"
 #include "PathTracing.h"
+#include "PSOPrecacheSettings.h"
 
 namespace
 {
@@ -57,13 +58,6 @@ TAutoConsoleVariable<int32> CVarPostProcessingDisableMaterials(
 	TEXT(" Allows to disable post process materials. \n"),
 	ECVF_Scalability | ECVF_RenderThreadSafe);
 
-static int32 GPostProcessingMaterialPSOPrecache = 1;
-static FAutoConsoleVariableRef CVarGPostProcessingMaterialPSOPrecache(
-	TEXT("r.PSOPrecache.PostProcessingMaterial"), 
-	GPostProcessingMaterialPSOPrecache, 
-	TEXT("Precache all possible required PSOs for loaded PostProcessing Materials."),
-	ECVF_ReadOnly);
-
 static FName NAME_SceneColor("SceneColor");
 
 static bool IsPostProcessStencilTestAllowed()
@@ -1547,7 +1541,7 @@ void FPostProcessMaterialPSOCollector::CollectPSOInitializers(
 	TArray<FPSOPrecacheData>& PSOInitializers
 )
 {
-	if (!Material.IsPostProcessMaterial() || GPostProcessingMaterialPSOPrecache == 0)
+	if (!Material.IsPostProcessMaterial() || !UPSOPrecacheSettingsManager::GetSettings().bPostProcessingMaterial)
 	{
 		return;
 	}
diff --git a/Engine/Source/Runtime/Renderer/Private/TranslucentLighting.cpp b/Engine/Source/Runtime/Renderer/Private/TranslucentLighting.cpp
--- a/Engine/Source/Runtime/Renderer/Private/TranslucentLighting.cpp
+++ b/Engine/Source/Runtime/Renderer/Private/TranslucentLighting.cpp
@@ -59,6 +59,7 @@
 #include "BasePassRendering.h"
 #include "LightFunctionAtlas.h"
 #include "BlueNoise.h"
+#include "PSOPrecacheSettings.h"
 
 using namespace LightFunctionAtlas;
 
@@ -227,13 +228,6 @@ static FAutoConsoleVariableRef CVarTranslucencyLightingVolumeDebugFreezePosition
 	ECVF_RenderThreadSafe
 );
 
-static int32 GTranslucencyLightingVolumeMaterialPSOPrecache = 1;
-static FAutoConsoleVariableRef CVarGTranslucencyLightingVolumeMaterialPSOPrecache(
-	TEXT("r.PSOPrecache.TranslucencyLightingVolumeMaterial"), 
-	GTranslucencyLightingVolumeMaterialPSOPrecache, 
-	TEXT("Precache all possible required Translucency Lighting Volume PSOs for loaded LightMaterials."),
-	ECVF_ReadOnly);
-
 static const TCHAR* TranslucentLightingMaterialPSOCollectorName = TEXT("TranslucentLightingMaterialPSOCollector");
 
 float GetTranslucencyLightingVolumePositionOffsetRadius()
@@ -3067,7 +3061,7 @@ void FTranslucentLightingMaterialPSOCollector::CollectPSOInitializers(
 	TArray<FPSOPrecacheData>& PSOInitializers
 )
 {
-	if (!Material.IsLightFunction() || GTranslucencyLightingVolumeMaterialPSOPrecache == 0)
+	if (!Material.IsLightFunction() || !UPSOPrecacheSettingsManager::GetSettings().bTranslucencyLightingVolume)
 	{
 		return;
 	}
diff --git a/Engine/Source/Runtime/Renderer/Private/WaterInfoTextureRendering.cpp b/Engine/Source/Runtime/Renderer/Private/WaterInfoTextureRendering.cpp
--- a/Engine/Source/Runtime/Renderer/Private/WaterInfoTextureRendering.cpp
+++ b/Engine/Source/Runtime/Renderer/Private/WaterInfoTextureRendering.cpp
@@ -11,6 +11,7 @@
 #include "BasePassRendering.h"
 #include "MobileBasePassRendering.h"
 #include "RenderCaptureInterface.h"
+#include "PSOPrecacheSettings.h"
 
 DECLARE_GPU_DRAWCALL_STAT(WaterInfoTexture);
 
@@ -32,13 +33,6 @@ static TAutoConsoleVariable<int32> CVarRenderCaptureNextWaterInfoDraws(
 	TEXT("Enable capturing of the water info texture for the next N draws"),
 	ECVF_RenderThreadSafe);
 
-static TAutoConsoleVariable<int32> CVarPSOPrecacheWaterInfoDepthPass(
-	TEXT("r.PSOPrecache.WaterInfoDepthPass"),
-	0,
-	TEXT("Precaches PSOs for WaterInfoDepthPass.\n"),
-	ECVF_ReadOnly
-);
-
 BEGIN_SHADER_PARAMETER_STRUCT(FWaterInfoTexturePassParameters, )
 	SHADER_PARAMETER_STRUCT_INCLUDE(FViewShaderParameters, View)
 	SHADER_PARAMETER_STRUCT_INCLUDE(FInstanceCullingDrawParams, InstanceCullingDrawParams)
@@ -554,8 +548,7 @@ bool FWaterInfoTextureDepthPassMeshProcessor::Process(
 void FWaterInfoTextureDepthPassMeshProcessor::CollectPSOInitializers(const FSceneTexturesConfig& SceneTexturesConfig, const FMaterial& Material, const FPSOPrecacheVertexFactoryData& VertexFactoryData, const FPSOPrecacheParams& PreCacheParams, TArray<FPSOPrecacheData>& PSOInitializers)
 {
 	// Early out if disabled
-	int32 CollectPSOs = CVarPSOPrecacheWaterInfoDepthPass.GetValueOnAnyThread();
-	if (CollectPSOs == 0)
+	if (!UPSOPrecacheSettingsManager::GetSettings().bWaterInfoDepth)
 	{
 		return;
 	}
diff --git a/Engine/Source/Runtime/SlateRHIRenderer/Private/SlateRHIRenderingPolicy.cpp b/Engine/Source/Runtime/SlateRHIRenderer/Private/SlateRHIRenderingPolicy.cpp
--- a/Engine/Source/Runtime/SlateRHIRenderer/Private/SlateRHIRenderingPolicy.cpp
+++ b/Engine/Source/Runtime/SlateRHIRenderer/Private/SlateRHIRenderingPolicy.cpp
@@ -36,6 +36,7 @@
 #include "Containers/StaticBitArray.h"
 #include "MeshPassProcessor.h"
 #include "PSOPrecacheValidation.h"
+#include "PSOPrecacheSettings.h"
 #include "VT/VirtualTextureFeedbackResource.h"
 
 extern void UpdateNoiseTextureParameters(FViewUniformShaderParameters& ViewUniformShaderParameters);
@@ -50,13 +51,6 @@ DECLARE_DWORD_COUNTER_STAT(TEXT("Num Vertices"), STAT_SlateVertexCount, STATGROU
 DECLARE_DWORD_COUNTER_STAT(TEXT("Clips (Scissor)"), STAT_SlateScissorClips, STATGROUP_Slate);
 DECLARE_DWORD_COUNTER_STAT(TEXT("Clips (Stencil)"), STAT_SlateStencilClips, STATGROUP_Slate);
 
-static int32 GSlateMaterialPSOPrecache = 1;
-static FAutoConsoleVariableRef CVarGSlateMaterialPSOPrecache(
-	TEXT("r.PSOPrecache.SlateMaterials"), 
-	GSlateMaterialPSOPrecache, 
-	TEXT("Precache all possible required PSOs for loaded Slate Materials."),
-	ECVF_ReadOnly);
-
 static TAutoConsoleVariable<int32> CVarSlateMaterialShaderMode(
 	TEXT("slate.MaterialShaderMode"),
 	0,
@@ -1381,7 +1375,7 @@ void DrawSlateRenderBatch(
 		State.GraphicsPSOInit.PrimitiveType = GetRHIPrimitiveType(RenderBatch.DrawPrimitiveType);
 		
 #if PSO_PRECACHING_VALIDATE
-		if (PSOCollectorStats::IsFullPrecachingValidationEnabled() && GSlateMaterialPSOPrecache > 0)
+		if (PSOCollectorStats::IsFullPrecachingValidationEnabled() && UPSOPrecacheSettingsManager::GetSettings().bSlateMaterials)
 		{
 			static const int32 MaterialPSOCollectorIndex = FPSOCollectorCreateManager::GetIndex(GetFeatureLevelShadingPath(GMaxRHIFeatureLevel), SlateMaterialPSOCollectorName);
 			// Material Render Proxy is not cached in the state, but could be done for better debugging if needed
@@ -1431,8 +1425,7 @@ void DrawSlateRenderBatch(
 		State.GraphicsPSOInit.PrimitiveType = GetRHIPrimitiveType(RenderBatch.DrawPrimitiveType);
 
 #if PSO_PRECACHING_VALIDATE
-		static IConsoleVariable* PrecacheGlobalShadersCVar = IConsoleManager::Get().FindConsoleVariable(TEXT("r.PSOPrecache.GlobalShaders"));
-		if (PSOCollectorStats::IsFullPrecachingValidationEnabled() && PrecacheGlobalShadersCVar && PrecacheGlobalShadersCVar->GetInt() > 0)
+		if (PSOCollectorStats::IsFullPrecachingValidationEnabled() && UPSOPrecacheSettingsManager::GetSettings().GlobalShaderMode > 0)
 		{			
 			static const int32 GlobalPSOCollectorIndex = FGlobalPSOCollectorManager::GetIndex(SlateGlobalPSOCollectorName);
 			PSOCollectorStats::CheckGlobalGraphicsPipelineStateInCache(State.GraphicsPSOInit, GlobalPSOCollectorIndex);
@@ -1928,7 +1921,7 @@ void FSlateMaterialPSOCollector::CollectPSOInitializers(
 	TArray<FPSOPrecacheData>& PSOInitializers
 )
 {
-	if (!Material.IsUIMaterial() || GSlateMaterialPSOPrecache == 0)
+	if (!Material.IsUIMaterial() || !UPSOPrecacheSettingsManager::GetSettings().bSlateMaterials)
 	{
 		return;
 	}
