Add mesh deformer usage flag.
The usage flag is used for more granular local vertex factory compilation.
The flag is validated on creation of skinned mesh proxies that use deformers.
If the material has bAutomaticallySetUsageInEditor then the usage flag will get auto set in editor.
If the usage flag is disabled then deformer meshes will use the default material.
For existing materials the flag will default to the current SkeletalMesh usage. For new materials it will default off.
#Synced-CL 49065038

diff --git a/Engine/Source/Runtime/Core/Public/UObject/FortniteMainBranchObjectVersions.inl b/Engine/Source/Runtime/Core/Public/UObject/FortniteMainBranchObjectVersions.inl
--- a/Engine/Source/Runtime/Core/Public/UObject/FortniteMainBranchObjectVersions.inl
+++ b/Engine/Source/Runtime/Core/Public/UObject/FortniteMainBranchObjectVersions.inl
@@ -721,5 +721,8 @@ FORTNITE_MAIN_VERSION(OptionalQuantizePerFaceBoundsInChaosTriMeshBVH, 235)
 // Add Dataflow Graph Topology Guid to UE::Dataflow:FGraph
 FORTNITE_MAIN_VERSION(DataflowAddGraphTopologyGuid, 236)
 
+// Add material usage flag for mesh deformers
+FORTNITE_MAIN_VERSION(AddMeshDeformerUsageFlag, 237)
+
 // -----<new versions can be added above this line>-------------------------------------------------
 #undef FORTNITE_MAIN_VERSION
diff --git a/Engine/Source/Runtime/Engine/Private/GPUSkinVertexFactory.cpp b/Engine/Source/Runtime/Engine/Private/GPUSkinVertexFactory.cpp
--- a/Engine/Source/Runtime/Engine/Private/GPUSkinVertexFactory.cpp
+++ b/Engine/Source/Runtime/Engine/Private/GPUSkinVertexFactory.cpp
@@ -527,9 +527,7 @@ bool FGPUBaseSkinVertexFactory::GetAlwaysUseDeformerForUnlimitedBoneInfluences(E
 	auto InnerFunc = [](EShaderPlatform Platform)
 	{
 		static FShaderPlatformCachedIniValue<bool> UseDeformerForUBICVar(TEXT("r.GPUSkin.AlwaysUseDeformerForUnlimitedBoneInfluences"));
-		const IMeshDeformerProvider* MeshDeformerProvider = IMeshDeformerProvider::Get();
-
-		return MeshDeformerProvider && MeshDeformerProvider->IsSupported(Platform) && UseDeformerForUBICVar.Get(Platform);
+		return AreMeshDeformersAvailable(Platform) && UseDeformerForUBICVar.Get(Platform);
 	};
 
 #if WITH_EDITOR
diff --git a/Engine/Source/Runtime/Engine/Private/LocalVertexFactory.cpp b/Engine/Source/Runtime/Engine/Private/LocalVertexFactory.cpp
--- a/Engine/Source/Runtime/Engine/Private/LocalVertexFactory.cpp
+++ b/Engine/Source/Runtime/Engine/Private/LocalVertexFactory.cpp
@@ -5,7 +5,6 @@
 =============================================================================*/
 
 #include "LocalVertexFactory.h"
-#include "Animation/MeshDeformerGeometry.h"
 #include "MeshBatch.h"
 #include "MeshDrawShaderBindings.h"
 #include "SkeletalRenderPublic.h"
@@ -314,8 +313,9 @@ bool FLocalVertexFactory::ShouldCompilePermutation(const FVertexFactoryShaderPer
 		return true;
 	}
 
-	// We can use the local vertex factory in pass through mode when rendering SkeletalMesh through the SkinCache
-	if (Parameters.MaterialParameters.bIsUsedWithSkeletalMesh && IsGPUSkinPassThroughSupported(Parameters.Platform))
+	// We use the local vertex factory in passthrough mode when using MeshDeformers or when using SkeletalMesh through the SkinCache.
+	if ((Parameters.MaterialParameters.bIsUsedWithSkeletalMesh && IsGPUSkinCacheAvailable(Parameters.Platform)) ||
+		(Parameters.MaterialParameters.bIsUsedWithMeshDeformer && AreMeshDeformersAvailable(Parameters.Platform)))
 	{
 		return true;
 	}
diff --git a/Engine/Source/Runtime/Engine/Private/Materials/Material.cpp b/Engine/Source/Runtime/Engine/Private/Materials/Material.cpp
--- a/Engine/Source/Runtime/Engine/Private/Materials/Material.cpp
+++ b/Engine/Source/Runtime/Engine/Private/Materials/Material.cpp
@@ -1542,6 +1542,7 @@ uint32 UMaterial::GetUsageFlags() const
 	UsageFlags |= bUsedWithStaticMesh ? (1u << MATUSAGE_StaticMesh) : 0;
 	UsageFlags |= bUsedWithEditorCompositing ? (1u << MATUSAGE_EditorCompositing) : 0;
 	UsageFlags |= bUsedWithNeuralNetworks ? (1u << MATUSAGE_NeuralNetworks) : 0;
+	UsageFlags |= bUsedWithMeshDeformer ? (1u << MATUSAGE_MeshDeformer) : 0;
 	PRAGMA_ENABLE_DEPRECATION_WARNINGS
 	return UsageFlags;
 }
@@ -1577,6 +1578,7 @@ bool UMaterial::GetUsageByFlag(EMaterialUsage Usage) const
 		case MATUSAGE_StaticMesh: UsageValue = bUsedWithStaticMesh; break;
 		case MATUSAGE_EditorCompositing: UsageValue = bUsedWithEditorCompositing; break;
 		case MATUSAGE_NeuralNetworks: UsageValue = bUsedWithNeuralNetworks; break;
+		case MATUSAGE_MeshDeformer: UsageValue = bUsedWithMeshDeformer; break;
 		default: UE_LOG(LogMaterial, Fatal,TEXT("Unknown material usage: %u"), (int32)Usage);
 	};
 	PRAGMA_ENABLE_DEPRECATION_WARNINGS
@@ -1852,6 +1854,10 @@ void UMaterial::SetUsageByFlag(EMaterialUsage Usage, bool NewValue)
 		{
 			bUsedWithNeuralNetworks = NewValue; break;
 		}
+		case MATUSAGE_MeshDeformer:
+		{
+			bUsedWithMeshDeformer = NewValue; break;
+		}
 		default: UE_LOG(LogMaterial, Fatal, TEXT("Unknown material usage: %u"), (int32)Usage);
 	};
 	PRAGMA_ENABLE_DEPRECATION_WARNINGS
@@ -3132,6 +3138,14 @@ void UMaterial::Serialize(FArchive& Ar)
 		}
 	}
 
+	if (Ar.IsLoading() && Ar.CustomVer(FFortniteMainBranchObjectVersion::GUID) < FFortniteMainBranchObjectVersion::AddMeshDeformerUsageFlag)
+	{
+		// Default mesh deformer usage to match the skeletal mesh usage.
+		PRAGMA_DISABLE_DEPRECATION_WARNINGS
+		bUsedWithMeshDeformer = bUsedWithSkeletalMesh;
+		PRAGMA_ENABLE_DEPRECATION_WARNINGS
+	}
+	
 #if WITH_EDITOR
 	if (Ar.IsSaving() && Ar.IsCooking() && ShouldStripUnusedDefaultTextures())
 	{
diff --git a/Engine/Source/Runtime/Engine/Private/Materials/MaterialInterface.cpp b/Engine/Source/Runtime/Engine/Private/Materials/MaterialInterface.cpp
--- a/Engine/Source/Runtime/Engine/Private/Materials/MaterialInterface.cpp
+++ b/Engine/Source/Runtime/Engine/Private/Materials/MaterialInterface.cpp
@@ -2722,6 +2722,7 @@ FString UMaterialInterface::GetUsageName(EMaterialUsage Usage)
 	case MATUSAGE_StaticMesh: UsageName = TEXT("StaticMesh"); break;
 	case MATUSAGE_EditorCompositing: UsageName = TEXT("EditorCompositing"); break;
 	case MATUSAGE_NeuralNetworks: UsageName = TEXT("NeuralNetworks"); break;
+	case MATUSAGE_MeshDeformer: UsageName = TEXT("MeshDeformer"); break;
 	default: UE_LOG(LogMaterial, Fatal, TEXT("Unknown material usage: %u"), (int32)Usage);
 	};
 	return UsageName;
diff --git a/Engine/Source/Runtime/Engine/Private/Materials/MaterialShared.cpp b/Engine/Source/Runtime/Engine/Private/Materials/MaterialShared.cpp
--- a/Engine/Source/Runtime/Engine/Private/Materials/MaterialShared.cpp
+++ b/Engine/Source/Runtime/Engine/Private/Materials/MaterialShared.cpp
@@ -1800,6 +1800,11 @@ bool FMaterialResource::IsUsedWithNeuralNetworks() const
 	return MaterialInstance ? MaterialInstance->GetUsageByFlag(MATUSAGE_NeuralNetworks) : Material->GetUsageByFlag(MATUSAGE_NeuralNetworks);
 }
 
+bool FMaterialResource::IsUsedWithMeshDeformer() const
+{
+	return MaterialInstance ? MaterialInstance->GetUsageByFlag(MATUSAGE_MeshDeformer) : Material->GetUsageByFlag(MATUSAGE_MeshDeformer);
+}
+
 bool FMaterialResource::IsUsedWithLandscape() const
 {
 	return false;
@@ -5936,13 +5941,14 @@ FMaterialShaderParametersBase::FMaterialShaderParametersBase(const FMaterial* In
 	bAlwaysEvaluateWorldPositionOffset = InMaterial->ShouldAlwaysEvaluateWorldPositionOffset();
 	bDisablePreExposureScale = InMaterial->GetDisablePreExposureScale();
 	bAllowVariableRateShading = InMaterial->IsVariableRateShadingAllowed();
-	bSamplesMaterialCache = InMaterial->SamplesMaterialCache();
-	bHasMaterialCacheOutput = InMaterial->HasMaterialCacheOutput();
 	bIsUsedWithStaticMesh = InMaterial->IsUsedWithStaticMesh();
 	bIsTranslucentSurfaceLighting = InMaterial->GetTranslucencyLightingMode() == ETranslucencyLightingMode::TLM_SurfacePerPixelLighting; // get lighting mode for mobile translucent receive CSM 
 	bHasFirstPersonInterpolation = InMaterial->HasFirstPersonOutput();
 	bUsesTemporalResponsiveness = InMaterial->UsesTemporalResponsiveness();
 	bUsesMotionVectorWorldOffset = InMaterial->UsesMotionVectorWorldOffset();
+	bSamplesMaterialCache = InMaterial->SamplesMaterialCache();
+	bHasMaterialCacheOutput = InMaterial->HasMaterialCacheOutput();
+	bIsUsedWithMeshDeformer = InMaterial->IsUsedWithMeshDeformer();
 }
 
 FMaterialShaderParameters::FMaterialShaderParameters(const FMaterial* InMaterial)
diff --git a/Engine/Source/Runtime/Engine/Private/SkeletalMeshLODRenderData.cpp b/Engine/Source/Runtime/Engine/Private/SkeletalMeshLODRenderData.cpp
--- a/Engine/Source/Runtime/Engine/Private/SkeletalMeshLODRenderData.cpp
+++ b/Engine/Source/Runtime/Engine/Private/SkeletalMeshLODRenderData.cpp
@@ -78,12 +78,6 @@ namespace
 	};
 }
 
-static bool IsMeshDeformerAvailable(EShaderPlatform InPlatform)
-{
-	static IMeshDeformerProvider* MeshDeformerProvider = IMeshDeformerProvider::Get();
-	return MeshDeformerProvider && MeshDeformerProvider->IsSupported(InPlatform);
-}
-
 /** 
  * This function returns true if the duplicate vertices should be cooked. 
  * The data is used to deal with seams along split vertices when recomputing normals at runtime.
@@ -117,7 +111,7 @@ static bool IsMeshDeformerSupportedByCookTargetPlatform(const ITargetPlatform& I
 	{
 		const EShaderPlatform LegacyShaderPlatform = ShaderFormatToLegacyShaderPlatform(TargetedShaderFormats[FormatIndex]);
 		
-		if (IsMeshDeformerAvailable(LegacyShaderPlatform))
+		if (AreMeshDeformersAvailable(LegacyShaderPlatform))
 		{
 			return true;
 		}
diff --git a/Engine/Source/Runtime/Engine/Private/SkeletalMeshSceneProxy.cpp b/Engine/Source/Runtime/Engine/Private/SkeletalMeshSceneProxy.cpp
--- a/Engine/Source/Runtime/Engine/Private/SkeletalMeshSceneProxy.cpp
+++ b/Engine/Source/Runtime/Engine/Private/SkeletalMeshSceneProxy.cpp
@@ -177,6 +177,8 @@ FSkeletalMeshSceneProxy::FSkeletalMeshSceneProxy(const FSkinnedMeshSceneProxyDes
 	extern bool GForceDefaultMaterial;
 	const bool bCastShadow = InMeshDesc.CastShadow;
 	bool bAnySectionCastsShadow = false;
+	bool bAnySectionHasSkelMeshUsage = false;
+	bool bAnySectionHasMeshDeformerUsage = false;
 	LODSections.Reserve(RenderData->LODRenderData.Num());
 	LODSections.AddZeroed(RenderData->LODRenderData.Num());
 	for(int32 LODIdx=0; LODIdx < RenderData->LODRenderData.Num(); LODIdx++)
@@ -213,10 +215,15 @@ FSkeletalMeshSceneProxy::FSkeletalMeshSceneProxy(const FSkinnedMeshSceneProxyDes
 				MaterialRelevance |= Material->GetRelevance(GetScene().GetShaderPlatform());
 			}
 
+			// validate usage based on technique
+			const bool bUsesMeshDeformer = InMeshDesc.GetMeshDeformerInstanceForLOD(LODIdx) != nullptr;
+			const EMaterialUsage RequiredUsage = bUsesMeshDeformer ? MATUSAGE_MeshDeformer : MATUSAGE_SkeletalMesh;
+			bAnySectionHasMeshDeformerUsage |= RequiredUsage == MATUSAGE_MeshDeformer;
+			bAnySectionHasSkelMeshUsage |= RequiredUsage == MATUSAGE_SkeletalMesh;
+			bool bValidUsage = Material && Material->CheckMaterialUsage_Concurrent(RequiredUsage);
+
 			// if this is a clothing section, then enabled and will be drawn but the corresponding original section should be disabled
 			const bool bClothSection = Section.HasClothingData();
-
-			bool bValidUsage = Material && Material->CheckMaterialUsage_Concurrent(MATUSAGE_SkeletalMesh);
 			if (bClothSection)
 			{
 				bValidUsage &= Material && Material->CheckMaterialUsage_Concurrent(MATUSAGE_Clothing);
@@ -271,12 +278,18 @@ FSkeletalMeshSceneProxy::FSkeletalMeshSceneProxy(const FSkinnedMeshSceneProxyDes
 
 	if (OverlayMaterial != nullptr)
 	{
-		if (!OverlayMaterial->CheckMaterialUsage_Concurrent(MATUSAGE_SkeletalMesh))
+		bool bInvalidUsage = false;
+		if (bAnySectionHasSkelMeshUsage && !OverlayMaterial->CheckMaterialUsage_Concurrent(MATUSAGE_SkeletalMesh))
 		{
-			OverlayMaterial = nullptr;
-			UE_LOG(LogSkeletalMesh, Warning, TEXT("Overlay material with missing usage flag was applied to skeletal mesh %s"),	*InMeshDesc.GetSkinnedAsset()->GetPathName());
+			bInvalidUsage = true;
+			UE_LOG(LogSkeletalMesh, Warning, TEXT("Overlay material with missing skel mesh usage flag was applied to skeletal mesh %s"), *InMeshDesc.GetSkinnedAsset()->GetPathName());
+		}
+		if (bAnySectionHasMeshDeformerUsage && !OverlayMaterial->CheckMaterialUsage_Concurrent(MATUSAGE_MeshDeformer))
+		{
+			bInvalidUsage = true;
+			UE_LOG(LogSkeletalMesh, Warning, TEXT("Overlay material with missing mesh deformer usage flag was applied to skeletal mesh %s"), *InMeshDesc.GetSkinnedAsset()->GetPathName());
 		}
-		else if (bForceDefaultMaterial)
+		if (bInvalidUsage || bForceDefaultMaterial)
 		{
 			OverlayMaterial = nullptr;
 		}
diff --git a/Engine/Source/Runtime/Engine/Public/MaterialShared.h b/Engine/Source/Runtime/Engine/Public/MaterialShared.h
--- a/Engine/Source/Runtime/Engine/Public/MaterialShared.h
+++ b/Engine/Source/Runtime/Engine/Public/MaterialShared.h
@@ -1128,6 +1128,9 @@ struct FMaterialShaderParametersBase
 			uint16 bIsSubstrateMaterial : 1;
 			uint16 bUsesTemporalResponsiveness : 1;
 			uint16 bUsesMotionVectorWorldOffset : 1;
+			uint16 bSamplesMaterialCache : 1;
+			uint16 bHasMaterialCacheOutput : 1;
+			uint16 bIsUsedWithMeshDeformer : 1;
 		};
 	};
 	EBlendMode BlendMode;
@@ -1208,8 +1211,6 @@ struct FMaterialShaderParametersBase
 			uint64 bAlwaysEvaluateWorldPositionOffset : 1;
 			uint64 bDisablePreExposureScale : 1;
 			uint64 bAllowVariableRateShading : 1;
-			uint8 bSamplesMaterialCache : 1;
-			uint8 bHasMaterialCacheOutput : 1;
 		};
 	};
 
@@ -2408,6 +2409,7 @@ public:
 	virtual bool IsUsedWithLidarPointCloud() const { return false; }
 	virtual bool IsUsedWithVirtualHeightfieldMesh() const { return false; }
 	virtual bool IsUsedWithNeuralNetworks() const { return false; }
+	virtual bool IsUsedWithMeshDeformer() const { return false; }
 	virtual bool IsUsedWithStaticMesh() const { return true; }
 	virtual bool IsFullyRough() const { return false; }
 	virtual bool GetForceCompatibleWithLightFunctionAtlas() const { return false; }
@@ -3223,6 +3225,7 @@ public:
 	ENGINE_API virtual bool IsUsedWithLidarPointCloud() const override;
 	ENGINE_API virtual bool IsUsedWithVirtualHeightfieldMesh() const override;
 	ENGINE_API virtual bool IsUsedWithNeuralNetworks() const override;
+	ENGINE_API virtual bool IsUsedWithMeshDeformer() const override;
 	ENGINE_API virtual bool IsUsedWithNanite() const override;
 	ENGINE_API virtual bool IsUsedWithVoxels() const override;
 	ENGINE_API virtual bool IsUsedWithVolumetricCloud() const override;
diff --git a/Engine/Source/Runtime/Engine/Public/Materials/Material.h b/Engine/Source/Runtime/Engine/Public/Materials/Material.h
--- a/Engine/Source/Runtime/Engine/Public/Materials/Material.h
+++ b/Engine/Source/Runtime/Engine/Public/Materials/Material.h
@@ -686,6 +686,14 @@ public:
 	UPROPERTY(EditAnywhere, BlueprintReadOnly, Category=Usage)
 	uint8 bUsedWithSkeletalMesh : 1;
 
+	/**
+	 * Indicates that the material and its instances can be use with mesh deformers
+	 * This will result in the shaders required to support mesh deformers being compiled which will increase shader compile time and memory usage.
+	 */
+	UE_DEPRECATED(5.8, "Always use the GetUsageByFlag() accessor.")
+	UPROPERTY(EditAnywhere, BlueprintReadOnly, Category = Usage)
+	uint8 bUsedWithMeshDeformer : 1;
+
 	/** 
 	 * Indicates that the material and its instances can be used with editor compositing  
 	 * This will result in the shaders required to support editor compositing being compiled which will increase shader compile time and memory usage.
diff --git a/Engine/Source/Runtime/Engine/Public/Materials/MaterialInterface.h b/Engine/Source/Runtime/Engine/Public/Materials/MaterialInterface.h
--- a/Engine/Source/Runtime/Engine/Public/Materials/MaterialInterface.h
+++ b/Engine/Source/Runtime/Engine/Public/Materials/MaterialInterface.h
@@ -123,6 +123,8 @@ enum EMaterialUsage : int
 	MATUSAGE_EditorCompositing,
 	/** Indicates that the material and its instances can be used with neural network engine. */
 	MATUSAGE_NeuralNetworks,
+	/** Indicates that the material and its instances can be used with mesh deformers. */
+	MATUSAGE_MeshDeformer,
 
 	MATUSAGE_MAX,
 };
diff --git a/Engine/Source/Runtime/RenderCore/Private/RenderUtils.cpp b/Engine/Source/Runtime/RenderCore/Private/RenderUtils.cpp
--- a/Engine/Source/Runtime/RenderCore/Private/RenderUtils.cpp
+++ b/Engine/Source/Runtime/RenderCore/Private/RenderUtils.cpp
@@ -1597,12 +1597,18 @@ RENDERCORE_API bool IsGPUSkinCacheAvailable(EShaderPlatform Platform)
 	return AreSkinCacheShadersEnabled(Platform) != 0 && IsGPUSkinCacheAllowed(Platform);
 }
 
+RENDERCORE_API bool AreMeshDeformersAvailable(EShaderPlatform Platform)
+{
+	// Note that we rely on the IMeshDeformerProvider being loaded before any calls to this function.
+	// OptimusSettings/ComputeFramework are PostConfigInit for that reason.
+	static IMeshDeformerProvider* MeshDeformerProvider = IMeshDeformerProvider::Get();
+	return MeshDeformerProvider && MeshDeformerProvider->IsSupported(Platform);
+}
+
 RENDERCORE_API bool IsGPUSkinPassThroughSupported(EShaderPlatform Platform)
 {
 	// Enable the GPUSkin passthrough path if we might use the GPUSkinCache or MeshDeformers.
-	static IMeshDeformerProvider* MeshDeformerProvider = IMeshDeformerProvider::Get();
-	bool bMeshDeformersAvailable = MeshDeformerProvider && MeshDeformerProvider->IsSupported(Platform);
-	return bMeshDeformersAvailable || IsGPUSkinCacheAvailable(Platform);
+	return IsGPUSkinCacheAvailable(Platform) || AreMeshDeformersAvailable(Platform);
 }
 
 RENDERCORE_API bool AreBufferSRVsAlwaysCreatedByDefault(EShaderPlatform ShaderPlatform)
diff --git a/Engine/Source/Runtime/RenderCore/Public/RenderUtils.h b/Engine/Source/Runtime/RenderCore/Public/RenderUtils.h
--- a/Engine/Source/Runtime/RenderCore/Public/RenderUtils.h
+++ b/Engine/Source/Runtime/RenderCore/Public/RenderUtils.h
@@ -545,6 +545,11 @@ RENDERCORE_API bool IsGPUSkinCacheAllowed(EShaderPlatform Platform);
  */
 RENDERCORE_API bool IsGPUSkinCacheAvailable(EShaderPlatform Platform);
 
+/**
+ * Checks if mesh deformers can be used for a specific platform.
+ */
+RENDERCORE_API bool AreMeshDeformersAvailable(EShaderPlatform Platform);
+
 /**
  * Does the platform support GPUSkinPassthrough permutations.
  * This knowledge can be used to indicate if we need to create SRV for index/vertex buffers.
