Skipping overridable serialization events and initialization tracking on Verse properties during PIE.

Also introduced an IsObjectEditable utility function for checking if objects belong in a world that allows them to be edited.

#rnx
#Synced-CL 48962394

diff --git a/Engine/Source/Runtime/Engine/Classes/Engine/World.h b/Engine/Source/Runtime/Engine/Classes/Engine/World.h
--- a/Engine/Source/Runtime/Engine/Classes/Engine/World.h
+++ b/Engine/Source/Runtime/Engine/Classes/Engine/World.h
@@ -1546,6 +1546,9 @@ public:
 
 #endif // WITH_EDITOR
 
+	/** Checks if a given object lives in a world that allows it to be edited */
+	static UE_API bool IsObjectEditable(const UObject* Object);
+
 	/**
 	 * Sets whether or not this world is ticked by the engine, but use it at your own risk! This could
 	 * have unintended consequences if used carelessly. That said, for worlds that are not interactive
diff --git a/Engine/Source/Runtime/Engine/Private/World.cpp b/Engine/Source/Runtime/Engine/Private/World.cpp
--- a/Engine/Source/Runtime/Engine/Private/World.cpp
+++ b/Engine/Source/Runtime/Engine/Private/World.cpp
@@ -9242,6 +9242,17 @@ bool UWorld::IsPreviewWorld() const
 	return WorldType == EWorldType::EditorPreview || WorldType == EWorldType::GamePreview;
 }
 
+bool UWorld::IsObjectEditable(const UObject* Object)
+{
+	bool bResult = false;
+	const UWorld* World = Object ? Object->GetWorld() : nullptr;
+	if (World)
+	{
+		bResult = (World->WorldType == EWorldType::Editor) || (World->WorldType == EWorldType::EditorPreview);
+	}
+	return bResult;
+}
+
 bool UWorld::UsesGameHiddenFlags() const
 {
 	return IsGameWorld();
