Fix API macro misuse

#rnx

#Synced-CL 49063813

diff --git a/Engine/Source/Runtime/Engine/Public/WorldPartition/HLOD/HLODActor.h b/Engine/Source/Runtime/Engine/Public/WorldPartition/HLOD/HLODActor.h
--- a/Engine/Source/Runtime/Engine/Public/WorldPartition/HLOD/HLODActor.h
+++ b/Engine/Source/Runtime/Engine/Public/WorldPartition/HLOD/HLODActor.h
@@ -135,11 +135,11 @@ public:
 	ENGINE_API void SetHLODRebuildPolicyDataSet(const FHLODRebuildPolicyDataSet& InHLODRebuildPolicyDataSet);
 
 	UE_DEPRECATED(5.8, "Hash data is now stored in a UHLODRebuildPolicyData object.")
-	ENGINE_API uint32 GetHLODHash() const { return HLODHash_DEPRECATED; }
+	uint32 GetHLODHash() const { return HLODHash_DEPRECATED; }
 	UE_DEPRECATED(5.8, "Hash data is now stored in a UHLODRebuildPolicyData object.")
-	ENGINE_API uint32 ComputeHLODHash() const { return 0; }
+	uint32 ComputeHLODHash() const { return 0; }
 	UE_DEPRECATED(5.8, "Hash data is now stored in a UHLODRebuildPolicyData object.")
-	ENGINE_API void SetHLODHash(uint32 InHLODHash, const FString& InHLODBuildReportContent) {}
+	void SetHLODHash(uint32 InHLODHash, const FString& InHLODBuildReportContent) {}
 	UE_DEPRECATED(5.8, "Use override with no parameter.")
 	ENGINE_API void UpdateHLODBuildReportContent(const FString& InHLODBuildReportContent);
 
