FSpline Refactor: Prepare FSpline for non-tangent based splines.

diff --git a/Engine/Source/Runtime/Engine/Private/Components/SplineComponent.cpp b/Engine/Source/Runtime/Engine/Private/Components/SplineComponent.cpp
--- a/Engine/Source/Runtime/Engine/Private/Components/SplineComponent.cpp
+++ b/Engine/Source/Runtime/Engine/Private/Components/SplineComponent.cpp
@@ -5,6 +5,7 @@
 =============================================================================*/
 
 #include "Components/SplineComponent.h"
+#include "Curves/Splines/TangentSpline.h"
 #include "Engine/Engine.h"
 #include "SceneView.h"
 #include "UObject/EditorObjectVersion.h"
@@ -27,7 +28,6 @@
 #include "Misc/TransactionObjectEvent.h"
 #include "UObject/ICookInfo.h"
 
-
 #include UE_INLINE_GENERATED_CPP_BY_NAME(SplineComponent)
 
 #if WITH_EDITOR
@@ -38,9 +38,9 @@
 
 DEFINE_LOG_CATEGORY_STATIC(LogSplineComponent, Log, All);
 
-const FInterpCurvePointVector USplineComponent::DummyPointPosition(0.0f, FVector::ZeroVector, FVector::ForwardVector, FVector::ForwardVector, CIM_Constant);
-const FInterpCurvePointQuat USplineComponent::DummyPointRotation(0.0f, FQuat::Identity);
-const FInterpCurvePointVector USplineComponent::DummyPointScale(0.0f, FVector::OneVector);
+const FInterpCurveVector USplineComponent::DummyPointPosition = {};
+const FInterpCurveQuat USplineComponent::DummyPointRotation = {};
+const FInterpCurveVector USplineComponent::DummyPointScale = {};
 
 namespace UE::SplineComponent
 {
@@ -132,12 +132,12 @@ static ELastAuthority GetEffectiveLastAuthority(const ELastAuthority InLastAutho
 	}
 	return UseLastAuthority;
 }
-
-static bool ShouldUseSplineCurves()
+	
+bool ShouldUseSplineCurvesStatic()
 {
-	return GUseSplineCurves || FSpline::IsEnabledStatic() == false;
+	return GUseSplineCurves || FSpline::GetDefaultImplementation() != UE::Spline::ESplineType::Tangent;
 }
-	
+
 }
 
 USplineMetadata::USplineMetadata(const FObjectInitializer& ObjectInitializer)
@@ -147,12 +147,19 @@ USplineMetadata::USplineMetadata(const FObjectInitializer& ObjectInitializer)
 
 bool USplineComponent::Validate() const
 {
-	if (!Spline.IsEnabled())
+	if (Spline.GetCurrentImplementation() != UE::Spline::ESplineType::Tangent)
 	{
 		// Can't be inconsistent if Spline is incapable of storing any data at all.
 		return true;
 	}
 
+	const FTangentSpline* TangentSpline = Spline.Get<UE::Spline::ESplineType::Tangent>();
+	if (!ensureAlways(TangentSpline))
+	{
+		// Something is wrong...
+		return false;
+	}
+
 	auto StringifySplinePointType = [](const ESplinePointType::Type InSplinePointType)
 	{
 		switch (InSplinePointType)
@@ -169,10 +176,10 @@ bool USplineComponent::Validate() const
 	bool bValid = true;
 	
 	const int32 NumSplineCurvesPoints = WarninglessSplineCurves().Position.Points.Num();
-	const int32 NumSplinePoints = Spline.GetNumControlPoints();
+	const int32 NumSplinePoints = TangentSpline->GetNumberOfPoints();
 	const int32 NumPoints = FMath::Min(NumSplinePoints, NumSplineCurvesPoints);
 
-	UE_LOG(LogSplineComponent, Display, TEXT("Validating SplineComponent (%s) - Authority: %s"), *GetPathName(), UE::SplineComponent::ShouldUseSplineCurves() ? TEXT("FSplineCurves") : TEXT("FSpline"));
+	UE_LOG(LogSplineComponent, Display, TEXT("Validating SplineComponent (%s) - Authority: %s"), *GetPathName(), ShouldUseSplineCurves() ? TEXT("FSplineCurves") : TEXT("FSpline"));
 	
 	if (NumSplineCurvesPoints != NumSplinePoints)
 	{
@@ -180,6 +187,15 @@ bool USplineComponent::Validate() const
 		bValid = false;
 	}
 
+	const float SplineCurvesLength = WarninglessSplineCurves().GetSplineLength();
+	const float SplineLength = TangentSpline->GetDistanceAtParameter(TangentSpline->GetParameterSpace().Max);
+	if (!FMath::IsNearlyEqual(SplineCurvesLength, SplineLength))
+	{
+		UE_LOG(LogSplineComponent, Display, TEXT("Internally inconsistent spline lengths (%f vs %f)!"), SplineCurvesLength, SplineLength);
+		bValid = false;
+	}
+	UE_LOG(LogSplineComponent, Display, TEXT("Spline lengths: (%f vs %f)"), SplineCurvesLength, SplineLength);
+
 	for (int32 Idx = 0; Idx < NumPoints; ++Idx)
 	{
 		const auto& SplineCurvesPositionPoint = WarninglessSplineCurves().Position.Points[Idx];
@@ -187,7 +203,7 @@ bool USplineComponent::Validate() const
 		const auto& SplineCurvesScalePoint = WarninglessSplineCurves().Scale.Points[Idx];
 
 		const ESplinePointType::Type SplineCurvesPointType = ConvertInterpCurveModeToSplinePointType(SplineCurvesPositionPoint.InterpMode);
-		const ESplinePointType::Type SplinePointType = ConvertInterpCurveModeToSplinePointType(Spline.GetSplinePointType(Idx));
+		const ESplinePointType::Type SplinePointType = ConvertInterpCurveModeToSplinePointType(TangentSpline->GetSplinePointType(Idx));
 		if (SplineCurvesPointType != SplinePointType)
 		{
 			UE_LOG(LogSplineComponent, Display, TEXT("Internally inconsistent point type (%s vs %s) at index %d!"), StringifySplinePointType(SplineCurvesPointType), StringifySplinePointType(SplinePointType), Idx );
@@ -195,7 +211,7 @@ bool USplineComponent::Validate() const
 		}
 		
 		const FVector SplineCurvesPosition = SplineCurvesPositionPoint.OutVal;
-		const FVector SplinePosition = Spline.GetLocation(Idx);
+		const FVector SplinePosition = TangentSpline->GetLocation(Idx);
 		if (!SplineCurvesPosition.Equals(SplinePosition))
 		{
 			UE_LOG(LogSplineComponent, Display, TEXT("Internally inconsistent position (%s vs %s) at index %d!"), *SplineCurvesPosition.ToCompactString(), *SplinePosition.ToCompactString(), Idx );
@@ -203,7 +219,7 @@ bool USplineComponent::Validate() const
 		}
 
 		const FVector SplineCurvesInTangent = SplineCurvesPositionPoint.ArriveTangent;
-		const FVector SplineInTangent = Spline.GetInTangent(Idx);
+		const FVector SplineInTangent = TangentSpline->GetInTangent(Idx);
 		if (!SplineCurvesInTangent.Equals(SplineInTangent))
 		{
 			UE_LOG(LogSplineComponent, Display, TEXT("Internally inconsistent arrive tangent (%s vs %s) at index %d!"), *SplineCurvesInTangent.ToCompactString(), *SplineInTangent.ToCompactString(), Idx );
@@ -220,7 +236,7 @@ bool USplineComponent::Validate() const
 		}
 
 		const FVector SplineCurvesOutTangent = SplineCurvesPositionPoint.LeaveTangent;
-		const FVector SplineOutTangent = Spline.GetOutTangent(Idx);
+		const FVector SplineOutTangent = TangentSpline->GetOutTangent(Idx);
 		if (!SplineCurvesOutTangent.Equals(SplineOutTangent))
 		{
 			UE_LOG(LogSplineComponent, Display, TEXT("Internally inconsistent leave tangent (%s vs %s) at index %d!"), *SplineCurvesOutTangent.ToCompactString(), *SplineOutTangent.ToCompactString(), Idx );
@@ -237,7 +253,7 @@ bool USplineComponent::Validate() const
 		}
 
 		const float SplineCurvesParam = SplineCurvesPositionPoint.InVal;
-		const float SplineParam = Spline.GetParameterAtIndex(Idx);
+		const float SplineParam = TangentSpline->GetParameterAtIndex(Idx);
 		if (!FMath::IsNearlyEqual(SplineCurvesParam, SplineParam))
 		{
 			UE_LOG(LogSplineComponent, Display, TEXT("Internally inconsistent parameter (%f vs %f) at index %d!"), SplineCurvesParam, SplineParam, Idx );
@@ -245,7 +261,7 @@ bool USplineComponent::Validate() const
 		}
 
 		const FQuat SplineCurvesRotation = SplineCurvesRotationPoint.OutVal;
-		const FQuat SplineRotation = Spline.GetRotation(Idx);
+		const FQuat SplineRotation = TangentSpline->GetRotation(Idx);
 		if (!SplineCurvesRotation.Equals(SplineRotation))
 		{
 			UE_LOG(LogSplineComponent, Display, TEXT("Internally inconsistent rotation (%s vs %s) at index %d!"), *SplineCurvesRotation.ToString(), *SplineRotation.ToString(), Idx );
@@ -253,7 +269,7 @@ bool USplineComponent::Validate() const
 		}
 
 		const FVector SplineCurvesScale = SplineCurvesScalePoint.OutVal;
-		const FVector SplineScale = Spline.GetScale(Idx);
+		const FVector SplineScale = TangentSpline->GetScale(Idx);
 		if (!SplineCurvesScale.Equals(SplineScale))
 		{
 			UE_LOG(LogSplineComponent, Display, TEXT("Internally inconsistent scale (%s vs %s) at index %d!"), *SplineCurvesScale.ToCompactString(), *SplineScale.ToCompactString(), Idx );
@@ -269,14 +285,9 @@ bool USplineComponent::Validate() const
 	return bValid;
 }
 
-void USplineComponent::PopulateFromLegacy()
-{
-	SynchronizeSplines();
-}
-
 FName USplineComponent::GetSplinePropertyName()
 {
-	if (UE::SplineComponent::ShouldUseSplineCurves())
+	if (UE::SplineComponent::ShouldUseSplineCurvesStatic())
 	{
 PRAGMA_DISABLE_DEPRECATION_WARNINGS
 		return GET_MEMBER_NAME_CHECKED(USplineComponent, SplineCurves);
@@ -335,9 +346,9 @@ USplineComponent::USplineComponent(const FObjectInitializer& ObjectInitializer)
 
 #if WITH_EDITORONLY_DATA
 	// Set these deprecated values up so that old assets with default values load correctly (and are subsequently upgraded during Serialize)
-	SplineInfo_DEPRECATED = Spline.GetSplinePointsPosition();
-	SplineRotInfo_DEPRECATED = Spline.GetSplinePointsRotation();
-	SplineScaleInfo_DEPRECATED = Spline.GetSplinePointsScale();
+	SplineInfo_DEPRECATED = WarninglessSplineCurves().Position;
+	SplineRotInfo_DEPRECATED = WarninglessSplineCurves().Rotation;
+	SplineScaleInfo_DEPRECATED = WarninglessSplineCurves().Scale;
 	SplineReparamTable_DEPRECATED = WarninglessSplineCurves().ReparamTable;
 #endif
 
@@ -387,26 +398,39 @@ FSpline USplineComponent::GetDefaultSpline()
 
 void USplineComponent::SetDefaultSpline()
 {
+	if (UE::SplineComponent::GUseSplineCurves)
+	{
+		Spline = FSpline(UE::Spline::ESplineType::Unimplemented);
+	}
+	else
+	{
+		Spline = FSpline(UE::Spline::ESplineType::Tangent);
+	}
+
 	const FVector StartPoint(0.0, 0.0, 0.0);
 	const FVector EndPoint(100.0, 0.0, 0.0);
 	constexpr float StartParam = 0.f;
 	constexpr float EndParam = 1.f;
 
-	Spline.Reset();
 	WarninglessSplineCurves().Position.Points.Reset(10);
     WarninglessSplineCurves().Rotation.Points.Reset(10);
     WarninglessSplineCurves().Scale.Points.Reset(10);
 
-	Spline.AddPoint(FSplinePoint(StartParam, StartPoint));
     WarninglessSplineCurves().Position.Points.Emplace(StartParam, StartPoint, FVector::ZeroVector, FVector::ZeroVector, CIM_CurveAuto);
     WarninglessSplineCurves().Rotation.Points.Emplace(StartParam, FQuat::Identity, FQuat::Identity, FQuat::Identity, CIM_CurveAuto);
     WarninglessSplineCurves().Scale.Points.Emplace(StartParam, FVector(1.0f), FVector::ZeroVector, FVector::ZeroVector, CIM_CurveAuto);
 
-	Spline.AddPoint(FSplinePoint(EndParam, EndPoint));
     WarninglessSplineCurves().Position.Points.Emplace(EndParam, EndPoint, FVector::ZeroVector, FVector::ZeroVector, CIM_CurveAuto);
     WarninglessSplineCurves().Rotation.Points.Emplace(EndParam, FQuat::Identity, FQuat::Identity, FQuat::Identity, CIM_CurveAuto);
     WarninglessSplineCurves().Scale.Points.Emplace(EndParam, FVector(1.0f), FVector::ZeroVector, FVector::ZeroVector, CIM_CurveAuto);
 
+	if (FTangentSpline* TangentSpline = Spline.Get<UE::Spline::ESplineType::Tangent>())
+	{
+		TangentSpline->Reset();
+		TangentSpline->AddPoint(FSplinePoint(StartParam, StartPoint));
+		TangentSpline->AddPoint(FSplinePoint(EndParam, EndPoint));
+	}
+
 	OnSplineChanged.Broadcast();
 
 	if (UE::SplineComponent::GValidateOnChange)
@@ -456,7 +480,7 @@ void USplineComponent::Serialize(FArchive& Ar)
 {
 	if (Ar.IsSaving())
 	{
-		LastAuthority = UE::SplineComponent::ShouldUseSplineCurves() ? ELastAuthority::SplineCurves : ELastAuthority::Spline;
+		LastAuthority = ShouldUseSplineCurves() ? ELastAuthority::SplineCurves : ELastAuthority::Spline;
 	}
 	
 	Super::Serialize(Ar);
@@ -476,14 +500,15 @@ void USplineComponent::Serialize(FArchive& Ar)
 		}
 
 #if WITH_EDITOR
-		// WasEnabled() is editor only, that data gets stripped from packaged builds.
-		if (!Spline.WasEnabled())
+		// If SplineComponent.UseSplineCurves is false and Spline was not previously implemented
+		// (implying it is not consistent with SplineCurves), initialize Spline from SplineCurves.
+		if (Spline.GetPreviousImplementation() == UE::Spline::ESplineType::Unimplemented && !UE::SplineComponent::GUseSplineCurves)
 		{
-			SynchronizeSplines();
+			Spline = WarninglessSplineCurves();
 		}
 #endif
 		
-		if (UE::SplineComponent::GSynchronizeOnLoad && Spline.IsEnabled() && !Validate())
+		if (UE::SplineComponent::GSynchronizeOnLoad && Spline.GetCurrentImplementation() == UE::Spline::ESplineType::Tangent && !Validate())
 		{
 			UE_LOG(LogSplineComponent, Display, TEXT("Internally inconsistent SplineComponent (%s) at load time!"), *GetPathName())
 			
@@ -506,7 +531,7 @@ void USplineComponent::Serialize(FArchive& Ar)
 
 			case ELastAuthority::SplineCurves:
 				UE_LOG(LogSplineComponent, Display, TEXT("Populating Spline from SplineCurves..."))
-				SynchronizeSplines();
+				SetSpline(WarninglessSplineCurves());
 				break;
 			}
 
@@ -522,7 +547,7 @@ void USplineComponent::Serialize(FArchive& Ar)
 				UE_LOG(LogSplineComponent, Display, TEXT("Successfully synchronized SplineComponent (%s) on load."), *GetPathName())
 			}
 		}
-		else if (UE::SplineComponent::GValidateOnLoad && Spline.IsEnabled() && !Validate())
+		else if (UE::SplineComponent::GValidateOnLoad && Spline.GetCurrentImplementation() == UE::Spline::ESplineType::Tangent && !Validate())
 		{
 			UE_LOG(LogSplineComponent, Display, TEXT("Internally inconsistent spline data at load time!"))
 		}
@@ -541,12 +566,17 @@ void USplineComponent::Serialize(FArchive& Ar)
 	const FPackageFileVersion ArchiveUEVersion = Ar.UEVer();
 	if (ArchiveUEVersion < VER_UE4_INTERPCURVE_SUPPORTS_LOOPING)
 	{
+		FTangentSpline* TangentSpline = Spline.Get<UE::Spline::ESplineType::Tangent>();
+
 		int32 LegacyNumPoints = WarninglessSplineCurves().Position.Points.Num();
 
 		// The start point is no longer cloned as the endpoint when the spline is looped, so remove the extra endpoint if present
 		if (bClosedLoop && ( GetLocationAtSplinePoint(0, ESplineCoordinateSpace::Local) == GetLocationAtSplinePoint(GetNumberOfSplinePoints() - 1, ESplineCoordinateSpace::Local)))
 		{
-			Spline.RemovePoint(Spline.GetNumControlPoints() - 1);
+			if (TangentSpline)
+			{
+				TangentSpline->RemovePoint(TangentSpline->GetNumberOfPoints() - 1);
+			}
 			WarninglessSplineCurves().Position.Points.RemoveAt(LegacyNumPoints - 1, EAllowShrinking::No);
 			LegacyNumPoints--;
 		}
@@ -563,8 +593,11 @@ void USplineComponent::Serialize(FArchive& Ar)
 		
 		// Importantly, these 2 functions guarantee that 1 rotation and scale value exists for each position,
 		// but they are default values.
-		Spline.ResetRotation();
-		Spline.ResetScale();
+		if (TangentSpline)
+		{
+			TangentSpline->ResetRotation();
+			TangentSpline->ResetScale();
+		}
 
 		UpdateSpline();
 	}
@@ -590,6 +623,21 @@ void USplineComponent::OnRegister()
 #endif
 }
 
+FSplineCurves FSplineCurves::FromSplineInterface(const FSpline& InProxy)
+{
+	FSplineCurves NewCurve;
+
+	if (const FTangentSpline* TangentSpline = InProxy.Get<UE::Spline::ESplineType::Tangent>())
+	{
+		NewCurve.Position = TangentSpline->GetSplinePointsPosition();
+		NewCurve.Rotation = TangentSpline->GetSplinePointsRotation();
+		NewCurve.Scale = TangentSpline->GetSplinePointsScale();
+	}
+
+	NewCurve.UpdateSpline(NewCurve.Position.bIsLooped, false);
+	return NewCurve;
+}
+
 void FSplineCurves::UpdateSpline(bool bClosedLoop, bool bStationaryEndpoints, int32 ReparamStepsPerSegment, bool bLoopPositionOverride, float LoopPosition, const FVector& Scale3D)
 {
 	TRACE_CPUPROFILER_EVENT_SCOPE(FSplineCurves::UpdateSpline);
@@ -651,26 +699,36 @@ void USplineComponent::UpdateSpline()
 {
 	TRACE_CPUPROFILER_EVENT_SCOPE(USplineComponent::UpdateSpline);
 
-	FSpline::FUpdateSplineParams Params{bClosedLoop, bStationaryEndpoints, ReparamStepsPerSegment, bLoopPositionOverride, LoopPosition, GetComponentTransform().GetScale3D()};
+	// If SplineComponent.UseSplineCurves is false and FSpline is not implemented, initialize it from SplineCurves.
+	if (UE::SplineComponent::GUseSplineCurves == false && Spline.GetCurrentImplementation() != UE::Spline::ESplineType::Tangent)
+	{
+		Spline = WarninglessSplineCurves();
+	}
+
+	FTangentSpline* TangentSpline = Spline.Get<UE::Spline::ESplineType::Tangent>();
+
+	FTangentSpline::FUpdateSplineParams Params{bClosedLoop, bStationaryEndpoints, ReparamStepsPerSegment, bLoopPositionOverride, LoopPosition, GetComponentTransform().GetScale3D()};
 	if (UE::SplineComponent::GOnlyUpdateActiveSpline)
 	{
-		if (UE::SplineComponent::ShouldUseSplineCurves())
+		if (ShouldUseSplineCurves())
 		{
 			WarninglessSplineCurves().UpdateSpline(Params.bClosedLoop, Params.bStationaryEndpoints, Params.ReparamStepsPerSegment, Params.bLoopPositionOverride, Params.LoopPosition, Params.Scale3D);
 		}
-		else
+		else if (ensureAlways(TangentSpline)) // ShouldUseSplineCurves should _never_ return false if TangentSpline would not be true
 		{
-			Spline.UpdateSpline(Params);
+			TangentSpline->UpdateSpline(Params);
 		}
 	}
 	else
 	{
 		WarninglessSplineCurves().UpdateSpline(Params.bClosedLoop, Params.bStationaryEndpoints, Params.ReparamStepsPerSegment, Params.bLoopPositionOverride, Params.LoopPosition, Params.Scale3D);
-		Spline.UpdateSpline(Params);
+
+		if (TangentSpline)
+		{
+			TangentSpline->UpdateSpline(Params);
+		}
 	}
 
-	// todo: make Spline replicate
-	//MARK_PROPERTY_DIRTY_FROM_NAME(USplineComponent, Spline, this);
 	MARK_PROPERTY_DIRTY_FROM_NAME(USplineComponent, SplineCurves, this);
 	MARK_PROPERTY_DIRTY_FROM_NAME(USplineComponent, bClosedLoop, this);
 	MARK_PROPERTY_DIRTY_FROM_NAME(USplineComponent, bStationaryEndpoints, this);
@@ -771,9 +829,17 @@ float FSplineCurves::GetSegmentLength(const int32 Index, const float Param, bool
 
 float USplineComponent::GetSegmentLength(const int32 Index, const float Param) const
 {
-	return UE::SplineComponent::ShouldUseSplineCurves()
-		? WarninglessSplineCurves().GetSegmentLength(Index, Param, bClosedLoop, GetComponentTransform().GetScale3D())
-		: Spline.GetSegmentLength(Index, Param, GetComponentTransform().GetScale3D());
+	if (ShouldUseSplineCurves())
+	{
+		return WarninglessSplineCurves().GetSegmentLength(Index, Param, bClosedLoop, GetComponentTransform().GetScale3D());
+	}
+	else if (const FTangentSpline* TangentSpline = Spline.Get<UE::Spline::ESplineType::Tangent>())
+	{
+		return TangentSpline->GetSegmentLength(Index, Param, GetComponentTransform().GetScale3D());
+	}
+
+	ensureAlwaysMsgf(false, TEXT("ShouldUseSplineCurves() returned false but Spline failed to yield an FTangentSpline!"));
+	return 0.f;
 }
 
 float USplineComponent::GetSegmentParamFromLength(const int32 Index, const float Length, const float SegmentLength) const
@@ -801,10 +867,22 @@ float USplineComponent::GetSegmentParamFromLength(const int32 Index, const float
 	//        (df/dt) = d(L(t))/dt = |dP/dt|
 
 	// TODO: check if this works OK with delta InVal != 1.0f
+	// ^ it definitely does not
 
-	int32 NumPoints = UE::SplineComponent::ShouldUseSplineCurves()
-		? WarninglessSplineCurves().Position.Points.Num()
-		: Spline.GetNumControlPoints();
+	int32 NumPoints = 0;
+	if (ShouldUseSplineCurves())
+	{
+		NumPoints = WarninglessSplineCurves().Position.Points.Num();
+	}
+	else if (const FTangentSpline* TangentSpline = Spline.Get<UE::Spline::ESplineType::Tangent>())
+	{
+		NumPoints = TangentSpline->GetNumberOfPoints();
+	}
+	else
+	{
+		ensureAlwaysMsgf(false, TEXT("ShouldUseSplineCurves() returned false but Spline failed to yield an FTangentSpline!"));
+		return 0.f;
+	}
 	
 	const int32 LastPoint = NumPoints - 1;
 
@@ -816,7 +894,7 @@ float USplineComponent::GetSegmentParamFromLength(const int32 Index, const float
 	// two iterations of Newton-Raphson is enough
 	for (int32 Iteration = 0; Iteration < 2; ++Iteration)
 	{
-		float TangentMagnitude = UE::SplineComponent::ShouldUseSplineCurves()
+		float TangentMagnitude = ShouldUseSplineCurves()
 			? WarninglessSplineCurves().Position.EvalDerivative(Index + Param, FVector::ZeroVector).Size()
 			: Spline.EvaluateDerivative(Index + Param).Size();
 		
@@ -832,7 +910,7 @@ float USplineComponent::GetSegmentParamFromLength(const int32 Index, const float
 
 FVector USplineComponent::GetLocationAtSplineInputKey(float InKey, ESplineCoordinateSpace::Type CoordinateSpace) const
 {
-	FVector Location = UE::SplineComponent::ShouldUseSplineCurves()
+	FVector Location = ShouldUseSplineCurves()
 		? WarninglessSplineCurves().Position.Eval(InKey, FVector::ZeroVector)
 		: Spline.Evaluate(InKey);
 
@@ -846,7 +924,7 @@ FVector USplineComponent::GetLocationAtSplineInputKey(float InKey, ESplineCoordi
 
 FVector USplineComponent::GetTangentAtSplineInputKey(float InKey, ESplineCoordinateSpace::Type CoordinateSpace) const
 {
-	FVector Tangent = UE::SplineComponent::ShouldUseSplineCurves()
+	FVector Tangent = ShouldUseSplineCurves()
 		? WarninglessSplineCurves().Position.EvalDerivative(InKey, FVector::ZeroVector)
 		: Spline.EvaluateDerivative(InKey);
 
@@ -860,7 +938,7 @@ FVector USplineComponent::GetTangentAtSplineInputKey(float InKey, ESplineCoordin
 
 FVector USplineComponent::GetDirectionAtSplineInputKey(float InKey, ESplineCoordinateSpace::Type CoordinateSpace) const
 {
-	FVector Direction = UE::SplineComponent::ShouldUseSplineCurves()
+	FVector Direction = ShouldUseSplineCurves()
 		? WarninglessSplineCurves().Position.EvalDerivative(InKey, FVector::ZeroVector).GetSafeNormal()
 		: Spline.EvaluateDerivative(InKey).GetSafeNormal();
 
@@ -880,12 +958,23 @@ FRotator USplineComponent::GetRotationAtSplineInputKey(float InKey, ESplineCoord
 
 FQuat USplineComponent::GetQuaternionAtSplineInputKey(float InKey, ESplineCoordinateSpace::Type CoordinateSpace) const
 {
-	FQuat Quat = UE::SplineComponent::ShouldUseSplineCurves()
-		? WarninglessSplineCurves().Rotation.Eval(InKey, FQuat::Identity)
-		: Spline.EvaluateRotation(InKey);
+	FQuat Quat = FQuat::Identity;
+	if (ShouldUseSplineCurves())
+	{
+		Quat = WarninglessSplineCurves().Rotation.Eval(InKey, FQuat::Identity);
+	}
+	else if (const FTangentSpline* TangentSpline = Spline.Get<UE::Spline::ESplineType::Tangent>())
+	{
+		Quat = TangentSpline->EvaluateRotation(InKey);
+	}
+	else
+	{
+		ensureAlwaysMsgf(false, TEXT("ShouldUseSplineCurves() returned false but Spline failed to yield an FTangentSpline!"));
+		return Quat;
+	}
 	Quat.Normalize();
 
-	FVector Direction = UE::SplineComponent::ShouldUseSplineCurves()
+	FVector Direction = ShouldUseSplineCurves()
 		? WarninglessSplineCurves().Position.EvalDerivative(InKey, FVector::ZeroVector).GetSafeNormal()
 		: Spline.EvaluateDerivative(InKey).GetSafeNormal();
 	
@@ -951,16 +1040,41 @@ float USplineComponent::GetRollAtSplineInputKey(float InKey, ESplineCoordinateSp
 
 FVector USplineComponent::GetScaleAtSplineInputKey(float InKey) const
 {
-	return UE::SplineComponent::ShouldUseSplineCurves()
-		? WarninglessSplineCurves().Scale.Eval(InKey, FVector(1.0f))
-		: Spline.EvaluateScale(InKey);
+	FVector Scale = FVector::OneVector;
+
+	if (ShouldUseSplineCurves())
+	{
+		Scale = WarninglessSplineCurves().Scale.Eval(InKey, FVector::OneVector);
+	}
+	else if (const FTangentSpline* TangentSpline = Spline.Get<UE::Spline::ESplineType::Tangent>())
+	{
+		Scale = TangentSpline->EvaluateScale(InKey);
+	}
+	else
+	{
+		ensureAlwaysMsgf(false, TEXT("ShouldUseSplineCurves() returned false but Spline failed to yield an FTangentSpline!"));
+	}
+
+	return Scale;
 }
 
 float USplineComponent::GetDistanceAlongSplineAtSplineInputKey(float InKey) const
 {
-	const int32 NumPoints = UE::SplineComponent::ShouldUseSplineCurves()
-		? WarninglessSplineCurves().Position.Points.Num()
-		: Spline.GetNumControlPoints();
+	int32 NumPoints = 0;
+	const FTangentSpline* TangentSpline = nullptr;
+	if (ShouldUseSplineCurves())
+	{
+		NumPoints = WarninglessSplineCurves().Position.Points.Num();
+	}
+	else if ((TangentSpline = Spline.Get<UE::Spline::ESplineType::Tangent>()))
+	{
+		NumPoints = TangentSpline->GetNumberOfPoints();
+	}
+	else
+	{
+		ensureAlwaysMsgf(false, TEXT("ShouldUseSplineCurves() returned false but Spline failed to yield an FTangentSpline!"));
+		return 0.f;
+	}
 	
 	const int32 NumSegments = bClosedLoop ? NumPoints : NumPoints - 1;
 
@@ -968,7 +1082,8 @@ float USplineComponent::GetDistanceAlongSplineAtSplineInputKey(float InKey) cons
 	{
 		float Distance;
 		
-		if (UE::SplineComponent::ShouldUseSplineCurves())
+		// !TangentSpline is true only if ShouldUseSplineCurves previously returned true.
+		if (!TangentSpline)
 		{
 			const int32 PointIndex = FMath::FloorToInt(InKey);
 			const float Fraction = InKey - PointIndex;
@@ -977,7 +1092,7 @@ float USplineComponent::GetDistanceAlongSplineAtSplineInputKey(float InKey) cons
 		}
 		else
 		{
-			Distance = Spline.GetDistanceAtParameter(InKey);
+			Distance = TangentSpline->GetDistanceAtParameter(InKey);
 		}
 		
 		return Distance;
@@ -995,18 +1110,32 @@ float USplineComponent::GetDistanceAlongSplineAtLocation(const FVector& InLocati
 	const FVector LocalLocation = (CoordinateSpace == ESplineCoordinateSpace::World) ? GetComponentTransform().InverseTransformPosition(InLocation) : InLocation;
 	float Dummy;
 	
-	const float Key = UE::SplineComponent::ShouldUseSplineCurves()
-		? WarninglessSplineCurves().Position.FindNearest(LocalLocation, Dummy)
-		: Spline.FindNearest(LocalLocation, Dummy);
+	float Key;
+	if (ShouldUseSplineCurves())
+	{
+		Key = WarninglessSplineCurves().Position.FindNearest(LocalLocation, Dummy);
+	}
+	else if (const FTangentSpline* TangentSpline = Spline.Get<UE::Spline::ESplineType::Tangent>())
+	{
+		Key = TangentSpline->FindNearest(LocalLocation, Dummy);
+	}
+	else
+	{
+		ensureAlwaysMsgf(false, TEXT("ShouldUseSplineCurves() returned false but Spline failed to yield an FTangentSpline!"));
+		return 0.f;
+	}
 	
 	return GetDistanceAlongSplineAtSplineInputKey(Key);
 }
 
 int32 USplineComponent::GetNumberOfPropertyValues(FName PropertyName) const
 {
-	if (UE::SplineComponent::ShouldUseSplineCurves() == false && Spline.SupportsAttributes())
+	if (ShouldUseSplineCurves() == false)
 	{
-		return Spline.NumAttributeValues<float>(PropertyName);
+		if (const FTangentSpline* TangentSpline = Spline.Get<UE::Spline::ESplineType::Tangent>())
+		{
+			return TangentSpline->NumAttributeValues<float>(PropertyName);
+		}
 	}
 
 	return 0;
@@ -1014,16 +1143,19 @@ int32 USplineComponent::GetNumberOfPropertyValues(FName PropertyName) const
 
 TArray<FName> USplineComponent::GetFloatPropertyChannels() const
 {
-	if (UE::SplineComponent::ShouldUseSplineCurves() == false && Spline.SupportsAttributes())
+	if (ShouldUseSplineCurves() == false)
 	{
-		return Spline.GetFloatPropertyChannels();
+		if (const FTangentSpline* TangentSpline = Spline.Get<UE::Spline::ESplineType::Tangent>())
+		{
+			return TangentSpline->GetAttributeChannelNamesByValueType<float>();
+		}
 	}
 
 	return TArray<FName>();
 }
 
 template<class T>
-bool CreatePropertyChannel(const USplineMetadata* Metadata, const FSpline& Spline, FName PropertyName)
+bool CreatePropertyChannel(const USplineMetadata* Metadata, FSpline& Spline, FName PropertyName)
 {
 	if (Metadata && Metadata->GetClass()->FindPropertyByName(PropertyName))
 	{
@@ -1031,11 +1163,20 @@ bool CreatePropertyChannel(const USplineMetadata* Metadata, const FSpline& Splin
 		return false;
 	}
 	
-	if (UE::SplineComponent::GUseSplineCurves == false && Spline.SupportsAttributes())
+	if (UE::SplineComponent::GUseSplineCurves == false)
 	{
-		if (!Spline.HasAttributeChannel(PropertyName))
+		if (FTangentSpline* TangentSpline = Spline.Get<UE::Spline::ESplineType::Tangent>())
 		{
-			return Spline.CreateAttributeChannel<T>(PropertyName);
+			if (!TangentSpline->HasAttributeChannel(PropertyName))
+			{
+				if (UE::Geometry::Spline::TTangentBezierSpline<T>* NewAttrChannel = TangentSpline->CreateAttributeChannel<UE::Geometry::Spline::TTangentBezierSpline<T>>(PropertyName))
+				{
+					NewAttrChannel->SetPreInfinityMode(UE::Geometry::Spline::EOutOfBoundsHandlingMode::Constant);
+					NewAttrChannel->SetPostInfinityMode(UE::Geometry::Spline::EOutOfBoundsHandlingMode::Constant);
+
+					return true;
+				}
+			}
 		}
 	}
 
@@ -1054,11 +1195,14 @@ T GetPropertyAtSplineInputKey(const USplineMetadata* Metadata, const FSpline& Sp
 		}
 	}
 	
-	if (UE::SplineComponent::GUseSplineCurves == false && Spline.SupportsAttributes())
+	if (UE::SplineComponent::GUseSplineCurves == false)
 	{
-		if (Spline.HasAttributeChannel(PropertyName))
+		if (const FTangentSpline* TangentSpline = Spline.Get<UE::Spline::ESplineType::Tangent>())
 		{
-			return Spline.EvaluateAttribute<T>(InKey, PropertyName);
+			if (TangentSpline->HasAttributeChannel(PropertyName))
+			{
+				return TangentSpline->EvaluateAttribute<T>(PropertyName, InKey);
+			}
 		}
 	}
 
@@ -1066,13 +1210,16 @@ T GetPropertyAtSplineInputKey(const USplineMetadata* Metadata, const FSpline& Sp
 }
 
 template<class T>
-int32 SetPropertyAtSplineInputKey(const FSpline& Spline, float InKey, const T& InValue, FName PropertyName)
+int32 AddPropertyAtSplineInputKey(FSpline& Spline, float InKey, const T& InValue, FName PropertyName)
 {
-	if (UE::SplineComponent::GUseSplineCurves == false && Spline.SupportsAttributes())
+	if (UE::SplineComponent::GUseSplineCurves == false)
 	{
-		if (Spline.HasAttributeChannel(PropertyName))
+		if (FTangentSpline* TangentSpline = Spline.Get<UE::Spline::ESplineType::Tangent>())
 		{
-			return Spline.AddAttributeValue<T>(InKey, InValue, PropertyName);
+			if (TangentSpline->HasAttributeChannel(PropertyName))
+			{
+				return TangentSpline->AddAttributeValue<T>(PropertyName, InValue, InKey);
+			}
 		}
 	}
 
@@ -1081,9 +1228,15 @@ int32 SetPropertyAtSplineInputKey(const FSpline& Spline, float InKey, const T& I
 
 float GetInputKeyAtIndex(const FSpline& Spline, int32 Index, FName PropertyName)
 {
-	if (UE::SplineComponent::GUseSplineCurves == false && Spline.SupportsAttributes())
+	if (UE::SplineComponent::GUseSplineCurves == false)
 	{
-		return Spline.GetAttributeParameter<float>(Index, PropertyName);
+		if (const FTangentSpline* TangentSpline = Spline.Get<UE::Spline::ESplineType::Tangent>())
+		{
+			if (TangentSpline->HasAttributeChannel(PropertyName))
+			{
+				return TangentSpline->GetAttributeParameter<float>(PropertyName, Index);
+			}
+		}
 	}
 
 	return 0.f;
@@ -1091,9 +1244,15 @@ float GetInputKeyAtIndex(const FSpline& Spline, int32 Index, FName PropertyName)
 
 int32 SetInputKeyAtIndex(FSpline& Spline, int32 Index, float InKey, FName PropertyName)
 {
-	if (UE::SplineComponent::GUseSplineCurves == false && Spline.SupportsAttributes())
+	if (UE::SplineComponent::GUseSplineCurves == false)
 	{
-		return Spline.SetAttributeParameter<float>(Index, InKey , PropertyName);
+		if (FTangentSpline* TangentSpline = Spline.Get<UE::Spline::ESplineType::Tangent>())
+		{
+			if (TangentSpline->HasAttributeChannel(PropertyName))
+			{
+				return TangentSpline->SetAttributeParameter<float>(PropertyName, Index, InKey);
+			}
+		}
 	}
 
 	return INDEX_NONE;
@@ -1102,9 +1261,15 @@ int32 SetInputKeyAtIndex(FSpline& Spline, int32 Index, float InKey, FName Proper
 template<typename T>
 float GetPropertyAtIndex(const FSpline& Spline, int32 Index, FName PropertyName)
 {
-	if (UE::SplineComponent::GUseSplineCurves == false && Spline.SupportsAttributes())
+	if (UE::SplineComponent::GUseSplineCurves == false)
 	{
-		return Spline.GetAttributeValue<T>(Index, PropertyName);
+		if (const FTangentSpline* TangentSpline = Spline.Get<UE::Spline::ESplineType::Tangent>())
+		{
+			if (TangentSpline->HasAttributeChannel(PropertyName))
+			{
+				return TangentSpline->GetAttributeValue<T>(PropertyName, Index);
+			}
+		}
 	}
 
 	return 0.f;
@@ -1113,30 +1278,54 @@ float GetPropertyAtIndex(const FSpline& Spline, int32 Index, FName PropertyName)
 template<typename T>
 void SetPropertyAtIndex(FSpline& Spline, int32 Index, float Value, FName PropertyName)
 {
-	if (UE::SplineComponent::GUseSplineCurves == false && Spline.SupportsAttributes())
+	if (UE::SplineComponent::GUseSplineCurves == false)
 	{
-		Spline.SetAttributeValue<T>(Index, Value, PropertyName);
+		if (FTangentSpline* TangentSpline = Spline.Get<UE::Spline::ESplineType::Tangent>())
+		{
+			if (TangentSpline->HasAttributeChannel(PropertyName))
+			{
+				TangentSpline->SetAttributeValue<T>(PropertyName, Value, Index);
+			}
+		}
 	}
 }
 
 void USplineComponent::RemovePropertyAtIndex(int32 Index, FName PropertyName)
 {
-	if (UE::SplineComponent::ShouldUseSplineCurves() == false && Spline.SupportsAttributes())
+	if (ShouldUseSplineCurves() == false)
 	{
-		Spline.RemoveAttributeValue<float>(Index, PropertyName);
+		if (FTangentSpline* TangentSpline = Spline.Get<UE::Spline::ESplineType::Tangent>())
+		{
+			if (TangentSpline->HasAttributeChannel(PropertyName))
+			{
+				TangentSpline->RemoveAttributeValue<float>(PropertyName, Index);
+			}
+		}
 	}
 }
 
 bool USplineComponent::SupportsAttributes() const
 {
-	return UE::SplineComponent::ShouldUseSplineCurves() == false && Spline.SupportsAttributes();
+	// If GetAs yields a Tangent spline, we support attributes.
+	if (ShouldUseSplineCurves() == false && Spline.Get<UE::Spline::ESplineType::Tangent>())
+	{
+		return true;
+	}
+
+	return false;
 }
 
 bool USplineComponent::RemovePropertyChannel(FName PropertyName)
 {
-	if (UE::SplineComponent::ShouldUseSplineCurves() == false && Spline.SupportsAttributes())
+	if (ShouldUseSplineCurves() == false)
 	{
-		return Spline.RemoveAttributeChannel(PropertyName);
+		if (FTangentSpline* TangentSpline = Spline.Get<UE::Spline::ESplineType::Tangent>())
+		{
+			if (TangentSpline->HasAttributeChannel(PropertyName))
+			{
+				return TangentSpline->RemoveAttributeChannel(PropertyName);
+			}
+		}
 	}
 
 	return false;
@@ -1152,9 +1341,9 @@ float USplineComponent::GetFloatPropertyAtSplineInputKey(float InKey, FName Prop
 	return GetPropertyAtSplineInputKey<float>(GetSplinePointsMetadata(), Spline, InKey, PropertyName);
 }
 
-int32 USplineComponent::SetFloatPropertyAtSplineInputKey(float InKey, float Value, FName PropertyName) const
+int32 USplineComponent::AddFloatPropertyAtSplineInputKey(float InKey, float Value, FName PropertyName)
 {
-	return SetPropertyAtSplineInputKey<float>(Spline, InKey, Value, PropertyName);
+	return AddPropertyAtSplineInputKey<float>(Spline, InKey, Value, PropertyName);
 }
 
 float USplineComponent::GetFloatPropertyInputKeyAtIndex(int32 Index, FName PropertyName) const
@@ -1334,7 +1523,10 @@ void USplineComponent::AddPoint(const FSplinePoint& InSplinePoint, bool bUpdateS
 		CIM_CurveAuto
 	), Index);
 	
-	Spline.AddPoint(InSplinePoint);
+	if (FTangentSpline* TangentSpline = Spline.Get<UE::Spline::ESplineType::Tangent>())
+	{
+		TangentSpline->AddPoint(InSplinePoint);
+	}
 	
 	if (USplineMetadata* Metadata = GetSplinePointsMetadata())
 	{
@@ -1408,11 +1600,14 @@ void USplineComponent::AddSplinePoint(const FVector& Position, ESplineCoordinate
 	NewPoint.Scale = FVector(1.f);
 	NewPoint.Type = ConvertInterpCurveModeToSplinePointType(CIM_CurveAuto);
 
-	switch (UE::SplineComponent::GAddSplinePointImplementation)
+	if (FTangentSpline* TangentSpline = Spline.Get<UE::Spline::ESplineType::Tangent>())
 	{
-	default: // fallthrough
-	case 0: Spline.AddPoint(NewPoint); break;
-	case 1: Spline.InsertPoint(NewPoint, GetNumberOfSplinePoints()); break;
+		switch (UE::SplineComponent::GAddSplinePointImplementation)
+		{
+		default: // fallthrough
+		case 0: TangentSpline->AddPoint(NewPoint); break;
+		case 1: TangentSpline->InsertPoint(NewPoint, GetNumberOfSplinePoints()); break;
+		}
 	}
 
 	WarninglessSplineCurves().Position.Points.Emplace(InKey, TransformedPosition, FVector::ZeroVector, FVector::ZeroVector, CIM_CurveAuto);
@@ -1461,7 +1656,10 @@ void USplineComponent::AddSplinePointAtIndex(const FVector& Position, int32 Inde
 		NewPoint.Rotation = FQuat::Identity.Rotator();
 		NewPoint.Scale = FVector(1.f);
 
-		Spline.InsertPoint(NewPoint, Index);
+		if (FTangentSpline* TangentSpline = Spline.Get<UE::Spline::ESplineType::Tangent>())
+		{
+			TangentSpline->InsertPoint(NewPoint, Index);
+		}
 
 		WarninglessSplineCurves().Position.Points.Insert(FInterpCurvePoint<FVector>(InKey, TransformedPosition, FVector::ZeroVector, FVector::ZeroVector, CIM_CurveAuto), Index);
 		WarninglessSplineCurves().Rotation.Points.Insert(FInterpCurvePoint<FQuat>(InKey, FQuat::Identity, FQuat::Identity, FQuat::Identity, CIM_CurveAuto), Index);
@@ -1506,7 +1704,10 @@ void USplineComponent::RemoveSplinePoint(int32 Index, bool bUpdateSpline)
 
 	if (Index >= 0 && Index < NumPoints)
 	{
-		Spline.RemovePoint(Index);
+		if (FTangentSpline* TangentSpline = Spline.Get<UE::Spline::ESplineType::Tangent>())
+		{
+			TangentSpline->RemovePoint(Index);
+		}
 
 		float KnotDelta = 0.0f;
 		if (Index > 0)
@@ -1574,6 +1775,7 @@ void USplineComponent::SetSplinePoints(const TArray<FVector>& Points, ESplineCoo
 	}
 
 	float InputKey = 0.0f;
+	FTangentSpline* TangentSpline = Spline.Get<UE::Spline::ESplineType::Tangent>();
 	for (const auto& Point : Points)
 	{
 		const FVector TransformedPoint = (CoordinateSpace == ESplineCoordinateSpace::World) ?
@@ -1587,7 +1789,10 @@ void USplineComponent::SetSplinePoints(const TArray<FVector>& Points, ESplineCoo
 		NewPoint.Rotation = FQuat::Identity.Rotator();
 		NewPoint.Scale = FVector(1.f);
 
-		Spline.AddPoint(NewPoint);
+		if (TangentSpline)
+		{
+			TangentSpline->AddPoint(NewPoint);
+		}
 
 		WarninglessSplineCurves().Position.Points.Emplace(InputKey, TransformedPoint, FVector::ZeroVector, FVector::ZeroVector, CIM_CurveAuto);
 		WarninglessSplineCurves().Rotation.Points.Emplace(InputKey, FQuat::Identity, FQuat::Identity, FQuat::Identity, CIM_CurveAuto);
@@ -1625,7 +1830,11 @@ void USplineComponent::SetLocationAtSplinePoint(int32 PointIndex, const FVector&
 		const FVector TransformedLocation = (CoordinateSpace == ESplineCoordinateSpace::World) ?
 			GetComponentTransform().InverseTransformPosition(InLocation) : InLocation;
 
-		Spline.SetLocation(PointIndex, TransformedLocation);
+		if (FTangentSpline* TangentSpline = Spline.Get<UE::Spline::ESplineType::Tangent>())
+		{
+			TangentSpline->SetLocation(PointIndex, TransformedLocation);
+		}
+
 		WarninglessSplineCurves().Position.Points[PointIndex].OutVal = TransformedLocation;
 
 		if (bUpdateSpline)
@@ -1658,8 +1867,11 @@ void USplineComponent::SetTangentsAtSplinePoint(int32 PointIndex, const FVector&
 		const FVector TransformedLeaveTangent = (CoordinateSpace == ESplineCoordinateSpace::World) ?
 			GetComponentTransform().InverseTransformVector(InLeaveTangent) : InLeaveTangent;
 
-		Spline.SetInTangent(PointIndex, TransformedArriveTangent);
-		Spline.SetOutTangent(PointIndex, TransformedLeaveTangent);
+		if (FTangentSpline* TangentSpline = Spline.Get<UE::Spline::ESplineType::Tangent>())
+		{
+			TangentSpline->SetInTangent(PointIndex, TransformedArriveTangent);
+			TangentSpline->SetOutTangent(PointIndex, TransformedLeaveTangent);
+		}
 
 		WarninglessSplineCurves().Position.Points[PointIndex].ArriveTangent = TransformedArriveTangent;
 		WarninglessSplineCurves().Position.Points[PointIndex].LeaveTangent = TransformedLeaveTangent;
@@ -1689,7 +1901,12 @@ void USplineComponent::SetUpVectorAtSplinePoint(int32 PointIndex, const FVector&
 			GetComponentTransform().InverseTransformVector(InUpVector.GetSafeNormal()) : InUpVector.GetSafeNormal();
 
 		FQuat Quat = FQuat::FindBetween(DefaultUpVector, TransformedUpVector);
-		Spline.SetRotation(PointIndex, Quat);
+
+		if (FTangentSpline* TangentSpline = Spline.Get<UE::Spline::ESplineType::Tangent>())
+		{
+			TangentSpline->SetRotation(PointIndex, Quat);
+		}
+
 		WarninglessSplineCurves().Rotation.Points[PointIndex].OutVal = Quat;
 
 		if (bUpdateSpline)
@@ -1757,8 +1974,12 @@ void USplineComponent::SetScaleAtSplinePoint(int32 PointIndex, const FVector& In
 	
 	if ((PointIndex >= 0) && (PointIndex < NumPoints))
 	{
+		if (FTangentSpline* TangentSpline = Spline.Get<UE::Spline::ESplineType::Tangent>())
+		{
+			TangentSpline->SetScale(PointIndex, InScaleVector);
+		}
+
 		WarninglessSplineCurves().Scale.Points[PointIndex].OutVal = InScaleVector;
-		Spline.SetScale(PointIndex, InScaleVector);
 
 		if (bUpdateSpline)
 		{
@@ -1780,9 +2001,20 @@ ESplinePointType::Type USplineComponent::GetSplinePointType(int32 PointIndex) co
 	
 	if (PointIndex >= 0 && PointIndex < NumPoints)
 	{
-		const EInterpCurveMode Mode = UE::SplineComponent::ShouldUseSplineCurves()
-			? WarninglessSplineCurves().Position.Points[PointIndex].InterpMode.GetValue()
-			: Spline.GetSplinePointType(PointIndex);
+		EInterpCurveMode Mode = CIM_Unknown;
+
+		if (ShouldUseSplineCurves())
+		{
+			Mode = WarninglessSplineCurves().Position.Points[PointIndex].InterpMode.GetValue();
+		}
+		else if (const FTangentSpline* TangentSpline = Spline.Get<UE::Spline::ESplineType::Tangent>())
+		{
+			Mode = TangentSpline->GetSplinePointType(PointIndex);
+		}
+		else
+		{
+			ensureAlwaysMsgf(false, TEXT("ShouldUseSplineCurves() returned false but Spline failed to yield an FTangentSpline!"));
+		}
 		
 		return ConvertInterpCurveModeToSplinePointType(Mode);
 	}
@@ -1796,7 +2028,11 @@ void USplineComponent::SetSplinePointType(int32 PointIndex, ESplinePointType::Ty
 	
 	if (PointIndex >= 0 && PointIndex < NumPoints)
 	{
-		Spline.SetSplinePointType(PointIndex, ConvertSplinePointTypeToInterpCurveMode(Type));
+		if (FTangentSpline* TangentSpline = Spline.Get<UE::Spline::ESplineType::Tangent>())
+		{
+			TangentSpline->SetSplinePointType(PointIndex, ConvertSplinePointTypeToInterpCurveMode(Type));
+		}
+
 		WarninglessSplineCurves().Position.Points[PointIndex].InterpMode = ConvertSplinePointTypeToInterpCurveMode(Type);
 		
 		if (bUpdateSpline)
@@ -1815,9 +2051,19 @@ void USplineComponent::SetSplinePointType(int32 PointIndex, ESplinePointType::Ty
 
 int32 USplineComponent::GetNumberOfSplinePoints() const
 {
-	return UE::SplineComponent::ShouldUseSplineCurves()
-		? WarninglessSplineCurves().Position.Points.Num()
-		: Spline.GetNumControlPoints();
+	if (ShouldUseSplineCurves())
+	{
+		return WarninglessSplineCurves().Position.Points.Num();
+	}
+	if (const FTangentSpline* TangentSpline = Spline.Get<UE::Spline::ESplineType::Tangent>())
+	{
+		return TangentSpline->GetNumberOfPoints();
+	}
+	else
+	{
+		ensureAlwaysMsgf(false, TEXT("ShouldUseSplineCurves() returned false but Spline failed to yield an FTangentSpline!"));
+		return 0;
+	}
 }
 
 int32 USplineComponent::GetNumberOfSplineSegments() const
@@ -1838,22 +2084,37 @@ float USplineComponent::GetInputKeyValueAtSplinePoint(int32 PointIndex) const
 
 	if (bClosedLoop && PointIndex == GetNumberOfSplinePoints())
 	{
-		if (UE::SplineComponent::ShouldUseSplineCurves())
+		if (ShouldUseSplineCurves())
 		{
 			return WarninglessSplineCurves().Position.Points[PointIndex - 1].InVal + WarninglessSplineCurves().Position.LoopKeyOffset;
 		}
+		else if (const FTangentSpline* TangentSpline = Spline.Get<UE::Spline::ESplineType::Tangent>())
+		{
+			return TangentSpline->GetParameterAtIndex(PointIndex);
+		}
 		else
 		{
-			return Spline.GetParameterAtIndex(PointIndex);
+			ensureAlwaysMsgf(false, TEXT("ShouldUseSplineCurves() returned false but Spline failed to yield an FTangentSpline!"));
+			return 0.f;
 		}
 	}
 	else
 	{
 		PointIndex = GetClampedIndex(PointIndex);
 
-		return UE::SplineComponent::ShouldUseSplineCurves()
-			? WarninglessSplineCurves().Position.Points[PointIndex].InVal
-			: Spline.GetParameterAtIndex(PointIndex);
+		if (ShouldUseSplineCurves())
+		{
+			return WarninglessSplineCurves().Position.Points[PointIndex].InVal;
+		}
+		else if (const FTangentSpline* TangentSpline = Spline.Get<UE::Spline::ESplineType::Tangent>())
+		{
+			return TangentSpline->GetParameterAtIndex(PointIndex);
+		}
+		else
+		{
+			ensureAlwaysMsgf(false, TEXT("ShouldUseSplineCurves() returned false but Spline failed to yield an FTangentSpline!"));
+			return 0.f;
+		}
 	}
 }
 
@@ -1879,9 +2140,20 @@ FVector USplineComponent::GetLocationAtSplinePoint(int32 PointIndex, ESplineCoor
 
 	PointIndex = GetClampedIndex(PointIndex);
 	
-	const FVector Location = UE::SplineComponent::ShouldUseSplineCurves()
-		? WarninglessSplineCurves().Position.Points[PointIndex].OutVal
-		: Spline.GetLocation(PointIndex);
+	FVector Location = FVector::Zero();
+	if (ShouldUseSplineCurves())
+	{
+		Location = WarninglessSplineCurves().Position.Points[PointIndex].OutVal;
+	}
+	else if (const FTangentSpline* TangentSpline = Spline.Get<UE::Spline::ESplineType::Tangent>())
+	{
+		Location = TangentSpline->GetLocation(PointIndex);
+	}
+	else
+	{
+		ensureAlwaysMsgf(false, TEXT("ShouldUseSplineCurves() returned false but Spline failed to yield an FTangentSpline!"));
+		return Location;
+	}
 	
 	return (CoordinateSpace == ESplineCoordinateSpace::World) ? GetComponentTransform().TransformPosition(Location) : Location;
 }
@@ -1905,9 +2177,21 @@ FVector USplineComponent::GetArriveTangentAtSplinePoint(int32 PointIndex, ESplin
 
 	PointIndex = GetClampedIndex(PointIndex);
 
-	const FVector Tangent = UE::SplineComponent::ShouldUseSplineCurves()
-		? WarninglessSplineCurves().Position.Points[PointIndex].ArriveTangent
-		: Spline.GetInTangent(PointIndex);
+	FVector Tangent = FVector::ZeroVector;
+
+	if (ShouldUseSplineCurves())
+	{
+		Tangent = WarninglessSplineCurves().Position.Points[PointIndex].ArriveTangent;
+	}
+	else if (const FTangentSpline* TangentSpline = Spline.Get<UE::Spline::ESplineType::Tangent>())
+	{
+		Tangent = TangentSpline->GetInTangent(PointIndex);
+	}
+	else
+	{
+		ensureAlwaysMsgf(false, TEXT("ShouldUseSplineCurves() returned false but Spline failed to yield an FTangentSpline!"));
+		return Tangent;
+	}
 
 	return (CoordinateSpace == ESplineCoordinateSpace::World) ? GetComponentTransform().TransformVector(Tangent) : Tangent;
 }
@@ -1921,9 +2205,21 @@ FVector USplineComponent::GetLeaveTangentAtSplinePoint(int32 PointIndex, ESpline
 
 	PointIndex = GetClampedIndex(PointIndex);
 
-	const FVector Tangent = UE::SplineComponent::ShouldUseSplineCurves()
-		? WarninglessSplineCurves().Position.Points[PointIndex].LeaveTangent
-		: Spline.GetOutTangent(PointIndex);
+	FVector Tangent = FVector::ZeroVector;
+
+	if (ShouldUseSplineCurves())
+	{
+		Tangent = WarninglessSplineCurves().Position.Points[PointIndex].LeaveTangent;
+	}
+	else if (const FTangentSpline* TangentSpline = Spline.Get<UE::Spline::ESplineType::Tangent>())
+	{
+		Tangent = TangentSpline->GetOutTangent(PointIndex);
+	}
+	else
+	{
+		ensureAlwaysMsgf(false, TEXT("ShouldUseSplineCurves() returned false but Spline failed to yield an FTangentSpline!"));
+		return Tangent;
+	}
 
 	return (CoordinateSpace == ESplineCoordinateSpace::World) ? GetComponentTransform().TransformVector(Tangent) : Tangent;
 }
@@ -1962,9 +2258,22 @@ FVector USplineComponent::GetScaleAtSplinePoint(int32 PointIndex) const
 
 	PointIndex = GetClampedIndex(PointIndex);
 
-	return UE::SplineComponent::ShouldUseSplineCurves()
-		? WarninglessSplineCurves().Scale.Points[PointIndex].OutVal
-		: Spline.GetScale(PointIndex);
+	FVector Scale = FVector::OneVector;
+
+	if (ShouldUseSplineCurves())
+	{
+		Scale = WarninglessSplineCurves().Scale.Points[PointIndex].OutVal;
+	}
+	else if (const FTangentSpline* TangentSpline = Spline.Get<UE::Spline::ESplineType::Tangent>())
+	{
+		Scale = TangentSpline->GetScale(PointIndex);
+	}
+	else
+	{
+		ensureAlwaysMsgf(false, TEXT("ShouldUseSplineCurves() returned false but Spline failed to yield an FTangentSpline!"));
+	}
+
+	return Scale;
 }
 
 FTransform USplineComponent::GetTransformAtSplinePoint(int32 PointIndex, ESplineCoordinateSpace::Type CoordinateSpace, bool bUseScale) const
@@ -1987,7 +2296,7 @@ float USplineComponent::GetDistanceAlongSplineAtSplinePoint(int32 PointIndex) co
 		return GetSplineLength();
 	}
 	
-	if (UE::SplineComponent::ShouldUseSplineCurves())
+	if (ShouldUseSplineCurves())
 	{
 		const int32 NumPoints = GetNumberOfSplinePoints();
         const int32 NumSegments = bClosedLoop ? NumPoints : NumPoints - 1;
@@ -2000,10 +2309,14 @@ float USplineComponent::GetDistanceAlongSplineAtSplinePoint(int32 PointIndex) co
         	return WarninglessSplineCurves().ReparamTable.Points[PointIndex * ReparamStepsPerSegment].InVal;
         }
 	}
+	else if (const FTangentSpline* TangentSpline = Spline.Get<UE::Spline::ESplineType::Tangent>())
+	{
+		const float ParameterAtIndex = TangentSpline->GetParameterAtIndex(PointIndex);
+		return TangentSpline->GetDistanceAtParameter(ParameterAtIndex);
+	}
 	else
 	{
-		const float ParameterAtIndex = Spline.GetParameterAtIndex(PointIndex);
-		return Spline.GetDistanceAtParameter(ParameterAtIndex);
+		ensureAlwaysMsgf(false, TEXT("ShouldUseSplineCurves() returned false but Spline failed to yield an FTangentSpline!"));
 	}
 
 	return 0.0f;
@@ -2024,9 +2337,20 @@ float FSplineCurves::GetSplineLength() const
 
 float USplineComponent::GetSplineLength() const
 {
-	return UE::SplineComponent::ShouldUseSplineCurves()
-		? WarninglessSplineCurves().GetSplineLength()
-		: Spline.GetSplineLength();
+	if (ShouldUseSplineCurves())
+	{
+		return WarninglessSplineCurves().GetSplineLength();
+	}
+	else if (const FTangentSpline* TangentSpline = Spline.Get<UE::Spline::ESplineType::Tangent>())
+	{
+		const float MaxParameter = TangentSpline->GetParameterSpace().Max;
+		return TangentSpline->GetDistanceAtParameter(MaxParameter);
+	}
+	else
+	{
+		ensureAlwaysMsgf(false, TEXT("ShouldUseSplineCurves() returned false but Spline failed to yield an FTangentSpline!"));
+		return 0.f;
+	}
 }
 
 void USplineComponent::SetDefaultUpVector(const FVector& UpVector, ESplineCoordinateSpace::Type CoordinateSpace)
@@ -2077,9 +2401,19 @@ float USplineComponent::GetInputKeyValueAtDistanceAlongSpline(float Distance) co
 		return 0.0f;
 	}
 	
-	return UE::SplineComponent::ShouldUseSplineCurves()
-		? WarninglessSplineCurves().ReparamTable.Eval(Distance, 0.0f)
-		: Spline.GetParameterAtDistance(Distance);
+	if (ShouldUseSplineCurves())
+	{
+		return WarninglessSplineCurves().ReparamTable.Eval(Distance, 0.0f);
+	}
+	else if (const FTangentSpline* TangentSpline = Spline.Get<UE::Spline::ESplineType::Tangent>())
+	{
+		return TangentSpline->GetParameterAtDistance(Distance);
+	}
+	else
+	{
+		ensureAlwaysMsgf(false, TEXT("ShouldUseSplineCurves() returned false but Spline failed to yield an FTangentSpline!"));
+		return 0.f;
+	}
 }
 
 float USplineComponent::GetTimeAtDistanceAlongSpline(float Distance) const
@@ -2091,9 +2425,20 @@ float USplineComponent::GetTimeAtDistanceAlongSpline(float Distance) const
 		return 0.0f;
 	}
 
-	const float Param = UE::SplineComponent::ShouldUseSplineCurves() 
-		? WarninglessSplineCurves().ReparamTable.Eval(Distance, 0.0f)
-		: Spline.GetParameterAtDistance(Distance);
+	float Param;
+	if (ShouldUseSplineCurves())
+	{
+		Param = WarninglessSplineCurves().ReparamTable.Eval(Distance, 0.0f);
+	}
+	else if (const FTangentSpline* TangentSpline = Spline.Get<UE::Spline::ESplineType::Tangent>())
+	{
+		Param = TangentSpline->GetParameterAtDistance(Distance);
+	}
+	else
+	{
+		ensureAlwaysMsgf(false, TEXT("ShouldUseSplineCurves() returned false but Spline failed to yield an FTangentSpline!"));
+		return 0.f;
+	}
 
 	// Convert Param to T and interpolate the duration interval (which yields Time).
 	return UE::Geometry::FInterval1f(0.f, Duration).Interpolate(GetInputKeyRange().GetT(Param));
@@ -2101,82 +2446,62 @@ float USplineComponent::GetTimeAtDistanceAlongSpline(float Distance) const
 
 FVector USplineComponent::GetLocationAtDistanceAlongSpline(float Distance, ESplineCoordinateSpace::Type CoordinateSpace) const
 {
-	const float Param = UE::SplineComponent::ShouldUseSplineCurves()
-		? WarninglessSplineCurves().ReparamTable.Eval(Distance, 0.0f)
-		: Spline.GetParameterAtDistance(Distance);
+	const float Param = GetInputKeyValueAtDistanceAlongSpline(Distance);
 	return GetLocationAtSplineInputKey(Param, CoordinateSpace);
 }
 
 FVector USplineComponent::GetTangentAtDistanceAlongSpline(float Distance, ESplineCoordinateSpace::Type CoordinateSpace) const
 {
-	const float Param = UE::SplineComponent::ShouldUseSplineCurves()
-		? WarninglessSplineCurves().ReparamTable.Eval(Distance, 0.0f)
-		: Spline.GetParameterAtDistance(Distance);
+	const float Param = GetInputKeyValueAtDistanceAlongSpline(Distance);
 	return GetTangentAtSplineInputKey(Param, CoordinateSpace);
 }
 
 
 FVector USplineComponent::GetDirectionAtDistanceAlongSpline(float Distance, ESplineCoordinateSpace::Type CoordinateSpace) const
 {
-	const float Param = UE::SplineComponent::ShouldUseSplineCurves()
-		? WarninglessSplineCurves().ReparamTable.Eval(Distance, 0.0f)
-		: Spline.GetParameterAtDistance(Distance);
+	const float Param = GetInputKeyValueAtDistanceAlongSpline(Distance);
 	return GetDirectionAtSplineInputKey(Param, CoordinateSpace);
 }
 
 FQuat USplineComponent::GetQuaternionAtDistanceAlongSpline(float Distance, ESplineCoordinateSpace::Type CoordinateSpace) const
 {
-	const float Param = UE::SplineComponent::ShouldUseSplineCurves()
-		? WarninglessSplineCurves().ReparamTable.Eval(Distance, 0.0f)
-		: Spline.GetParameterAtDistance(Distance);
+	const float Param = GetInputKeyValueAtDistanceAlongSpline(Distance);
 	return GetQuaternionAtSplineInputKey(Param, CoordinateSpace);
 }
 
 FRotator USplineComponent::GetRotationAtDistanceAlongSpline(float Distance, ESplineCoordinateSpace::Type CoordinateSpace) const
 {
-	const float Param = UE::SplineComponent::ShouldUseSplineCurves()
-		? WarninglessSplineCurves().ReparamTable.Eval(Distance, 0.0f)
-		: Spline.GetParameterAtDistance(Distance);
+	const float Param = GetInputKeyValueAtDistanceAlongSpline(Distance);
 	return GetRotationAtSplineInputKey(Param, CoordinateSpace);
 }
 
 FVector USplineComponent::GetUpVectorAtDistanceAlongSpline(float Distance, ESplineCoordinateSpace::Type CoordinateSpace) const
 {
-	const float Param = UE::SplineComponent::ShouldUseSplineCurves()
-		? WarninglessSplineCurves().ReparamTable.Eval(Distance, 0.0f)
-		: Spline.GetParameterAtDistance(Distance);
+	const float Param = GetInputKeyValueAtDistanceAlongSpline(Distance);
 	return GetUpVectorAtSplineInputKey(Param, CoordinateSpace);
 }
 
 FVector USplineComponent::GetRightVectorAtDistanceAlongSpline(float Distance, ESplineCoordinateSpace::Type CoordinateSpace) const
 {
-	const float Param = UE::SplineComponent::ShouldUseSplineCurves()
-		? WarninglessSplineCurves().ReparamTable.Eval(Distance, 0.0f)
-		: Spline.GetParameterAtDistance(Distance);
+	const float Param = GetInputKeyValueAtDistanceAlongSpline(Distance);
 	return GetRightVectorAtSplineInputKey(Param, CoordinateSpace);
 }
 
 float USplineComponent::GetRollAtDistanceAlongSpline(float Distance, ESplineCoordinateSpace::Type CoordinateSpace) const
 {
-	const float Param = UE::SplineComponent::ShouldUseSplineCurves()
-		? WarninglessSplineCurves().ReparamTable.Eval(Distance, 0.0f)
-		: Spline.GetParameterAtDistance(Distance);
+	const float Param = GetInputKeyValueAtDistanceAlongSpline(Distance);
 	return GetRollAtSplineInputKey(Param, CoordinateSpace);
 }
 
 FVector USplineComponent::GetScaleAtDistanceAlongSpline(float Distance) const
 {
-	const float Param = UE::SplineComponent::ShouldUseSplineCurves()
-		? WarninglessSplineCurves().ReparamTable.Eval(Distance, 0.0f)
-		: Spline.GetParameterAtDistance(Distance);
+	const float Param = GetInputKeyValueAtDistanceAlongSpline(Distance);
 	return GetScaleAtSplineInputKey(Param);
 }
 
 FTransform USplineComponent::GetTransformAtDistanceAlongSpline(float Distance, ESplineCoordinateSpace::Type CoordinateSpace, bool bUseScale) const
 {
-	const float Param = UE::SplineComponent::ShouldUseSplineCurves()
-		? WarninglessSplineCurves().ReparamTable.Eval(Distance, 0.0f)
-		: Spline.GetParameterAtDistance(Distance);
+	const float Param = GetInputKeyValueAtDistanceAlongSpline(Distance);
 	return GetTransformAtSplineInputKey(Param, CoordinateSpace, bUseScale);
 }
 
@@ -2354,18 +2679,40 @@ float USplineComponent::FindInputKeyClosestToWorldLocation(const FVector& WorldL
 {
 	const FVector LocalLocation = GetComponentTransform().InverseTransformPosition(WorldLocation);
 	float Dummy;
-	return UE::SplineComponent::ShouldUseSplineCurves()
-		? WarninglessSplineCurves().Position.FindNearest(LocalLocation, Dummy)
-		: Spline.FindNearest(LocalLocation, Dummy);
+
+	if (ShouldUseSplineCurves())
+	{
+		return WarninglessSplineCurves().Position.FindNearest(LocalLocation, Dummy);
+	}
+	else if (const FTangentSpline* TangentSpline = Spline.Get<UE::Spline::ESplineType::Tangent>())
+	{
+		return TangentSpline->FindNearest(LocalLocation, Dummy);
+	}
+	else
+	{
+		ensureAlwaysMsgf(false, TEXT("ShouldUseSplineCurves() returned false but Spline failed to yield an FTangentSpline!"));
+		return 0.f;
+	}
 }
 
 float USplineComponent::FindInputKeyOnSegmentClosestToWorldLocation(const FVector& WorldLocation, int32 Index) const
 {
 	const FVector LocalLocation = GetComponentTransform().InverseTransformPosition(WorldLocation);
 	float Dummy;
-	return UE::SplineComponent::ShouldUseSplineCurves()
-		? WarninglessSplineCurves().Position.FindNearestOnSegment(LocalLocation, Index, Dummy)
-		: Spline.FindNearestOnSegment(LocalLocation, Index, Dummy);
+
+	if (ShouldUseSplineCurves())
+	{
+		return WarninglessSplineCurves().Position.FindNearestOnSegment(LocalLocation, Index, Dummy);
+	}
+	else if (const FTangentSpline* TangentSpline = Spline.Get<UE::Spline::ESplineType::Tangent>())
+	{
+		return TangentSpline->FindNearestOnSegment(LocalLocation, Index, Dummy);
+	}
+	else
+	{
+		ensureAlwaysMsgf(false, TEXT("ShouldUseSplineCurves() returned false but Spline failed to yield an FTangentSpline!"));
+		return 0.f;
+	}
 }
 
 FVector USplineComponent::FindLocationClosestToWorldLocation(const FVector& WorldLocation, ESplineCoordinateSpace::Type CoordinateSpace) const
@@ -2653,25 +3000,26 @@ bool USplineComponent::ConvertSplineToPolyline_InDistanceRange(ESplineCoordinate
 	
 	int32 SegmentStart;
 	int32 SegmentEnd;
-	if (UE::SplineComponent::ShouldUseSplineCurves())
+	if (ShouldUseSplineCurves())
 	{
 		int32 StartIndex = WarninglessSplineCurves().ReparamTable.GetPointIndexForInputValue(RangeStart);
 		int32 EndIndex = WarninglessSplineCurves().ReparamTable.GetPointIndexForInputValue(RangeEnd);
 		SegmentStart = StartIndex / ReparamStepsPerSegment;
 		SegmentEnd = FMath::Min(NumSegments, 1 + EndIndex / ReparamStepsPerSegment);
 	}
-	else
+	else if (const FTangentSpline* TangentSpline = Spline.Get<UE::Spline::ESplineType::Tangent>())
 	{
-		// with new spline, we are strongly assuming that point index is always same as param
-		// so, we can convert param directly to segment index.
-		// we do this without making the assumptions about the reparameterization table (unlike the other branch)
-
-		const float StartParam = Spline.GetParameterAtDistance(RangeStart);
-		const float EndParam = Spline.GetParameterAtDistance(RangeEnd);
+		const float StartParam = TangentSpline->GetParameterAtDistance(RangeStart);
+		const float EndParam = TangentSpline->GetParameterAtDistance(RangeEnd);
 
 		SegmentStart = FMath::FloorToInt32(StartParam);
 		SegmentEnd = FMath::Min(NumSegments, FMath::CeilToInt32(EndParam));
 	}
+	else
+	{
+		ensureAlwaysMsgf(false, TEXT("ShouldUseSplineCurves() returned false but Spline failed to yield an FTangentSpline!"));
+		return false;
+	}
 
 	TArray<FVector> NewPoints;
 	TArray<double> NewDistances;
@@ -2820,7 +3168,7 @@ FPrimitiveSceneProxy* USplineComponent::CreateSceneProxy()
 		FSplinePDISceneProxy(const USplineComponent* InComponent)
 			: FPrimitiveSceneProxy(InComponent)
 			, bDrawDebug(InComponent->bDrawDebug)
-			, SplineInfo(UE::SplineComponent::ShouldUseSplineCurves() ? InComponent->WarninglessSplineCurves().Position : InComponent->Spline.GetSplinePointsPosition())
+			, SplineInfo(InComponent->GetSplinePointsPosition())
 #if WITH_EDITORONLY_DATA
 			, LineColor(InComponent->EditorUnselectedSplineSegmentColor)
 #else
@@ -2894,7 +3242,7 @@ FPrimitiveSceneProxy* USplineComponent::CreateSceneProxy()
 		FSplineMeshSceneProxy(const USplineComponent* InComponent)
 			: FPrimitiveSceneProxy(InComponent)
 			, bDrawDebug(InComponent->bDrawDebug)
-			, SplineInfo(UE::SplineComponent::ShouldUseSplineCurves() ? InComponent->WarninglessSplineCurves().Position : InComponent->Spline.GetSplinePointsPosition())
+			, SplineInfo(InComponent->GetSplinePointsPosition())
 #if WITH_EDITORONLY_DATA
 			, LineColor(InComponent->EditorUnselectedSplineSegmentColor)
 #else
@@ -3390,7 +3738,7 @@ void USplineComponent::Draw(FPrimitiveDrawInterface* PDI, const FSceneView* View
 
 FBoxSphereBounds USplineComponent::CalcBounds(const FTransform& LocalToWorld) const
 {
-	const auto& InterpCurve = UE::SplineComponent::ShouldUseSplineCurves() ? WarninglessSplineCurves().Position : Spline.GetSplinePointsPosition();
+	const auto& InterpCurve = GetSplinePointsPosition();
 
 #if SPLINE_FAST_BOUNDS_CALCULATION
 	FBox BoundingBox(0);
@@ -3401,8 +3749,8 @@ FBoxSphereBounds USplineComponent::CalcBounds(const FTransform& LocalToWorld) co
 
 	return FBoxSphereBounds(BoundingBox.TransformBy(LocalToWorld));
 #else
-	const int32 NumPoints = GetNumberOfSplinePoints();
-	const int32 NumSegments = bClosedLoop ? NumPoints : NumPoints - 1;
+	const int32 NumPoints = InterpCurve.Points.Num();
+	const int32 NumSegments = InterpCurve.bIsLooped ? NumPoints : NumPoints - 1;
 
 	FVector Min(WORLD_MAX);
 	FVector Max(-WORLD_MAX);
@@ -3444,7 +3792,7 @@ bool USplineComponent::GetIgnoreBoundsForEditorFocus() const
 
 TStructOnScope<FActorComponentInstanceData> USplineComponent::GetComponentInstanceData() const
 {
-	if (UE::SplineComponent::ShouldUseSplineCurves())
+	if (ShouldUseSplineCurves())
 	{
 		TStructOnScope<FActorComponentInstanceData> InstanceData = MakeStructOnScope<FActorComponentInstanceData, FSplineInstanceData>(this);
 		FSplineInstanceData* SplineInstanceData = InstanceData.Cast<FSplineInstanceData>();
@@ -3483,17 +3831,17 @@ TStructOnScope<FActorComponentInstanceData> USplineComponent::GetComponentInstan
 
 FSplineCurves USplineComponent::GetSplineCurves() const
 {
-	return UE::SplineComponent::ShouldUseSplineCurves()
+	return ShouldUseSplineCurves()
 		? WarninglessSplineCurves()
 		: FSplineCurves::FromSplineInterface(Spline);
 }
 
 FSpline USplineComponent::GetSpline() const
 {
-	if (UE::SplineComponent::ShouldUseSplineCurves())
+	if (ShouldUseSplineCurves())
 	{
 		FSpline Intermediate;
-		Intermediate = SplineCurves;
+		Intermediate = WarninglessSplineCurves();
 		return Intermediate;
 	}
 	else
@@ -3518,7 +3866,13 @@ void USplineComponent::SetSpline(const FSpline& InSpline)
 void USplineComponent::SetSpline(const FSplineCurves& InSplineCurves)
 {
 	WarninglessSplineCurves() = InSplineCurves;
-	Spline = InSplineCurves;
+
+	// If Spline is already tangent backed or SplineComponent.UseSplineCurves is false, we can copy assign to the FSpline.
+	// This will force it to instantiate an FTangentSpline and populate it with the data from InSplineCurves.
+	if (Spline.GetCurrentImplementation() == UE::Spline::ESplineType::Tangent || !UE::SplineComponent::GUseSplineCurves)
+	{
+		Spline = InSplineCurves;
+	}
 
 	if (UE::SplineComponent::GValidateOnChange)
 	{
@@ -3530,30 +3884,70 @@ void USplineComponent::SetSpline(const FSplineCurves& InSplineCurves)
 
 int32 USplineComponent::GetVersion() const
 {
-	return UE::SplineComponent::ShouldUseSplineCurves()
-		? WarninglessSplineCurves().Version
-		: Spline.GetVersion();
+	if (ShouldUseSplineCurves())
+	{
+		return WarninglessSplineCurves().Version;
+	}
+	else if (const FTangentSpline* TangentSpline = Spline.Get<UE::Spline::ESplineType::Tangent>())
+	{
+		return TangentSpline->GetVersion();
+	}
+	else
+	{
+		ensureAlwaysMsgf(false, TEXT("ShouldUseSplineCurves() returned false but Spline failed to yield an FTangentSpline!"));
+		return 0;
+	}
 }
 
 const FInterpCurveVector& USplineComponent::GetSplinePointsPosition() const
 {
-	return UE::SplineComponent::ShouldUseSplineCurves()
-		? WarninglessSplineCurves().Position
-		: Spline.GetSplinePointsPosition();
+	if (ShouldUseSplineCurves())
+	{
+		return WarninglessSplineCurves().Position;
+	}
+	else if (const FTangentSpline* TangentSpline = Spline.Get<UE::Spline::ESplineType::Tangent>())
+	{
+		return TangentSpline->GetSplinePointsPosition();
+	}
+	else
+	{
+		ensureAlwaysMsgf(false, TEXT("ShouldUseSplineCurves() returned false but Spline failed to yield an FTangentSpline!"));
+		return DummyPointPosition;
+	}
 }
 
 const FInterpCurveQuat& USplineComponent::GetSplinePointsRotation() const
 {
-	return UE::SplineComponent::ShouldUseSplineCurves()
-		? WarninglessSplineCurves().Rotation
-		: Spline.GetSplinePointsRotation();
+	if (ShouldUseSplineCurves())
+	{
+		return WarninglessSplineCurves().Rotation;
+	}
+	else if (const FTangentSpline* TangentSpline = Spline.Get<UE::Spline::ESplineType::Tangent>())
+	{
+		return TangentSpline->GetSplinePointsRotation();
+	}
+	else
+	{
+		ensureAlwaysMsgf(false, TEXT("ShouldUseSplineCurves() returned false but Spline failed to yield an FTangentSpline!"));
+		return DummyPointRotation;
+	}
 }
 
 const FInterpCurveVector& USplineComponent::GetSplinePointsScale() const
 {
-	return UE::SplineComponent::ShouldUseSplineCurves()
-		? WarninglessSplineCurves().Scale
-		: Spline.GetSplinePointsScale();
+	if (ShouldUseSplineCurves())
+	{
+		return WarninglessSplineCurves().Scale;
+	}
+	else if (const FTangentSpline* TangentSpline = Spline.Get<UE::Spline::ESplineType::Tangent>())
+	{
+		return TangentSpline->GetSplinePointsScale();
+	}
+	else
+	{
+		ensureAlwaysMsgf(false, TEXT("ShouldUseSplineCurves() returned false but Spline failed to yield an FTangentSpline!"));
+		return DummyPointScale;
+	}
 }
 
 TArray<ESplinePointType::Type> USplineComponent::GetEnabledSplinePointTypes() const
@@ -3719,6 +4113,11 @@ void USplineComponent::PostTransacted(const FTransactionObjectEvent& Transaction
 }
 #endif
 
+bool USplineComponent::ShouldUseSplineCurves() const
+{
+	return UE::SplineComponent::GUseSplineCurves || Spline.GetCurrentImplementation() != UE::Spline::ESplineType::Tangent;
+}
+
 UE::Geometry::FInterval1f USplineComponent::GetInputKeyRange() const
 {
 	UE::Geometry::FInterval1f Range;
