FSpline Refactor: Prepare FSpline for non-tangent based splines.

diff --git a/Engine/Source/Runtime/Engine/Private/Curves/Spline.cpp b/Engine/Source/Runtime/Engine/Private/Curves/Spline.cpp
--- a/Engine/Source/Runtime/Engine/Private/Curves/Spline.cpp
+++ b/Engine/Source/Runtime/Engine/Private/Curves/Spline.cpp
@@ -4,33 +4,22 @@
 
 #include "Algo/Sort.h"
 #include "BoxTypes.h"
-#include "Async/TransactionallySafeMutex.h"
 #include "Components/SplineComponent.h"	// for FSplineCurves, ideally removed
+#include "Curves/Splines/TangentSpline.h"
 #include "Misc/Base64.h"
-#include "Misc/TransactionallySafeRWLock.h"
 #include "Serialization/MemoryReader.h"
 #include "Serialization/MemoryWriter.h"
 #include "UObject/FortniteMainBranchObjectVersion.h"
 
-PRAGMA_DISABLE_EXPERIMENTAL_WARNINGS
-#include "Splines/TangentBezierSpline.h"
-#include "Splines/MultiSpline.h"
-PRAGMA_ENABLE_EXPERIMENTAL_WARNINGS
-
 DEFINE_LOG_CATEGORY_STATIC(LogSpline, Log, All);
 
 #include UE_INLINE_GENERATED_CPP_BY_NAME(Spline)
 
 PRAGMA_DISABLE_EXPERIMENTAL_WARNINGS
 
-using UE::Geometry::FInterval1f;
-
-const FInterpCurveVector FSpline::PositionCurve;
-const FInterpCurveQuat FSpline::RotationCurve;
-const FInterpCurveVector FSpline::ScaleCurve;
-
 namespace UE::Spline
 {
+
 int32 GImplementation = 0;
 FAutoConsoleVariableRef CVarSplineImplementation(
 	TEXT("Spline.Implementation"),
@@ -38,3524 +27,353 @@ FAutoConsoleVariableRef CVarSplineImplementation(
 	TEXT("0) Not Implemented - 1) Legacy Implementation - 2) New Implementation")
 #if !WITH_EDITOR
 	,ECVF_ReadOnly		// In non-editor builds, this CVar should exist so that GImplementation is properly set but it should not be changed by the user.
-#endif
-);
-
-int32 GParameterizationPolicy = 0;
-FAutoConsoleVariableRef CVarParameterizationPolicy(
-	TEXT("Spline.ParameterizationPolicy"),
-	GParameterizationPolicy,
-	TEXT("0) Uniform - 1) Chord Length - 2) Centripetal")
-);
-	
-bool GApproximateTangents = false;
-FAutoConsoleVariableRef CVarApproximateTangents(
-	TEXT("Spline.ApproximateTangents"),
-	GApproximateTangents,
-	TEXT("True if we should approximate tangents using the central difference formula.")
-);
-
-bool GFallbackFindNearest = false;
-FAutoConsoleVariableRef CVarFallbackFindNearest(
-	TEXT("Spline.FallbackFindNearest"),
-	GFallbackFindNearest,
-	TEXT("True if we should implement FindNearest and FindNearestOnSegment using an intermediate spline representation. Only applies if Spline.Implementation == 2.")
-);
-
-bool GUseLegacyPositionEvaluation = false;
-FAutoConsoleVariableRef CVarUseLegacyPositionEvaluation(
-	TEXT("Spline.UseLegacyPositionEvaluation"),
-	GUseLegacyPositionEvaluation,
-	TEXT("If true, evaluating the position channel always routes through an interp curve.")
-);
-
-bool GUseLegacyRotationEvaluation = false;
-FAutoConsoleVariableRef CVarUseLegacyRotationEvaluation(
-	TEXT("Spline.UseLegacyRotationEvaluation"),
-	GUseLegacyRotationEvaluation,
-	TEXT("If true, evaluating the rotation channel always routes through an interp curve.")
-);
-
-bool GUseLegacyScaleEvaluation = false;
-FAutoConsoleVariableRef CVarUseLegacyScaleEvaluation(
-	TEXT("Spline.UseLegacyScaleEvaluation"),
-	GUseLegacyScaleEvaluation,
-	TEXT("If true, evaluating the scale channel always routes through an interp curve.")
-);
-
-bool GImmediatelyUpdateLegacyCurves = false;
-FAutoConsoleVariableRef CVarImmediatelyUpdateLegacyCurves(
-	TEXT("Spline.ImmediatelyUpdateLegacyCurves"),
-	GImmediatelyUpdateLegacyCurves,
-	TEXT("If true, mutating operations immediately rebuild legacy curves. If false, legacy curves are updated only when requested.")
-);
-
-// WARNING! THIS BREAKS EXISTING TEXT REPRESENTATION WHEN CHANGED, IT WILL ABSOLUTELY TRASH YOUR STUFF IF YOU CHANGE WILLY-NILLY
-bool GEncodeAsHex = true;
-// FAutoConsoleVariableRef CVarEncodeAsHex(
-// 	TEXT("Spline.EncodeAsHex"),
-// 	GEncodeAsHex,
-// 	TEXT("If true, splines are imported/exported using hexadecimal encoding. If false, the data is encoded in base 64.")
-// );
-
-bool GValidateRotScale = false;
-FAutoConsoleVariableRef CVarValidateRotScale(
-	TEXT("Spline.ValidateRotScale"),
-	GValidateRotScale,
-	TEXT("True if we should validate rotation and scale attributes when structurally modifying the spline.")
-);
-
-bool GSkipDoubleUpdate = true;
-FAutoConsoleVariableRef CVarSkipDoubleUpdate(
-	TEXT("Spline.SkipDoubleUpdate"),
-	GSkipDoubleUpdate,
-	TEXT("When true, skip updating the spline twice if UpdateSpline with parameters is called.")
-);
-
-}
-
-UE::Geometry::Spline::EParameterizationPolicy GetParameterizationPolicy()
-{
-	switch (UE::Spline::GParameterizationPolicy)
-	{
-	default:	// fallthrough
-	case 0: return UE::Geometry::Spline::EParameterizationPolicy::Uniform;
-	case 1: return UE::Geometry::Spline::EParameterizationPolicy::ChordLength;
-	case 2: return UE::Geometry::Spline::EParameterizationPolicy::Centripetal;
-	}
-}
-
-/**
- * Converts from EInterpCurveMode to ETangentMode
- */
-inline UE::Geometry::Spline::ETangentMode ConvertInterpCurveModeToTangentMode(EInterpCurveMode Mode)
-{
-	switch (Mode)
-	{
-	case CIM_Linear:
-		return UE::Geometry::Spline::ETangentMode::LegacyLinear; // Default to legacy for compatibility
-	case CIM_CurveAuto:
-		return UE::Geometry::Spline::ETangentMode::LegacyAuto; // Default to legacy for compatibility
-	case CIM_Constant:
-		return UE::Geometry::Spline::ETangentMode::Constant;
-	case CIM_CurveUser:
-		return UE::Geometry::Spline::ETangentMode::User;
-	case CIM_CurveBreak:
-		return UE::Geometry::Spline::ETangentMode::Broken;
-	case CIM_CurveAutoClamped:
-		return UE::Geometry::Spline::ETangentMode::LegacyAutoClamped; // Default to legacy for compatibility
-	case CIM_Unknown:
-	default:
-		return UE::Geometry::Spline::ETangentMode::Unknown;
-	}
-}
-
-/**
- * Converts from ETangentMode to EInterpCurveMode
- */
-inline EInterpCurveMode ConvertTangentModeToInterpCurveMode(UE::Geometry::Spline::ETangentMode Mode)
-{
-	switch (Mode)
-	{
-	case UE::Geometry::Spline::ETangentMode::LegacyLinear:
-	case UE::Geometry::Spline::ETangentMode::Linear:
-		return CIM_Linear;
-	case UE::Geometry::Spline::ETangentMode::Auto:
-	case UE::Geometry::Spline::ETangentMode::LegacyAuto:
-		return CIM_CurveAuto;
-	case UE::Geometry::Spline::ETangentMode::Constant:
-		return CIM_Constant;
-	case UE::Geometry::Spline::ETangentMode::User:
-		return CIM_CurveUser;
-	case UE::Geometry::Spline::ETangentMode::Broken:
-		return CIM_CurveBreak;
-	case UE::Geometry::Spline::ETangentMode::AutoClamped:
-	case UE::Geometry::Spline::ETangentMode::LegacyAutoClamped:
-		return CIM_CurveAutoClamped;
-	case UE::Geometry::Spline::ETangentMode::Unknown:
-	default:
-		return CIM_Unknown;
-	}
-}
-
-/**
- * FNewSpline Definition
- *
- * A spline that provides tangent-based control over curve shape while using 
- * piecewise Bezier curves internally for evaluation. Supports both manual
- * tangent control and automatic tangent computation.
- */
-
-class FNewSpline: public UE::Geometry::Spline::TMultiSpline<UE::Geometry::Spline::FTangentBezierSpline3d>
-{
-	using FTangentBezierControlPoint = UE::Geometry::Spline::TTangentBezierControlPoint<FVector3d>;
-	
-public:
-	
-    FNewSpline();
-    FNewSpline(const FNewSpline& Other);
-	FNewSpline(const FLegacySpline& Other);
-	FNewSpline(const FSplineCurves& Other);
-	virtual ~FNewSpline() override = default;
-    FNewSpline& operator=(const FNewSpline& Other);
-
-	bool operator!=(const FNewSpline& Other) const;
-
-    virtual bool Serialize(FArchive& Ar) override;
-    friend FArchive& operator<<(FArchive& Ar, FNewSpline& Spline)
-    {
-        Spline.Serialize(Ar);
-        return Ar;
-    }
-
-	/**
-	 * Returns the parameter space of FNewSpline. All public functions taking a parameter are only defined when the provided value is in this interval.
-	 * @return [0, NumSegments] when NumSegment != 0, otherwise undefined (empty interval).
-	 */
-	FInterval1f GetSegmentSpace() const;
-	float GetSegmentLength(const int32 Index, const float Param, const FVector& Scale3D = FVector(1.0f)) const;
-	float GetSplineLength();
-	
-    virtual void SetClosedLoop(bool bInClosedLoop) override;
-	void SetClosedLoop(bool bInClosedLoop, bool bUpdateSpline);
-
-	void Reset();
-	
-	void ResetRotation();
-	void ResetScale();
-	
-    void AddPoint(const FSplinePoint& Point);
-	void InsertPoint(const FSplinePoint& Point, int32 Index);
-	FSplinePoint GetPoint(int32 Index) const;
-    void RemovePoint(int32 Index);
-
-	void SetLocation(int32 Index, const FVector& InLocation);
-	FVector GetLocation(const int32 Index) const;
-	
-	void SetInTangent(const int32 Index, const FVector& InTangent);
-	FVector GetInTangent(const int32 Index) const;
-	
-	void SetOutTangent(const int32 Index, const FVector& OutTangent);
-	FVector GetOutTangent(const int32 Index) const;
-
-	void SetRotation(int32 Index, const FQuat& InRotation);
-	FQuat GetRotation(const int32 Index) const;
-	
-	void SetScale(int32 Index, const FVector& InScale);
-	FVector GetScale(const int32 Index) const;
-	
-	void SetSplinePointType(int32 Index, EInterpCurveMode Type);
-	EInterpCurveMode GetSplinePointType(int32 Index) const;
-
-	float GetParameterAtIndex(int32 Index) const;
-	float GetParameterAtDistance(float Distance);
-	float GetDistanceAtParameter(float Parameter);
-
-    virtual float FindNearest(const FVector& InLocation, float& OutSquaredDistance) const override;
-    float FindNearestOnSegment(const FVector& InLocation, int32 SegmentIndex, float& OutSquaredDistance) const;
-	
-	/** Parameter based evaluation API. */
-	FVector EvaluatePosition(float Parameter) const;	// Not named Evaluate because we would be shadowing non-virtual base, this maps parameters
-    FVector EvaluateDerivative(float Parameter) const;
-	FQuat EvaluateRotation(float Parameter) const;
-	FVector EvaluateScale(float Parameter) const;
-
-	// Templated Channel Creation/Query functions
-	template <typename AttrType> int32 NumAttributeValues(FName Name) const;
-
-	// Templated Attribute Interaction by Index functions
-	template <typename AttrType> AttrType GetAttributeValue(const FName& Name, int32 Index) const;
-	template <typename AttrType> void SetAttributeValue(FName Name, const AttrType& Value, int32 Index);
-	template <typename AttrType> void RemoveAttributeValue(FName Name, int32 Index);
-	template <typename AttrType> int32 SetAttributeParameter(const FName& Name, int32 Index, float ParentSpaceParameter);
-	template <typename AttrType> float GetAttributeParameter(const FName& Name, int32 Index) const;
-
-	// Templated Attribute Interaction by Parameter functions
-	template <typename AttrType> int32 AddAttributeValue(FName Name, const AttrType& Value, float Parameter);
-	
-	/** Updates the spline using the current configuration. */
-	void UpdateSpline();
-
-	/** Updates the spline configuration, then updates the spline. */
-	void UpdateSpline(const FSpline::FUpdateSplineParams& Params);
-
-	const FInterpCurveVector& GetSplinePointsPosition() const { return PositionCurve; }
-	const FInterpCurveQuat& GetSplinePointsRotation() const { return RotationCurve; }
-	const FInterpCurveVector& GetSplinePointsScale() const { return ScaleCurve; }
-
-private:
-
-	/** Inspired by FDynamicMesh3::FChangeStamp. */
-	struct FChangeStamp
-	{
-		FChangeStamp()
-		{}
-		
-		UE_NONCOPYABLE(FChangeStamp);
-
-		/** Updates the change stamp in an thread-safe, transactionally-safe way. */
-		void Increment()
-		{
-			Mutex.Lock();
-			++Value;
-			Mutex.Unlock();
-		}
-
-		/** Returns the current change value in a thread-safe, transactionally-safe way. */
-		uint32 GetValue() const
-		{
-			Mutex.Lock();
-			uint32 Result = Value;
-			Mutex.Unlock();
-
-			return Result;
-		}
-
-	private:
-		
-		/** Guards `Value`. */
-		mutable UE::FTransactionallySafeMutex Mutex;
-
-		/** The change stamp is incremented when modifications occur. It's guarded by `Mutex`. */
-		uint32 Value = 1;
-	};
-
-	FInterpCurveFloat ReparamTable;
-	FChangeStamp ReparamTableNextVersion;
-	uint32 ReparamTableVersion = 0;
-	FTransactionallySafeRWLock ReparamTableRWLock;
-	
-	static inline const FName RotationAttrName = FName("Rotation");
-	static inline const FName ScaleAttrName = FName("Scale");
-
-	/** Legacy Curves: */
-    mutable FInterpCurveVector PositionCurve;
-	mutable FInterpCurveQuat RotationCurve;
-	mutable FInterpCurveVector ScaleCurve;
-	FChangeStamp LegacyCurvesNextVersion;
-	mutable uint32 LegacyCurvesVersion = 0;
-	
-	int32 ReparamStepsPerSegment = 10;
-
-private:
-
-	// In order for FSpline to access the parameter mapping functions below.
-	// The alternative is to redefine a bunch of MultiSpline functions to do the remapping internally, but this is a bit of a pain.
-	friend struct FSpline;
-	
-	void Reparameterize();
-
-	/** Convert from [0, 1] to [0, NumSegments] */
-	float FromNormalizedSpace(float Parameter) const;
-	
-	/** Convert from [0, NumSegments] to [0, 1] */
-	float ToNormalizedSpace(float Parameter) const;
-
-	/** Convert from internal spline space to [0, NumSegments] */
-	float FromInternalSplineSpace(float Parameter) const;
-	
-	/** Convert from [0, NumSegments] to internal spline space */
-	float ToInternalSplineSpace(float Parameter) const;
-
-	void MarkReparamTableDirty();
-	
-	/** Clears and repopulates reparameterization attribute channels. */
-	void UpdateReparamTable();
-
-	void MarkLegacyCurvesDirty();
-
-	/** Clears and repopulates PositionCurve, RotationCurve, and ScaleCurve */
-	void RebuildLegacyCurves() const;
-
-	/** Converts FSplinePoint to FTangentBezierControlPoint */
-    FTangentBezierControlPoint ConvertToTangentBezierControlPoint(const FSplinePoint& Point) const;
-
-    /**
-     * Updates the attributes of a point in the spline
-     * @param Point 
-     * @param PointIndex 
-     */
-    void UpdatePointAttributes(const FSplinePoint& Point, int32 PointIndex);
-	
-	/**
-	 * Converts an index (possibly fractional) to a parameter value in internal spline space.
-	 */
-	float ConvertIndexToInternalParameter(int32 Index, float Fraction = 0.0f) const;
-
-    /**
-	 * Converts an internal spline space parameter value to the closest index
-	 */
-	int32 ConvertInternalParameterToNearestPointIndex(float Parameter) const;
-
-	void ValidateRotScale() const;
-
-	// Helper functions for keeping attribute channels stable when reparameterizing.
-	template <typename AttrType> TArray<AttrType> CacheChannelKnots(const FName& InName) const;
-	template <typename AttrType> void RestoreChannelKnots(const FName& InName, const TArray<AttrType>& InKnots);
-};
-
-template <typename AttrType>
-int32 FNewSpline::AddAttributeValue(FName Name, const AttrType& Value, float Parameter)
-{
-	using AttrSplineType = UE::Geometry::Spline::TTangentBezierSpline<AttrType>;
-	using ControlPointType = UE::Geometry::Spline::TTangentBezierControlPoint<AttrType>;
-	
-	if (AttrSplineType* Channel = GetTypedAttributeChannel<AttrSplineType>(Name))
-	{
-		ControlPointType ControlPoint = ControlPointType(Value);
-		int32 NumPoints = Channel->GetNumPoints();
-
-		float ParentSpaceParameter = ToInternalSplineSpace(Parameter);
-		FInterval1f MappedChildSpace = GetMappedChildSpace(Name);
-		
-		// Special case for empty spline
-		if (NumPoints == 0)
-		{
-			Channel->AppendPoint(ControlPoint);
-			SetAttributeChannelRange(Name, FInterval1f(ParentSpaceParameter, ParentSpaceParameter), EMappingRangeSpace::Parent);
-			return 0;
-		}
-
-		if (NumPoints == 1)
-		{
-			if (ParentSpaceParameter > MappedChildSpace.Min)
-			{
-				Channel->AppendPoint(ControlPoint);
-				Channel->Reparameterize(UE::Geometry::Spline::EParameterizationPolicy::Uniform);
-				SetAttributeChannelRange(Name, FInterval1f(MappedChildSpace.Min, ParentSpaceParameter), EMappingRangeSpace::Parent);
-				return 1;
-			}
-			else
-			{
-				Channel->PrependPoint(ControlPoint);
-				Channel->Reparameterize(UE::Geometry::Spline::EParameterizationPolicy::Uniform);
-				SetAttributeChannelRange(Name, FInterval1f(ParentSpaceParameter, MappedChildSpace.Min), EMappingRangeSpace::Parent);
-				return 0;
-			}
-		}
-
-		// append case
-		if (ParentSpaceParameter > MappedChildSpace.Max)
-		{
-			// It is important to compute ChildSpaceParameter before the AppendPoint call, otherwise it will not be correct.
-			float ChildSpaceParameter = MapParameterToChildSpace(Name, ParentSpaceParameter);
-			Channel->AppendPoint(ControlPoint);
-
-			// By growing the child space and the mapped parent range proportionally, we keep the internal points stable in parent space.
-			Channel->SetParameter(Channel->GetNumPoints() - 1, ChildSpaceParameter);
-			SetAttributeChannelRange(Name, FInterval1f(MappedChildSpace.Min, ParentSpaceParameter), EMappingRangeSpace::Parent);
-			
-			return Channel->GetNumPoints() - 1;
-		}
-
-		// prepend case
-		if (ParentSpaceParameter < MappedChildSpace.Min)
-		{
-			// It is important to compute ChildSpaceParameter before the AppendPoint call, otherwise it will not be correct.
-			float ChildSpaceParameter = MapParameterToChildSpace(Name, ParentSpaceParameter);
-			Channel->PrependPoint(ControlPoint);
-
-			// By growing the child space and the mapped parent range proportionally, we keep the internal points stable in parent space.
-			Channel->SetParameter(0, ChildSpaceParameter);
-			SetAttributeChannelRange(Name, FInterval1f(ParentSpaceParameter, MappedChildSpace.Max), EMappingRangeSpace::Parent);
-			
-			return 0;
-		}
-		
-		float ChildSpaceParameter = MapParameterToChildSpace(Name, ParentSpaceParameter);
-		float LocalT;
-		int32 SegmentIndex = Channel->FindSegmentIndex(ChildSpaceParameter, LocalT);
-		
-		int NewPointIndex = Channel->InsertPointAtSegmentParam(SegmentIndex, LocalT, ControlPoint);
-		return NewPointIndex;
-	}
-
-	return INDEX_NONE;
-}
-
-template <typename AttrType>
-void FNewSpline::SetAttributeValue(FName Name, const AttrType& Value, int32 Index)
-{
-	using AttrSplineType = UE::Geometry::Spline::TTangentBezierSpline<AttrType>;
-
-	if (AttrSplineType* Channel = GetTypedAttributeChannel<AttrSplineType>(Name))
-	{
-		Channel->SetValue(Index, Value);
-	}
-}
-
-template <typename AttrType>
-void FNewSpline::RemoveAttributeValue(FName Name, int32 Index)
-{
-	using AttrSplineType = UE::Geometry::Spline::TTangentBezierSpline<AttrType>;
-
-	if (AttrSplineType* Channel = GetTypedAttributeChannel<AttrSplineType>(Name))
-	{
-		if (Channel->GetNumPoints() > 1)
-		{
-			if (Index == 0)
-			{
-				FInterval1f NewMappingRange = GetMappedChildSpace(Name);
-				NewMappingRange.Min = ToInternalSplineSpace(GetAttributeParameter<AttrType>(Name, Index + 1));
-				SetAttributeChannelRange(Name, NewMappingRange, EMappingRangeSpace::Parent);
-			}
-			else if (Index == Channel->GetNumPoints() - 1)
-			{
-				FInterval1f NewMappingRange = GetMappedChildSpace(Name);
-				NewMappingRange.Max = ToInternalSplineSpace(GetAttributeParameter<AttrType>(Name, Index - 1));
-				SetAttributeChannelRange(Name, NewMappingRange, EMappingRangeSpace::Parent);
-			}
-		}
-		
-		Channel->RemovePoint(Index);
-	}
-}
-
-template <typename AttrType>
-AttrType FNewSpline::GetAttributeValue(const FName& Name, int32 Index) const
-{
-	using AttrSplineType = UE::Geometry::Spline::TTangentBezierSpline<AttrType>;
-	
-	if (AttrSplineType* Channel = GetTypedAttributeChannel<AttrSplineType>(Name))
-	{
-		return Channel->GetValue(Index);
-	}
-
-	return AttrType();
-}
-
-template <typename AttrType>
-int32 FNewSpline::NumAttributeValues(FName Name) const
-{
-	using AttrSplineType = UE::Geometry::Spline::TTangentBezierSpline<AttrType>;
-	
-	if (AttrSplineType* Channel = GetTypedAttributeChannel<AttrSplineType>(Name))
-	{
-		return Channel->GetNumPoints();
-	}
-
-	return 0;
-}
-
-template <typename AttrType>
-int32 FNewSpline::SetAttributeParameter(const FName& Name, int32 Index, float Parameter)
-{
-	using AttrSplineType = UE::Geometry::Spline::TTangentBezierSpline<AttrType>;
-	
-	if (AttrSplineType* Channel = GetTypedAttributeChannel<AttrSplineType>(Name))
-	{
-		float ParentSpaceParameter = ToInternalSplineSpace(Parameter);
-
-		// Prevent collapse of the channel space and keep it starting at 0.
-		auto SanitizeChannelSpace = [Channel]()
-		{
-			const FInterval1f ChannelSpace = Channel->GetParameterSpace();
-			const float KnotOffset = -1.f * ChannelSpace.Min;
-			const float KnotScaleFactor = 1.f / (ChannelSpace.Max - ChannelSpace.Min);
-			TArray<UE::Geometry::Spline::FKnot> Knots = Channel->GetKnotVector();
-			for (UE::Geometry::Spline::FKnot& Knot : Knots)
-			{
-				Knot.Value += KnotOffset;
-				Knot.Value *= KnotScaleFactor;
-			}
-			Channel->SetKnotVector(Knots);
-		};
-		
-		// Helper to prevent code duplication.
-		// Uses current value of Index and ParentSpaceParameter at the time of calling to update the attribute parameter.
-		auto SetParentSpaceParameter = [this, &ParentSpaceParameter, Name, Channel, &Index]() -> int32
-		{
-			const float CurrentChildSpaceParameter = MapParameterToChildSpace(Name, ToInternalSplineSpace(GetAttributeParameter<AttrType>(Name, Index)));
-			const float DesiredChildSpaceParameter = MapParameterToChildSpace(Name, ParentSpaceParameter);
-
-			constexpr float MinStep = 2.f * UE_KINDA_SMALL_NUMBER;
-			float NewChildSpaceParameter = DesiredChildSpaceParameter > CurrentChildSpaceParameter
-				? FMath::Max(DesiredChildSpaceParameter, CurrentChildSpaceParameter + MinStep)
-				: FMath::Min(DesiredChildSpaceParameter, CurrentChildSpaceParameter - MinStep);
-			
-			return Channel->SetParameter(Index, NewChildSpaceParameter);
-		};
-
-		/**
-		 * Cases we handle below:
-		 * 1) Moving the only existing attribute.
-		 * 2) Moving an endpoint for a 2 point channel.
-		 * 3) Moving the first endpoint for a 3+ point channel.
-		 * 4) Moving the last endpoint for a 3+ point channel.
-		 * 5) Moving an internal point.
-		 */
-		
-		if (Channel->GetNumPoints() == 1)						// Case 1: Moving the only attribute
-		{
-			SetAttributeChannelRange(Name, FInterval1f(ParentSpaceParameter, ParentSpaceParameter), EMappingRangeSpace::Parent);
-			return Index;
-		}
-		else if (Channel->GetNumPoints() == 2)					// Case 2: Moving an end point while no internal points exist.
-		{
-			FInterval1f MappedRange = GetMappedChildSpace(Name);
-			float& RangeBound = Index == 0 ? MappedRange.Min : MappedRange.Max;
-			RangeBound = ParentSpaceParameter;
-
-			if (MappedRange.Min > MappedRange.Max)
-			{
-				// The mapping range will flip, swap end points and un-flip.
-				Channel->SetParameter(1, Channel->GetParameter(0) - 1.f);	// It doesn't actually matter what we do here, as long as we get the ordering & mapping range correct.
-				SetAttributeChannelRange(Name, FInterval1f(MappedRange.Max, MappedRange.Min), EMappingRangeSpace::Parent);
-				Index = Index == 0 ? 1 : 0;
-			}
-			else
-			{
-				SetAttributeChannelRange(Name, MappedRange, EMappingRangeSpace::Parent);
-			}
-
-			return Index;
-		}
-		else if (Index == 0)									// Case 3: Moving the first point (which has exactly 1 neighbor)
-		{
-			// Save internal attribute point parameters for later re-parameterization after changing the child space mapping.
-			TArray<float> InternalParameters = CacheChannelKnots<float>(Name);
-
-			float ParentSpaceUpperBound = GetMappedChildSpace(Name).Max;
-			const float NeighborParentSpaceParameter = ToInternalSplineSpace(GetAttributeParameter<AttrType>(Name, Index + 1));
-			const bool bAttributesWillReorder = ParentSpaceParameter > NeighborParentSpaceParameter;
-
-			// If the end point is passing its neighbor, we need to actually reshuffle the points.
-			if (bAttributesWillReorder)
-			{
-				Index = SetParentSpaceParameter();
-
-				// Shift values left up to the last invalidated index.
-				const int32 InvalidatedInternalParameterIdx = FMath::Clamp(Index - 1, 0, InternalParameters.Num() - 1);
-				for (int InternalParameterIdx = 1; InternalParameterIdx <= InvalidatedInternalParameterIdx; ++InternalParameterIdx)
-				{
-					InternalParameters[InternalParameterIdx - 1] = InternalParameters[InternalParameterIdx];
-				}
-				
-				// Do one of 2 things to the actual invalidated index:
-				if (Index == Channel->GetNumPoints() - 1)
-				{
-					// Old upper end point is now an internal point and needs to be re-parameterized.
-					InternalParameters[InvalidatedInternalParameterIdx] = ParentSpaceUpperBound;
-				}
-				else
-				{
-					// Old lower end point is now an internal point and needs to be re-parameterized.
-					InternalParameters[InvalidatedInternalParameterIdx] = ParentSpaceParameter;
-				}
-			}
-			
-			// Updates the lower bound of the mapping range.
-			ParentSpaceUpperBound = Index == Channel->GetNumPoints() - 1 ? ParentSpaceParameter : ParentSpaceUpperBound;
-			const float ParentSpaceLowerBound = bAttributesWillReorder ? NeighborParentSpaceParameter : ParentSpaceParameter;
-			SetAttributeChannelRange(Name, FInterval1f(ParentSpaceLowerBound, ParentSpaceUpperBound), EMappingRangeSpace::Parent);
-
-			// Re-parameterize the internal points so that they do not move in parent space after the mapping range changed
-			RestoreChannelKnots(Name, InternalParameters);
-			
-			if (bAttributesWillReorder)
-			{
-				SanitizeChannelSpace();
-			}
-			
-			return Index;
-		}
-		else if (Index == Channel->GetNumPoints() - 1)			// Case 4: Moving the last point (which has exactly 1 neighbor)
-		{
-			// Save internal attribute point parameters for later re-parameterization after changing the child space mapping.
-			TArray<float> InternalParameters = CacheChannelKnots<float>(Name);
-
-			float ParentSpaceLowerBound = GetMappedChildSpace(Name).Min;
-			const float NeighborParentSpaceParameter = ToInternalSplineSpace(GetAttributeParameter<AttrType>(Name, Index - 1));
-			const bool bAttributesWillReorder = ParentSpaceParameter < NeighborParentSpaceParameter;
-			
-			// If the end point is passing its neighbor, we need to actually reshuffle the points.
-			if (bAttributesWillReorder)
-			{
-				Index = SetParentSpaceParameter();
-
-				// Shift values right up to the last invalidated index.
-				const int32 InvalidatedInternalParameterIdx = FMath::Clamp(Index - 1, 0, InternalParameters.Num() - 1);
-				for (int InternalParameterIdx = InternalParameters.Num() - 2; InternalParameterIdx >= InvalidatedInternalParameterIdx; --InternalParameterIdx)
-				{
-					InternalParameters[InternalParameterIdx + 1] = InternalParameters[InternalParameterIdx];
-				}
-				
-				// Do one of 2 things to the actual invalidated index:
-				if (Index == 0)
-				{
-					// Old upper end point is now an internal point and needs to be re-parameterized.
-					InternalParameters[InvalidatedInternalParameterIdx] = ParentSpaceLowerBound;
-				}
-				else
-				{
-					// Old lower end point is now an internal point and needs to be re-parameterized.
-					InternalParameters[InvalidatedInternalParameterIdx] = ParentSpaceParameter;
-				}
-			}
-
-			// Update the child space mapping
-			ParentSpaceLowerBound = Index == 0 ? ParentSpaceParameter : ParentSpaceLowerBound;
-			float ParentSpaceUpperBound = bAttributesWillReorder ? NeighborParentSpaceParameter : ParentSpaceParameter; // Prevent over-shrinking if we reshuffle.
-			SetAttributeChannelRange(Name, FInterval1f(ParentSpaceLowerBound, ParentSpaceUpperBound), EMappingRangeSpace::Parent);
-			
-			// Re-parameterize the internal points so that they do not move in parent space after the mapping range changed
-			RestoreChannelKnots(Name, InternalParameters);
-
-			if (bAttributesWillReorder)
-			{
-				SanitizeChannelSpace();
-			}
-		
-			return Index;
-		}
-		else													// Case 5: Moving an internal point (which has exactly 2 neighbors)
-		{
-			FInterval1f MappedRange = GetMappedChildSpace(Name);
-			Index = SetParentSpaceParameter();
-			
-			if (Index == 0)
-			{
-				// The internal point is now the lower end point
-				MappedRange.Min = ParentSpaceParameter;
-			}
-			else if (Index == Channel->GetNumPoints() - 1)
-			{
-				// The internal point is now the upper end point
-				MappedRange.Max = ParentSpaceParameter;
-			}
-			
-			SetAttributeChannelRange(Name, MappedRange, EMappingRangeSpace::Parent);
-			
-			return Index;
-		}
-	}
-	
-	return INDEX_NONE;
-}
-
-template <typename AttrType>
-float FNewSpline::GetAttributeParameter(const FName& Name, int32 Index) const
-{
-	using AttrSplineType = UE::Geometry::Spline::TTangentBezierSpline<AttrType>;
-	
-	if (AttrSplineType* Channel = GetTypedAttributeChannel<AttrSplineType>(Name))
-	{
-		return FromInternalSplineSpace(MapParameterFromChildSpace(Name, Channel->GetParameter(Index)));
-	}
-	return 0.0f;
-}
-
-/** FLegacySpline Definition */
-
-class FLegacySpline
-{
-	FInterpCurveVector PositionCurve;
-	FInterpCurveQuat RotationCurve;
-	FInterpCurveVector ScaleCurve;
-	FInterpCurveFloat ReparamTable;
-	
-public:
-
-	ENGINE_API FLegacySpline() = default;
-	ENGINE_API FLegacySpline(const FNewSpline& Other);
-	ENGINE_API FLegacySpline(const FSplineCurves& Other);
-	
-	/* Control Point Index Interface */
-
-	void AddPoint(const FSplinePoint& InPoint);
-	void InsertPoint(const FSplinePoint& InPoint, int32 Index);
-	FSplinePoint GetPoint(const int32 Index) const;
-	void RemovePoint(int32 Index);
-	
-	void SetLocation(int32 Index, const FVector& InLocation);
-	FVector GetLocation(const int32 Index) const;
-	
-	void SetInTangent(const int32 Index, const FVector& InTangent);
-	FVector GetInTangent(const int32 Index) const;
-	
-	void SetOutTangent(const int32 Index, const FVector& OutTangent);
-	FVector GetOutTangent(const int32 Index) const;
-
-	void SetRotation(int32 Index, const FQuat& InRotation);
-	FQuat GetRotation(const int32 Index) const;
-	
-	void SetScale(int32 Index, const FVector& InScale);
-	FVector GetScale(const int32 Index) const;
-	
-	void SetSplinePointType(int32 Index, EInterpCurveMode Type);
-	EInterpCurveMode GetSplinePointType(int32 Index) const;
-
-	float GetParameterAtIndex(int32 Index) const;
-	float GetParameterAtDistance(float Distance) const;
-	float GetDistanceAtParameter(float Parameter) const;
-	
-	/* Parameter Interface */
-
-	FVector Evaluate(float Param) const;
-	FVector EvaluateDerivative(float Param) const;
-	FQuat EvaluateRotation(float Param) const;
-	FVector EvaluateScale(float Param) const;
-	float FindNearest(const FVector& InLocation, float& OutSquaredDist) const;
-	float FindNearestOnSegment(const FVector& InLocation, int32 SegmentIndex, float& OutSquaredDist) const;
-	
-	/* Misc Interface */
-	
-	bool operator==(const FLegacySpline& Other) const
-	{
-		return PositionCurve == Other.PositionCurve
-			&& RotationCurve == Other.RotationCurve
-			&& ScaleCurve == Other.ScaleCurve;
-	}
-
-	bool operator!=(const FLegacySpline& Other) const
-	{
-		return !(*this == Other);
-	}
-	
-	friend ENGINE_API FArchive& operator<<(FArchive& Ar, FLegacySpline& Spline)
-	{
-		Spline.Serialize(Ar);
-		return Ar;
-	}
-	bool Serialize(FArchive& Ar);
-	
-	FInterpCurveVector& GetSplinePointsPosition() { return PositionCurve; }
-	const FInterpCurveVector& GetSplinePointsPosition() const { return PositionCurve; }
-	FInterpCurveQuat& GetSplinePointsRotation() { return RotationCurve; }
-	const FInterpCurveQuat& GetSplinePointsRotation() const { return RotationCurve; }
-	FInterpCurveVector& GetSplinePointsScale() { return ScaleCurve; }
-	const FInterpCurveVector& GetSplinePointsScale() const { return ScaleCurve; }
-
-	/** Returns the length of the specified spline segment up to the parametric value given */
-	float GetSegmentLength(const int32 Index, const float Param, const FVector& Scale3D = FVector(1.0f)) const;
-
-	/** Returns total length along this spline */
-	float GetSplineLength() const;
-	
-	/** Returns the total number of control points on this spline. */
-	int32 GetNumControlPoints() const;
-	
-	/** Reset the spline to an empty spline. */
-	void Reset();
-	
-	/** Reset the rotation attribute channel to default values. */
-	void ResetRotation();
-	
-	/** Reset the scale attribute channel to default values. */
-	void ResetScale();
-	
-	/** Update the spline's internal data according to the passed-in params. */
-	void UpdateSpline(const FSpline::FUpdateSplineParams& InParams);
-};
-
-
-
-/** FSpline Implementation */
-
-FSpline::FSpline()
-: CurrentImplementation(UE::Spline::GImplementation)
-, Version(0xffffffff)
-{
-#if WITH_EDITOR
-	PreviousImplementation = 0;
-#endif
-	
-	switch (CurrentImplementation)
-	{
-	default: break;		// Intentionally doing nothing
-	case 1: LegacyData = MakeShared<FLegacySpline>(); break;
-	case 2: NewData = MakeShared<FNewSpline>(); break;
-	}
-}
-
-FSpline::FSpline(const FSpline& Other)
-	: FSpline()
-{
-	*this = Other;
-}
-
-FSpline& FSpline::operator=(const FSpline& Other)
-{
-	CurrentImplementation = Other.CurrentImplementation;
-
-	switch (CurrentImplementation)
-	{
-	default:
-		break;
-	case 1:
-		*LegacyData = Other.LegacyData ? *Other.LegacyData : FLegacySpline();
-		break;
-	case 2:
-		*NewData = Other.NewData ? *Other.NewData : FNewSpline();
-		break;
-	}
-	
-	Version++;
-	
-	return *this;
-}
-
-FSpline& FSpline::operator=(const FSplineCurves& Other)
-{
-	switch (CurrentImplementation)
-	{
-	default:
-		break;
-	case 1:
-		*LegacyData = FLegacySpline(Other);
-		break;
-	case 2:
-		*NewData = FNewSpline(Other);
-		break;
-	}
-	
-	Version++;
-	
-	return *this;
-}
-
-void FSpline::AddPoint(const FSplinePoint& InPoint)
-{
-	switch (CurrentImplementation)
-	{
-	default: break;		// Intentionally doing nothing
-	case 1: LegacyData->AddPoint(InPoint); break;
-	case 2: NewData->AddPoint(InPoint); break;
-	}
-}
-
-void FSpline::InsertPoint(const FSplinePoint& InPoint, int32 Index)
-{
-	switch (CurrentImplementation)
-	{
-	default: break;		// Intentionally doing nothing
-	case 1: LegacyData->InsertPoint(InPoint, Index); break;
-	case 2: NewData->InsertPoint(InPoint, Index); break;
-	}
-}
-
-FSplinePoint FSpline::GetPoint(const int32 Index) const
-{
-	switch (CurrentImplementation)
-	{
-	default: return FSplinePoint();		// Intentionally doing nothing
-	case 1: return LegacyData->GetPoint(Index);
-	case 2: return NewData->GetPoint(Index);
-	}
-}
-
-void FSpline::RemovePoint(const int32 Index)
-{
-	switch (CurrentImplementation)
-	{
-	default: break;		// Intentionally doing nothing
-	case 1: LegacyData->RemovePoint(Index); break;
-	case 2: NewData->RemovePoint(Index); break;
-	}
-}
-
-void FSpline::SetLocation(int32 Index, const FVector& InLocation)
-{
-	switch (CurrentImplementation)
-	{
-	default: break;		// Intentionally doing nothing
-	case 1: LegacyData->SetLocation(Index, InLocation); break;
-	case 2: NewData->SetLocation(Index, InLocation); break;
-	}
-}
-
-FVector FSpline::GetLocation(const int32 Index) const
-{
-	switch (CurrentImplementation)
-	{
-	default: return FVector();
-	case 1: return LegacyData->GetLocation(Index);
-	case 2: return NewData->GetLocation(Index);
-	}
-}
-
-void FSpline::SetInTangent(const int32 Index, const FVector& InTangent)
-{
-	switch (CurrentImplementation)
-	{
-	default: break;
-	case 1: LegacyData->SetInTangent(Index, InTangent); break;
-	case 2: NewData->SetInTangent(Index, InTangent); break;
-	}
-}
-
-FVector FSpline::GetInTangent(const int32 Index) const
-{
-	switch (CurrentImplementation)
-	{
-	default: return FVector();
-	case 1: return LegacyData->GetInTangent(Index);
-	case 2: return NewData->GetInTangent(Index);
-	}
-}
-
-void FSpline::SetOutTangent(const int32 Index, const FVector& OutTangent)
-{
-	switch (CurrentImplementation)
-	{
-	default: break;
-	case 1: LegacyData->SetOutTangent(Index, OutTangent); break;
-	case 2: NewData->SetOutTangent(Index, OutTangent); break;
-	}
-}
-
-FVector FSpline::GetOutTangent(const int32 Index) const
-{
-	switch (CurrentImplementation)
-	{
-	default: return FVector();
-	case 1: return LegacyData->GetOutTangent(Index);
-	case 2: return NewData->GetOutTangent(Index);
-	}
-}
-
-void FSpline::SetRotation(int32 Index, const FQuat& InRotation)
-{
-	switch (CurrentImplementation)
-	{
-	default: break;
-	case 1: LegacyData->SetRotation(Index, InRotation); break;
-	case 2: NewData->SetRotation(Index, InRotation); break;
-	}
-}
-
-FQuat FSpline::GetRotation(const int32 Index) const
-{
-	switch (CurrentImplementation)
-	{
-	default: return FQuat();
-	case 1: return LegacyData->GetRotation(Index);
-	case 2: return NewData->GetRotation(Index);
-	}
-}
-
-void FSpline::SetScale(int32 Index, const FVector& InScale)
-{
-	switch (CurrentImplementation)
-	{
-	default: break;
-	case 1: LegacyData->SetScale(Index, InScale); break;
-	case 2: NewData->SetScale(Index, InScale); break;
-	}
-}
-
-FVector FSpline::GetScale(const int32 Index) const
-{
-	switch (CurrentImplementation)
-	{
-	default: return FVector(1.0);
-	case 1: return LegacyData->GetScale(Index);
-	case 2: return NewData->GetScale(Index);
-	}
-}
-
-void FSpline::SetSplinePointType(int32 Index, EInterpCurveMode Type)
-{
-	switch (CurrentImplementation)
-	{
-	default: break;
-	case 1: LegacyData->SetSplinePointType(Index, Type); break;
-	case 2: NewData->SetSplinePointType(Index, Type); break;
-	}
-}
-
-EInterpCurveMode FSpline::GetSplinePointType(int32 Index) const
-{
-	switch (CurrentImplementation)
-	{
-	default: return CIM_Unknown;
-	case 1: return LegacyData->GetSplinePointType(Index);
-	case 2: return NewData->GetSplinePointType(Index);
-	}
-}
-
-float FSpline::GetParameterAtIndex(int32 Index) const
-{
-	switch (CurrentImplementation)
-	{
-	default: return 0.f;
-	case 1: return LegacyData->GetParameterAtIndex(Index);
-	case 2: return NewData->GetParameterAtIndex(Index);
-	}
-}
-
-float FSpline::GetParameterAtDistance(float Distance) const
-{
-	switch (CurrentImplementation)
-	{
-	default: return 0.f;
-	case 1: return LegacyData->GetParameterAtDistance(Distance);
-	case 2: return NewData->GetParameterAtDistance(Distance);
-	}
-}
-
-float FSpline::GetDistanceAtParameter(float Parameter) const
-{
-	switch (CurrentImplementation)
-	{
-	default: return 0.f;
-	case 1: return LegacyData->GetDistanceAtParameter(Parameter);
-	case 2: return NewData->GetDistanceAtParameter(Parameter);
-	}
-}
-
-FQuat FSpline::GetOrientation(int32 Index) const
-{
-	return GetOrientation(GetParameterAtIndex(Index));
-}
-
-void FSpline::SetOrientation(int32 Index, const FQuat& InOrientation)
-{
-	if (Index < 0 || Index >= GetNumControlPoints())
-	{
-		return;
-	}
-	
-	// Work backwards to compute the rotation that is currently being applied.
-	const FQuat RelativeRotation = InOrientation * GetOrientation(Index).Inverse();
-	
-	// Store rotation which transforms the world up vector to the local up vector.
-	SetRotation(Index, FQuat::FindBetween(FVector::UpVector, InOrientation.GetUpVector()));
-
-	// Align tangents with rotation, preserving magnitude.
-	const FVector OldInTangent = GetInTangent(Index);
-	const FVector OldInTangentDirection = OldInTangent.GetSafeNormal();
-	const float InTangentMag = OldInTangent.Length();
-	const FVector NewInTangentDirection = RelativeRotation.RotateVector(OldInTangentDirection);
-	
-	const FVector OldOutTangent = GetOutTangent(Index);
-	const FVector OldOutTangentDirection = OldOutTangent.GetSafeNormal();
-	const float OutTangentMag = OldOutTangent.Length();
-	const FVector NewOutTangentDirection = RelativeRotation.RotateVector(OldOutTangentDirection);
-	
-	SetInTangent(Index, InTangentMag * NewInTangentDirection);
-	SetOutTangent(Index, OutTangentMag * NewOutTangentDirection);
-}
-
-FVector FSpline::Evaluate(float Param) const
-{
-	switch (CurrentImplementation)
-	{
-	default: return FVector();
-	case 1: return LegacyData->Evaluate(Param);
-	case 2: return NewData->EvaluatePosition(Param);
-	}
-}
-
-FVector FSpline::EvaluateDerivative(float Param) const
-{
-	if (UE::Spline::GApproximateTangents)
-	{
-		// Approximate using central difference.
-		
-		// This computes the tangent direction using central difference and
-		// assumes that tangent magnitude is linearly changing between control points.
-		// While the assumption about magnitude is probably wrong, it works well.
-		
-		auto ClampValidParameterRange = [Min = 0.f, Max = GetNumControlPoints() - 1](float Param)
-			{ return FMath::Clamp(Param, Min, Max); };
-
-		constexpr float H = UE_KINDA_SMALL_NUMBER;
-		const float ParamLow = ClampValidParameterRange(Param - H);
-		const float ParamHigh = ClampValidParameterRange(Param + H);
-		const FVector Tangent = ((Evaluate(ParamHigh) - Evaluate(ParamLow)) / (ParamHigh - ParamLow)).GetSafeNormal();
-
-		const int32 Index1 = FMath::Clamp((int32)Param, 0, GetNumControlPoints() - 1);
-		const int32 Index2 = FMath::Clamp((int32)(Param+1), 0, GetNumControlPoints() - 1);
-		const float Mag1 = GetInTangent(Index1).Length();
-		const float Mag2 = GetInTangent(Index2).Length();
-		const float Mag = FMath::Lerp(Mag1, Mag2, FMath::Frac(Param));
-
-		return Tangent * Mag;
-	}
-	else
-	{
-		switch (CurrentImplementation)
-		{
-		default: return FVector();
-		case 1: return LegacyData->EvaluateDerivative(Param);
-		case 2: return NewData->EvaluateDerivative(Param);
-		}
-	}
-}
-
-FQuat FSpline::EvaluateRotation(float Param) const
-{
-	switch (CurrentImplementation)
-	{
-	default: return FQuat();
-	case 1: return LegacyData->EvaluateRotation(Param);
-	case 2: return NewData->EvaluateRotation(Param);
-	}
-}
-
-FVector FSpline::EvaluateScale(float Param) const
-{
-	switch (CurrentImplementation)
-	{
-	default: return FVector();
-	case 1: return LegacyData->EvaluateScale(Param);
-	case 2: return NewData->EvaluateScale(Param);
-	}	
-}
-
-FQuat FSpline::GetOrientation(float Param) const
-{
-	FQuat Rotation = EvaluateRotation(Param);
-	Rotation.Normalize();
-
-	FVector Direction = EvaluateDerivative(Param);
-	Direction = Direction.GetSafeNormal();
-
-	const FVector UpVector = Rotation.RotateVector(FVector::UpVector);
-	
-	return FRotationMatrix::MakeFromXZ(Direction, UpVector).ToQuat();
-}
-
-bool FSpline::HasAttributeChannel(FName AttributeName) const
-{
-	if (NewData && SupportsAttributes())
-	{
-		return NewData->HasAttributeChannel(AttributeName);
-	}
-
-	return false;
-}
-
-bool FSpline::RemoveAttributeChannel(FName AttributeName) const
-{
-	if (NewData && SupportsAttributes())
-	{
-		return NewData->RemoveAttributeChannel(AttributeName);
-	}
-
-	return false;
-}
-
-TArray<FName> FSpline::GetFloatPropertyChannels() const
-{
-	TArray<FName> ChannelNames;
-	if (NewData && SupportsAttributes())
-	{
-		ChannelNames = NewData->GetAttributeChannelNamesByValueType<float>();
-	}
-	return ChannelNames;
-}
-
-TArray<FName> FSpline::GetVectorPropertyChannels() const
-{
-	TArray<FName> ChannelNames;
-	if (NewData && SupportsAttributes())
-	{
-		ChannelNames = NewData->GetAttributeChannelNamesByValueType<FVector>();
-
-		// Don't report these internal attribute channels to the caller
-		ChannelNames.Remove(NewData->ScaleAttrName);
-	}
-	return ChannelNames;
-}
-
-template <typename AttrType>
-float FSpline::GetAttributeParameter(int32 Index, const FName& Name) const
-{
-	if (NewData && SupportsAttributes())
-	{
-		return NewData->GetAttributeParameter<AttrType>(Name, Index);
-	}
-
-	return 0.f;
-}
-template ENGINE_API float FSpline::GetAttributeParameter<float>(int32 Index, const FName& Name) const;
-template ENGINE_API float FSpline::GetAttributeParameter<FVector>(int32 Index, const FName& Name) const;
-
-template <typename AttrType>
-int32 FSpline::SetAttributeParameter(int32 Index, float Parameter, const FName& Name)
-{
-	if (NewData && SupportsAttributes())
-	{
-		return NewData->SetAttributeParameter<AttrType>(Name, Index, Parameter);
-	}
-	
-	return Index;
-}
-template ENGINE_API int32 FSpline::SetAttributeParameter<float>(int32 Index, float Parameter, const FName& Name);
-template ENGINE_API int32 FSpline::SetAttributeParameter<FVector>(int32 Index, float Parameter, const FName& Name);
-
-template <typename AttrType>
-int32 FSpline::NumAttributeValues(FName AttributeName) const
-{
-	if (NewData && SupportsAttributes())
-	{
-		return NewData->NumAttributeValues<AttrType>(AttributeName);
-	}
-
-	return 0;
-}
-template ENGINE_API int32 FSpline::NumAttributeValues<float>(FName AttributeName) const;
-template ENGINE_API int32 FSpline::NumAttributeValues<FVector>(FName AttributeName) const;
-
-template<typename AttrType>
-AttrType FSpline::GetAttributeValue(int32 Index, const FName& Name) const
-{
-	if (NewData && SupportsAttributes())
-	{
-		return NewData->GetAttributeValue<AttrType>(Name, Index);
-	}
-	
-	return AttrType();
-}
-template ENGINE_API float FSpline::GetAttributeValue<float>(int32 Index, const FName& Name) const;
-template ENGINE_API FVector FSpline::GetAttributeValue<FVector>(int32 Index, const FName& Name) const;
-
-template<typename AttrType>
-void FSpline::SetAttributeValue(int32 Index, const AttrType& Value, const FName& Name)
-{
-	if (NewData && SupportsAttributes())
-	{
-		NewData->SetAttributeValue<AttrType>(Name, Value, Index);
-	}
-}
-template ENGINE_API void FSpline::SetAttributeValue<float>(int32 Index, const float& Value, const FName& Name);
-template ENGINE_API void FSpline::SetAttributeValue<FVector>(int32 Index, const FVector& Value, const FName& Name);
-
-template<typename AttrType>
-bool FSpline::CreateAttributeChannel(FName AttributeName) const
-{
-	if (NewData && SupportsAttributes())
-	{
-		if (UE::Geometry::Spline::TTangentBezierSpline<AttrType>* NewAttrChannel = NewData->CreateAttributeChannel<UE::Geometry::Spline::TTangentBezierSpline<AttrType>>(AttributeName); NewAttrChannel != nullptr)
-		{
-			NewAttrChannel->SetPreInfinityMode(UE::Geometry::Spline::EOutOfBoundsHandlingMode::Constant);
-			NewAttrChannel->SetPostInfinityMode(UE::Geometry::Spline::EOutOfBoundsHandlingMode::Constant);
-
-			return true;
-		}
-	}
-	
-	return false;
-}
-template ENGINE_API bool FSpline::CreateAttributeChannel<float>(FName AttributeName) const;
-template ENGINE_API bool FSpline::CreateAttributeChannel<FVector>(FName AttributeName) const;
-
-template<typename AttrType>
-int32 FSpline::AddAttributeValue(float Param, const AttrType& Value, FName AttributeName) const
-{
-	if (NewData && SupportsAttributes())
-	{
-		return NewData->AddAttributeValue(AttributeName, Value, Param);
-	}
-	
-	return INDEX_NONE;
-}
-template ENGINE_API int32 FSpline::AddAttributeValue<float>(float Param, const float& Value, FName AttributeName) const;
-template ENGINE_API int32 FSpline::AddAttributeValue<FVector>(float Param, const FVector& Value, FName AttributeName) const;
-
-template <typename AttrType>
-void FSpline::RemoveAttributeValue(int32 Index, FName AttributeName)
-{
-	if (NewData && SupportsAttributes())
-	{
-		NewData->RemoveAttributeValue<AttrType>(AttributeName, Index);
-	}
-}
-template ENGINE_API void FSpline::RemoveAttributeValue<float>(int32 Index, FName AttributeName);
-template ENGINE_API void FSpline::RemoveAttributeValue<FVector>(int32 Index, FName AttributeName);
-
-template<typename AttrType>
-AttrType FSpline::EvaluateAttribute(float Param, FName AttributeName) const
-{
-	if (NewData && SupportsAttributes())
-	{
-		return NewData->EvaluateAttribute<AttrType>(AttributeName, NewData->ToInternalSplineSpace(Param));
-	}
-	return AttrType(0);
-}
-template ENGINE_API float FSpline::EvaluateAttribute<float>(float Param, FName AttributeName) const;
-template ENGINE_API FVector FSpline::EvaluateAttribute<FVector>(float Param, FName AttributeName) const;
-
-float FSpline::FindNearest(const FVector& InLocation, float& OutSquaredDist) const
-{
-	switch (CurrentImplementation)
-	{
-	default: return 0.f;
-	case 1: return LegacyData->FindNearest(InLocation, OutSquaredDist);
-	case 2: return NewData->FindNearest(InLocation, OutSquaredDist);
-	}
-}
-
-float FSpline::FindNearestOnSegment(const FVector& InLocation, int32 SegmentIndex, float& OutSquaredDist) const
-{
-	switch (CurrentImplementation)
-	{
-	default: return 0.f;
-	case 1: return LegacyData->FindNearestOnSegment(InLocation, SegmentIndex, OutSquaredDist);
-	case 2: return NewData->FindNearestOnSegment(InLocation, SegmentIndex, OutSquaredDist);
-	}
-}
-
-bool FSpline::operator==(const FSpline& Other) const
-{
-	if (CurrentImplementation == Other.CurrentImplementation)
-	{
-		if (LegacyData && Other.LegacyData)
-		{
-			return *LegacyData == *Other.LegacyData;
-		}
-
-		if (NewData && Other.NewData)
-		{
-			return *NewData == *Other.NewData;
-		}
-	}
-
-	return false;
-}
-
-bool FSpline::operator!=(const FSpline& Other) const
-{
-	return !(*this == Other);
-}
-
-bool FSpline::IsNearlyEqual(const FSpline& Other, const float RelativeTolerance, const float AbsoluteTolerance) const
-{
-	if (!ensureAlways(IsNew() && Other.IsNew()))
-	{
-		return false;
-	}
-
-	const int32 NumPoints = NewData->GetSpline().GetNumPoints();
-	if (Other.NewData->GetSpline().GetNumPoints() != NumPoints)
-	{
-		return false;
-	}
-
-	auto ScalarNearlyEqual = [AbsoluteTolerance, RelativeTolerance](const float& A, const float& B)
-	{
-		const float Delta = FMath::Abs(A - B);
-		const float Tolerance = FMath::Max(AbsoluteTolerance, RelativeTolerance * FMath::Max(FMath::Abs(A), FMath::Abs(B)));
-		return Delta <= Tolerance;
-	};
-
-	auto VectorNearlyEqual = [AbsoluteTolerance, RelativeTolerance](const FVector& A, const FVector& B)
-	{
-		const float MagDelta = (A - B).Length();
-		const float MagTolerance = FMath::Max(AbsoluteTolerance, RelativeTolerance * FMath::Max(A.Length(), B.Length()));
-		return MagDelta <= MagTolerance;
-	};
-
-	auto QuatComponentsNearlyEqual = [AbsoluteTolerance, RelativeTolerance](const FQuat& A, const FQuat& B)
-	{
-		const float DX = FMath::Abs(A.X - B.X);
-		const float TX = FMath::Max(AbsoluteTolerance, RelativeTolerance * FMath::Max(FMath::Abs(A.X), FMath::Abs(B.X)));
-		if (DX > TX) return false;
-
-		const float DY = FMath::Abs(A.Y - B.Y);
-		const float TY = FMath::Max(AbsoluteTolerance, RelativeTolerance * FMath::Max(FMath::Abs(A.Y), FMath::Abs(B.Y)));
-		if (DY > TY) return false;
-
-		const float DZ = FMath::Abs(A.Z - B.Z);
-		const float TZ = FMath::Max(AbsoluteTolerance, RelativeTolerance * FMath::Max(FMath::Abs(A.Z), FMath::Abs(B.Z)));
-		if (DZ > TZ) return false;
-
-		const float DW = FMath::Abs(A.W - B.W);
-		const float TW = FMath::Max(AbsoluteTolerance, RelativeTolerance * FMath::Max(FMath::Abs(A.W), FMath::Abs(B.W)));
-		return DW <= TW;
-	};
-
-	for (int32 PointIdx = 0; PointIdx < NumPoints; ++PointIdx)
-	{
-		if (!VectorNearlyEqual(NewData->GetLocation(PointIdx), Other.NewData->GetLocation(PointIdx)))
-		{
-			return false;
-		}
-
-		if (!VectorNearlyEqual(NewData->GetInTangent(PointIdx), Other.NewData->GetInTangent(PointIdx)))
-		{
-			return false;
-		}
-
-		if (!VectorNearlyEqual(NewData->GetOutTangent(PointIdx), Other.NewData->GetOutTangent(PointIdx)))
-		{
-			return false;
-		}
-
-		const FQuat Rotation = NewData->GetRotation(PointIdx);
-		const FQuat OtherRotation = Other.NewData->GetRotation(PointIdx);
-		const FQuat OtherRotationNeg = FQuat(-OtherRotation.X, -OtherRotation.Y, -OtherRotation.Z, -OtherRotation.W);
-		if (!QuatComponentsNearlyEqual(Rotation, OtherRotation) && !QuatComponentsNearlyEqual(Rotation, OtherRotationNeg))
-		{
-			return false;
-		}
-
-		if (!VectorNearlyEqual(NewData->GetScale(PointIdx), Other.NewData->GetScale(PointIdx)))
-		{
-			return false;
-		}
-
-		if (!ScalarNearlyEqual(NewData->GetParameterAtIndex(PointIdx), Other.NewData->GetParameterAtIndex(PointIdx)))
-		{
-			return false;
-		}
-	}
-
-	const TArray<FName> OtherFloatChannelNames = Other.GetFloatPropertyChannels();
-	for (const FName& FloatChannelName : GetFloatPropertyChannels())
-	{
-		if (!OtherFloatChannelNames.Contains(FloatChannelName))
-		{
-			return false;
-		}
-
-		const int32 NumFloatChannelPoints = NumAttributeValues<float>(FloatChannelName);
-		if (Other.NumAttributeValues<float>(FloatChannelName) != NumFloatChannelPoints)
-		{
-			return false;
-		}
-
-		for (int32 FloatChannelIdx = 0; FloatChannelIdx < NumFloatChannelPoints; ++FloatChannelIdx)
-		{
-			if (!ScalarNearlyEqual(
-				GetAttributeValue<float>(FloatChannelIdx, FloatChannelName),
-				Other.GetAttributeValue<float>(FloatChannelIdx, FloatChannelName)))
-			{
-				return false;
-			}
-
-			if (!ScalarNearlyEqual(
-				GetAttributeParameter<float>(FloatChannelIdx, FloatChannelName),
-				Other.GetAttributeParameter<float>(FloatChannelIdx, FloatChannelName)))
-			{
-				return false;
-			}
-		}
-	}
-
-	return true;
-}
-
-bool FSpline::Serialize(FArchive& Ar)
-{
-	Ar.UsingCustomVersion(FFortniteMainBranchObjectVersion::GUID);
-	
-	// Data format:
-
-	// Byte 1 - The data format, determine by CurrentImplementation at the time of last save
-	// Remaining N Bytes - Spline data (or empty, if byte 0 is 0). The format is determined by byte 1.
-	
-	if (Ar.IsLoading())
-	{
-		SerializeLoad(Ar);
-	}
-	else
-	{
-		SerializeSave(Ar);
-	}
-
-	return true;
-}
-
-void FSpline::SerializeLoad(FArchive& Ar)
-{
-	int8 PreviousImpl;
-	Ar << PreviousImpl;
-#if WITH_EDITOR
-	PreviousImplementation = PreviousImpl;
-#endif
-
-	auto WasEnabled = [PreviousImpl]()
-	{ return PreviousImpl != 0; };
-
-	auto WasLegacy = [PreviousImpl]()
-	{ return PreviousImpl == 1; };
-		
-	if (WasEnabled())
-	{
-		if (WasLegacy())
-		{
-			if (IsLegacy())
-			{
-				LegacyData = MakeShared<FLegacySpline>();
-				Ar << *LegacyData;
-			}
-			else if (IsNew())
-			{
-				FLegacySpline IntermediateData;
-				Ar << IntermediateData;
-				NewData = MakeShared<FNewSpline>(IntermediateData);
-			}
-			else if (!IsEnabled())
-			{
-				// Intentionally doing nothing with data
-				FLegacySpline IntermediateData;
-				Ar << IntermediateData;
-			}
-		}
-		else
-		{
-			if (IsLegacy())
-			{
-				FNewSpline IntermediateData;
-				Ar << IntermediateData;
-				LegacyData = MakeShared<FLegacySpline>(IntermediateData);
-			}
-			else if (IsNew())
-			{
-				NewData = MakeShared<FNewSpline>();
-				Ar << *NewData;
-			}
-			else if (!IsEnabled())
-			{
-				// Intentionally doing nothing with data
-				FNewSpline IntermediateData;
-				Ar << IntermediateData;
-			}
-		}
-	}
-}
-
-void FSpline::SerializeSave(FArchive& Ar) const
-{
-	uint8 CurrentImpl = CurrentImplementation;
-	Ar << CurrentImpl;
-		
-	if (IsLegacy())
-	{
-		Ar << *LegacyData;
-	}
-	else if (IsNew())
-	{
-		Ar << *NewData;
-	}
-}
-
-bool FSpline::ExportTextItem(FString& ValueStr, FSpline const& DefaultValue, class UObject* Parent, int32 PortFlags, class UObject* ExportRootScope) const
-{
-	// serialize our spline
-    TArray<uint8> SplineWriteBuffer;
-    FMemoryWriter MemWriter(SplineWriteBuffer);
-    SerializeSave(MemWriter);
-
-	if (UE::Spline::GEncodeAsHex)
-	{
-		FString HexStr = FString::FromHexBlob(SplineWriteBuffer.GetData(), SplineWriteBuffer.Num());
-		ValueStr = FString::Printf(TEXT("SplineData SplineDataLen=%d SplineData=%s\r\n"), HexStr.Len(), *HexStr);
-	}
-	else
-	{
-		FString Base64String = FBase64::Encode(SplineWriteBuffer);
-
-		// Base64 encoding uses the '/' character, but T3D interprets '//' as some kind of
-		// terminator (?). If it occurs then the string passed to ImportTextItem() will
-		// come back as full of nullptrs. So we will swap in '-' here, and swap back to '/' in ImportTextItem()
-		Base64String.ReplaceCharInline('/', '-');
-
-		ValueStr = FString::Printf(TEXT("SplineData SplineDataLen=%d SplineData=%s\r\n"), Base64String.Len(), *Base64String);
-	}
-	
-	return true;
-}
-
-bool FSpline::ImportTextItem(const TCHAR*& SourceText, int32 PortFlags, UObject* Parent, FOutputDevice* ErrorText)
-{
-	if (FParse::Command(&SourceText, TEXT("SplineData")))	
-	{
-		static const TCHAR SplineDataLenToken[] = TEXT("SplineDataLen=");
-		const TCHAR* FoundSplineDataLenStart = FCString::Strifind(SourceText, SplineDataLenToken);
-		if (FoundSplineDataLenStart)
-		{
-			SourceText = FoundSplineDataLenStart + FCString::Strlen(SplineDataLenToken);
-			int32 SplineDataLen = FCString::Atoi(SourceText);
-
-			static const TCHAR SplineDataToken[] = TEXT("SplineData=");
-			const TCHAR* FoundSplineDataStart = FCString::Strifind(SourceText, SplineDataToken);
-			if (FoundSplineDataStart)
-			{
-				SourceText = FoundSplineDataStart + FCString::Strlen(SplineDataToken);
-				FString SplineData = FString::ConstructFromPtrSize(SourceText, SplineDataLen);
-
-				// fix-up the hack applied to the Base64-encoded string in ExportTextItem()
-				SplineData.ReplaceCharInline('-', '/');
-
-				TArray<uint8> SplineReadBuffer;
-				bool bDecoded;
-				
-				if (UE::Spline::GEncodeAsHex)
-				{
-					SplineReadBuffer.SetNum(SplineDataLen);
-					bDecoded = FString::ToHexBlob(SplineData, SplineReadBuffer.GetData(), SplineDataLen);
-				}
-				else
-				{
-					bDecoded = FBase64::Decode(SplineData, SplineReadBuffer);
-				}
-
-				if (bDecoded)
-				{
-					FMemoryReader MemReader(SplineReadBuffer);
-					SerializeLoad(MemReader);
-				}
-			}
-		}
-
-
-	}
-	
-	return true;
-}
-
-const FInterpCurveVector& FSpline::GetSplinePointsPosition() const
-{
-	switch (CurrentImplementation)
-	{
-	default:
-		return PositionCurve;
-	case 1:
-		return LegacyData->GetSplinePointsPosition();
-	case 2:
-		NewData->RebuildLegacyCurves();
-		return NewData->GetSplinePointsPosition();
-	}
-}
-
-const FInterpCurveQuat& FSpline::GetSplinePointsRotation() const
-{
-	switch (CurrentImplementation)
-	{
-	default:
-		return RotationCurve;
-	case 1:
-		return LegacyData->GetSplinePointsRotation();
-	case 2:
-		NewData->RebuildLegacyCurves();
-		return NewData->GetSplinePointsRotation();
-	}
-}
-
-const FInterpCurveVector& FSpline::GetSplinePointsScale() const
-{
-	switch (CurrentImplementation)
-	{
-	default:
-		return ScaleCurve;
-	case 1:
-		return LegacyData->GetSplinePointsScale();
-	case 2:
-		NewData->RebuildLegacyCurves();
-		return NewData->GetSplinePointsScale();
-	}
-}
-
-float FSpline::GetSegmentLength(const int32 Index, const float Param, const FVector& Scale3D) const
-{
-	switch (CurrentImplementation)
-	{
-	default: return 0.f;
-	case 1: return LegacyData->GetSegmentLength(Index, Param, Scale3D);
-	case 2: return NewData->GetSegmentLength(Index, Param, Scale3D);
-	}
-}
-
-float FSpline::GetSplineLength() const
-{
-	switch (CurrentImplementation)
-	{
-	default: return 0.f;	// Intentionally doing nothing
-	case 1: return LegacyData->GetSplineLength();
-	case 2: return NewData->GetSplineLength();
-	}
-}
-
-int32 FSpline::GetNumSegments() const
-{
-	switch (CurrentImplementation)
-	{
-	default: return 0;
-	case 1: return LegacyData->GetSplinePointsPosition().bIsLooped ? LegacyData->GetNumControlPoints() : FMath::Max(0, LegacyData->GetNumControlPoints() - 1);
-	case 2: return NewData->GetSpline().GetNumberOfSegments();
-	}
-}
-
-int32 FSpline::GetNumControlPoints() const
-{
-	switch (CurrentImplementation)
-	{
-	default: return 0;
-	case 1: return LegacyData->GetNumControlPoints();
-	case 2: return NewData->GetSpline().GetNumPoints();
-	}
-}
-
-void FSpline::Reset()
-{
-	switch (CurrentImplementation)
-	{
-	default: break;
-	case 1: LegacyData->Reset(); break;
-	case 2: NewData->Reset(); break;
-	}
-}
-	
-void FSpline::ResetRotation()
-{
-	switch (CurrentImplementation)
-	{
-	default: break;
-	case 1: LegacyData->ResetRotation(); break;
-	case 2: NewData->ResetRotation(); break;
-	}
-}
-	
-void FSpline::ResetScale()
-{
-	switch (CurrentImplementation)
-	{
-	default: break;
-	case 1: LegacyData->ResetScale(); break;
-	case 2: NewData->ResetScale(); break;
-	}
-}
-
-void FSpline::SetClosedLoop(bool bClosed)
-{
-	TRACE_CPUPROFILER_EVENT_SCOPE(FSpline::SetClosedLoop);
-
-	CachedUpdateSplineParams.bClosedLoop = bClosed;
-	UpdateSpline();
-}
-
-bool FSpline::IsClosedLoop() const
-{
-	return CachedUpdateSplineParams.bClosedLoop;
-}
-
-void FSpline::UpdateSpline(const FUpdateSplineParams& InParams)
-{
-	TRACE_CPUPROFILER_EVENT_SCOPE(FSpline::UpdateSpline);
-
-#if WITH_EDITOR
-	if (UE::Spline::GImplementation != CurrentImplementation)
-	{
-		OnSplineImplementationChanged();
-	}
-#endif
-		
-	switch (CurrentImplementation)
-    {
-    default: break;
-    case 1: LegacyData->UpdateSpline(InParams); break;
-    case 2: NewData->UpdateSpline(InParams); break;
-    }
-
-	CachedUpdateSplineParams = InParams;
-	++Version;
-}
-
-void FSpline::UpdateSpline()
-{
-	UpdateSpline(CachedUpdateSplineParams);
-}
-
-TSharedPtr<UE::Geometry::Spline::TSplineInterface<FVector>> FSpline::GetSplineInterface() const
-{
-	if (IsNew())
-	{
-		return StaticCastSharedPtr<UE::Geometry::Spline::TSplineInterface<FVector>>(NewData);
-	}
-
-	return nullptr;
-}
-
-bool FSpline::IsEnabledStatic()
-{
-	return UE::Spline::GImplementation != 0;
-}
-
-#if WITH_EDITOR
-void FSpline::OnSplineImplementationChanged()
-{
-	using namespace UE::Spline;
-
-	// This implements a state machine I drew out.
-	// I am pretty confident it covers all cases because it handles all edges of a
-	// directed digraph (of 3 nodes, the 3 possible states, with self-edges) where edges represent state transitions.
-
-	// Do you see the nice triangle?
-	
-	if (GImplementation == CurrentImplementation)			// 0->0, 1->1, 2->2
-	{
-	}
-	else if (GImplementation == 0)							// 1->0, 2->0
-	{
-		LegacyData.Reset();
-		NewData.Reset();
-	}
-	else if (CurrentImplementation == 0)					// 0->1, 0->2
-	{
-		// Here we just need to alloc but not copy
-		// GImplementation guaranteed to be 1 or 2
-		if (GImplementation == 1)
-		{
-			LegacyData = MakeShared<FLegacySpline>();
-		}
-		else if (GImplementation == 2)
-		{
-			NewData = MakeShared<FNewSpline>();
-		}
-	}
-	else if (CurrentImplementation == 1)					// 1->2
-	{
-		NewData = MakeShared<FNewSpline>(*LegacyData);
-		LegacyData.Reset();
-	}
-	else if (CurrentImplementation == 2)					// 2->1
-	{
-		LegacyData = MakeShared<FLegacySpline>(*NewData);
-		NewData.Reset();
-	}
-
-	CurrentImplementation = GImplementation;
-}
-#endif
-
-
-
-/** FNewSpline Implementation */
-
-FNewSpline::FNewSpline()
-{
-	// This is a hack. We just need to make sure that this spline type registers itself with the spline registry.
-	// This should be removed in the future once we can register this spline type automatically without instantiation.
-	UE::Geometry::Spline::TTangentBezierSpline<float> AutoRegister;
-	
-	CreateAttributeChannel<UE::Geometry::Spline::TTangentBezierSpline<FQuat>>(RotationAttrName);
-	CreateAttributeChannel<UE::Geometry::Spline::TTangentBezierSpline<FVector>>(ScaleAttrName);
-	
-	ValidateRotScale();
-}
-
-FNewSpline::FNewSpline(const FNewSpline& Other)
-{
-    *this = Other;
-
-	ValidateRotScale();
-}
-
-FNewSpline::FNewSpline(const FLegacySpline& Other)
-{
-	CreateAttributeChannel<UE::Geometry::Spline::TTangentBezierSpline<FQuat>>(RotationAttrName);
-	CreateAttributeChannel<UE::Geometry::Spline::TTangentBezierSpline<FVector>>(ScaleAttrName);
-	
-	TArray<FTangentBezierControlPoint> Points;
-    for (int32 Idx = 0; Idx < Other.GetNumControlPoints(); Idx++)
-    {
-	    FTangentBezierControlPoint& Point = Points.AddDefaulted_GetRef();
-
-    	Point.Position = Other.GetLocation(Idx);
-    	Point.TangentIn = Other.GetInTangent(Idx);
-    	Point.TangentOut = Other.GetOutTangent(Idx);
-    	Point.TangentMode = ConvertInterpCurveModeToTangentMode(Other.GetSplinePointType(Idx));
-    }
-
-	// initialize spline
-	GetSpline().SetControlPoints(Points);
-	ResetRotation();
-	ResetScale();
-	
-	// setup rotation/scale attributes
-	for (int32 Idx = 0; Idx < GetSpline().GetNumPoints(); Idx++)
-	{    		
-		SetAttributeValue<FQuat>(RotationAttrName, Other.GetRotation(Idx), Idx);
-		SetAttributeValue<FVector>(ScaleAttrName, Other.GetScale(Idx), Idx);
-	}
-	
-	ValidateRotScale();
-
-	FSpline::FUpdateSplineParams UpdateParams;
-	UpdateParams.bClosedLoop = Other.GetSplinePointsPosition().bIsLooped;
-	
-    UpdateSpline(UpdateParams);
-}
-
-FNewSpline::FNewSpline(const FSplineCurves& Other)
-{
-	CreateAttributeChannel<UE::Geometry::Spline::TTangentBezierSpline<FQuat>>(RotationAttrName);
-	CreateAttributeChannel<UE::Geometry::Spline::TTangentBezierSpline<FVector>>(ScaleAttrName);
-	
-	TArray<FTangentBezierControlPoint> Points;
-    for (int32 Idx = 0; Idx < Other.Position.Points.Num(); Idx++)
-    {
-	    const FInterpCurvePoint<FVector>& Position = Other.Position.Points[Idx];
-    	
-    	FTangentBezierControlPoint& Point = Points.AddDefaulted_GetRef();
-
-    	Point.Position = Position.OutVal;
-    	Point.TangentIn = Position.ArriveTangent;
-    	Point.TangentOut = Position.LeaveTangent;
-    	Point.TangentMode = ConvertInterpCurveModeToTangentMode(Position.InterpMode);
-    }
-	
-	// initialize spline
-	GetSpline().SetControlPoints(Points);
-	ResetRotation();
-	ResetScale();
-	
-    // setup rotation/scale attributes
-	for (int32 Idx = 0; Idx < Other.Position.Points.Num(); Idx++)
-	{
-		FQuat Rotation = FQuat::Identity;
-		if (Other.Rotation.Points.IsValidIndex(Idx))
-		{
-			Rotation = Other.Rotation.Points[Idx].OutVal;
-		}
-		else
-		{
-			UE_LOG(LogSpline, Warning, TEXT("Number of rotation values (%u) different than the number of position values (%u)!"), Other.Rotation.Points.Num(), Other.Position.Points.Num());
-		}
-    	
-		FVector Scale = FVector::OneVector;
-		if (Other.Scale.Points.IsValidIndex(Idx))
-		{
-			Scale = Other.Scale.Points[Idx].OutVal;
-		}
-		else
-		{
-			UE_LOG(LogSpline, Warning, TEXT("Number of scale values (%u) different than the number of position values (%u)!"), Other.Scale.Points.Num(), Other.Position.Points.Num());
-		}
-    	
-		SetAttributeValue<FQuat>(RotationAttrName, Rotation, Idx);
-		SetAttributeValue<FVector>(ScaleAttrName, Scale, Idx);
-	}
-
-	ValidateRotScale();
-	
-	FSpline::FUpdateSplineParams UpdateParams;
-	UpdateParams.bClosedLoop = Other.Position.bIsLooped;
-	
-    UpdateSpline(UpdateParams);
-}
-
-/** Copy assignment */
-FNewSpline& FNewSpline::operator=(const FNewSpline& Other)
-{
-    if (this != &Other)
-    {
-		TMultiSpline::operator=(Other);
-		PositionCurve = Other.PositionCurve;
-		RotationCurve = Other.RotationCurve;
-		ScaleCurve = Other.ScaleCurve;
-		ReparamStepsPerSegment = Other.ReparamStepsPerSegment;
-    }
-	
-    return *this;
-}
-
-bool FNewSpline::operator!=(const FNewSpline& Other) const
-{
-    return !(*this == Other);
-}
-
-bool FNewSpline::Serialize(FArchive& Ar)
-{
-    TMultiSpline::Serialize(Ar);
-    Ar << ReparamStepsPerSegment;
-
-	if (Ar.IsLoading())
-	{
-		Reparameterize();
-
-		// Do not serialize legacy curves, only generate if loading.
-		RebuildLegacyCurves();
-	}
-	
-    return true;
-}
-
-FInterval1f FNewSpline::GetSegmentSpace() const
-{
-	return GetNumberOfSegments() == 0 ? FInterval1f::Empty() : FInterval1f(0.f, static_cast<float>(GetNumberOfSegments()));
-}
-
-float FNewSpline::GetSegmentLength(const int32 Index, const float Param, const FVector& Scale3D) const
-{
-	const int32 NumPoints = GetSpline().GetNumPoints();
-	const int32 LastPoint = NumPoints - 1;
-
-	check(Index >= 0 && ((IsClosedLoop() && Index < NumPoints) || (!IsClosedLoop() && Index < LastPoint)));
-	check(Param >= 0.0f && Param <= 1.0f);
-
-	// Evaluate the length of a Hermite spline segment.
-	// This calculates the integral of |dP/dt| dt, where P(t) is the spline equation with components (x(t), y(t), z(t)).
-	// This isn't solvable analytically, so we use a numerical method (Legendre-Gauss quadrature) which performs very well
-	// with functions of this type, even with very few samples.  In this case, just 5 samples is sufficient to yield a
-	// reasonable result.
-
-	struct FLegendreGaussCoefficient
-	{
-		float Abscissa;
-		float Weight;
-	};
-
-	static const FLegendreGaussCoefficient LegendreGaussCoefficients[] =
-	{
-		{ 0.0f, 0.5688889f },
-		{ -0.5384693f, 0.47862867f },
-		{ 0.5384693f, 0.47862867f },
-		{ -0.90617985f, 0.23692688f },
-		{ 0.90617985f, 0.23692688f }
-	};
-
-	const auto& P0 = GetSpline().GetValue(Index);
-	const auto& T0 = GetSpline().GetTangentOut(Index);
-	const auto& P1 = GetSpline().GetValue(Index == LastPoint ? 0 : Index + 1);
-	const auto& T1 = GetSpline().GetTangentIn(Index == LastPoint ? 0 : Index + 1);
-
-	// Special cases for linear
-	if (GetSpline().GetTangentModes()[Index] == UE::Geometry::Spline::ETangentMode::Linear ||
-		GetSpline().GetTangentModes()[Index] == UE::Geometry::Spline::ETangentMode::LegacyLinear)
-	{
-		return ((P1 - P0) * Scale3D).Size() * Param;
-	}
-
-	// Cache the coefficients to be fed into the function to calculate the spline derivative at each sample point as they are constant.
-	const FVector Coeff1 = ((P0 - P1) * 2.0f + T0 + T1) * 3.0f;
-	const FVector Coeff2 = (P1 - P0) * 6.0f - T0 * 4.0f - T1 * 2.0f;
-	const FVector Coeff3 = T0;
-
-	const float HalfParam = Param * 0.5f;
-
-	float Length = 0.0f;
-	for (const auto& LegendreGaussCoefficient : LegendreGaussCoefficients)
-	{
-		// Calculate derivative at each Legendre-Gauss sample, and perform a weighted sum
-		const float Alpha = HalfParam * (1.0f + LegendreGaussCoefficient.Abscissa);
-		const FVector Derivative = ((Coeff1 * Alpha + Coeff2) * Alpha + Coeff3) * Scale3D;
-		Length += Derivative.Size() * LegendreGaussCoefficient.Weight;
-	}
-	Length *= HalfParam;
-
-	return Length;
-}
-
-float FNewSpline::GetSplineLength()
-{	
-	// Evaluate the ParamToLength channel at the very end of the spline
-	const float MaxParameter = static_cast<float>(GetNumberOfSegments());
-    return GetDistanceAtParameter(MaxParameter);
-}
-
-void FNewSpline::SetClosedLoop(bool bInClosedLoop)
-{
-	SetClosedLoop(bInClosedLoop, true);
-}
-
-void FNewSpline::SetClosedLoop(bool bInClosedLoop, bool bUpdateSpline)
-{
-	if (bInClosedLoop == GetSpline().IsClosedLoop())
-	{
-		return;
-	}
-	TRACE_CPUPROFILER_EVENT_SCOPE(FNewSpline::SetClosedLoop);
-
-	MarkLegacyCurvesDirty();
-
-	GetSpline().SetClosedLoop(bInClosedLoop);
-	
-	if (UE::Geometry::Spline::TTangentBezierSpline<FVector>* ScaleChild = GetTypedAttributeChannel<UE::Geometry::Spline::TTangentBezierSpline<FVector>>(ScaleAttrName))
-	{
-		ScaleChild->SetClosedLoop(bInClosedLoop);
-		ScaleChild->SetKnotVector(GetSpline().GetKnotVector());
-	}
-		
-	if (UE::Geometry::Spline::TTangentBezierSpline<FQuat>* RotChild = GetTypedAttributeChannel<UE::Geometry::Spline::TTangentBezierSpline<FQuat>>(RotationAttrName))
-	{
-		RotChild->SetClosedLoop(bInClosedLoop);
-		RotChild->SetKnotVector(GetSpline().GetKnotVector());
-	}
-	
-	if (bUpdateSpline)
-	{
-		UpdateSpline();
-	}
-}
-
-void FNewSpline::Reset()
-{
-	Clear();
-	GetSpline().SetTangentModes(TArray<UE::Geometry::Spline::ETangentMode>());
-	PositionCurve.Reset();
-	RotationCurve.Reset();
-	ScaleCurve.Reset();
-}
-
-void FNewSpline::ResetRotation()
-{
-    ClearAttributeChannel(RotationAttrName);
-
-    for (int32 Idx = 0; Idx < GetSpline().GetNumPoints(); Idx++)
-    {    		
-    	AddAttributeValue<FQuat>(RotationAttrName, FQuat::Identity, static_cast<float>(Idx));
-    }
-
-	SetAttributeChannelRange(RotationAttrName, FInterval1f(0.f, 1.f), EMappingRangeSpace::Normalized);
-	if (UE::Geometry::Spline::TTangentBezierSpline<FQuat>* RotSpline = GetTypedAttributeChannel<UE::Geometry::Spline::TTangentBezierSpline<FQuat>>(RotationAttrName))
-	{
-		RotSpline->SetKnotVector(GetSpline().GetKnotVector());
-	}
-}
-
-void FNewSpline::ResetScale()
-{
-    ClearAttributeChannel(ScaleAttrName);
-	
-    for (int32 Idx = 0; Idx < GetSpline().GetNumPoints(); Idx++)
-    {    		
-        AddAttributeValue<FVector>(ScaleAttrName, FVector::OneVector, static_cast<float>(Idx));
-    }
-
-	SetAttributeChannelRange(ScaleAttrName, FInterval1f(0.f, 1.f), EMappingRangeSpace::Normalized);
-	if (UE::Geometry::Spline::TTangentBezierSpline<FVector>* ScaleSpline = GetTypedAttributeChannel<UE::Geometry::Spline::TTangentBezierSpline<FVector>>(ScaleAttrName))
-	{
-		ScaleSpline->SetKnotVector(GetSpline().GetKnotVector());
-	}
-}
-
-void FNewSpline::AddPoint(const FSplinePoint& Point)
-{
-	ValidateRotScale();
-	ON_SCOPE_EXIT
-	{
-		SetAttributeChannelRange(ScaleAttrName, FInterval1f(0.f, 1.f), EMappingRangeSpace::Normalized);
-		SetAttributeChannelRange(RotationAttrName, FInterval1f(0.f, 1.f), EMappingRangeSpace::Normalized);
-		ValidateRotScale();
-	};
-	
-	FTangentBezierControlPoint ControlPoint = ConvertToTangentBezierControlPoint(Point);
-	const int32 NumPoints = GetSpline().GetNumPoints();
-    
-	// Convert InputKey to parameter value
-	int32 Index = FMath::FloorToInt(Point.InputKey);
-	float Fraction = Point.InputKey - Index;
-	
-	if (NumPoints == 0 || Index >= NumPoints)
-	{
-		GetSpline().AppendPoint(ControlPoint);
-
-		if (UE::Geometry::Spline::TTangentBezierSpline<FVector>* ScaleChild = GetTypedAttributeChannel<UE::Geometry::Spline::TTangentBezierSpline<FVector>>(ScaleAttrName))
-		{
-			ScaleChild->AppendPoint(UE::Geometry::Spline::TTangentBezierSpline<FVector>::FTangentBezierControlPoint(Point.Scale));
-			ScaleChild->SetKnotVector(GetSpline().GetKnotVector());
-		}
-		
-		if (UE::Geometry::Spline::TTangentBezierSpline<FQuat>* RotChild = GetTypedAttributeChannel<UE::Geometry::Spline::TTangentBezierSpline<FQuat>>(RotationAttrName))
-		{
-			RotChild->AppendPoint(UE::Geometry::Spline::TTangentBezierSpline<FQuat>::FTangentBezierControlPoint(Point.Rotation.Quaternion()));
-			RotChild->SetKnotVector(GetSpline().GetKnotVector());
-		}
-		
-		MarkLegacyCurvesDirty();
-		return;
-	}
-	else if (Index == 0)
-    {
-    	GetSpline().PrependPoint(ControlPoint);
-    	
-		if (UE::Geometry::Spline::TTangentBezierSpline<FVector>* ScaleChild = GetTypedAttributeChannel<UE::Geometry::Spline::TTangentBezierSpline<FVector>>(ScaleAttrName))
-		{
-    		ScaleChild->PrependPoint(UE::Geometry::Spline::TTangentBezierSpline<FVector>::FTangentBezierControlPoint(Point.Scale));
-			ScaleChild->SetKnotVector(GetSpline().GetKnotVector());
-		}
-    
-		if (UE::Geometry::Spline::TTangentBezierSpline<FQuat>* RotChild = GetTypedAttributeChannel<UE::Geometry::Spline::TTangentBezierSpline<FQuat>>(RotationAttrName))
-		{
-    		RotChild->PrependPoint(UE::Geometry::Spline::TTangentBezierSpline<FQuat>::FTangentBezierControlPoint(Point.Rotation.Quaternion()));
-			RotChild->SetKnotVector(GetSpline().GetKnotVector());
-		}
-    	
-    	MarkLegacyCurvesDirty();
-    	return;
-    }
-	
-	// Insert at calculated parameter
-	int32 InsertedIndex = GetSpline().InsertPointAtSegmentParam(Index, Fraction, ControlPoint);
-	float NewPointInternalParam = GetSpline().GetParameter(InsertedIndex);
-	
-	float ScaleParam = MapParameterToChildSpace(ScaleAttrName, NewPointInternalParam);
-	if (UE::Geometry::Spline::TTangentBezierSpline<FVector>* ScaleChild = GetTypedAttributeChannel<UE::Geometry::Spline::TTangentBezierSpline<FVector>>(ScaleAttrName))
-	{
-		ScaleChild->InsertPointAtGlobalParam(ScaleParam, UE::Geometry::Spline::TTangentBezierSpline<FVector>::FTangentBezierControlPoint(Point.Scale));
-		ScaleChild->SetKnotVector(GetSpline().GetKnotVector());
-	}
-	
-	float RotParam = MapParameterToChildSpace(RotationAttrName, NewPointInternalParam);
-	if (UE::Geometry::Spline::TTangentBezierSpline<FQuat>* RotChild = GetTypedAttributeChannel<UE::Geometry::Spline::TTangentBezierSpline<FQuat>>(RotationAttrName))
-	{
-		RotChild->InsertPointAtGlobalParam(RotParam, UE::Geometry::Spline::TTangentBezierSpline<FQuat>::FTangentBezierControlPoint(Point.Rotation.Quaternion()));
-		RotChild->SetKnotVector(GetSpline().GetKnotVector());
-	}
-	
-	MarkLegacyCurvesDirty();
-}
-
-void FNewSpline::InsertPoint(const FSplinePoint& Point, int32 Index)
-{
-	ValidateRotScale();
-	ON_SCOPE_EXIT
-	{
-		SetAttributeChannelRange(ScaleAttrName, FInterval1f(0.f, 1.f), EMappingRangeSpace::Normalized);
-		SetAttributeChannelRange(RotationAttrName, FInterval1f(0.f, 1.f), EMappingRangeSpace::Normalized);
-		ValidateRotScale();
-	};
-	
-	FTangentBezierControlPoint ControlPoint = ConvertToTangentBezierControlPoint(Point);
-	const int32 NumPoints = GetSpline().GetNumPoints();
-    
-	// Special case for empty spline
-	if (NumPoints == 0 || Index >= NumPoints)
-	{
-		GetSpline().AppendPoint(ControlPoint);
-		
-		if (UE::Geometry::Spline::TTangentBezierSpline<FVector>* ScaleChild = GetTypedAttributeChannel<UE::Geometry::Spline::TTangentBezierSpline<FVector>>(ScaleAttrName))
-		{
-			ScaleChild->AppendPoint(UE::Geometry::Spline::TTangentBezierSpline<FVector>::FTangentBezierControlPoint(Point.Scale));
-			ScaleChild->SetKnotVector(GetSpline().GetKnotVector());
-		}
-		
-		if (UE::Geometry::Spline::TTangentBezierSpline<FQuat>* RotChild = GetTypedAttributeChannel<UE::Geometry::Spline::TTangentBezierSpline<FQuat>>(RotationAttrName))
-		{
-			RotChild->AppendPoint(UE::Geometry::Spline::TTangentBezierSpline<FQuat>::FTangentBezierControlPoint(Point.Rotation.Quaternion()));
-			RotChild->SetKnotVector(GetSpline().GetKnotVector());
-		}
-		
-		MarkLegacyCurvesDirty();
-		return;
-	}
-	else if (Index == 0)
-	{
-		GetSpline().PrependPoint(ControlPoint);
-		
-		if (UE::Geometry::Spline::TTangentBezierSpline<FVector>* ScaleChild = GetTypedAttributeChannel<UE::Geometry::Spline::TTangentBezierSpline<FVector>>(ScaleAttrName))
-		{
-			ScaleChild->PrependPoint(UE::Geometry::Spline::TTangentBezierSpline<FVector>::FTangentBezierControlPoint(Point.Scale));
-			ScaleChild->SetKnotVector(GetSpline().GetKnotVector());
-		}
-	
-		if (UE::Geometry::Spline::TTangentBezierSpline<FQuat>* RotChild = GetTypedAttributeChannel<UE::Geometry::Spline::TTangentBezierSpline<FQuat>>(RotationAttrName))
-		{
-			RotChild->PrependPoint(UE::Geometry::Spline::TTangentBezierSpline<FQuat>::FTangentBezierControlPoint(Point.Rotation.Quaternion()));
-			RotChild->SetKnotVector(GetSpline().GetKnotVector());
-		}
-		
-		MarkLegacyCurvesDirty();
-		return;
-	}
-	
-	Index = FMath::Clamp(Index, 0, GetSpline().GetNumPoints());
-	
-	int32 InsertedIndex = GetSpline().InsertPointAtPosition(Index, ControlPoint);
-	float NewPointInternalParam = GetSpline().GetParameter(InsertedIndex);
-	
-	float ScaleParam = MapParameterToChildSpace(ScaleAttrName, NewPointInternalParam);
-	if (UE::Geometry::Spline::TTangentBezierSpline<FVector>* ScaleChild = GetTypedAttributeChannel<UE::Geometry::Spline::TTangentBezierSpline<FVector>>(ScaleAttrName))
-	{
-		ScaleChild->InsertPointAtGlobalParam(ScaleParam, UE::Geometry::Spline::TTangentBezierSpline<FVector>::FTangentBezierControlPoint(Point.Scale));
-		ScaleChild->SetKnotVector(GetSpline().GetKnotVector());
-	}
-	
-	float RotParam = MapParameterToChildSpace(RotationAttrName, NewPointInternalParam);
-	if (UE::Geometry::Spline::TTangentBezierSpline<FQuat>* RotChild = GetTypedAttributeChannel<UE::Geometry::Spline::TTangentBezierSpline<FQuat>>(RotationAttrName))
-	{
-		RotChild->InsertPointAtGlobalParam(RotParam, UE::Geometry::Spline::TTangentBezierSpline<FQuat>::FTangentBezierControlPoint(Point.Rotation.Quaternion()));
-		RotChild->SetKnotVector(GetSpline().GetKnotVector());
-	}
-	
-	MarkLegacyCurvesDirty();
-}
-
-FSplinePoint FNewSpline::GetPoint(int32 Index) const
-{
-    FSplinePoint Point;
-    
-    if (Index < 0 || Index >= GetSpline().GetNumPoints())
-    {
-    	return Point;
-    }
-
-	Point.InputKey = GetParameterAtIndex(Index);
-    Point.Position = GetSpline().GetValue(Index);
-	Point.ArriveTangent = GetSpline().GetTangentIn(Index);
-    Point.LeaveTangent = GetSpline().GetTangentOut(Index);
-    Point.Rotation = GetAttributeValue<FQuat>(RotationAttrName, Index).Rotator();
-    Point.Scale = GetAttributeValue<FVector>(ScaleAttrName, Index);
-	EInterpCurveMode InterpCurveMode = ConvertTangentModeToInterpCurveMode(GetSpline().GetTangentModes()[Index]);
-    Point.Type = ConvertInterpCurveModeToSplinePointType(InterpCurveMode);
-
-    return Point;
-}
-
-/** Removes a point from the spline */
-void FNewSpline::RemovePoint(int32 Index)
-{
-	ValidateRotScale();
-	ON_SCOPE_EXIT
-	{
-		SetAttributeChannelRange(ScaleAttrName, FInterval1f(0.f, 1.f), EMappingRangeSpace::Normalized);
-		SetAttributeChannelRange(RotationAttrName, FInterval1f(0.f, 1.f), EMappingRangeSpace::Normalized);
-		
-		ValidateRotScale();
-	};
-	
-	if (Index < 0 || Index >= GetSpline().GetNumPoints())
-	{
-		return;
-	}
-	
-	GetSpline().RemovePoint(Index);
-
-	if (UE::Geometry::Spline::TTangentBezierSpline<FVector>* ScaleChild = GetTypedAttributeChannel<UE::Geometry::Spline::TTangentBezierSpline<FVector>>(ScaleAttrName))
-	{
-		ScaleChild->RemovePoint(Index);
-		ScaleChild->SetKnotVector(GetSpline().GetKnotVector());
-	}
-	
-	if (UE::Geometry::Spline::TTangentBezierSpline<FQuat>* RotChild = GetTypedAttributeChannel<UE::Geometry::Spline::TTangentBezierSpline<FQuat>>(RotationAttrName))
-	{
-		RotChild->RemovePoint(Index);
-		RotChild->SetKnotVector(GetSpline().GetKnotVector());
-	}
-	
-    MarkLegacyCurvesDirty();
-}
-
-void FNewSpline::SetLocation(int32 Index, const FVector& InLocation)
-{
-    if (Index < 0 || Index >= GetSpline().GetNumPoints())
-    {
-    	return;
-    }
-	
-	GetSpline().ModifyPoint(Index, {InLocation, GetSpline().GetTangentIn(Index), GetSpline().GetTangentOut(Index), GetSpline().GetTangentMode(Index)});
-	
-	MarkReparamTableDirty();
-    MarkLegacyCurvesDirty();
-}
-
-FVector FNewSpline::GetLocation(const int32 Index) const
-{
-	if (Index < 0 || Index >= GetSpline().GetNumPoints())
-    {
-    	return FVector();
-    }
-
-    return GetSpline().GetValue(Index);
-}
-
-void FNewSpline::SetInTangent(const int32 Index, const FVector& InTangent)
-{
-	if (Index < 0 || Index >= GetSpline().GetNumPoints())
-    {
-    	return;
-    }
-
-	GetSpline().SetPointTangentMode(Index, UE::Geometry::Spline::ETangentMode::User);
-    GetSpline().SetTangentIn(Index, InTangent);
-
-	MarkReparamTableDirty();
-    MarkLegacyCurvesDirty();
-}
-
-FVector FNewSpline::GetInTangent(const int32 Index) const
-{
-	if (Index < 0 || Index >= GetSpline().GetNumPoints())
-    {
-    	return FVector();
-    }
-
-    return GetSpline().GetTangentIn(Index);
-}
-
-void FNewSpline::SetOutTangent(const int32 Index, const FVector& OutTangent)
-{
-	if (Index < 0 || Index >= GetSpline().GetNumPoints())
-    {
-    	return;
-    }
-	
-	GetSpline().SetPointTangentMode(Index, UE::Geometry::Spline::ETangentMode::User);
-	GetSpline().SetTangentOut(Index, OutTangent);
-
-	MarkReparamTableDirty();
-    MarkLegacyCurvesDirty();
-}
-
-FVector FNewSpline::GetOutTangent(const int32 Index) const
-{
-	if (Index < 0 || Index >= GetSpline().GetNumPoints())
-    {
-    	return FVector();
-    }
-
-    return GetSpline().GetTangentOut(Index);
-}
-
-void FNewSpline::SetRotation(int32 Index, const FQuat& InRotation)
-{
-	if (Index < 0 || Index >= GetSpline().GetNumPoints())
-    {
-    	return;
-    }
-
-	ValidateRotScale();
-
-	SetAttributeValue<FQuat>(RotationAttrName, InRotation, Index);
-	
-    MarkLegacyCurvesDirty();
-}
-
-FQuat FNewSpline::GetRotation(const int32 Index) const
-{
-	if (Index < 0 || Index >= GetSpline().GetNumPoints())
-    {
-    	return FQuat::Identity;
-    }
-	
-    return GetAttributeValue<FQuat>(RotationAttrName, Index);
-}
-
-void FNewSpline::SetScale(int32 Index, const FVector& InScale)
-{
-	if (Index < 0 || Index >= GetSpline().GetNumPoints())
-    {
-    	return;
-    }
-
-	ValidateRotScale();
-
-	SetAttributeValue<FVector>(ScaleAttrName, InScale, Index);
-    
-    MarkLegacyCurvesDirty();
-}
-
-FVector FNewSpline::GetScale(const int32 Index) const
-{
-    if (Index < 0 || Index >= GetSpline().GetNumPoints())
-    {
-    	return FVector::OneVector;
-    }
-
-    return GetAttributeValue<FVector>(ScaleAttrName, Index);
-}
-
-void FNewSpline::SetSplinePointType(int32 Index, EInterpCurveMode Type)
-{
-    if (Index < 0 || Index >= GetSpline().GetNumPoints())
-    {
-    	return;
-    }
-	
-	GetSpline().SetPointTangentMode(Index, ConvertInterpCurveModeToTangentMode(Type));
-
-	MarkReparamTableDirty();
-    MarkLegacyCurvesDirty();
-}
-
-EInterpCurveMode FNewSpline::GetSplinePointType(int32 Index) const
-{
-    if (Index < 0 || Index >= GetSpline().GetNumPoints())
-    {
-    	return CIM_Unknown;
-    }
-	UE::Geometry::Spline::ETangentMode TangentMode = GetSpline().GetTangentMode(Index);
-    return ConvertTangentModeToInterpCurveMode(TangentMode);
-}
-
-float FNewSpline::GetParameterAtIndex(int32 Index) const
-{
-    if (Index < 0 || (!IsClosedLoop() && Index >= GetSpline().GetNumPoints()) || (IsClosedLoop() && Index > GetSpline().GetNumPoints()))
-    {
-    	return 0.f;
-    }
-
-	return FromInternalSplineSpace(GetSpline().GetParameter(Index));
-}
-
-float FNewSpline::GetParameterAtDistance(float Distance)
-{
-	UpdateReparamTable();
-	
-	UE::TReadScopeLock Lock(ReparamTableRWLock);
-	
-	return ReparamTable.Eval(Distance);
-}
-
-float FNewSpline::GetDistanceAtParameter(float Parameter)
-{
-	if (!GetSegmentSpace().Contains(Parameter))
-	{
-		return 0.f;
-	}
-	
-	UpdateReparamTable();
-	
-	UE::TReadScopeLock Lock(ReparamTableRWLock);
-
-	const float ParameterMax = static_cast<float>(GetNumberOfSegments());
-	const float Key = (Parameter / ParameterMax) * (ReparamTable.Points.Num() - 1);
-	const int32 LowerKey = FMath::FloorToInt32(Key);
-	ensureAlways(LowerKey >= 0 && LowerKey < ReparamTable.Points.Num());
-	const int32 UpperKey = FMath::CeilToInt32(Key);
-	ensureAlways(UpperKey >= 0 && UpperKey < ReparamTable.Points.Num());
-	const float Alpha = FMath::Frac(Key);
-	return FMath::Lerp(ReparamTable.Points[LowerKey].InVal, ReparamTable.Points[UpperKey].InVal, Alpha);
-}
-
-float FNewSpline::FindNearest(const FVector& InLocation, float& OutSquaredDistance) const
-{
-    if (UE::Spline::GFallbackFindNearest)
-    {
-    	return PositionCurve.FindNearest(InLocation, OutSquaredDistance);
-    }
-    else
-    {
-	    return FromInternalSplineSpace(GetSpline().FindNearest(InLocation, OutSquaredDistance));
-    }
-}
-
-float FNewSpline::FindNearestOnSegment(const FVector& InLocation, int32 SegmentIndex, float& OutSquaredDistance) const
-{
-    if (UE::Spline::GFallbackFindNearest)
-    {
-        if (!PositionCurve.Points.IsValidIndex(SegmentIndex))
-        {
-            return 0.f;
-        }
-    
-        return PositionCurve.FindNearestOnSegment(InLocation, SegmentIndex, OutSquaredDistance);
-    }
-    else
-    {
-	   return FromInternalSplineSpace(GetSpline().FindNearestOnSegment(InLocation, SegmentIndex, OutSquaredDistance));
-    }
-}
-
-FVector FNewSpline::EvaluatePosition(float Parameter) const
-{
-	if (UE::Spline::GUseLegacyPositionEvaluation)
-	{
-		RebuildLegacyCurves();
-		return PositionCurve.Eval(Parameter);
-	}
-	else
-	{
-		return Evaluate(ToInternalSplineSpace(Parameter));
-	}
-}
-
-FVector FNewSpline::EvaluateDerivative(float Parameter) const
-{
-	Parameter = ToInternalSplineSpace(Parameter);
-    return GetSpline().GetTangent(Parameter);
-}
-
-FQuat FNewSpline::EvaluateRotation(float Parameter) const
-{
-	if (UE::Spline::GUseLegacyRotationEvaluation)
-	{
-		RebuildLegacyCurves();
-		return RotationCurve.Eval(Parameter);
-	}
-	else
-	{
-		return EvaluateAttribute<FQuat>(RotationAttrName, ToInternalSplineSpace(Parameter));
-	}
-}
-
-FVector FNewSpline::EvaluateScale(float Parameter) const
-{
-	if (UE::Spline::GUseLegacyScaleEvaluation)
-	{
-		RebuildLegacyCurves();
-		return ScaleCurve.Eval(Parameter);
-	}
-	else
-	{
-		return EvaluateAttribute<FVector>(ScaleAttrName, ToInternalSplineSpace(Parameter));
-	}
-}
-
-void FNewSpline::UpdateSpline()
-{
-	TRACE_CPUPROFILER_EVENT_SCOPE(FNewSpline::UpdateSpline);
-
-	// Update our parameterization, forward it to the rotation and scale channels, and then recompute tangents.
-	Reparameterize();
-
-    MarkReparamTableDirty();
-
-	// Updates legacy curves based on points, never evaluates internal spline.
-    RebuildLegacyCurves();					
-}
-
-void FNewSpline::UpdateSpline(const FSpline::FUpdateSplineParams& Params)
-{
-	TRACE_CPUPROFILER_EVENT_SCOPE(FNewSpline::UpdateSpline);
-
-	// todo: handled unhandled params
-   
-	if (UE::Spline::GSkipDoubleUpdate)
-	{
-		// Avoid updating the spline twice.
-		constexpr bool bUpdateSpline = false;
-		SetClosedLoop(Params.bClosedLoop, bUpdateSpline);
-	}
-	else
-	{
-		SetClosedLoop(Params.bClosedLoop);
-	}
-
-	GetSpline().SetStationaryEndpoints(Params.bStationaryEndpoints);
-    ReparamStepsPerSegment = Params.ReparamStepsPerSegment;
-
-    UpdateSpline();
-}
-
-void FNewSpline::Reparameterize()
-{
-	TRACE_CPUPROFILER_EVENT_SCOPE(FNewSpline::Reparameterize);
-
-	// Caching segment space attribute parameters so we can restore them (necessary because we store the mapping in parent space which are about to be invalidated).
-	TMap<FName, TPair<FInterval1f, TArray<float>>> CachedChannelParameterizations;
-	for (const FName& ChannelName : GetAttributeChannelNamesByValueType<float>())
-	{
-		if (ChannelName == RotationAttrName || ChannelName == ScaleAttrName || NumAttributeValues<float>(ChannelName) == 0)
-		{
-			continue;
-		}
-
-		ensureAlwaysMsgf(GetAttributeChannelRangeSpace(ChannelName) == EMappingRangeSpace::Parent, TEXT("Attribute channel mapping range space is not Parent!"));
-
-		FInterval1f ChannelRangeInParentSpace = GetAttributeChannelRange(ChannelName);
-		FInterval1f ChannelRangeInSegmentSpace = FInterval1f(FromInternalSplineSpace(ChannelRangeInParentSpace.Min), FromInternalSplineSpace(ChannelRangeInParentSpace.Max));
-
-		CachedChannelParameterizations.FindOrAdd(ChannelName) = TPair<FInterval1f, TArray<float>>(ChannelRangeInSegmentSpace, CacheChannelKnots<float>(ChannelName));
-
-		for (float& InternalKnot : CachedChannelParameterizations[ChannelName].Value)
-		{
-			InternalKnot = FromInternalSplineSpace(InternalKnot);
-		}
-	}
-
-	GetSpline().Reparameterize(GetParameterizationPolicy());
-	GetSpline().UpdateTangents();
-
-	for (const FName& ChannelName : GetAttributeChannelNamesByValueType<float>())
-	{
-		if (ChannelName == RotationAttrName || ChannelName == ScaleAttrName || NumAttributeValues<float>(ChannelName) == 0)
-		{
-			continue;
-		}
-
-		const FInterval1f& ChannelRangeInSegmentSpace = CachedChannelParameterizations[ChannelName].Key;
-		const FInterval1f ChannelRangeInParentSpace = FInterval1f(ToInternalSplineSpace(ChannelRangeInSegmentSpace.Min), ToInternalSplineSpace(ChannelRangeInSegmentSpace.Max));
-
-		SetAttributeChannelRange(ChannelName, ChannelRangeInParentSpace, EMappingRangeSpace::Parent);
-
-		for (float& InternalKnot : CachedChannelParameterizations[ChannelName].Value)
-		{
-			InternalKnot = ToInternalSplineSpace(InternalKnot);
-		}
-
-		RestoreChannelKnots(ChannelName, CachedChannelParameterizations[ChannelName].Value);
-		if (UE::Geometry::Spline::TTangentBezierSpline<float>* Channel = GetTypedAttributeChannel<UE::Geometry::Spline::TTangentBezierSpline<float>>(ChannelName))
-		{
-			Channel->UpdateTangents();
-		}
-	}
-
-	if (UE::Geometry::Spline::TTangentBezierSpline<FVector>* ScaleSpline = GetTypedAttributeChannel<UE::Geometry::Spline::TTangentBezierSpline<FVector>>(ScaleAttrName))
-	{
-		ScaleSpline->SetKnotVector(GetSpline().GetKnotVector());
-		ScaleSpline->UpdateTangents();
-	}
-
-	if (UE::Geometry::Spline::TTangentBezierSpline<FQuat>* RotSpline = GetTypedAttributeChannel<UE::Geometry::Spline::TTangentBezierSpline<FQuat>>(RotationAttrName))
-	{
-		RotSpline->SetKnotVector(GetSpline().GetKnotVector());
-		RotSpline->UpdateTangents();
-	}
-}
-
-float FNewSpline::FromNormalizedSpace(float Parameter) const
-{
-	return GetNumberOfSegments() * Parameter;
-}
-
-float FNewSpline::ToNormalizedSpace(float Parameter) const
-{
-	return Parameter / GetNumberOfSegments();
-}
-
-float FNewSpline::FromInternalSplineSpace(float Parameter) const
-{
-	float OutLocalParam;
-	const int32 SegmentIndex = GetSpline().FindSegmentIndex(Parameter, OutLocalParam);
-	return static_cast<float>(SegmentIndex) + OutLocalParam;
-}
-float FNewSpline::ToInternalSplineSpace(float Parameter) const
-{
-	const int32 NumSegments = GetNumberOfSegments();
-	float ClampedParam = FMath::Clamp(Parameter, 0.0f, static_cast<float>(NumSegments));
-	const int32 SegmentIndex = FMath::Min(FMath::FloorToInt(ClampedParam), NumSegments - 1);
-	const float LocalT = ClampedParam - static_cast<float>(SegmentIndex); // Parameter [0,1] within segment
-	return GetSegmentParameterRange(SegmentIndex).Interpolate(LocalT);
-}
-
-void FNewSpline::MarkReparamTableDirty()
-{
-	ReparamTableNextVersion.Increment();
-}
-
-void FNewSpline::UpdateReparamTable()
-{
-	const int32 NumSegments = GetNumberOfSegments();
-
-	auto IsReparamTableDirty = [this]()
-	{
-		return ReparamTableNextVersion.GetValue() != ReparamTableVersion;
-	};
-	
-	if (!IsReparamTableDirty() || ReparamStepsPerSegment == 0 || NumSegments == 0)
-	{
-		return;
-	}
-
-	UE::TWriteScopeLock Lock(ReparamTableRWLock);
-
-	if (!IsReparamTableDirty())
-	{
-		return;
-	}
-
-	// We can't rely on the next version to not change during the update.
-	const uint32 CachedNextVersion = ReparamTableNextVersion.GetValue();
-
-	ReparamTable.Points.Reset(NumSegments * ReparamStepsPerSegment + 1);
-	float AccumulatedLength = 0.0f;
-	for (int32 SegmentIndex = 0; SegmentIndex < NumSegments; ++SegmentIndex)
-	{
-		ReparamTable.Points.Emplace(AccumulatedLength, SegmentIndex, 0.0f, 0.0f, CIM_Linear);
-		for (int32 Step = 1; Step < ReparamStepsPerSegment; ++Step)
-		{
-			const float Param = static_cast<float>(Step) / ReparamStepsPerSegment;
-			ReparamTable.Points.Emplace(GetSegmentLength(SegmentIndex, Param) + AccumulatedLength, SegmentIndex + Param, 0.0f, 0.0f, CIM_Linear);
-		}
-
-		AccumulatedLength += GetSegmentLength(SegmentIndex, 1.0f);
-	}
-	ReparamTable.Points.Emplace(AccumulatedLength, static_cast<float>(NumSegments), 0.0f, 0.0f, CIM_Linear);
-
-	ReparamTableVersion = CachedNextVersion;
-}
-
-void FNewSpline::MarkLegacyCurvesDirty()
-{
-	if (UE::Spline::GImmediatelyUpdateLegacyCurves)
-	{
-		RebuildLegacyCurves();
-	}
-	else
-	{
-		LegacyCurvesNextVersion.Increment();
-	}
-}
-
-void FNewSpline::RebuildLegacyCurves() const
-{
-	TRACE_CPUPROFILER_EVENT_SCOPE(FNewSpline::RebuildLegacyCurves);
-
-	// Curves are dirty if:
-	// 1) Current version is not the same as next version (if next version == version, it has not been dirtied)
-	// 2) UE::Spline::GImmediatelyUpdateLegacyCurves == true (legacy curves are updated any time they might be dirty)
-	auto AreLegacyCurvesDirty = [this]()
-	{
-		return LegacyCurvesNextVersion.GetValue() != LegacyCurvesVersion || UE::Spline::GImmediatelyUpdateLegacyCurves;
-	};
-	
-	if (!AreLegacyCurvesDirty())
-	{
-		return;
-	}
-
-	PositionCurve.Points.Reset(GetSpline().GetNumPoints());
-	RotationCurve.Points.Reset(GetSpline().GetNumPoints());
-	ScaleCurve.Points.Reset(GetSpline().GetNumPoints());
-
-    PositionCurve.bIsLooped = IsClosedLoop();
-    PositionCurve.LoopKeyOffset = 1.f;
-    RotationCurve.bIsLooped = IsClosedLoop();
-	RotationCurve.LoopKeyOffset = 1.f;
-    ScaleCurve.bIsLooped = IsClosedLoop();
-	ScaleCurve.LoopKeyOffset = 1.f;
-	
-	using FRotationChannel = UE::Geometry::Spline::TTangentBezierSpline<FQuat>;
-	using FScaleChannel = UE::Geometry::Spline::TTangentBezierSpline<FVector>;
-
-	const FRotationChannel* RotationChannel = GetTypedAttributeChannel<FRotationChannel>(RotationAttrName);
-	const FScaleChannel* ScaleChannel = GetTypedAttributeChannel<FScaleChannel>(ScaleAttrName);
-
-    for (int32 i = 0; i < GetSpline().GetNumPoints(); ++i)
-    {
-    	const float AttributeParameter = static_cast<float>(i);
-    	
-    	// Additionally populate interp curves here so we don't double loop
-    	PositionCurve.Points.Emplace(AttributeParameter, GetSpline().GetValue(i), GetSpline().GetTangentIn(i), GetSpline().GetTangentOut(i), ConvertTangentModeToInterpCurveMode(GetSpline().GetTangentModes()[i]));
-		RotationCurve.Points.Emplace(AttributeParameter, RotationChannel ? RotationChannel->GetValue(i) : FQuat::Identity);
-    	RotationCurve.Points.Last().InterpMode = CIM_CurveAuto;
-		ScaleCurve.Points.Emplace(AttributeParameter, ScaleChannel ? ScaleChannel->GetValue(i) : FVector::OneVector);
-    	ScaleCurve.Points.Last().InterpMode = CIM_CurveAuto;
-    }
-
-	Algo::SortBy(RotationCurve.Points, &FInterpCurvePoint<FQuat>::InVal);
-	Algo::SortBy(ScaleCurve.Points, &FInterpCurvePoint<FVector>::InVal);
-
-	PositionCurve.AutoSetTangents(0.f, false);
-	RotationCurve.AutoSetTangents(0.f, false);
-	ScaleCurve.AutoSetTangents(0.f, false);
-    
-    LegacyCurvesVersion = LegacyCurvesNextVersion.GetValue();
-}
-
-FNewSpline::FTangentBezierControlPoint FNewSpline::ConvertToTangentBezierControlPoint(const FSplinePoint& Point) const
-{
-	FTangentBezierControlPoint NewPoint;
-	NewPoint.Position = Point.Position;
-	NewPoint.TangentIn = Point.ArriveTangent;
-	NewPoint.TangentOut = Point.LeaveTangent;
-	NewPoint.TangentMode = ConvertInterpCurveModeToTangentMode(
-		ConvertSplinePointTypeToInterpCurveMode(Point.Type));
-	return NewPoint;
-}
-
-void FNewSpline::UpdatePointAttributes(const FSplinePoint& Point, int32 PointIndex)
-{
-	if (PointIndex < 0)
-	{
-		return;
-	}
-	
-	SetRotation(PointIndex, Point.Rotation.Quaternion());
-	SetScale(PointIndex, Point.Scale);
-        
-	MarkLegacyCurvesDirty();
-}
-
-float FNewSpline::ConvertIndexToInternalParameter(int32 Index, float Fraction) const
-{
-	// Handle empty spline
-	if (GetSpline().GetNumPoints() <= 1)
-		return 0.0f;
-    
-	// Clamp to valid range
-	Index = FMath::Clamp(Index, 0, GetSpline().GetNumPoints() - 1);
-    
-	// For exact index, just return parameter at that index
-	if (FMath::IsNearlyZero(Fraction))
-		return GetSpline().GetParameter(Index);
-    
-	// For fractional indices, interpolate
-	int32 NextIndex = FMath::Min(Index + 1, GetSpline().GetNumPoints() - 1);
-	float StartParam = GetSpline().GetParameter(Index);
-	float EndParam = GetSpline().GetParameter(NextIndex);
-    
-	return FMath::Lerp(StartParam, EndParam, Fraction);
-}
-
-int32 FNewSpline::ConvertInternalParameterToNearestPointIndex(float Parameter) const
-{
-	// Handle empty spline
-	if (GetSpline().GetNumPoints() <= 1)
-		return 0;
-    
-	// Find segment containing this parameter
-	for (int32 i = 0; i < GetSpline().GetNumPoints() - 1; i++)
-	{
-		float StartParam = GetSpline().GetParameter(i);
-		float EndParam = GetSpline().GetParameter(i + 1);
-        
-		if (Parameter >= StartParam && Parameter <= EndParam)
-		{
-			// Return closer index
-			float Fraction = (Parameter - StartParam) / (EndParam - StartParam);
-			return (Fraction <= 0.5f) ? i : (i + 1);
-		}
-	}
-    
-	// Out of range
-	return (Parameter < GetSpline().GetParameter(0)) ? 0 : (GetSpline().GetNumPoints() - 1);
-}
-
-void FNewSpline::ValidateRotScale() const
-{
-	if (!UE::Spline::GValidateRotScale)
-	{
-		return;
-	}
-	
-	// Validate child existence
-	UE::Geometry::Spline::TTangentBezierSpline<FQuat>* RotSpline = GetTypedAttributeChannel<UE::Geometry::Spline::TTangentBezierSpline<FQuat>>(RotationAttrName);
-	UE::Geometry::Spline::TTangentBezierSpline<FVector>* ScaleSpline = GetTypedAttributeChannel<UE::Geometry::Spline::TTangentBezierSpline<FVector>>(ScaleAttrName);
-	ensure (RotSpline && ScaleSpline);
-
-	if (RotSpline && ScaleSpline)
-	{
-		// Validate number of points
-		int32 NumControlPoints = GetSpline().GetNumPoints();
-		int32 NumRotPoints = RotSpline->GetNumPoints();
-		int32 NumScalePoints = ScaleSpline->GetNumPoints();
-		ensure(NumControlPoints == NumRotPoints && NumRotPoints == NumScalePoints);
-
-		// Validate control point -> attr point mapping
-		for (int32 i = 0; i < NumControlPoints; i++)
-		{
-			float InternalParam = GetSpline().GetParameter(i);
-		
-			float ExpectedRotParam = MapParameterToChildSpace(RotationAttrName, InternalParam);
-			float ActualRotParam = RotSpline->GetParameter(i);
-			ensure(FMath::IsNearlyEqual(ExpectedRotParam, ActualRotParam, UE_KINDA_SMALL_NUMBER));
-		
-			float ExpectedScaleParam = MapParameterToChildSpace(ScaleAttrName, InternalParam);
-			float ActualScaleParam = ScaleSpline->GetParameter(i);
-			ensure(FMath::IsNearlyEqual(ExpectedScaleParam, ActualScaleParam, UE_KINDA_SMALL_NUMBER));
-		}
-	}
-}
+#endif
+);
 
-template <typename AttrType>
-TArray<AttrType> FNewSpline::CacheChannelKnots(const FName& InName) const
+inline UE::Spline::ESplineType GetDefaultSplineType()
 {
-	using AttrSplineType = UE::Geometry::Spline::TTangentBezierSpline<AttrType>;
-
-	TArray<float> InternalParameters;
-
-	if (AttrSplineType* Channel = GetTypedAttributeChannel<AttrSplineType>(InName))
+	switch (GImplementation)
 	{
-		for (int InternalIdx = 1; InternalIdx < Channel->GetNumPoints() - 1; ++InternalIdx)
-		{
-			InternalParameters.Add(MapParameterFromChildSpace(InName, Channel->GetParameter(InternalIdx)));
-		}
+	default:	// fallthrough
+	case 0:		return UE::Spline::ESplineType::Unimplemented;
+	case 1:		return UE::Spline::ESplineType::LegacyTangent;
+	case 2:		return UE::Spline::ESplineType::Tangent;
 	}
-		
-	return InternalParameters;
 }
 
-template <typename AttrType>
-void FNewSpline::RestoreChannelKnots(const FName& InName, const TArray<AttrType>& InKnots)
-{
-	using AttrSplineType = UE::Geometry::Spline::TTangentBezierSpline<AttrType>;
-
-	if (InKnots.Num() == 0) return;
-
-	if (AttrSplineType* Channel = GetTypedAttributeChannel<AttrSplineType>(InName))
-	{
-		// Provided knots are for internal points, bail out if we don't have the correct number of knots.
-		if (ensureAlwaysMsgf(InKnots.Num() == Channel->GetNumPoints() - 2, TEXT("Invalid number of knots!")))
-		{
-			for (int InternalIdx = 1; InternalIdx < Channel->GetNumPoints() - 1; ++InternalIdx)
-			{
-				Channel->SetParameter(InternalIdx, MapParameterToChildSpace(InName, InKnots[InternalIdx - 1]));
-			}
-		}
-	}
 }
 
-/** FLegacySpline Implementation */
+using UE::Geometry::FInterval1f;
 
-FLegacySpline::FLegacySpline(const FNewSpline& Other)
-	: FLegacySpline()
-{
-	PositionCurve = Other.GetSplinePointsPosition();
-	RotationCurve = Other.GetSplinePointsRotation();
-	ScaleCurve = Other.GetSplinePointsScale();
+/** FSpline Implementation */
 
-	FSpline::FUpdateSplineParams InParams;
-	InParams.bClosedLoop = PositionCurve.bIsLooped;
-	UpdateSpline(InParams);
+FSpline::FSpline()
+{
+	SetCurrentImplementation(UE::Spline::GetDefaultSplineType());
 
-	PositionCurve.AutoSetTangents(0.f, false);
+#if WITH_EDITOR
+	PreviousImplementation = UE::Spline::ESplineType::Unimplemented;
+#endif
 }
 
-FLegacySpline::FLegacySpline(const FSplineCurves& InSpline)
-	: FLegacySpline()
+FSpline::FSpline(UE::Spline::ESplineType InType)
 {
-	PositionCurve = InSpline.Position;
-	RotationCurve = InSpline.Rotation;
-	ScaleCurve = InSpline.Scale;
-	ReparamTable = InSpline.ReparamTable;
+	SetCurrentImplementation(InType);
+
+#if WITH_EDITOR
+	PreviousImplementation = UE::Spline::ESplineType::Unimplemented;
+#endif
 }
 
-void FLegacySpline::AddPoint(const FSplinePoint& InPoint)
+FSpline::FSpline(const FSpline& Other)
 {
-	auto UpperBound = [this](float Value)
-	{
-		int32 Count = PositionCurve.Points.Num();
-		int32 First = 0;
-
-		while (Count > 0)
-		{
-			const int32 Middle = Count / 2;
-			if (Value >= PositionCurve.Points[First + Middle].InVal)
-			{
-				First += Middle + 1;
-				Count -= Middle + 1;
-			}
-			else
-			{
-				Count = Middle;
-			}
-		}
-
-		return First;
-	};
-	
-	int32 Index = UpperBound(InPoint.InputKey);
-	
-	PositionCurve.Points.Insert(FInterpCurvePoint<FVector>(
-		InPoint.InputKey,
-		InPoint.Position,
-		InPoint.ArriveTangent,
-		InPoint.LeaveTangent,
-		ConvertSplinePointTypeToInterpCurveMode(InPoint.Type)
-	), Index);
-
-	RotationCurve.Points.Insert(FInterpCurvePoint<FQuat>(
-		InPoint.InputKey,
-		InPoint.Rotation.Quaternion(),
-		FQuat::Identity,
-		FQuat::Identity,
-		CIM_CurveAuto
-	), Index);
-
-	ScaleCurve.Points.Insert(FInterpCurvePoint<FVector>(
-		InPoint.InputKey,
-		InPoint.Scale,
-		FVector::ZeroVector,
-		FVector::ZeroVector,
-		CIM_CurveAuto
-	), Index);
+	*this = Other;
 }
 
-void FLegacySpline::InsertPoint(const FSplinePoint& InPoint, int32 Index)
+FSpline& FSpline::operator=(const FSpline& Other)
 {
-	const float InKey = (Index == 0) ? 0.0f : GetParameterAtIndex(Index - 1) + 1.0f;
-	
-	PositionCurve.Points.Insert(FInterpCurvePoint<FVector>(
-		InKey,
-		InPoint.Position,
-		InPoint.ArriveTangent,
-		InPoint.LeaveTangent,
-		ConvertSplinePointTypeToInterpCurveMode(InPoint.Type)
-	), Index);
-
-	RotationCurve.Points.Insert(FInterpCurvePoint<FQuat>(
-		InKey,
-		InPoint.Rotation.Quaternion(),
-		FQuat::Identity,
-		FQuat::Identity,
-		CIM_CurveAuto
-	), Index);
-
-	ScaleCurve.Points.Insert(FInterpCurvePoint<FVector>(
-		InKey,
-		InPoint.Scale,
-		FVector::ZeroVector,
-		FVector::ZeroVector,
-		CIM_CurveAuto
-	), Index);
-
-	// Increment all point input values after the inserted element
-	for (int32 Idx = Index + 1; Idx < PositionCurve.Points.Num(); Idx++)
+	SetCurrentImplementation(Other.CurrentImplementation);
+
+	switch (CurrentImplementation)
 	{
-		PositionCurve.Points[Idx].InVal += 1.f;
-		RotationCurve.Points[Idx].InVal += 1.f;
-		ScaleCurve.Points[Idx].InVal += 1.f;
+	default:
+		break;
+
+	case UE::Spline::ESplineType::LegacyTangent:
+		*LegacySplineImpl = *Other.LegacySplineImpl;
+		break;
+
+	case UE::Spline::ESplineType::Tangent:
+		*Get<UE::Spline::ESplineType::Tangent>() = *Other.Get<UE::Spline::ESplineType::Tangent>();
+		break;
 	}
+		
+	return *this;
 }
 
-FSplinePoint FLegacySpline::GetPoint(const int32 Index) const
+FSpline& FSpline::operator=(const FSplineCurves& Other)
 {
-	FSplinePoint Point;
+	SetCurrentImplementation(UE::Spline::ESplineType::Tangent);
 
-	if (!PositionCurve.Points.IsValidIndex(Index))
-	{
-		return Point;
-	}
-	
-	Point.InputKey = PositionCurve.Points[Index].InVal;
-	Point.Position = PositionCurve.Points[Index].OutVal;
-	Point.ArriveTangent = PositionCurve.Points[Index].ArriveTangent;
-	Point.LeaveTangent = PositionCurve.Points[Index].LeaveTangent;
-	Point.Rotation = RotationCurve.Points[Index].OutVal.Rotator();
-	Point.Scale = ScaleCurve.Points[Index].OutVal;
-	Point.Type = ConvertInterpCurveModeToSplinePointType(PositionCurve.Points[Index].InterpMode);
+	*Get<UE::Spline::ESplineType::Tangent>() = FTangentSpline(Other);
 	
-	return Point;
+	return *this;
 }
 
-void FLegacySpline::RemovePoint(int32 Index)
+FVector FSpline::Evaluate(float Param) const
 {
-	if (!PositionCurve.Points.IsValidIndex(Index))
+	switch (CurrentImplementation)
 	{
-		return;
-	}
-	
-	PositionCurve.Points.RemoveAt(Index);
-	RotationCurve.Points.RemoveAt(Index);
-	ScaleCurve.Points.RemoveAt(Index);
+	default:
+		break;
 
-	while (Index < GetNumControlPoints())
-	{
-		PositionCurve.Points[Index].InVal -= 1.0f;
-		RotationCurve.Points[Index].InVal -= 1.0f;
-		ScaleCurve.Points[Index].InVal -= 1.0f;
-		Index++;
-	}
-}
+	case UE::Spline::ESplineType::LegacyTangent:
+		return LegacySplineImpl->PositionCurve.Eval(Param);
 
-void FLegacySpline::SetLocation(int32 Index, const FVector& InLocation)
-{
-	if (!PositionCurve.Points.IsValidIndex(Index))
-	{
-		return;
+	case UE::Spline::ESplineType::Tangent:
+		return Get<UE::Spline::ESplineType::Tangent>()->EvaluatePosition(Param);
 	}
-	
-	PositionCurve.Points[Index].OutVal = InLocation;
-}
 
-FVector FLegacySpline::GetLocation(const int32 Index) const
-{
-	return PositionCurve.Points.IsValidIndex(Index) ? PositionCurve.Points[Index].OutVal : FVector();
+	return FVector();
 }
 
-void FLegacySpline::SetInTangent(const int32 Index, const FVector& InTangent)
+FVector FSpline::EvaluateDerivative(float Param) const
 {
-	if (!PositionCurve.Points.IsValidIndex(Index))
+	switch (CurrentImplementation)
 	{
-		return;
-	}
-	
-	PositionCurve.Points[Index].ArriveTangent = InTangent;
-	PositionCurve.Points[Index].InterpMode = CIM_CurveUser;
-}
+	default:
+		break;
 
-FVector FLegacySpline::GetInTangent(const int32 Index) const
-{
-	return PositionCurve.Points.IsValidIndex(Index) ? PositionCurve.Points[Index].ArriveTangent : FVector();
-}
+	case UE::Spline::ESplineType::LegacyTangent:
+		return LegacySplineImpl->PositionCurve.EvalDerivative(Param);
 
-void FLegacySpline::SetOutTangent(const int32 Index, const FVector& OutTangent)
-{
-	if (!PositionCurve.Points.IsValidIndex(Index))
-	{
-		return;
+	case UE::Spline::ESplineType::Tangent:
+		return Get<UE::Spline::ESplineType::Tangent>()->EvaluateDerivative(Param);
 	}
-	
-	PositionCurve.Points[Index].LeaveTangent = OutTangent;
-	PositionCurve.Points[Index].InterpMode = CIM_CurveUser;
-}
 
-FVector FLegacySpline::GetOutTangent(const int32 Index) const
-{
-	return PositionCurve.Points.IsValidIndex(Index) ? PositionCurve.Points[Index].LeaveTangent : FVector();
+	return FVector();
 }
 
-void FLegacySpline::SetRotation(int32 Index, const FQuat& InRotation)
+int32 FSpline::GetNumberOfSegments() const
 {
-	if (!RotationCurve.Points.IsValidIndex(Index))
+	switch (CurrentImplementation)
 	{
-		return;
-	}
-	
-	RotationCurve.Points[Index].OutVal = InRotation;
-}
+	default:
+		return 0;
 
-FQuat FLegacySpline::GetRotation(const int32 Index) const
-{
-	return RotationCurve.Points.IsValidIndex(Index) ? RotationCurve.Points[Index].OutVal : FQuat();
-}
+	case UE::Spline::ESplineType::LegacyTangent:
+		return LegacySplineImpl->PositionCurve.bIsLooped ? LegacySplineImpl->PositionCurve.Points.Num() : FMath::Max(LegacySplineImpl->PositionCurve.Points.Num() - 1, 0);
 
-void FLegacySpline::SetScale(int32 Index, const FVector& InScale)
-{
-	if (!ScaleCurve.Points.IsValidIndex(Index))
-	{
-		return;
+	case UE::Spline::ESplineType::Tangent:
+		// fallthrough for all ISplineInterfaces
+		return SplineImpl->GetNumberOfSegments();
 	}
-	
-	ScaleCurve.Points[Index].OutVal = InScale;
 }
 
-FVector FLegacySpline::GetScale(const int32 Index) const
+void FSpline::Reset()
 {
-	return ScaleCurve.Points.IsValidIndex(Index) ? ScaleCurve.Points[Index].OutVal : FVector();
+	switch (CurrentImplementation)
+	{
+	default:
+		break;
+
+	case UE::Spline::ESplineType::LegacyTangent:
+		LegacySplineImpl->PositionCurve.Reset();
+		LegacySplineImpl->PositionCurve.ClearLoopKey();
+		LegacySplineImpl->RotationCurve.Reset();
+		LegacySplineImpl->RotationCurve.ClearLoopKey();
+		LegacySplineImpl->ScaleCurve.Reset();
+		LegacySplineImpl->ScaleCurve.ClearLoopKey();
+		LegacySplineImpl->ReparamTable.Reset();
+		LegacySplineImpl->ReparamTable.ClearLoopKey();
+		break;
+
+	case UE::Spline::ESplineType::Tangent:
+		// fallthrough for all ISplineInterfaces
+		SplineImpl->Clear();
+		break;
+	}
 }
 
-void FLegacySpline::SetSplinePointType(int32 Index, EInterpCurveMode Type)
+bool FSpline::operator==(const FSpline& Other) const
 {
-	if (!PositionCurve.Points.IsValidIndex(Index))
+	if (CurrentImplementation != Other.CurrentImplementation)
 	{
-		return;
+		return false;
 	}
-	
-	PositionCurve.Points[Index].InterpMode = Type;
-}
 
-EInterpCurveMode FLegacySpline::GetSplinePointType(const int32 Index) const
-{
-	return PositionCurve.Points.IsValidIndex(Index) ? PositionCurve.Points[Index].InterpMode.GetValue() : CIM_Unknown;
-}
+	switch (CurrentImplementation)
+	{
+	default:
+		return false;
 
-float FLegacySpline::GetParameterAtIndex(int32 Index) const
-{
-	return PositionCurve.Points.IsValidIndex(Index) ? PositionCurve.Points[Index].InVal : 0.f;
+	case UE::Spline::ESplineType::LegacyTangent:
+		return LegacySplineImpl->PositionCurve == Other.LegacySplineImpl->PositionCurve &&
+			LegacySplineImpl->RotationCurve == Other.LegacySplineImpl->RotationCurve &&
+			LegacySplineImpl->ScaleCurve == Other.LegacySplineImpl->ScaleCurve;
+
+	case UE::Spline::ESplineType::Tangent:
+		// fallthrough for all ISplineInterfaces
+		return SplineImpl->IsEqual(Other.SplineImpl.Get());
+
+	}
 }
 
-float FLegacySpline::GetParameterAtDistance(float Distance) const
+bool FSpline::operator!=(const FSpline& Other) const
 {
-	return ReparamTable.Eval(Distance);
+	return !(*this == Other);
 }
 
-float FLegacySpline::GetDistanceAtParameter(float Parameter) const
+bool FSpline::Serialize(FArchive& Ar)
 {
-	// this might be duplicating what an interp curve can already do...
+	Ar.UsingCustomVersion(FFortniteMainBranchObjectVersion::GUID);
 	
-	const float ParameterMax = PositionCurve.Points.Last().InVal;
-	const float Key = (Parameter / ParameterMax) * (ReparamTable.Points.Num() - 1);
-	const int32 LowerKey = FMath::FloorToInt32(Key);
-	ensureAlways(LowerKey >= 0 && LowerKey < ReparamTable.Points.Num());
-	const int32 UpperKey = FMath::CeilToInt32(Key);
-	ensureAlways(UpperKey >= 0 && UpperKey < ReparamTable.Points.Num());
-	const float Alpha = FMath::Frac(Key);
-	const float Distance = FMath::Lerp(ReparamTable.Points[LowerKey].InVal, ReparamTable.Points[UpperKey].InVal, Alpha);
-
-	return Distance;
-}
+	// Data format:
 
-FVector FLegacySpline::Evaluate(float Param) const
-{
-	return PositionCurve.Eval(Param);
-}
+	// Byte 1 - The data format, determine by CurrentImplementation at the time of last save
+	// Remaining N Bytes - Spline data (or empty, if byte 0 is 0). The format is determined by byte 1.
+	
+	if (Ar.IsLoading())
+	{
+		SerializeLoad(Ar);
+	}
+	else
+	{
+		SerializeSave(Ar);
+	}
 
-FVector FLegacySpline::EvaluateDerivative(float Param) const
-{
-	return PositionCurve.EvalDerivative(Param);
+	return true;
 }
 
-FQuat FLegacySpline::EvaluateRotation(float Param) const
+void FSpline::SerializeLoad(FArchive& Ar)
 {
-	return RotationCurve.Eval(Param);
-}
+	uint8 PreviousImpl;
+	Ar << PreviousImpl;
 
-FVector FLegacySpline::EvaluateScale(float Param) const
-{
-	return ScaleCurve.Eval(Param);
-}
+#if WITH_EDITOR
+	PreviousImplementation = static_cast<UE::Spline::ESplineType>(PreviousImpl);
+#endif
 
-float FLegacySpline::FindNearest(const FVector& InLocation, float& OutSquaredDist) const
-{
-	return PositionCurve.FindNearest(InLocation, OutSquaredDist);
-}
+	SetCurrentImplementation(static_cast<UE::Spline::ESplineType>(PreviousImpl));
 
-float FLegacySpline::FindNearestOnSegment(const FVector& InLocation, int32 SegmentIndex, float& OutSquaredDist) const
-{
-	if (!PositionCurve.Points.IsValidIndex(SegmentIndex))
+	switch (CurrentImplementation)
 	{
-		return 0.f;
+	default: break;
+
+	case UE::Spline::ESplineType::LegacyTangent: Ar << *LegacySplineImpl; break;
+		
+	case UE::Spline::ESplineType::Tangent:
+		// fallthrough for all ISplineInterfaces
+		Ar << *SplineImpl; break;
 	}
-	
-	return PositionCurve.FindNearestOnSegment(InLocation, SegmentIndex, OutSquaredDist);
 }
 
-bool FLegacySpline::Serialize(FArchive& Ar)
+void FSpline::SerializeSave(FArchive& Ar) const
 {
-	Ar.UsingCustomVersion(FFortniteMainBranchObjectVersion::GUID);
-	Ar << PositionCurve;
-	Ar << RotationCurve;
-	Ar << ScaleCurve;
-	Ar << ReparamTable;
+	uint8 CurrentImpl = static_cast<uint8>(CurrentImplementation);
+	Ar << CurrentImpl;
 
-	return true;
-}
+	switch (CurrentImplementation)
+	{
+	default: break;
 
-float FLegacySpline::GetSegmentLength(const int32 Index, const float Param, const FVector& Scale3D) const
-{
-	const int32 NumPoints = PositionCurve.Points.Num();
-	const int32 LastPoint = NumPoints - 1;
+	case UE::Spline::ESplineType::LegacyTangent: Ar << *LegacySplineImpl; break;
 
-	check(Index >= 0 && ((PositionCurve.bIsLooped && Index < NumPoints) || (!PositionCurve.bIsLooped && Index < LastPoint)));
-	check(Param >= 0.0f && Param <= 1.0f);
+	case UE::Spline::ESplineType::Tangent:
+		// fallthrough for all ISplineInterfaces
+		Ar << *SplineImpl; break;
+	}
+}
 
-	// Evaluate the length of a Hermite spline segment.
-	// This calculates the integral of |dP/dt| dt, where P(t) is the spline equation with components (x(t), y(t), z(t)).
-	// This isn't solvable analytically, so we use a numerical method (Legendre-Gauss quadrature) which performs very well
-	// with functions of this type, even with very few samples.  In this case, just 5 samples is sufficient to yield a
-	// reasonable result.
+// todo: make cvar when importing handles it changing.
+constexpr bool bEncodeAsHex = true;
 
-	struct FLegendreGaussCoefficient
-	{
-		float Abscissa;
-		float Weight;
-	};
+bool FSpline::ExportTextItem(FString& ValueStr, FSpline const& DefaultValue, class UObject* Parent, int32 PortFlags, class UObject* ExportRootScope) const
+{
+	// serialize our spline
+    TArray<uint8> SplineWriteBuffer;
+    FMemoryWriter MemWriter(SplineWriteBuffer);
+    SerializeSave(MemWriter);
 
-	static const FLegendreGaussCoefficient LegendreGaussCoefficients[] =
+	if (bEncodeAsHex)
 	{
-		{ 0.0f, 0.5688889f },
-		{ -0.5384693f, 0.47862867f },
-		{ 0.5384693f, 0.47862867f },
-		{ -0.90617985f, 0.23692688f },
-		{ 0.90617985f, 0.23692688f }
-	};
-
-	const auto& StartPoint = PositionCurve.Points[Index];
-	const auto& EndPoint = PositionCurve.Points[Index == LastPoint ? 0 : Index + 1];
-
-	const auto& P0 = StartPoint.OutVal;
-	const auto& T0 = StartPoint.LeaveTangent;
-	const auto& P1 = EndPoint.OutVal;
-	const auto& T1 = EndPoint.ArriveTangent;
-
-	// Special cases for linear or constant segments
-	if (StartPoint.InterpMode == CIM_Linear)
-	{
-		return ((P1 - P0) * Scale3D).Size() * Param;
+		FString HexStr = FString::FromHexBlob(SplineWriteBuffer.GetData(), SplineWriteBuffer.Num());
+		ValueStr = FString::Printf(TEXT("SplineData SplineDataLen=%d SplineData=%s\r\n"), HexStr.Len(), *HexStr);
 	}
-	else if (StartPoint.InterpMode == CIM_Constant)
+	else
 	{
-		// Special case: constant interpolation acts like distance = 0 for all p in [0, 1[ but for p == 1, the distance returned is the linear distance between start and end
-		return Param == 1.f ? ((P1 - P0) * Scale3D).Size() : 0.0f;
-	}
-
-	// Cache the coefficients to be fed into the function to calculate the spline derivative at each sample point as they are constant.
-	const FVector Coeff1 = ((P0 - P1) * 2.0f + T0 + T1) * 3.0f;
-	const FVector Coeff2 = (P1 - P0) * 6.0f - T0 * 4.0f - T1 * 2.0f;
-	const FVector Coeff3 = T0;
+		FString Base64String = FBase64::Encode(SplineWriteBuffer);
 
-	const float HalfParam = Param * 0.5f;
+		// Base64 encoding uses the '/' character, but T3D interprets '//' as some kind of
+		// terminator (?). If it occurs then the string passed to ImportTextItem() will
+		// come back as full of nullptrs. So we will swap in '-' here, and swap back to '/' in ImportTextItem()
+		Base64String.ReplaceCharInline('/', '-');
 
-	float Length = 0.0f;
-	for (const auto& LegendreGaussCoefficient : LegendreGaussCoefficients)
-	{
-		// Calculate derivative at each Legendre-Gauss sample, and perform a weighted sum
-		const float Alpha = HalfParam * (1.0f + LegendreGaussCoefficient.Abscissa);
-		const FVector Derivative = ((Coeff1 * Alpha + Coeff2) * Alpha + Coeff3) * Scale3D;
-		Length += Derivative.Size() * LegendreGaussCoefficient.Weight;
+		ValueStr = FString::Printf(TEXT("SplineData SplineDataLen=%d SplineData=%s\r\n"), Base64String.Len(), *Base64String);
 	}
-	Length *= HalfParam;
-
-	return Length;
+	
+	return true;
 }
 
-float FLegacySpline::GetSplineLength() const
+bool FSpline::ImportTextItem(const TCHAR*& SourceText, int32 PortFlags, UObject* Parent, FOutputDevice* ErrorText)
 {
-	const int32 NumPoints = ReparamTable.Points.Num();
+	// todo: SplineData= is followed by a byte identifying the implementation type, we should
+	// validate this on import and potentially change spline type if it doesn't match our current type
 
-	// This is given by the input of the last entry in the remap table
-	if (NumPoints > 0)
+	if (FParse::Command(&SourceText, TEXT("SplineData")))	
 	{
-		return ReparamTable.Points.Last().InVal;
-	}
-
-	return 0.0f;
-}
+		static const TCHAR SplineDataLenToken[] = TEXT("SplineDataLen=");
+		const TCHAR* FoundSplineDataLenStart = FCString::Strifind(SourceText, SplineDataLenToken);
+		if (FoundSplineDataLenStart)
+		{
+			SourceText = FoundSplineDataLenStart + FCString::Strlen(SplineDataLenToken);
+			int32 SplineDataLen = FCString::Atoi(SourceText);
 
-int32 FLegacySpline::GetNumControlPoints() const
-{
-	return PositionCurve.Points.Num();
-}
+			static const TCHAR SplineDataToken[] = TEXT("SplineData=");
+			const TCHAR* FoundSplineDataStart = FCString::Strifind(SourceText, SplineDataToken);
+			if (FoundSplineDataStart)
+			{
+				SourceText = FoundSplineDataStart + FCString::Strlen(SplineDataToken);
+				FString SplineData = FString::ConstructFromPtrSize(SourceText, SplineDataLen);
 
-void FLegacySpline::Reset()
-{
-	PositionCurve.Points.Reset();
-	RotationCurve.Points.Reset();
-	ScaleCurve.Points.Reset();
-}
+				// fix-up the hack applied to the Base64-encoded string in ExportTextItem()
+				SplineData.ReplaceCharInline('-', '/');
 
-void FLegacySpline::ResetRotation()
-{
-	RotationCurve.Points.Reset(PositionCurve.Points.Num());
+				TArray<uint8> SplineReadBuffer;
+				bool bDecoded;
+				
+				if (bEncodeAsHex)
+				{
+					SplineReadBuffer.SetNum(SplineDataLen / 2);
+					bDecoded = FString::ToHexBlob(SplineData, SplineReadBuffer.GetData(), SplineReadBuffer.Num());
+				}
+				else
+				{
+					bDecoded = FBase64::Decode(SplineData, SplineReadBuffer);
+				}
 
-	for (int32 Count = RotationCurve.Points.Num(); Count < PositionCurve.Points.Num(); Count++)
-	{
-		RotationCurve.Points.Emplace(static_cast<float>(Count), FQuat::Identity, FQuat::Identity, FQuat::Identity, CIM_CurveAuto);
+				if (bDecoded)
+				{
+					FMemoryReader MemReader(SplineReadBuffer);
+					SerializeLoad(MemReader);
+				}
+			}
+		}
 	}
+	
+	return true;
 }
 
-void FLegacySpline::ResetScale()
-{
-	ScaleCurve.Points.Reset(PositionCurve.Points.Num());
-
-	for (int32 Count = ScaleCurve.Points.Num(); Count < PositionCurve.Points.Num(); Count++)
-	{
-		ScaleCurve.Points.Emplace(static_cast<float>(Count), FVector(1.0f), FVector::ZeroVector, FVector::ZeroVector, CIM_CurveAuto);
-	}
+UE::Spline::ESplineType FSpline::GetCurrentImplementation() const
+{ 
+	return CurrentImplementation;
 }
 
-void FLegacySpline::UpdateSpline(const FSpline::FUpdateSplineParams& InParams)
+#if WITH_EDITOR
+UE::Spline::ESplineType FSpline::GetPreviousImplementation() const
 {
-	TRACE_CPUPROFILER_EVENT_SCOPE(FLegacySpline::UpdateSpline);
-
-	const int32 NumPoints = PositionCurve.Points.Num();
-	check(RotationCurve.Points.Num() == NumPoints && ScaleCurve.Points.Num() == NumPoints);
-
-#if DO_CHECK
-	// Ensure input keys are strictly ascending
-	for (int32 Index = 1; Index < NumPoints; Index++)
-	{
-		ensureAlways(PositionCurve.Points[Index - 1].InVal < PositionCurve.Points[Index].InVal);
-	}
+	return PreviousImplementation;
+}
 #endif
 
-	// Ensure splines' looping status matches with that of the spline component
-	if (InParams.bClosedLoop)
-	{
-		const float LastKey = PositionCurve.Points.Num() > 0 ? PositionCurve.Points.Last().InVal : 0.0f;
-		const float LoopKey = InParams.bLoopPositionOverride ? InParams.LoopPosition : LastKey + 1.0f;
-		PositionCurve.SetLoopKey(LoopKey);
-		RotationCurve.SetLoopKey(LoopKey);
-		ScaleCurve.SetLoopKey(LoopKey);
-	}
-	else
-	{
-		PositionCurve.ClearLoopKey();
-		RotationCurve.ClearLoopKey();
-		ScaleCurve.ClearLoopKey();
-	}
+UE::Spline::ESplineType FSpline::GetDefaultImplementation()
+{
+	return UE::Spline::GetDefaultSplineType();
+}
 
-	// Automatically set the tangents on any CurveAuto keys
-	PositionCurve.AutoSetTangents(0.0f, InParams.bStationaryEndpoints);
-	RotationCurve.AutoSetTangents(0.0f, InParams.bStationaryEndpoints);
-	ScaleCurve.AutoSetTangents(0.0f, InParams.bStationaryEndpoints);
+void FSpline::SetCurrentImplementation(UE::Spline::ESplineType Implementation)
+{
+	CurrentImplementation = Implementation;
 
-	// Now initialize the spline reparam table
-	const int32 NumSegments = PositionCurve.bIsLooped ? NumPoints : FMath::Max(0, NumPoints - 1);
+	LegacySplineImpl.Reset();
+	SplineImpl.Reset();
 
-	// Start by clearing it
-	ReparamTable.Points.Reset(NumSegments * InParams.ReparamStepsPerSegment + 1);
-	float AccumulatedLength = 0.0f;
-	for (int32 SegmentIndex = 0; SegmentIndex < NumSegments; ++SegmentIndex)
+	switch (CurrentImplementation)
 	{
-		for (int32 Step = 0; Step < InParams.ReparamStepsPerSegment; ++Step)
-		{
-			const float Param = static_cast<float>(Step) / InParams.ReparamStepsPerSegment;
-			const float SegmentLength = (Step == 0) ? 0.0f : GetSegmentLength(SegmentIndex, Param, InParams.Scale3D);
-
-			ReparamTable.Points.Emplace(SegmentLength + AccumulatedLength, SegmentIndex + Param, 0.0f, 0.0f, CIM_Linear);
-		}
-		AccumulatedLength += GetSegmentLength(SegmentIndex, 1.0f, InParams.Scale3D);
+	default: break;		// Intentionally doing nothing
+	case UE::Spline::ESplineType::LegacyTangent:	LegacySplineImpl = MakeUnique<FLegacyTangentSpline>(); break;
+	case UE::Spline::ESplineType::Tangent:			SplineImpl = MakeUnique<FTangentSpline>(); break;
 	}
-	ReparamTable.Points.Emplace(AccumulatedLength, static_cast<float>(NumSegments), 0.0f, 0.0f, CIM_Linear);
 }
 
 PRAGMA_ENABLE_EXPERIMENTAL_WARNINGS
\ No newline at end of file
