[NetPhysics] Fix to how GetForwardPredictionTime() is computed, using the cheaper GetLastReceivedServerFrame() instead of GetLatestReceivedTargetStateFrame() which currently takes a physics read lock.

#Synced-CL 49060427

diff --git a/Engine/Source/Runtime/Engine/Private/PhysicsEngine/NetworkPhysicsComponent.cpp b/Engine/Source/Runtime/Engine/Private/PhysicsEngine/NetworkPhysicsComponent.cpp
--- a/Engine/Source/Runtime/Engine/Private/PhysicsEngine/NetworkPhysicsComponent.cpp
+++ b/Engine/Source/Runtime/Engine/Private/PhysicsEngine/NetworkPhysicsComponent.cpp
@@ -2967,10 +2967,10 @@ const float FAsyncNetworkPhysicsComponent::GetForwardPredictionTime() const
 	const Chaos::FPhysicsSolverBase* BaseSolver = GetSolver();
 	if (const Chaos::FRewindData* RewindData = BaseSolver ? BaseSolver->GetRewindData() : nullptr)
 	{
-		int32 LatestReceivedTargetStateFrame = GetLatestReceivedTargetStateFrame();
-		if (LatestReceivedTargetStateFrame != INDEX_NONE)
+		int32 LatestReceivedServerFrame = GetLastReceivedServerFrame();
+		if (LatestReceivedServerFrame != INDEX_NONE)
 		{
-			const int32 MeasuredInputDelayFrames = RewindData->GetLatestFrame() - LatestReceivedTargetStateFrame;
+			const int32 MeasuredInputDelayFrames = RewindData->GetLatestFrame() - LatestReceivedServerFrame;
 			const float MeasuredInputDelaySeconds = FMath::Max(MeasuredInputDelayFrames, 0.0f) * BaseSolver->GetAsyncDeltaTime();
 			return MeasuredInputDelaySeconds;
 		}
