Caching the RenderTarget's ViewRect for the game thread before passing the canvas to the render thread and using that value to avoid data race conflicts when drawing HUD Stats

#Synced-CL 48916169

diff --git a/Engine/Source/Runtime/Engine/Private/Slate/DebugCanvas.cpp b/Engine/Source/Runtime/Engine/Private/Slate/DebugCanvas.cpp
--- a/Engine/Source/Runtime/Engine/Private/Slate/DebugCanvas.cpp
+++ b/Engine/Source/Runtime/Engine/Private/Slate/DebugCanvas.cpp
@@ -30,7 +30,14 @@ public:
 	/** FRenderTarget interface */
 	virtual FIntPoint GetSizeXY() const
 	{
-		return ViewRect.Size();
+		if (IsInRenderingThread())
+		{
+			return ViewRect_RenderThread.Size();
+		}
+		else
+		{
+			return ViewRect_GameThread.Size();
+		}
 	}
 
 	/** Sets the texture that this target renders to */
@@ -56,20 +63,31 @@ public:
 		return RDGTexture;
 	}
 
-	void SetViewRect(const FIntRect& InViewRect)
+	void SetViewRect_GameThread(const FIntRect& InViewRect)
 	{
-		ViewRect = InViewRect;
+		ViewRect_GameThread = InViewRect;
+	}
+
+	void SetViewRect_RenderThread(const FIntRect& InViewRect)
+	{
+		ViewRect_RenderThread = InViewRect;
 	}
 
 	/** Gets the viewport rect for the render target */
-	const FIntRect& GetViewRect() const 
+	const FIntRect& GetViewRect_GameThread() const 
 	{
-		return ViewRect;
+		return ViewRect_GameThread;
+	}
+
+	const FIntRect& GetViewRect_RenderThread() const 
+	{
+		return ViewRect_RenderThread;
 	}
 
 private:
 	FRDGTexture* RDGTexture = nullptr;
-	FIntRect ViewRect;
+	FIntRect ViewRect_GameThread;
+	FIntRect ViewRect_RenderThread;
 };
 
 #define INVALID_LAYER_ID UINT_MAX
@@ -172,7 +190,6 @@ FCanvas* FDebugCanvasDrawer::GetGameThreadDebugCanvas()
 	return GameThreadCanvas.Get();
 }
 
-
 void FDebugCanvasDrawer::BeginRenderingCanvas( const FIntRect& CanvasRect )
 {
 	if( CanvasRect.Size().X > 0 && CanvasRect.Size().Y > 0 )
@@ -186,6 +203,17 @@ void FDebugCanvasDrawer::BeginRenderingCanvas( const FIntRect& CanvasRect )
 			LayerTexture->ResizeTarget(CanvasRect.Width(), CanvasRect.Height());
 			XLayerResource = LayerTexture->GameThread_GetRenderTargetResource();
 		}
+
+		// Caching the ViewRect for the game thread before passing the canvas to the render thread
+		if (CanvasToRender.IsValid() && CanvasToRender->IsUsingInternalTexture())
+		{
+			RenderTarget->SetViewRect_GameThread(FIntRect(FIntPoint(0, 0), CanvasToRender->GetParentCanvasSize()));
+		}
+		else
+		{
+			RenderTarget->SetViewRect_GameThread(CanvasRect);
+		}
+
 		ENQUEUE_RENDER_COMMAND(BeginRenderingDebugCanvas)(
 			[CanvasDrawer, CanvasToRender, XLayerResource, CanvasRect = CanvasRect](FRHICommandListImmediate& RHICmdList)
 			{
@@ -298,7 +326,7 @@ void FDebugCanvasDrawer::Draw_RenderThread(FRDGBuilder& GraphBuilder, const FDra
 		}
 		else
 		{
-			RenderThreadCanvas->SetRenderTargetRect(RenderTarget->GetViewRect());
+			RenderThreadCanvas->SetRenderTargetRect(RenderTarget->GetViewRect_RenderThread());
 		}
 
 		RenderThreadCanvas->Flush_RenderThread(GraphBuilder, true);
@@ -324,11 +352,11 @@ void FDebugCanvasDrawer::SetRenderThreadCanvas( const FIntRect& InCanvasRect, FC
 	check( IsInRenderingThread() );
 	if (Canvas->IsUsingInternalTexture())
 	{
-		RenderTarget->SetViewRect(FIntRect(FIntPoint(0, 0), Canvas->GetParentCanvasSize()));
+		RenderTarget->SetViewRect_RenderThread(FIntRect(FIntPoint(0, 0), Canvas->GetParentCanvasSize()));
 	}
 	else
 	{
-		RenderTarget->SetViewRect(InCanvasRect);
+		RenderTarget->SetViewRect_RenderThread(InCanvasRect);
 	}
 	RenderThreadCanvas = Canvas;
 	LayerResource = InLayerResource;
