[Backout] - CL49076318
#rnx
Original CL Desc
-----------------------------------------------------------------
FSpline Refactor: Prepare FSpline for non-tangent based splines.

#Synced-CL 49044627

#Synced-CL 49076318

diff --git a/Engine/Plugins/PCG/Source/PCG/Private/Data/PCGSplineStruct.cpp b/Engine/Plugins/PCG/Source/PCG/Private/Data/PCGSplineStruct.cpp
--- a/Engine/Plugins/PCG/Source/PCG/Private/Data/PCGSplineStruct.cpp
+++ b/Engine/Plugins/PCG/Source/PCG/Private/Data/PCGSplineStruct.cpp
@@ -2,7 +2,6 @@
 
 #include "Data/PCGSplineStruct.h"
 
-#include "Curves/Splines/TangentSpline.h"
 #include "Metadata/PCGMetadataCommon.h"
 
 #include UE_INLINE_GENERATED_CPP_BY_NAME(PCGSplineStruct)
@@ -97,17 +96,8 @@ namespace PCGSplineStruct
 //
 //		return FBoxSphereBounds(BoundingBox.TransformBy(LocalToWorld));
 //#else
-		const FTangentSpline* TangentSpline = Spline.Get<UE::Spline::ESplineType::Tangent>();
-
-		if (!TangentSpline)
-		{
-			ensureAlwaysMsgf(false, TEXT("ShouldUseSplineCurves() returned false but Spline failed to yield an FTangentSpline!"));
-			return FBoxSphereBounds();
-		}
-
-		const int32 NumPoints = TangentSpline->GetSpline().GetNumPoints();
+		const int32 NumPoints = Spline.GetNumControlPoints();
 		const int32 NumSegments = bClosedLoop ? NumPoints : NumPoints - 1;
-		const FInterpCurveVector& SplinePoints = TangentSpline->GetSplinePointsPosition();
 
 		FVector Min(WORLD_MAX);
 		FVector Max(-WORLD_MAX);
@@ -117,11 +107,11 @@ namespace PCGSplineStruct
 			{
 				const bool bLoopSegment = (Index == NumPoints - 1);
 				const int32 NextIndex = bLoopSegment ? 0 : Index + 1;
-				const FInterpCurvePoint<FVector>& ThisInterpPoint = SplinePoints.Points[Index];
-				FInterpCurvePoint<FVector> NextInterpPoint = SplinePoints.Points[NextIndex];
+				const FInterpCurvePoint<FVector>& ThisInterpPoint = Spline.GetSplinePointsPosition().Points[Index];
+				FInterpCurvePoint<FVector> NextInterpPoint = Spline.GetSplinePointsPosition().Points[NextIndex];
 				if (bLoopSegment)
 				{
-					NextInterpPoint.InVal = ThisInterpPoint.InVal + SplinePoints.LoopKeyOffset;
+					NextInterpPoint.InVal = ThisInterpPoint.InVal + Spline.GetSplinePointsPosition().LoopKeyOffset;
 				}
 
 				CurveVectorFindIntervalBounds(ThisInterpPoint, NextInterpPoint, Min, Max);
@@ -129,7 +119,7 @@ namespace PCGSplineStruct
 		}
 		else if (NumPoints == 1)
 		{
-			Min = Max = SplinePoints.Points[0].OutVal;
+			Min = Max = Spline.GetSplinePointsPosition().Points[0].OutVal;
 		}
 		else
 		{
@@ -162,15 +152,6 @@ void FPCGSplineStruct::Initialize(const USplineComponent* InSplineComponent)
 
 void FPCGSplineStruct::Initialize(const TArray<FSplinePoint>& InSplinePoints, bool bIsClosedLoop, const FTransform& InTransform, TArray<PCGMetadataEntryKey> InOptionalEntryKeys)
 {
-	if (PCGSplineStruct::GUseSplineCurves)
-	{
-		Spline = FSpline(UE::Spline::ESplineType::Unimplemented);
-	}
-	else
-	{
-		Spline = FSpline(UE::Spline::ESplineType::Tangent);
-	}
-
 	Transform = InTransform;
 	DefaultUpVector = FVector::ZAxisVector;
 	ReparamStepsPerSegment = 10; // default value in USplineComponent
@@ -206,6 +187,7 @@ void FPCGSplineStruct::ApplyTo(USplineComponent* InSplineComponent) const
 	InSplineComponent->DefaultUpVector = DefaultUpVector;
 	InSplineComponent->ReparamStepsPerSegment = ReparamStepsPerSegment;
 
+	// Not using a ternary so we are not doing any implicit conversions (SetSpline is overloaded)
 	if (ShouldUseSplineCurves())
 	{
 		InSplineComponent->SetSpline(WarninglessSplineCurves());
@@ -270,10 +252,7 @@ FPCGSplineStruct& FPCGSplineStruct::operator=(FPCGSplineStruct&& Other)
 
 void FPCGSplineStruct::AddPoint(const FSplinePoint& InSplinePoint, bool bUpdateSpline, const int64* InOptionalControlPointEntryKey)
 {
-	if (FTangentSpline* TangentSpline = Spline.Get<UE::Spline::ESplineType::Tangent>())
-	{
-		TangentSpline->AddPoint(InSplinePoint);
-	}
+	Spline.AddPoint(InSplinePoint);
  
 	const int32 Index = PCGSplineStruct::UpperBound(WarninglessSplineCurves().Position.Points, InSplinePoint.InputKey);
 	
@@ -345,181 +324,79 @@ void FPCGSplineStruct::AddPoints(const TArray<FSplinePoint>& InSplinePoints, boo
 void FPCGSplineStruct::SetLocation(int32 Index, const FVector& InLocation)
 {
 	WarninglessSplineCurves().Position.Points[Index].OutVal = InLocation;
-
-	if (FTangentSpline* TangentSpline = Spline.Get<UE::Spline::ESplineType::Tangent>())
-	{
-		TangentSpline->SetLocation(Index, InLocation);
-	}
+	Spline.SetLocation(Index, InLocation);
 }
 
 FVector FPCGSplineStruct::GetLocation(const int32 Index) const
 {
-	FVector Location = FVector::ZeroVector;
-
-	if (ShouldUseSplineCurves())
-	{
-		Location = WarninglessSplineCurves().Position.Points[Index].OutVal;
-	}
-	else if (const FTangentSpline* TangentSpline = Spline.Get<UE::Spline::ESplineType::Tangent>())
-	{
-		Location = TangentSpline->GetLocation(Index);
-	}
-	else
-	{
-		ensureAlwaysMsgf(false, TEXT("ShouldUseSplineCurves() returned false but Spline failed to yield an FTangentSpline!"));
-	}
-
-	return Location;
+	return ShouldUseSplineCurves()
+		? WarninglessSplineCurves().Position.Points[Index].OutVal
+		: Spline.GetLocation(Index);
 }
 	
 void FPCGSplineStruct::SetArriveTangent(const int32 Index, const FVector& InArriveTangent)
 {
 	WarninglessSplineCurves().Position.Points[Index].ArriveTangent = InArriveTangent;
-
-	if (FTangentSpline* TangentSpline = Spline.Get<UE::Spline::ESplineType::Tangent>())
-	{
-		TangentSpline->SetInTangent(Index, InArriveTangent);
-	}
+	Spline.SetInTangent(Index, InArriveTangent);
 }
 
 FVector FPCGSplineStruct::GetArriveTangent(const int32 Index) const
 {
-	FVector Tangent = FVector::ZeroVector;
-
-	if (ShouldUseSplineCurves())
-	{
-		Tangent = WarninglessSplineCurves().Position.Points[Index].ArriveTangent;
-	}
-	else if (const FTangentSpline* TangentSpline = Spline.Get<UE::Spline::ESplineType::Tangent>())
-	{
-		Tangent = TangentSpline->GetInTangent(Index);
-	}
-	else
-	{
-		ensureAlwaysMsgf(false, TEXT("ShouldUseSplineCurves() returned false but Spline failed to yield an FTangentSpline!"));
-	}
-
-	return Tangent;
+	return ShouldUseSplineCurves()
+		? WarninglessSplineCurves().Position.Points[Index].ArriveTangent
+		: Spline.GetInTangent(Index);
 }
 	
 void FPCGSplineStruct::SetLeaveTangent(const int32 Index, const FVector& InLeaveTangent)
 {
 	WarninglessSplineCurves().Position.Points[Index].LeaveTangent = InLeaveTangent;
-
-	if (FTangentSpline* TangentSpline = Spline.Get<UE::Spline::ESplineType::Tangent>())
-	{
-		TangentSpline->SetOutTangent(Index, InLeaveTangent);
-	}
+	Spline.SetOutTangent(Index, InLeaveTangent);
 }
 
 FVector FPCGSplineStruct::GetLeaveTangent(const int32 Index) const
 {
-	FVector Tangent = FVector::ZeroVector;
-
-	if (ShouldUseSplineCurves())
-	{
-		Tangent = WarninglessSplineCurves().Position.Points[Index].LeaveTangent;
-	}
-	else if (const FTangentSpline* TangentSpline = Spline.Get<UE::Spline::ESplineType::Tangent>())
-	{
-		Tangent = TangentSpline->GetOutTangent(Index);
-	}
-	else
-	{
-		ensureAlwaysMsgf(false, TEXT("ShouldUseSplineCurves() returned false but Spline failed to yield an FTangentSpline!"));
-	}
-
-	return Tangent;
+	return ShouldUseSplineCurves()
+		? WarninglessSplineCurves().Position.Points[Index].LeaveTangent
+		: Spline.GetOutTangent(Index);
 }
 
 void FPCGSplineStruct::SetRotation(int32 Index, const FQuat& InRotation)
 {
 	WarninglessSplineCurves().Rotation.Points[Index].OutVal = InRotation;
-
-	if (FTangentSpline* TangentSpline = Spline.Get<UE::Spline::ESplineType::Tangent>())
-	{
-		TangentSpline->SetRotation(Index, InRotation);
-	}
+	Spline.SetRotation(Index, InRotation);
 }
 
 FQuat FPCGSplineStruct::GetRotation(const int32 Index) const
 {
-	FQuat Quat = FQuat::Identity;
-
-	if (ShouldUseSplineCurves())
-	{
-		Quat = WarninglessSplineCurves().Rotation.Points[Index].OutVal;
-	}
-	else if (const FTangentSpline* TangentSpline = Spline.Get<UE::Spline::ESplineType::Tangent>())
-	{
-		Quat = TangentSpline->GetRotation(Index);
-	}
-	else
-	{
-		ensureAlwaysMsgf(false, TEXT("ShouldUseSplineCurves() returned false but Spline failed to yield an FTangentSpline!"));
-	}
-
-	return Quat;
+	return ShouldUseSplineCurves()
+		? WarninglessSplineCurves().Rotation.Points[Index].OutVal
+		: Spline.GetRotation(Index);
 }
 	
 void FPCGSplineStruct::SetScale(int32 Index, const FVector& InScale)
 {
 	WarninglessSplineCurves().Scale.Points[Index].OutVal = InScale;
-
-	if (FTangentSpline* TangentSpline = Spline.Get<UE::Spline::ESplineType::Tangent>())
-	{
-		TangentSpline->SetScale(Index, InScale);
-	}
+	Spline.SetScale(Index, InScale);
 }
 
 FVector FPCGSplineStruct::GetScale(const int32 Index) const
 {
-	FVector Scale = FVector::OneVector;
-
-	if (ShouldUseSplineCurves())
-	{
-		Scale = WarninglessSplineCurves().Scale.Points[Index].OutVal;
-	}
-	else if (const FTangentSpline* TangentSpline = Spline.Get<UE::Spline::ESplineType::Tangent>())
-	{
-		Scale = TangentSpline->GetScale(Index);
-	}
-	else
-	{
-		ensureAlwaysMsgf(false, TEXT("ShouldUseSplineCurves() returned false but Spline failed to yield an FTangentSpline!"));
-	}
-
-	return Scale;
+	return ShouldUseSplineCurves()
+		? WarninglessSplineCurves().Scale.Points[Index].OutVal
+		: Spline.GetScale(Index);
 }
 	
 void FPCGSplineStruct::SetSplinePointType(int32 Index, EInterpCurveMode Type)
 {
 	WarninglessSplineCurves().Position.Points[Index].InterpMode = Type;
-
-	if (FTangentSpline* TangentSpline = Spline.Get<UE::Spline::ESplineType::Tangent>())
-	{
-		TangentSpline->SetSplinePointType(Index, Type);
-	}
+	Spline.SetSplinePointType(Index, Type);
 }
 
 EInterpCurveMode FPCGSplineStruct::GetSplinePointType(int32 Index) const
 {
-	EInterpCurveMode Mode = CIM_Unknown;
-
-	if (ShouldUseSplineCurves())
-	{
-		Mode = WarninglessSplineCurves().Position.Points[Index].InterpMode.GetValue();
-	}
-	else if (const FTangentSpline* TangentSpline = Spline.Get<UE::Spline::ESplineType::Tangent>())
-	{
-		Mode = TangentSpline->GetSplinePointType(Index);
-	}
-	else
-	{
-		ensureAlwaysMsgf(false, TEXT("ShouldUseSplineCurves() returned false but Spline failed to yield an FTangentSpline!"));
-	}
-
-	return Mode;
+	return ShouldUseSplineCurves()
+		? WarninglessSplineCurves().Position.Points[Index].InterpMode.GetValue()
+		: Spline.GetSplinePointType(Index);
 }
 
 const TArray<FName>& FPCGSplineStruct::GetAttributeChannelNames() const
@@ -531,17 +408,14 @@ const TArray<FName>& FPCGSplineStruct::GetAttributeChannelNames() const
 			CachedSortedAttributeChannelNames.Empty();
 		}
 	}
-	else if (const FTangentSpline* TangentSpline = Spline.Get<UE::Spline::ESplineType::Tangent>())
+	else if (TArray<FName> Channels = Spline.GetFloatPropertyChannels(); Channels.Num() != CachedSortedAttributeChannelNames.Num())
 	{
-		if (TArray<FName> Channels = TangentSpline->GetAttributeChannelNamesByValueType<float>(); Channels.Num() != CachedSortedAttributeChannelNames.Num())
+		Channels.Sort([](const FName& A, const FName& B)
+			{ return A.ToString() < B.ToString(); });
+		
+		if (PCG::TScopeLock ScopeLock(Lock); Channels.Num() != CachedSortedAttributeChannelNames.Num())
 		{
-			Channels.Sort([](const FName& A, const FName& B)
-				{ return A.ToString() < B.ToString(); });
-
-			if (PCG::TScopeLock ScopeLock(Lock); Channels.Num() != CachedSortedAttributeChannelNames.Num())
-			{
-				CachedSortedAttributeChannelNames = MoveTemp(Channels);
-			}
+			CachedSortedAttributeChannelNames = MoveTemp(Channels);
 		}
 	}
 
@@ -550,41 +424,21 @@ const TArray<FName>& FPCGSplineStruct::GetAttributeChannelNames() const
 
 float FPCGSplineStruct::GetAttributeValue(float InputKey, FName ChannelName) const
 {
-	if (ShouldUseSplineCurves())
-	{
-		return 0.f;
-	}
-	else if (const FTangentSpline* TangentSpline = Spline.Get<UE::Spline::ESplineType::Tangent>())
-	{
-		return TangentSpline->EvaluateAttribute<float>(ChannelName, InputKey);
-	}
-	else
-	{
-		ensureAlwaysMsgf(false, TEXT("ShouldUseSplineCurves() returned false but Spline failed to yield an FTangentSpline!"));
-		return 0.f;
-	}
+	return ShouldUseSplineCurves()
+		? 0.f
+		: Spline.EvaluateAttribute<float>(InputKey, ChannelName);
 }
 
 void FPCGSplineStruct::UpdateSpline()
 {
-	// If SplineComponent.UseSplineCurves is false and FSpline is not implemented, initialize it from SplineCurves.
-	if (PCGSplineStruct::GUseSplineCurves == false && Spline.GetCurrentImplementation() != UE::Spline::ESplineType::Tangent)
-	{
-		Spline = WarninglessSplineCurves();
-	}
-
 	const bool bLoopPositionOverride = false;
 	const bool bStationaryEndpoints = false;
 	const float LoopPosition = 0.0f;
 
-	FTangentSpline::FUpdateSplineParams Params{bClosedLoop, bStationaryEndpoints, ReparamStepsPerSegment, bLoopPositionOverride, LoopPosition};
+	FSpline::FUpdateSplineParams Params{bClosedLoop, bStationaryEndpoints, ReparamStepsPerSegment, bLoopPositionOverride, LoopPosition};
 
 	WarninglessSplineCurves().UpdateSpline(Params.bClosedLoop, Params.bStationaryEndpoints, Params.ReparamStepsPerSegment, Params.bLoopPositionOverride, Params.LoopPosition, Params.Scale3D);
-
-	if (FTangentSpline* TangentSpline = Spline.Get<UE::Spline::ESplineType::Tangent>())
-	{
-		TangentSpline->UpdateSpline(Params);
-	}
+	Spline.UpdateSpline(Params);
 }
 
 int FPCGSplineStruct::GetNumberOfSplineSegments() const
@@ -595,36 +449,16 @@ int FPCGSplineStruct::GetNumberOfSplineSegments() const
 
 int FPCGSplineStruct::GetNumberOfPoints() const
 {
-	if (ShouldUseSplineCurves())
-	{
-		return WarninglessSplineCurves().Position.Points.Num();
-	}
-	else if (const FTangentSpline* TangentSpline = Spline.Get<UE::Spline::ESplineType::Tangent>())
-	{
-		return TangentSpline->GetSpline().GetNumPoints();
-	}
-	else
-	{
-		ensureAlwaysMsgf(false, TEXT("ShouldUseSplineCurves() returned false but Spline failed to yield an FTangentSpline!"));
-		return 0;
-	}
+	return ShouldUseSplineCurves()
+		? WarninglessSplineCurves().Position.Points.Num()
+		: Spline.GetNumControlPoints();
 }
 
 float FPCGSplineStruct::GetLoopKeyOffset() const
 {
-	if (ShouldUseSplineCurves())
-	{
-		return WarninglessSplineCurves().Position.LoopKeyOffset;
-	}
-	else if (const FTangentSpline* TangentSpline = Spline.Get<UE::Spline::ESplineType::Tangent>())
-	{
-		return TangentSpline->GetSplinePointsPosition().LoopKeyOffset;
-	}
-	else
-	{
-		ensureAlwaysMsgf(false, TEXT("ShouldUseSplineCurves() returned false but Spline failed to yield an FTangentSpline!"));
-		return 0.f;
-	}
+	return ShouldUseSplineCurves()
+		? WarninglessSplineCurves().Position.LoopKeyOffset
+		: Spline.GetSplinePointsPosition().LoopKeyOffset;
 }
 
 void FPCGSplineStruct::SetLoopKeyOffset(float LoopKeyOffset)
@@ -636,34 +470,17 @@ void FPCGSplineStruct::SetLoopKeyOffset(float LoopKeyOffset)
 		const bool bLoopPositionOverride = true;
 		const float LoopPosition = LoopKeyOffset + (WarninglessSplineCurves().Position.Points.IsEmpty() ? 0.0f : WarninglessSplineCurves().Position.Points.Last().InVal);
 
-		FTangentSpline::FUpdateSplineParams Params{ bClosedLoop, bStationaryEndpoints, ReparamStepsPerSegment, bLoopPositionOverride, LoopPosition };
-
+		FSpline::FUpdateSplineParams Params{ bClosedLoop, bStationaryEndpoints, ReparamStepsPerSegment, bLoopPositionOverride, LoopPosition };
 		WarninglessSplineCurves().UpdateSpline(Params.bClosedLoop, Params.bStationaryEndpoints, Params.ReparamStepsPerSegment, Params.bLoopPositionOverride, Params.LoopPosition, Params.Scale3D);
-
-		if (FTangentSpline* TangentSpline = Spline.Get<UE::Spline::ESplineType::Tangent>())
-		{
-			TangentSpline->UpdateSpline(Params);
-
-		}
+		Spline.UpdateSpline(Params);
 	}
 }
 
 FVector::FReal FPCGSplineStruct::GetSplineLength() const
 {
-	if (ShouldUseSplineCurves())
-	{
-		return WarninglessSplineCurves().GetSplineLength();
-	}
-	else if (const FTangentSpline* TangentSpline = Spline.Get<UE::Spline::ESplineType::Tangent>())
-	{
-		const float MaxParameter = TangentSpline->GetParameterSpace().Max;
-		return TangentSpline->GetDistanceAtParameter(MaxParameter);
-	}
-	else
-	{
-		ensureAlwaysMsgf(false, TEXT("ShouldUseSplineCurves() returned false but Spline failed to yield an FTangentSpline!"));
-		return 0.f;
-	}
+	return ShouldUseSplineCurves()
+		? WarninglessSplineCurves().GetSplineLength()
+		: Spline.GetSplineLength();
 }
 
 FBox FPCGSplineStruct::GetBounds() const
@@ -689,9 +506,11 @@ FBox FPCGSplineStruct::GetSegmentBounds(int32 SegmentIndex) const
 		return FBox(EForceInit::ForceInit);
 	}
 
-	const FInterpCurveVector& SplinePointsPosition = GetSplinePointsPosition();
+	const FInterpCurveVector& SplinePointsPosition = ShouldUseSplineCurves()
+		? WarninglessSplineCurves().Position
+		: Spline.GetSplinePointsPosition();
 
-	const int32 NumPoints = SplinePointsPosition.Points.Num();
+	const int32 NumPoints = GetNumberOfPoints();
 	const bool bLoopSegment = (SegmentIndex == NumPoints - 1);
 	const int32 NextIndex = bLoopSegment ? 0 : SegmentIndex + 1;
 	const FInterpCurvePoint<FVector>& ThisInterpPoint = SplinePointsPosition.Points[SegmentIndex];
@@ -711,71 +530,23 @@ FBox FPCGSplineStruct::GetSegmentBounds(int32 SegmentIndex) const
 
 const FInterpCurveVector& FPCGSplineStruct::GetSplinePointsScale() const
 {
-	const FTangentSpline* TangentSpline = Spline.Get<UE::Spline::ESplineType::Tangent>();
-
-	// We normally branch directly on ShouldUseSplineCurves(), but because this function returns a const-ref
-	// we need a safe fallback if ShouldUseSplineCurves() is false and TangentSpline is unexpectedly nullptr.
-	// Cache the value and assert the contract up front, then fall back to the spline curves data if needed.
-	const bool bShouldUseSplineCurves = ShouldUseSplineCurves();
-	if (!bShouldUseSplineCurves)
-	{
-		ensureAlwaysMsgf(TangentSpline, TEXT("ShouldUseSplineCurves() returned false but Spline failed to yield an FTangentSpline!"));
-	}
-
-	if (bShouldUseSplineCurves || !TangentSpline)
-	{
-		return WarninglessSplineCurves().Scale;
-	}
-	else
-	{
-		return TangentSpline->GetSplinePointsScale();
-	}
+	return ShouldUseSplineCurves()
+		? WarninglessSplineCurves().Scale
+		: Spline.GetSplinePointsScale();
 }
 
 const FInterpCurveVector& FPCGSplineStruct::GetSplinePointsPosition() const
 {
-	const FTangentSpline* TangentSpline = Spline.Get<UE::Spline::ESplineType::Tangent>();
-
-	// We normally branch directly on ShouldUseSplineCurves(), but because this function returns a const-ref
-	// we need a safe fallback if ShouldUseSplineCurves() is false and TangentSpline is unexpectedly nullptr.
-	// Cache the value and assert the contract up front, then fall back to the spline curves data if needed.
-	const bool bShouldUseSplineCurves = ShouldUseSplineCurves();
-	if (!bShouldUseSplineCurves)
-	{
-		ensureAlwaysMsgf(TangentSpline, TEXT("ShouldUseSplineCurves() returned false but Spline failed to yield an FTangentSpline!"));
-	}
-
-	if (bShouldUseSplineCurves || !TangentSpline)
-	{
-		return WarninglessSplineCurves().Position;
-	}
-	else
-	{
-		return TangentSpline->GetSplinePointsPosition();
-	}
+	return ShouldUseSplineCurves()
+		? WarninglessSplineCurves().Position
+		: Spline.GetSplinePointsPosition();
 }
 
 const FInterpCurveQuat& FPCGSplineStruct::GetSplinePointsRotation() const
 {
-	const FTangentSpline* TangentSpline = Spline.Get<UE::Spline::ESplineType::Tangent>();
-
-	// We normally branch directly on ShouldUseSplineCurves(), but because this function returns a const-ref
-	// we need a safe fallback if ShouldUseSplineCurves() is false and TangentSpline is unexpectedly nullptr.
-	// Cache the value and assert the contract up front, then fall back to the spline curves data if needed.
-	const bool bShouldUseSplineCurves = ShouldUseSplineCurves();
-	if (!bShouldUseSplineCurves)
-	{
-		ensureAlwaysMsgf(TangentSpline, TEXT("ShouldUseSplineCurves() returned false but Spline failed to yield an FTangentSpline!"));
-	}
-
-	if (bShouldUseSplineCurves || !TangentSpline)
-	{
-		return WarninglessSplineCurves().Rotation;
-	}
-	else
-	{
-		return TangentSpline->GetSplinePointsRotation();
-	}
+	return ShouldUseSplineCurves()
+		? WarninglessSplineCurves().Rotation
+		: Spline.GetSplinePointsRotation();
 }
 
 FVector::FReal FPCGSplineStruct::GetDistanceAlongSplineAtSplinePoint(int32 PointIndex) const
@@ -799,14 +570,10 @@ FVector::FReal FPCGSplineStruct::GetDistanceAlongSplineAtSplinePoint(int32 Point
 			return WarninglessSplineCurves().ReparamTable.Points[PointIndex * ReparamStepsPerSegment].InVal;
 		}
 	}
-	else if (const FTangentSpline* TangentSpline = Spline.Get<UE::Spline::ESplineType::Tangent>())
-	{
-		const float ParameterAtIndex = TangentSpline->GetParameterAtIndex(PointIndex);
-		return TangentSpline->GetDistanceAtParameter(ParameterAtIndex);
-	}
 	else
 	{
-		ensureAlwaysMsgf(false, TEXT("ShouldUseSplineCurves() returned false but Spline failed to yield an FTangentSpline!"));
+		const float ParameterAtIndex = Spline.GetParameterAtIndex(PointIndex);
+		return Spline.GetDistanceAtParameter(ParameterAtIndex);
 	}
 
 	return 0.0f;
@@ -814,13 +581,19 @@ FVector::FReal FPCGSplineStruct::GetDistanceAlongSplineAtSplinePoint(int32 Point
 
 FVector FPCGSplineStruct::GetLocationAtDistanceAlongSpline(FVector::FReal Distance, ESplineCoordinateSpace::Type CoordinateSpace) const
 {
-	const float Param = GetInputKeyAtDistanceAlongSpline(Distance);
+	const float Param = ShouldUseSplineCurves()
+		? WarninglessSplineCurves().ReparamTable.Eval(Distance, 0.0f)
+		: Spline.GetParameterAtDistance(Distance);
+	
 	return GetLocationAtSplineInputKey(Param, CoordinateSpace);
 }
 
 FTransform FPCGSplineStruct::GetTransformAtDistanceAlongSpline(FVector::FReal Distance, ESplineCoordinateSpace::Type CoordinateSpace, bool bUseScale) const
 {
-	const float Param = GetInputKeyAtDistanceAlongSpline(Distance);
+	const float Param = ShouldUseSplineCurves()
+		? WarninglessSplineCurves().ReparamTable.Eval(Distance, 0.0f)
+		: Spline.GetParameterAtDistance(Distance);
+	
 	return GetTransformAtSplineInputKey(Param, CoordinateSpace, bUseScale);
 }
 
@@ -832,20 +605,10 @@ float FPCGSplineStruct::GetInputKeyAtDistanceAlongSpline(FVector::FReal Distance
 	{
 		return 0.0f;
 	}
-
-	if (ShouldUseSplineCurves())
-	{
-		return WarninglessSplineCurves().ReparamTable.Eval(Distance, 0.0f);
-	}
-	else if (const FTangentSpline* TangentSpline = Spline.Get<UE::Spline::ESplineType::Tangent>())
-	{
-		return TangentSpline->GetParameterAtDistance(Distance);
-	}
-	else
-	{
-		ensureAlwaysMsgf(false, TEXT("ShouldUseSplineCurves() returned false but Spline failed to yield an FTangentSpline!"));
-		return 0.f;
-	}
+	
+	return ShouldUseSplineCurves()
+		? WarninglessSplineCurves().ReparamTable.Eval(Distance, 0.0f)
+		: Spline.GetParameterAtDistance(Distance);
 }
 
 FVector FPCGSplineStruct::GetRightVectorAtSplineInputKey(float InKey, ESplineCoordinateSpace::Type CoordinateSpace) const
@@ -881,51 +644,33 @@ float FPCGSplineStruct::FindInputKeyClosestToWorldLocation(const FVector& WorldL
 {
 	const FVector LocalLocation = Transform.InverseTransformPosition(WorldLocation);
 	float Dummy;
-
-	if (ShouldUseSplineCurves())
-	{
-		return WarninglessSplineCurves().Position.FindNearest(LocalLocation, Dummy);
-	}
-	else if (const FTangentSpline* TangentSpline = Spline.Get<UE::Spline::ESplineType::Tangent>())
-	{
-		return TangentSpline->FindNearest(LocalLocation, Dummy);
-	}
-	else
-	{
-		ensureAlwaysMsgf(false, TEXT("ShouldUseSplineCurves() returned false but Spline failed to yield an FTangentSpline!"));
-		return 0.f;
-	}
+	return ShouldUseSplineCurves()
+		? WarninglessSplineCurves().Position.FindNearest(LocalLocation, Dummy)
+		: Spline.FindNearest(LocalLocation, Dummy);
 }
 
 float FPCGSplineStruct::FindInputKeyOnSegmentClosestToWorldLocation(const FVector& WorldLocation, int32 Index) const
 {
 	const FVector LocalLocation = Transform.InverseTransformPosition(WorldLocation);
 	float Dummy;
-
-	if (ShouldUseSplineCurves())
-	{
-		return WarninglessSplineCurves().Position.FindNearestOnSegment(LocalLocation, Index, Dummy);
-	}
-	else if (const FTangentSpline* TangentSpline = Spline.Get<UE::Spline::ESplineType::Tangent>())
-	{
-		return TangentSpline->FindNearestOnSegment(LocalLocation, Index, Dummy);
-	}
-	else
-	{
-		ensureAlwaysMsgf(false, TEXT("ShouldUseSplineCurves() returned false but Spline failed to yield an FTangentSpline!"));
-		return 0.f;
-	}
+	return ShouldUseSplineCurves()
+		? WarninglessSplineCurves().Position.FindNearestOnSegment(LocalLocation, Index, Dummy)
+		: Spline.FindNearestOnSegment(LocalLocation, Index, Dummy);
 }
 
 TTuple<int, float> FPCGSplineStruct::GetSegmentStartIndexAndKeyAtInputKey(float InKey) const
 {
-	const int32 Index = GetSplinePointsPosition().GetPointIndexForInputValue(InKey);
+	const int32 Index = ShouldUseSplineCurves()
+		? WarninglessSplineCurves().Position.GetPointIndexForInputValue(InKey)
+		: Spline.GetSplinePointsPosition().GetPointIndexForInputValue(InKey);
 	return {Index, GetInputKeyAtSegmentStart(Index)};
 }
 
 float FPCGSplineStruct::GetInputKeyAtSegmentStart(int InSegmentIndex) const
 {
-	const auto& SplinePointsPosition = GetSplinePointsPosition();
+	const auto& SplinePointsPosition = ShouldUseSplineCurves()
+		? WarninglessSplineCurves().Position
+		: Spline.GetSplinePointsPosition();
 
 	if (IsClosedLoop() && !SplinePointsPosition.Points.IsEmpty() && InSegmentIndex >= SplinePointsPosition.Points.Num())
 	{
@@ -974,20 +719,9 @@ FVector FPCGSplineStruct::GetLocationAtSplineInputKey(float InKey, ESplineCoordi
 
 FQuat FPCGSplineStruct::GetQuaternionAtSplineInputKey(float InKey, ESplineCoordinateSpace::Type CoordinateSpace) const
 {
-	FQuat Quat = FQuat::Identity;
-	if (ShouldUseSplineCurves())
-	{
-		Quat = WarninglessSplineCurves().Rotation.Eval(InKey, FQuat::Identity);
-	}
-	else if (const FTangentSpline* TangentSpline = Spline.Get<UE::Spline::ESplineType::Tangent>())
-	{
-		Quat = TangentSpline->EvaluateRotation(InKey);
-	}
-	else
-	{
-		ensureAlwaysMsgf(false, TEXT("ShouldUseSplineCurves() returned false but Spline failed to yield an FTangentSpline!"));
-		return Quat;
-	}
+	FQuat Quat = ShouldUseSplineCurves()
+		? WarninglessSplineCurves().Rotation.Eval(InKey, FQuat::Identity)
+		: Spline.EvaluateRotation(InKey);
 	Quat.Normalize();
 
 	FVector Direction = ShouldUseSplineCurves()
@@ -1008,22 +742,9 @@ FQuat FPCGSplineStruct::GetQuaternionAtSplineInputKey(float InKey, ESplineCoordi
 
 FVector FPCGSplineStruct::GetScaleAtSplineInputKey(float InKey) const
 {
-	FVector Scale = FVector::OneVector;
-
-	if (ShouldUseSplineCurves())
-	{
-		Scale = WarninglessSplineCurves().Scale.Eval(InKey, FVector::OneVector);
-	}
-	else if (const FTangentSpline* TangentSpline = Spline.Get<UE::Spline::ESplineType::Tangent>())
-	{
-		Scale = TangentSpline->EvaluateScale(InKey);
-	}
-	else
-	{
-		ensureAlwaysMsgf(false, TEXT("ShouldUseSplineCurves() returned false but Spline failed to yield an FTangentSpline!"));
-	}
-
-	return Scale;
+	return ShouldUseSplineCurves()
+		? WarninglessSplineCurves().Scale.Eval(InKey, FVector(1.0f))
+		: Spline.GetSplinePointsScale().Eval(InKey, FVector(1.0f));
 }
 
 FVector FPCGSplineStruct::GetTangentAtSplineInputKey(float InKey, ESplineCoordinateSpace::Type CoordinateSpace) const
@@ -1170,9 +891,9 @@ void FPCGSplineStruct::PostSerialize(const FArchive& Ar)
 #if WITH_EDITOR
 	if (Ar.IsLoading() && Ar.IsPersistent() && !Ar.HasAnyPortFlags(PPF_Duplicate | PPF_DuplicateForPIE))
 	{
-		// If PCGSplineStruct.UseSplineCurves is false and Spline was not previously implemented
-		// (implying it is not consistent with SplineCurves), initialize Spline from SplineCurves.
-		if (Spline.GetPreviousImplementation() == UE::Spline::ESplineType::Unimplemented && !PCGSplineStruct::GUseSplineCurves)
+		// This is kind of a hack. It will be true for every FPCGSplineStruct which did not have an FSpline before but does not.
+		// It is also going to be true if we already had a Spline but did not save it while it was configured to actually do anything.
+		if (!Spline.WasEnabled())
 		{
 			Spline = WarninglessSplineCurves();
 		}
@@ -1182,5 +903,5 @@ void FPCGSplineStruct::PostSerialize(const FArchive& Ar)
 
 bool FPCGSplineStruct::ShouldUseSplineCurves() const
 {
-	return PCGSplineStruct::GUseSplineCurves || Spline.GetCurrentImplementation() != UE::Spline::ESplineType::Tangent;
+	return PCGSplineStruct::GUseSplineCurves || Spline.IsEnabled() == false;
 }
\ No newline at end of file
diff --git a/Engine/Source/Editor/ComponentVisualizers/Private/SplineComponentVisualizer.cpp b/Engine/Source/Editor/ComponentVisualizers/Private/SplineComponentVisualizer.cpp
--- a/Engine/Source/Editor/ComponentVisualizers/Private/SplineComponentVisualizer.cpp
+++ b/Engine/Source/Editor/ComponentVisualizers/Private/SplineComponentVisualizer.cpp
@@ -3576,16 +3576,13 @@ void FSplineComponentVisualizer::OnAddAttributeKey()
 	
 	const float NearestParam = SplineComp->FindInputKeyClosestToWorldLocation(SelectionState->GetSelectedSplinePosition());
 	const float CurrentValue = SplineComp->GetFloatPropertyAtSplineInputKey(NearestParam, SelectionState->GetSelectedAttributeName());
-	const int32 NewAttributeIndex = SplineComp->AddFloatPropertyAtSplineInputKey(NearestParam, CurrentValue, SelectionState->GetSelectedAttributeName());
+	const int32 NewAttributeIndex = SplineComp->SetFloatPropertyAtSplineInputKey(NearestParam, CurrentValue, SelectionState->GetSelectedAttributeName());
 
 	// Update selection
-	if (NewAttributeIndex != INDEX_NONE)
-	{
-		SelectionState->SetSelectedSegmentIndex(INDEX_NONE);
-		SelectionState->SetSelectedAttribute(NewAttributeIndex, SelectionState->GetSelectedAttributeName());
-		SelectionState->SetSelectedSplinePosition(FVector::ZeroVector);
-		SelectionState->SetCachedRotation(SplineComp->GetQuaternionAtSplineInputKey(NearestParam, ESplineCoordinateSpace::World));
-	}
+	SelectionState->SetSelectedSegmentIndex(INDEX_NONE);
+	SelectionState->SetSelectedAttribute(NewAttributeIndex, SelectionState->GetSelectedAttributeName());
+	SelectionState->SetSelectedSplinePosition(FVector::ZeroVector);
+	SelectionState->SetCachedRotation(SplineComp->GetQuaternionAtSplineInputKey(NearestParam, ESplineCoordinateSpace::World));
 
 	SplineComp->bSplineHasBeenEdited = true;
 	NotifyPropertiesModified(SplineComp, SplineProperties, EPropertyChangeType::ArrayAdd);
diff --git a/Engine/Source/Editor/DetailCustomizations/Private/SplineComponentDetails.cpp b/Engine/Source/Editor/DetailCustomizations/Private/SplineComponentDetails.cpp
--- a/Engine/Source/Editor/DetailCustomizations/Private/SplineComponentDetails.cpp
+++ b/Engine/Source/Editor/DetailCustomizations/Private/SplineComponentDetails.cpp
@@ -1005,7 +1005,7 @@ bool FSplinePointDetails::CanSetInputKey() const
 	}
 
 PRAGMA_DISABLE_DEPRECATION_WARNINGS
-	bool bUsingSplineCurves = SplineComp->ShouldUseSplineCurves();
+	bool bUsingSplineCurves = SplineComp->GetSplinePropertyName() == GET_MEMBER_NAME_CHECKED(USplineComponent, SplineCurves);
 PRAGMA_ENABLE_DEPRECATION_WARNINGS
 
 	return IsOnePointSelected() && bUsingSplineCurves;
diff --git a/Engine/Source/Runtime/Engine/Classes/Components/SplineComponent.h b/Engine/Source/Runtime/Engine/Classes/Components/SplineComponent.h
--- a/Engine/Source/Runtime/Engine/Classes/Components/SplineComponent.h
+++ b/Engine/Source/Runtime/Engine/Classes/Components/SplineComponent.h
@@ -76,7 +76,15 @@ struct FSplineCurves
 {
 	GENERATED_BODY()
 
-	static FSplineCurves FromSplineInterface(const FSpline& InProxy);
+	static FSplineCurves FromSplineInterface(const FSpline& InProxy)
+	{
+		FSplineCurves NewCurve;
+		NewCurve.Position = InProxy.GetSplinePointsPosition();
+		NewCurve.Rotation = InProxy.GetSplinePointsRotation();
+		NewCurve.Scale = InProxy.GetSplinePointsScale();
+		NewCurve.UpdateSpline(NewCurve.Position.bIsLooped, false);
+		return NewCurve;
+	}
 	
 	/** Spline built from position data. */
 	UPROPERTY()
@@ -222,13 +230,15 @@ protected:
 	
 public:
 
+	UFUNCTION(CallInEditor, Category="Spline")
+	void PopulateFromLegacy();
+
 	/** Return the name of the authoritative spline property. */
 	ENGINE_API static FName GetSplinePropertyName();
 
 	/** Return the set of names of the backing spline properties. */
 	ENGINE_API static TSet<FName> GetSplinePropertyNames();
 	
-	UE_DEPRECATED(5.8, TEXT("SplineCurves will be made protected."))
 	UPROPERTY(EditAnywhere, Replicated, Category=Points)
 	FSplineCurves SplineCurves;
 
@@ -406,6 +416,9 @@ public:
 
 	UE_EXPERIMENTAL(5.7, "Spline display setting change notifications are experimental while the API is under active development.")
 	FSimpleMulticastDelegate& GetOnSplineDisplayChanged() { return OnSplineDisplayChanged; }
+	
+	/** Reset all spline data to match the current state of SplineCurves */
+	void SynchronizeSplines() { SetSpline(WarninglessSplineCurves()); }
 
 	ENGINE_API int32 GetVersion() const;
 	ENGINE_API const FInterpCurveVector& GetSplinePointsPosition() const;
@@ -953,8 +966,8 @@ public:
 	 * @return The index of the new attribute.
 	 */
 	UFUNCTION(BlueprintCallable, Category = "Spline|Metadata")
-	ENGINE_API int32 AddFloatPropertyAtSplineInputKey(float InKey, float Value, FName PropertyName);
-
+	ENGINE_API int32 SetFloatPropertyAtSplineInputKey(float InKey, float Value, FName PropertyName) const;
+	
 	/**
 	 * Get a float attribute parameter by index.
 	 * @param Index The index of the specified attribute.
@@ -1005,18 +1018,6 @@ public:
 	FInterpCurveQuat& GetSplinePointsRotation() { return WarninglessSplineCurves().Rotation; }
 	FInterpCurveVector& GetSplinePointsScale() { return WarninglessSplineCurves().Scale; }
 
-public: // Deprecated Interface
-
-	UE_DEPRECATED(5.8, TEXT("SynchronizeSplines is deprecated."))
-	void SynchronizeSplines() {}
-
-	UE_DEPRECATED(5.8, TEXT("PopulateFromLegacy is deprecated."))
-	void PopulateFromLegacy() {}
-
-	UE_DEPRECATED(5.8, TEXT("SetFloatPropertyAtSplineInputKey is deprecated, please use AddFloatPropertyAtSplineInputKey"))
-	UFUNCTION(BlueprintCallable, Category = "Spline|Metadata", meta = (DeprecatedFunction, DeprecationMessage = "Please use AddFloatPropertyAtSplineInputKey"))
-	ENGINE_API int32 SetFloatPropertyAtSplineInputKey(float InKey, float Value, FName PropertyName) const { return INDEX_NONE; }
-
 protected:
 
 	/** Returns the length of the specified spline segment up to the parametric value given */
@@ -1033,9 +1034,6 @@ protected:
 	
 private:
 	
-	/** Whether or not FSplineCurves is the authoritative data. If false, we assume Spline is implemented as FTangentSpline. */
-	ENGINE_API bool ShouldUseSplineCurves() const;
-
 	/** Returns an interval representing the full range of valid input keys, or an empty interval if the spline is empty. */
 	UE::Geometry::FInterval1f GetInputKeyRange() const;
 	
@@ -1054,9 +1052,9 @@ private:
 	}
 	
 	/** The dummy value used for queries when there are no point in a spline */
-	static ENGINE_API const FInterpCurveVector DummyPointPosition;
-	static ENGINE_API const FInterpCurveQuat DummyPointRotation;
-	static ENGINE_API const FInterpCurveVector DummyPointScale;
+	static ENGINE_API const FInterpCurvePointVector DummyPointPosition;
+	static ENGINE_API const FInterpCurvePointQuat DummyPointRotation;
+	static ENGINE_API const FInterpCurvePointVector DummyPointScale;
 
 	// Remove when appending _DEPRECATED to SplineCurves.
 	FSplineCurves& WarninglessSplineCurves()
@@ -1094,9 +1092,6 @@ PRAGMA_ENABLE_DEPRECATION_WARNINGS
 	// FSplineComponentVisualizer will access some private members when attempting to call NotifyPropertiesModified
 	friend class FSplineComponentVisualizer;
 
-	// FSplinePointDetails will access ShouldUseSplineCurves for enabling/disabling some editing functionality
-	friend class FSplinePointDetails;
-
 	// AFunctionalSplineTestBase accesses private methods for functional testing.
 	friend class AFunctionalSplineTestBase;
 };
diff --git a/Engine/Source/Runtime/Engine/Classes/Curves/Spline.h b/Engine/Source/Runtime/Engine/Classes/Curves/Spline.h
--- a/Engine/Source/Runtime/Engine/Classes/Curves/Spline.h
+++ b/Engine/Source/Runtime/Engine/Classes/Curves/Spline.h
@@ -2,37 +2,18 @@
 
 #pragma once
 
+#include "CoreMinimal.h"
+#include "Math/InterpCurve.h"
 #include "Splines/SplineInterfaces.h"
-#include "Curves/Splines/TangentSpline.h"	// only for FLegacyTangentSpline, move that to its own header.
 
 #include "Spline.generated.h"
 
-struct UE_EXPERIMENTAL(5.7, "New spline APIs are experimental.") FSpline;
-
-namespace UE::Spline
-{
-enum class ESplineType : uint8
-{
-	Unimplemented = 0,
-	LegacyTangent = 1,
-	Tangent = 2,
-};
-
-template<ESplineType Impl>
-struct TSplineTypeMap;
-
-template<>
-struct TSplineTypeMap<ESplineType::LegacyTangent>
-{
-	using Type = FLegacyTangentSpline;
-};
+struct FSplinePoint;
+struct FSplineCurves;
+class FLegacySpline;
+class FNewSpline;
 
-template<>
-struct TSplineTypeMap<ESplineType::Tangent>
-{
-	using Type = FTangentSpline;
-};
-} // UE::Spline
+struct UE_EXPERIMENTAL(5.7, "New spline APIs are experimental.") FSpline;
 
 /**
  * A general purpose, reflected spline.
@@ -44,125 +25,186 @@ struct FSpline
 	GENERATED_BODY()
 
 	ENGINE_API FSpline();
-	ENGINE_API explicit FSpline(UE::Spline::ESplineType InType);
-
 	ENGINE_API FSpline(const FSpline& Other);
 	ENGINE_API FSpline& operator=(const FSpline& Other);
-
-	ENGINE_API FSpline(FSpline&& Other) = default;
-	ENGINE_API FSpline& operator=(FSpline&& Other) = default;
-
 	ENGINE_API FSpline& operator=(const FSplineCurves& Other);
+	
+	/* Control Point Index Interface */
+	
+	/** Adds point by parameter. The new point is inserted after all points with parameter values less than OR EQUAL to the new point's parameter.
+	 * Nick would like to change this behavior in the future. See UE-250236. */
+	ENGINE_API void AddPoint(const FSplinePoint& InPoint);
+	/** Adds point by index. Provided parameter is ignored. The new point is inserted before all points with parameter values greater than the new point's parameter. */
+	ENGINE_API void InsertPoint(const FSplinePoint& InPoint, int32 Index);
+	ENGINE_API FSplinePoint GetPoint(const int32 Index) const;
+	ENGINE_API void RemovePoint(const int32 Index);
+	
+	ENGINE_API void SetLocation(int32 Index, const FVector& InLocation);
+	ENGINE_API FVector GetLocation(const int32 Index) const;
+	
+	ENGINE_API void SetInTangent(const int32 Index, const FVector& InTangent);
+	ENGINE_API FVector GetInTangent(const int32 Index) const;
+	
+	ENGINE_API void SetOutTangent(const int32 Index, const FVector& OutTangent);
+	ENGINE_API FVector GetOutTangent(const int32 Index) const;
 
-	ENGINE_API FVector Evaluate(float Param) const;
+	ENGINE_API void SetRotation(int32 Index, const FQuat& InRotation);
+	ENGINE_API FQuat GetRotation(const int32 Index) const;
+	
+	ENGINE_API void SetScale(int32 Index, const FVector& InScale);
+	ENGINE_API FVector GetScale(const int32 Index) const;
+	
+	ENGINE_API void SetSplinePointType(int32 Index, EInterpCurveMode Type);
+	ENGINE_API EInterpCurveMode GetSplinePointType(int32 Index) const;
 
-	ENGINE_API FVector EvaluateDerivative(float Param) const;
+	ENGINE_API float GetParameterAtIndex(int32 Index) const;
+	ENGINE_API float GetParameterAtDistance(float Distance) const;
+	ENGINE_API float GetDistanceAtParameter(float Parameter) const;
 
-	ENGINE_API int32 GetNumberOfSegments() const;
+	ENGINE_API FQuat GetOrientation(int32 Index) const;
+	ENGINE_API void SetOrientation(int32 Index, const FQuat& InOrientation);
+	
+	/* Parameter Interface */
 
-	/** Reset the spline to an empty spline. Preserves the current implementation. */
-	ENGINE_API void Reset();
+	ENGINE_API FVector Evaluate(float Param) const;
+	ENGINE_API FVector EvaluateDerivative(float Param) const;
+	ENGINE_API FQuat EvaluateRotation(float Param) const;
+	ENGINE_API FVector EvaluateScale(float Param) const;
 
-	/**
-	 * Fetch the underlying concrete spline implementation.
-	 *
-	 * Pointers returned by Get<Impl>() are invalidated by:
-	 * - Serialize(FArchive&) when Ar.IsLoading() is true
-	 * - ImportTextItem(...)
-	 * - operator=(...)
-	 * - SetCurrentImplementation(...)
-	 * - destruction of this FSpline
-	 *
-	 * @return A valid pointer if Impl == GetCurrentImplementation(), otherwise nullptr.
-	 */
-	template<UE::Spline::ESplineType Impl>
-	typename UE::Spline::TSplineTypeMap<Impl>::Type* Get();
-
-	/** A const version of Get<Impl> */
-	template<UE::Spline::ESplineType Impl>
-	const typename UE::Spline::TSplineTypeMap<Impl>::Type* Get() const;
-		
+	ENGINE_API FQuat GetOrientation(float Param) const;
+	
+	/* Attribute Interface */
+
+	// Non-Templated functions
+	bool SupportsAttributes() const { return IsNew(); }
+	ENGINE_API bool HasAttributeChannel(FName AttributeName) const;
+	ENGINE_API bool RemoveAttributeChannel(FName AttributeName) const;
+	ENGINE_API TArray<FName> GetFloatPropertyChannels() const;
+	ENGINE_API TArray<FName> GetVectorPropertyChannels() const;
+
+	template <typename AttrType> ENGINE_API float GetAttributeParameter(int32 Index, const FName& Name) const;
+	template <typename AttrType> ENGINE_API int32 SetAttributeParameter(int32 Index, float Parameter, const FName& Name);
+	template <typename AttrType> ENGINE_API int32 NumAttributeValues(FName AttributeName) const;
+	template <typename AttrType> ENGINE_API AttrType GetAttributeValue(int32 Index, const FName& Name) const;
+	template <typename AttrType> ENGINE_API void SetAttributeValue(int32 Index, const AttrType& Value, const FName& Name);
+	template <typename AttrType> ENGINE_API bool CreateAttributeChannel(FName AttributeName) const;
+	template <typename AttrType> ENGINE_API int32 AddAttributeValue(float Param, const AttrType& Value, FName AttributeName) const;
+	template <typename AttrType> ENGINE_API void RemoveAttributeValue(int32 Index, FName AttributeName);
+	template <typename AttrType> ENGINE_API AttrType EvaluateAttribute(float Param, FName AttributeName) const;
+
+	ENGINE_API float FindNearest(const FVector& InLocation, float& OutSquaredDist) const;
+	ENGINE_API float FindNearestOnSegment(const FVector& InLocation, int32 SegmentIndex, float& OutSquaredDist) const;
+	
+	/* Misc Interface */
+	
 	ENGINE_API bool operator==(const FSpline& Other) const;
 	ENGINE_API bool operator!=(const FSpline& Other) const;
+
+	ENGINE_API bool IsNearlyEqual(const FSpline& Other, const float RelativeTolerance, const float AbsoluteTolerance) const;
 	
 	friend FArchive& operator<<(FArchive& Ar, FSpline& Spline)
-	{ 
+	{
 		Spline.Serialize(Ar);
 		return Ar;
 	}
-
 	ENGINE_API bool Serialize(FArchive& Ar);
+	ENGINE_API void SerializeLoad(FArchive& Ar);
+	ENGINE_API void SerializeSave(FArchive& Ar) const;
 	ENGINE_API bool ExportTextItem(FString& ValueStr, FSpline const& DefaultValue, class UObject* Parent, int32 PortFlags, class UObject* ExportRootScope) const;
 	ENGINE_API bool ImportTextItem(const TCHAR*& Buffer, int32 PortFlags, UObject* Parent, FOutputDevice* ErrorText);
 
-	ENGINE_API UE::Spline::ESplineType GetCurrentImplementation() const;
+	uint32 GetVersion() const { return Version; }
+	
+	ENGINE_API const FInterpCurveVector& GetSplinePointsPosition() const;
+	ENGINE_API const FInterpCurveQuat& GetSplinePointsRotation() const;
+	ENGINE_API const FInterpCurveVector& GetSplinePointsScale() const;
 
-#if WITH_EDITOR
-	/** Returns the implementation of this spline when when most recently loaded. */
-	ENGINE_API UE::Spline::ESplineType GetPreviousImplementation() const;
-#endif
+	/** Returns the length of the specified spline segment up to the parametric value given */
+	ENGINE_API float GetSegmentLength(const int32 Index, const float Param, const FVector& Scale3D = FVector(1.0f)) const;
 
-	/** Returns the default implementation for new spline. */
-	ENGINE_API static UE::Spline::ESplineType GetDefaultImplementation();
+	/** Returns total length along this spline */
+	ENGINE_API float GetSplineLength() const;
 
-private:
+	/** Returns the total number of segments on this spline. */
+	ENGINE_API int32 GetNumSegments() const;
+	
+	/** Returns the total number of control points on this spline. */
+	ENGINE_API int32 GetNumControlPoints() const;
+	
+	/** Reset the spline to an empty spline. */
+	ENGINE_API void Reset();
 	
-	/** Sets the current implementation and performs any necessary allocation. */
-	void SetCurrentImplementation(UE::Spline::ESplineType Implementation);
+	/** Reset the rotation attribute channel to default values. */
+	ENGINE_API void ResetRotation();
+	
+	/** Reset the scale attribute channel to default values. */
+	ENGINE_API void ResetScale();
 
-	void SerializeLoad(FArchive& Ar);
-	void SerializeSave(FArchive& Ar) const;
+	ENGINE_API void SetClosedLoop(bool bClosed);
+	ENGINE_API bool IsClosedLoop() const;
+	
+	struct FUpdateSplineParams
+	{
+		bool bClosedLoop = false;
+		bool bStationaryEndpoints = false;
+		int32 ReparamStepsPerSegment = 10;
+		bool bLoopPositionOverride = false;
+		float LoopPosition = 0.0f;
+		FVector Scale3D = FVector(1.0f);
+	};
+	
+	/** Update the spline's internal data according to the passed-in params. */
+	ENGINE_API void UpdateSpline(const FUpdateSplineParams& InParams);
 
-private:
+	/** Update the spline's internal data according to the most recently used update params (or default params if never updated). */
+	ENGINE_API void UpdateSpline();
 
-	UE::Spline::ESplineType CurrentImplementation;
+	ENGINE_API TSharedPtr<UE::Geometry::Spline::TSplineInterface<FVector>> GetSplineInterface() const;
 
+private:
+	
+	static ENGINE_API const FInterpCurveVector PositionCurve;
+	static ENGINE_API const FInterpCurveQuat RotationCurve;
+	static ENGINE_API const FInterpCurveVector ScaleCurve;
+	
+	// Used for upgrade logic in spline component.
+	// Not ideal, but allows us to automatically populate the proxy
+	// at serialize time when we might otherwise not be able to.
+	friend class USplineComponent;
+	friend struct FPCGSplineStruct;
+	
 #if WITH_EDITOR
-	UE::Spline::ESplineType PreviousImplementation;
+	uint8 PreviousImplementation;
 #endif
+	uint8 CurrentImplementation;
+	
+	uint32 Version;
 
-	// Valid when CurrentImplementation is ESplineType::LegacyTangent
-	TUniquePtr<FLegacyTangentSpline> LegacySplineImpl;
-
-	// Valid when CurrentImplementation is not ESplineType::Unimplemented and not ESplineType::LegacyTangent;
-	TUniquePtr<UE::Geometry::Spline::TSplineInterface<FVector>> SplineImpl;
-};
+	FUpdateSplineParams CachedUpdateSplineParams;
+	
+	// probably better implemented as a TSharedPtr<ISplineInterface> or something
+	struct {													// Invalid when CurrentImplementation is 0 or 3.
+		TSharedPtr<FLegacySpline> LegacyData;					// Valid when CurrentImplementation is 1.
+		TSharedPtr<FNewSpline> NewData;							// Valid when CurrentImplementation is 2.
+	};
+	
+	bool IsEnabled() const { return CurrentImplementation != 0; }
+	bool IsLegacy() const { return CurrentImplementation == 1; }
+	bool IsNew() const { return CurrentImplementation == 2; }
 
-template<UE::Spline::ESplineType Impl>
-typename UE::Spline::TSplineTypeMap<Impl>::Type* FSpline::Get()
-{
-	if (Impl != CurrentImplementation)
-	{
-		return nullptr;
-	}
+#if WITH_EDITOR
+	bool WasEnabled() const { return PreviousImplementation != 0; }
+	bool WasLegacy() const { return PreviousImplementation == 1; }
+	bool WasNew() const { return PreviousImplementation == 2; }
 
-	if constexpr (Impl == UE::Spline::ESplineType::LegacyTangent)
-	{
-		return LegacySplineImpl.Get();
-	}
-	else
-	{
-		return static_cast<typename UE::Spline::TSplineTypeMap<Impl>::Type*>(SplineImpl.Get());
-	}
-}
+	void OnSplineImplementationChanged();
+#endif
 
-template<UE::Spline::ESplineType Impl>
-const typename UE::Spline::TSplineTypeMap<Impl>::Type* FSpline::Get() const
-{
-	if (Impl != CurrentImplementation)
-	{
-		return nullptr;
-	}
+public:
 
-	if constexpr (Impl == UE::Spline::ESplineType::LegacyTangent)
-	{
-		return LegacySplineImpl.Get();
-	}
-	else
-	{
-		return static_cast<const typename UE::Spline::TSplineTypeMap<Impl>::Type*>(SplineImpl.Get());
-	}
-}
+	// True if a given FSpline will actually be implemented when instantiated.
+	static bool IsEnabledStatic();
+};
 
 template<>
 struct TStructOpsTypeTraits<FSpline> : public TStructOpsTypeTraitsBase2<FSpline>
diff --git a/Engine/Source/Runtime/Engine/Classes/Curves/Splines/TangentSpline.h b/Engine/Source/Runtime/Engine/Classes/Curves/Splines/TangentSpline.h
deleted file mode 100644
diff --git a/Engine/Source/Runtime/Engine/Private/Components/SplineComponent.cpp b/Engine/Source/Runtime/Engine/Private/Components/SplineComponent.cpp
--- a/Engine/Source/Runtime/Engine/Private/Components/SplineComponent.cpp
+++ b/Engine/Source/Runtime/Engine/Private/Components/SplineComponent.cpp
@@ -5,7 +5,6 @@
 =============================================================================*/
 
 #include "Components/SplineComponent.h"
-#include "Curves/Splines/TangentSpline.h"
 #include "Engine/Engine.h"
 #include "SceneView.h"
 #include "UObject/EditorObjectVersion.h"
@@ -28,6 +27,7 @@
 #include "Misc/TransactionObjectEvent.h"
 #include "UObject/ICookInfo.h"
 
+
 #include UE_INLINE_GENERATED_CPP_BY_NAME(SplineComponent)
 
 #if WITH_EDITOR
@@ -38,9 +38,9 @@
 
 DEFINE_LOG_CATEGORY_STATIC(LogSplineComponent, Log, All);
 
-const FInterpCurveVector USplineComponent::DummyPointPosition = {};
-const FInterpCurveQuat USplineComponent::DummyPointRotation = {};
-const FInterpCurveVector USplineComponent::DummyPointScale = {};
+const FInterpCurvePointVector USplineComponent::DummyPointPosition(0.0f, FVector::ZeroVector, FVector::ForwardVector, FVector::ForwardVector, CIM_Constant);
+const FInterpCurvePointQuat USplineComponent::DummyPointRotation(0.0f, FQuat::Identity);
+const FInterpCurvePointVector USplineComponent::DummyPointScale(0.0f, FVector::OneVector);
 
 namespace UE::SplineComponent
 {
@@ -132,12 +132,12 @@ static ELastAuthority GetEffectiveLastAuthority(const ELastAuthority InLastAutho
 	}
 	return UseLastAuthority;
 }
-	
-bool ShouldUseSplineCurvesStatic()
+
+static bool ShouldUseSplineCurves()
 {
-	return GUseSplineCurves || FSpline::GetDefaultImplementation() != UE::Spline::ESplineType::Tangent;
+	return GUseSplineCurves || FSpline::IsEnabledStatic() == false;
 }
-
+	
 }
 
 USplineMetadata::USplineMetadata(const FObjectInitializer& ObjectInitializer)
@@ -147,19 +147,12 @@ USplineMetadata::USplineMetadata(const FObjectInitializer& ObjectInitializer)
 
 bool USplineComponent::Validate() const
 {
-	if (Spline.GetCurrentImplementation() != UE::Spline::ESplineType::Tangent)
+	if (!Spline.IsEnabled())
 	{
 		// Can't be inconsistent if Spline is incapable of storing any data at all.
 		return true;
 	}
 
-	const FTangentSpline* TangentSpline = Spline.Get<UE::Spline::ESplineType::Tangent>();
-	if (!ensureAlways(TangentSpline))
-	{
-		// Something is wrong...
-		return false;
-	}
-
 	auto StringifySplinePointType = [](const ESplinePointType::Type InSplinePointType)
 	{
 		switch (InSplinePointType)
@@ -176,10 +169,10 @@ bool USplineComponent::Validate() const
 	bool bValid = true;
 	
 	const int32 NumSplineCurvesPoints = WarninglessSplineCurves().Position.Points.Num();
-	const int32 NumSplinePoints = TangentSpline->GetNumberOfPoints();
+	const int32 NumSplinePoints = Spline.GetNumControlPoints();
 	const int32 NumPoints = FMath::Min(NumSplinePoints, NumSplineCurvesPoints);
 
-	UE_LOG(LogSplineComponent, Display, TEXT("Validating SplineComponent (%s) - Authority: %s"), *GetPathName(), ShouldUseSplineCurves() ? TEXT("FSplineCurves") : TEXT("FSpline"));
+	UE_LOG(LogSplineComponent, Display, TEXT("Validating SplineComponent (%s) - Authority: %s"), *GetPathName(), UE::SplineComponent::ShouldUseSplineCurves() ? TEXT("FSplineCurves") : TEXT("FSpline"));
 	
 	if (NumSplineCurvesPoints != NumSplinePoints)
 	{
@@ -187,15 +180,6 @@ bool USplineComponent::Validate() const
 		bValid = false;
 	}
 
-	const float SplineCurvesLength = WarninglessSplineCurves().GetSplineLength();
-	const float SplineLength = TangentSpline->GetDistanceAtParameter(TangentSpline->GetParameterSpace().Max);
-	if (!FMath::IsNearlyEqual(SplineCurvesLength, SplineLength))
-	{
-		UE_LOG(LogSplineComponent, Display, TEXT("Internally inconsistent spline lengths (%f vs %f)!"), SplineCurvesLength, SplineLength);
-		bValid = false;
-	}
-	UE_LOG(LogSplineComponent, Display, TEXT("Spline lengths: (%f vs %f)"), SplineCurvesLength, SplineLength);
-
 	for (int32 Idx = 0; Idx < NumPoints; ++Idx)
 	{
 		const auto& SplineCurvesPositionPoint = WarninglessSplineCurves().Position.Points[Idx];
@@ -203,7 +187,7 @@ bool USplineComponent::Validate() const
 		const auto& SplineCurvesScalePoint = WarninglessSplineCurves().Scale.Points[Idx];
 
 		const ESplinePointType::Type SplineCurvesPointType = ConvertInterpCurveModeToSplinePointType(SplineCurvesPositionPoint.InterpMode);
-		const ESplinePointType::Type SplinePointType = ConvertInterpCurveModeToSplinePointType(TangentSpline->GetSplinePointType(Idx));
+		const ESplinePointType::Type SplinePointType = ConvertInterpCurveModeToSplinePointType(Spline.GetSplinePointType(Idx));
 		if (SplineCurvesPointType != SplinePointType)
 		{
 			UE_LOG(LogSplineComponent, Display, TEXT("Internally inconsistent point type (%s vs %s) at index %d!"), StringifySplinePointType(SplineCurvesPointType), StringifySplinePointType(SplinePointType), Idx );
@@ -211,7 +195,7 @@ bool USplineComponent::Validate() const
 		}
 		
 		const FVector SplineCurvesPosition = SplineCurvesPositionPoint.OutVal;
-		const FVector SplinePosition = TangentSpline->GetLocation(Idx);
+		const FVector SplinePosition = Spline.GetLocation(Idx);
 		if (!SplineCurvesPosition.Equals(SplinePosition))
 		{
 			UE_LOG(LogSplineComponent, Display, TEXT("Internally inconsistent position (%s vs %s) at index %d!"), *SplineCurvesPosition.ToCompactString(), *SplinePosition.ToCompactString(), Idx );
@@ -219,7 +203,7 @@ bool USplineComponent::Validate() const
 		}
 
 		const FVector SplineCurvesInTangent = SplineCurvesPositionPoint.ArriveTangent;
-		const FVector SplineInTangent = TangentSpline->GetInTangent(Idx);
+		const FVector SplineInTangent = Spline.GetInTangent(Idx);
 		if (!SplineCurvesInTangent.Equals(SplineInTangent))
 		{
 			UE_LOG(LogSplineComponent, Display, TEXT("Internally inconsistent arrive tangent (%s vs %s) at index %d!"), *SplineCurvesInTangent.ToCompactString(), *SplineInTangent.ToCompactString(), Idx );
@@ -236,7 +220,7 @@ bool USplineComponent::Validate() const
 		}
 
 		const FVector SplineCurvesOutTangent = SplineCurvesPositionPoint.LeaveTangent;
-		const FVector SplineOutTangent = TangentSpline->GetOutTangent(Idx);
+		const FVector SplineOutTangent = Spline.GetOutTangent(Idx);
 		if (!SplineCurvesOutTangent.Equals(SplineOutTangent))
 		{
 			UE_LOG(LogSplineComponent, Display, TEXT("Internally inconsistent leave tangent (%s vs %s) at index %d!"), *SplineCurvesOutTangent.ToCompactString(), *SplineOutTangent.ToCompactString(), Idx );
@@ -253,7 +237,7 @@ bool USplineComponent::Validate() const
 		}
 
 		const float SplineCurvesParam = SplineCurvesPositionPoint.InVal;
-		const float SplineParam = TangentSpline->GetParameterAtIndex(Idx);
+		const float SplineParam = Spline.GetParameterAtIndex(Idx);
 		if (!FMath::IsNearlyEqual(SplineCurvesParam, SplineParam))
 		{
 			UE_LOG(LogSplineComponent, Display, TEXT("Internally inconsistent parameter (%f vs %f) at index %d!"), SplineCurvesParam, SplineParam, Idx );
@@ -261,7 +245,7 @@ bool USplineComponent::Validate() const
 		}
 
 		const FQuat SplineCurvesRotation = SplineCurvesRotationPoint.OutVal;
-		const FQuat SplineRotation = TangentSpline->GetRotation(Idx);
+		const FQuat SplineRotation = Spline.GetRotation(Idx);
 		if (!SplineCurvesRotation.Equals(SplineRotation))
 		{
 			UE_LOG(LogSplineComponent, Display, TEXT("Internally inconsistent rotation (%s vs %s) at index %d!"), *SplineCurvesRotation.ToString(), *SplineRotation.ToString(), Idx );
@@ -269,7 +253,7 @@ bool USplineComponent::Validate() const
 		}
 
 		const FVector SplineCurvesScale = SplineCurvesScalePoint.OutVal;
-		const FVector SplineScale = TangentSpline->GetScale(Idx);
+		const FVector SplineScale = Spline.GetScale(Idx);
 		if (!SplineCurvesScale.Equals(SplineScale))
 		{
 			UE_LOG(LogSplineComponent, Display, TEXT("Internally inconsistent scale (%s vs %s) at index %d!"), *SplineCurvesScale.ToCompactString(), *SplineScale.ToCompactString(), Idx );
@@ -285,9 +269,14 @@ bool USplineComponent::Validate() const
 	return bValid;
 }
 
+void USplineComponent::PopulateFromLegacy()
+{
+	SynchronizeSplines();
+}
+
 FName USplineComponent::GetSplinePropertyName()
 {
-	if (UE::SplineComponent::ShouldUseSplineCurvesStatic())
+	if (UE::SplineComponent::ShouldUseSplineCurves())
 	{
 PRAGMA_DISABLE_DEPRECATION_WARNINGS
 		return GET_MEMBER_NAME_CHECKED(USplineComponent, SplineCurves);
@@ -346,9 +335,9 @@ USplineComponent::USplineComponent(const FObjectInitializer& ObjectInitializer)
 
 #if WITH_EDITORONLY_DATA
 	// Set these deprecated values up so that old assets with default values load correctly (and are subsequently upgraded during Serialize)
-	SplineInfo_DEPRECATED = WarninglessSplineCurves().Position;
-	SplineRotInfo_DEPRECATED = WarninglessSplineCurves().Rotation;
-	SplineScaleInfo_DEPRECATED = WarninglessSplineCurves().Scale;
+	SplineInfo_DEPRECATED = Spline.GetSplinePointsPosition();
+	SplineRotInfo_DEPRECATED = Spline.GetSplinePointsRotation();
+	SplineScaleInfo_DEPRECATED = Spline.GetSplinePointsScale();
 	SplineReparamTable_DEPRECATED = WarninglessSplineCurves().ReparamTable;
 #endif
 
@@ -398,39 +387,26 @@ FSpline USplineComponent::GetDefaultSpline()
 
 void USplineComponent::SetDefaultSpline()
 {
-	if (UE::SplineComponent::GUseSplineCurves)
-	{
-		Spline = FSpline(UE::Spline::ESplineType::Unimplemented);
-	}
-	else
-	{
-		Spline = FSpline(UE::Spline::ESplineType::Tangent);
-	}
-
 	const FVector StartPoint(0.0, 0.0, 0.0);
 	const FVector EndPoint(100.0, 0.0, 0.0);
 	constexpr float StartParam = 0.f;
 	constexpr float EndParam = 1.f;
 
+	Spline.Reset();
 	WarninglessSplineCurves().Position.Points.Reset(10);
     WarninglessSplineCurves().Rotation.Points.Reset(10);
     WarninglessSplineCurves().Scale.Points.Reset(10);
 
+	Spline.AddPoint(FSplinePoint(StartParam, StartPoint));
     WarninglessSplineCurves().Position.Points.Emplace(StartParam, StartPoint, FVector::ZeroVector, FVector::ZeroVector, CIM_CurveAuto);
     WarninglessSplineCurves().Rotation.Points.Emplace(StartParam, FQuat::Identity, FQuat::Identity, FQuat::Identity, CIM_CurveAuto);
     WarninglessSplineCurves().Scale.Points.Emplace(StartParam, FVector(1.0f), FVector::ZeroVector, FVector::ZeroVector, CIM_CurveAuto);
 
+	Spline.AddPoint(FSplinePoint(EndParam, EndPoint));
     WarninglessSplineCurves().Position.Points.Emplace(EndParam, EndPoint, FVector::ZeroVector, FVector::ZeroVector, CIM_CurveAuto);
     WarninglessSplineCurves().Rotation.Points.Emplace(EndParam, FQuat::Identity, FQuat::Identity, FQuat::Identity, CIM_CurveAuto);
     WarninglessSplineCurves().Scale.Points.Emplace(EndParam, FVector(1.0f), FVector::ZeroVector, FVector::ZeroVector, CIM_CurveAuto);
 
-	if (FTangentSpline* TangentSpline = Spline.Get<UE::Spline::ESplineType::Tangent>())
-	{
-		TangentSpline->Reset();
-		TangentSpline->AddPoint(FSplinePoint(StartParam, StartPoint));
-		TangentSpline->AddPoint(FSplinePoint(EndParam, EndPoint));
-	}
-
 	OnSplineChanged.Broadcast();
 
 	if (UE::SplineComponent::GValidateOnChange)
@@ -480,7 +456,7 @@ void USplineComponent::Serialize(FArchive& Ar)
 {
 	if (Ar.IsSaving())
 	{
-		LastAuthority = ShouldUseSplineCurves() ? ELastAuthority::SplineCurves : ELastAuthority::Spline;
+		LastAuthority = UE::SplineComponent::ShouldUseSplineCurves() ? ELastAuthority::SplineCurves : ELastAuthority::Spline;
 	}
 	
 	Super::Serialize(Ar);
@@ -500,15 +476,14 @@ void USplineComponent::Serialize(FArchive& Ar)
 		}
 
 #if WITH_EDITOR
-		// If SplineComponent.UseSplineCurves is false and Spline was not previously implemented
-		// (implying it is not consistent with SplineCurves), initialize Spline from SplineCurves.
-		if (Spline.GetPreviousImplementation() == UE::Spline::ESplineType::Unimplemented && !UE::SplineComponent::GUseSplineCurves)
+		// WasEnabled() is editor only, that data gets stripped from packaged builds.
+		if (!Spline.WasEnabled())
 		{
-			Spline = WarninglessSplineCurves();
+			SynchronizeSplines();
 		}
 #endif
 		
-		if (UE::SplineComponent::GSynchronizeOnLoad && Spline.GetCurrentImplementation() == UE::Spline::ESplineType::Tangent && !Validate())
+		if (UE::SplineComponent::GSynchronizeOnLoad && Spline.IsEnabled() && !Validate())
 		{
 			UE_LOG(LogSplineComponent, Display, TEXT("Internally inconsistent SplineComponent (%s) at load time!"), *GetPathName())
 			
@@ -531,7 +506,7 @@ void USplineComponent::Serialize(FArchive& Ar)
 
 			case ELastAuthority::SplineCurves:
 				UE_LOG(LogSplineComponent, Display, TEXT("Populating Spline from SplineCurves..."))
-				SetSpline(WarninglessSplineCurves());
+				SynchronizeSplines();
 				break;
 			}
 
@@ -547,7 +522,7 @@ void USplineComponent::Serialize(FArchive& Ar)
 				UE_LOG(LogSplineComponent, Display, TEXT("Successfully synchronized SplineComponent (%s) on load."), *GetPathName())
 			}
 		}
-		else if (UE::SplineComponent::GValidateOnLoad && Spline.GetCurrentImplementation() == UE::Spline::ESplineType::Tangent && !Validate())
+		else if (UE::SplineComponent::GValidateOnLoad && Spline.IsEnabled() && !Validate())
 		{
 			UE_LOG(LogSplineComponent, Display, TEXT("Internally inconsistent spline data at load time!"))
 		}
@@ -566,17 +541,12 @@ void USplineComponent::Serialize(FArchive& Ar)
 	const FPackageFileVersion ArchiveUEVersion = Ar.UEVer();
 	if (ArchiveUEVersion < VER_UE4_INTERPCURVE_SUPPORTS_LOOPING)
 	{
-		FTangentSpline* TangentSpline = Spline.Get<UE::Spline::ESplineType::Tangent>();
-
 		int32 LegacyNumPoints = WarninglessSplineCurves().Position.Points.Num();
 
 		// The start point is no longer cloned as the endpoint when the spline is looped, so remove the extra endpoint if present
 		if (bClosedLoop && ( GetLocationAtSplinePoint(0, ESplineCoordinateSpace::Local) == GetLocationAtSplinePoint(GetNumberOfSplinePoints() - 1, ESplineCoordinateSpace::Local)))
 		{
-			if (TangentSpline)
-			{
-				TangentSpline->RemovePoint(TangentSpline->GetNumberOfPoints() - 1);
-			}
+			Spline.RemovePoint(Spline.GetNumControlPoints() - 1);
 			WarninglessSplineCurves().Position.Points.RemoveAt(LegacyNumPoints - 1, EAllowShrinking::No);
 			LegacyNumPoints--;
 		}
@@ -593,11 +563,8 @@ void USplineComponent::Serialize(FArchive& Ar)
 		
 		// Importantly, these 2 functions guarantee that 1 rotation and scale value exists for each position,
 		// but they are default values.
-		if (TangentSpline)
-		{
-			TangentSpline->ResetRotation();
-			TangentSpline->ResetScale();
-		}
+		Spline.ResetRotation();
+		Spline.ResetScale();
 
 		UpdateSpline();
 	}
@@ -623,21 +590,6 @@ void USplineComponent::OnRegister()
 #endif
 }
 
-FSplineCurves FSplineCurves::FromSplineInterface(const FSpline& InProxy)
-{
-	FSplineCurves NewCurve;
-
-	if (const FTangentSpline* TangentSpline = InProxy.Get<UE::Spline::ESplineType::Tangent>())
-	{
-		NewCurve.Position = TangentSpline->GetSplinePointsPosition();
-		NewCurve.Rotation = TangentSpline->GetSplinePointsRotation();
-		NewCurve.Scale = TangentSpline->GetSplinePointsScale();
-	}
-
-	NewCurve.UpdateSpline(NewCurve.Position.bIsLooped, false);
-	return NewCurve;
-}
-
 void FSplineCurves::UpdateSpline(bool bClosedLoop, bool bStationaryEndpoints, int32 ReparamStepsPerSegment, bool bLoopPositionOverride, float LoopPosition, const FVector& Scale3D)
 {
 	TRACE_CPUPROFILER_EVENT_SCOPE(FSplineCurves::UpdateSpline);
@@ -699,36 +651,26 @@ void USplineComponent::UpdateSpline()
 {
 	TRACE_CPUPROFILER_EVENT_SCOPE(USplineComponent::UpdateSpline);
 
-	// If SplineComponent.UseSplineCurves is false and FSpline is not implemented, initialize it from SplineCurves.
-	if (UE::SplineComponent::GUseSplineCurves == false && Spline.GetCurrentImplementation() != UE::Spline::ESplineType::Tangent)
-	{
-		Spline = WarninglessSplineCurves();
-	}
-
-	FTangentSpline* TangentSpline = Spline.Get<UE::Spline::ESplineType::Tangent>();
-
-	FTangentSpline::FUpdateSplineParams Params{bClosedLoop, bStationaryEndpoints, ReparamStepsPerSegment, bLoopPositionOverride, LoopPosition, GetComponentTransform().GetScale3D()};
+	FSpline::FUpdateSplineParams Params{bClosedLoop, bStationaryEndpoints, ReparamStepsPerSegment, bLoopPositionOverride, LoopPosition, GetComponentTransform().GetScale3D()};
 	if (UE::SplineComponent::GOnlyUpdateActiveSpline)
 	{
-		if (ShouldUseSplineCurves())
+		if (UE::SplineComponent::ShouldUseSplineCurves())
 		{
 			WarninglessSplineCurves().UpdateSpline(Params.bClosedLoop, Params.bStationaryEndpoints, Params.ReparamStepsPerSegment, Params.bLoopPositionOverride, Params.LoopPosition, Params.Scale3D);
 		}
-		else if (ensureAlways(TangentSpline)) // ShouldUseSplineCurves should _never_ return false if TangentSpline would not be true
+		else
 		{
-			TangentSpline->UpdateSpline(Params);
+			Spline.UpdateSpline(Params);
 		}
 	}
 	else
 	{
 		WarninglessSplineCurves().UpdateSpline(Params.bClosedLoop, Params.bStationaryEndpoints, Params.ReparamStepsPerSegment, Params.bLoopPositionOverride, Params.LoopPosition, Params.Scale3D);
-
-		if (TangentSpline)
-		{
-			TangentSpline->UpdateSpline(Params);
-		}
+		Spline.UpdateSpline(Params);
 	}
 
+	// todo: make Spline replicate
+	//MARK_PROPERTY_DIRTY_FROM_NAME(USplineComponent, Spline, this);
 	MARK_PROPERTY_DIRTY_FROM_NAME(USplineComponent, SplineCurves, this);
 	MARK_PROPERTY_DIRTY_FROM_NAME(USplineComponent, bClosedLoop, this);
 	MARK_PROPERTY_DIRTY_FROM_NAME(USplineComponent, bStationaryEndpoints, this);
@@ -829,17 +771,9 @@ float FSplineCurves::GetSegmentLength(const int32 Index, const float Param, bool
 
 float USplineComponent::GetSegmentLength(const int32 Index, const float Param) const
 {
-	if (ShouldUseSplineCurves())
-	{
-		return WarninglessSplineCurves().GetSegmentLength(Index, Param, bClosedLoop, GetComponentTransform().GetScale3D());
-	}
-	else if (const FTangentSpline* TangentSpline = Spline.Get<UE::Spline::ESplineType::Tangent>())
-	{
-		return TangentSpline->GetSegmentLength(Index, Param, GetComponentTransform().GetScale3D());
-	}
-
-	ensureAlwaysMsgf(false, TEXT("ShouldUseSplineCurves() returned false but Spline failed to yield an FTangentSpline!"));
-	return 0.f;
+	return UE::SplineComponent::ShouldUseSplineCurves()
+		? WarninglessSplineCurves().GetSegmentLength(Index, Param, bClosedLoop, GetComponentTransform().GetScale3D())
+		: Spline.GetSegmentLength(Index, Param, GetComponentTransform().GetScale3D());
 }
 
 float USplineComponent::GetSegmentParamFromLength(const int32 Index, const float Length, const float SegmentLength) const
@@ -867,22 +801,10 @@ float USplineComponent::GetSegmentParamFromLength(const int32 Index, const float
 	//        (df/dt) = d(L(t))/dt = |dP/dt|
 
 	// TODO: check if this works OK with delta InVal != 1.0f
-	// ^ it definitely does not
 
-	int32 NumPoints = 0;
-	if (ShouldUseSplineCurves())
-	{
-		NumPoints = WarninglessSplineCurves().Position.Points.Num();
-	}
-	else if (const FTangentSpline* TangentSpline = Spline.Get<UE::Spline::ESplineType::Tangent>())
-	{
-		NumPoints = TangentSpline->GetNumberOfPoints();
-	}
-	else
-	{
-		ensureAlwaysMsgf(false, TEXT("ShouldUseSplineCurves() returned false but Spline failed to yield an FTangentSpline!"));
-		return 0.f;
-	}
+	int32 NumPoints = UE::SplineComponent::ShouldUseSplineCurves()
+		? WarninglessSplineCurves().Position.Points.Num()
+		: Spline.GetNumControlPoints();
 	
 	const int32 LastPoint = NumPoints - 1;
 
@@ -894,7 +816,7 @@ float USplineComponent::GetSegmentParamFromLength(const int32 Index, const float
 	// two iterations of Newton-Raphson is enough
 	for (int32 Iteration = 0; Iteration < 2; ++Iteration)
 	{
-		float TangentMagnitude = ShouldUseSplineCurves()
+		float TangentMagnitude = UE::SplineComponent::ShouldUseSplineCurves()
 			? WarninglessSplineCurves().Position.EvalDerivative(Index + Param, FVector::ZeroVector).Size()
 			: Spline.EvaluateDerivative(Index + Param).Size();
 		
@@ -910,7 +832,7 @@ float USplineComponent::GetSegmentParamFromLength(const int32 Index, const float
 
 FVector USplineComponent::GetLocationAtSplineInputKey(float InKey, ESplineCoordinateSpace::Type CoordinateSpace) const
 {
-	FVector Location = ShouldUseSplineCurves()
+	FVector Location = UE::SplineComponent::ShouldUseSplineCurves()
 		? WarninglessSplineCurves().Position.Eval(InKey, FVector::ZeroVector)
 		: Spline.Evaluate(InKey);
 
@@ -924,7 +846,7 @@ FVector USplineComponent::GetLocationAtSplineInputKey(float InKey, ESplineCoordi
 
 FVector USplineComponent::GetTangentAtSplineInputKey(float InKey, ESplineCoordinateSpace::Type CoordinateSpace) const
 {
-	FVector Tangent = ShouldUseSplineCurves()
+	FVector Tangent = UE::SplineComponent::ShouldUseSplineCurves()
 		? WarninglessSplineCurves().Position.EvalDerivative(InKey, FVector::ZeroVector)
 		: Spline.EvaluateDerivative(InKey);
 
@@ -938,7 +860,7 @@ FVector USplineComponent::GetTangentAtSplineInputKey(float InKey, ESplineCoordin
 
 FVector USplineComponent::GetDirectionAtSplineInputKey(float InKey, ESplineCoordinateSpace::Type CoordinateSpace) const
 {
-	FVector Direction = ShouldUseSplineCurves()
+	FVector Direction = UE::SplineComponent::ShouldUseSplineCurves()
 		? WarninglessSplineCurves().Position.EvalDerivative(InKey, FVector::ZeroVector).GetSafeNormal()
 		: Spline.EvaluateDerivative(InKey).GetSafeNormal();
 
@@ -958,23 +880,12 @@ FRotator USplineComponent::GetRotationAtSplineInputKey(float InKey, ESplineCoord
 
 FQuat USplineComponent::GetQuaternionAtSplineInputKey(float InKey, ESplineCoordinateSpace::Type CoordinateSpace) const
 {
-	FQuat Quat = FQuat::Identity;
-	if (ShouldUseSplineCurves())
-	{
-		Quat = WarninglessSplineCurves().Rotation.Eval(InKey, FQuat::Identity);
-	}
-	else if (const FTangentSpline* TangentSpline = Spline.Get<UE::Spline::ESplineType::Tangent>())
-	{
-		Quat = TangentSpline->EvaluateRotation(InKey);
-	}
-	else
-	{
-		ensureAlwaysMsgf(false, TEXT("ShouldUseSplineCurves() returned false but Spline failed to yield an FTangentSpline!"));
-		return Quat;
-	}
+	FQuat Quat = UE::SplineComponent::ShouldUseSplineCurves()
+		? WarninglessSplineCurves().Rotation.Eval(InKey, FQuat::Identity)
+		: Spline.EvaluateRotation(InKey);
 	Quat.Normalize();
 
-	FVector Direction = ShouldUseSplineCurves()
+	FVector Direction = UE::SplineComponent::ShouldUseSplineCurves()
 		? WarninglessSplineCurves().Position.EvalDerivative(InKey, FVector::ZeroVector).GetSafeNormal()
 		: Spline.EvaluateDerivative(InKey).GetSafeNormal();
 	
@@ -1040,41 +951,16 @@ float USplineComponent::GetRollAtSplineInputKey(float InKey, ESplineCoordinateSp
 
 FVector USplineComponent::GetScaleAtSplineInputKey(float InKey) const
 {
-	FVector Scale = FVector::OneVector;
-
-	if (ShouldUseSplineCurves())
-	{
-		Scale = WarninglessSplineCurves().Scale.Eval(InKey, FVector::OneVector);
-	}
-	else if (const FTangentSpline* TangentSpline = Spline.Get<UE::Spline::ESplineType::Tangent>())
-	{
-		Scale = TangentSpline->EvaluateScale(InKey);
-	}
-	else
-	{
-		ensureAlwaysMsgf(false, TEXT("ShouldUseSplineCurves() returned false but Spline failed to yield an FTangentSpline!"));
-	}
-
-	return Scale;
+	return UE::SplineComponent::ShouldUseSplineCurves()
+		? WarninglessSplineCurves().Scale.Eval(InKey, FVector(1.0f))
+		: Spline.EvaluateScale(InKey);
 }
 
 float USplineComponent::GetDistanceAlongSplineAtSplineInputKey(float InKey) const
 {
-	int32 NumPoints = 0;
-	const FTangentSpline* TangentSpline = nullptr;
-	if (ShouldUseSplineCurves())
-	{
-		NumPoints = WarninglessSplineCurves().Position.Points.Num();
-	}
-	else if ((TangentSpline = Spline.Get<UE::Spline::ESplineType::Tangent>()))
-	{
-		NumPoints = TangentSpline->GetNumberOfPoints();
-	}
-	else
-	{
-		ensureAlwaysMsgf(false, TEXT("ShouldUseSplineCurves() returned false but Spline failed to yield an FTangentSpline!"));
-		return 0.f;
-	}
+	const int32 NumPoints = UE::SplineComponent::ShouldUseSplineCurves()
+		? WarninglessSplineCurves().Position.Points.Num()
+		: Spline.GetNumControlPoints();
 	
 	const int32 NumSegments = bClosedLoop ? NumPoints : NumPoints - 1;
 
@@ -1082,8 +968,7 @@ float USplineComponent::GetDistanceAlongSplineAtSplineInputKey(float InKey) cons
 	{
 		float Distance;
 		
-		// !TangentSpline is true only if ShouldUseSplineCurves previously returned true.
-		if (!TangentSpline)
+		if (UE::SplineComponent::ShouldUseSplineCurves())
 		{
 			const int32 PointIndex = FMath::FloorToInt(InKey);
 			const float Fraction = InKey - PointIndex;
@@ -1092,7 +977,7 @@ float USplineComponent::GetDistanceAlongSplineAtSplineInputKey(float InKey) cons
 		}
 		else
 		{
-			Distance = TangentSpline->GetDistanceAtParameter(InKey);
+			Distance = Spline.GetDistanceAtParameter(InKey);
 		}
 		
 		return Distance;
@@ -1110,32 +995,18 @@ float USplineComponent::GetDistanceAlongSplineAtLocation(const FVector& InLocati
 	const FVector LocalLocation = (CoordinateSpace == ESplineCoordinateSpace::World) ? GetComponentTransform().InverseTransformPosition(InLocation) : InLocation;
 	float Dummy;
 	
-	float Key;
-	if (ShouldUseSplineCurves())
-	{
-		Key = WarninglessSplineCurves().Position.FindNearest(LocalLocation, Dummy);
-	}
-	else if (const FTangentSpline* TangentSpline = Spline.Get<UE::Spline::ESplineType::Tangent>())
-	{
-		Key = TangentSpline->FindNearest(LocalLocation, Dummy);
-	}
-	else
-	{
-		ensureAlwaysMsgf(false, TEXT("ShouldUseSplineCurves() returned false but Spline failed to yield an FTangentSpline!"));
-		return 0.f;
-	}
+	const float Key = UE::SplineComponent::ShouldUseSplineCurves()
+		? WarninglessSplineCurves().Position.FindNearest(LocalLocation, Dummy)
+		: Spline.FindNearest(LocalLocation, Dummy);
 	
 	return GetDistanceAlongSplineAtSplineInputKey(Key);
 }
 
 int32 USplineComponent::GetNumberOfPropertyValues(FName PropertyName) const
 {
-	if (ShouldUseSplineCurves() == false)
+	if (UE::SplineComponent::ShouldUseSplineCurves() == false && Spline.SupportsAttributes())
 	{
-		if (const FTangentSpline* TangentSpline = Spline.Get<UE::Spline::ESplineType::Tangent>())
-		{
-			return TangentSpline->NumAttributeValues<float>(PropertyName);
-		}
+		return Spline.NumAttributeValues<float>(PropertyName);
 	}
 
 	return 0;
@@ -1143,19 +1014,16 @@ int32 USplineComponent::GetNumberOfPropertyValues(FName PropertyName) const
 
 TArray<FName> USplineComponent::GetFloatPropertyChannels() const
 {
-	if (ShouldUseSplineCurves() == false)
+	if (UE::SplineComponent::ShouldUseSplineCurves() == false && Spline.SupportsAttributes())
 	{
-		if (const FTangentSpline* TangentSpline = Spline.Get<UE::Spline::ESplineType::Tangent>())
-		{
-			return TangentSpline->GetAttributeChannelNamesByValueType<float>();
-		}
+		return Spline.GetFloatPropertyChannels();
 	}
 
 	return TArray<FName>();
 }
 
 template<class T>
-bool CreatePropertyChannel(const USplineMetadata* Metadata, FSpline& Spline, FName PropertyName)
+bool CreatePropertyChannel(const USplineMetadata* Metadata, const FSpline& Spline, FName PropertyName)
 {
 	if (Metadata && Metadata->GetClass()->FindPropertyByName(PropertyName))
 	{
@@ -1163,20 +1031,11 @@ bool CreatePropertyChannel(const USplineMetadata* Metadata, FSpline& Spline, FNa
 		return false;
 	}
 	
-	if (UE::SplineComponent::GUseSplineCurves == false)
+	if (UE::SplineComponent::GUseSplineCurves == false && Spline.SupportsAttributes())
 	{
-		if (FTangentSpline* TangentSpline = Spline.Get<UE::Spline::ESplineType::Tangent>())
+		if (!Spline.HasAttributeChannel(PropertyName))
 		{
-			if (!TangentSpline->HasAttributeChannel(PropertyName))
-			{
-				if (UE::Geometry::Spline::TTangentBezierSpline<T>* NewAttrChannel = TangentSpline->CreateAttributeChannel<UE::Geometry::Spline::TTangentBezierSpline<T>>(PropertyName))
-				{
-					NewAttrChannel->SetPreInfinityMode(UE::Geometry::Spline::EOutOfBoundsHandlingMode::Constant);
-					NewAttrChannel->SetPostInfinityMode(UE::Geometry::Spline::EOutOfBoundsHandlingMode::Constant);
-
-					return true;
-				}
-			}
+			return Spline.CreateAttributeChannel<T>(PropertyName);
 		}
 	}
 
@@ -1195,14 +1054,11 @@ T GetPropertyAtSplineInputKey(const USplineMetadata* Metadata, const FSpline& Sp
 		}
 	}
 	
-	if (UE::SplineComponent::GUseSplineCurves == false)
+	if (UE::SplineComponent::GUseSplineCurves == false && Spline.SupportsAttributes())
 	{
-		if (const FTangentSpline* TangentSpline = Spline.Get<UE::Spline::ESplineType::Tangent>())
+		if (Spline.HasAttributeChannel(PropertyName))
 		{
-			if (TangentSpline->HasAttributeChannel(PropertyName))
-			{
-				return TangentSpline->EvaluateAttribute<T>(PropertyName, InKey);
-			}
+			return Spline.EvaluateAttribute<T>(InKey, PropertyName);
 		}
 	}
 
@@ -1210,16 +1066,13 @@ T GetPropertyAtSplineInputKey(const USplineMetadata* Metadata, const FSpline& Sp
 }
 
 template<class T>
-int32 AddPropertyAtSplineInputKey(FSpline& Spline, float InKey, const T& InValue, FName PropertyName)
+int32 SetPropertyAtSplineInputKey(const FSpline& Spline, float InKey, const T& InValue, FName PropertyName)
 {
-	if (UE::SplineComponent::GUseSplineCurves == false)
+	if (UE::SplineComponent::GUseSplineCurves == false && Spline.SupportsAttributes())
 	{
-		if (FTangentSpline* TangentSpline = Spline.Get<UE::Spline::ESplineType::Tangent>())
+		if (Spline.HasAttributeChannel(PropertyName))
 		{
-			if (TangentSpline->HasAttributeChannel(PropertyName))
-			{
-				return TangentSpline->AddAttributeValue<T>(PropertyName, InValue, InKey);
-			}
+			return Spline.AddAttributeValue<T>(InKey, InValue, PropertyName);
 		}
 	}
 
@@ -1228,15 +1081,9 @@ int32 AddPropertyAtSplineInputKey(FSpline& Spline, float InKey, const T& InValue
 
 float GetInputKeyAtIndex(const FSpline& Spline, int32 Index, FName PropertyName)
 {
-	if (UE::SplineComponent::GUseSplineCurves == false)
+	if (UE::SplineComponent::GUseSplineCurves == false && Spline.SupportsAttributes())
 	{
-		if (const FTangentSpline* TangentSpline = Spline.Get<UE::Spline::ESplineType::Tangent>())
-		{
-			if (TangentSpline->HasAttributeChannel(PropertyName))
-			{
-				return TangentSpline->GetAttributeParameter<float>(PropertyName, Index);
-			}
-		}
+		return Spline.GetAttributeParameter<float>(Index, PropertyName);
 	}
 
 	return 0.f;
@@ -1244,15 +1091,9 @@ float GetInputKeyAtIndex(const FSpline& Spline, int32 Index, FName PropertyName)
 
 int32 SetInputKeyAtIndex(FSpline& Spline, int32 Index, float InKey, FName PropertyName)
 {
-	if (UE::SplineComponent::GUseSplineCurves == false)
+	if (UE::SplineComponent::GUseSplineCurves == false && Spline.SupportsAttributes())
 	{
-		if (FTangentSpline* TangentSpline = Spline.Get<UE::Spline::ESplineType::Tangent>())
-		{
-			if (TangentSpline->HasAttributeChannel(PropertyName))
-			{
-				return TangentSpline->SetAttributeParameter<float>(PropertyName, Index, InKey);
-			}
-		}
+		return Spline.SetAttributeParameter<float>(Index, InKey , PropertyName);
 	}
 
 	return INDEX_NONE;
@@ -1261,15 +1102,9 @@ int32 SetInputKeyAtIndex(FSpline& Spline, int32 Index, float InKey, FName Proper
 template<typename T>
 float GetPropertyAtIndex(const FSpline& Spline, int32 Index, FName PropertyName)
 {
-	if (UE::SplineComponent::GUseSplineCurves == false)
+	if (UE::SplineComponent::GUseSplineCurves == false && Spline.SupportsAttributes())
 	{
-		if (const FTangentSpline* TangentSpline = Spline.Get<UE::Spline::ESplineType::Tangent>())
-		{
-			if (TangentSpline->HasAttributeChannel(PropertyName))
-			{
-				return TangentSpline->GetAttributeValue<T>(PropertyName, Index);
-			}
-		}
+		return Spline.GetAttributeValue<T>(Index, PropertyName);
 	}
 
 	return 0.f;
@@ -1278,54 +1113,30 @@ float GetPropertyAtIndex(const FSpline& Spline, int32 Index, FName PropertyName)
 template<typename T>
 void SetPropertyAtIndex(FSpline& Spline, int32 Index, float Value, FName PropertyName)
 {
-	if (UE::SplineComponent::GUseSplineCurves == false)
+	if (UE::SplineComponent::GUseSplineCurves == false && Spline.SupportsAttributes())
 	{
-		if (FTangentSpline* TangentSpline = Spline.Get<UE::Spline::ESplineType::Tangent>())
-		{
-			if (TangentSpline->HasAttributeChannel(PropertyName))
-			{
-				TangentSpline->SetAttributeValue<T>(PropertyName, Value, Index);
-			}
-		}
+		Spline.SetAttributeValue<T>(Index, Value, PropertyName);
 	}
 }
 
 void USplineComponent::RemovePropertyAtIndex(int32 Index, FName PropertyName)
 {
-	if (ShouldUseSplineCurves() == false)
+	if (UE::SplineComponent::ShouldUseSplineCurves() == false && Spline.SupportsAttributes())
 	{
-		if (FTangentSpline* TangentSpline = Spline.Get<UE::Spline::ESplineType::Tangent>())
-		{
-			if (TangentSpline->HasAttributeChannel(PropertyName))
-			{
-				TangentSpline->RemoveAttributeValue<float>(PropertyName, Index);
-			}
-		}
+		Spline.RemoveAttributeValue<float>(Index, PropertyName);
 	}
 }
 
 bool USplineComponent::SupportsAttributes() const
 {
-	// If GetAs yields a Tangent spline, we support attributes.
-	if (ShouldUseSplineCurves() == false && Spline.Get<UE::Spline::ESplineType::Tangent>())
-	{
-		return true;
-	}
-
-	return false;
+	return UE::SplineComponent::ShouldUseSplineCurves() == false && Spline.SupportsAttributes();
 }
 
 bool USplineComponent::RemovePropertyChannel(FName PropertyName)
 {
-	if (ShouldUseSplineCurves() == false)
+	if (UE::SplineComponent::ShouldUseSplineCurves() == false && Spline.SupportsAttributes())
 	{
-		if (FTangentSpline* TangentSpline = Spline.Get<UE::Spline::ESplineType::Tangent>())
-		{
-			if (TangentSpline->HasAttributeChannel(PropertyName))
-			{
-				return TangentSpline->RemoveAttributeChannel(PropertyName);
-			}
-		}
+		return Spline.RemoveAttributeChannel(PropertyName);
 	}
 
 	return false;
@@ -1341,9 +1152,9 @@ float USplineComponent::GetFloatPropertyAtSplineInputKey(float InKey, FName Prop
 	return GetPropertyAtSplineInputKey<float>(GetSplinePointsMetadata(), Spline, InKey, PropertyName);
 }
 
-int32 USplineComponent::AddFloatPropertyAtSplineInputKey(float InKey, float Value, FName PropertyName)
+int32 USplineComponent::SetFloatPropertyAtSplineInputKey(float InKey, float Value, FName PropertyName) const
 {
-	return AddPropertyAtSplineInputKey<float>(Spline, InKey, Value, PropertyName);
+	return SetPropertyAtSplineInputKey<float>(Spline, InKey, Value, PropertyName);
 }
 
 float USplineComponent::GetFloatPropertyInputKeyAtIndex(int32 Index, FName PropertyName) const
@@ -1523,10 +1334,7 @@ void USplineComponent::AddPoint(const FSplinePoint& InSplinePoint, bool bUpdateS
 		CIM_CurveAuto
 	), Index);
 	
-	if (FTangentSpline* TangentSpline = Spline.Get<UE::Spline::ESplineType::Tangent>())
-	{
-		TangentSpline->AddPoint(InSplinePoint);
-	}
+	Spline.AddPoint(InSplinePoint);
 	
 	if (USplineMetadata* Metadata = GetSplinePointsMetadata())
 	{
@@ -1600,14 +1408,11 @@ void USplineComponent::AddSplinePoint(const FVector& Position, ESplineCoordinate
 	NewPoint.Scale = FVector(1.f);
 	NewPoint.Type = ConvertInterpCurveModeToSplinePointType(CIM_CurveAuto);
 
-	if (FTangentSpline* TangentSpline = Spline.Get<UE::Spline::ESplineType::Tangent>())
+	switch (UE::SplineComponent::GAddSplinePointImplementation)
 	{
-		switch (UE::SplineComponent::GAddSplinePointImplementation)
-		{
-		default: // fallthrough
-		case 0: TangentSpline->AddPoint(NewPoint); break;
-		case 1: TangentSpline->InsertPoint(NewPoint, GetNumberOfSplinePoints()); break;
-		}
+	default: // fallthrough
+	case 0: Spline.AddPoint(NewPoint); break;
+	case 1: Spline.InsertPoint(NewPoint, GetNumberOfSplinePoints()); break;
 	}
 
 	WarninglessSplineCurves().Position.Points.Emplace(InKey, TransformedPosition, FVector::ZeroVector, FVector::ZeroVector, CIM_CurveAuto);
@@ -1656,10 +1461,7 @@ void USplineComponent::AddSplinePointAtIndex(const FVector& Position, int32 Inde
 		NewPoint.Rotation = FQuat::Identity.Rotator();
 		NewPoint.Scale = FVector(1.f);
 
-		if (FTangentSpline* TangentSpline = Spline.Get<UE::Spline::ESplineType::Tangent>())
-		{
-			TangentSpline->InsertPoint(NewPoint, Index);
-		}
+		Spline.InsertPoint(NewPoint, Index);
 
 		WarninglessSplineCurves().Position.Points.Insert(FInterpCurvePoint<FVector>(InKey, TransformedPosition, FVector::ZeroVector, FVector::ZeroVector, CIM_CurveAuto), Index);
 		WarninglessSplineCurves().Rotation.Points.Insert(FInterpCurvePoint<FQuat>(InKey, FQuat::Identity, FQuat::Identity, FQuat::Identity, CIM_CurveAuto), Index);
@@ -1704,10 +1506,7 @@ void USplineComponent::RemoveSplinePoint(int32 Index, bool bUpdateSpline)
 
 	if (Index >= 0 && Index < NumPoints)
 	{
-		if (FTangentSpline* TangentSpline = Spline.Get<UE::Spline::ESplineType::Tangent>())
-		{
-			TangentSpline->RemovePoint(Index);
-		}
+		Spline.RemovePoint(Index);
 
 		float KnotDelta = 0.0f;
 		if (Index > 0)
@@ -1775,7 +1574,6 @@ void USplineComponent::SetSplinePoints(const TArray<FVector>& Points, ESplineCoo
 	}
 
 	float InputKey = 0.0f;
-	FTangentSpline* TangentSpline = Spline.Get<UE::Spline::ESplineType::Tangent>();
 	for (const auto& Point : Points)
 	{
 		const FVector TransformedPoint = (CoordinateSpace == ESplineCoordinateSpace::World) ?
@@ -1789,10 +1587,7 @@ void USplineComponent::SetSplinePoints(const TArray<FVector>& Points, ESplineCoo
 		NewPoint.Rotation = FQuat::Identity.Rotator();
 		NewPoint.Scale = FVector(1.f);
 
-		if (TangentSpline)
-		{
-			TangentSpline->AddPoint(NewPoint);
-		}
+		Spline.AddPoint(NewPoint);
 
 		WarninglessSplineCurves().Position.Points.Emplace(InputKey, TransformedPoint, FVector::ZeroVector, FVector::ZeroVector, CIM_CurveAuto);
 		WarninglessSplineCurves().Rotation.Points.Emplace(InputKey, FQuat::Identity, FQuat::Identity, FQuat::Identity, CIM_CurveAuto);
@@ -1830,11 +1625,7 @@ void USplineComponent::SetLocationAtSplinePoint(int32 PointIndex, const FVector&
 		const FVector TransformedLocation = (CoordinateSpace == ESplineCoordinateSpace::World) ?
 			GetComponentTransform().InverseTransformPosition(InLocation) : InLocation;
 
-		if (FTangentSpline* TangentSpline = Spline.Get<UE::Spline::ESplineType::Tangent>())
-		{
-			TangentSpline->SetLocation(PointIndex, TransformedLocation);
-		}
-
+		Spline.SetLocation(PointIndex, TransformedLocation);
 		WarninglessSplineCurves().Position.Points[PointIndex].OutVal = TransformedLocation;
 
 		if (bUpdateSpline)
@@ -1867,11 +1658,8 @@ void USplineComponent::SetTangentsAtSplinePoint(int32 PointIndex, const FVector&
 		const FVector TransformedLeaveTangent = (CoordinateSpace == ESplineCoordinateSpace::World) ?
 			GetComponentTransform().InverseTransformVector(InLeaveTangent) : InLeaveTangent;
 
-		if (FTangentSpline* TangentSpline = Spline.Get<UE::Spline::ESplineType::Tangent>())
-		{
-			TangentSpline->SetInTangent(PointIndex, TransformedArriveTangent);
-			TangentSpline->SetOutTangent(PointIndex, TransformedLeaveTangent);
-		}
+		Spline.SetInTangent(PointIndex, TransformedArriveTangent);
+		Spline.SetOutTangent(PointIndex, TransformedLeaveTangent);
 
 		WarninglessSplineCurves().Position.Points[PointIndex].ArriveTangent = TransformedArriveTangent;
 		WarninglessSplineCurves().Position.Points[PointIndex].LeaveTangent = TransformedLeaveTangent;
@@ -1901,12 +1689,7 @@ void USplineComponent::SetUpVectorAtSplinePoint(int32 PointIndex, const FVector&
 			GetComponentTransform().InverseTransformVector(InUpVector.GetSafeNormal()) : InUpVector.GetSafeNormal();
 
 		FQuat Quat = FQuat::FindBetween(DefaultUpVector, TransformedUpVector);
-
-		if (FTangentSpline* TangentSpline = Spline.Get<UE::Spline::ESplineType::Tangent>())
-		{
-			TangentSpline->SetRotation(PointIndex, Quat);
-		}
-
+		Spline.SetRotation(PointIndex, Quat);
 		WarninglessSplineCurves().Rotation.Points[PointIndex].OutVal = Quat;
 
 		if (bUpdateSpline)
@@ -1974,12 +1757,8 @@ void USplineComponent::SetScaleAtSplinePoint(int32 PointIndex, const FVector& In
 	
 	if ((PointIndex >= 0) && (PointIndex < NumPoints))
 	{
-		if (FTangentSpline* TangentSpline = Spline.Get<UE::Spline::ESplineType::Tangent>())
-		{
-			TangentSpline->SetScale(PointIndex, InScaleVector);
-		}
-
 		WarninglessSplineCurves().Scale.Points[PointIndex].OutVal = InScaleVector;
+		Spline.SetScale(PointIndex, InScaleVector);
 
 		if (bUpdateSpline)
 		{
@@ -2001,20 +1780,9 @@ ESplinePointType::Type USplineComponent::GetSplinePointType(int32 PointIndex) co
 	
 	if (PointIndex >= 0 && PointIndex < NumPoints)
 	{
-		EInterpCurveMode Mode = CIM_Unknown;
-
-		if (ShouldUseSplineCurves())
-		{
-			Mode = WarninglessSplineCurves().Position.Points[PointIndex].InterpMode.GetValue();
-		}
-		else if (const FTangentSpline* TangentSpline = Spline.Get<UE::Spline::ESplineType::Tangent>())
-		{
-			Mode = TangentSpline->GetSplinePointType(PointIndex);
-		}
-		else
-		{
-			ensureAlwaysMsgf(false, TEXT("ShouldUseSplineCurves() returned false but Spline failed to yield an FTangentSpline!"));
-		}
+		const EInterpCurveMode Mode = UE::SplineComponent::ShouldUseSplineCurves()
+			? WarninglessSplineCurves().Position.Points[PointIndex].InterpMode.GetValue()
+			: Spline.GetSplinePointType(PointIndex);
 		
 		return ConvertInterpCurveModeToSplinePointType(Mode);
 	}
@@ -2028,11 +1796,7 @@ void USplineComponent::SetSplinePointType(int32 PointIndex, ESplinePointType::Ty
 	
 	if (PointIndex >= 0 && PointIndex < NumPoints)
 	{
-		if (FTangentSpline* TangentSpline = Spline.Get<UE::Spline::ESplineType::Tangent>())
-		{
-			TangentSpline->SetSplinePointType(PointIndex, ConvertSplinePointTypeToInterpCurveMode(Type));
-		}
-
+		Spline.SetSplinePointType(PointIndex, ConvertSplinePointTypeToInterpCurveMode(Type));
 		WarninglessSplineCurves().Position.Points[PointIndex].InterpMode = ConvertSplinePointTypeToInterpCurveMode(Type);
 		
 		if (bUpdateSpline)
@@ -2051,19 +1815,9 @@ void USplineComponent::SetSplinePointType(int32 PointIndex, ESplinePointType::Ty
 
 int32 USplineComponent::GetNumberOfSplinePoints() const
 {
-	if (ShouldUseSplineCurves())
-	{
-		return WarninglessSplineCurves().Position.Points.Num();
-	}
-	if (const FTangentSpline* TangentSpline = Spline.Get<UE::Spline::ESplineType::Tangent>())
-	{
-		return TangentSpline->GetNumberOfPoints();
-	}
-	else
-	{
-		ensureAlwaysMsgf(false, TEXT("ShouldUseSplineCurves() returned false but Spline failed to yield an FTangentSpline!"));
-		return 0;
-	}
+	return UE::SplineComponent::ShouldUseSplineCurves()
+		? WarninglessSplineCurves().Position.Points.Num()
+		: Spline.GetNumControlPoints();
 }
 
 int32 USplineComponent::GetNumberOfSplineSegments() const
@@ -2084,37 +1838,22 @@ float USplineComponent::GetInputKeyValueAtSplinePoint(int32 PointIndex) const
 
 	if (bClosedLoop && PointIndex == GetNumberOfSplinePoints())
 	{
-		if (ShouldUseSplineCurves())
+		if (UE::SplineComponent::ShouldUseSplineCurves())
 		{
 			return WarninglessSplineCurves().Position.Points[PointIndex - 1].InVal + WarninglessSplineCurves().Position.LoopKeyOffset;
 		}
-		else if (const FTangentSpline* TangentSpline = Spline.Get<UE::Spline::ESplineType::Tangent>())
-		{
-			return TangentSpline->GetParameterAtIndex(PointIndex);
-		}
 		else
 		{
-			ensureAlwaysMsgf(false, TEXT("ShouldUseSplineCurves() returned false but Spline failed to yield an FTangentSpline!"));
-			return 0.f;
+			return Spline.GetParameterAtIndex(PointIndex);
 		}
 	}
 	else
 	{
 		PointIndex = GetClampedIndex(PointIndex);
 
-		if (ShouldUseSplineCurves())
-		{
-			return WarninglessSplineCurves().Position.Points[PointIndex].InVal;
-		}
-		else if (const FTangentSpline* TangentSpline = Spline.Get<UE::Spline::ESplineType::Tangent>())
-		{
-			return TangentSpline->GetParameterAtIndex(PointIndex);
-		}
-		else
-		{
-			ensureAlwaysMsgf(false, TEXT("ShouldUseSplineCurves() returned false but Spline failed to yield an FTangentSpline!"));
-			return 0.f;
-		}
+		return UE::SplineComponent::ShouldUseSplineCurves()
+			? WarninglessSplineCurves().Position.Points[PointIndex].InVal
+			: Spline.GetParameterAtIndex(PointIndex);
 	}
 }
 
@@ -2140,20 +1879,9 @@ FVector USplineComponent::GetLocationAtSplinePoint(int32 PointIndex, ESplineCoor
 
 	PointIndex = GetClampedIndex(PointIndex);
 	
-	FVector Location = FVector::Zero();
-	if (ShouldUseSplineCurves())
-	{
-		Location = WarninglessSplineCurves().Position.Points[PointIndex].OutVal;
-	}
-	else if (const FTangentSpline* TangentSpline = Spline.Get<UE::Spline::ESplineType::Tangent>())
-	{
-		Location = TangentSpline->GetLocation(PointIndex);
-	}
-	else
-	{
-		ensureAlwaysMsgf(false, TEXT("ShouldUseSplineCurves() returned false but Spline failed to yield an FTangentSpline!"));
-		return Location;
-	}
+	const FVector Location = UE::SplineComponent::ShouldUseSplineCurves()
+		? WarninglessSplineCurves().Position.Points[PointIndex].OutVal
+		: Spline.GetLocation(PointIndex);
 	
 	return (CoordinateSpace == ESplineCoordinateSpace::World) ? GetComponentTransform().TransformPosition(Location) : Location;
 }
@@ -2177,21 +1905,9 @@ FVector USplineComponent::GetArriveTangentAtSplinePoint(int32 PointIndex, ESplin
 
 	PointIndex = GetClampedIndex(PointIndex);
 
-	FVector Tangent = FVector::ZeroVector;
-
-	if (ShouldUseSplineCurves())
-	{
-		Tangent = WarninglessSplineCurves().Position.Points[PointIndex].ArriveTangent;
-	}
-	else if (const FTangentSpline* TangentSpline = Spline.Get<UE::Spline::ESplineType::Tangent>())
-	{
-		Tangent = TangentSpline->GetInTangent(PointIndex);
-	}
-	else
-	{
-		ensureAlwaysMsgf(false, TEXT("ShouldUseSplineCurves() returned false but Spline failed to yield an FTangentSpline!"));
-		return Tangent;
-	}
+	const FVector Tangent = UE::SplineComponent::ShouldUseSplineCurves()
+		? WarninglessSplineCurves().Position.Points[PointIndex].ArriveTangent
+		: Spline.GetInTangent(PointIndex);
 
 	return (CoordinateSpace == ESplineCoordinateSpace::World) ? GetComponentTransform().TransformVector(Tangent) : Tangent;
 }
@@ -2205,21 +1921,9 @@ FVector USplineComponent::GetLeaveTangentAtSplinePoint(int32 PointIndex, ESpline
 
 	PointIndex = GetClampedIndex(PointIndex);
 
-	FVector Tangent = FVector::ZeroVector;
-
-	if (ShouldUseSplineCurves())
-	{
-		Tangent = WarninglessSplineCurves().Position.Points[PointIndex].LeaveTangent;
-	}
-	else if (const FTangentSpline* TangentSpline = Spline.Get<UE::Spline::ESplineType::Tangent>())
-	{
-		Tangent = TangentSpline->GetOutTangent(PointIndex);
-	}
-	else
-	{
-		ensureAlwaysMsgf(false, TEXT("ShouldUseSplineCurves() returned false but Spline failed to yield an FTangentSpline!"));
-		return Tangent;
-	}
+	const FVector Tangent = UE::SplineComponent::ShouldUseSplineCurves()
+		? WarninglessSplineCurves().Position.Points[PointIndex].LeaveTangent
+		: Spline.GetOutTangent(PointIndex);
 
 	return (CoordinateSpace == ESplineCoordinateSpace::World) ? GetComponentTransform().TransformVector(Tangent) : Tangent;
 }
@@ -2258,22 +1962,9 @@ FVector USplineComponent::GetScaleAtSplinePoint(int32 PointIndex) const
 
 	PointIndex = GetClampedIndex(PointIndex);
 
-	FVector Scale = FVector::OneVector;
-
-	if (ShouldUseSplineCurves())
-	{
-		Scale = WarninglessSplineCurves().Scale.Points[PointIndex].OutVal;
-	}
-	else if (const FTangentSpline* TangentSpline = Spline.Get<UE::Spline::ESplineType::Tangent>())
-	{
-		Scale = TangentSpline->GetScale(PointIndex);
-	}
-	else
-	{
-		ensureAlwaysMsgf(false, TEXT("ShouldUseSplineCurves() returned false but Spline failed to yield an FTangentSpline!"));
-	}
-
-	return Scale;
+	return UE::SplineComponent::ShouldUseSplineCurves()
+		? WarninglessSplineCurves().Scale.Points[PointIndex].OutVal
+		: Spline.GetScale(PointIndex);
 }
 
 FTransform USplineComponent::GetTransformAtSplinePoint(int32 PointIndex, ESplineCoordinateSpace::Type CoordinateSpace, bool bUseScale) const
@@ -2296,7 +1987,7 @@ float USplineComponent::GetDistanceAlongSplineAtSplinePoint(int32 PointIndex) co
 		return GetSplineLength();
 	}
 	
-	if (ShouldUseSplineCurves())
+	if (UE::SplineComponent::ShouldUseSplineCurves())
 	{
 		const int32 NumPoints = GetNumberOfSplinePoints();
         const int32 NumSegments = bClosedLoop ? NumPoints : NumPoints - 1;
@@ -2309,14 +2000,10 @@ float USplineComponent::GetDistanceAlongSplineAtSplinePoint(int32 PointIndex) co
         	return WarninglessSplineCurves().ReparamTable.Points[PointIndex * ReparamStepsPerSegment].InVal;
         }
 	}
-	else if (const FTangentSpline* TangentSpline = Spline.Get<UE::Spline::ESplineType::Tangent>())
-	{
-		const float ParameterAtIndex = TangentSpline->GetParameterAtIndex(PointIndex);
-		return TangentSpline->GetDistanceAtParameter(ParameterAtIndex);
-	}
 	else
 	{
-		ensureAlwaysMsgf(false, TEXT("ShouldUseSplineCurves() returned false but Spline failed to yield an FTangentSpline!"));
+		const float ParameterAtIndex = Spline.GetParameterAtIndex(PointIndex);
+		return Spline.GetDistanceAtParameter(ParameterAtIndex);
 	}
 
 	return 0.0f;
@@ -2337,20 +2024,9 @@ float FSplineCurves::GetSplineLength() const
 
 float USplineComponent::GetSplineLength() const
 {
-	if (ShouldUseSplineCurves())
-	{
-		return WarninglessSplineCurves().GetSplineLength();
-	}
-	else if (const FTangentSpline* TangentSpline = Spline.Get<UE::Spline::ESplineType::Tangent>())
-	{
-		const float MaxParameter = TangentSpline->GetParameterSpace().Max;
-		return TangentSpline->GetDistanceAtParameter(MaxParameter);
-	}
-	else
-	{
-		ensureAlwaysMsgf(false, TEXT("ShouldUseSplineCurves() returned false but Spline failed to yield an FTangentSpline!"));
-		return 0.f;
-	}
+	return UE::SplineComponent::ShouldUseSplineCurves()
+		? WarninglessSplineCurves().GetSplineLength()
+		: Spline.GetSplineLength();
 }
 
 void USplineComponent::SetDefaultUpVector(const FVector& UpVector, ESplineCoordinateSpace::Type CoordinateSpace)
@@ -2401,19 +2077,9 @@ float USplineComponent::GetInputKeyValueAtDistanceAlongSpline(float Distance) co
 		return 0.0f;
 	}
 	
-	if (ShouldUseSplineCurves())
-	{
-		return WarninglessSplineCurves().ReparamTable.Eval(Distance, 0.0f);
-	}
-	else if (const FTangentSpline* TangentSpline = Spline.Get<UE::Spline::ESplineType::Tangent>())
-	{
-		return TangentSpline->GetParameterAtDistance(Distance);
-	}
-	else
-	{
-		ensureAlwaysMsgf(false, TEXT("ShouldUseSplineCurves() returned false but Spline failed to yield an FTangentSpline!"));
-		return 0.f;
-	}
+	return UE::SplineComponent::ShouldUseSplineCurves()
+		? WarninglessSplineCurves().ReparamTable.Eval(Distance, 0.0f)
+		: Spline.GetParameterAtDistance(Distance);
 }
 
 float USplineComponent::GetTimeAtDistanceAlongSpline(float Distance) const
@@ -2425,20 +2091,9 @@ float USplineComponent::GetTimeAtDistanceAlongSpline(float Distance) const
 		return 0.0f;
 	}
 
-	float Param;
-	if (ShouldUseSplineCurves())
-	{
-		Param = WarninglessSplineCurves().ReparamTable.Eval(Distance, 0.0f);
-	}
-	else if (const FTangentSpline* TangentSpline = Spline.Get<UE::Spline::ESplineType::Tangent>())
-	{
-		Param = TangentSpline->GetParameterAtDistance(Distance);
-	}
-	else
-	{
-		ensureAlwaysMsgf(false, TEXT("ShouldUseSplineCurves() returned false but Spline failed to yield an FTangentSpline!"));
-		return 0.f;
-	}
+	const float Param = UE::SplineComponent::ShouldUseSplineCurves() 
+		? WarninglessSplineCurves().ReparamTable.Eval(Distance, 0.0f)
+		: Spline.GetParameterAtDistance(Distance);
 
 	// Convert Param to T and interpolate the duration interval (which yields Time).
 	return UE::Geometry::FInterval1f(0.f, Duration).Interpolate(GetInputKeyRange().GetT(Param));
@@ -2446,62 +2101,82 @@ float USplineComponent::GetTimeAtDistanceAlongSpline(float Distance) const
 
 FVector USplineComponent::GetLocationAtDistanceAlongSpline(float Distance, ESplineCoordinateSpace::Type CoordinateSpace) const
 {
-	const float Param = GetInputKeyValueAtDistanceAlongSpline(Distance);
+	const float Param = UE::SplineComponent::ShouldUseSplineCurves()
+		? WarninglessSplineCurves().ReparamTable.Eval(Distance, 0.0f)
+		: Spline.GetParameterAtDistance(Distance);
 	return GetLocationAtSplineInputKey(Param, CoordinateSpace);
 }
 
 FVector USplineComponent::GetTangentAtDistanceAlongSpline(float Distance, ESplineCoordinateSpace::Type CoordinateSpace) const
 {
-	const float Param = GetInputKeyValueAtDistanceAlongSpline(Distance);
+	const float Param = UE::SplineComponent::ShouldUseSplineCurves()
+		? WarninglessSplineCurves().ReparamTable.Eval(Distance, 0.0f)
+		: Spline.GetParameterAtDistance(Distance);
 	return GetTangentAtSplineInputKey(Param, CoordinateSpace);
 }
 
 
 FVector USplineComponent::GetDirectionAtDistanceAlongSpline(float Distance, ESplineCoordinateSpace::Type CoordinateSpace) const
 {
-	const float Param = GetInputKeyValueAtDistanceAlongSpline(Distance);
+	const float Param = UE::SplineComponent::ShouldUseSplineCurves()
+		? WarninglessSplineCurves().ReparamTable.Eval(Distance, 0.0f)
+		: Spline.GetParameterAtDistance(Distance);
 	return GetDirectionAtSplineInputKey(Param, CoordinateSpace);
 }
 
 FQuat USplineComponent::GetQuaternionAtDistanceAlongSpline(float Distance, ESplineCoordinateSpace::Type CoordinateSpace) const
 {
-	const float Param = GetInputKeyValueAtDistanceAlongSpline(Distance);
+	const float Param = UE::SplineComponent::ShouldUseSplineCurves()
+		? WarninglessSplineCurves().ReparamTable.Eval(Distance, 0.0f)
+		: Spline.GetParameterAtDistance(Distance);
 	return GetQuaternionAtSplineInputKey(Param, CoordinateSpace);
 }
 
 FRotator USplineComponent::GetRotationAtDistanceAlongSpline(float Distance, ESplineCoordinateSpace::Type CoordinateSpace) const
 {
-	const float Param = GetInputKeyValueAtDistanceAlongSpline(Distance);
+	const float Param = UE::SplineComponent::ShouldUseSplineCurves()
+		? WarninglessSplineCurves().ReparamTable.Eval(Distance, 0.0f)
+		: Spline.GetParameterAtDistance(Distance);
 	return GetRotationAtSplineInputKey(Param, CoordinateSpace);
 }
 
 FVector USplineComponent::GetUpVectorAtDistanceAlongSpline(float Distance, ESplineCoordinateSpace::Type CoordinateSpace) const
 {
-	const float Param = GetInputKeyValueAtDistanceAlongSpline(Distance);
+	const float Param = UE::SplineComponent::ShouldUseSplineCurves()
+		? WarninglessSplineCurves().ReparamTable.Eval(Distance, 0.0f)
+		: Spline.GetParameterAtDistance(Distance);
 	return GetUpVectorAtSplineInputKey(Param, CoordinateSpace);
 }
 
 FVector USplineComponent::GetRightVectorAtDistanceAlongSpline(float Distance, ESplineCoordinateSpace::Type CoordinateSpace) const
 {
-	const float Param = GetInputKeyValueAtDistanceAlongSpline(Distance);
+	const float Param = UE::SplineComponent::ShouldUseSplineCurves()
+		? WarninglessSplineCurves().ReparamTable.Eval(Distance, 0.0f)
+		: Spline.GetParameterAtDistance(Distance);
 	return GetRightVectorAtSplineInputKey(Param, CoordinateSpace);
 }
 
 float USplineComponent::GetRollAtDistanceAlongSpline(float Distance, ESplineCoordinateSpace::Type CoordinateSpace) const
 {
-	const float Param = GetInputKeyValueAtDistanceAlongSpline(Distance);
+	const float Param = UE::SplineComponent::ShouldUseSplineCurves()
+		? WarninglessSplineCurves().ReparamTable.Eval(Distance, 0.0f)
+		: Spline.GetParameterAtDistance(Distance);
 	return GetRollAtSplineInputKey(Param, CoordinateSpace);
 }
 
 FVector USplineComponent::GetScaleAtDistanceAlongSpline(float Distance) const
 {
-	const float Param = GetInputKeyValueAtDistanceAlongSpline(Distance);
+	const float Param = UE::SplineComponent::ShouldUseSplineCurves()
+		? WarninglessSplineCurves().ReparamTable.Eval(Distance, 0.0f)
+		: Spline.GetParameterAtDistance(Distance);
 	return GetScaleAtSplineInputKey(Param);
 }
 
 FTransform USplineComponent::GetTransformAtDistanceAlongSpline(float Distance, ESplineCoordinateSpace::Type CoordinateSpace, bool bUseScale) const
 {
-	const float Param = GetInputKeyValueAtDistanceAlongSpline(Distance);
+	const float Param = UE::SplineComponent::ShouldUseSplineCurves()
+		? WarninglessSplineCurves().ReparamTable.Eval(Distance, 0.0f)
+		: Spline.GetParameterAtDistance(Distance);
 	return GetTransformAtSplineInputKey(Param, CoordinateSpace, bUseScale);
 }
 
@@ -2679,40 +2354,18 @@ float USplineComponent::FindInputKeyClosestToWorldLocation(const FVector& WorldL
 {
 	const FVector LocalLocation = GetComponentTransform().InverseTransformPosition(WorldLocation);
 	float Dummy;
-
-	if (ShouldUseSplineCurves())
-	{
-		return WarninglessSplineCurves().Position.FindNearest(LocalLocation, Dummy);
-	}
-	else if (const FTangentSpline* TangentSpline = Spline.Get<UE::Spline::ESplineType::Tangent>())
-	{
-		return TangentSpline->FindNearest(LocalLocation, Dummy);
-	}
-	else
-	{
-		ensureAlwaysMsgf(false, TEXT("ShouldUseSplineCurves() returned false but Spline failed to yield an FTangentSpline!"));
-		return 0.f;
-	}
+	return UE::SplineComponent::ShouldUseSplineCurves()
+		? WarninglessSplineCurves().Position.FindNearest(LocalLocation, Dummy)
+		: Spline.FindNearest(LocalLocation, Dummy);
 }
 
 float USplineComponent::FindInputKeyOnSegmentClosestToWorldLocation(const FVector& WorldLocation, int32 Index) const
 {
 	const FVector LocalLocation = GetComponentTransform().InverseTransformPosition(WorldLocation);
 	float Dummy;
-
-	if (ShouldUseSplineCurves())
-	{
-		return WarninglessSplineCurves().Position.FindNearestOnSegment(LocalLocation, Index, Dummy);
-	}
-	else if (const FTangentSpline* TangentSpline = Spline.Get<UE::Spline::ESplineType::Tangent>())
-	{
-		return TangentSpline->FindNearestOnSegment(LocalLocation, Index, Dummy);
-	}
-	else
-	{
-		ensureAlwaysMsgf(false, TEXT("ShouldUseSplineCurves() returned false but Spline failed to yield an FTangentSpline!"));
-		return 0.f;
-	}
+	return UE::SplineComponent::ShouldUseSplineCurves()
+		? WarninglessSplineCurves().Position.FindNearestOnSegment(LocalLocation, Index, Dummy)
+		: Spline.FindNearestOnSegment(LocalLocation, Index, Dummy);
 }
 
 FVector USplineComponent::FindLocationClosestToWorldLocation(const FVector& WorldLocation, ESplineCoordinateSpace::Type CoordinateSpace) const
@@ -3000,26 +2653,25 @@ bool USplineComponent::ConvertSplineToPolyline_InDistanceRange(ESplineCoordinate
 	
 	int32 SegmentStart;
 	int32 SegmentEnd;
-	if (ShouldUseSplineCurves())
+	if (UE::SplineComponent::ShouldUseSplineCurves())
 	{
 		int32 StartIndex = WarninglessSplineCurves().ReparamTable.GetPointIndexForInputValue(RangeStart);
 		int32 EndIndex = WarninglessSplineCurves().ReparamTable.GetPointIndexForInputValue(RangeEnd);
 		SegmentStart = StartIndex / ReparamStepsPerSegment;
 		SegmentEnd = FMath::Min(NumSegments, 1 + EndIndex / ReparamStepsPerSegment);
 	}
-	else if (const FTangentSpline* TangentSpline = Spline.Get<UE::Spline::ESplineType::Tangent>())
+	else
 	{
-		const float StartParam = TangentSpline->GetParameterAtDistance(RangeStart);
-		const float EndParam = TangentSpline->GetParameterAtDistance(RangeEnd);
+		// with new spline, we are strongly assuming that point index is always same as param
+		// so, we can convert param directly to segment index.
+		// we do this without making the assumptions about the reparameterization table (unlike the other branch)
+
+		const float StartParam = Spline.GetParameterAtDistance(RangeStart);
+		const float EndParam = Spline.GetParameterAtDistance(RangeEnd);
 
 		SegmentStart = FMath::FloorToInt32(StartParam);
 		SegmentEnd = FMath::Min(NumSegments, FMath::CeilToInt32(EndParam));
 	}
-	else
-	{
-		ensureAlwaysMsgf(false, TEXT("ShouldUseSplineCurves() returned false but Spline failed to yield an FTangentSpline!"));
-		return false;
-	}
 
 	TArray<FVector> NewPoints;
 	TArray<double> NewDistances;
@@ -3168,7 +2820,7 @@ FPrimitiveSceneProxy* USplineComponent::CreateSceneProxy()
 		FSplinePDISceneProxy(const USplineComponent* InComponent)
 			: FPrimitiveSceneProxy(InComponent)
 			, bDrawDebug(InComponent->bDrawDebug)
-			, SplineInfo(InComponent->GetSplinePointsPosition())
+			, SplineInfo(UE::SplineComponent::ShouldUseSplineCurves() ? InComponent->WarninglessSplineCurves().Position : InComponent->Spline.GetSplinePointsPosition())
 #if WITH_EDITORONLY_DATA
 			, LineColor(InComponent->EditorUnselectedSplineSegmentColor)
 #else
@@ -3242,7 +2894,7 @@ FPrimitiveSceneProxy* USplineComponent::CreateSceneProxy()
 		FSplineMeshSceneProxy(const USplineComponent* InComponent)
 			: FPrimitiveSceneProxy(InComponent)
 			, bDrawDebug(InComponent->bDrawDebug)
-			, SplineInfo(InComponent->GetSplinePointsPosition())
+			, SplineInfo(UE::SplineComponent::ShouldUseSplineCurves() ? InComponent->WarninglessSplineCurves().Position : InComponent->Spline.GetSplinePointsPosition())
 #if WITH_EDITORONLY_DATA
 			, LineColor(InComponent->EditorUnselectedSplineSegmentColor)
 #else
@@ -3738,7 +3390,7 @@ void USplineComponent::Draw(FPrimitiveDrawInterface* PDI, const FSceneView* View
 
 FBoxSphereBounds USplineComponent::CalcBounds(const FTransform& LocalToWorld) const
 {
-	const auto& InterpCurve = GetSplinePointsPosition();
+	const auto& InterpCurve = UE::SplineComponent::ShouldUseSplineCurves() ? WarninglessSplineCurves().Position : Spline.GetSplinePointsPosition();
 
 #if SPLINE_FAST_BOUNDS_CALCULATION
 	FBox BoundingBox(0);
@@ -3749,8 +3401,8 @@ FBoxSphereBounds USplineComponent::CalcBounds(const FTransform& LocalToWorld) co
 
 	return FBoxSphereBounds(BoundingBox.TransformBy(LocalToWorld));
 #else
-	const int32 NumPoints = InterpCurve.Points.Num();
-	const int32 NumSegments = InterpCurve.bIsLooped ? NumPoints : NumPoints - 1;
+	const int32 NumPoints = GetNumberOfSplinePoints();
+	const int32 NumSegments = bClosedLoop ? NumPoints : NumPoints - 1;
 
 	FVector Min(WORLD_MAX);
 	FVector Max(-WORLD_MAX);
@@ -3792,7 +3444,7 @@ bool USplineComponent::GetIgnoreBoundsForEditorFocus() const
 
 TStructOnScope<FActorComponentInstanceData> USplineComponent::GetComponentInstanceData() const
 {
-	if (ShouldUseSplineCurves())
+	if (UE::SplineComponent::ShouldUseSplineCurves())
 	{
 		TStructOnScope<FActorComponentInstanceData> InstanceData = MakeStructOnScope<FActorComponentInstanceData, FSplineInstanceData>(this);
 		FSplineInstanceData* SplineInstanceData = InstanceData.Cast<FSplineInstanceData>();
@@ -3831,17 +3483,17 @@ TStructOnScope<FActorComponentInstanceData> USplineComponent::GetComponentInstan
 
 FSplineCurves USplineComponent::GetSplineCurves() const
 {
-	return ShouldUseSplineCurves()
+	return UE::SplineComponent::ShouldUseSplineCurves()
 		? WarninglessSplineCurves()
 		: FSplineCurves::FromSplineInterface(Spline);
 }
 
 FSpline USplineComponent::GetSpline() const
 {
-	if (ShouldUseSplineCurves())
+	if (UE::SplineComponent::ShouldUseSplineCurves())
 	{
 		FSpline Intermediate;
-		Intermediate = WarninglessSplineCurves();
+		Intermediate = SplineCurves;
 		return Intermediate;
 	}
 	else
@@ -3866,13 +3518,7 @@ void USplineComponent::SetSpline(const FSpline& InSpline)
 void USplineComponent::SetSpline(const FSplineCurves& InSplineCurves)
 {
 	WarninglessSplineCurves() = InSplineCurves;
-
-	// If Spline is already tangent backed or SplineComponent.UseSplineCurves is false, we can copy assign to the FSpline.
-	// This will force it to instantiate an FTangentSpline and populate it with the data from InSplineCurves.
-	if (Spline.GetCurrentImplementation() == UE::Spline::ESplineType::Tangent || !UE::SplineComponent::GUseSplineCurves)
-	{
-		Spline = InSplineCurves;
-	}
+	Spline = InSplineCurves;
 
 	if (UE::SplineComponent::GValidateOnChange)
 	{
@@ -3884,70 +3530,30 @@ void USplineComponent::SetSpline(const FSplineCurves& InSplineCurves)
 
 int32 USplineComponent::GetVersion() const
 {
-	if (ShouldUseSplineCurves())
-	{
-		return WarninglessSplineCurves().Version;
-	}
-	else if (const FTangentSpline* TangentSpline = Spline.Get<UE::Spline::ESplineType::Tangent>())
-	{
-		return TangentSpline->GetVersion();
-	}
-	else
-	{
-		ensureAlwaysMsgf(false, TEXT("ShouldUseSplineCurves() returned false but Spline failed to yield an FTangentSpline!"));
-		return 0;
-	}
+	return UE::SplineComponent::ShouldUseSplineCurves()
+		? WarninglessSplineCurves().Version
+		: Spline.GetVersion();
 }
 
 const FInterpCurveVector& USplineComponent::GetSplinePointsPosition() const
 {
-	if (ShouldUseSplineCurves())
-	{
-		return WarninglessSplineCurves().Position;
-	}
-	else if (const FTangentSpline* TangentSpline = Spline.Get<UE::Spline::ESplineType::Tangent>())
-	{
-		return TangentSpline->GetSplinePointsPosition();
-	}
-	else
-	{
-		ensureAlwaysMsgf(false, TEXT("ShouldUseSplineCurves() returned false but Spline failed to yield an FTangentSpline!"));
-		return DummyPointPosition;
-	}
+	return UE::SplineComponent::ShouldUseSplineCurves()
+		? WarninglessSplineCurves().Position
+		: Spline.GetSplinePointsPosition();
 }
 
 const FInterpCurveQuat& USplineComponent::GetSplinePointsRotation() const
 {
-	if (ShouldUseSplineCurves())
-	{
-		return WarninglessSplineCurves().Rotation;
-	}
-	else if (const FTangentSpline* TangentSpline = Spline.Get<UE::Spline::ESplineType::Tangent>())
-	{
-		return TangentSpline->GetSplinePointsRotation();
-	}
-	else
-	{
-		ensureAlwaysMsgf(false, TEXT("ShouldUseSplineCurves() returned false but Spline failed to yield an FTangentSpline!"));
-		return DummyPointRotation;
-	}
+	return UE::SplineComponent::ShouldUseSplineCurves()
+		? WarninglessSplineCurves().Rotation
+		: Spline.GetSplinePointsRotation();
 }
 
 const FInterpCurveVector& USplineComponent::GetSplinePointsScale() const
 {
-	if (ShouldUseSplineCurves())
-	{
-		return WarninglessSplineCurves().Scale;
-	}
-	else if (const FTangentSpline* TangentSpline = Spline.Get<UE::Spline::ESplineType::Tangent>())
-	{
-		return TangentSpline->GetSplinePointsScale();
-	}
-	else
-	{
-		ensureAlwaysMsgf(false, TEXT("ShouldUseSplineCurves() returned false but Spline failed to yield an FTangentSpline!"));
-		return DummyPointScale;
-	}
+	return UE::SplineComponent::ShouldUseSplineCurves()
+		? WarninglessSplineCurves().Scale
+		: Spline.GetSplinePointsScale();
 }
 
 TArray<ESplinePointType::Type> USplineComponent::GetEnabledSplinePointTypes() const
@@ -4113,11 +3719,6 @@ void USplineComponent::PostTransacted(const FTransactionObjectEvent& Transaction
 }
 #endif
 
-bool USplineComponent::ShouldUseSplineCurves() const
-{
-	return UE::SplineComponent::GUseSplineCurves || Spline.GetCurrentImplementation() != UE::Spline::ESplineType::Tangent;
-}
-
 UE::Geometry::FInterval1f USplineComponent::GetInputKeyRange() const
 {
 	UE::Geometry::FInterval1f Range;
diff --git a/Engine/Source/Runtime/Engine/Private/Curves/Spline.cpp b/Engine/Source/Runtime/Engine/Private/Curves/Spline.cpp
--- a/Engine/Source/Runtime/Engine/Private/Curves/Spline.cpp
+++ b/Engine/Source/Runtime/Engine/Private/Curves/Spline.cpp
@@ -4,22 +4,33 @@
 
 #include "Algo/Sort.h"
 #include "BoxTypes.h"
+#include "Async/TransactionallySafeMutex.h"
 #include "Components/SplineComponent.h"	// for FSplineCurves, ideally removed
-#include "Curves/Splines/TangentSpline.h"
 #include "Misc/Base64.h"
+#include "Misc/TransactionallySafeRWLock.h"
 #include "Serialization/MemoryReader.h"
 #include "Serialization/MemoryWriter.h"
 #include "UObject/FortniteMainBranchObjectVersion.h"
 
+PRAGMA_DISABLE_EXPERIMENTAL_WARNINGS
+#include "Splines/TangentBezierSpline.h"
+#include "Splines/MultiSpline.h"
+PRAGMA_ENABLE_EXPERIMENTAL_WARNINGS
+
 DEFINE_LOG_CATEGORY_STATIC(LogSpline, Log, All);
 
 #include UE_INLINE_GENERATED_CPP_BY_NAME(Spline)
 
 PRAGMA_DISABLE_EXPERIMENTAL_WARNINGS
 
+using UE::Geometry::FInterval1f;
+
+const FInterpCurveVector FSpline::PositionCurve;
+const FInterpCurveQuat FSpline::RotationCurve;
+const FInterpCurveVector FSpline::ScaleCurve;
+
 namespace UE::Spline
 {
-
 int32 GImplementation = 0;
 FAutoConsoleVariableRef CVarSplineImplementation(
 	TEXT("Spline.Implementation"),
@@ -30,173 +41,1379 @@
 #endif
 );
 
-inline UE::Spline::ESplineType GetDefaultSplineType()
+int32 GParameterizationPolicy = 0;
+FAutoConsoleVariableRef CVarParameterizationPolicy(
+	TEXT("Spline.ParameterizationPolicy"),
+	GParameterizationPolicy,
+	TEXT("0) Uniform - 1) Chord Length - 2) Centripetal")
+);
+	
+bool GApproximateTangents = false;
+FAutoConsoleVariableRef CVarApproximateTangents(
+	TEXT("Spline.ApproximateTangents"),
+	GApproximateTangents,
+	TEXT("True if we should approximate tangents using the central difference formula.")
+);
+
+bool GFallbackFindNearest = false;
+FAutoConsoleVariableRef CVarFallbackFindNearest(
+	TEXT("Spline.FallbackFindNearest"),
+	GFallbackFindNearest,
+	TEXT("True if we should implement FindNearest and FindNearestOnSegment using an intermediate spline representation. Only applies if Spline.Implementation == 2.")
+);
+
+bool GUseLegacyPositionEvaluation = false;
+FAutoConsoleVariableRef CVarUseLegacyPositionEvaluation(
+	TEXT("Spline.UseLegacyPositionEvaluation"),
+	GUseLegacyPositionEvaluation,
+	TEXT("If true, evaluating the position channel always routes through an interp curve.")
+);
+
+bool GUseLegacyRotationEvaluation = false;
+FAutoConsoleVariableRef CVarUseLegacyRotationEvaluation(
+	TEXT("Spline.UseLegacyRotationEvaluation"),
+	GUseLegacyRotationEvaluation,
+	TEXT("If true, evaluating the rotation channel always routes through an interp curve.")
+);
+
+bool GUseLegacyScaleEvaluation = false;
+FAutoConsoleVariableRef CVarUseLegacyScaleEvaluation(
+	TEXT("Spline.UseLegacyScaleEvaluation"),
+	GUseLegacyScaleEvaluation,
+	TEXT("If true, evaluating the scale channel always routes through an interp curve.")
+);
+
+bool GImmediatelyUpdateLegacyCurves = false;
+FAutoConsoleVariableRef CVarImmediatelyUpdateLegacyCurves(
+	TEXT("Spline.ImmediatelyUpdateLegacyCurves"),
+	GImmediatelyUpdateLegacyCurves,
+	TEXT("If true, mutating operations immediately rebuild legacy curves. If false, legacy curves are updated only when requested.")
+);
+
+// WARNING! THIS BREAKS EXISTING TEXT REPRESENTATION WHEN CHANGED, IT WILL ABSOLUTELY TRASH YOUR STUFF IF YOU CHANGE WILLY-NILLY
+bool GEncodeAsHex = true;
+// FAutoConsoleVariableRef CVarEncodeAsHex(
+// 	TEXT("Spline.EncodeAsHex"),
+// 	GEncodeAsHex,
+// 	TEXT("If true, splines are imported/exported using hexadecimal encoding. If false, the data is encoded in base 64.")
+// );
+
+bool GValidateRotScale = false;
+FAutoConsoleVariableRef CVarValidateRotScale(
+	TEXT("Spline.ValidateRotScale"),
+	GValidateRotScale,
+	TEXT("True if we should validate rotation and scale attributes when structurally modifying the spline.")
+);
+
+bool GSkipDoubleUpdate = true;
+FAutoConsoleVariableRef CVarSkipDoubleUpdate(
+	TEXT("Spline.SkipDoubleUpdate"),
+	GSkipDoubleUpdate,
+	TEXT("When true, skip updating the spline twice if UpdateSpline with parameters is called.")
+);
+
+}
+
+UE::Geometry::Spline::EParameterizationPolicy GetParameterizationPolicy()
 {
-	switch (GImplementation)
+	switch (UE::Spline::GParameterizationPolicy)
 	{
 	default:	// fallthrough
-	case 0:		return UE::Spline::ESplineType::Unimplemented;
-	case 1:		return UE::Spline::ESplineType::LegacyTangent;
-	case 2:		return UE::Spline::ESplineType::Tangent;
+	case 0: return UE::Geometry::Spline::EParameterizationPolicy::Uniform;
+	case 1: return UE::Geometry::Spline::EParameterizationPolicy::ChordLength;
+	case 2: return UE::Geometry::Spline::EParameterizationPolicy::Centripetal;
 	}
 }
 
+/**
+ * Converts from EInterpCurveMode to ETangentMode
+ */
+inline UE::Geometry::Spline::ETangentMode ConvertInterpCurveModeToTangentMode(EInterpCurveMode Mode)
+{
+	switch (Mode)
+	{
+	case CIM_Linear:
+		return UE::Geometry::Spline::ETangentMode::LegacyLinear; // Default to legacy for compatibility
+	case CIM_CurveAuto:
+		return UE::Geometry::Spline::ETangentMode::LegacyAuto; // Default to legacy for compatibility
+	case CIM_Constant:
+		return UE::Geometry::Spline::ETangentMode::Constant;
+	case CIM_CurveUser:
+		return UE::Geometry::Spline::ETangentMode::User;
+	case CIM_CurveBreak:
+		return UE::Geometry::Spline::ETangentMode::Broken;
+	case CIM_CurveAutoClamped:
+		return UE::Geometry::Spline::ETangentMode::LegacyAutoClamped; // Default to legacy for compatibility
+	case CIM_Unknown:
+	default:
+		return UE::Geometry::Spline::ETangentMode::Unknown;
+	}
 }
 
-using UE::Geometry::FInterval1f;
+/**
+ * Converts from ETangentMode to EInterpCurveMode
+ */
+inline EInterpCurveMode ConvertTangentModeToInterpCurveMode(UE::Geometry::Spline::ETangentMode Mode)
+{
+	switch (Mode)
+	{
+	case UE::Geometry::Spline::ETangentMode::LegacyLinear:
+	case UE::Geometry::Spline::ETangentMode::Linear:
+		return CIM_Linear;
+	case UE::Geometry::Spline::ETangentMode::Auto:
+	case UE::Geometry::Spline::ETangentMode::LegacyAuto:
+		return CIM_CurveAuto;
+	case UE::Geometry::Spline::ETangentMode::Constant:
+		return CIM_Constant;
+	case UE::Geometry::Spline::ETangentMode::User:
+		return CIM_CurveUser;
+	case UE::Geometry::Spline::ETangentMode::Broken:
+		return CIM_CurveBreak;
+	case UE::Geometry::Spline::ETangentMode::AutoClamped:
+	case UE::Geometry::Spline::ETangentMode::LegacyAutoClamped:
+		return CIM_CurveAutoClamped;
+	case UE::Geometry::Spline::ETangentMode::Unknown:
+	default:
+		return CIM_Unknown;
+	}
+}
 
+/**
+ * FNewSpline Definition
+ *
+ * A spline that provides tangent-based control over curve shape while using 
+ * piecewise Bezier curves internally for evaluation. Supports both manual
+ * tangent control and automatic tangent computation.
+ */
+
+class FNewSpline: public UE::Geometry::Spline::TMultiSpline<UE::Geometry::Spline::FTangentBezierSpline3d>
+{
+	using FTangentBezierControlPoint = UE::Geometry::Spline::TTangentBezierControlPoint<FVector3d>;
+	
+public:
+	
+    FNewSpline();
+    FNewSpline(const FNewSpline& Other);
+	FNewSpline(const FLegacySpline& Other);
+	FNewSpline(const FSplineCurves& Other);
+	virtual ~FNewSpline() override = default;
+    FNewSpline& operator=(const FNewSpline& Other);
+
+	bool operator!=(const FNewSpline& Other) const;
+
+    virtual bool Serialize(FArchive& Ar) override;
+    friend FArchive& operator<<(FArchive& Ar, FNewSpline& Spline)
+    {
+        Spline.Serialize(Ar);
+        return Ar;
+    }
+
+	/**
+	 * Returns the parameter space of FNewSpline. All public functions taking a parameter are only defined when the provided value is in this interval.
+	 * @return [0, NumSegments] when NumSegment != 0, otherwise undefined (empty interval).
+	 */
+	FInterval1f GetSegmentSpace() const;
+	float GetSegmentLength(const int32 Index, const float Param, const FVector& Scale3D = FVector(1.0f)) const;
+	float GetSplineLength();
+	
+    virtual void SetClosedLoop(bool bInClosedLoop) override;
+	void SetClosedLoop(bool bInClosedLoop, bool bUpdateSpline);
+
+	void Reset();
+	
+	void ResetRotation();
+	void ResetScale();
+	
+    void AddPoint(const FSplinePoint& Point);
+	void InsertPoint(const FSplinePoint& Point, int32 Index);
+	FSplinePoint GetPoint(int32 Index) const;
+    void RemovePoint(int32 Index);
+
+	void SetLocation(int32 Index, const FVector& InLocation);
+	FVector GetLocation(const int32 Index) const;
+	
+	void SetInTangent(const int32 Index, const FVector& InTangent);
+	FVector GetInTangent(const int32 Index) const;
+	
+	void SetOutTangent(const int32 Index, const FVector& OutTangent);
+	FVector GetOutTangent(const int32 Index) const;
+
+	void SetRotation(int32 Index, const FQuat& InRotation);
+	FQuat GetRotation(const int32 Index) const;
+	
+	void SetScale(int32 Index, const FVector& InScale);
+	FVector GetScale(const int32 Index) const;
+	
+	void SetSplinePointType(int32 Index, EInterpCurveMode Type);
+	EInterpCurveMode GetSplinePointType(int32 Index) const;
+
+	float GetParameterAtIndex(int32 Index) const;
+	float GetParameterAtDistance(float Distance);
+	float GetDistanceAtParameter(float Parameter);
+
+    virtual float FindNearest(const FVector& InLocation, float& OutSquaredDistance) const override;
+    float FindNearestOnSegment(const FVector& InLocation, int32 SegmentIndex, float& OutSquaredDistance) const;
+	
+	/** Parameter based evaluation API. */
+	FVector EvaluatePosition(float Parameter) const;	// Not named Evaluate because we would be shadowing non-virtual base, this maps parameters
+    FVector EvaluateDerivative(float Parameter) const;
+	FQuat EvaluateRotation(float Parameter) const;
+	FVector EvaluateScale(float Parameter) const;
+
+	// Templated Channel Creation/Query functions
+	template <typename AttrType> int32 NumAttributeValues(FName Name) const;
+
+	// Templated Attribute Interaction by Index functions
+	template <typename AttrType> AttrType GetAttributeValue(const FName& Name, int32 Index) const;
+	template <typename AttrType> void SetAttributeValue(FName Name, const AttrType& Value, int32 Index);
+	template <typename AttrType> void RemoveAttributeValue(FName Name, int32 Index);
+	template <typename AttrType> int32 SetAttributeParameter(const FName& Name, int32 Index, float ParentSpaceParameter);
+	template <typename AttrType> float GetAttributeParameter(const FName& Name, int32 Index) const;
+
+	// Templated Attribute Interaction by Parameter functions
+	template <typename AttrType> int32 AddAttributeValue(FName Name, const AttrType& Value, float Parameter);
+	
+	/** Updates the spline using the current configuration. */
+	void UpdateSpline();
+
+	/** Updates the spline configuration, then updates the spline. */
+	void UpdateSpline(const FSpline::FUpdateSplineParams& Params);
+
+	const FInterpCurveVector& GetSplinePointsPosition() const { return PositionCurve; }
+	const FInterpCurveQuat& GetSplinePointsRotation() const { return RotationCurve; }
+	const FInterpCurveVector& GetSplinePointsScale() const { return ScaleCurve; }
+
+private:
+
+	/** Inspired by FDynamicMesh3::FChangeStamp. */
+	struct FChangeStamp
+	{
+		FChangeStamp()
+		{}
+		
+		UE_NONCOPYABLE(FChangeStamp);
+
+		/** Updates the change stamp in an thread-safe, transactionally-safe way. */
+		void Increment()
+		{
+			Mutex.Lock();
+			++Value;
+			Mutex.Unlock();
+		}
+
+		/** Returns the current change value in a thread-safe, transactionally-safe way. */
+		uint32 GetValue() const
+		{
+			Mutex.Lock();
+			uint32 Result = Value;
+			Mutex.Unlock();
+
+			return Result;
+		}
+
+	private:
+		
+		/** Guards `Value`. */
+		mutable UE::FTransactionallySafeMutex Mutex;
+
+		/** The change stamp is incremented when modifications occur. It's guarded by `Mutex`. */
+		uint32 Value = 1;
+	};
+
+	FInterpCurveFloat ReparamTable;
+	FChangeStamp ReparamTableNextVersion;
+	uint32 ReparamTableVersion = 0;
+	FTransactionallySafeRWLock ReparamTableRWLock;
+	
+	static inline const FName RotationAttrName = FName("Rotation");
+	static inline const FName ScaleAttrName = FName("Scale");
+
+	/** Legacy Curves: */
+    mutable FInterpCurveVector PositionCurve;
+	mutable FInterpCurveQuat RotationCurve;
+	mutable FInterpCurveVector ScaleCurve;
+	FChangeStamp LegacyCurvesNextVersion;
+	mutable uint32 LegacyCurvesVersion = 0;
+	
+	int32 ReparamStepsPerSegment = 10;
+
+private:
+
+	// In order for FSpline to access the parameter mapping functions below.
+	// The alternative is to redefine a bunch of MultiSpline functions to do the remapping internally, but this is a bit of a pain.
+	friend struct FSpline;
+	
+	void Reparameterize();
+
+	/** Convert from [0, 1] to [0, NumSegments] */
+	float FromNormalizedSpace(float Parameter) const;
+	
+	/** Convert from [0, NumSegments] to [0, 1] */
+	float ToNormalizedSpace(float Parameter) const;
+
+	/** Convert from internal spline space to [0, NumSegments] */
+	float FromInternalSplineSpace(float Parameter) const;
+	
+	/** Convert from [0, NumSegments] to internal spline space */
+	float ToInternalSplineSpace(float Parameter) const;
+
+	void MarkReparamTableDirty();
+	
+	/** Clears and repopulates reparameterization attribute channels. */
+	void UpdateReparamTable();
+
+	void MarkLegacyCurvesDirty();
+
+	/** Clears and repopulates PositionCurve, RotationCurve, and ScaleCurve */
+	void RebuildLegacyCurves() const;
+
+	/** Converts FSplinePoint to FTangentBezierControlPoint */
+    FTangentBezierControlPoint ConvertToTangentBezierControlPoint(const FSplinePoint& Point) const;
+
+    /**
+     * Updates the attributes of a point in the spline
+     * @param Point 
+     * @param PointIndex 
+     */
+    void UpdatePointAttributes(const FSplinePoint& Point, int32 PointIndex);
+	
+	/**
+	 * Converts an index (possibly fractional) to a parameter value in internal spline space.
+	 */
+	float ConvertIndexToInternalParameter(int32 Index, float Fraction = 0.0f) const;
+
+    /**
+	 * Converts an internal spline space parameter value to the closest index
+	 */
+	int32 ConvertInternalParameterToNearestPointIndex(float Parameter) const;
+
+	void ValidateRotScale() const;
+
+	// Helper functions for keeping attribute channels stable when reparameterizing.
+	template <typename AttrType> TArray<AttrType> CacheChannelKnots(const FName& InName) const;
+	template <typename AttrType> void RestoreChannelKnots(const FName& InName, const TArray<AttrType>& InKnots);
+};
+
+template <typename AttrType>
+int32 FNewSpline::AddAttributeValue(FName Name, const AttrType& Value, float Parameter)
+{
+	using AttrSplineType = UE::Geometry::Spline::TTangentBezierSpline<AttrType>;
+	using ControlPointType = UE::Geometry::Spline::TTangentBezierControlPoint<AttrType>;
+	
+	if (AttrSplineType* Channel = GetTypedAttributeChannel<AttrSplineType>(Name))
+	{
+		ControlPointType ControlPoint = ControlPointType(Value);
+		int32 NumPoints = Channel->GetNumPoints();
+
+		float ParentSpaceParameter = ToInternalSplineSpace(Parameter);
+		FInterval1f MappedChildSpace = GetMappedChildSpace(Name);
+		
+		// Special case for empty spline
+		if (NumPoints == 0)
+		{
+			Channel->AppendPoint(ControlPoint);
+			SetAttributeChannelRange(Name, FInterval1f(ParentSpaceParameter, ParentSpaceParameter), EMappingRangeSpace::Parent);
+			return 0;
+		}
+
+		if (NumPoints == 1)
+		{
+			if (ParentSpaceParameter > MappedChildSpace.Min)
+			{
+				Channel->AppendPoint(ControlPoint);
+				Channel->Reparameterize(UE::Geometry::Spline::EParameterizationPolicy::Uniform);
+				SetAttributeChannelRange(Name, FInterval1f(MappedChildSpace.Min, ParentSpaceParameter), EMappingRangeSpace::Parent);
+				return 1;
+			}
+			else
+			{
+				Channel->PrependPoint(ControlPoint);
+				Channel->Reparameterize(UE::Geometry::Spline::EParameterizationPolicy::Uniform);
+				SetAttributeChannelRange(Name, FInterval1f(ParentSpaceParameter, MappedChildSpace.Min), EMappingRangeSpace::Parent);
+				return 0;
+			}
+		}
+
+		// append case
+		if (ParentSpaceParameter > MappedChildSpace.Max)
+		{
+			// It is important to compute ChildSpaceParameter before the AppendPoint call, otherwise it will not be correct.
+			float ChildSpaceParameter = MapParameterToChildSpace(Name, ParentSpaceParameter);
+			Channel->AppendPoint(ControlPoint);
+
+			// By growing the child space and the mapped parent range proportionally, we keep the internal points stable in parent space.
+			Channel->SetParameter(Channel->GetNumPoints() - 1, ChildSpaceParameter);
+			SetAttributeChannelRange(Name, FInterval1f(MappedChildSpace.Min, ParentSpaceParameter), EMappingRangeSpace::Parent);
+			
+			return Channel->GetNumPoints() - 1;
+		}
+
+		// prepend case
+		if (ParentSpaceParameter < MappedChildSpace.Min)
+		{
+			// It is important to compute ChildSpaceParameter before the AppendPoint call, otherwise it will not be correct.
+			float ChildSpaceParameter = MapParameterToChildSpace(Name, ParentSpaceParameter);
+			Channel->PrependPoint(ControlPoint);
+
+			// By growing the child space and the mapped parent range proportionally, we keep the internal points stable in parent space.
+			Channel->SetParameter(0, ChildSpaceParameter);
+			SetAttributeChannelRange(Name, FInterval1f(ParentSpaceParameter, MappedChildSpace.Max), EMappingRangeSpace::Parent);
+			
+			return 0;
+		}
+		
+		float ChildSpaceParameter = MapParameterToChildSpace(Name, ParentSpaceParameter);
+		float LocalT;
+		int32 SegmentIndex = Channel->FindSegmentIndex(ChildSpaceParameter, LocalT);
+		
+		int NewPointIndex = Channel->InsertPointAtSegmentParam(SegmentIndex, LocalT, ControlPoint);
+		return NewPointIndex;
+	}
+
+	return INDEX_NONE;
+}
+
+template <typename AttrType>
+void FNewSpline::SetAttributeValue(FName Name, const AttrType& Value, int32 Index)
+{
+	using AttrSplineType = UE::Geometry::Spline::TTangentBezierSpline<AttrType>;
+
+	if (AttrSplineType* Channel = GetTypedAttributeChannel<AttrSplineType>(Name))
+	{
+		Channel->SetValue(Index, Value);
+	}
+}
+
+template <typename AttrType>
+void FNewSpline::RemoveAttributeValue(FName Name, int32 Index)
+{
+	using AttrSplineType = UE::Geometry::Spline::TTangentBezierSpline<AttrType>;
+
+	if (AttrSplineType* Channel = GetTypedAttributeChannel<AttrSplineType>(Name))
+	{
+		if (Channel->GetNumPoints() > 1)
+		{
+			if (Index == 0)
+			{
+				FInterval1f NewMappingRange = GetMappedChildSpace(Name);
+				NewMappingRange.Min = ToInternalSplineSpace(GetAttributeParameter<AttrType>(Name, Index + 1));
+				SetAttributeChannelRange(Name, NewMappingRange, EMappingRangeSpace::Parent);
+			}
+			else if (Index == Channel->GetNumPoints() - 1)
+			{
+				FInterval1f NewMappingRange = GetMappedChildSpace(Name);
+				NewMappingRange.Max = ToInternalSplineSpace(GetAttributeParameter<AttrType>(Name, Index - 1));
+				SetAttributeChannelRange(Name, NewMappingRange, EMappingRangeSpace::Parent);
+			}
+		}
+		
+		Channel->RemovePoint(Index);
+	}
+}
+
+template <typename AttrType>
+AttrType FNewSpline::GetAttributeValue(const FName& Name, int32 Index) const
+{
+	using AttrSplineType = UE::Geometry::Spline::TTangentBezierSpline<AttrType>;
+	
+	if (AttrSplineType* Channel = GetTypedAttributeChannel<AttrSplineType>(Name))
+	{
+		return Channel->GetValue(Index);
+	}
+
+	return AttrType();
+}
+
+template <typename AttrType>
+int32 FNewSpline::NumAttributeValues(FName Name) const
+{
+	using AttrSplineType = UE::Geometry::Spline::TTangentBezierSpline<AttrType>;
+	
+	if (AttrSplineType* Channel = GetTypedAttributeChannel<AttrSplineType>(Name))
+	{
+		return Channel->GetNumPoints();
+	}
+
+	return 0;
+}
+
+template <typename AttrType>
+int32 FNewSpline::SetAttributeParameter(const FName& Name, int32 Index, float Parameter)
+{
+	using AttrSplineType = UE::Geometry::Spline::TTangentBezierSpline<AttrType>;
+	
+	if (AttrSplineType* Channel = GetTypedAttributeChannel<AttrSplineType>(Name))
+	{
+		float ParentSpaceParameter = ToInternalSplineSpace(Parameter);
+
+		// Prevent collapse of the channel space and keep it starting at 0.
+		auto SanitizeChannelSpace = [Channel]()
+		{
+			const FInterval1f ChannelSpace = Channel->GetParameterSpace();
+			const float KnotOffset = -1.f * ChannelSpace.Min;
+			const float KnotScaleFactor = 1.f / (ChannelSpace.Max - ChannelSpace.Min);
+			TArray<UE::Geometry::Spline::FKnot> Knots = Channel->GetKnotVector();
+			for (UE::Geometry::Spline::FKnot& Knot : Knots)
+			{
+				Knot.Value += KnotOffset;
+				Knot.Value *= KnotScaleFactor;
+			}
+			Channel->SetKnotVector(Knots);
+		};
+		
+		// Helper to prevent code duplication.
+		// Uses current value of Index and ParentSpaceParameter at the time of calling to update the attribute parameter.
+		auto SetParentSpaceParameter = [this, &ParentSpaceParameter, Name, Channel, &Index]() -> int32
+		{
+			const float CurrentChildSpaceParameter = MapParameterToChildSpace(Name, ToInternalSplineSpace(GetAttributeParameter<AttrType>(Name, Index)));
+			const float DesiredChildSpaceParameter = MapParameterToChildSpace(Name, ParentSpaceParameter);
+
+			constexpr float MinStep = 2.f * UE_KINDA_SMALL_NUMBER;
+			float NewChildSpaceParameter = DesiredChildSpaceParameter > CurrentChildSpaceParameter
+				? FMath::Max(DesiredChildSpaceParameter, CurrentChildSpaceParameter + MinStep)
+				: FMath::Min(DesiredChildSpaceParameter, CurrentChildSpaceParameter - MinStep);
+			
+			return Channel->SetParameter(Index, NewChildSpaceParameter);
+		};
+
+		/**
+		 * Cases we handle below:
+		 * 1) Moving the only existing attribute.
+		 * 2) Moving an endpoint for a 2 point channel.
+		 * 3) Moving the first endpoint for a 3+ point channel.
+		 * 4) Moving the last endpoint for a 3+ point channel.
+		 * 5) Moving an internal point.
+		 */
+		
+		if (Channel->GetNumPoints() == 1)						// Case 1: Moving the only attribute
+		{
+			SetAttributeChannelRange(Name, FInterval1f(ParentSpaceParameter, ParentSpaceParameter), EMappingRangeSpace::Parent);
+			return Index;
+		}
+		else if (Channel->GetNumPoints() == 2)					// Case 2: Moving an end point while no internal points exist.
+		{
+			FInterval1f MappedRange = GetMappedChildSpace(Name);
+			float& RangeBound = Index == 0 ? MappedRange.Min : MappedRange.Max;
+			RangeBound = ParentSpaceParameter;
+
+			if (MappedRange.Min > MappedRange.Max)
+			{
+				// The mapping range will flip, swap end points and un-flip.
+				Channel->SetParameter(1, Channel->GetParameter(0) - 1.f);	// It doesn't actually matter what we do here, as long as we get the ordering & mapping range correct.
+				SetAttributeChannelRange(Name, FInterval1f(MappedRange.Max, MappedRange.Min), EMappingRangeSpace::Parent);
+				Index = Index == 0 ? 1 : 0;
+			}
+			else
+			{
+				SetAttributeChannelRange(Name, MappedRange, EMappingRangeSpace::Parent);
+			}
+
+			return Index;
+		}
+		else if (Index == 0)									// Case 3: Moving the first point (which has exactly 1 neighbor)
+		{
+			// Save internal attribute point parameters for later re-parameterization after changing the child space mapping.
+			TArray<float> InternalParameters = CacheChannelKnots<float>(Name);
+
+			float ParentSpaceUpperBound = GetMappedChildSpace(Name).Max;
+			const float NeighborParentSpaceParameter = ToInternalSplineSpace(GetAttributeParameter<AttrType>(Name, Index + 1));
+			const bool bAttributesWillReorder = ParentSpaceParameter > NeighborParentSpaceParameter;
+
+			// If the end point is passing its neighbor, we need to actually reshuffle the points.
+			if (bAttributesWillReorder)
+			{
+				Index = SetParentSpaceParameter();
+
+				// Shift values left up to the last invalidated index.
+				const int32 InvalidatedInternalParameterIdx = FMath::Clamp(Index - 1, 0, InternalParameters.Num() - 1);
+				for (int InternalParameterIdx = 1; InternalParameterIdx <= InvalidatedInternalParameterIdx; ++InternalParameterIdx)
+				{
+					InternalParameters[InternalParameterIdx - 1] = InternalParameters[InternalParameterIdx];
+				}
+				
+				// Do one of 2 things to the actual invalidated index:
+				if (Index == Channel->GetNumPoints() - 1)
+				{
+					// Old upper end point is now an internal point and needs to be re-parameterized.
+					InternalParameters[InvalidatedInternalParameterIdx] = ParentSpaceUpperBound;
+				}
+				else
+				{
+					// Old lower end point is now an internal point and needs to be re-parameterized.
+					InternalParameters[InvalidatedInternalParameterIdx] = ParentSpaceParameter;
+				}
+			}
+			
+			// Updates the lower bound of the mapping range.
+			ParentSpaceUpperBound = Index == Channel->GetNumPoints() - 1 ? ParentSpaceParameter : ParentSpaceUpperBound;
+			const float ParentSpaceLowerBound = bAttributesWillReorder ? NeighborParentSpaceParameter : ParentSpaceParameter;
+			SetAttributeChannelRange(Name, FInterval1f(ParentSpaceLowerBound, ParentSpaceUpperBound), EMappingRangeSpace::Parent);
+
+			// Re-parameterize the internal points so that they do not move in parent space after the mapping range changed
+			RestoreChannelKnots(Name, InternalParameters);
+			
+			if (bAttributesWillReorder)
+			{
+				SanitizeChannelSpace();
+			}
+			
+			return Index;
+		}
+		else if (Index == Channel->GetNumPoints() - 1)			// Case 4: Moving the last point (which has exactly 1 neighbor)
+		{
+			// Save internal attribute point parameters for later re-parameterization after changing the child space mapping.
+			TArray<float> InternalParameters = CacheChannelKnots<float>(Name);
+
+			float ParentSpaceLowerBound = GetMappedChildSpace(Name).Min;
+			const float NeighborParentSpaceParameter = ToInternalSplineSpace(GetAttributeParameter<AttrType>(Name, Index - 1));
+			const bool bAttributesWillReorder = ParentSpaceParameter < NeighborParentSpaceParameter;
+			
+			// If the end point is passing its neighbor, we need to actually reshuffle the points.
+			if (bAttributesWillReorder)
+			{
+				Index = SetParentSpaceParameter();
+
+				// Shift values right up to the last invalidated index.
+				const int32 InvalidatedInternalParameterIdx = FMath::Clamp(Index - 1, 0, InternalParameters.Num() - 1);
+				for (int InternalParameterIdx = InternalParameters.Num() - 2; InternalParameterIdx >= InvalidatedInternalParameterIdx; --InternalParameterIdx)
+				{
+					InternalParameters[InternalParameterIdx + 1] = InternalParameters[InternalParameterIdx];
+				}
+				
+				// Do one of 2 things to the actual invalidated index:
+				if (Index == 0)
+				{
+					// Old upper end point is now an internal point and needs to be re-parameterized.
+					InternalParameters[InvalidatedInternalParameterIdx] = ParentSpaceLowerBound;
+				}
+				else
+				{
+					// Old lower end point is now an internal point and needs to be re-parameterized.
+					InternalParameters[InvalidatedInternalParameterIdx] = ParentSpaceParameter;
+				}
+			}
+
+			// Update the child space mapping
+			ParentSpaceLowerBound = Index == 0 ? ParentSpaceParameter : ParentSpaceLowerBound;
+			float ParentSpaceUpperBound = bAttributesWillReorder ? NeighborParentSpaceParameter : ParentSpaceParameter; // Prevent over-shrinking if we reshuffle.
+			SetAttributeChannelRange(Name, FInterval1f(ParentSpaceLowerBound, ParentSpaceUpperBound), EMappingRangeSpace::Parent);
+			
+			// Re-parameterize the internal points so that they do not move in parent space after the mapping range changed
+			RestoreChannelKnots(Name, InternalParameters);
+
+			if (bAttributesWillReorder)
+			{
+				SanitizeChannelSpace();
+			}
+		
+			return Index;
+		}
+		else													// Case 5: Moving an internal point (which has exactly 2 neighbors)
+		{
+			FInterval1f MappedRange = GetMappedChildSpace(Name);
+			Index = SetParentSpaceParameter();
+			
+			if (Index == 0)
+			{
+				// The internal point is now the lower end point
+				MappedRange.Min = ParentSpaceParameter;
+			}
+			else if (Index == Channel->GetNumPoints() - 1)
+			{
+				// The internal point is now the upper end point
+				MappedRange.Max = ParentSpaceParameter;
+			}
+			
+			SetAttributeChannelRange(Name, MappedRange, EMappingRangeSpace::Parent);
+			
+			return Index;
+		}
+	}
+	
+	return INDEX_NONE;
+}
+
+template <typename AttrType>
+float FNewSpline::GetAttributeParameter(const FName& Name, int32 Index) const
+{
+	using AttrSplineType = UE::Geometry::Spline::TTangentBezierSpline<AttrType>;
+	
+	if (AttrSplineType* Channel = GetTypedAttributeChannel<AttrSplineType>(Name))
+	{
+		return FromInternalSplineSpace(MapParameterFromChildSpace(Name, Channel->GetParameter(Index)));
+	}
+	return 0.0f;
+}
+
+/** FLegacySpline Definition */
+
+class FLegacySpline
+{
+	FInterpCurveVector PositionCurve;
+	FInterpCurveQuat RotationCurve;
+	FInterpCurveVector ScaleCurve;
+	FInterpCurveFloat ReparamTable;
+	
+public:
+
+	ENGINE_API FLegacySpline() = default;
+	ENGINE_API FLegacySpline(const FNewSpline& Other);
+	ENGINE_API FLegacySpline(const FSplineCurves& Other);
+	
+	/* Control Point Index Interface */
+
+	void AddPoint(const FSplinePoint& InPoint);
+	void InsertPoint(const FSplinePoint& InPoint, int32 Index);
+	FSplinePoint GetPoint(const int32 Index) const;
+	void RemovePoint(int32 Index);
+	
+	void SetLocation(int32 Index, const FVector& InLocation);
+	FVector GetLocation(const int32 Index) const;
+	
+	void SetInTangent(const int32 Index, const FVector& InTangent);
+	FVector GetInTangent(const int32 Index) const;
+	
+	void SetOutTangent(const int32 Index, const FVector& OutTangent);
+	FVector GetOutTangent(const int32 Index) const;
+
+	void SetRotation(int32 Index, const FQuat& InRotation);
+	FQuat GetRotation(const int32 Index) const;
+	
+	void SetScale(int32 Index, const FVector& InScale);
+	FVector GetScale(const int32 Index) const;
+	
+	void SetSplinePointType(int32 Index, EInterpCurveMode Type);
+	EInterpCurveMode GetSplinePointType(int32 Index) const;
+
+	float GetParameterAtIndex(int32 Index) const;
+	float GetParameterAtDistance(float Distance) const;
+	float GetDistanceAtParameter(float Parameter) const;
+	
+	/* Parameter Interface */
+
+	FVector Evaluate(float Param) const;
+	FVector EvaluateDerivative(float Param) const;
+	FQuat EvaluateRotation(float Param) const;
+	FVector EvaluateScale(float Param) const;
+	float FindNearest(const FVector& InLocation, float& OutSquaredDist) const;
+	float FindNearestOnSegment(const FVector& InLocation, int32 SegmentIndex, float& OutSquaredDist) const;
+	
+	/* Misc Interface */
+	
+	bool operator==(const FLegacySpline& Other) const
+	{
+		return PositionCurve == Other.PositionCurve
+			&& RotationCurve == Other.RotationCurve
+			&& ScaleCurve == Other.ScaleCurve;
+	}
+
+	bool operator!=(const FLegacySpline& Other) const
+	{
+		return !(*this == Other);
+	}
+	
+	friend ENGINE_API FArchive& operator<<(FArchive& Ar, FLegacySpline& Spline)
+	{
+		Spline.Serialize(Ar);
+		return Ar;
+	}
+	bool Serialize(FArchive& Ar);
+	
+	FInterpCurveVector& GetSplinePointsPosition() { return PositionCurve; }
+	const FInterpCurveVector& GetSplinePointsPosition() const { return PositionCurve; }
+	FInterpCurveQuat& GetSplinePointsRotation() { return RotationCurve; }
+	const FInterpCurveQuat& GetSplinePointsRotation() const { return RotationCurve; }
+	FInterpCurveVector& GetSplinePointsScale() { return ScaleCurve; }
+	const FInterpCurveVector& GetSplinePointsScale() const { return ScaleCurve; }
+
+	/** Returns the length of the specified spline segment up to the parametric value given */
+	float GetSegmentLength(const int32 Index, const float Param, const FVector& Scale3D = FVector(1.0f)) const;
+
+	/** Returns total length along this spline */
+	float GetSplineLength() const;
+	
+	/** Returns the total number of control points on this spline. */
+	int32 GetNumControlPoints() const;
+	
+	/** Reset the spline to an empty spline. */
+	void Reset();
+	
+	/** Reset the rotation attribute channel to default values. */
+	void ResetRotation();
+	
+	/** Reset the scale attribute channel to default values. */
+	void ResetScale();
+	
+	/** Update the spline's internal data according to the passed-in params. */
+	void UpdateSpline(const FSpline::FUpdateSplineParams& InParams);
+};
+
+
+
 /** FSpline Implementation */
 
 FSpline::FSpline()
+: CurrentImplementation(UE::Spline::GImplementation)
+, Version(0xffffffff)
 {
-	SetCurrentImplementation(UE::Spline::GetDefaultSplineType());
-
 #if WITH_EDITOR
-	PreviousImplementation = UE::Spline::ESplineType::Unimplemented;
+	PreviousImplementation = 0;
 #endif
+	
+	switch (CurrentImplementation)
+	{
+	default: break;		// Intentionally doing nothing
+	case 1: LegacyData = MakeShared<FLegacySpline>(); break;
+	case 2: NewData = MakeShared<FNewSpline>(); break;
+	}
 }
 
-FSpline::FSpline(UE::Spline::ESplineType InType)
-{
-	SetCurrentImplementation(InType);
-
-#if WITH_EDITOR
-	PreviousImplementation = UE::Spline::ESplineType::Unimplemented;
-#endif
-}
-
 FSpline::FSpline(const FSpline& Other)
+	: FSpline()
 {
 	*this = Other;
 }
 
 FSpline& FSpline::operator=(const FSpline& Other)
 {
-	SetCurrentImplementation(Other.CurrentImplementation);
+	CurrentImplementation = Other.CurrentImplementation;
 
 	switch (CurrentImplementation)
 	{
 	default:
 		break;
-
-	case UE::Spline::ESplineType::LegacyTangent:
-		*LegacySplineImpl = *Other.LegacySplineImpl;
+	case 1:
+		*LegacyData = Other.LegacyData ? *Other.LegacyData : FLegacySpline();
 		break;
-
-	case UE::Spline::ESplineType::Tangent:
-		*Get<UE::Spline::ESplineType::Tangent>() = *Other.Get<UE::Spline::ESplineType::Tangent>();
+	case 2:
+		*NewData = Other.NewData ? *Other.NewData : FNewSpline();
 		break;
 	}
-		
+	
+	Version++;
+	
 	return *this;
 }
 
 FSpline& FSpline::operator=(const FSplineCurves& Other)
 {
-	SetCurrentImplementation(UE::Spline::ESplineType::Tangent);
-
-	*Get<UE::Spline::ESplineType::Tangent>() = FTangentSpline(Other);
+	switch (CurrentImplementation)
+	{
+	default:
+		break;
+	case 1:
+		*LegacyData = FLegacySpline(Other);
+		break;
+	case 2:
+		*NewData = FNewSpline(Other);
+		break;
+	}
+	
+	Version++;
 	
 	return *this;
 }
 
+void FSpline::AddPoint(const FSplinePoint& InPoint)
+{
+	switch (CurrentImplementation)
+	{
+	default: break;		// Intentionally doing nothing
+	case 1: LegacyData->AddPoint(InPoint); break;
+	case 2: NewData->AddPoint(InPoint); break;
+	}
+}
+
+void FSpline::InsertPoint(const FSplinePoint& InPoint, int32 Index)
+{
+	switch (CurrentImplementation)
+	{
+	default: break;		// Intentionally doing nothing
+	case 1: LegacyData->InsertPoint(InPoint, Index); break;
+	case 2: NewData->InsertPoint(InPoint, Index); break;
+	}
+}
+
+FSplinePoint FSpline::GetPoint(const int32 Index) const
+{
+	switch (CurrentImplementation)
+	{
+	default: return FSplinePoint();		// Intentionally doing nothing
+	case 1: return LegacyData->GetPoint(Index);
+	case 2: return NewData->GetPoint(Index);
+	}
+}
+
+void FSpline::RemovePoint(const int32 Index)
+{
+	switch (CurrentImplementation)
+	{
+	default: break;		// Intentionally doing nothing
+	case 1: LegacyData->RemovePoint(Index); break;
+	case 2: NewData->RemovePoint(Index); break;
+	}
+}
+
+void FSpline::SetLocation(int32 Index, const FVector& InLocation)
+{
+	switch (CurrentImplementation)
+	{
+	default: break;		// Intentionally doing nothing
+	case 1: LegacyData->SetLocation(Index, InLocation); break;
+	case 2: NewData->SetLocation(Index, InLocation); break;
+	}
+}
+
+FVector FSpline::GetLocation(const int32 Index) const
+{
+	switch (CurrentImplementation)
+	{
+	default: return FVector();
+	case 1: return LegacyData->GetLocation(Index);
+	case 2: return NewData->GetLocation(Index);
+	}
+}
+
+void FSpline::SetInTangent(const int32 Index, const FVector& InTangent)
+{
+	switch (CurrentImplementation)
+	{
+	default: break;
+	case 1: LegacyData->SetInTangent(Index, InTangent); break;
+	case 2: NewData->SetInTangent(Index, InTangent); break;
+	}
+}
+
+FVector FSpline::GetInTangent(const int32 Index) const
+{
+	switch (CurrentImplementation)
+	{
+	default: return FVector();
+	case 1: return LegacyData->GetInTangent(Index);
+	case 2: return NewData->GetInTangent(Index);
+	}
+}
+
+void FSpline::SetOutTangent(const int32 Index, const FVector& OutTangent)
+{
+	switch (CurrentImplementation)
+	{
+	default: break;
+	case 1: LegacyData->SetOutTangent(Index, OutTangent); break;
+	case 2: NewData->SetOutTangent(Index, OutTangent); break;
+	}
+}
+
+FVector FSpline::GetOutTangent(const int32 Index) const
+{
+	switch (CurrentImplementation)
+	{
+	default: return FVector();
+	case 1: return LegacyData->GetOutTangent(Index);
+	case 2: return NewData->GetOutTangent(Index);
+	}
+}
+
+void FSpline::SetRotation(int32 Index, const FQuat& InRotation)
+{
+	switch (CurrentImplementation)
+	{
+	default: break;
+	case 1: LegacyData->SetRotation(Index, InRotation); break;
+	case 2: NewData->SetRotation(Index, InRotation); break;
+	}
+}
+
+FQuat FSpline::GetRotation(const int32 Index) const
+{
+	switch (CurrentImplementation)
+	{
+	default: return FQuat();
+	case 1: return LegacyData->GetRotation(Index);
+	case 2: return NewData->GetRotation(Index);
+	}
+}
+
+void FSpline::SetScale(int32 Index, const FVector& InScale)
+{
+	switch (CurrentImplementation)
+	{
+	default: break;
+	case 1: LegacyData->SetScale(Index, InScale); break;
+	case 2: NewData->SetScale(Index, InScale); break;
+	}
+}
+
+FVector FSpline::GetScale(const int32 Index) const
+{
+	switch (CurrentImplementation)
+	{
+	default: return FVector(1.0);
+	case 1: return LegacyData->GetScale(Index);
+	case 2: return NewData->GetScale(Index);
+	}
+}
+
+void FSpline::SetSplinePointType(int32 Index, EInterpCurveMode Type)
+{
+	switch (CurrentImplementation)
+	{
+	default: break;
+	case 1: LegacyData->SetSplinePointType(Index, Type); break;
+	case 2: NewData->SetSplinePointType(Index, Type); break;
+	}
+}
+
+EInterpCurveMode FSpline::GetSplinePointType(int32 Index) const
+{
+	switch (CurrentImplementation)
+	{
+	default: return CIM_Unknown;
+	case 1: return LegacyData->GetSplinePointType(Index);
+	case 2: return NewData->GetSplinePointType(Index);
+	}
+}
+
+float FSpline::GetParameterAtIndex(int32 Index) const
+{
+	switch (CurrentImplementation)
+	{
+	default: return 0.f;
+	case 1: return LegacyData->GetParameterAtIndex(Index);
+	case 2: return NewData->GetParameterAtIndex(Index);
+	}
+}
+
+float FSpline::GetParameterAtDistance(float Distance) const
+{
+	switch (CurrentImplementation)
+	{
+	default: return 0.f;
+	case 1: return LegacyData->GetParameterAtDistance(Distance);
+	case 2: return NewData->GetParameterAtDistance(Distance);
+	}
+}
+
+float FSpline::GetDistanceAtParameter(float Parameter) const
+{
+	switch (CurrentImplementation)
+	{
+	default: return 0.f;
+	case 1: return LegacyData->GetDistanceAtParameter(Parameter);
+	case 2: return NewData->GetDistanceAtParameter(Parameter);
+	}
+}
+
+FQuat FSpline::GetOrientation(int32 Index) const
+{
+	return GetOrientation(GetParameterAtIndex(Index));
+}
+
+void FSpline::SetOrientation(int32 Index, const FQuat& InOrientation)
+{
+	if (Index < 0 || Index >= GetNumControlPoints())
+	{
+		return;
+	}
+	
+	// Work backwards to compute the rotation that is currently being applied.
+	const FQuat RelativeRotation = InOrientation * GetOrientation(Index).Inverse();
+	
+	// Store rotation which transforms the world up vector to the local up vector.
+	SetRotation(Index, FQuat::FindBetween(FVector::UpVector, InOrientation.GetUpVector()));
+
+	// Align tangents with rotation, preserving magnitude.
+	const FVector OldInTangent = GetInTangent(Index);
+	const FVector OldInTangentDirection = OldInTangent.GetSafeNormal();
+	const float InTangentMag = OldInTangent.Length();
+	const FVector NewInTangentDirection = RelativeRotation.RotateVector(OldInTangentDirection);
+	
+	const FVector OldOutTangent = GetOutTangent(Index);
+	const FVector OldOutTangentDirection = OldOutTangent.GetSafeNormal();
+	const float OutTangentMag = OldOutTangent.Length();
+	const FVector NewOutTangentDirection = RelativeRotation.RotateVector(OldOutTangentDirection);
+	
+	SetInTangent(Index, InTangentMag * NewInTangentDirection);
+	SetOutTangent(Index, OutTangentMag * NewOutTangentDirection);
+}
+
 FVector FSpline::Evaluate(float Param) const
 {
 	switch (CurrentImplementation)
 	{
-	default:
-		break;
-
-	case UE::Spline::ESplineType::LegacyTangent:
-		return LegacySplineImpl->PositionCurve.Eval(Param);
-
-	case UE::Spline::ESplineType::Tangent:
-		return Get<UE::Spline::ESplineType::Tangent>()->EvaluatePosition(Param);
+	default: return FVector();
+	case 1: return LegacyData->Evaluate(Param);
+	case 2: return NewData->EvaluatePosition(Param);
 	}
-
-	return FVector();
 }
 
 FVector FSpline::EvaluateDerivative(float Param) const
 {
-	switch (CurrentImplementation)
+	if (UE::Spline::GApproximateTangents)
 	{
-	default:
-		break;
+		// Approximate using central difference.
+		
+		// This computes the tangent direction using central difference and
+		// assumes that tangent magnitude is linearly changing between control points.
+		// While the assumption about magnitude is probably wrong, it works well.
+		
+		auto ClampValidParameterRange = [Min = 0.f, Max = GetNumControlPoints() - 1](float Param)
+			{ return FMath::Clamp(Param, Min, Max); };
 
-	case UE::Spline::ESplineType::LegacyTangent:
-		return LegacySplineImpl->PositionCurve.EvalDerivative(Param);
+		constexpr float H = UE_KINDA_SMALL_NUMBER;
+		const float ParamLow = ClampValidParameterRange(Param - H);
+		const float ParamHigh = ClampValidParameterRange(Param + H);
+		const FVector Tangent = ((Evaluate(ParamHigh) - Evaluate(ParamLow)) / (ParamHigh - ParamLow)).GetSafeNormal();
 
-	case UE::Spline::ESplineType::Tangent:
-		return Get<UE::Spline::ESplineType::Tangent>()->EvaluateDerivative(Param);
+		const int32 Index1 = FMath::Clamp((int32)Param, 0, GetNumControlPoints() - 1);
+		const int32 Index2 = FMath::Clamp((int32)(Param+1), 0, GetNumControlPoints() - 1);
+		const float Mag1 = GetInTangent(Index1).Length();
+		const float Mag2 = GetInTangent(Index2).Length();
+		const float Mag = FMath::Lerp(Mag1, Mag2, FMath::Frac(Param));
+
+		return Tangent * Mag;
 	}
-
-	return FVector();
+	else
+	{
+		switch (CurrentImplementation)
+		{
+		default: return FVector();
+		case 1: return LegacyData->EvaluateDerivative(Param);
+		case 2: return NewData->EvaluateDerivative(Param);
+		}
+	}
 }
 
-int32 FSpline::GetNumberOfSegments() const
+FQuat FSpline::EvaluateRotation(float Param) const
 {
 	switch (CurrentImplementation)
 	{
-	default:
-		return 0;
-
-	case UE::Spline::ESplineType::LegacyTangent:
-		return LegacySplineImpl->PositionCurve.bIsLooped ? LegacySplineImpl->PositionCurve.Points.Num() : FMath::Max(LegacySplineImpl->PositionCurve.Points.Num() - 1, 0);
-
-	case UE::Spline::ESplineType::Tangent:
-		// fallthrough for all ISplineInterfaces
-		return SplineImpl->GetNumberOfSegments();
+	default: return FQuat();
+	case 1: return LegacyData->EvaluateRotation(Param);
+	case 2: return NewData->EvaluateRotation(Param);
 	}
 }
 
-void FSpline::Reset()
+FVector FSpline::EvaluateScale(float Param) const
 {
 	switch (CurrentImplementation)
 	{
-	default:
-		break;
+	default: return FVector();
+	case 1: return LegacyData->EvaluateScale(Param);
+	case 2: return NewData->EvaluateScale(Param);
+	}	
+}
 
-	case UE::Spline::ESplineType::LegacyTangent:
-		LegacySplineImpl->PositionCurve.Reset();
-		LegacySplineImpl->PositionCurve.ClearLoopKey();
-		LegacySplineImpl->RotationCurve.Reset();
-		LegacySplineImpl->RotationCurve.ClearLoopKey();
-		LegacySplineImpl->ScaleCurve.Reset();
-		LegacySplineImpl->ScaleCurve.ClearLoopKey();
-		LegacySplineImpl->ReparamTable.Reset();
-		LegacySplineImpl->ReparamTable.ClearLoopKey();
-		break;
+FQuat FSpline::GetOrientation(float Param) const
+{
+	FQuat Rotation = EvaluateRotation(Param);
+	Rotation.Normalize();
 
-	case UE::Spline::ESplineType::Tangent:
-		// fallthrough for all ISplineInterfaces
-		SplineImpl->Clear();
-		break;
+	FVector Direction = EvaluateDerivative(Param);
+	Direction = Direction.GetSafeNormal();
+
+	const FVector UpVector = Rotation.RotateVector(FVector::UpVector);
+	
+	return FRotationMatrix::MakeFromXZ(Direction, UpVector).ToQuat();
+}
+
+bool FSpline::HasAttributeChannel(FName AttributeName) const
+{
+	if (NewData && SupportsAttributes())
+	{
+		return NewData->HasAttributeChannel(AttributeName);
+	}
+
+	return false;
+}
+
+bool FSpline::RemoveAttributeChannel(FName AttributeName) const
+{
+	if (NewData && SupportsAttributes())
+	{
+		return NewData->RemoveAttributeChannel(AttributeName);
+	}
+
+	return false;
+}
+
+TArray<FName> FSpline::GetFloatPropertyChannels() const
+{
+	TArray<FName> ChannelNames;
+	if (NewData && SupportsAttributes())
+	{
+		ChannelNames = NewData->GetAttributeChannelNamesByValueType<float>();
+	}
+	return ChannelNames;
+}
+
+TArray<FName> FSpline::GetVectorPropertyChannels() const
+{
+	TArray<FName> ChannelNames;
+	if (NewData && SupportsAttributes())
+	{
+		ChannelNames = NewData->GetAttributeChannelNamesByValueType<FVector>();
+
+		// Don't report these internal attribute channels to the caller
+		ChannelNames.Remove(NewData->ScaleAttrName);
+	}
+	return ChannelNames;
+}
+
+template <typename AttrType>
+float FSpline::GetAttributeParameter(int32 Index, const FName& Name) const
+{
+	if (NewData && SupportsAttributes())
+	{
+		return NewData->GetAttributeParameter<AttrType>(Name, Index);
+	}
+
+	return 0.f;
+}
+template ENGINE_API float FSpline::GetAttributeParameter<float>(int32 Index, const FName& Name) const;
+template ENGINE_API float FSpline::GetAttributeParameter<FVector>(int32 Index, const FName& Name) const;
+
+template <typename AttrType>
+int32 FSpline::SetAttributeParameter(int32 Index, float Parameter, const FName& Name)
+{
+	if (NewData && SupportsAttributes())
+	{
+		return NewData->SetAttributeParameter<AttrType>(Name, Index, Parameter);
+	}
+	
+	return Index;
+}
+template ENGINE_API int32 FSpline::SetAttributeParameter<float>(int32 Index, float Parameter, const FName& Name);
+template ENGINE_API int32 FSpline::SetAttributeParameter<FVector>(int32 Index, float Parameter, const FName& Name);
+
+template <typename AttrType>
+int32 FSpline::NumAttributeValues(FName AttributeName) const
+{
+	if (NewData && SupportsAttributes())
+	{
+		return NewData->NumAttributeValues<AttrType>(AttributeName);
+	}
+
+	return 0;
+}
+template ENGINE_API int32 FSpline::NumAttributeValues<float>(FName AttributeName) const;
+template ENGINE_API int32 FSpline::NumAttributeValues<FVector>(FName AttributeName) const;
+
+template<typename AttrType>
+AttrType FSpline::GetAttributeValue(int32 Index, const FName& Name) const
+{
+	if (NewData && SupportsAttributes())
+	{
+		return NewData->GetAttributeValue<AttrType>(Name, Index);
+	}
+	
+	return AttrType();
+}
+template ENGINE_API float FSpline::GetAttributeValue<float>(int32 Index, const FName& Name) const;
+template ENGINE_API FVector FSpline::GetAttributeValue<FVector>(int32 Index, const FName& Name) const;
+
+template<typename AttrType>
+void FSpline::SetAttributeValue(int32 Index, const AttrType& Value, const FName& Name)
+{
+	if (NewData && SupportsAttributes())
+	{
+		NewData->SetAttributeValue<AttrType>(Name, Value, Index);
+	}
+}
+template ENGINE_API void FSpline::SetAttributeValue<float>(int32 Index, const float& Value, const FName& Name);
+template ENGINE_API void FSpline::SetAttributeValue<FVector>(int32 Index, const FVector& Value, const FName& Name);
+
+template<typename AttrType>
+bool FSpline::CreateAttributeChannel(FName AttributeName) const
+{
+	if (NewData && SupportsAttributes())
+	{
+		if (UE::Geometry::Spline::TTangentBezierSpline<AttrType>* NewAttrChannel = NewData->CreateAttributeChannel<UE::Geometry::Spline::TTangentBezierSpline<AttrType>>(AttributeName); NewAttrChannel != nullptr)
+		{
+			NewAttrChannel->SetPreInfinityMode(UE::Geometry::Spline::EOutOfBoundsHandlingMode::Constant);
+			NewAttrChannel->SetPostInfinityMode(UE::Geometry::Spline::EOutOfBoundsHandlingMode::Constant);
+
+			return true;
+		}
+	}
+	
+	return false;
+}
+template ENGINE_API bool FSpline::CreateAttributeChannel<float>(FName AttributeName) const;
+template ENGINE_API bool FSpline::CreateAttributeChannel<FVector>(FName AttributeName) const;
+
+template<typename AttrType>
+int32 FSpline::AddAttributeValue(float Param, const AttrType& Value, FName AttributeName) const
+{
+	if (NewData && SupportsAttributes())
+	{
+		return NewData->AddAttributeValue(AttributeName, Value, Param);
+	}
+	
+	return INDEX_NONE;
+}
+template ENGINE_API int32 FSpline::AddAttributeValue<float>(float Param, const float& Value, FName AttributeName) const;
+template ENGINE_API int32 FSpline::AddAttributeValue<FVector>(float Param, const FVector& Value, FName AttributeName) const;
+
+template <typename AttrType>
+void FSpline::RemoveAttributeValue(int32 Index, FName AttributeName)
+{
+	if (NewData && SupportsAttributes())
+	{
+		NewData->RemoveAttributeValue<AttrType>(AttributeName, Index);
+	}
+}
+template ENGINE_API void FSpline::RemoveAttributeValue<float>(int32 Index, FName AttributeName);
+template ENGINE_API void FSpline::RemoveAttributeValue<FVector>(int32 Index, FName AttributeName);
+
+template<typename AttrType>
+AttrType FSpline::EvaluateAttribute(float Param, FName AttributeName) const
+{
+	if (NewData && SupportsAttributes())
+	{
+		return NewData->EvaluateAttribute<AttrType>(AttributeName, NewData->ToInternalSplineSpace(Param));
+	}
+	return AttrType(0);
+}
+template ENGINE_API float FSpline::EvaluateAttribute<float>(float Param, FName AttributeName) const;
+template ENGINE_API FVector FSpline::EvaluateAttribute<FVector>(float Param, FName AttributeName) const;
+
+float FSpline::FindNearest(const FVector& InLocation, float& OutSquaredDist) const
+{
+	switch (CurrentImplementation)
+	{
+	default: return 0.f;
+	case 1: return LegacyData->FindNearest(InLocation, OutSquaredDist);
+	case 2: return NewData->FindNearest(InLocation, OutSquaredDist);
 	}
 }
 
+float FSpline::FindNearestOnSegment(const FVector& InLocation, int32 SegmentIndex, float& OutSquaredDist) const
+{
+	switch (CurrentImplementation)
+	{
+	default: return 0.f;
+	case 1: return LegacyData->FindNearestOnSegment(InLocation, SegmentIndex, OutSquaredDist);
+	case 2: return NewData->FindNearestOnSegment(InLocation, SegmentIndex, OutSquaredDist);
+	}
+}
+
 bool FSpline::operator==(const FSpline& Other) const
 {
-	if (CurrentImplementation != Other.CurrentImplementation)
+	if (CurrentImplementation == Other.CurrentImplementation)
 	{
-		return false;
+		if (LegacyData && Other.LegacyData)
+		{
+			return *LegacyData == *Other.LegacyData;
+		}
+
+		if (NewData && Other.NewData)
+		{
+			return *NewData == *Other.NewData;
+		}
 	}
 
-	switch (CurrentImplementation)
-	{
-	default:
-		return false;
-
-	case UE::Spline::ESplineType::LegacyTangent:
-		return LegacySplineImpl->PositionCurve == Other.LegacySplineImpl->PositionCurve &&
-			LegacySplineImpl->RotationCurve == Other.LegacySplineImpl->RotationCurve &&
-			LegacySplineImpl->ScaleCurve == Other.LegacySplineImpl->ScaleCurve;
-
-	case UE::Spline::ESplineType::Tangent:
-		// fallthrough for all ISplineInterfaces
-		return SplineImpl->IsEqual(Other.SplineImpl.Get());
-
-	}
+	return false;
 }
 
 bool FSpline::operator!=(const FSpline& Other) const
@@ -204,6 +1421,123 @@
 	return !(*this == Other);
 }
 
+bool FSpline::IsNearlyEqual(const FSpline& Other, const float RelativeTolerance, const float AbsoluteTolerance) const
+{
+	if (!ensureAlways(IsNew() && Other.IsNew()))
+	{
+		return false;
+	}
+
+	const int32 NumPoints = NewData->GetSpline().GetNumPoints();
+	if (Other.NewData->GetSpline().GetNumPoints() != NumPoints)
+	{
+		return false;
+	}
+
+	auto ScalarNearlyEqual = [AbsoluteTolerance, RelativeTolerance](const float& A, const float& B)
+	{
+		const float Delta = FMath::Abs(A - B);
+		const float Tolerance = FMath::Max(AbsoluteTolerance, RelativeTolerance * FMath::Max(FMath::Abs(A), FMath::Abs(B)));
+		return Delta <= Tolerance;
+	};
+
+	auto VectorNearlyEqual = [AbsoluteTolerance, RelativeTolerance](const FVector& A, const FVector& B)
+	{
+		const float MagDelta = (A - B).Length();
+		const float MagTolerance = FMath::Max(AbsoluteTolerance, RelativeTolerance * FMath::Max(A.Length(), B.Length()));
+		return MagDelta <= MagTolerance;
+	};
+
+	auto QuatComponentsNearlyEqual = [AbsoluteTolerance, RelativeTolerance](const FQuat& A, const FQuat& B)
+	{
+		const float DX = FMath::Abs(A.X - B.X);
+		const float TX = FMath::Max(AbsoluteTolerance, RelativeTolerance * FMath::Max(FMath::Abs(A.X), FMath::Abs(B.X)));
+		if (DX > TX) return false;
+
+		const float DY = FMath::Abs(A.Y - B.Y);
+		const float TY = FMath::Max(AbsoluteTolerance, RelativeTolerance * FMath::Max(FMath::Abs(A.Y), FMath::Abs(B.Y)));
+		if (DY > TY) return false;
+
+		const float DZ = FMath::Abs(A.Z - B.Z);
+		const float TZ = FMath::Max(AbsoluteTolerance, RelativeTolerance * FMath::Max(FMath::Abs(A.Z), FMath::Abs(B.Z)));
+		if (DZ > TZ) return false;
+
+		const float DW = FMath::Abs(A.W - B.W);
+		const float TW = FMath::Max(AbsoluteTolerance, RelativeTolerance * FMath::Max(FMath::Abs(A.W), FMath::Abs(B.W)));
+		return DW <= TW;
+	};
+
+	for (int32 PointIdx = 0; PointIdx < NumPoints; ++PointIdx)
+	{
+		if (!VectorNearlyEqual(NewData->GetLocation(PointIdx), Other.NewData->GetLocation(PointIdx)))
+		{
+			return false;
+		}
+
+		if (!VectorNearlyEqual(NewData->GetInTangent(PointIdx), Other.NewData->GetInTangent(PointIdx)))
+		{
+			return false;
+		}
+
+		if (!VectorNearlyEqual(NewData->GetOutTangent(PointIdx), Other.NewData->GetOutTangent(PointIdx)))
+		{
+			return false;
+		}
+
+		const FQuat Rotation = NewData->GetRotation(PointIdx);
+		const FQuat OtherRotation = Other.NewData->GetRotation(PointIdx);
+		const FQuat OtherRotationNeg = FQuat(-OtherRotation.X, -OtherRotation.Y, -OtherRotation.Z, -OtherRotation.W);
+		if (!QuatComponentsNearlyEqual(Rotation, OtherRotation) && !QuatComponentsNearlyEqual(Rotation, OtherRotationNeg))
+		{
+			return false;
+		}
+
+		if (!VectorNearlyEqual(NewData->GetScale(PointIdx), Other.NewData->GetScale(PointIdx)))
+		{
+			return false;
+		}
+
+		if (!ScalarNearlyEqual(NewData->GetParameterAtIndex(PointIdx), Other.NewData->GetParameterAtIndex(PointIdx)))
+		{
+			return false;
+		}
+	}
+
+	const TArray<FName> OtherFloatChannelNames = Other.GetFloatPropertyChannels();
+	for (const FName& FloatChannelName : GetFloatPropertyChannels())
+	{
+		if (!OtherFloatChannelNames.Contains(FloatChannelName))
+		{
+			return false;
+		}
+
+		const int32 NumFloatChannelPoints = NumAttributeValues<float>(FloatChannelName);
+		if (Other.NumAttributeValues<float>(FloatChannelName) != NumFloatChannelPoints)
+		{
+			return false;
+		}
+
+		for (int32 FloatChannelIdx = 0; FloatChannelIdx < NumFloatChannelPoints; ++FloatChannelIdx)
+		{
+			if (!ScalarNearlyEqual(
+				GetAttributeValue<float>(FloatChannelIdx, FloatChannelName),
+				Other.GetAttributeValue<float>(FloatChannelIdx, FloatChannelName)))
+			{
+				return false;
+			}
+
+			if (!ScalarNearlyEqual(
+				GetAttributeParameter<float>(FloatChannelIdx, FloatChannelName),
+				Other.GetAttributeParameter<float>(FloatChannelIdx, FloatChannelName)))
+			{
+				return false;
+			}
+		}
+	}
+
+	return true;
+}
+
 bool FSpline::Serialize(FArchive& Ar)
 {
 	Ar.UsingCustomVersion(FFortniteMainBranchObjectVersion::GUID);
@@ -227,47 +1561,78 @@
 
 void FSpline::SerializeLoad(FArchive& Ar)
 {
-	uint8 PreviousImpl;
+	int8 PreviousImpl;
 	Ar << PreviousImpl;
-
 #if WITH_EDITOR
-	PreviousImplementation = static_cast<UE::Spline::ESplineType>(PreviousImpl);
+	PreviousImplementation = PreviousImpl;
 #endif
 
-	SetCurrentImplementation(static_cast<UE::Spline::ESplineType>(PreviousImpl));
+	auto WasEnabled = [PreviousImpl]()
+	{ return PreviousImpl != 0; };
 
-	switch (CurrentImplementation)
+	auto WasLegacy = [PreviousImpl]()
+	{ return PreviousImpl == 1; };
+		
+	if (WasEnabled())
 	{
-	default: break;
-
-	case UE::Spline::ESplineType::LegacyTangent: Ar << *LegacySplineImpl; break;
-		
-	case UE::Spline::ESplineType::Tangent:
-		// fallthrough for all ISplineInterfaces
-		Ar << *SplineImpl; break;
+		if (WasLegacy())
+		{
+			if (IsLegacy())
+			{
+				LegacyData = MakeShared<FLegacySpline>();
+				Ar << *LegacyData;
+			}
+			else if (IsNew())
+			{
+				FLegacySpline IntermediateData;
+				Ar << IntermediateData;
+				NewData = MakeShared<FNewSpline>(IntermediateData);
+			}
+			else if (!IsEnabled())
+			{
+				// Intentionally doing nothing with data
+				FLegacySpline IntermediateData;
+				Ar << IntermediateData;
+			}
+		}
+		else
+		{
+			if (IsLegacy())
+			{
+				FNewSpline IntermediateData;
+				Ar << IntermediateData;
+				LegacyData = MakeShared<FLegacySpline>(IntermediateData);
+			}
+			else if (IsNew())
+			{
+				NewData = MakeShared<FNewSpline>();
+				Ar << *NewData;
+			}
+			else if (!IsEnabled())
+			{
+				// Intentionally doing nothing with data
+				FNewSpline IntermediateData;
+				Ar << IntermediateData;
+			}
+		}
 	}
 }
 
 void FSpline::SerializeSave(FArchive& Ar) const
 {
-	uint8 CurrentImpl = static_cast<uint8>(CurrentImplementation);
+	uint8 CurrentImpl = CurrentImplementation;
 	Ar << CurrentImpl;
-
-	switch (CurrentImplementation)
+		
+	if (IsLegacy())
 	{
-	default: break;
-
-	case UE::Spline::ESplineType::LegacyTangent: Ar << *LegacySplineImpl; break;
-
-	case UE::Spline::ESplineType::Tangent:
-		// fallthrough for all ISplineInterfaces
-		Ar << *SplineImpl; break;
+		Ar << *LegacyData;
+	}
+	else if (IsNew())
+	{
+		Ar << *NewData;
 	}
 }
 
-// todo: make cvar when importing handles it changing.
-constexpr bool bEncodeAsHex = true;
-
 bool FSpline::ExportTextItem(FString& ValueStr, FSpline const& DefaultValue, class UObject* Parent, int32 PortFlags, class UObject* ExportRootScope) const
 {
 	// serialize our spline
@@ -275,7 +1640,7 @@
     FMemoryWriter MemWriter(SplineWriteBuffer);
     SerializeSave(MemWriter);
 
-	if (bEncodeAsHex)
+	if (UE::Spline::GEncodeAsHex)
 	{
 		FString HexStr = FString::FromHexBlob(SplineWriteBuffer.GetData(), SplineWriteBuffer.Num());
 		ValueStr = FString::Printf(TEXT("SplineData SplineDataLen=%d SplineData=%s\r\n"), HexStr.Len(), *HexStr);
@@ -297,9 +1662,6 @@
 
 bool FSpline::ImportTextItem(const TCHAR*& SourceText, int32 PortFlags, UObject* Parent, FOutputDevice* ErrorText)
 {
-	// todo: SplineData= is followed by a byte identifying the implementation type, we should
-	// validate this on import and potentially change spline type if it doesn't match our current type
-
 	if (FParse::Command(&SourceText, TEXT("SplineData")))	
 	{
 		static const TCHAR SplineDataLenToken[] = TEXT("SplineDataLen=");
@@ -322,10 +1684,10 @@
 				TArray<uint8> SplineReadBuffer;
 				bool bDecoded;
 				
-				if (bEncodeAsHex)
+				if (UE::Spline::GEncodeAsHex)
 				{
-					SplineReadBuffer.SetNum(SplineDataLen / 2);
-					bDecoded = FString::ToHexBlob(SplineData, SplineReadBuffer.GetData(), SplineReadBuffer.Num());
+					SplineReadBuffer.SetNum(SplineDataLen);
+					bDecoded = FString::ToHexBlob(SplineData, SplineReadBuffer.GetData(), SplineDataLen);
 				}
 				else
 				{
@@ -339,41 +1701,1861 @@
 				}
 			}
 		}
-	}
-	
-	return true;
-}
-
-UE::Spline::ESplineType FSpline::GetCurrentImplementation() const
-{ 
-	return CurrentImplementation;
-}
-
-#if WITH_EDITOR
-UE::Spline::ESplineType FSpline::GetPreviousImplementation() const
-{
-	return PreviousImplementation;
-}
-#endif
-
-UE::Spline::ESplineType FSpline::GetDefaultImplementation()
-{
-	return UE::Spline::GetDefaultSplineType();
-}
-
-void FSpline::SetCurrentImplementation(UE::Spline::ESplineType Implementation)
-{
-	CurrentImplementation = Implementation;
-
-	LegacySplineImpl.Reset();
-	SplineImpl.Reset();
-
-	switch (CurrentImplementation)
-	{
-	default: break;		// Intentionally doing nothing
-	case UE::Spline::ESplineType::LegacyTangent:	LegacySplineImpl = MakeUnique<FLegacyTangentSpline>(); break;
-	case UE::Spline::ESplineType::Tangent:			SplineImpl = MakeUnique<FTangentSpline>(); break;
-	}
-}
-
+
+
+	}
+	
+	return true;
+}
+
+const FInterpCurveVector& FSpline::GetSplinePointsPosition() const
+{
+	switch (CurrentImplementation)
+	{
+	default:
+		return PositionCurve;
+	case 1:
+		return LegacyData->GetSplinePointsPosition();
+	case 2:
+		NewData->RebuildLegacyCurves();
+		return NewData->GetSplinePointsPosition();
+	}
+}
+
+const FInterpCurveQuat& FSpline::GetSplinePointsRotation() const
+{
+	switch (CurrentImplementation)
+	{
+	default:
+		return RotationCurve;
+	case 1:
+		return LegacyData->GetSplinePointsRotation();
+	case 2:
+		NewData->RebuildLegacyCurves();
+		return NewData->GetSplinePointsRotation();
+	}
+}
+
+const FInterpCurveVector& FSpline::GetSplinePointsScale() const
+{
+	switch (CurrentImplementation)
+	{
+	default:
+		return ScaleCurve;
+	case 1:
+		return LegacyData->GetSplinePointsScale();
+	case 2:
+		NewData->RebuildLegacyCurves();
+		return NewData->GetSplinePointsScale();
+	}
+}
+
+float FSpline::GetSegmentLength(const int32 Index, const float Param, const FVector& Scale3D) const
+{
+	switch (CurrentImplementation)
+	{
+	default: return 0.f;
+	case 1: return LegacyData->GetSegmentLength(Index, Param, Scale3D);
+	case 2: return NewData->GetSegmentLength(Index, Param, Scale3D);
+	}
+}
+
+float FSpline::GetSplineLength() const
+{
+	switch (CurrentImplementation)
+	{
+	default: return 0.f;	// Intentionally doing nothing
+	case 1: return LegacyData->GetSplineLength();
+	case 2: return NewData->GetSplineLength();
+	}
+}
+
+int32 FSpline::GetNumSegments() const
+{
+	switch (CurrentImplementation)
+	{
+	default: return 0;
+	case 1: return LegacyData->GetSplinePointsPosition().bIsLooped ? LegacyData->GetNumControlPoints() : FMath::Max(0, LegacyData->GetNumControlPoints() - 1);
+	case 2: return NewData->GetSpline().GetNumberOfSegments();
+	}
+}
+
+int32 FSpline::GetNumControlPoints() const
+{
+	switch (CurrentImplementation)
+	{
+	default: return 0;
+	case 1: return LegacyData->GetNumControlPoints();
+	case 2: return NewData->GetSpline().GetNumPoints();
+	}
+}
+
+void FSpline::Reset()
+{
+	switch (CurrentImplementation)
+	{
+	default: break;
+	case 1: LegacyData->Reset(); break;
+	case 2: NewData->Reset(); break;
+	}
+}
+	
+void FSpline::ResetRotation()
+{
+	switch (CurrentImplementation)
+	{
+	default: break;
+	case 1: LegacyData->ResetRotation(); break;
+	case 2: NewData->ResetRotation(); break;
+	}
+}
+	
+void FSpline::ResetScale()
+{
+	switch (CurrentImplementation)
+	{
+	default: break;
+	case 1: LegacyData->ResetScale(); break;
+	case 2: NewData->ResetScale(); break;
+	}
+}
+
+void FSpline::SetClosedLoop(bool bClosed)
+{
+	TRACE_CPUPROFILER_EVENT_SCOPE(FSpline::SetClosedLoop);
+
+	CachedUpdateSplineParams.bClosedLoop = bClosed;
+	UpdateSpline();
+}
+
+bool FSpline::IsClosedLoop() const
+{
+	return CachedUpdateSplineParams.bClosedLoop;
+}
+
+void FSpline::UpdateSpline(const FUpdateSplineParams& InParams)
+{
+	TRACE_CPUPROFILER_EVENT_SCOPE(FSpline::UpdateSpline);
+
+#if WITH_EDITOR
+	if (UE::Spline::GImplementation != CurrentImplementation)
+	{
+		OnSplineImplementationChanged();
+	}
+#endif
+		
+	switch (CurrentImplementation)
+    {
+    default: break;
+    case 1: LegacyData->UpdateSpline(InParams); break;
+    case 2: NewData->UpdateSpline(InParams); break;
+    }
+
+	CachedUpdateSplineParams = InParams;
+	++Version;
+}
+
+void FSpline::UpdateSpline()
+{
+	UpdateSpline(CachedUpdateSplineParams);
+}
+
+TSharedPtr<UE::Geometry::Spline::TSplineInterface<FVector>> FSpline::GetSplineInterface() const
+{
+	if (IsNew())
+	{
+		return StaticCastSharedPtr<UE::Geometry::Spline::TSplineInterface<FVector>>(NewData);
+	}
+
+	return nullptr;
+}
+
+bool FSpline::IsEnabledStatic()
+{
+	return UE::Spline::GImplementation != 0;
+}
+
+#if WITH_EDITOR
+void FSpline::OnSplineImplementationChanged()
+{
+	using namespace UE::Spline;
+
+	// This implements a state machine I drew out.
+	// I am pretty confident it covers all cases because it handles all edges of a
+	// directed digraph (of 3 nodes, the 3 possible states, with self-edges) where edges represent state transitions.
+
+	// Do you see the nice triangle?
+	
+	if (GImplementation == CurrentImplementation)			// 0->0, 1->1, 2->2
+	{
+	}
+	else if (GImplementation == 0)							// 1->0, 2->0
+	{
+		LegacyData.Reset();
+		NewData.Reset();
+	}
+	else if (CurrentImplementation == 0)					// 0->1, 0->2
+	{
+		// Here we just need to alloc but not copy
+		// GImplementation guaranteed to be 1 or 2
+		if (GImplementation == 1)
+		{
+			LegacyData = MakeShared<FLegacySpline>();
+		}
+		else if (GImplementation == 2)
+		{
+			NewData = MakeShared<FNewSpline>();
+		}
+	}
+	else if (CurrentImplementation == 1)					// 1->2
+	{
+		NewData = MakeShared<FNewSpline>(*LegacyData);
+		LegacyData.Reset();
+	}
+	else if (CurrentImplementation == 2)					// 2->1
+	{
+		LegacyData = MakeShared<FLegacySpline>(*NewData);
+		NewData.Reset();
+	}
+
+	CurrentImplementation = GImplementation;
+}
+#endif
+
+
+
+/** FNewSpline Implementation */
+
+FNewSpline::FNewSpline()
+{
+	// This is a hack. We just need to make sure that this spline type registers itself with the spline registry.
+	// This should be removed in the future once we can register this spline type automatically without instantiation.
+	UE::Geometry::Spline::TTangentBezierSpline<float> AutoRegister;
+	
+	CreateAttributeChannel<UE::Geometry::Spline::TTangentBezierSpline<FQuat>>(RotationAttrName);
+	CreateAttributeChannel<UE::Geometry::Spline::TTangentBezierSpline<FVector>>(ScaleAttrName);
+	
+	ValidateRotScale();
+}
+
+FNewSpline::FNewSpline(const FNewSpline& Other)
+{
+    *this = Other;
+
+	ValidateRotScale();
+}
+
+FNewSpline::FNewSpline(const FLegacySpline& Other)
+{
+	CreateAttributeChannel<UE::Geometry::Spline::TTangentBezierSpline<FQuat>>(RotationAttrName);
+	CreateAttributeChannel<UE::Geometry::Spline::TTangentBezierSpline<FVector>>(ScaleAttrName);
+	
+	TArray<FTangentBezierControlPoint> Points;
+    for (int32 Idx = 0; Idx < Other.GetNumControlPoints(); Idx++)
+    {
+	    FTangentBezierControlPoint& Point = Points.AddDefaulted_GetRef();
+
+    	Point.Position = Other.GetLocation(Idx);
+    	Point.TangentIn = Other.GetInTangent(Idx);
+    	Point.TangentOut = Other.GetOutTangent(Idx);
+    	Point.TangentMode = ConvertInterpCurveModeToTangentMode(Other.GetSplinePointType(Idx));
+    }
+
+	// initialize spline
+	GetSpline().SetControlPoints(Points);
+	ResetRotation();
+	ResetScale();
+	
+	// setup rotation/scale attributes
+	for (int32 Idx = 0; Idx < GetSpline().GetNumPoints(); Idx++)
+	{    		
+		SetAttributeValue<FQuat>(RotationAttrName, Other.GetRotation(Idx), Idx);
+		SetAttributeValue<FVector>(ScaleAttrName, Other.GetScale(Idx), Idx);
+	}
+	
+	ValidateRotScale();
+
+	FSpline::FUpdateSplineParams UpdateParams;
+	UpdateParams.bClosedLoop = Other.GetSplinePointsPosition().bIsLooped;
+	
+    UpdateSpline(UpdateParams);
+}
+
+FNewSpline::FNewSpline(const FSplineCurves& Other)
+{
+	CreateAttributeChannel<UE::Geometry::Spline::TTangentBezierSpline<FQuat>>(RotationAttrName);
+	CreateAttributeChannel<UE::Geometry::Spline::TTangentBezierSpline<FVector>>(ScaleAttrName);
+	
+	TArray<FTangentBezierControlPoint> Points;
+    for (int32 Idx = 0; Idx < Other.Position.Points.Num(); Idx++)
+    {
+	    const FInterpCurvePoint<FVector>& Position = Other.Position.Points[Idx];
+    	
+    	FTangentBezierControlPoint& Point = Points.AddDefaulted_GetRef();
+
+    	Point.Position = Position.OutVal;
+    	Point.TangentIn = Position.ArriveTangent;
+    	Point.TangentOut = Position.LeaveTangent;
+    	Point.TangentMode = ConvertInterpCurveModeToTangentMode(Position.InterpMode);
+    }
+	
+	// initialize spline
+	GetSpline().SetControlPoints(Points);
+	ResetRotation();
+	ResetScale();
+	
+    // setup rotation/scale attributes
+	for (int32 Idx = 0; Idx < Other.Position.Points.Num(); Idx++)
+	{
+		FQuat Rotation = FQuat::Identity;
+		if (Other.Rotation.Points.IsValidIndex(Idx))
+		{
+			Rotation = Other.Rotation.Points[Idx].OutVal;
+		}
+		else
+		{
+			UE_LOG(LogSpline, Warning, TEXT("Number of rotation values (%u) different than the number of position values (%u)!"), Other.Rotation.Points.Num(), Other.Position.Points.Num());
+		}
+    	
+		FVector Scale = FVector::OneVector;
+		if (Other.Scale.Points.IsValidIndex(Idx))
+		{
+			Scale = Other.Scale.Points[Idx].OutVal;
+		}
+		else
+		{
+			UE_LOG(LogSpline, Warning, TEXT("Number of scale values (%u) different than the number of position values (%u)!"), Other.Scale.Points.Num(), Other.Position.Points.Num());
+		}
+    	
+		SetAttributeValue<FQuat>(RotationAttrName, Rotation, Idx);
+		SetAttributeValue<FVector>(ScaleAttrName, Scale, Idx);
+	}
+
+	ValidateRotScale();
+	
+	FSpline::FUpdateSplineParams UpdateParams;
+	UpdateParams.bClosedLoop = Other.Position.bIsLooped;
+	
+    UpdateSpline(UpdateParams);
+}
+
+/** Copy assignment */
+FNewSpline& FNewSpline::operator=(const FNewSpline& Other)
+{
+    if (this != &Other)
+    {
+		TMultiSpline::operator=(Other);
+		PositionCurve = Other.PositionCurve;
+		RotationCurve = Other.RotationCurve;
+		ScaleCurve = Other.ScaleCurve;
+		ReparamStepsPerSegment = Other.ReparamStepsPerSegment;
+    }
+	
+    return *this;
+}
+
+bool FNewSpline::operator!=(const FNewSpline& Other) const
+{
+    return !(*this == Other);
+}
+
+bool FNewSpline::Serialize(FArchive& Ar)
+{
+    TMultiSpline::Serialize(Ar);
+    Ar << ReparamStepsPerSegment;
+
+	if (Ar.IsLoading())
+	{
+		Reparameterize();
+
+		// Do not serialize legacy curves, only generate if loading.
+		RebuildLegacyCurves();
+	}
+	
+    return true;
+}
+
+FInterval1f FNewSpline::GetSegmentSpace() const
+{
+	return GetNumberOfSegments() == 0 ? FInterval1f::Empty() : FInterval1f(0.f, static_cast<float>(GetNumberOfSegments()));
+}
+
+float FNewSpline::GetSegmentLength(const int32 Index, const float Param, const FVector& Scale3D) const
+{
+	const int32 NumPoints = GetSpline().GetNumPoints();
+	const int32 LastPoint = NumPoints - 1;
+
+	check(Index >= 0 && ((IsClosedLoop() && Index < NumPoints) || (!IsClosedLoop() && Index < LastPoint)));
+	check(Param >= 0.0f && Param <= 1.0f);
+
+	// Evaluate the length of a Hermite spline segment.
+	// This calculates the integral of |dP/dt| dt, where P(t) is the spline equation with components (x(t), y(t), z(t)).
+	// This isn't solvable analytically, so we use a numerical method (Legendre-Gauss quadrature) which performs very well
+	// with functions of this type, even with very few samples.  In this case, just 5 samples is sufficient to yield a
+	// reasonable result.
+
+	struct FLegendreGaussCoefficient
+	{
+		float Abscissa;
+		float Weight;
+	};
+
+	static const FLegendreGaussCoefficient LegendreGaussCoefficients[] =
+	{
+		{ 0.0f, 0.5688889f },
+		{ -0.5384693f, 0.47862867f },
+		{ 0.5384693f, 0.47862867f },
+		{ -0.90617985f, 0.23692688f },
+		{ 0.90617985f, 0.23692688f }
+	};
+
+	const auto& P0 = GetSpline().GetValue(Index);
+	const auto& T0 = GetSpline().GetTangentOut(Index);
+	const auto& P1 = GetSpline().GetValue(Index == LastPoint ? 0 : Index + 1);
+	const auto& T1 = GetSpline().GetTangentIn(Index == LastPoint ? 0 : Index + 1);
+
+	// Special cases for linear
+	if (GetSpline().GetTangentModes()[Index] == UE::Geometry::Spline::ETangentMode::Linear ||
+		GetSpline().GetTangentModes()[Index] == UE::Geometry::Spline::ETangentMode::LegacyLinear)
+	{
+		return ((P1 - P0) * Scale3D).Size() * Param;
+	}
+
+	// Cache the coefficients to be fed into the function to calculate the spline derivative at each sample point as they are constant.
+	const FVector Coeff1 = ((P0 - P1) * 2.0f + T0 + T1) * 3.0f;
+	const FVector Coeff2 = (P1 - P0) * 6.0f - T0 * 4.0f - T1 * 2.0f;
+	const FVector Coeff3 = T0;
+
+	const float HalfParam = Param * 0.5f;
+
+	float Length = 0.0f;
+	for (const auto& LegendreGaussCoefficient : LegendreGaussCoefficients)
+	{
+		// Calculate derivative at each Legendre-Gauss sample, and perform a weighted sum
+		const float Alpha = HalfParam * (1.0f + LegendreGaussCoefficient.Abscissa);
+		const FVector Derivative = ((Coeff1 * Alpha + Coeff2) * Alpha + Coeff3) * Scale3D;
+		Length += Derivative.Size() * LegendreGaussCoefficient.Weight;
+	}
+	Length *= HalfParam;
+
+	return Length;
+}
+
+float FNewSpline::GetSplineLength()
+{	
+	// Evaluate the ParamToLength channel at the very end of the spline
+	const float MaxParameter = static_cast<float>(GetNumberOfSegments());
+    return GetDistanceAtParameter(MaxParameter);
+}
+
+void FNewSpline::SetClosedLoop(bool bInClosedLoop)
+{
+	SetClosedLoop(bInClosedLoop, true);
+}
+
+void FNewSpline::SetClosedLoop(bool bInClosedLoop, bool bUpdateSpline)
+{
+	if (bInClosedLoop == GetSpline().IsClosedLoop())
+	{
+		return;
+	}
+	TRACE_CPUPROFILER_EVENT_SCOPE(FNewSpline::SetClosedLoop);
+
+	MarkLegacyCurvesDirty();
+
+	GetSpline().SetClosedLoop(bInClosedLoop);
+	
+	if (UE::Geometry::Spline::TTangentBezierSpline<FVector>* ScaleChild = GetTypedAttributeChannel<UE::Geometry::Spline::TTangentBezierSpline<FVector>>(ScaleAttrName))
+	{
+		ScaleChild->SetClosedLoop(bInClosedLoop);
+		ScaleChild->SetKnotVector(GetSpline().GetKnotVector());
+	}
+		
+	if (UE::Geometry::Spline::TTangentBezierSpline<FQuat>* RotChild = GetTypedAttributeChannel<UE::Geometry::Spline::TTangentBezierSpline<FQuat>>(RotationAttrName))
+	{
+		RotChild->SetClosedLoop(bInClosedLoop);
+		RotChild->SetKnotVector(GetSpline().GetKnotVector());
+	}
+	
+	if (bUpdateSpline)
+	{
+		UpdateSpline();
+	}
+}
+
+void FNewSpline::Reset()
+{
+	Clear();
+	GetSpline().SetTangentModes(TArray<UE::Geometry::Spline::ETangentMode>());
+	PositionCurve.Reset();
+	RotationCurve.Reset();
+	ScaleCurve.Reset();
+}
+
+void FNewSpline::ResetRotation()
+{
+    ClearAttributeChannel(RotationAttrName);
+
+    for (int32 Idx = 0; Idx < GetSpline().GetNumPoints(); Idx++)
+    {    		
+    	AddAttributeValue<FQuat>(RotationAttrName, FQuat::Identity, static_cast<float>(Idx));
+    }
+
+	SetAttributeChannelRange(RotationAttrName, FInterval1f(0.f, 1.f), EMappingRangeSpace::Normalized);
+	if (UE::Geometry::Spline::TTangentBezierSpline<FQuat>* RotSpline = GetTypedAttributeChannel<UE::Geometry::Spline::TTangentBezierSpline<FQuat>>(RotationAttrName))
+	{
+		RotSpline->SetKnotVector(GetSpline().GetKnotVector());
+	}
+}
+
+void FNewSpline::ResetScale()
+{
+    ClearAttributeChannel(ScaleAttrName);
+	
+    for (int32 Idx = 0; Idx < GetSpline().GetNumPoints(); Idx++)
+    {    		
+        AddAttributeValue<FVector>(ScaleAttrName, FVector::OneVector, static_cast<float>(Idx));
+    }
+
+	SetAttributeChannelRange(ScaleAttrName, FInterval1f(0.f, 1.f), EMappingRangeSpace::Normalized);
+	if (UE::Geometry::Spline::TTangentBezierSpline<FVector>* ScaleSpline = GetTypedAttributeChannel<UE::Geometry::Spline::TTangentBezierSpline<FVector>>(ScaleAttrName))
+	{
+		ScaleSpline->SetKnotVector(GetSpline().GetKnotVector());
+	}
+}
+
+void FNewSpline::AddPoint(const FSplinePoint& Point)
+{
+	ValidateRotScale();
+	ON_SCOPE_EXIT
+	{
+		SetAttributeChannelRange(ScaleAttrName, FInterval1f(0.f, 1.f), EMappingRangeSpace::Normalized);
+		SetAttributeChannelRange(RotationAttrName, FInterval1f(0.f, 1.f), EMappingRangeSpace::Normalized);
+		ValidateRotScale();
+	};
+	
+	FTangentBezierControlPoint ControlPoint = ConvertToTangentBezierControlPoint(Point);
+	const int32 NumPoints = GetSpline().GetNumPoints();
+    
+	// Convert InputKey to parameter value
+	int32 Index = FMath::FloorToInt(Point.InputKey);
+	float Fraction = Point.InputKey - Index;
+	
+	if (NumPoints == 0 || Index >= NumPoints)
+	{
+		GetSpline().AppendPoint(ControlPoint);
+
+		if (UE::Geometry::Spline::TTangentBezierSpline<FVector>* ScaleChild = GetTypedAttributeChannel<UE::Geometry::Spline::TTangentBezierSpline<FVector>>(ScaleAttrName))
+		{
+			ScaleChild->AppendPoint(UE::Geometry::Spline::TTangentBezierSpline<FVector>::FTangentBezierControlPoint(Point.Scale));
+			ScaleChild->SetKnotVector(GetSpline().GetKnotVector());
+		}
+		
+		if (UE::Geometry::Spline::TTangentBezierSpline<FQuat>* RotChild = GetTypedAttributeChannel<UE::Geometry::Spline::TTangentBezierSpline<FQuat>>(RotationAttrName))
+		{
+			RotChild->AppendPoint(UE::Geometry::Spline::TTangentBezierSpline<FQuat>::FTangentBezierControlPoint(Point.Rotation.Quaternion()));
+			RotChild->SetKnotVector(GetSpline().GetKnotVector());
+		}
+		
+		MarkLegacyCurvesDirty();
+		return;
+	}
+	else if (Index == 0)
+    {
+    	GetSpline().PrependPoint(ControlPoint);
+    	
+		if (UE::Geometry::Spline::TTangentBezierSpline<FVector>* ScaleChild = GetTypedAttributeChannel<UE::Geometry::Spline::TTangentBezierSpline<FVector>>(ScaleAttrName))
+		{
+    		ScaleChild->PrependPoint(UE::Geometry::Spline::TTangentBezierSpline<FVector>::FTangentBezierControlPoint(Point.Scale));
+			ScaleChild->SetKnotVector(GetSpline().GetKnotVector());
+		}
+    
+		if (UE::Geometry::Spline::TTangentBezierSpline<FQuat>* RotChild = GetTypedAttributeChannel<UE::Geometry::Spline::TTangentBezierSpline<FQuat>>(RotationAttrName))
+		{
+    		RotChild->PrependPoint(UE::Geometry::Spline::TTangentBezierSpline<FQuat>::FTangentBezierControlPoint(Point.Rotation.Quaternion()));
+			RotChild->SetKnotVector(GetSpline().GetKnotVector());
+		}
+    	
+    	MarkLegacyCurvesDirty();
+    	return;
+    }
+	
+	// Insert at calculated parameter
+	int32 InsertedIndex = GetSpline().InsertPointAtSegmentParam(Index, Fraction, ControlPoint);
+	float NewPointInternalParam = GetSpline().GetParameter(InsertedIndex);
+	
+	float ScaleParam = MapParameterToChildSpace(ScaleAttrName, NewPointInternalParam);
+	if (UE::Geometry::Spline::TTangentBezierSpline<FVector>* ScaleChild = GetTypedAttributeChannel<UE::Geometry::Spline::TTangentBezierSpline<FVector>>(ScaleAttrName))
+	{
+		ScaleChild->InsertPointAtGlobalParam(ScaleParam, UE::Geometry::Spline::TTangentBezierSpline<FVector>::FTangentBezierControlPoint(Point.Scale));
+		ScaleChild->SetKnotVector(GetSpline().GetKnotVector());
+	}
+	
+	float RotParam = MapParameterToChildSpace(RotationAttrName, NewPointInternalParam);
+	if (UE::Geometry::Spline::TTangentBezierSpline<FQuat>* RotChild = GetTypedAttributeChannel<UE::Geometry::Spline::TTangentBezierSpline<FQuat>>(RotationAttrName))
+	{
+		RotChild->InsertPointAtGlobalParam(RotParam, UE::Geometry::Spline::TTangentBezierSpline<FQuat>::FTangentBezierControlPoint(Point.Rotation.Quaternion()));
+		RotChild->SetKnotVector(GetSpline().GetKnotVector());
+	}
+	
+	MarkLegacyCurvesDirty();
+}
+
+void FNewSpline::InsertPoint(const FSplinePoint& Point, int32 Index)
+{
+	ValidateRotScale();
+	ON_SCOPE_EXIT
+	{
+		SetAttributeChannelRange(ScaleAttrName, FInterval1f(0.f, 1.f), EMappingRangeSpace::Normalized);
+		SetAttributeChannelRange(RotationAttrName, FInterval1f(0.f, 1.f), EMappingRangeSpace::Normalized);
+		ValidateRotScale();
+	};
+	
+	FTangentBezierControlPoint ControlPoint = ConvertToTangentBezierControlPoint(Point);
+	const int32 NumPoints = GetSpline().GetNumPoints();
+    
+	// Special case for empty spline
+	if (NumPoints == 0 || Index >= NumPoints)
+	{
+		GetSpline().AppendPoint(ControlPoint);
+		
+		if (UE::Geometry::Spline::TTangentBezierSpline<FVector>* ScaleChild = GetTypedAttributeChannel<UE::Geometry::Spline::TTangentBezierSpline<FVector>>(ScaleAttrName))
+		{
+			ScaleChild->AppendPoint(UE::Geometry::Spline::TTangentBezierSpline<FVector>::FTangentBezierControlPoint(Point.Scale));
+			ScaleChild->SetKnotVector(GetSpline().GetKnotVector());
+		}
+		
+		if (UE::Geometry::Spline::TTangentBezierSpline<FQuat>* RotChild = GetTypedAttributeChannel<UE::Geometry::Spline::TTangentBezierSpline<FQuat>>(RotationAttrName))
+		{
+			RotChild->AppendPoint(UE::Geometry::Spline::TTangentBezierSpline<FQuat>::FTangentBezierControlPoint(Point.Rotation.Quaternion()));
+			RotChild->SetKnotVector(GetSpline().GetKnotVector());
+		}
+		
+		MarkLegacyCurvesDirty();
+		return;
+	}
+	else if (Index == 0)
+	{
+		GetSpline().PrependPoint(ControlPoint);
+		
+		if (UE::Geometry::Spline::TTangentBezierSpline<FVector>* ScaleChild = GetTypedAttributeChannel<UE::Geometry::Spline::TTangentBezierSpline<FVector>>(ScaleAttrName))
+		{
+			ScaleChild->PrependPoint(UE::Geometry::Spline::TTangentBezierSpline<FVector>::FTangentBezierControlPoint(Point.Scale));
+			ScaleChild->SetKnotVector(GetSpline().GetKnotVector());
+		}
+	
+		if (UE::Geometry::Spline::TTangentBezierSpline<FQuat>* RotChild = GetTypedAttributeChannel<UE::Geometry::Spline::TTangentBezierSpline<FQuat>>(RotationAttrName))
+		{
+			RotChild->PrependPoint(UE::Geometry::Spline::TTangentBezierSpline<FQuat>::FTangentBezierControlPoint(Point.Rotation.Quaternion()));
+			RotChild->SetKnotVector(GetSpline().GetKnotVector());
+		}
+		
+		MarkLegacyCurvesDirty();
+		return;
+	}
+	
+	Index = FMath::Clamp(Index, 0, GetSpline().GetNumPoints());
+	
+	int32 InsertedIndex = GetSpline().InsertPointAtPosition(Index, ControlPoint);
+	float NewPointInternalParam = GetSpline().GetParameter(InsertedIndex);
+	
+	float ScaleParam = MapParameterToChildSpace(ScaleAttrName, NewPointInternalParam);
+	if (UE::Geometry::Spline::TTangentBezierSpline<FVector>* ScaleChild = GetTypedAttributeChannel<UE::Geometry::Spline::TTangentBezierSpline<FVector>>(ScaleAttrName))
+	{
+		ScaleChild->InsertPointAtGlobalParam(ScaleParam, UE::Geometry::Spline::TTangentBezierSpline<FVector>::FTangentBezierControlPoint(Point.Scale));
+		ScaleChild->SetKnotVector(GetSpline().GetKnotVector());
+	}
+	
+	float RotParam = MapParameterToChildSpace(RotationAttrName, NewPointInternalParam);
+	if (UE::Geometry::Spline::TTangentBezierSpline<FQuat>* RotChild = GetTypedAttributeChannel<UE::Geometry::Spline::TTangentBezierSpline<FQuat>>(RotationAttrName))
+	{
+		RotChild->InsertPointAtGlobalParam(RotParam, UE::Geometry::Spline::TTangentBezierSpline<FQuat>::FTangentBezierControlPoint(Point.Rotation.Quaternion()));
+		RotChild->SetKnotVector(GetSpline().GetKnotVector());
+	}
+	
+	MarkLegacyCurvesDirty();
+}
+
+FSplinePoint FNewSpline::GetPoint(int32 Index) const
+{
+    FSplinePoint Point;
+    
+    if (Index < 0 || Index >= GetSpline().GetNumPoints())
+    {
+    	return Point;
+    }
+
+	Point.InputKey = GetParameterAtIndex(Index);
+    Point.Position = GetSpline().GetValue(Index);
+	Point.ArriveTangent = GetSpline().GetTangentIn(Index);
+    Point.LeaveTangent = GetSpline().GetTangentOut(Index);
+    Point.Rotation = GetAttributeValue<FQuat>(RotationAttrName, Index).Rotator();
+    Point.Scale = GetAttributeValue<FVector>(ScaleAttrName, Index);
+	EInterpCurveMode InterpCurveMode = ConvertTangentModeToInterpCurveMode(GetSpline().GetTangentModes()[Index]);
+    Point.Type = ConvertInterpCurveModeToSplinePointType(InterpCurveMode);
+
+    return Point;
+}
+
+/** Removes a point from the spline */
+void FNewSpline::RemovePoint(int32 Index)
+{
+	ValidateRotScale();
+	ON_SCOPE_EXIT
+	{
+		SetAttributeChannelRange(ScaleAttrName, FInterval1f(0.f, 1.f), EMappingRangeSpace::Normalized);
+		SetAttributeChannelRange(RotationAttrName, FInterval1f(0.f, 1.f), EMappingRangeSpace::Normalized);
+		
+		ValidateRotScale();
+	};
+	
+	if (Index < 0 || Index >= GetSpline().GetNumPoints())
+	{
+		return;
+	}
+	
+	GetSpline().RemovePoint(Index);
+
+	if (UE::Geometry::Spline::TTangentBezierSpline<FVector>* ScaleChild = GetTypedAttributeChannel<UE::Geometry::Spline::TTangentBezierSpline<FVector>>(ScaleAttrName))
+	{
+		ScaleChild->RemovePoint(Index);
+		ScaleChild->SetKnotVector(GetSpline().GetKnotVector());
+	}
+	
+	if (UE::Geometry::Spline::TTangentBezierSpline<FQuat>* RotChild = GetTypedAttributeChannel<UE::Geometry::Spline::TTangentBezierSpline<FQuat>>(RotationAttrName))
+	{
+		RotChild->RemovePoint(Index);
+		RotChild->SetKnotVector(GetSpline().GetKnotVector());
+	}
+	
+    MarkLegacyCurvesDirty();
+}
+
+void FNewSpline::SetLocation(int32 Index, const FVector& InLocation)
+{
+    if (Index < 0 || Index >= GetSpline().GetNumPoints())
+    {
+    	return;
+    }
+	
+	GetSpline().ModifyPoint(Index, {InLocation, GetSpline().GetTangentIn(Index), GetSpline().GetTangentOut(Index), GetSpline().GetTangentMode(Index)});
+	
+	MarkReparamTableDirty();
+    MarkLegacyCurvesDirty();
+}
+
+FVector FNewSpline::GetLocation(const int32 Index) const
+{
+	if (Index < 0 || Index >= GetSpline().GetNumPoints())
+    {
+    	return FVector();
+    }
+
+    return GetSpline().GetValue(Index);
+}
+
+void FNewSpline::SetInTangent(const int32 Index, const FVector& InTangent)
+{
+	if (Index < 0 || Index >= GetSpline().GetNumPoints())
+    {
+    	return;
+    }
+
+	GetSpline().SetPointTangentMode(Index, UE::Geometry::Spline::ETangentMode::User);
+    GetSpline().SetTangentIn(Index, InTangent);
+
+	MarkReparamTableDirty();
+    MarkLegacyCurvesDirty();
+}
+
+FVector FNewSpline::GetInTangent(const int32 Index) const
+{
+	if (Index < 0 || Index >= GetSpline().GetNumPoints())
+    {
+    	return FVector();
+    }
+
+    return GetSpline().GetTangentIn(Index);
+}
+
+void FNewSpline::SetOutTangent(const int32 Index, const FVector& OutTangent)
+{
+	if (Index < 0 || Index >= GetSpline().GetNumPoints())
+    {
+    	return;
+    }
+	
+	GetSpline().SetPointTangentMode(Index, UE::Geometry::Spline::ETangentMode::User);
+	GetSpline().SetTangentOut(Index, OutTangent);
+
+	MarkReparamTableDirty();
+    MarkLegacyCurvesDirty();
+}
+
+FVector FNewSpline::GetOutTangent(const int32 Index) const
+{
+	if (Index < 0 || Index >= GetSpline().GetNumPoints())
+    {
+    	return FVector();
+    }
+
+    return GetSpline().GetTangentOut(Index);
+}
+
+void FNewSpline::SetRotation(int32 Index, const FQuat& InRotation)
+{
+	if (Index < 0 || Index >= GetSpline().GetNumPoints())
+    {
+    	return;
+    }
+
+	ValidateRotScale();
+
+	SetAttributeValue<FQuat>(RotationAttrName, InRotation, Index);
+	
+    MarkLegacyCurvesDirty();
+}
+
+FQuat FNewSpline::GetRotation(const int32 Index) const
+{
+	if (Index < 0 || Index >= GetSpline().GetNumPoints())
+    {
+    	return FQuat::Identity;
+    }
+	
+    return GetAttributeValue<FQuat>(RotationAttrName, Index);
+}
+
+void FNewSpline::SetScale(int32 Index, const FVector& InScale)
+{
+	if (Index < 0 || Index >= GetSpline().GetNumPoints())
+    {
+    	return;
+    }
+
+	ValidateRotScale();
+
+	SetAttributeValue<FVector>(ScaleAttrName, InScale, Index);
+    
+    MarkLegacyCurvesDirty();
+}
+
+FVector FNewSpline::GetScale(const int32 Index) const
+{
+    if (Index < 0 || Index >= GetSpline().GetNumPoints())
+    {
+    	return FVector::OneVector;
+    }
+
+    return GetAttributeValue<FVector>(ScaleAttrName, Index);
+}
+
+void FNewSpline::SetSplinePointType(int32 Index, EInterpCurveMode Type)
+{
+    if (Index < 0 || Index >= GetSpline().GetNumPoints())
+    {
+    	return;
+    }
+	
+	GetSpline().SetPointTangentMode(Index, ConvertInterpCurveModeToTangentMode(Type));
+
+	MarkReparamTableDirty();
+    MarkLegacyCurvesDirty();
+}
+
+EInterpCurveMode FNewSpline::GetSplinePointType(int32 Index) const
+{
+    if (Index < 0 || Index >= GetSpline().GetNumPoints())
+    {
+    	return CIM_Unknown;
+    }
+	UE::Geometry::Spline::ETangentMode TangentMode = GetSpline().GetTangentMode(Index);
+    return ConvertTangentModeToInterpCurveMode(TangentMode);
+}
+
+float FNewSpline::GetParameterAtIndex(int32 Index) const
+{
+    if (Index < 0 || (!IsClosedLoop() && Index >= GetSpline().GetNumPoints()) || (IsClosedLoop() && Index > GetSpline().GetNumPoints()))
+    {
+    	return 0.f;
+    }
+
+	return FromInternalSplineSpace(GetSpline().GetParameter(Index));
+}
+
+float FNewSpline::GetParameterAtDistance(float Distance)
+{
+	UpdateReparamTable();
+	
+	UE::TReadScopeLock Lock(ReparamTableRWLock);
+	
+	return ReparamTable.Eval(Distance);
+}
+
+float FNewSpline::GetDistanceAtParameter(float Parameter)
+{
+	if (!GetSegmentSpace().Contains(Parameter))
+	{
+		return 0.f;
+	}
+	
+	UpdateReparamTable();
+	
+	UE::TReadScopeLock Lock(ReparamTableRWLock);
+
+	const float ParameterMax = static_cast<float>(GetNumberOfSegments());
+	const float Key = (Parameter / ParameterMax) * (ReparamTable.Points.Num() - 1);
+	const int32 LowerKey = FMath::FloorToInt32(Key);
+	ensureAlways(LowerKey >= 0 && LowerKey < ReparamTable.Points.Num());
+	const int32 UpperKey = FMath::CeilToInt32(Key);
+	ensureAlways(UpperKey >= 0 && UpperKey < ReparamTable.Points.Num());
+	const float Alpha = FMath::Frac(Key);
+	return FMath::Lerp(ReparamTable.Points[LowerKey].InVal, ReparamTable.Points[UpperKey].InVal, Alpha);
+}
+
+float FNewSpline::FindNearest(const FVector& InLocation, float& OutSquaredDistance) const
+{
+    if (UE::Spline::GFallbackFindNearest)
+    {
+    	return PositionCurve.FindNearest(InLocation, OutSquaredDistance);
+    }
+    else
+    {
+	    return FromInternalSplineSpace(GetSpline().FindNearest(InLocation, OutSquaredDistance));
+    }
+}
+
+float FNewSpline::FindNearestOnSegment(const FVector& InLocation, int32 SegmentIndex, float& OutSquaredDistance) const
+{
+    if (UE::Spline::GFallbackFindNearest)
+    {
+        if (!PositionCurve.Points.IsValidIndex(SegmentIndex))
+        {
+            return 0.f;
+        }
+    
+        return PositionCurve.FindNearestOnSegment(InLocation, SegmentIndex, OutSquaredDistance);
+    }
+    else
+    {
+	   return FromInternalSplineSpace(GetSpline().FindNearestOnSegment(InLocation, SegmentIndex, OutSquaredDistance));
+    }
+}
+
+FVector FNewSpline::EvaluatePosition(float Parameter) const
+{
+	if (UE::Spline::GUseLegacyPositionEvaluation)
+	{
+		RebuildLegacyCurves();
+		return PositionCurve.Eval(Parameter);
+	}
+	else
+	{
+		return Evaluate(ToInternalSplineSpace(Parameter));
+	}
+}
+
+FVector FNewSpline::EvaluateDerivative(float Parameter) const
+{
+	Parameter = ToInternalSplineSpace(Parameter);
+    return GetSpline().GetTangent(Parameter);
+}
+
+FQuat FNewSpline::EvaluateRotation(float Parameter) const
+{
+	if (UE::Spline::GUseLegacyRotationEvaluation)
+	{
+		RebuildLegacyCurves();
+		return RotationCurve.Eval(Parameter);
+	}
+	else
+	{
+		return EvaluateAttribute<FQuat>(RotationAttrName, ToInternalSplineSpace(Parameter));
+	}
+}
+
+FVector FNewSpline::EvaluateScale(float Parameter) const
+{
+	if (UE::Spline::GUseLegacyScaleEvaluation)
+	{
+		RebuildLegacyCurves();
+		return ScaleCurve.Eval(Parameter);
+	}
+	else
+	{
+		return EvaluateAttribute<FVector>(ScaleAttrName, ToInternalSplineSpace(Parameter));
+	}
+}
+
+void FNewSpline::UpdateSpline()
+{
+	TRACE_CPUPROFILER_EVENT_SCOPE(FNewSpline::UpdateSpline);
+
+	// Update our parameterization, forward it to the rotation and scale channels, and then recompute tangents.
+	Reparameterize();
+
+    MarkReparamTableDirty();
+
+	// Updates legacy curves based on points, never evaluates internal spline.
+    RebuildLegacyCurves();					
+}
+
+void FNewSpline::UpdateSpline(const FSpline::FUpdateSplineParams& Params)
+{
+	TRACE_CPUPROFILER_EVENT_SCOPE(FNewSpline::UpdateSpline);
+
+	// todo: handled unhandled params
+   
+	if (UE::Spline::GSkipDoubleUpdate)
+	{
+		// Avoid updating the spline twice.
+		constexpr bool bUpdateSpline = false;
+		SetClosedLoop(Params.bClosedLoop, bUpdateSpline);
+	}
+	else
+	{
+		SetClosedLoop(Params.bClosedLoop);
+	}
+
+	GetSpline().SetStationaryEndpoints(Params.bStationaryEndpoints);
+    ReparamStepsPerSegment = Params.ReparamStepsPerSegment;
+
+    UpdateSpline();
+}
+
+void FNewSpline::Reparameterize()
+{
+	TRACE_CPUPROFILER_EVENT_SCOPE(FNewSpline::Reparameterize);
+
+	// Caching segment space attribute parameters so we can restore them (necessary because we store the mapping in parent space which are about to be invalidated).
+	TMap<FName, TPair<FInterval1f, TArray<float>>> CachedChannelParameterizations;
+	for (const FName& ChannelName : GetAttributeChannelNamesByValueType<float>())
+	{
+		if (ChannelName == RotationAttrName || ChannelName == ScaleAttrName || NumAttributeValues<float>(ChannelName) == 0)
+		{
+			continue;
+		}
+
+		ensureAlwaysMsgf(GetAttributeChannelRangeSpace(ChannelName) == EMappingRangeSpace::Parent, TEXT("Attribute channel mapping range space is not Parent!"));
+
+		FInterval1f ChannelRangeInParentSpace = GetAttributeChannelRange(ChannelName);
+		FInterval1f ChannelRangeInSegmentSpace = FInterval1f(FromInternalSplineSpace(ChannelRangeInParentSpace.Min), FromInternalSplineSpace(ChannelRangeInParentSpace.Max));
+
+		CachedChannelParameterizations.FindOrAdd(ChannelName) = TPair<FInterval1f, TArray<float>>(ChannelRangeInSegmentSpace, CacheChannelKnots<float>(ChannelName));
+
+		for (float& InternalKnot : CachedChannelParameterizations[ChannelName].Value)
+		{
+			InternalKnot = FromInternalSplineSpace(InternalKnot);
+		}
+	}
+
+	GetSpline().Reparameterize(GetParameterizationPolicy());
+	GetSpline().UpdateTangents();
+
+	for (const FName& ChannelName : GetAttributeChannelNamesByValueType<float>())
+	{
+		if (ChannelName == RotationAttrName || ChannelName == ScaleAttrName || NumAttributeValues<float>(ChannelName) == 0)
+		{
+			continue;
+		}
+
+		const FInterval1f& ChannelRangeInSegmentSpace = CachedChannelParameterizations[ChannelName].Key;
+		const FInterval1f ChannelRangeInParentSpace = FInterval1f(ToInternalSplineSpace(ChannelRangeInSegmentSpace.Min), ToInternalSplineSpace(ChannelRangeInSegmentSpace.Max));
+
+		SetAttributeChannelRange(ChannelName, ChannelRangeInParentSpace, EMappingRangeSpace::Parent);
+
+		for (float& InternalKnot : CachedChannelParameterizations[ChannelName].Value)
+		{
+			InternalKnot = ToInternalSplineSpace(InternalKnot);
+		}
+
+		RestoreChannelKnots(ChannelName, CachedChannelParameterizations[ChannelName].Value);
+		if (UE::Geometry::Spline::TTangentBezierSpline<float>* Channel = GetTypedAttributeChannel<UE::Geometry::Spline::TTangentBezierSpline<float>>(ChannelName))
+		{
+			Channel->UpdateTangents();
+		}
+	}
+
+	if (UE::Geometry::Spline::TTangentBezierSpline<FVector>* ScaleSpline = GetTypedAttributeChannel<UE::Geometry::Spline::TTangentBezierSpline<FVector>>(ScaleAttrName))
+	{
+		ScaleSpline->SetKnotVector(GetSpline().GetKnotVector());
+		ScaleSpline->UpdateTangents();
+	}
+
+	if (UE::Geometry::Spline::TTangentBezierSpline<FQuat>* RotSpline = GetTypedAttributeChannel<UE::Geometry::Spline::TTangentBezierSpline<FQuat>>(RotationAttrName))
+	{
+		RotSpline->SetKnotVector(GetSpline().GetKnotVector());
+		RotSpline->UpdateTangents();
+	}
+}
+
+float FNewSpline::FromNormalizedSpace(float Parameter) const
+{
+	return GetNumberOfSegments() * Parameter;
+}
+
+float FNewSpline::ToNormalizedSpace(float Parameter) const
+{
+	return Parameter / GetNumberOfSegments();
+}
+
+float FNewSpline::FromInternalSplineSpace(float Parameter) const
+{
+	float OutLocalParam;
+	const int32 SegmentIndex = GetSpline().FindSegmentIndex(Parameter, OutLocalParam);
+	return static_cast<float>(SegmentIndex) + OutLocalParam;
+}
+float FNewSpline::ToInternalSplineSpace(float Parameter) const
+{
+	const int32 NumSegments = GetNumberOfSegments();
+	float ClampedParam = FMath::Clamp(Parameter, 0.0f, static_cast<float>(NumSegments));
+	const int32 SegmentIndex = FMath::Min(FMath::FloorToInt(ClampedParam), NumSegments - 1);
+	const float LocalT = ClampedParam - static_cast<float>(SegmentIndex); // Parameter [0,1] within segment
+	return GetSegmentParameterRange(SegmentIndex).Interpolate(LocalT);
+}
+
+void FNewSpline::MarkReparamTableDirty()
+{
+	ReparamTableNextVersion.Increment();
+}
+
+void FNewSpline::UpdateReparamTable()
+{
+	const int32 NumSegments = GetNumberOfSegments();
+
+	auto IsReparamTableDirty = [this]()
+	{
+		return ReparamTableNextVersion.GetValue() != ReparamTableVersion;
+	};
+	
+	if (!IsReparamTableDirty() || ReparamStepsPerSegment == 0 || NumSegments == 0)
+	{
+		return;
+	}
+
+	UE::TWriteScopeLock Lock(ReparamTableRWLock);
+
+	if (!IsReparamTableDirty())
+	{
+		return;
+	}
+
+	// We can't rely on the next version to not change during the update.
+	const uint32 CachedNextVersion = ReparamTableNextVersion.GetValue();
+
+	ReparamTable.Points.Reset(NumSegments * ReparamStepsPerSegment + 1);
+	float AccumulatedLength = 0.0f;
+	for (int32 SegmentIndex = 0; SegmentIndex < NumSegments; ++SegmentIndex)
+	{
+		ReparamTable.Points.Emplace(AccumulatedLength, SegmentIndex, 0.0f, 0.0f, CIM_Linear);
+		for (int32 Step = 1; Step < ReparamStepsPerSegment; ++Step)
+		{
+			const float Param = static_cast<float>(Step) / ReparamStepsPerSegment;
+			ReparamTable.Points.Emplace(GetSegmentLength(SegmentIndex, Param) + AccumulatedLength, SegmentIndex + Param, 0.0f, 0.0f, CIM_Linear);
+		}
+
+		AccumulatedLength += GetSegmentLength(SegmentIndex, 1.0f);
+	}
+	ReparamTable.Points.Emplace(AccumulatedLength, static_cast<float>(NumSegments), 0.0f, 0.0f, CIM_Linear);
+
+	ReparamTableVersion = CachedNextVersion;
+}
+
+void FNewSpline::MarkLegacyCurvesDirty()
+{
+	if (UE::Spline::GImmediatelyUpdateLegacyCurves)
+	{
+		RebuildLegacyCurves();
+	}
+	else
+	{
+		LegacyCurvesNextVersion.Increment();
+	}
+}
+
+void FNewSpline::RebuildLegacyCurves() const
+{
+	TRACE_CPUPROFILER_EVENT_SCOPE(FNewSpline::RebuildLegacyCurves);
+
+	// Curves are dirty if:
+	// 1) Current version is not the same as next version (if next version == version, it has not been dirtied)
+	// 2) UE::Spline::GImmediatelyUpdateLegacyCurves == true (legacy curves are updated any time they might be dirty)
+	auto AreLegacyCurvesDirty = [this]()
+	{
+		return LegacyCurvesNextVersion.GetValue() != LegacyCurvesVersion || UE::Spline::GImmediatelyUpdateLegacyCurves;
+	};
+	
+	if (!AreLegacyCurvesDirty())
+	{
+		return;
+	}
+
+	PositionCurve.Points.Reset(GetSpline().GetNumPoints());
+	RotationCurve.Points.Reset(GetSpline().GetNumPoints());
+	ScaleCurve.Points.Reset(GetSpline().GetNumPoints());
+
+    PositionCurve.bIsLooped = IsClosedLoop();
+    PositionCurve.LoopKeyOffset = 1.f;
+    RotationCurve.bIsLooped = IsClosedLoop();
+	RotationCurve.LoopKeyOffset = 1.f;
+    ScaleCurve.bIsLooped = IsClosedLoop();
+	ScaleCurve.LoopKeyOffset = 1.f;
+	
+	using FRotationChannel = UE::Geometry::Spline::TTangentBezierSpline<FQuat>;
+	using FScaleChannel = UE::Geometry::Spline::TTangentBezierSpline<FVector>;
+
+	const FRotationChannel* RotationChannel = GetTypedAttributeChannel<FRotationChannel>(RotationAttrName);
+	const FScaleChannel* ScaleChannel = GetTypedAttributeChannel<FScaleChannel>(ScaleAttrName);
+
+    for (int32 i = 0; i < GetSpline().GetNumPoints(); ++i)
+    {
+    	const float AttributeParameter = static_cast<float>(i);
+    	
+    	// Additionally populate interp curves here so we don't double loop
+    	PositionCurve.Points.Emplace(AttributeParameter, GetSpline().GetValue(i), GetSpline().GetTangentIn(i), GetSpline().GetTangentOut(i), ConvertTangentModeToInterpCurveMode(GetSpline().GetTangentModes()[i]));
+		RotationCurve.Points.Emplace(AttributeParameter, RotationChannel ? RotationChannel->GetValue(i) : FQuat::Identity);
+    	RotationCurve.Points.Last().InterpMode = CIM_CurveAuto;
+		ScaleCurve.Points.Emplace(AttributeParameter, ScaleChannel ? ScaleChannel->GetValue(i) : FVector::OneVector);
+    	ScaleCurve.Points.Last().InterpMode = CIM_CurveAuto;
+    }
+
+	Algo::SortBy(RotationCurve.Points, &FInterpCurvePoint<FQuat>::InVal);
+	Algo::SortBy(ScaleCurve.Points, &FInterpCurvePoint<FVector>::InVal);
+
+	PositionCurve.AutoSetTangents(0.f, false);
+	RotationCurve.AutoSetTangents(0.f, false);
+	ScaleCurve.AutoSetTangents(0.f, false);
+    
+    LegacyCurvesVersion = LegacyCurvesNextVersion.GetValue();
+}
+
+FNewSpline::FTangentBezierControlPoint FNewSpline::ConvertToTangentBezierControlPoint(const FSplinePoint& Point) const
+{
+	FTangentBezierControlPoint NewPoint;
+	NewPoint.Position = Point.Position;
+	NewPoint.TangentIn = Point.ArriveTangent;
+	NewPoint.TangentOut = Point.LeaveTangent;
+	NewPoint.TangentMode = ConvertInterpCurveModeToTangentMode(
+		ConvertSplinePointTypeToInterpCurveMode(Point.Type));
+	return NewPoint;
+}
+
+void FNewSpline::UpdatePointAttributes(const FSplinePoint& Point, int32 PointIndex)
+{
+	if (PointIndex < 0)
+	{
+		return;
+	}
+	
+	SetRotation(PointIndex, Point.Rotation.Quaternion());
+	SetScale(PointIndex, Point.Scale);
+        
+	MarkLegacyCurvesDirty();
+}
+
+float FNewSpline::ConvertIndexToInternalParameter(int32 Index, float Fraction) const
+{
+	// Handle empty spline
+	if (GetSpline().GetNumPoints() <= 1)
+		return 0.0f;
+    
+	// Clamp to valid range
+	Index = FMath::Clamp(Index, 0, GetSpline().GetNumPoints() - 1);
+    
+	// For exact index, just return parameter at that index
+	if (FMath::IsNearlyZero(Fraction))
+		return GetSpline().GetParameter(Index);
+    
+	// For fractional indices, interpolate
+	int32 NextIndex = FMath::Min(Index + 1, GetSpline().GetNumPoints() - 1);
+	float StartParam = GetSpline().GetParameter(Index);
+	float EndParam = GetSpline().GetParameter(NextIndex);
+    
+	return FMath::Lerp(StartParam, EndParam, Fraction);
+}
+
+int32 FNewSpline::ConvertInternalParameterToNearestPointIndex(float Parameter) const
+{
+	// Handle empty spline
+	if (GetSpline().GetNumPoints() <= 1)
+		return 0;
+    
+	// Find segment containing this parameter
+	for (int32 i = 0; i < GetSpline().GetNumPoints() - 1; i++)
+	{
+		float StartParam = GetSpline().GetParameter(i);
+		float EndParam = GetSpline().GetParameter(i + 1);
+        
+		if (Parameter >= StartParam && Parameter <= EndParam)
+		{
+			// Return closer index
+			float Fraction = (Parameter - StartParam) / (EndParam - StartParam);
+			return (Fraction <= 0.5f) ? i : (i + 1);
+		}
+	}
+    
+	// Out of range
+	return (Parameter < GetSpline().GetParameter(0)) ? 0 : (GetSpline().GetNumPoints() - 1);
+}
+
+void FNewSpline::ValidateRotScale() const
+{
+	if (!UE::Spline::GValidateRotScale)
+	{
+		return;
+	}
+	
+	// Validate child existence
+	UE::Geometry::Spline::TTangentBezierSpline<FQuat>* RotSpline = GetTypedAttributeChannel<UE::Geometry::Spline::TTangentBezierSpline<FQuat>>(RotationAttrName);
+	UE::Geometry::Spline::TTangentBezierSpline<FVector>* ScaleSpline = GetTypedAttributeChannel<UE::Geometry::Spline::TTangentBezierSpline<FVector>>(ScaleAttrName);
+	ensure (RotSpline && ScaleSpline);
+
+	if (RotSpline && ScaleSpline)
+	{
+		// Validate number of points
+		int32 NumControlPoints = GetSpline().GetNumPoints();
+		int32 NumRotPoints = RotSpline->GetNumPoints();
+		int32 NumScalePoints = ScaleSpline->GetNumPoints();
+		ensure(NumControlPoints == NumRotPoints && NumRotPoints == NumScalePoints);
+
+		// Validate control point -> attr point mapping
+		for (int32 i = 0; i < NumControlPoints; i++)
+		{
+			float InternalParam = GetSpline().GetParameter(i);
+		
+			float ExpectedRotParam = MapParameterToChildSpace(RotationAttrName, InternalParam);
+			float ActualRotParam = RotSpline->GetParameter(i);
+			ensure(FMath::IsNearlyEqual(ExpectedRotParam, ActualRotParam, UE_KINDA_SMALL_NUMBER));
+		
+			float ExpectedScaleParam = MapParameterToChildSpace(ScaleAttrName, InternalParam);
+			float ActualScaleParam = ScaleSpline->GetParameter(i);
+			ensure(FMath::IsNearlyEqual(ExpectedScaleParam, ActualScaleParam, UE_KINDA_SMALL_NUMBER));
+		}
+	}
+}
+
+template <typename AttrType>
+TArray<AttrType> FNewSpline::CacheChannelKnots(const FName& InName) const
+{
+	using AttrSplineType = UE::Geometry::Spline::TTangentBezierSpline<AttrType>;
+
+	TArray<float> InternalParameters;
+
+	if (AttrSplineType* Channel = GetTypedAttributeChannel<AttrSplineType>(InName))
+	{
+		for (int InternalIdx = 1; InternalIdx < Channel->GetNumPoints() - 1; ++InternalIdx)
+		{
+			InternalParameters.Add(MapParameterFromChildSpace(InName, Channel->GetParameter(InternalIdx)));
+		}
+	}
+		
+	return InternalParameters;
+}
+
+template <typename AttrType>
+void FNewSpline::RestoreChannelKnots(const FName& InName, const TArray<AttrType>& InKnots)
+{
+	using AttrSplineType = UE::Geometry::Spline::TTangentBezierSpline<AttrType>;
+
+	if (InKnots.Num() == 0) return;
+
+	if (AttrSplineType* Channel = GetTypedAttributeChannel<AttrSplineType>(InName))
+	{
+		// Provided knots are for internal points, bail out if we don't have the correct number of knots.
+		if (ensureAlwaysMsgf(InKnots.Num() == Channel->GetNumPoints() - 2, TEXT("Invalid number of knots!")))
+		{
+			for (int InternalIdx = 1; InternalIdx < Channel->GetNumPoints() - 1; ++InternalIdx)
+			{
+				Channel->SetParameter(InternalIdx, MapParameterToChildSpace(InName, InKnots[InternalIdx - 1]));
+			}
+		}
+	}
+}
+
+/** FLegacySpline Implementation */
+
+FLegacySpline::FLegacySpline(const FNewSpline& Other)
+	: FLegacySpline()
+{
+	PositionCurve = Other.GetSplinePointsPosition();
+	RotationCurve = Other.GetSplinePointsRotation();
+	ScaleCurve = Other.GetSplinePointsScale();
+
+	FSpline::FUpdateSplineParams InParams;
+	InParams.bClosedLoop = PositionCurve.bIsLooped;
+	UpdateSpline(InParams);
+
+	PositionCurve.AutoSetTangents(0.f, false);
+}
+
+FLegacySpline::FLegacySpline(const FSplineCurves& InSpline)
+	: FLegacySpline()
+{
+	PositionCurve = InSpline.Position;
+	RotationCurve = InSpline.Rotation;
+	ScaleCurve = InSpline.Scale;
+	ReparamTable = InSpline.ReparamTable;
+}
+
+void FLegacySpline::AddPoint(const FSplinePoint& InPoint)
+{
+	auto UpperBound = [this](float Value)
+	{
+		int32 Count = PositionCurve.Points.Num();
+		int32 First = 0;
+
+		while (Count > 0)
+		{
+			const int32 Middle = Count / 2;
+			if (Value >= PositionCurve.Points[First + Middle].InVal)
+			{
+				First += Middle + 1;
+				Count -= Middle + 1;
+			}
+			else
+			{
+				Count = Middle;
+			}
+		}
+
+		return First;
+	};
+	
+	int32 Index = UpperBound(InPoint.InputKey);
+	
+	PositionCurve.Points.Insert(FInterpCurvePoint<FVector>(
+		InPoint.InputKey,
+		InPoint.Position,
+		InPoint.ArriveTangent,
+		InPoint.LeaveTangent,
+		ConvertSplinePointTypeToInterpCurveMode(InPoint.Type)
+	), Index);
+
+	RotationCurve.Points.Insert(FInterpCurvePoint<FQuat>(
+		InPoint.InputKey,
+		InPoint.Rotation.Quaternion(),
+		FQuat::Identity,
+		FQuat::Identity,
+		CIM_CurveAuto
+	), Index);
+
+	ScaleCurve.Points.Insert(FInterpCurvePoint<FVector>(
+		InPoint.InputKey,
+		InPoint.Scale,
+		FVector::ZeroVector,
+		FVector::ZeroVector,
+		CIM_CurveAuto
+	), Index);
+}
+
+void FLegacySpline::InsertPoint(const FSplinePoint& InPoint, int32 Index)
+{
+	const float InKey = (Index == 0) ? 0.0f : GetParameterAtIndex(Index - 1) + 1.0f;
+	
+	PositionCurve.Points.Insert(FInterpCurvePoint<FVector>(
+		InKey,
+		InPoint.Position,
+		InPoint.ArriveTangent,
+		InPoint.LeaveTangent,
+		ConvertSplinePointTypeToInterpCurveMode(InPoint.Type)
+	), Index);
+
+	RotationCurve.Points.Insert(FInterpCurvePoint<FQuat>(
+		InKey,
+		InPoint.Rotation.Quaternion(),
+		FQuat::Identity,
+		FQuat::Identity,
+		CIM_CurveAuto
+	), Index);
+
+	ScaleCurve.Points.Insert(FInterpCurvePoint<FVector>(
+		InKey,
+		InPoint.Scale,
+		FVector::ZeroVector,
+		FVector::ZeroVector,
+		CIM_CurveAuto
+	), Index);
+
+	// Increment all point input values after the inserted element
+	for (int32 Idx = Index + 1; Idx < PositionCurve.Points.Num(); Idx++)
+	{
+		PositionCurve.Points[Idx].InVal += 1.f;
+		RotationCurve.Points[Idx].InVal += 1.f;
+		ScaleCurve.Points[Idx].InVal += 1.f;
+	}
+}
+
+FSplinePoint FLegacySpline::GetPoint(const int32 Index) const
+{
+	FSplinePoint Point;
+
+	if (!PositionCurve.Points.IsValidIndex(Index))
+	{
+		return Point;
+	}
+	
+	Point.InputKey = PositionCurve.Points[Index].InVal;
+	Point.Position = PositionCurve.Points[Index].OutVal;
+	Point.ArriveTangent = PositionCurve.Points[Index].ArriveTangent;
+	Point.LeaveTangent = PositionCurve.Points[Index].LeaveTangent;
+	Point.Rotation = RotationCurve.Points[Index].OutVal.Rotator();
+	Point.Scale = ScaleCurve.Points[Index].OutVal;
+	Point.Type = ConvertInterpCurveModeToSplinePointType(PositionCurve.Points[Index].InterpMode);
+	
+	return Point;
+}
+
+void FLegacySpline::RemovePoint(int32 Index)
+{
+	if (!PositionCurve.Points.IsValidIndex(Index))
+	{
+		return;
+	}
+	
+	PositionCurve.Points.RemoveAt(Index);
+	RotationCurve.Points.RemoveAt(Index);
+	ScaleCurve.Points.RemoveAt(Index);
+
+	while (Index < GetNumControlPoints())
+	{
+		PositionCurve.Points[Index].InVal -= 1.0f;
+		RotationCurve.Points[Index].InVal -= 1.0f;
+		ScaleCurve.Points[Index].InVal -= 1.0f;
+		Index++;
+	}
+}
+
+void FLegacySpline::SetLocation(int32 Index, const FVector& InLocation)
+{
+	if (!PositionCurve.Points.IsValidIndex(Index))
+	{
+		return;
+	}
+	
+	PositionCurve.Points[Index].OutVal = InLocation;
+}
+
+FVector FLegacySpline::GetLocation(const int32 Index) const
+{
+	return PositionCurve.Points.IsValidIndex(Index) ? PositionCurve.Points[Index].OutVal : FVector();
+}
+
+void FLegacySpline::SetInTangent(const int32 Index, const FVector& InTangent)
+{
+	if (!PositionCurve.Points.IsValidIndex(Index))
+	{
+		return;
+	}
+	
+	PositionCurve.Points[Index].ArriveTangent = InTangent;
+	PositionCurve.Points[Index].InterpMode = CIM_CurveUser;
+}
+
+FVector FLegacySpline::GetInTangent(const int32 Index) const
+{
+	return PositionCurve.Points.IsValidIndex(Index) ? PositionCurve.Points[Index].ArriveTangent : FVector();
+}
+
+void FLegacySpline::SetOutTangent(const int32 Index, const FVector& OutTangent)
+{
+	if (!PositionCurve.Points.IsValidIndex(Index))
+	{
+		return;
+	}
+	
+	PositionCurve.Points[Index].LeaveTangent = OutTangent;
+	PositionCurve.Points[Index].InterpMode = CIM_CurveUser;
+}
+
+FVector FLegacySpline::GetOutTangent(const int32 Index) const
+{
+	return PositionCurve.Points.IsValidIndex(Index) ? PositionCurve.Points[Index].LeaveTangent : FVector();
+}
+
+void FLegacySpline::SetRotation(int32 Index, const FQuat& InRotation)
+{
+	if (!RotationCurve.Points.IsValidIndex(Index))
+	{
+		return;
+	}
+	
+	RotationCurve.Points[Index].OutVal = InRotation;
+}
+
+FQuat FLegacySpline::GetRotation(const int32 Index) const
+{
+	return RotationCurve.Points.IsValidIndex(Index) ? RotationCurve.Points[Index].OutVal : FQuat();
+}
+
+void FLegacySpline::SetScale(int32 Index, const FVector& InScale)
+{
+	if (!ScaleCurve.Points.IsValidIndex(Index))
+	{
+		return;
+	}
+	
+	ScaleCurve.Points[Index].OutVal = InScale;
+}
+
+FVector FLegacySpline::GetScale(const int32 Index) const
+{
+	return ScaleCurve.Points.IsValidIndex(Index) ? ScaleCurve.Points[Index].OutVal : FVector();
+}
+
+void FLegacySpline::SetSplinePointType(int32 Index, EInterpCurveMode Type)
+{
+	if (!PositionCurve.Points.IsValidIndex(Index))
+	{
+		return;
+	}
+	
+	PositionCurve.Points[Index].InterpMode = Type;
+}
+
+EInterpCurveMode FLegacySpline::GetSplinePointType(const int32 Index) const
+{
+	return PositionCurve.Points.IsValidIndex(Index) ? PositionCurve.Points[Index].InterpMode.GetValue() : CIM_Unknown;
+}
+
+float FLegacySpline::GetParameterAtIndex(int32 Index) const
+{
+	return PositionCurve.Points.IsValidIndex(Index) ? PositionCurve.Points[Index].InVal : 0.f;
+}
+
+float FLegacySpline::GetParameterAtDistance(float Distance) const
+{
+	return ReparamTable.Eval(Distance);
+}
+
+float FLegacySpline::GetDistanceAtParameter(float Parameter) const
+{
+	// this might be duplicating what an interp curve can already do...
+	
+	const float ParameterMax = PositionCurve.Points.Last().InVal;
+	const float Key = (Parameter / ParameterMax) * (ReparamTable.Points.Num() - 1);
+	const int32 LowerKey = FMath::FloorToInt32(Key);
+	ensureAlways(LowerKey >= 0 && LowerKey < ReparamTable.Points.Num());
+	const int32 UpperKey = FMath::CeilToInt32(Key);
+	ensureAlways(UpperKey >= 0 && UpperKey < ReparamTable.Points.Num());
+	const float Alpha = FMath::Frac(Key);
+	const float Distance = FMath::Lerp(ReparamTable.Points[LowerKey].InVal, ReparamTable.Points[UpperKey].InVal, Alpha);
+
+	return Distance;
+}
+
+FVector FLegacySpline::Evaluate(float Param) const
+{
+	return PositionCurve.Eval(Param);
+}
+
+FVector FLegacySpline::EvaluateDerivative(float Param) const
+{
+	return PositionCurve.EvalDerivative(Param);
+}
+
+FQuat FLegacySpline::EvaluateRotation(float Param) const
+{
+	return RotationCurve.Eval(Param);
+}
+
+FVector FLegacySpline::EvaluateScale(float Param) const
+{
+	return ScaleCurve.Eval(Param);
+}
+
+float FLegacySpline::FindNearest(const FVector& InLocation, float& OutSquaredDist) const
+{
+	return PositionCurve.FindNearest(InLocation, OutSquaredDist);
+}
+
+float FLegacySpline::FindNearestOnSegment(const FVector& InLocation, int32 SegmentIndex, float& OutSquaredDist) const
+{
+	if (!PositionCurve.Points.IsValidIndex(SegmentIndex))
+	{
+		return 0.f;
+	}
+	
+	return PositionCurve.FindNearestOnSegment(InLocation, SegmentIndex, OutSquaredDist);
+}
+
+bool FLegacySpline::Serialize(FArchive& Ar)
+{
+	Ar.UsingCustomVersion(FFortniteMainBranchObjectVersion::GUID);
+	Ar << PositionCurve;
+	Ar << RotationCurve;
+	Ar << ScaleCurve;
+	Ar << ReparamTable;
+
+	return true;
+}
+
+float FLegacySpline::GetSegmentLength(const int32 Index, const float Param, const FVector& Scale3D) const
+{
+	const int32 NumPoints = PositionCurve.Points.Num();
+	const int32 LastPoint = NumPoints - 1;
+
+	check(Index >= 0 && ((PositionCurve.bIsLooped && Index < NumPoints) || (!PositionCurve.bIsLooped && Index < LastPoint)));
+	check(Param >= 0.0f && Param <= 1.0f);
+
+	// Evaluate the length of a Hermite spline segment.
+	// This calculates the integral of |dP/dt| dt, where P(t) is the spline equation with components (x(t), y(t), z(t)).
+	// This isn't solvable analytically, so we use a numerical method (Legendre-Gauss quadrature) which performs very well
+	// with functions of this type, even with very few samples.  In this case, just 5 samples is sufficient to yield a
+	// reasonable result.
+
+	struct FLegendreGaussCoefficient
+	{
+		float Abscissa;
+		float Weight;
+	};
+
+	static const FLegendreGaussCoefficient LegendreGaussCoefficients[] =
+	{
+		{ 0.0f, 0.5688889f },
+		{ -0.5384693f, 0.47862867f },
+		{ 0.5384693f, 0.47862867f },
+		{ -0.90617985f, 0.23692688f },
+		{ 0.90617985f, 0.23692688f }
+	};
+
+	const auto& StartPoint = PositionCurve.Points[Index];
+	const auto& EndPoint = PositionCurve.Points[Index == LastPoint ? 0 : Index + 1];
+
+	const auto& P0 = StartPoint.OutVal;
+	const auto& T0 = StartPoint.LeaveTangent;
+	const auto& P1 = EndPoint.OutVal;
+	const auto& T1 = EndPoint.ArriveTangent;
+
+	// Special cases for linear or constant segments
+	if (StartPoint.InterpMode == CIM_Linear)
+	{
+		return ((P1 - P0) * Scale3D).Size() * Param;
+	}
+	else if (StartPoint.InterpMode == CIM_Constant)
+	{
+		// Special case: constant interpolation acts like distance = 0 for all p in [0, 1[ but for p == 1, the distance returned is the linear distance between start and end
+		return Param == 1.f ? ((P1 - P0) * Scale3D).Size() : 0.0f;
+	}
+
+	// Cache the coefficients to be fed into the function to calculate the spline derivative at each sample point as they are constant.
+	const FVector Coeff1 = ((P0 - P1) * 2.0f + T0 + T1) * 3.0f;
+	const FVector Coeff2 = (P1 - P0) * 6.0f - T0 * 4.0f - T1 * 2.0f;
+	const FVector Coeff3 = T0;
+
+	const float HalfParam = Param * 0.5f;
+
+	float Length = 0.0f;
+	for (const auto& LegendreGaussCoefficient : LegendreGaussCoefficients)
+	{
+		// Calculate derivative at each Legendre-Gauss sample, and perform a weighted sum
+		const float Alpha = HalfParam * (1.0f + LegendreGaussCoefficient.Abscissa);
+		const FVector Derivative = ((Coeff1 * Alpha + Coeff2) * Alpha + Coeff3) * Scale3D;
+		Length += Derivative.Size() * LegendreGaussCoefficient.Weight;
+	}
+	Length *= HalfParam;
+
+	return Length;
+}
+
+float FLegacySpline::GetSplineLength() const
+{
+	const int32 NumPoints = ReparamTable.Points.Num();
+
+	// This is given by the input of the last entry in the remap table
+	if (NumPoints > 0)
+	{
+		return ReparamTable.Points.Last().InVal;
+	}
+
+	return 0.0f;
+}
+
+int32 FLegacySpline::GetNumControlPoints() const
+{
+	return PositionCurve.Points.Num();
+}
+
+void FLegacySpline::Reset()
+{
+	PositionCurve.Points.Reset();
+	RotationCurve.Points.Reset();
+	ScaleCurve.Points.Reset();
+}
+
+void FLegacySpline::ResetRotation()
+{
+	RotationCurve.Points.Reset(PositionCurve.Points.Num());
+
+	for (int32 Count = RotationCurve.Points.Num(); Count < PositionCurve.Points.Num(); Count++)
+	{
+		RotationCurve.Points.Emplace(static_cast<float>(Count), FQuat::Identity, FQuat::Identity, FQuat::Identity, CIM_CurveAuto);
+	}
+}
+
+void FLegacySpline::ResetScale()
+{
+	ScaleCurve.Points.Reset(PositionCurve.Points.Num());
+
+	for (int32 Count = ScaleCurve.Points.Num(); Count < PositionCurve.Points.Num(); Count++)
+	{
+		ScaleCurve.Points.Emplace(static_cast<float>(Count), FVector(1.0f), FVector::ZeroVector, FVector::ZeroVector, CIM_CurveAuto);
+	}
+}
+
+void FLegacySpline::UpdateSpline(const FSpline::FUpdateSplineParams& InParams)
+{
+	TRACE_CPUPROFILER_EVENT_SCOPE(FLegacySpline::UpdateSpline);
+
+	const int32 NumPoints = PositionCurve.Points.Num();
+	check(RotationCurve.Points.Num() == NumPoints && ScaleCurve.Points.Num() == NumPoints);
+
+#if DO_CHECK
+	// Ensure input keys are strictly ascending
+	for (int32 Index = 1; Index < NumPoints; Index++)
+	{
+		ensureAlways(PositionCurve.Points[Index - 1].InVal < PositionCurve.Points[Index].InVal);
+	}
+#endif
+
+	// Ensure splines' looping status matches with that of the spline component
+	if (InParams.bClosedLoop)
+	{
+		const float LastKey = PositionCurve.Points.Num() > 0 ? PositionCurve.Points.Last().InVal : 0.0f;
+		const float LoopKey = InParams.bLoopPositionOverride ? InParams.LoopPosition : LastKey + 1.0f;
+		PositionCurve.SetLoopKey(LoopKey);
+		RotationCurve.SetLoopKey(LoopKey);
+		ScaleCurve.SetLoopKey(LoopKey);
+	}
+	else
+	{
+		PositionCurve.ClearLoopKey();
+		RotationCurve.ClearLoopKey();
+		ScaleCurve.ClearLoopKey();
+	}
+
+	// Automatically set the tangents on any CurveAuto keys
+	PositionCurve.AutoSetTangents(0.0f, InParams.bStationaryEndpoints);
+	RotationCurve.AutoSetTangents(0.0f, InParams.bStationaryEndpoints);
+	ScaleCurve.AutoSetTangents(0.0f, InParams.bStationaryEndpoints);
+
+	// Now initialize the spline reparam table
+	const int32 NumSegments = PositionCurve.bIsLooped ? NumPoints : FMath::Max(0, NumPoints - 1);
+
+	// Start by clearing it
+	ReparamTable.Points.Reset(NumSegments * InParams.ReparamStepsPerSegment + 1);
+	float AccumulatedLength = 0.0f;
+	for (int32 SegmentIndex = 0; SegmentIndex < NumSegments; ++SegmentIndex)
+	{
+		for (int32 Step = 0; Step < InParams.ReparamStepsPerSegment; ++Step)
+		{
+			const float Param = static_cast<float>(Step) / InParams.ReparamStepsPerSegment;
+			const float SegmentLength = (Step == 0) ? 0.0f : GetSegmentLength(SegmentIndex, Param, InParams.Scale3D);
+
+			ReparamTable.Points.Emplace(SegmentLength + AccumulatedLength, SegmentIndex + Param, 0.0f, 0.0f, CIM_Linear);
+		}
+		AccumulatedLength += GetSegmentLength(SegmentIndex, 1.0f, InParams.Scale3D);
+	}
+	ReparamTable.Points.Emplace(AccumulatedLength, static_cast<float>(NumSegments), 0.0f, 0.0f, CIM_Linear);
+}
+
 PRAGMA_ENABLE_EXPERIMENTAL_WARNINGS
diff --git a/Engine/Source/Runtime/Engine/Private/Curves/Splines/TangentSpline.cpp b/Engine/Source/Runtime/Engine/Private/Curves/Splines/TangentSpline.cpp
deleted file mode 100644
