Deprecated LocalToWorld as it's unused and causes a race with the skeletal mesh updater and the scene transform loop.

#Synced-CL 49075818

diff --git a/Engine/Plugins/Runtime/HairStrands/Source/HairStrandsCore/Private/GroomGeometryCache.cpp b/Engine/Plugins/Runtime/HairStrands/Source/HairStrandsCore/Private/GroomGeometryCache.cpp
--- a/Engine/Plugins/Runtime/HairStrands/Source/HairStrandsCore/Private/GroomGeometryCache.cpp
+++ b/Engine/Plugins/Runtime/HairStrands/Source/HairStrandsCore/Private/GroomGeometryCache.cpp
@@ -29,7 +29,6 @@
  {
 	if (SceneProxy)
 	{
-		Out.LocalToWorld = FTransform(SceneProxy->GetLocalToWorld());
 		Out.LODIndex = 0;
 		if (bOutputTriangleData)
 		{
diff --git a/Engine/Plugins/Runtime/HairStrands/Source/HairStrandsCore/Private/GroomManager.cpp b/Engine/Plugins/Runtime/HairStrands/Source/HairStrandsCore/Private/GroomManager.cpp
--- a/Engine/Plugins/Runtime/HairStrands/Source/HairStrandsCore/Private/GroomManager.cpp
+++ b/Engine/Plugins/Runtime/HairStrands/Source/HairStrandsCore/Private/GroomManager.cpp
@@ -534,8 +534,6 @@ static void GetOrAllocateCachedGeometry(
 		return;
 	}
 	
-	Out.LocalToWorld = FTransform(Proxy->GetLocalToWorld());
-
 	const FSkeletalMeshObject* SkeletalMeshObject = Proxy->GetMeshObject();
 	if (SkeletalMeshObject == nullptr)
 	{
diff --git a/Engine/Source/Runtime/Engine/Private/SkeletalMeshSceneProxy.cpp b/Engine/Source/Runtime/Engine/Private/SkeletalMeshSceneProxy.cpp
--- a/Engine/Source/Runtime/Engine/Private/SkeletalMeshSceneProxy.cpp
+++ b/Engine/Source/Runtime/Engine/Private/SkeletalMeshSceneProxy.cpp
@@ -1637,7 +1637,6 @@ bool FSkeletalMeshSceneProxy::GetMaterialTextureScales(int32 LODIndex, int32 Sec
 void FSkeletalMeshSceneProxy::OnTransformChanged(FRHICommandListBase& RHICmdList)
 {
 	// OnTransformChanged is called on the following frame after FSkeletalMeshObject::Update(), thus omit '+ 1' to frame number.
-	MeshObject->SetTransform(GetLocalToWorld(), GetScene().GetFrameNumber());
 	MeshObject->RefreshClothingTransforms(GetLocalToWorld(), GetScene().GetFrameNumber());
 }
 
diff --git a/Engine/Source/Runtime/Engine/Private/SkeletalRenderGPUSkin.cpp b/Engine/Source/Runtime/Engine/Private/SkeletalRenderGPUSkin.cpp
--- a/Engine/Source/Runtime/Engine/Private/SkeletalRenderGPUSkin.cpp
+++ b/Engine/Source/Runtime/Engine/Private/SkeletalRenderGPUSkin.cpp
@@ -2044,23 +2044,6 @@ void FSkeletalMeshObjectGPUSkin::UpdateSceneExtensionHeader(FRHICommandListBase&
 	}
 }
 
-FMatrix FSkeletalMeshObjectGPUSkin::GetTransform() const
-{
-	if (DynamicData)
-	{
-		return DynamicData->LocalToWorld;
-	}
-	return FMatrix();
-}
-
-void FSkeletalMeshObjectGPUSkin::SetTransform(const FMatrix& InNewLocalToWorld, uint32 FrameNumber)
-{
-	if (DynamicData)
-	{
-		DynamicData->LocalToWorld = InNewLocalToWorld;
-	}
-}
-
 void FSkeletalMeshObjectGPUSkin::RefreshClothingTransforms(const FMatrix& InNewLocalToWorld, uint32 FrameNumber)
 {
 	if(DynamicData && DynamicData->ClothingSimData.Num() > 0)
@@ -2765,7 +2748,6 @@ bool FSkeletalMeshObjectGPUSkin::GetCachedGeometry(FRDGBuilder& GraphBuilder, FC
 	}
 
 	OutCachedGeometry.LODIndex = LodIndex;
-	OutCachedGeometry.LocalToWorld = FTransform(GetTransform());
 	return true;
 }
 
@@ -2795,7 +2777,6 @@ int32 FDynamicSkelMeshObjectDataGPUSkin::Reset()
 #if RHI_RAYTRACING
 	bAnySegmentUsesWorldPositionOffset = false;
 #endif
-	LocalToWorld = FMatrix::Identity;
 
 	PreviousReferenceToLocal.Empty();
 	PreviousReferenceToLocalForRayTracing.Empty();
@@ -2998,9 +2979,6 @@ void FDynamicSkelMeshObjectDataGPUSkin::InitDynamicSkelMeshObjectDataGPUSkin(
 		}
 	}
 
-	// Update local to world transform
-	LocalToWorld = InDynamicData.GetComponentTransform().ToMatrixWithScale();
-
 	// Update the clothing simulation mesh positions and normals
 	if (const IClothSimulationDataProvider* Provider = InDynamicData.GetClothSimulationDataProvider())
 	{
diff --git a/Engine/Source/Runtime/Engine/Private/SkeletalRenderGPUSkin.h b/Engine/Source/Runtime/Engine/Private/SkeletalRenderGPUSkin.h
--- a/Engine/Source/Runtime/Engine/Private/SkeletalRenderGPUSkin.h
+++ b/Engine/Source/Runtime/Engine/Private/SkeletalRenderGPUSkin.h
@@ -130,9 +130,6 @@ public:
     /** store transform of the cloth object **/
     FMatrix ClothObjectLocalToWorld;
 
-	/** store transform of the object **/
-	FMatrix LocalToWorld;
-
 	/** a weight factor to blend between simulated positions and skinned positions */	
 	float ClothBlendWeight;
 
@@ -493,8 +490,6 @@ public:
 		uint32 NumVertices = 0;
 	};
 
-	FMatrix GetTransform() const;
-	virtual void SetTransform(const FMatrix& InNewLocalToWorld, uint32 FrameNumber) override;
 	virtual void RefreshClothingTransforms(const FMatrix& InNewLocalToWorld, uint32 FrameNumber) override;
 	virtual void UpdateSkinWeightBuffer(const TArrayView<const FSkelMeshComponentLODInfo> InLODInfo) override;
 
diff --git a/Engine/Source/Runtime/Engine/Public/CachedGeometry.h b/Engine/Source/Runtime/Engine/Public/CachedGeometry.h
--- a/Engine/Source/Runtime/Engine/Public/CachedGeometry.h
+++ b/Engine/Source/Runtime/Engine/Public/CachedGeometry.h
@@ -42,5 +42,7 @@ struct FCachedGeometry
 
 	int32 LODIndex = -1;
 	TArray<Section> Sections;
+
+	UE_DEPRECATED(5.8, "Use scene proxy transform instead.")
 	FTransform LocalToWorld = FTransform::Identity;
 };
