Static mesh flag applies the auto set.
Static mesh always renders default material when static mesh flag not set.
#Synced-CL 49065038

diff --git a/Engine/Source/Runtime/Engine/Private/StaticMeshSceneProxy.cpp b/Engine/Source/Runtime/Engine/Private/StaticMeshSceneProxy.cpp
--- a/Engine/Source/Runtime/Engine/Private/StaticMeshSceneProxy.cpp
+++ b/Engine/Source/Runtime/Engine/Private/StaticMeshSceneProxy.cpp
@@ -2667,8 +2667,10 @@ FStaticMeshSceneProxy::FLODInfo::FLODInfo(const FStaticMeshSceneProxyDesc& InPro
 			}
 		}
 
-		// If there isn't an applied material, or if we need static lighting and it doesn't support it, fall back to the default material.
-		if(!SectionInfo.Material || (bHasSurfaceStaticLighting && !SectionInfo.Material->CheckMaterialUsage_Concurrent(MATUSAGE_StaticLighting)))
+		// If there isn't an applied material, or if the material doesn't have static mesh usage or if we need static lighting and the material doesn't support it, then fall back to the default material.
+		if (!SectionInfo.Material || 
+			!SectionInfo.Material->CheckMaterialUsage_Concurrent(MATUSAGE_StaticMesh) ||
+			(bHasSurfaceStaticLighting && !SectionInfo.Material->CheckMaterialUsage_Concurrent(MATUSAGE_StaticLighting)))
 		{
 			SectionInfo.Material = UMaterial::GetDefaultMaterial(MD_Surface);
 			if (SectionInfo.OverlayMaterial)
