FSpline Refactor: Prepare FSpline for non-tangent based splines.

#was-backout 49084292
#Synced-CL 49076822

diff --git a/Engine/Source/Runtime/Engine/Classes/Components/SplineComponent.h b/Engine/Source/Runtime/Engine/Classes/Components/SplineComponent.h
--- a/Engine/Source/Runtime/Engine/Classes/Components/SplineComponent.h
+++ b/Engine/Source/Runtime/Engine/Classes/Components/SplineComponent.h
@@ -76,15 +76,7 @@ struct FSplineCurves
 {
 	GENERATED_BODY()
 
-	static FSplineCurves FromSplineInterface(const FSpline& InProxy)
-	{
-		FSplineCurves NewCurve;
-		NewCurve.Position = InProxy.GetSplinePointsPosition();
-		NewCurve.Rotation = InProxy.GetSplinePointsRotation();
-		NewCurve.Scale = InProxy.GetSplinePointsScale();
-		NewCurve.UpdateSpline(NewCurve.Position.bIsLooped, false);
-		return NewCurve;
-	}
+	static FSplineCurves FromSplineInterface(const FSpline& InProxy);
 	
 	/** Spline built from position data. */
 	UPROPERTY()
@@ -230,9 +222,6 @@ protected:
 	
 public:
 
-	UFUNCTION(CallInEditor, Category="Spline")
-	void PopulateFromLegacy();
-
 	/** Return the name of the authoritative spline property. */
 	ENGINE_API static FName GetSplinePropertyName();
 
@@ -416,9 +405,6 @@ public:
 
 	UE_EXPERIMENTAL(5.7, "Spline display setting change notifications are experimental while the API is under active development.")
 	FSimpleMulticastDelegate& GetOnSplineDisplayChanged() { return OnSplineDisplayChanged; }
-	
-	/** Reset all spline data to match the current state of SplineCurves */
-	void SynchronizeSplines() { SetSpline(WarninglessSplineCurves()); }
 
 	ENGINE_API int32 GetVersion() const;
 	ENGINE_API const FInterpCurveVector& GetSplinePointsPosition() const;
@@ -966,8 +952,8 @@ public:
 	 * @return The index of the new attribute.
 	 */
 	UFUNCTION(BlueprintCallable, Category = "Spline|Metadata")
-	ENGINE_API int32 SetFloatPropertyAtSplineInputKey(float InKey, float Value, FName PropertyName) const;
-	
+	ENGINE_API int32 AddFloatPropertyAtSplineInputKey(float InKey, float Value, FName PropertyName);
+
 	/**
 	 * Get a float attribute parameter by index.
 	 * @param Index The index of the specified attribute.
@@ -1018,6 +1004,18 @@ public:
 	FInterpCurveQuat& GetSplinePointsRotation() { return WarninglessSplineCurves().Rotation; }
 	FInterpCurveVector& GetSplinePointsScale() { return WarninglessSplineCurves().Scale; }
 
+public: // Deprecated Interface
+
+	UE_DEPRECATED(5.8, "SynchronizeSplines is deprecated.")
+	void SynchronizeSplines() {}
+
+	UE_DEPRECATED(5.8, "PopulateFromLegacy is deprecated.")
+	void PopulateFromLegacy() {}
+
+	UE_DEPRECATED(5.8, "SetFloatPropertyAtSplineInputKey is deprecated, please use AddFloatPropertyAtSplineInputKey")
+	UFUNCTION(BlueprintCallable, Category = "Spline|Metadata", meta = (DeprecatedFunction, DeprecationMessage = "Please use AddFloatPropertyAtSplineInputKey"))
+	int32 SetFloatPropertyAtSplineInputKey(float InKey, float Value, FName PropertyName) const { return INDEX_NONE; }
+
 protected:
 
 	/** Returns the length of the specified spline segment up to the parametric value given */
@@ -1034,6 +1032,9 @@ protected:
 	
 private:
 	
+	/** Whether or not FSplineCurves is the authoritative data. If false, we assume Spline is implemented as FTangentSpline. */
+	ENGINE_API bool ShouldUseSplineCurves() const;
+
 	/** Returns an interval representing the full range of valid input keys, or an empty interval if the spline is empty. */
 	UE::Geometry::FInterval1f GetInputKeyRange() const;
 	
@@ -1052,9 +1053,9 @@ private:
 	}
 	
 	/** The dummy value used for queries when there are no point in a spline */
-	static ENGINE_API const FInterpCurvePointVector DummyPointPosition;
-	static ENGINE_API const FInterpCurvePointQuat DummyPointRotation;
-	static ENGINE_API const FInterpCurvePointVector DummyPointScale;
+	static ENGINE_API const FInterpCurveVector DummyPointPosition;
+	static ENGINE_API const FInterpCurveQuat DummyPointRotation;
+	static ENGINE_API const FInterpCurveVector DummyPointScale;
 
 	// Remove when appending _DEPRECATED to SplineCurves.
 	FSplineCurves& WarninglessSplineCurves()
@@ -1092,6 +1093,9 @@ PRAGMA_ENABLE_DEPRECATION_WARNINGS
 	// FSplineComponentVisualizer will access some private members when attempting to call NotifyPropertiesModified
 	friend class FSplineComponentVisualizer;
 
+	// FSplinePointDetails will access ShouldUseSplineCurves for enabling/disabling some editing functionality
+	friend class FSplinePointDetails;
+
 	// AFunctionalSplineTestBase accesses private methods for functional testing.
 	friend class AFunctionalSplineTestBase;
 };
diff --git a/Engine/Source/Runtime/Engine/Classes/Curves/Spline.h b/Engine/Source/Runtime/Engine/Classes/Curves/Spline.h
--- a/Engine/Source/Runtime/Engine/Classes/Curves/Spline.h
+++ b/Engine/Source/Runtime/Engine/Classes/Curves/Spline.h
@@ -2,19 +2,38 @@
 
 #pragma once
 
-#include "CoreMinimal.h"
-#include "Math/InterpCurve.h"
 #include "Splines/SplineInterfaces.h"
+#include "Curves/Splines/TangentSpline.h"	// only for FLegacyTangentSpline, move that to its own header.
 
 #include "Spline.generated.h"
 
-struct FSplinePoint;
-struct FSplineCurves;
-class FLegacySpline;
-class FNewSpline;
-
 struct UE_EXPERIMENTAL(5.7, "New spline APIs are experimental.") FSpline;
 
+namespace UE::Spline
+{
+enum class ESplineType : uint8
+{
+	Unimplemented = 0,
+	LegacyTangent = 1,
+	Tangent = 2,
+};
+
+template<ESplineType Impl>
+struct TSplineTypeMap;
+
+template<>
+struct TSplineTypeMap<ESplineType::LegacyTangent>
+{
+	using Type = FLegacyTangentSpline;
+};
+
+template<>
+struct TSplineTypeMap<ESplineType::Tangent>
+{
+	using Type = FTangentSpline;
+};
+} // UE::Spline
+
 /**
  * A general purpose, reflected spline.
  * The implementation can be configured at runtime.
@@ -25,186 +44,127 @@ struct FSpline
 	GENERATED_BODY()
 
 	ENGINE_API FSpline();
+	ENGINE_API explicit FSpline(UE::Spline::ESplineType InType);
+
 	ENGINE_API FSpline(const FSpline& Other);
 	ENGINE_API FSpline& operator=(const FSpline& Other);
-	ENGINE_API FSpline& operator=(const FSplineCurves& Other);
-	
-	/* Control Point Index Interface */
-	
-	/** Adds point by parameter. The new point is inserted after all points with parameter values less than OR EQUAL to the new point's parameter.
-	 * Nick would like to change this behavior in the future. See UE-250236. */
-	ENGINE_API void AddPoint(const FSplinePoint& InPoint);
-	/** Adds point by index. Provided parameter is ignored. The new point is inserted before all points with parameter values greater than the new point's parameter. */
-	ENGINE_API void InsertPoint(const FSplinePoint& InPoint, int32 Index);
-	ENGINE_API FSplinePoint GetPoint(const int32 Index) const;
-	ENGINE_API void RemovePoint(const int32 Index);
-	
-	ENGINE_API void SetLocation(int32 Index, const FVector& InLocation);
-	ENGINE_API FVector GetLocation(const int32 Index) const;
-	
-	ENGINE_API void SetInTangent(const int32 Index, const FVector& InTangent);
-	ENGINE_API FVector GetInTangent(const int32 Index) const;
-	
-	ENGINE_API void SetOutTangent(const int32 Index, const FVector& OutTangent);
-	ENGINE_API FVector GetOutTangent(const int32 Index) const;
 
-	ENGINE_API void SetRotation(int32 Index, const FQuat& InRotation);
-	ENGINE_API FQuat GetRotation(const int32 Index) const;
-	
-	ENGINE_API void SetScale(int32 Index, const FVector& InScale);
-	ENGINE_API FVector GetScale(const int32 Index) const;
-	
-	ENGINE_API void SetSplinePointType(int32 Index, EInterpCurveMode Type);
-	ENGINE_API EInterpCurveMode GetSplinePointType(int32 Index) const;
+	FSpline(FSpline&& Other) = default;
+	FSpline& operator=(FSpline&& Other) = default;
 
-	ENGINE_API float GetParameterAtIndex(int32 Index) const;
-	ENGINE_API float GetParameterAtDistance(float Distance) const;
-	ENGINE_API float GetDistanceAtParameter(float Parameter) const;
-
-	ENGINE_API FQuat GetOrientation(int32 Index) const;
-	ENGINE_API void SetOrientation(int32 Index, const FQuat& InOrientation);
-	
-	/* Parameter Interface */
+	ENGINE_API FSpline& operator=(const FSplineCurves& Other);
 
 	ENGINE_API FVector Evaluate(float Param) const;
+
 	ENGINE_API FVector EvaluateDerivative(float Param) const;
-	ENGINE_API FQuat EvaluateRotation(float Param) const;
-	ENGINE_API FVector EvaluateScale(float Param) const;
 
-	ENGINE_API FQuat GetOrientation(float Param) const;
-	
-	/* Attribute Interface */
-
-	// Non-Templated functions
-	bool SupportsAttributes() const { return IsNew(); }
-	ENGINE_API bool HasAttributeChannel(FName AttributeName) const;
-	ENGINE_API bool RemoveAttributeChannel(FName AttributeName) const;
-	ENGINE_API TArray<FName> GetFloatPropertyChannels() const;
-	ENGINE_API TArray<FName> GetVectorPropertyChannels() const;
-
-	template <typename AttrType> ENGINE_API float GetAttributeParameter(int32 Index, const FName& Name) const;
-	template <typename AttrType> ENGINE_API int32 SetAttributeParameter(int32 Index, float Parameter, const FName& Name);
-	template <typename AttrType> ENGINE_API int32 NumAttributeValues(FName AttributeName) const;
-	template <typename AttrType> ENGINE_API AttrType GetAttributeValue(int32 Index, const FName& Name) const;
-	template <typename AttrType> ENGINE_API void SetAttributeValue(int32 Index, const AttrType& Value, const FName& Name);
-	template <typename AttrType> ENGINE_API bool CreateAttributeChannel(FName AttributeName) const;
-	template <typename AttrType> ENGINE_API int32 AddAttributeValue(float Param, const AttrType& Value, FName AttributeName) const;
-	template <typename AttrType> ENGINE_API void RemoveAttributeValue(int32 Index, FName AttributeName);
-	template <typename AttrType> ENGINE_API AttrType EvaluateAttribute(float Param, FName AttributeName) const;
-
-	ENGINE_API float FindNearest(const FVector& InLocation, float& OutSquaredDist) const;
-	ENGINE_API float FindNearestOnSegment(const FVector& InLocation, int32 SegmentIndex, float& OutSquaredDist) const;
-	
-	/* Misc Interface */
-	
+	ENGINE_API UE::Geometry::FInterval1f GetParameterSpace() const;
+
+	ENGINE_API int32 GetNumberOfSegments() const;
+
+	/** Reset the spline to an empty spline. Preserves the current implementation. */
+	ENGINE_API void Reset();
+
+	/**
+	 * Fetch the underlying concrete spline implementation.
+	 *
+	 * Pointers returned by Get<Impl>() are invalidated by:
+	 * - Serialize(FArchive&) when Ar.IsLoading() is true
+	 * - ImportTextItem(...)
+	 * - operator=(...)
+	 * - SetCurrentImplementation(...)
+	 * - destruction of this FSpline
+	 *
+	 * @return A valid pointer if Impl == GetCurrentImplementation(), otherwise nullptr.
+	 */
+	template<UE::Spline::ESplineType Impl>
+	typename UE::Spline::TSplineTypeMap<Impl>::Type* Get();
+
+	/** A const version of Get<Impl> */
+	template<UE::Spline::ESplineType Impl>
+	const typename UE::Spline::TSplineTypeMap<Impl>::Type* Get() const;
+		
 	ENGINE_API bool operator==(const FSpline& Other) const;
 	ENGINE_API bool operator!=(const FSpline& Other) const;
-
-	ENGINE_API bool IsNearlyEqual(const FSpline& Other, const float RelativeTolerance, const float AbsoluteTolerance) const;
 	
 	friend FArchive& operator<<(FArchive& Ar, FSpline& Spline)
-	{
+	{ 
 		Spline.Serialize(Ar);
 		return Ar;
 	}
+
 	ENGINE_API bool Serialize(FArchive& Ar);
-	ENGINE_API void SerializeLoad(FArchive& Ar);
-	ENGINE_API void SerializeSave(FArchive& Ar) const;
 	ENGINE_API bool ExportTextItem(FString& ValueStr, FSpline const& DefaultValue, class UObject* Parent, int32 PortFlags, class UObject* ExportRootScope) const;
 	ENGINE_API bool ImportTextItem(const TCHAR*& Buffer, int32 PortFlags, UObject* Parent, FOutputDevice* ErrorText);
 
-	uint32 GetVersion() const { return Version; }
-	
-	ENGINE_API const FInterpCurveVector& GetSplinePointsPosition() const;
-	ENGINE_API const FInterpCurveQuat& GetSplinePointsRotation() const;
-	ENGINE_API const FInterpCurveVector& GetSplinePointsScale() const;
+	ENGINE_API UE::Spline::ESplineType GetCurrentImplementation() const;
 
-	/** Returns the length of the specified spline segment up to the parametric value given */
-	ENGINE_API float GetSegmentLength(const int32 Index, const float Param, const FVector& Scale3D = FVector(1.0f)) const;
+#if WITH_EDITOR
+	/** Returns the implementation of this spline when when most recently loaded. */
+	ENGINE_API UE::Spline::ESplineType GetPreviousImplementation() const;
+#endif
 
-	/** Returns total length along this spline */
-	ENGINE_API float GetSplineLength() const;
+	/** Returns the default implementation for new spline. */
+	ENGINE_API static UE::Spline::ESplineType GetDefaultImplementation();
 
-	/** Returns the total number of segments on this spline. */
-	ENGINE_API int32 GetNumSegments() const;
-	
-	/** Returns the total number of control points on this spline. */
-	ENGINE_API int32 GetNumControlPoints() const;
-	
-	/** Reset the spline to an empty spline. */
-	ENGINE_API void Reset();
-	
-	/** Reset the rotation attribute channel to default values. */
-	ENGINE_API void ResetRotation();
+private:
 	
-	/** Reset the scale attribute channel to default values. */
-	ENGINE_API void ResetScale();
+	/** Sets the current implementation and performs any necessary allocation. */
+	void SetCurrentImplementation(UE::Spline::ESplineType Implementation);
 
-	ENGINE_API void SetClosedLoop(bool bClosed);
-	ENGINE_API bool IsClosedLoop() const;
-	
-	struct FUpdateSplineParams
-	{
-		bool bClosedLoop = false;
-		bool bStationaryEndpoints = false;
-		int32 ReparamStepsPerSegment = 10;
-		bool bLoopPositionOverride = false;
-		float LoopPosition = 0.0f;
-		FVector Scale3D = FVector(1.0f);
-	};
-	
-	/** Update the spline's internal data according to the passed-in params. */
-	ENGINE_API void UpdateSpline(const FUpdateSplineParams& InParams);
+	void SerializeLoad(FArchive& Ar);
+	void SerializeSave(FArchive& Ar) const;
 
-	/** Update the spline's internal data according to the most recently used update params (or default params if never updated). */
-	ENGINE_API void UpdateSpline();
+private:
 
-	ENGINE_API TSharedPtr<UE::Geometry::Spline::TSplineInterface<FVector>> GetSplineInterface() const;
+	UE::Spline::ESplineType CurrentImplementation;
 
-private:
-	
-	static ENGINE_API const FInterpCurveVector PositionCurve;
-	static ENGINE_API const FInterpCurveQuat RotationCurve;
-	static ENGINE_API const FInterpCurveVector ScaleCurve;
-	
-	// Used for upgrade logic in spline component.
-	// Not ideal, but allows us to automatically populate the proxy
-	// at serialize time when we might otherwise not be able to.
-	friend class USplineComponent;
-	friend struct FPCGSplineStruct;
-	
 #if WITH_EDITOR
-	uint8 PreviousImplementation;
+	UE::Spline::ESplineType PreviousImplementation;
 #endif
-	uint8 CurrentImplementation;
-	
-	uint32 Version;
 
-	FUpdateSplineParams CachedUpdateSplineParams;
-	
-	// probably better implemented as a TSharedPtr<ISplineInterface> or something
-	struct {													// Invalid when CurrentImplementation is 0 or 3.
-		TSharedPtr<FLegacySpline> LegacyData;					// Valid when CurrentImplementation is 1.
-		TSharedPtr<FNewSpline> NewData;							// Valid when CurrentImplementation is 2.
-	};
-	
-	bool IsEnabled() const { return CurrentImplementation != 0; }
-	bool IsLegacy() const { return CurrentImplementation == 1; }
-	bool IsNew() const { return CurrentImplementation == 2; }
+	// Valid when CurrentImplementation is ESplineType::LegacyTangent
+	TUniquePtr<FLegacyTangentSpline> LegacySplineImpl;
 
-#if WITH_EDITOR
-	bool WasEnabled() const { return PreviousImplementation != 0; }
-	bool WasLegacy() const { return PreviousImplementation == 1; }
-	bool WasNew() const { return PreviousImplementation == 2; }
+	// Valid when CurrentImplementation is not ESplineType::Unimplemented and not ESplineType::LegacyTangent;
+	TUniquePtr<UE::Geometry::Spline::TSplineInterface<FVector>> SplineImpl;
+};
 
-	void OnSplineImplementationChanged();
-#endif
+template<UE::Spline::ESplineType Impl>
+typename UE::Spline::TSplineTypeMap<Impl>::Type* FSpline::Get()
+{
+	if (Impl != CurrentImplementation)
+	{
+		return nullptr;
+	}
 
-public:
+	if constexpr (Impl == UE::Spline::ESplineType::LegacyTangent)
+	{
+		return LegacySplineImpl.Get();
+	}
+	else
+	{
+		return static_cast<typename UE::Spline::TSplineTypeMap<Impl>::Type*>(SplineImpl.Get());
+	}
+}
 
-	// True if a given FSpline will actually be implemented when instantiated.
-	static bool IsEnabledStatic();
-};
+template<UE::Spline::ESplineType Impl>
+const typename UE::Spline::TSplineTypeMap<Impl>::Type* FSpline::Get() const
+{
+	if (Impl != CurrentImplementation)
+	{
+		return nullptr;
+	}
+
+	if constexpr (Impl == UE::Spline::ESplineType::LegacyTangent)
+	{
+		return LegacySplineImpl.Get();
+	}
+	else
+	{
+		return static_cast<const typename UE::Spline::TSplineTypeMap<Impl>::Type*>(SplineImpl.Get());
+	}
+}
 
 template<>
 struct TStructOpsTypeTraits<FSpline> : public TStructOpsTypeTraitsBase2<FSpline>
