SoundWaveProxyPlayer

Fix crash during rendering with speed changes, resulting in Resampling after originally not resampling.

During resampling, we use the SourceAudio buffer to hold on to remaining samples not used in the previous block. So make sure to reset SourceAudio buffer after generating source audio when *not resampling* so we don't send the wrong data in during resampling.

#rnx
#Synced-CL 49063521

diff --git a/Engine/Source/Runtime/Engine/Private/SoundWaveProxyPlayer.cpp b/Engine/Source/Runtime/Engine/Private/SoundWaveProxyPlayer.cpp
--- a/Engine/Source/Runtime/Engine/Private/SoundWaveProxyPlayer.cpp
+++ b/Engine/Source/Runtime/Engine/Private/SoundWaveProxyPlayer.cpp
@@ -310,6 +310,7 @@ void FSoundWaveProxyPlayer::RenderMultiChannelAudio(Audio::FMultichannelBufferVi
 			SourceAudio.SetNum(OutputNumFrames * NumChannels);
 			GenerateSourceAudio(SourceAudio);
 			ConvertDeinterleave->ProcessAudio(SourceAudio, OutMultiChannelAudio);
+			SourceAudio.SetNum(0, EAllowShrinking::No);
 		}
 	}
 	else
