Fix state/input being applied on frozen particle during resim and used in a normal frame (flag not cleared)

#Synced-CL 49101379

diff --git a/Engine/Source/Runtime/Engine/Private/PhysicsEngine/NetworkPhysicsComponent.cpp b/Engine/Source/Runtime/Engine/Private/PhysicsEngine/NetworkPhysicsComponent.cpp
--- a/Engine/Source/Runtime/Engine/Private/PhysicsEngine/NetworkPhysicsComponent.cpp
+++ b/Engine/Source/Runtime/Engine/Private/PhysicsEngine/NetworkPhysicsComponent.cpp
@@ -2180,10 +2180,21 @@ void FAsyncNetworkPhysicsComponent::OnPreProcessInputs_Internal(const int32 Phys
 #endif
 
 	{
+		Chaos::FReadPhysicsObjectInterface_Internal Interface = Chaos::FPhysicsObjectInternalInterface::GetRead();
+		
+		bool bCanApplyData = true;
+		if (bIsSolverResim)
+		{
+			if (Chaos::FPBDRigidParticleHandle* Particle = Interface.GetRigidParticle(PhysicsObject))
+			{
+				bCanApplyData = (Particle->ResimType() != Chaos::EResimType::FrozenDuringResim);
+			}
+		}
+		
 		// Apply replicated state on clients if we are resimulating or on simulated proxies if setting is enabled
 		const bool bApplySimProxyState = ComponentSettings.GetApplySimProxyStateAtRuntime() && !bIsServer && !IsLocallyControlled();
-
-		if ((bApplySimProxyState || bIsSolverResim) && StateHistory && StateData)
+		
+		if ((bApplySimProxyState || bIsSolverResim) && StateHistory && StateData && bCanApplyData)
 		{
 			FNetworkPhysicsPayload* PhysicsData = StateData.Get();
 			PhysicsData->LocalFrame = PhysicsStep;
@@ -2226,7 +2237,7 @@ void FAsyncNetworkPhysicsComponent::OnPreProcessInputs_Internal(const int32 Phys
 		const bool bApplySimProxyInput = ComponentSettings.GetApplySimProxyInputAtRuntime() && !bIsServer && !IsLocallyControlled();
 		const bool bApplyServerInput = bIsServer && !IsLocallyControlled();
 
-		if ((bApplyServerInput || bApplySimProxyInput || bIsSolverResim) && InputHistory && InputData)
+		if ((bApplyServerInput || bApplySimProxyInput || bIsSolverResim) && InputHistory && InputData && bCanApplyData)
 		{
 			FNetworkPhysicsPayload* PhysicsData = InputData.Get();
 			int32 NextExpectedInputFrame = PhysicsStep;
