[World] Support LevelInstance Enter/Exit edition callbacks

#Synced-CL 49104468

diff --git a/Engine/Source/Editor/LevelInstanceEditor/Private/LevelInstanceEditorModule.h b/Engine/Source/Editor/LevelInstanceEditor/Private/LevelInstanceEditorModule.h
--- a/Engine/Source/Editor/LevelInstanceEditor/Private/LevelInstanceEditorModule.h
+++ b/Engine/Source/Editor/LevelInstanceEditor/Private/LevelInstanceEditorModule.h
@@ -42,6 +42,9 @@ public:
 	DECLARE_DERIVED_EVENT(FLevelInstanceEditorModule, ILevelInstanceEditorModule::FTryExitEditorModeEvent, FTryExitEditorModeEvent);
 	virtual FTryExitEditorModeEvent& OnTryExitEditorMode() override { return TryExitEditorModeEvent; }
 
+	virtual FOnLevelInstanceEditModeChanged& OnLevelInstanceEditModeChanged() override { return OnLevelInstanceEditModeChangedDelegate; }
+	virtual void BroadcastOnLevelInstanceEditModeChanged(ILevelInstanceInterface* LevelInstance, ELevelInstanceEditMode NewMode) override { OnLevelInstanceEditModeChangedDelegate.Broadcast(LevelInstance, NewMode); }
+
 	virtual bool IsEditInPlaceStreamingEnabled() const override;
 	virtual bool IsSubSelectionEnabled() const override;
 
@@ -132,5 +135,7 @@ private:
 	FExitEditorModeEvent ExitEditorModeEvent;
 	FTryExitEditorModeEvent TryExitEditorModeEvent;
 
+	FOnLevelInstanceEditModeChanged OnLevelInstanceEditModeChangedDelegate;
+
 	TScriptInterface<IInputBehaviorSource> DefaultBehaviorSource;
 };
diff --git a/Engine/Source/Runtime/Engine/Private/LevelInstance/LevelInstanceSubsystem.cpp b/Engine/Source/Runtime/Engine/Private/LevelInstance/LevelInstanceSubsystem.cpp
--- a/Engine/Source/Runtime/Engine/Private/LevelInstance/LevelInstanceSubsystem.cpp
+++ b/Engine/Source/Runtime/Engine/Private/LevelInstance/LevelInstanceSubsystem.cpp
@@ -485,6 +485,12 @@ void ULevelInstanceSubsystem::OnWorldCleanup(UWorld* InWorld, bool bSessionEnded
 			InWorld->RemoveStreamingLevels(StreamingLevels);
 		}
 
+		if (LevelInstanceEdit.IsValid())
+		{
+			ILevelInstanceInterface* LI = LevelInstanceEdit->GetLevelInstance();
+			BroadcastOnLevelInstanceEditModeChanged(LI, ILevelInstanceEditorModule::ELevelInstanceEditMode::None);
+		}
+
 		LoadedLevelInstances.Empty();
 		LevelInstanceEdit.Reset();
 		PropertyOverrideEdit.Reset();
@@ -562,6 +568,8 @@ void ULevelInstanceSubsystem::RegisterLoadedLevelStreamingLevelInstanceEditor(UL
 		check(!LevelInstanceEdit.IsValid());
 		ILevelInstanceInterface* LevelInstance = LevelStreaming->GetLevelInstance();
 		LevelInstanceEdit = MakeUnique<FLevelInstanceEdit>(LevelStreaming, LevelInstance);
+
+		BroadcastOnLevelInstanceEditModeChanged(LevelInstance, ILevelInstanceEditorModule::ELevelInstanceEditMode::Original);
 	}
 }
 
@@ -1851,6 +1859,19 @@ void ULevelInstanceSubsystem::BreakLevelInstance_Impl(ILevelInstanceInterface* L
 	}
 }
 
+void ULevelInstanceSubsystem::BroadcastOnLevelInstanceEditModeChanged(ILevelInstanceInterface* LevelInstance, ILevelInstanceEditorModule::ELevelInstanceEditMode NewMode)
+{
+	if (LevelInstance == nullptr)
+	{
+		return;
+	}
+
+	if (ILevelInstanceEditorModule* EditorModule = FModuleManager::Get().GetModulePtr<ILevelInstanceEditorModule>("LevelInstanceEditor"))
+	{
+		EditorModule->BroadcastOnLevelInstanceEditModeChanged(LevelInstance, NewMode);
+	}
+}
+
 bool ULevelInstanceSubsystem::LevelInstanceHasLevelScriptBlueprint(const ILevelInstanceInterface* LevelInstance) const
 {
 	if (LevelInstance)
@@ -2676,6 +2697,7 @@ bool ULevelInstanceSubsystem::CommitLevelInstanceInternal(TUniquePtr<FLevelInsta
 	GEditor->SelectNone(true, true);
 
 	// Remove from streaming level...
+	BroadcastOnLevelInstanceEditModeChanged(LevelInstance, ILevelInstanceEditorModule::ELevelInstanceEditMode::None);
 	InLevelInstanceEdit.Reset();
 
 	if (bChangesCommitted)
@@ -3046,7 +3068,10 @@ void ULevelInstanceSubsystem::EditLevelInstancePropertyOverrides(ILevelInstanceI
 	UActorEditorContextSubsystem::Get()->PushContext();
 
 	// Load Edit LevelInstance level
-	ULevelStreamingLevelInstanceEditorPropertyOverride::Load(GetLevelInstance(LevelInstanceID));
+	LevelInstance = GetLevelInstance(LevelInstanceID);
+	BroadcastOnLevelInstanceEditModeChanged(LevelInstance, ILevelInstanceEditorModule::ELevelInstanceEditMode::Override);
+
+	ULevelStreamingLevelInstanceEditorPropertyOverride::Load(LevelInstance);
 
 	// Select Context Actor
 	SelectActorFromActorName(GetLevelInstance(LevelInstanceID), ActorNameToSelect);
@@ -3126,6 +3151,8 @@ bool ULevelInstanceSubsystem::CommitLevelInstancePropertyOverridesInternal(TUniq
 		UpdateLevelInstancesFromPropertyOverrideAsset(PreviousOverrideAsset, LevelInstanceWithOverrides->GetPropertyOverrideAsset());
 	}
 
+	BroadcastOnLevelInstanceEditModeChanged(LevelInstanceWithOverrides, ILevelInstanceEditorModule::ELevelInstanceEditMode::None);
+
 	BlockOnLoading();
 
 	if (AActor* LevelInstanceToSelect = Cast<AActor>(GetLevelInstance(LevelInstanceWithOverridesID)))
diff --git a/Engine/Source/Runtime/Engine/Public/LevelInstance/ILevelInstanceEditorModule.h b/Engine/Source/Runtime/Engine/Public/LevelInstance/ILevelInstanceEditorModule.h
--- a/Engine/Source/Runtime/Engine/Public/LevelInstance/ILevelInstanceEditorModule.h
+++ b/Engine/Source/Runtime/Engine/Public/LevelInstance/ILevelInstanceEditorModule.h
@@ -29,6 +29,17 @@ public:
 	DECLARE_EVENT(ILevelInstanceEditorModule, FTryExitEditorModeEvent);
 	virtual FTryExitEditorModeEvent& OnTryExitEditorMode() = 0;
 
+	/** Broadcasts when the state of a LI being edited changes */
+	enum class ELevelInstanceEditMode : uint8
+	{
+		None,     /** The level instance is not being edited */
+		Original, /** The level instance's original is being edited */
+		Override  /** A level instance override is being edited */
+	};
+	DECLARE_MULTICAST_DELEGATE_TwoParams(FOnLevelInstanceEditModeChanged, ILevelInstanceInterface* /*LevelInstance*/, ELevelInstanceEditMode /*NewMode*/);
+	virtual FOnLevelInstanceEditModeChanged& OnLevelInstanceEditModeChanged() = 0;
+	virtual void BroadcastOnLevelInstanceEditModeChanged(ILevelInstanceInterface* LevelInstance, ELevelInstanceEditMode NewMode) = 0;
+
 	virtual bool IsEditInPlaceStreamingEnabled() const = 0;
 	virtual bool IsSubSelectionEnabled() const = 0;
 
diff --git a/Engine/Source/Runtime/Engine/Public/LevelInstance/LevelInstanceSubsystem.h b/Engine/Source/Runtime/Engine/Public/LevelInstance/LevelInstanceSubsystem.h
--- a/Engine/Source/Runtime/Engine/Public/LevelInstance/LevelInstanceSubsystem.h
+++ b/Engine/Source/Runtime/Engine/Public/LevelInstance/LevelInstanceSubsystem.h
@@ -7,6 +7,7 @@
 #include "UObject/NameTypes.h"
 #include "UObject/GCObject.h"
 #include "Folder.h"
+#include "LevelInstance/ILevelInstanceEditorModule.h"
 #include "LevelInstance/LevelInstanceTypes.h"
 #include "WorldPartition/Filter/WorldPartitionActorFilter.h"
 #include "WorldPartition/WorldPartitionHandle.h"
@@ -237,6 +238,8 @@ private:
 
 	ENGINE_API void BreakLevelInstance_Impl(ILevelInstanceInterface* LevelInstance, uint32 Levels, TArray<AActor*>& OutMovedActors, ELevelInstanceBreakFlags Flags);
 
+	void BroadcastOnLevelInstanceEditModeChanged(ILevelInstanceInterface* LevelInstance, ILevelInstanceEditorModule::ELevelInstanceEditMode NewMode);
+
 	static ENGINE_API bool ShouldIgnoreDirtyPackage(UPackage* DirtyPackage, const UWorld* EditingWorld);
 
 	class FLevelInstanceEdit
