Add support for not setting the optional Face and Vertex remap arrays on the Chaos implicit triangle mesh.

The Face map will not be set if the existing UBodySetup::bSupportUVsAndFaceRemap is false, as it will not be used in this case.
The Vertex map will not be set if a new UBodySetup::bSupportVertexRemap is false (defaults to true), as there was previously no way to control this per-asset.

#Synced-CL 48982928

diff --git a/Engine/Source/Runtime/Engine/Classes/PhysicsEngine/BodySetup.h b/Engine/Source/Runtime/Engine/Classes/PhysicsEngine/BodySetup.h
--- a/Engine/Source/Runtime/Engine/Classes/PhysicsEngine/BodySetup.h
+++ b/Engine/Source/Runtime/Engine/Classes/PhysicsEngine/BodySetup.h
@@ -115,6 +115,9 @@ struct FCookBodySetupInfo
 	/** Whether to support face remap, needed for physical material masks */
 	bool bSupportFaceRemap;
 
+	/** Whether to support vertex remap, needed for skeletal mesh per-poly collision and any other FBodyInstance::UpdateTriMeshVertices calls */
+	bool bSupportVertexRemap;
+
 	/** Error generating cook info for trimesh*/
 	bool bTriMeshError;
 };
@@ -194,6 +197,14 @@ class UBodySetup : public UBodySetupCore
 	UPROPERTY()
 	uint8 bSupportUVsAndFaceRemap:1;
 
+	/**
+	 * If true, the physics triangle mesh will store the vertex remap table. This is needed
+	 * to support skeletal mesh per-poly collision, and any other FBodyInstance::UpdateTriMeshVertices calls.
+	 * Both this flag and the p.Chaos.TriMeshPerPolySupport cvar must be true for the vertex remap table to be set.
+	 */
+	UPROPERTY()
+	uint8 bSupportVertexRemap:1;
+
 	/** Flag used to know if we have created the physics convex and tri meshes from the cooked data yet */
 	uint8 bCreatedPhysicsMeshes:1;
 
diff --git a/Engine/Source/Runtime/Engine/Private/PhysicsEngine/BodySetup.cpp b/Engine/Source/Runtime/Engine/Private/PhysicsEngine/BodySetup.cpp
--- a/Engine/Source/Runtime/Engine/Private/PhysicsEngine/BodySetup.cpp
+++ b/Engine/Source/Runtime/Engine/Private/PhysicsEngine/BodySetup.cpp
@@ -62,6 +62,7 @@ FCookBodySetupInfo::FCookBodySetupInfo() :
 	bCookTriMesh(false),
 	bSupportUVFromHitResults(false),
 	bSupportFaceRemap(false),
+	bSupportVertexRemap(true),
 	bTriMeshError(false)
 {
 }
@@ -139,6 +140,7 @@ UBodySetup::UBodySetup(const FObjectInitializer& ObjectInitializer)
 	bGenerateMirroredCollision = true;
 	bGenerateNonMirroredCollision = true;
 	bSupportUVsAndFaceRemap = false;
+	bSupportVertexRemap = true;
 	DefaultInstance.SetObjectType(ECC_PhysicsBody);
 #if WITH_EDITORONLY_DATA
 	BuildScale_DEPRECATED = 1.0f;
@@ -326,6 +328,7 @@ void UBodySetup::GetCookInfo(FCookBodySetupInfo& OutCookInfo, EPhysXMeshCookFlag
 			OutCookInfo.TriMeshCookFlags = CookFlags;
 
 			OutCookInfo.bSupportFaceRemap = bSupportUVsAndFaceRemap;
+			OutCookInfo.bSupportVertexRemap = bSupportVertexRemap;
 		}
 		else
 		{
@@ -1523,6 +1526,11 @@ void UBodySetup::GetGeometryDDCKey(FString& OutString) const
 	{
 		OutString += FString("_1");
 	}
+
+	if (!bSupportVertexRemap)
+	{
+		OutString += FString("_V");
+	}
 }
 
 void UBodySetup::PostInitProperties()
diff --git a/Engine/Source/Runtime/Engine/Private/PhysicsEngine/Experimental/ChaosCooking.cpp b/Engine/Source/Runtime/Engine/Private/PhysicsEngine/Experimental/ChaosCooking.cpp
--- a/Engine/Source/Runtime/Engine/Private/PhysicsEngine/Experimental/ChaosCooking.cpp
+++ b/Engine/Source/Runtime/Engine/Private/PhysicsEngine/Experimental/ChaosCooking.cpp
@@ -323,8 +323,17 @@ namespace Chaos
 					}
 				}
 
-				TUniquePtr<TArray<int32>> OutFaceRemapPtr = MakeUnique<TArray<int32>>(OutFaceRemap);
-				TUniquePtr<TArray<int32>> OutVertexRemapPtr = Chaos::TriMeshPerPolySupport ? MakeUnique<TArray<int32>>(OutVertexRemap) : nullptr;
+				
+				TUniquePtr<TArray<int32>> OutFaceRemapPtr = InParams.bSupportFaceRemap ? MakeUnique<TArray<int32>>(OutFaceRemap) : nullptr;
+				TUniquePtr<TArray<int32>> OutVertexRemapPtr = (InParams.bSupportVertexRemap && Chaos::TriMeshPerPolySupport) ? MakeUnique<TArray<int32>>(OutVertexRemap) : nullptr;
+				if (!OutFaceRemapPtr)
+				{
+					OutFaceRemap.Empty();
+				}
+				if (!OutVertexRemapPtr)
+				{
+					OutVertexRemap.Empty();
+				}
 				Chaos::FTriangleMeshImplicitObject::FBuildSettings MeshBuildSettings;
 				MeshBuildSettings.PerFaceBounds = InParams.TriangleMeshDesc.bReducedMemoryRepresentation ? Chaos::FTrimeshBVH::EPerFaceBoundsRepresentation::Byte : Chaos::FTrimeshBVH::EPerFaceBoundsRepresentation::VectorizedFloat;
 				constexpr bool bCullsBackFaceRaycast = false;
@@ -333,7 +342,10 @@ namespace Chaos
 				OutTriangleMeshes.Emplace(TriangleMesh);
 
 				// Propagate remapped indices from the FTriangleMeshImplicitObject back to the remap array
-				CopyUpdatedFaceRemapFromTriangleMesh(*TriangleMesh, OutFaceRemap);
+				if (InParams.bSupportFaceRemap)
+				{
+					CopyUpdatedFaceRemapFromTriangleMesh(*TriangleMesh, OutFaceRemap);
+				}
 			};
 
 			if(FinalVerts.Num() < TNumericLimits<uint16>::Max())
