Fix default material PSO precaching compilation issues due to missing vertex element on SM5 without manual vertex fetch:
- Take Instance Origin, Transform and InstanceLightmapAndShadowMapUVBias stream into account when generating default vertex stream inputs for instance static mesh vertex factories
- Make sure all streams are added for GeometryCacheVertexFactory (not just position)

#Synced-CL 49066199

diff --git a/Engine/Source/Runtime/Engine/Private/GeometryCacheVertexFactory.cpp b/Engine/Source/Runtime/Engine/Private/GeometryCacheVertexFactory.cpp
--- a/Engine/Source/Runtime/Engine/Private/GeometryCacheVertexFactory.cpp
+++ b/Engine/Source/Runtime/Engine/Private/GeometryCacheVertexFactory.cpp
@@ -242,46 +242,42 @@ void FGeometryCacheVertexVertexFactory::GetPSOPrecacheVertexFetchElements(EVerte
 	if (VertexInputStreamType == EVertexInputStreamType::PositionAndNormalOnly)
 	{
 		// 2-axis TangentBasis components in a single buffer, hence *2u
-		OutElements.Add(FVertexElement(1, 0, VET_PackedNormal, 1, sizeof(FPackedNormal)*2u, false));
+		OutElements.Add(FVertexElement(1, 0, VET_PackedNormal, 1, sizeof(FPackedNormal) * 2u, false));
 	}
 
-	if (UseGPUScene(GMaxRHIShaderPlatform, GMaxRHIFeatureLevel) 
-		&& !PlatformGPUSceneUsesUniformBufferView(GMaxRHIShaderPlatform))
+	switch (VertexInputStreamType)
 	{
-		switch (VertexInputStreamType)
+		case EVertexInputStreamType::Default:
 		{
-			case EVertexInputStreamType::Default:
-			{			
-				// Make sure all required elements are added because manual vertex fetch is not supported (see GetVertexElements above)
-				uint8 TangentBasisAttributes[2] = { 1, 2 };
-				for (int32 AxisIndex = 0; AxisIndex < 2; AxisIndex++)
-				{
-					OutElements.Add(FVertexElement(TangentBasisAttributes[AxisIndex], 0, VET_PackedNormal, TangentBasisAttributes[AxisIndex], sizeof(FPackedNormal), false));
-				}
-				OutElements.Add(FVertexElement(3, 0, VET_Color, 3, sizeof(FColor), false));
-				OutElements.Add(FVertexElement(4, 0, VET_Float3, 4, sizeof(FVector3f), false));
-				for (int32 CoordinateIndex = 0; CoordinateIndex < MAX_STATIC_TEXCOORDS / 2; CoordinateIndex++)
-				{
-					OutElements.Add(FVertexElement(5 + CoordinateIndex, 0, VET_Half4, 5 + CoordinateIndex, sizeof(FVector2f), false));
-				}
-
-				OutElements.Add(FVertexElement(9, 0, VET_UInt, 13, sizeof(uint32), true));
-
-				break;
-			}
-			case EVertexInputStreamType::PositionOnly:
+			// Make sure all required elements are added because manual vertex fetch is not supported (see GetVertexElements above)
+			uint8 TangentBasisAttributes[2] = { 1, 2 };
+			for (int32 AxisIndex = 0; AxisIndex < 2; AxisIndex++)
 			{
-				OutElements.Add(FVertexElement(1, 0, VET_UInt, 1, sizeof(uint32), true));
-				break;
+				OutElements.Add(FVertexElement(TangentBasisAttributes[AxisIndex], 0, VET_PackedNormal, TangentBasisAttributes[AxisIndex], sizeof(FPackedNormal), false));
 			}
-			case EVertexInputStreamType::PositionAndNormalOnly:
+			OutElements.Add(FVertexElement(3, 0, VET_Color, 3, sizeof(FColor), false));
+			OutElements.Add(FVertexElement(4, 0, VET_Float3, 4, sizeof(FVector3f), false));
+			for (int32 CoordinateIndex = 0; CoordinateIndex < MAX_STATIC_TEXCOORDS / 2; CoordinateIndex++)
 			{
-				OutElements.Add(FVertexElement(2, 0, VET_UInt, 2, sizeof(uint32), true));
-				break;
+				OutElements.Add(FVertexElement(5 + CoordinateIndex, 0, VET_Half4, 5 + CoordinateIndex, sizeof(FVector2f), false));
 			}
-			default:
-				checkNoEntry();
+
+			OutElements.Add(FVertexElement(9, 0, VET_UInt, 13, sizeof(uint32), true));
+
+			break;
+		}
+		case EVertexInputStreamType::PositionOnly:
+		{
+			OutElements.Add(FVertexElement(1, 0, VET_UInt, 1, sizeof(uint32), true));
+			break;
+		}
+		case EVertexInputStreamType::PositionAndNormalOnly:
+		{
+			OutElements.Add(FVertexElement(2, 0, VET_UInt, 2, sizeof(uint32), true));
+			break;
 		}
+		default:
+			checkNoEntry();
 	}
 }
 
diff --git a/Engine/Source/Runtime/Engine/Private/InstancedStaticMesh.cpp b/Engine/Source/Runtime/Engine/Private/InstancedStaticMesh.cpp
--- a/Engine/Source/Runtime/Engine/Private/InstancedStaticMesh.cpp
+++ b/Engine/Source/Runtime/Engine/Private/InstancedStaticMesh.cpp
@@ -731,6 +731,25 @@ void FInstancedStaticMeshVertexFactory::GetPSOPrecacheVertexFetchElements(EVerte
 {
 	// Fallback to local vertex factory because manual vertex fetch is supported
 	FLocalVertexFactory::GetPSOPrecacheVertexFetchElements(VertexInputStreamType, Elements);
+	
+	const bool bCanUseGPUScene = UseGPUScene(GMaxRHIShaderPlatform, GMaxRHIFeatureLevel);
+	const bool bMobileUsesGPUScene = MobileSupportsGPUScene();
+	
+	if (GMaxRHIFeatureLevel > ERHIFeatureLevel::ES3_1 || !bMobileUsesGPUScene)
+	{
+		uint8 StreamIndex = Elements.Last().StreamIndex + 1;
+
+		// InstanceOriginComponent
+		Elements.Add(FVertexElement(StreamIndex++, 0, VET_Float4, 8, 0, false));
+
+		// InstanceTransformComponent
+		Elements.Add(FVertexElement(StreamIndex, 0, VET_Float4, 9, 0, false));
+		Elements.Add(FVertexElement(StreamIndex, 8, VET_Float4, 10, 0, false));
+		Elements.Add(FVertexElement(StreamIndex++, 16, VET_Float4, 11, 0, false));
+
+		// InstanceLightmapAndShadowMapUVBiasComponent
+		Elements.Add(FVertexElement(StreamIndex++, 16, VET_Short4N, 12, 0, false));
+	}
 }
 
 /**
