Skeleton: handle case-sensitive only renaming
- FReferenceSkeleton supports case-sensitive renaming but the USkeletonModifier and the USkeleton did not

#rnx

#Synced-CL 49084977

diff --git a/Engine/Plugins/Runtime/MeshModelingToolset/Source/SkeletalMeshModifiers/Private/SkeletonModifier.cpp b/Engine/Plugins/Runtime/MeshModelingToolset/Source/SkeletalMeshModifiers/Private/SkeletonModifier.cpp
--- a/Engine/Plugins/Runtime/MeshModelingToolset/Source/SkeletalMeshModifiers/Private/SkeletonModifier.cpp
+++ b/Engine/Plugins/Runtime/MeshModelingToolset/Source/SkeletalMeshModifiers/Private/SkeletonModifier.cpp
@@ -723,6 +723,12 @@ ESkeletalMeshModificationType USkeletonModifier::PreCommitSkeletalMesh()
 					{
 						EnumAddFlags(Modifications, ESkeletalMeshModificationType::TransformChanged);
 					}
+					
+					// case-sensitive change only
+					if (OtherBoneInfo.Name == NewBoneName && !OtherBoneInfo.Name.IsEqual(NewBoneName, ENameCase::CaseSensitive))
+					{
+						EnumAddFlags(Modifications, ESkeletalMeshModificationType::BonesRenamed);
+					}
 				}
 			}
 		}
diff --git a/Engine/Source/Runtime/Engine/Private/Animation/Skeleton.cpp b/Engine/Source/Runtime/Engine/Private/Animation/Skeleton.cpp
--- a/Engine/Source/Runtime/Engine/Private/Animation/Skeleton.cpp
+++ b/Engine/Source/Runtime/Engine/Private/Animation/Skeleton.cpp
@@ -994,9 +994,13 @@ bool USkeleton::MergeBonesToBoneTree(const USkinnedAsset* InSkinnedAsset, const
 		if( IsCompatibleMesh(InSkinnedAsset, bDoParentChainCheck) )
 		{
 			TRACE_CPUPROFILER_EVENT_SCOPE(USkeleton::MergeBonesToBoneTree::CompatibleBranch);
+			
+			const FReferenceSkeleton& InputRefSkeleton = InSkinnedAsset->GetRefSkeleton();
+			const TArray<FMeshBoneInfo>& InputBoneInfos = InputRefSkeleton.GetRefBoneInfo();
+			
 			// Exclude bones who do not have a parent.
 			TArray<int32> FilteredRequiredBones;
-			FAnimationRuntime::ExcludeBonesWithNoParents(RequiredRefBones, InSkinnedAsset->GetRefSkeleton(), FilteredRequiredBones);
+			FAnimationRuntime::ExcludeBonesWithNoParents(RequiredRefBones, InputRefSkeleton, FilteredRequiredBones);
 
 			// Two modifier passes: add bones that dont already exist, then set bone parents that have changed
 			{
@@ -1005,12 +1009,13 @@ bool USkeleton::MergeBonesToBoneTree(const USkinnedAsset* InSkinnedAsset, const
 				// Check for bone's existence
 				for (const int32 MeshBoneIndex : FilteredRequiredBones)
 				{
-					const int32 SkeletonBoneIndex = ReferenceSkeleton.FindRawBoneIndex(InSkinnedAsset->GetRefSkeleton().GetBoneName(MeshBoneIndex));
+					const FName& InputBoneName = InputRefSkeleton.GetBoneName(MeshBoneIndex);
+					const int32 SkeletonBoneIndex = ReferenceSkeleton.FindRawBoneIndex(InputBoneName);
 					
 					// Bone doesn't already exist. Add it.
 					if( SkeletonBoneIndex == INDEX_NONE )
 					{
-						FMeshBoneInfo NewMeshBoneInfo = InSkinnedAsset->GetRefSkeleton().GetRefBoneInfo()[MeshBoneIndex];
+						FMeshBoneInfo NewMeshBoneInfo = InputBoneInfos[MeshBoneIndex];
 						// Fix up ParentIndex for our new Skeleton.
 						if( ReferenceSkeleton.GetRawBoneNum() == 0 )
 						{
@@ -1018,13 +1023,22 @@ bool USkeleton::MergeBonesToBoneTree(const USkinnedAsset* InSkinnedAsset, const
 						}
 						else
 						{
-							NewMeshBoneInfo.ParentIndex = ReferenceSkeleton.FindRawBoneIndex(InSkinnedAsset->GetRefSkeleton().GetBoneName(InSkinnedAsset->GetRefSkeleton().GetParentIndex(MeshBoneIndex)));
+							NewMeshBoneInfo.ParentIndex = ReferenceSkeleton.FindRawBoneIndex(InputRefSkeleton.GetBoneName(InputRefSkeleton.GetParentIndex(MeshBoneIndex)));
 						}
 
-						RefSkelModifier.Add(NewMeshBoneInfo, InSkinnedAsset->GetRefSkeleton().GetRefBonePose()[MeshBoneIndex]);
+						RefSkelModifier.Add(NewMeshBoneInfo, InputRefSkeleton.GetRefBonePose()[MeshBoneIndex]);
 						BoneTree.AddDefaulted(1);
 						bShouldHandleHierarchyChange = true;
 					}
+					else
+					{
+						// handle case-sensitive only change
+						const FName& BoneName = ReferenceSkeleton.GetBoneName(SkeletonBoneIndex);
+						if (InputBoneName == BoneName && !InputBoneName.IsEqual(BoneName, ENameCase::CaseSensitive))
+						{
+							RefSkelModifier.Rename(BoneName, InputBoneName);
+						}
+					}
 				}
 			}
 
@@ -1038,15 +1052,15 @@ bool USkeleton::MergeBonesToBoneTree(const USkinnedAsset* InSkinnedAsset, const
 				// Check for different parents
 				for (const int32 MeshBoneIndex : FilteredRequiredBones)
 				{
-					const FName BoneName = InSkinnedAsset->GetRefSkeleton().GetBoneName(MeshBoneIndex);
+					const FName BoneName = InputRefSkeleton.GetBoneName(MeshBoneIndex);
 					const int32 SkeletonBoneIndex = ReferenceSkeleton.FindRawBoneIndex(BoneName);
 
 					if(SkeletonBoneIndex != INDEX_NONE)
 					{
 						// Bone exists, check if it is in the same place in the hierarchy
-						const int32 MeshParentIndex = InSkinnedAsset->GetRefSkeleton().GetParentIndex(MeshBoneIndex);
+						const int32 MeshParentIndex = InputRefSkeleton.GetParentIndex(MeshBoneIndex);
 						const int32 SkeletonParentIndex = ReferenceSkeleton.GetRawParentIndex(SkeletonBoneIndex);
-						const FName MeshParentName = MeshParentIndex != INDEX_NONE ? InSkinnedAsset->GetRefSkeleton().GetBoneName(MeshParentIndex) : NAME_None;
+						const FName MeshParentName = MeshParentIndex != INDEX_NONE ? InputRefSkeleton.GetBoneName(MeshParentIndex) : NAME_None;
 						const FName SkeletonParentName = SkeletonParentIndex != INDEX_NONE ? ReferenceSkeleton.GetRawRefBoneInfo()[SkeletonParentIndex].Name : NAME_None;
 
 						if(MeshParentName != SkeletonParentName)
