Add permutation count to the recompileshaders listtypes command

Looks like this:

```
ShaderTypeName, Filename, PermutationCount
FBufferTextureMappedCopyCS, /NNEDenoiserShaders/NNEDenoiserShadersMappedCopy.usf, 125
FTextureBufferMappedCopyCS, /NNEDenoiserShaders/NNEDenoiserShadersMappedCopy.usf, 125
FStochasticLightingTileClassificationMarkCS, /Engine/Private/StochasticLighting/StochasticLightingTileClassification.usf, 3072
```

#Synced-CL 49108023

diff --git a/Engine/Source/Runtime/Engine/Private/ShaderCompiler/ShaderCompiler.cpp b/Engine/Source/Runtime/Engine/Private/ShaderCompiler/ShaderCompiler.cpp
--- a/Engine/Source/Runtime/Engine/Private/ShaderCompiler/ShaderCompiler.cpp
+++ b/Engine/Source/Runtime/Engine/Private/ShaderCompiler/ShaderCompiler.cpp
@@ -4355,10 +4355,10 @@ namespace
 {
 	void ListAllShaderTypes()
 	{
-		UE_LOG(LogShaderCompilers, Display, TEXT("ShaderTypeName, Filename"));
+		UE_LOG(LogShaderCompilers, Display, TEXT("ShaderTypeName, Filename, PermutationCount"));
 		for (TLinkedList<FShaderType*>::TIterator It(FShaderType::GetTypeList()); It; It.Next())
 		{
-			UE_LOG(LogShaderCompilers, Display, TEXT("%s, %s "), (*It)->GetName(), (*It)->GetShaderFilename());
+			UE_LOG(LogShaderCompilers, Display, TEXT("%s, %s, %d"), (*It)->GetName(), (*It)->GetShaderFilename(), (*It)->GetPermutationCount());
 		}
 
 		UE_LOG(LogShaderCompilers, Display, TEXT("VertexFactoryTypeName, Filename"));
