Fixed bug where an estimate of zero was considered okay. Default is -1 when no estimate supplied.
Fixed bug where skinned asset build estimate could be truncated to zero due to int32/int64 issue.
Static meshes now include include an estimate for assets without bulk data, these are typically editor/engine assets which have some overhead, not zero.
Added Insights trace regions to asset compilation to more easily see total async work in progress.
#rnx

diff --git a/Engine/Source/Runtime/Engine/Private/AssetCompilingManager.cpp b/Engine/Source/Runtime/Engine/Private/AssetCompilingManager.cpp
--- a/Engine/Source/Runtime/Engine/Private/AssetCompilingManager.cpp
+++ b/Engine/Source/Runtime/Engine/Private/AssetCompilingManager.cpp
@@ -188,9 +188,9 @@ public:
 				LargestDefecit = CurrentDefecit;
 				UE_LOG(LogAsyncCompilation, VeryVerbose, TEXT(
 					"Memory available for asset compilation less than exptected: (Hard Limit %.3f MB - SystemAvailableMemory %.3f MB) < MemoryAvailableAgainstHardLimit = %.3f MB"),
-					   HardMemoryLimit / (1024 * 1024.f),
-					   SystemAvailableMemory / (1024 * 1024.f),
-					   MemoryAvailableAgainstHardLimit / (1024 * 1024.f));
+					   HardMemoryLimit / (1024.f * 1024.f),
+					   SystemAvailableMemory / (1024.f * 1024.f),
+					   MemoryAvailableAgainstHardLimit / (1024.f * 1024.f));
 			}
 		}
 
@@ -252,11 +252,11 @@ public:
 			FShaderCompileMemoryUsage SCWMemoryUsage = GShaderCompilingManager->GetExternalMemoryUsage();
 			FShaderCompilerMemoryLimits SCWMemoryLimits = GetShaderCompilerCurrentMemoryLimits();		
 
-			return SCWMemoryUsage.VirtualMemory <= static_cast<uint64>(SCWMemoryLimits.CurrentLimit * 1024 * 1024);
+			return SCWMemoryUsage.VirtualMemory <= static_cast<uint64>(SCWMemoryLimits.CurrentLimit * 1024ull * 1024ull);
 		}
 		bool IsMemoryAvailableToLoan() const
 		{
-			int64 AvailableAssetCompilationMemory = GetAvailableMemory() / (1024 * 1024);	
+			int64 AvailableAssetCompilationMemory = GetAvailableMemory() / (1024ull * 1024ull);	
 			return AvailableAssetCompilationMemory >= MemoryLoanThreshold;
 		}
 		bool ShouldRetrieveFromSCW() const 
@@ -281,6 +281,7 @@ public:
 		}
 		// Used to make sure this value remains constant between OnScheduled and OnUnscheduled
 		int64 RequiredMemory = 0;
+		uint64 TraceRegionId = 0;
 		MemorySharingInfo& MemoryInfo;
 	protected:
 		void DoThreadedWork() override
@@ -360,8 +361,8 @@ public:
 		TRACE_COUNTER_SET(AsyncCompilationConcurrency, GetCurrentConcurrency());
 
 		FShaderCompilerMemoryLimits SCWMemoryLimits = GetShaderCompilerCurrentMemoryLimits();
-		TRACE_COUNTER_SET(AsyncCompilationSCWCurrentLimit, SCWMemoryLimits.CurrentLimit * 1024 * 1024);
-		TRACE_COUNTER_SET(AsyncCompilationSCWBaseLimit, SCWMemoryLimits.BaseLimit * 1024 * 1024);
+		TRACE_COUNTER_SET(AsyncCompilationSCWCurrentLimit, SCWMemoryLimits.CurrentLimit * 1024ull * 1024ull);
+		TRACE_COUNTER_SET(AsyncCompilationSCWBaseLimit, SCWMemoryLimits.BaseLimit * 1024ull * 1024ull);
 	}
 
 	void OnScheduled(const IQueuedWork* InWork) override
@@ -376,10 +377,14 @@ public:
 			DebugName = TEXT("No DebugName");
 		}
 
-		UE_LOG(LogAsyncCompilation, Verbose, TEXT("OnScheduled [%s] Work->RequiredMemory = %.3f MB TotalEstimatedMemory = %.3f MB"),
+		Work->TraceRegionId = TRACE_BEGIN_REGION_WITH_ID(*FString::Printf(TEXT("%s (%.3f MiB)"), DebugName, Work->RequiredMemory/(1024*1024.f)), TEXT("AsyncCompilation"));
+
+		UE_LOG(LogAsyncCompilation, Verbose, TEXT("OnScheduled [%s] Work->RequiredMemory = %.3f MB TotalEstimatedMemory = %.3f MB (TraceId %d)"),
 				DebugName,
-				Work->RequiredMemory/(1024*1024.f),
-				TotalEstimatedMemory/(1024*1024.f));
+				Work->RequiredMemory/(1024.f*1024.f),
+				TotalEstimatedMemory/(1024.f*1024.f),
+				Work->TraceRegionId
+			);
 
 		UpdateCounters();
 	}
@@ -389,11 +394,18 @@ public:
 		FMemoryBoundScheduledWork* Work = (FMemoryBoundScheduledWork*)InWork;
 		TotalEstimatedMemory -= Work->RequiredMemory; // atomic
 		check(TotalEstimatedMemory >= 0);
+
+		if (Work->TraceRegionId)
+		{
+			TRACE_END_REGION_WITH_ID(Work->TraceRegionId);
+		}
 		
 		// don't call Work->GetDebugName , it's gone now
-		UE_LOG(LogAsyncCompilation, Verbose, TEXT("OnUnscheduled Work->RequiredMemory = %.3f MB TotalEstimatedMemory = %.3f MB"),
-				Work->RequiredMemory/(1024*1024.f),
-				TotalEstimatedMemory/(1024*1024.f));
+		UE_LOG(LogAsyncCompilation, Verbose, TEXT("OnUnscheduled Work->RequiredMemory = %.3f MB TotalEstimatedMemory = %.3f MB (TraceId %d)"),
+				Work->RequiredMemory/(1024.f*1024.f),
+				TotalEstimatedMemory/(1024.f*1024.f),
+				Work->TraceRegionId
+			);
 
 		UpdateCounters();
 		
@@ -471,10 +483,10 @@ public:
 							"BEWARE: AssetCompile memory estimate is greater than the hard memory limit, but we're running it [{TaskName}] anyway! "
 							"RequiredMemory = {TotalEstimatedMemory} MiB + {RequiredMemory} MiB, AvailableAssetCompilationMemory = {AvailableAssetCompilationMemory} MiB, HardMemoryLimit = {HardMemoryLimit} MiB",
 							("TaskName", DebugName),
-							("RequiredMemory", FString::SanitizeFloat(NewRequiredMemory / (1024 * 1024.f), 3)),
-							("TotalEstimatedMemory", FString::SanitizeFloat(TotalEstimatedMemory / (1024 * 1024.f), 3)),
-							("AvailableAssetCompilationMemory", FString::SanitizeFloat(AvailableAssetCompilationMemory / (1024 * 1024.f), 3)),
-							("HardMemoryLimit", FString::SanitizeFloat(HardMemoryLimit / (1024 * 1024.f), 3))
+							("RequiredMemory", FString::SanitizeFloat(NewRequiredMemory / (1024.f * 1024.f), 3)),
+							("TotalEstimatedMemory", FString::SanitizeFloat(TotalEstimatedMemory / (1024.f * 1024.f), 3)),
+							("AvailableAssetCompilationMemory", FString::SanitizeFloat(AvailableAssetCompilationMemory / (1024.f * 1024.f), 3)),
+							("HardMemoryLimit", FString::SanitizeFloat(HardMemoryLimit / (1024.f * 1024.f), 3))
 						);
 					}
 					else
@@ -483,9 +495,9 @@ public:
 							"BEWARE: AssetCompile memory estimate is greater than available, but we're running it [{TaskName}] anyway! "
 							"RequiredMemory = {TotalEstimatedMemory} MiB + {RequiredMemory} MiB, MemoryLimit = {AvailableAssetCompilationMemory} MiB ",
 							("TaskName", DebugName),
-							("RequiredMemory", FString::SanitizeFloat(NewRequiredMemory / (1024 * 1024.f), 3)),
-							("TotalEstimatedMemory", FString::SanitizeFloat(TotalEstimatedMemory / (1024 * 1024.f), 3)),
-							("AvailableAssetCompilationMemory", FString::SanitizeFloat(AvailableAssetCompilationMemory / (1024 * 1024.f), 3))
+							("RequiredMemory", FString::SanitizeFloat(NewRequiredMemory / (1024.f * 1024.f), 3)),
+							("TotalEstimatedMemory", FString::SanitizeFloat(TotalEstimatedMemory / (1024.f * 1024.f), 3)),
+							("AvailableAssetCompilationMemory", FString::SanitizeFloat(AvailableAssetCompilationMemory / (1024.f * 1024.f), 3))
 						);
 					}
 				}
@@ -525,7 +537,7 @@ public:
 			DebugName = TEXT("No DebugName");
 		}
 
-		if ( RequiredMemory < 0 )
+		if ( RequiredMemory <= 0 )
 		{
 			// log the types that are failing to set required memory
 			UE_LOG(LogAsyncCompilation, Display, TEXT("AddQueuedWork [%s] did not supply a RequiredMemory estimate. This is a bug, all asset types should provide an estimate. When no estimate is provided asset compilation is slowed as the queue has to operate very conservatively."),DebugName);
diff --git a/Engine/Source/Runtime/Engine/Private/SkinnedAsset.cpp b/Engine/Source/Runtime/Engine/Private/SkinnedAsset.cpp
--- a/Engine/Source/Runtime/Engine/Private/SkinnedAsset.cpp
+++ b/Engine/Source/Runtime/Engine/Private/SkinnedAsset.cpp
@@ -39,7 +39,7 @@ DEFINE_LOG_CATEGORY(LogSkinnedAsset)
 
 static TAutoConsoleVariable<int32> CVarMemoryForSkinnedAssetCompile (
 	TEXT("Memory.MemoryForSkinnedAssetCompile"),
-	1024 * 4, // 4GiB
+	4608,		//	4.5 GiB
 	TEXT("Memory in MiB set aside for a SkinnedAsset compile job\n"),
 	ECVF_Default);
 
@@ -141,7 +141,7 @@ void USkinnedAsset::PostLoad()
 		EQueuedWorkPriority BasePriority = FSkinnedAssetCompilingManager::Get().GetBasePriority(this);
 
 		AsyncTask = MakeUnique<FSkinnedAssetAsyncBuildTask>(this, MoveTemp(Context));
-		int64 RequiredMemory = 1024 * 1024 * CVarMemoryForSkinnedAssetCompile.GetValueOnAnyThread();
+		const int64 RequiredMemory = 1024ull * 1024ull * CVarMemoryForSkinnedAssetCompile.GetValueOnAnyThread();
 		if (!HasAnyDependenciesCompiling())
 		{
 			AsyncTask->StartBackgroundTask(ThreadPool, BasePriority, EQueuedWorkFlags::DoNotRunInsideBusyWait, RequiredMemory, TEXT("SkinnedAsset"));
diff --git a/Engine/Source/Runtime/Engine/Private/StaticMesh.cpp b/Engine/Source/Runtime/Engine/Private/StaticMesh.cpp
--- a/Engine/Source/Runtime/Engine/Private/StaticMesh.cpp
+++ b/Engine/Source/Runtime/Engine/Private/StaticMesh.cpp
@@ -7578,6 +7578,14 @@ int64 UStaticMesh::GetBuildRequiredMemoryEstimate() const
 		}
 	}
 
+	if (BiggestBulkDataSize == 0)
+	{
+		// When there is no bulk data size, the mesh isn't going to be built from source data.
+		// Rather than zero, returns a small fixed size of 4MiB to account for the overhead.
+		// For example, some engine meshes fall under this.
+		return 1024 * 1024 * 4;
+	}
+
 	// Rough estimate of the memory that should be required to build that mesh.
 	// Used -ddc=cold on the command line and opened some big meshes from the content browser one by one while
 	// monitoring the editor commit memory spike generated by the compilation.
