[ActorFolders] Make intermediate folder APIs internal; code outside the engine should only use FActorFolders.

#Synced-CL 48871118

diff --git a/Engine/Source/Editor/UnrealEd/Private/WorldFolders.cpp b/Engine/Source/Editor/UnrealEd/Private/WorldFolders.cpp
--- a/Engine/Source/Editor/UnrealEd/Private/WorldFolders.cpp
+++ b/Engine/Source/Editor/UnrealEd/Private/WorldFolders.cpp
@@ -21,6 +21,27 @@
 
 DEFINE_LOG_CATEGORY(LogWorldFolders);
 
+void FActorPlacementFolder::Reset()
+{
+	Path = NAME_None;
+	RootObjectPtr.Reset();
+	ActorFolderGuid.Invalidate();
+}
+
+FActorPlacementFolder& FActorPlacementFolder::operator= (const FFolder& InOtherFolder)
+{
+	Path = InOtherFolder.GetPath();
+	RootObjectPtr = InOtherFolder.GetRootObjectPtr();
+	ActorFolderGuid = InOtherFolder.GetActorFolderGuid();
+	return *this;
+}
+
+FFolder FActorPlacementFolder::GetFolder() const
+{
+	FFolder::FRootObject RootObject = FFolder::FRootObject(RootObjectPtr.Get());
+	return ActorFolderGuid.IsValid() ? FFolder(RootObject, ActorFolderGuid) : FFolder(RootObject, Path);
+}
+
 void UWorldFolders::Initialize(UWorld* InWorld)
 {
 	TRACE_CPUPROFILER_EVENT_SCOPE(UWorldFolders::Initialize);
diff --git a/Engine/Source/Editor/UnrealEd/Public/WorldFolders.h b/Engine/Source/Editor/UnrealEd/Public/WorldFolders.h
--- a/Engine/Source/Editor/UnrealEd/Public/WorldFolders.h
+++ b/Engine/Source/Editor/UnrealEd/Public/WorldFolders.h
@@ -34,68 +34,59 @@ struct FActorPlacementFolder
 {
 	GENERATED_BODY()
 
+	UE_INTERNAL
 	UPROPERTY()
 	FName Path;
 
+	UE_INTERNAL
 	UPROPERTY()
 	TWeakObjectPtr<UObject> RootObjectPtr;
 	
+	UE_INTERNAL
 	UPROPERTY()
 	FGuid ActorFolderGuid;
 
-	void Reset()
-	{
-		Path = NAME_None;
-		RootObjectPtr.Reset();
-		ActorFolderGuid.Invalidate();
-	}
+	FActorPlacementFolder& operator= (const FFolder& InOtherFolder);
 
-	FActorPlacementFolder& operator = (const FFolder& InOtherFolder)
-	{
-		Path = InOtherFolder.GetPath();
-		RootObjectPtr = InOtherFolder.GetRootObjectPtr();
-		ActorFolderGuid = InOtherFolder.GetActorFolderGuid();
-		return *this;
-	}
-
-	FFolder GetFolder() const
-	{
-		FFolder::FRootObject RootObject = FFolder::FRootObject(RootObjectPtr.Get());
-		return ActorFolderGuid.IsValid() ? FFolder(RootObject, ActorFolderGuid) : FFolder(RootObject, Path);
-	}
+	UE_INTERNAL void Reset();
+	UE_INTERNAL FFolder GetFolder() const;
 };
 
-/** Per-World Actor Folders UObject (used to support undo/redo reliably) */
+/**
+ * Per-World Actor Folders UObject (used to support undo/redo reliably)
+ * This class is reserved for low-level Engine use, therefore relevant APIs are marked UE_INTERNAL.
+ * The public, non-engine API for Folder operations is through the FActorFolders struct.
+ */
 UCLASS(MinimalAPI)
 class UWorldFolders : public UObject
 {
 public:
 	GENERATED_BODY()
 
-	UNREALED_API void Initialize(UWorld* InWorld);
-
-	UNREALED_API void RebuildList();
-	UNREALED_API bool AddFolder(const FFolder& InFolder);
-	UNREALED_API bool RemoveFolder(const FFolder& InFolder, bool bShouldDeleteFolder = false);
-	UNREALED_API bool RenameFolder(const FFolder& InOldFolder, const FFolder& InNewFolder);
-	UNREALED_API bool IsFolderExpanded(const FFolder& InFolder) const;
-	UNREALED_API bool SetIsFolderExpanded(const FFolder& InFolder, bool bIsExpanded);
-	UNREALED_API FFolder GetActorEditorContextFolder(bool bMustMatchCurrentLevel) const;
-	UNREALED_API bool SetActorEditorContextFolder(const FFolder& InFolder);
-	UNREALED_API void PushActorEditorContext(bool bDuplicateContext = false);
-	UNREALED_API void PopActorEditorContext();
-	UNREALED_API bool ContainsFolder(const FFolder& InFolder) const;
-	UNREALED_API void ForEachFolder(TFunctionRef<bool(const FFolder&)> Operation);
-	UNREALED_API void ForEachFolderWithRootObject(const FFolder::FRootObject& InFolderRootObject, TFunctionRef<bool(const FFolder&)> Operation);
-	UNREALED_API void SaveState();
-	UNREALED_API UWorld* GetWorld() const;
+	UE_INTERNAL UNREALED_API void Initialize(UWorld* InWorld);
+
+	UE_INTERNAL UNREALED_API void RebuildList();
+	UE_INTERNAL UNREALED_API bool AddFolder(const FFolder& InFolder);
+	UE_INTERNAL UNREALED_API bool RemoveFolder(const FFolder& InFolder, bool bShouldDeleteFolder = false);
+	UE_INTERNAL UNREALED_API bool RenameFolder(const FFolder& InOldFolder, const FFolder& InNewFolder);
+	UE_INTERNAL UNREALED_API bool IsFolderExpanded(const FFolder& InFolder) const;
+	UE_INTERNAL UNREALED_API bool SetIsFolderExpanded(const FFolder& InFolder, bool bIsExpanded);
+	UE_INTERNAL UNREALED_API FFolder GetActorEditorContextFolder(bool bMustMatchCurrentLevel) const;
+	UE_INTERNAL UNREALED_API bool SetActorEditorContextFolder(const FFolder& InFolder);
+	UE_INTERNAL UNREALED_API void PushActorEditorContext(bool bDuplicateContext = false);
+	UE_INTERNAL UNREALED_API void PopActorEditorContext();
+	UE_INTERNAL UNREALED_API bool ContainsFolder(const FFolder& InFolder) const;
+	UE_INTERNAL UNREALED_API void ForEachFolder(TFunctionRef<bool(const FFolder&)> Operation);
+	UE_INTERNAL UNREALED_API void ForEachFolderWithRootObject(const FFolder::FRootObject& InFolderRootObject, TFunctionRef<bool(const FFolder&)> Operation);
+	UE_INTERNAL UNREALED_API void SaveState();
+	UE_INTERNAL UNREALED_API UWorld* GetWorld() const;
 
 	//~ Begin UObject
-	UNREALED_API virtual void Serialize(FArchive& Ar) override;
+	UE_INTERNAL UNREALED_API virtual void Serialize(FArchive& Ar) override;
 	//~ End UObject
 
 	//~ Begin Deprecated
-	UNREALED_API FActorFolderProps* GetFolderProperties(const FFolder& InFolder);
+	UE_INTERNAL UNREALED_API FActorFolderProps* GetFolderProperties(const FFolder& InFolder);
 	//~ End Deprecated
 
 private:
diff --git a/Engine/Source/Editor/UnrealEd/Public/WorldFoldersImplementation.h b/Engine/Source/Editor/UnrealEd/Public/WorldFoldersImplementation.h
--- a/Engine/Source/Editor/UnrealEd/Public/WorldFoldersImplementation.h
+++ b/Engine/Source/Editor/UnrealEd/Public/WorldFoldersImplementation.h
@@ -8,16 +8,20 @@ class UWorld;
 class UWorldFolders;
 struct FFolder;
 
+/**
+ * This class is reserved for low-level Engine use, therefore relevant APIs are marked UE_INTERNAL.
+ * The public, non-engine API for Folder operations is through the FActorFolders struct.
+ */
 class FWorldFoldersImplementation
 {
 public:
 	FWorldFoldersImplementation(UWorldFolders& Owner);
 	virtual ~FWorldFoldersImplementation() = default;
-	virtual bool AddFolder(const FFolder& InFolder) { return true; }
-	virtual bool RemoveFolder(const FFolder& InFolder, bool bShouldDeleteFolder) { return true; }
-	virtual bool RenameFolder(const FFolder& InOldFolder, const FFolder& InNewFolder) { return true; }
-	virtual bool ContainsFolder(const FFolder& InFolder) const;
-	UWorld* GetWorld() const;
+	UE_INTERNAL virtual bool AddFolder(const FFolder& InFolder) { return true; }
+	UE_INTERNAL virtual bool RemoveFolder(const FFolder& InFolder, bool bShouldDeleteFolder) { return true; }
+	UE_INTERNAL virtual bool RenameFolder(const FFolder& InOldFolder, const FFolder& InNewFolder) { return true; }
+	UE_INTERNAL virtual bool ContainsFolder(const FFolder& InFolder) const;
+	UE_INTERNAL UWorld* GetWorld() const;
 
 protected:
 	UWorldFolders& Owner;
diff --git a/Engine/Source/Editor/UnrealEd/Public/WorldPersistentFolders.h b/Engine/Source/Editor/UnrealEd/Public/WorldPersistentFolders.h
--- a/Engine/Source/Editor/UnrealEd/Public/WorldPersistentFolders.h
+++ b/Engine/Source/Editor/UnrealEd/Public/WorldPersistentFolders.h
@@ -8,7 +8,11 @@
 
 class ULevel;
 
-// Class handling a list of actor folder objects of a world
+/**
+ * Class handling a list of actor folder objects of a world
+ * This class is reserved for low-level Engine use, therefore relevant APIs are marked UE_INTERNAL.
+ * The public, non-engine API for Folder operations is through the FActorFolders struct.
+ */
 class FWorldPersistentFolders : public FWorldFoldersImplementation
 {
 	typedef FWorldFoldersImplementation Super;
@@ -17,13 +21,13 @@ public:
 	FWorldPersistentFolders(UWorldFolders& InWorldFolders);
 	virtual ~FWorldPersistentFolders() = default;
 
-	UNREALED_API static UActorFolder* GetActorFolder(const FFolder& InFolder, UWorld* InWorld, bool bInAllowCreate = false);
+	UE_INTERNAL UNREALED_API static UActorFolder* GetActorFolder(const FFolder& InFolder, UWorld* InWorld, bool bInAllowCreate = false);
 
 	//~ Begin FWorldFoldersImplementation
-	virtual bool AddFolder(const FFolder& InFolder) override;
-	virtual bool RemoveFolder(const FFolder& InFolder, bool bShouldDeleteFolder) override;
-	virtual bool RenameFolder(const FFolder& InOldFolder, const FFolder& InNewFolder) override;
-	virtual bool ContainsFolder(const FFolder& InFolder) const override;
+	UE_INTERNAL virtual bool AddFolder(const FFolder& InFolder) override;
+	UE_INTERNAL virtual bool RemoveFolder(const FFolder& InFolder, bool bShouldDeleteFolder) override;
+	UE_INTERNAL virtual bool RenameFolder(const FFolder& InOldFolder, const FFolder& InNewFolder) override;
+	UE_INTERNAL virtual bool ContainsFolder(const FFolder& InFolder) const override;
 	//~ End FWorldFoldersImplementation
 
 private:
diff --git a/Engine/Source/Editor/UnrealEd/Public/WorldTransientFolders.h b/Engine/Source/Editor/UnrealEd/Public/WorldTransientFolders.h
--- a/Engine/Source/Editor/UnrealEd/Public/WorldTransientFolders.h
+++ b/Engine/Source/Editor/UnrealEd/Public/WorldTransientFolders.h
@@ -4,7 +4,11 @@
 
 #include "WorldFoldersImplementation.h"
 
-// Class handling a list of transient actor folders of a world
+/**
+ * Class handling a list of transient actor folders of a world
+ * This class is reserved for low-level Engine use, therefore relevant APIs are marked UE_INTERNAL.
+ * The public, non-engine API for Folder operations is through the FActorFolders struct.
+ */
 class FWorldTransientFolders : public FWorldFoldersImplementation
 {
 	typedef FWorldFoldersImplementation Super;
@@ -14,6 +18,6 @@ public:
 	~FWorldTransientFolders() = default;
 
 	//~ Begin FWorldFoldersImplementation
-	virtual bool RenameFolder(const FFolder& InOldFolder, const FFolder& InNewFolder) override;
+	UE_INTERNAL virtual bool RenameFolder(const FFolder& InOldFolder, const FFolder& InNewFolder) override;
 	//~ End FWorldFoldersImplementation
 };
diff --git a/Engine/Source/Runtime/Engine/Private/ActorFolder.cpp b/Engine/Source/Runtime/Engine/Private/ActorFolder.cpp
--- a/Engine/Source/Runtime/Engine/Private/ActorFolder.cpp
+++ b/Engine/Source/Runtime/Engine/Private/ActorFolder.cpp
@@ -69,13 +69,6 @@ namespace ActorFolder
 	static const FName NAME_OuterPackageName("OuterPackageName");
 };
 
-void UActorFolder::GetAssetRegistryTags(TArray<FAssetRegistryTag>& OutTags) const
-{
-	PRAGMA_DISABLE_DEPRECATION_WARNINGS;
-	Super::GetAssetRegistryTags(OutTags);
-	PRAGMA_ENABLE_DEPRECATION_WARNINGS;
-}
-
 void UActorFolder::GetAssetRegistryTags(FAssetRegistryTagsContext Context) const
 {
 	Super::GetAssetRegistryTags(Context);
diff --git a/Engine/Source/Runtime/Engine/Public/ActorFolder.h b/Engine/Source/Runtime/Engine/Public/ActorFolder.h
--- a/Engine/Source/Runtime/Engine/Public/ActorFolder.h
+++ b/Engine/Source/Runtime/Engine/Public/ActorFolder.h
@@ -12,6 +12,10 @@
 class ULevel;
 struct FActorFolderDesc;
 
+/**
+ * This class is reserved for low-level Engine use, therefore relevant APIs are marked UE_INTERNAL.
+ * The public, non-engine API for Folder operations is through the FActorFolders struct.
+ */
 UCLASS(Within = Level, MinimalAPI)
 class UActorFolder : public UObject
 {
@@ -20,55 +24,55 @@ class UActorFolder : public UObject
 #if WITH_EDITOR
 public:
 	// Helper that creates a new UActorFolder
-	static ENGINE_API UActorFolder* Create(ULevel* InOuter, const FString& InFolderLabel, UActorFolder* InParent);
+	UE_INTERNAL static ENGINE_API UActorFolder* Create(ULevel* InOuter, const FString& InFolderLabel, UActorFolder* InParent);
 
-	UE_DEPRECATED(5.4, "This function is deprecated.")
-	static EPackageFlags GetExternalPackageFlags() { return (PKG_EditorOnly | PKG_ContainsMapData | PKG_NewlyCreated); }
-	
 	// Returns actor folder info stored in its package
-	static ENGINE_API FActorFolderDesc GetAssetRegistryInfoFromPackage(FName PackageName);
+	UE_INTERNAL static ENGINE_API FActorFolderDesc GetAssetRegistryInfoFromPackage(FName PackageName);
 	
 	//~ Begin UObject
-	ENGINE_API virtual bool IsAsset() const override;
-	ENGINE_API virtual void GetAssetRegistryTags(FAssetRegistryTagsContext Context) const override;
-	UE_DEPRECATED(5.4, "Implement the version that takes FAssetRegistryTagsContext instead.")
-	ENGINE_API virtual void GetAssetRegistryTags(TArray<FAssetRegistryTag>& OutTags) const override;
-	ENGINE_API virtual void PostLoad() override;
+	UE_INTERNAL ENGINE_API virtual bool IsAsset() const override;
+	UE_INTERNAL ENGINE_API virtual void GetAssetRegistryTags(FAssetRegistryTagsContext Context) const override;
+	UE_INTERNAL ENGINE_API virtual void PostLoad() override;
 	//~ End UObject
 
-	ENGINE_API void SetParent(UActorFolder* InParent);
-	ENGINE_API UActorFolder* GetParent(bool bSkipDeleted = true) const;
+	UE_INTERNAL ENGINE_API void SetParent(UActorFolder* InParent);
+	UE_INTERNAL ENGINE_API UActorFolder* GetParent(bool bSkipDeleted = true) const;
 
-	ENGINE_API void SetLabel(const FString& InFolderLabel);
-	const FString& GetLabel() const { return FolderLabel; }
+	UE_INTERNAL ENGINE_API void SetLabel(const FString& InFolderLabel);
+	UE_INTERNAL const FString& GetLabel() const { return FolderLabel; }
 
 	// Control whether the folder is initially expanded or not
-	ENGINE_API void SetIsInitiallyExpanded(bool bInFolderInitiallyExpanded);
-	bool IsInitiallyExpanded() const { return bFolderInitiallyExpanded; }
+	UE_INTERNAL ENGINE_API void SetIsInitiallyExpanded(bool bInFolderInitiallyExpanded);
+	UE_INTERNAL bool IsInitiallyExpanded() const { return bFolderInitiallyExpanded; }
 
-	const FGuid& GetGuid() const { return FolderGuid; }
-	ENGINE_API FName GetPath() const;
-	ENGINE_API FString GetDisplayName() const;
+	UE_INTERNAL const FGuid& GetGuid() const { return FolderGuid; }
+	UE_INTERNAL ENGINE_API FName GetPath() const;
+	UE_INTERNAL ENGINE_API FString GetDisplayName() const;
 
-	ENGINE_API void MarkAsDeleted();
-	bool IsMarkedAsDeleted() const { return bIsDeleted; }
+	UE_INTERNAL ENGINE_API void MarkAsDeleted();
+	UE_INTERNAL bool IsMarkedAsDeleted() const { return bIsDeleted; }
 	
 	// Checks if folder is valid (if it's not deleted)
-	bool IsValid() const { return !IsMarkedAsDeleted(); }
+	UE_INTERNAL bool IsValid() const
+	{
+		PRAGMA_DISABLE_INTERNAL_WARNINGS
+		return !IsMarkedAsDeleted();
+		PRAGMA_ENABLE_INTERNAL_WARNINGS
+	}
 
 	// Remaps parent folder to the first parent folder not marked as deleted
-	ENGINE_API void Fixup();
+	UE_INTERNAL ENGINE_API void Fixup();
 	// Detects and clears invalid parent folder
-	ENGINE_API void FixupParentFolder();
+	UE_INTERNAL ENGINE_API void FixupParentFolder();
 
-	ENGINE_API FFolder GetFolder() const;
+	UE_INTERNAL ENGINE_API FFolder GetFolder() const;
 
 	/**
 	 * Set the folder packaging mode.
 	* @param bExternal will set the folder packaging mode to external if true, to internal otherwise
 	* @param bShouldDirty should dirty or not the level package
 	*/
-	ENGINE_API void SetPackageExternal(bool bExternal, bool bShouldDirty = true);
+	UE_INTERNAL ENGINE_API void SetPackageExternal(bool bExternal, bool bShouldDirty = true);
 
 private:
 	ENGINE_API FName GetPathInternal(UActorFolder* InSkipFolder) const;
