Fixed a bug where the EVTProducerPriority::VirtualTextureStreamingPriority is not properly set when configured through a device profile. The FTextureLODGroup::SetupGroup function would always default the priority to be Normal, which overwrote any other value that might have been set in the device profile.

#Synced-CL 49159615

diff --git a/Engine/Source/Runtime/Engine/Classes/Engine/TextureLODSettings.h b/Engine/Source/Runtime/Engine/Classes/Engine/TextureLODSettings.h
--- a/Engine/Source/Runtime/Engine/Classes/Engine/TextureLODSettings.h
+++ b/Engine/Source/Runtime/Engine/Classes/Engine/TextureLODSettings.h
@@ -3,6 +3,7 @@
 #pragma once
 
 #include "CoreMinimal.h"
+#include "VirtualTextureEnum.h"
 #include "UObject/ObjectMacros.h"
 #include "UObject/Object.h"
 #include "Engine/TextureDefines.h"
@@ -133,7 +134,7 @@ struct FTextureLODGroup
 	int32 VirtualTextureTileSizeBias = 0;
 
 	UPROPERTY()
-	EVTProducerPriority VirtualTextureStreamingPriority;
+	EVTProducerPriority VirtualTextureStreamingPriority = EVTProducerPriority::Normal;
 
 	/** Allows us to override max anisotropy. If unspecified, uses r.MaxAnisotropy */
 	UPROPERTY()
diff --git a/Engine/Source/Runtime/Engine/Private/TextureLODSettings.cpp b/Engine/Source/Runtime/Engine/Private/TextureLODSettings.cpp
--- a/Engine/Source/Runtime/Engine/Private/TextureLODSettings.cpp
+++ b/Engine/Source/Runtime/Engine/Private/TextureLODSettings.cpp
@@ -144,8 +144,6 @@ void FTextureLODGroup::SetupGroup()
 			Filter = ETextureSamplerFilter::AnisotropicLinear;
 		}
 	}
-
-	VirtualTextureStreamingPriority = EVTProducerPriority::Normal;
 }
 
 bool FTextureLODGroup::operator==(const FTextureLODGroup& Other) const
