ActorPrimitiveColorHandler - Default initialize FPrimitiveColorHandler::bAvailalbleInEditor to true

#Synced-CL 49035529

diff --git a/Engine/Source/Runtime/Engine/Classes/GameFramework/ActorPrimitiveColorHandler.h b/Engine/Source/Runtime/Engine/Classes/GameFramework/ActorPrimitiveColorHandler.h
--- a/Engine/Source/Runtime/Engine/Classes/GameFramework/ActorPrimitiveColorHandler.h
+++ b/Engine/Source/Runtime/Engine/Classes/GameFramework/ActorPrimitiveColorHandler.h
@@ -29,7 +29,7 @@ public:
 		FName HandlerName;
 		FText HandlerText;
 		FText HandlerToolTipText;
-		bool bAvailalbleInEditor;
+		bool bAvailalbleInEditor = true;
 		FGetColorFunc GetColorFunc;
 		FActivateFunc ActivateFunc;
 		FDeactivateFunc DeactivateFunc;
