World Partition - HLOD: Changed verbosity of some debug output from Log to Display

#Synced-CL 49092369

diff --git a/Engine/Plugins/Editor/WorldPartitionHLODUtilities/Source/Private/WorldPartition/HLOD/Utilities/WorldPartitionHLODUtilities.cpp b/Engine/Plugins/Editor/WorldPartitionHLODUtilities/Source/Private/WorldPartition/HLOD/Utilities/WorldPartitionHLODUtilities.cpp
--- a/Engine/Plugins/Editor/WorldPartitionHLODUtilities/Source/Private/WorldPartition/HLOD/Utilities/WorldPartitionHLODUtilities.cpp
+++ b/Engine/Plugins/Editor/WorldPartitionHLODUtilities/Source/Private/WorldPartition/HLOD/Utilities/WorldPartitionHLODUtilities.cpp
@@ -837,7 +837,6 @@ bool FWorldPartitionHLODUtilities::BuildHLOD(const FHLODBuildParams& InBuildPara
 
 	if (RebuildDecision != EHLODRebuildPolicyDecision::ApproveRebuild)
 	{
-		UE_LOG(LogHLODBuilder, Verbose, TEXT("HLOD actor \"%s\" doesn't need to be rebuilt."), *HLODActor->GetActorLabel());
 		return false;
 	}
 
@@ -923,10 +922,10 @@ PRAGMA_ENABLE_DEPRECATION_WARNINGS
 
 			if (BuildResult.HLODComponents.IsEmpty())
 			{
-				UE_LOG(LogHLODBuilder, Log, TEXT("HLOD generation created no component for HLOD actor %s. Listing source components:"), *HLODActor->GetActorLabel());
+				UE_LOG(LogHLODBuilder, Display, TEXT("HLOD generation created no component for HLOD actor %s. Listing source components:"), *HLODActor->GetActorLabel());
 				for (UActorComponent* SourceComponent : HLODBuildContext.SourceComponents)
 				{
-					UE_LOG(LogHLODBuilder, Log, TEXT("\t* %s"), *SourceComponent->GetPathName());
+					UE_LOG(LogHLODBuilder, Display, TEXT("\t* %s"), *SourceComponent->GetPathName());
 				}
 			}
 
@@ -1059,4 +1058,4 @@ void FWorldPartitionHLODUtilities::SetHLODBuildEvaluator(IWorldPartitionHLODUtil
 	HLODBuildEvaluatorDelegate = BuildEvaluatorDelegate;
 }
 
-#endif
\ No newline at end of file
+#endif
diff --git a/Engine/Source/Runtime/Engine/Private/WorldPartition/HLOD/HLODRebuildPolicy.cpp b/Engine/Source/Runtime/Engine/Private/WorldPartition/HLOD/HLODRebuildPolicy.cpp
--- a/Engine/Source/Runtime/Engine/Private/WorldPartition/HLOD/HLODRebuildPolicy.cpp
+++ b/Engine/Source/Runtime/Engine/Private/WorldPartition/HLOD/HLODRebuildPolicy.cpp
@@ -87,7 +87,7 @@ EHLODRebuildPolicyDecision UHLODRebuildPolicy::Evaluate(const AWorldPartitionHLO
 {
 	EHLODRebuildPolicyDecision FinalDecision = EHLODRebuildPolicyDecision::None;
 
-	UE_LOG(LogHLODBuilder, Log, TEXT("Evaluating HLOD rebuild policies..."));
+	UE_LOG(LogHLODBuilder, Display, TEXT("Evaluating HLOD rebuild policies..."));
 
 	for (TSubclassOf<UHLODRebuildPolicy> HLODRebuildPolicyClass : GetRebuildPoliciesClasses())
 	{
@@ -121,8 +121,8 @@ EHLODRebuildPolicyDecision UHLODRebuildPolicy::Evaluate(const AWorldPartitionHLO
 		FString PolicyName = UHLODRebuildPolicyCDO->GetClass()->GetName();
 		FString DecisionString = StaticEnum<EHLODRebuildPolicyDecision>()->GetNameStringByValue(static_cast<int64>(Decision));
 		
-		UE_CLOG( Reason.IsEmpty(), LogHLODBuilder, Log, TEXT(" * %s -> %s"), *PolicyName, *DecisionString);
-		UE_CLOG(!Reason.IsEmpty(), LogHLODBuilder, Log, TEXT(" * %s -> %s (%s)"), *PolicyName, *DecisionString, *Reason);
+		UE_CLOG( Reason.IsEmpty(), LogHLODBuilder, Display, TEXT(" * %s -> %s"), *PolicyName, *DecisionString);
+		UE_CLOG(!Reason.IsEmpty(), LogHLODBuilder, Display, TEXT(" * %s -> %s (%s)"), *PolicyName, *DecisionString, *Reason);
 		
 		if (Decision != EHLODRebuildPolicyDecision::None)
 		{
@@ -136,6 +136,9 @@ EHLODRebuildPolicyDecision UHLODRebuildPolicy::Evaluate(const AWorldPartitionHLO
 		}
 	}
 
+	FString FinalDecisionString = StaticEnum<EHLODRebuildPolicyDecision>()->GetNameStringByValue(static_cast<int64>(FinalDecision));
+	UE_LOG(LogHLODBuilder, Display, TEXT("Evaluated HLOD rebuild policies, final decision: %s"), *FinalDecisionString);
+
 	return FinalDecision;
 }
 
