ULevelInstancePropertyOverridePolicy : Add missing exports and constructor

#rnx
#Synced-CL 0

diff --git a/Engine/Source/Runtime/Engine/Private/LevelInstance/LevelInstancePropertyOverridePolicy.cpp b/Engine/Source/Runtime/Engine/Private/LevelInstance/LevelInstancePropertyOverridePolicy.cpp
--- a/Engine/Source/Runtime/Engine/Private/LevelInstance/LevelInstancePropertyOverridePolicy.cpp
+++ b/Engine/Source/Runtime/Engine/Private/LevelInstance/LevelInstancePropertyOverridePolicy.cpp
@@ -4,6 +4,11 @@
 
 #include UE_INLINE_GENERATED_CPP_BY_NAME(LevelInstancePropertyOverridePolicy)
 
+ULevelInstancePropertyOverridePolicy::ULevelInstancePropertyOverridePolicy(FVTableHelper& Helper)
+	: Super(Helper)
+{
+}
+
 #if WITH_EDITOR
 
 #include "UObject/PropertyOptional.h"
diff --git a/Engine/Source/Runtime/Engine/Public/LevelInstance/LevelInstancePropertyOverridePolicy.h b/Engine/Source/Runtime/Engine/Public/LevelInstance/LevelInstancePropertyOverridePolicy.h
--- a/Engine/Source/Runtime/Engine/Public/LevelInstance/LevelInstancePropertyOverridePolicy.h
+++ b/Engine/Source/Runtime/Engine/Public/LevelInstance/LevelInstancePropertyOverridePolicy.h
@@ -13,7 +13,8 @@ class ULevelInstancePropertyOverridePolicy : public UWorldPartitionPropertyOverr
 {
 	GENERATED_BODY()
 public:
-	ULevelInstancePropertyOverridePolicy() {}
+	ENGINE_API ULevelInstancePropertyOverridePolicy(FVTableHelper& Helper);
+	ULevelInstancePropertyOverridePolicy() = default;
 	virtual ~ULevelInstancePropertyOverridePolicy() = default;
 
 #if WITH_EDITOR
