[StandaloneHLOD] Skip non-map packages when scanning HLOD directory in search of Standalone HLODs.

#Synced-CL 48893542

diff --git a/Engine/Source/Runtime/Engine/Private/WorldPartition/HLOD/StandaloneHLODSubsystem.cpp b/Engine/Source/Runtime/Engine/Private/WorldPartition/HLOD/StandaloneHLODSubsystem.cpp
--- a/Engine/Source/Runtime/Engine/Private/WorldPartition/HLOD/StandaloneHLODSubsystem.cpp
+++ b/Engine/Source/Runtime/Engine/Private/WorldPartition/HLOD/StandaloneHLODSubsystem.cpp
@@ -6,6 +6,7 @@
 
 #if WITH_EDITOR
 
+#include "FileHelpers.h"
 #include "LevelInstance/LevelInstanceActor.h"
 #include "Settings/EditorExperimentalSettings.h"
 #include "WorldPartition/ActorDescContainerSubsystem.h"
@@ -278,6 +279,11 @@ void UWorldPartitionStandaloneHLODSubsystem::UpdateStandaloneHLODActors(UWorldPa
 				continue;
 			}
 
+			if (!FEditorFileUtils::IsMapPackageAsset(Package))
+			{
+				continue;
+			}
+
 			FString PackageName = FPackageName::FilenameToLongPackageName(Package);
 			const FSoftObjectPath LODWorldPackagePath(FTopLevelAssetPath(FName(*PackageName), FName(*FPackageName::GetLongPackageAssetName(PackageName))));
 
