FSkeletalMeshSourceModel::EnsureRawMeshBulkDataIsConvertedToNew: Protect against null MeshDescriptionBulkData via ensure, similarly to how other functions are doing.
This doesn't address the root cause, but it does prevent a crash.

#Synced-CL 48962394

diff --git a/Engine/Source/Runtime/Engine/Private/SkeletalMeshSourceModel.cpp b/Engine/Source/Runtime/Engine/Private/SkeletalMeshSourceModel.cpp
--- a/Engine/Source/Runtime/Engine/Private/SkeletalMeshSourceModel.cpp
+++ b/Engine/Source/Runtime/Engine/Private/SkeletalMeshSourceModel.cpp
@@ -294,6 +294,11 @@ USkeletalMesh* FSkeletalMeshSourceModel::GetOwner() const
 
 void FSkeletalMeshSourceModel::EnsureRawMeshBulkDataIsConvertedToNew()
 {
+	if (!ensure(MeshDescriptionBulkData))
+	{
+		return;
+	}
+
 	if (!MeshDescriptionBulkData->IsBulkDataValid() && RawMeshBulkData.IsValid())
 	{
 		ConvertRawMeshToMeshDescriptionBulkData();
