Adding the "stat thermals" command to display the thermal state of the device. Currently only available on Android/iOS.

#Synced-CL 49202707

diff --git a/Engine/Config/BaseInput.ini b/Engine/Config/BaseInput.ini
--- a/Engine/Config/BaseInput.ini
+++ b/Engine/Config/BaseInput.ini
@@ -129,6 +129,7 @@ MaxScrollbackSize=1024
 +ManualAutoCompleteList=(Command="Stat NET",Desc=)
 +ManualAutoCompleteList=(Command="Stat LIST",Desc="<Groups/Sets/Group> List groups of stats, saved sets, or specific stats within a specified group")
 +ManualAutoCompleteList=(Command="Stat splitscreen",Desc=)
++ManualAutoCompleteList=(Command="Stat THERMALS",Desc="Shows thermal state of the device if supported")
 +ManualAutoCompleteList=(Command="MemReport",Desc="Outputs memory stats to a profile file. -Full gives more data, -Log outputs to the log")
 +ManualAutoCompleteList=(Command="ListTextures",Desc="[Streaming|NonStreaming|Forced] [-Alphasort] [-csv] Lists all loaded textures and their current memory footprint")
 +ManualAutoCompleteList=(Command="ListStreamingTextures",Desc="Lists info for all streaming textures")
diff --git a/Engine/Source/Runtime/Engine/Classes/Engine/Engine.h b/Engine/Source/Runtime/Engine/Classes/Engine/Engine.h
--- a/Engine/Source/Runtime/Engine/Classes/Engine/Engine.h
+++ b/Engine/Source/Runtime/Engine/Classes/Engine/Engine.h
@@ -3994,6 +3994,7 @@ private:
 	ENGINE_API bool ToggleStatNamedEvents(UWorld* World, FCommonViewportClient* ViewportClient, const TCHAR* Stream = nullptr);
 	ENGINE_API bool ToggleStatVerboseNamedEvents(UWorld* World, FCommonViewportClient* ViewportClient, const TCHAR* Stream = nullptr);
 	ENGINE_API bool ToggleStatUnit(UWorld* World, FCommonViewportClient* ViewportClient, const TCHAR* Stream = nullptr);
+	ENGINE_API bool ToggleStatThermals(UWorld* World, FCommonViewportClient* ViewportClient, const TCHAR* Stream = nullptr);
 #if !UE_BUILD_SHIPPING
 	ENGINE_API bool PostStatSoundModulatorHelp(UWorld* World, FCommonViewportClient* ViewportClient, const TCHAR* Stream = nullptr);
 	ENGINE_API bool ToggleStatUnitCriticalPath(UWorld* World, FCommonViewportClient* ViewportClient, const TCHAR* Stream = nullptr);
@@ -4026,6 +4027,7 @@ private:
 	ENGINE_API int32 RenderStatLevels(UWorld* World, FViewport* Viewport, FCanvas* Canvas, int32 X, int32 Y, const FVector* ViewLocation = nullptr, const FRotator* ViewRotation = nullptr);
 	ENGINE_API int32 RenderStatLevelMap(UWorld* World, FViewport* Viewport, FCanvas* Canvas, int32 X, int32 Y, const FVector* ViewLocation = nullptr, const FRotator* ViewRotation = nullptr);
 	ENGINE_API int32 RenderStatUnit(UWorld* World, FViewport* Viewport, FCanvas* Canvas, int32 X, int32 Y, const FVector* ViewLocation = nullptr, const FRotator* ViewRotation = nullptr);
+	ENGINE_API int32 RenderStatThermals(UWorld* World, FViewport* Viewport, FCanvas* Canvas, int32 X, int32 Y, const FVector* ViewLocation = nullptr, const FRotator* ViewRotation = nullptr);
 	ENGINE_API int32 RenderStatDrawCount(UWorld* World, FViewport* Viewport, FCanvas* Canvas, int32 X, int32 Y, const FVector* ViewLocation = nullptr, const FRotator* ViewRotation = nullptr);
 #if !UE_BUILD_SHIPPING
 	ENGINE_API int32 RenderStatSoundReverb(UWorld* World, FViewport* Viewport, FCanvas* Canvas, int32 X, int32 Y, const FVector* ViewLocation = nullptr, const FRotator* ViewRotation = nullptr);
diff --git a/Engine/Source/Runtime/Engine/Classes/Engine/GameViewportClient.h b/Engine/Source/Runtime/Engine/Classes/Engine/GameViewportClient.h
--- a/Engine/Source/Runtime/Engine/Classes/Engine/GameViewportClient.h
+++ b/Engine/Source/Runtime/Engine/Classes/Engine/GameViewportClient.h
@@ -699,6 +699,14 @@ public:
 		return StatUnitData;
 	}
 
+	/**
+	 * Get a ptr to the stat thermals data for this viewport
+	 */
+	virtual FStatThermalsData* GetStatThermalsData() const override
+	{
+		return StatThermalsData;
+	}
+
 	/**
 	 * Get a ptr to the stat unit data for this viewport
 	 */
@@ -1073,6 +1081,9 @@ private:
 	/** Data needed to display perframe stat tracking when STAT UNIT is enabled */
 	FStatUnitData* StatUnitData;
 
+	/** Data needed to display thermal state when STAT THERMALS is enabled */
+	FStatThermalsData* StatThermalsData;
+
 	/** Data needed to display perframe stat tracking when STAT HITCHES is enabled */
 	FStatHitchesData* StatHitchesData;
 
diff --git a/Engine/Source/Runtime/Engine/Private/GameViewportClient.cpp b/Engine/Source/Runtime/Engine/Private/GameViewportClient.cpp
--- a/Engine/Source/Runtime/Engine/Private/GameViewportClient.cpp
+++ b/Engine/Source/Runtime/Engine/Private/GameViewportClient.cpp
@@ -348,6 +348,7 @@ UGameViewportClient::UGameViewportClient(const FObjectInitializer& ObjectInitial
 	if (HasAnyFlags(RF_ClassDefaultObject) == false)
 	{
 		StatUnitData = new FStatUnitData();
+		StatThermalsData = new FStatThermalsData();
 		StatHitchesData = new FStatHitchesData();
 		FCoreDelegates::StatCheckEnabled.AddUObject(this, &UGameViewportClient::HandleViewportStatCheckEnabled);
 		FCoreDelegates::StatEnabled.AddUObject(this, &UGameViewportClient::HandleViewportStatEnabled);
@@ -424,6 +425,12 @@ UGameViewportClient::~UGameViewportClient()
 		delete StatUnitData;
 		StatUnitData = NULL;
 	}
+	
+	if (StatThermalsData)
+	{
+		delete StatThermalsData;
+		StatThermalsData = NULL;
+	}
 }
 
 void UGameViewportClient::PostInitProperties()
diff --git a/Engine/Source/Runtime/Engine/Private/UnrealClient.cpp b/Engine/Source/Runtime/Engine/Private/UnrealClient.cpp
--- a/Engine/Source/Runtime/Engine/Private/UnrealClient.cpp
+++ b/Engine/Source/Runtime/Engine/Private/UnrealClient.cpp
@@ -1358,6 +1358,72 @@ int32 FStatUnitData::DrawStat(FViewport* InViewport, FCanvas* InCanvas, int32 In
 	return InY;
 }
 
+int32 FStatThermalsData::DrawStat(FViewport* InViewport, FCanvas* InCanvas, int32 InX, int32 InY)
+{
+	const EDeviceThermalState CurrentThermalState = FPlatformMisc::GetDeviceThermalState();
+
+	FLinearColor StatColor = FColor::White;
+	FString SeverityStr;
+	switch (CurrentThermalState)
+	{
+	case EDeviceThermalState::Unsupported:
+		SeverityStr = TEXT("Unsupported");
+		StatColor = FColor::Blue;
+		break;
+	case EDeviceThermalState::None:
+		SeverityStr = TEXT("None");
+		StatColor = FColor::Green;
+		break;
+	case EDeviceThermalState::Light:
+		SeverityStr = TEXT("Light");
+		StatColor = FColor::Yellow;
+		break;
+	case EDeviceThermalState::Moderate:
+		SeverityStr = TEXT("Moderate");
+		StatColor = FColor::Orange;
+		break;
+	case  EDeviceThermalState::Severe:
+		SeverityStr = TEXT("Severe");
+		StatColor = FColor::Red;
+		break;
+	case  EDeviceThermalState::Critical:
+		SeverityStr = TEXT("Critical");
+		StatColor = FColor::Red;
+		break;
+	case  EDeviceThermalState::Emergency:
+		SeverityStr = TEXT("Emergency");
+		StatColor = FColor::Red;
+		break;
+	default:
+		SeverityStr = TEXT("Unknown");
+		StatColor = FColor::Blue;
+		break;
+	}
+
+	UFont* SmallFont = GEngine->GetSmallFont();
+	check(SmallFont);
+
+	const FString LabelText = TEXT("Thermal throttle state:   ");
+	const int32 TitleSize = SmallFont->GetStringSize(*LabelText);
+
+	const int32 LabelCoordX = InX - TitleSize;
+	InCanvas->DrawShadowedString(
+		LabelCoordX,
+		InY,
+		*LabelText,
+		SmallFont,
+		FColor::White);
+
+	InCanvas->DrawShadowedString(
+		InX,
+		InY,
+		*SeverityStr,
+		SmallFont,
+		StatColor);
+
+	return InY;
+}
+
 int32 FStatHitchesData::DrawStat(FViewport* InViewport, FCanvas* InCanvas, int32 InX, int32 InY)
 {
 	const double CurrentTime = FPlatformTime::Seconds();
diff --git a/Engine/Source/Runtime/Engine/Private/UnrealEngine.cpp b/Engine/Source/Runtime/Engine/Private/UnrealEngine.cpp
--- a/Engine/Source/Runtime/Engine/Private/UnrealEngine.cpp
+++ b/Engine/Source/Runtime/Engine/Private/UnrealEngine.cpp
@@ -2357,6 +2357,7 @@ void UEngine::Init(IEngineLoop* InEngineLoop)
 	EngineStats.Add(FEngineStatFuncs(TEXT("STAT_FPS"), TEXT("STATCAT_Engine"), FText::FromString(TEXT("Display average FPS.")), FEngineStatRender::CreateUObject(this, &UEngine::RenderStatFPS), FEngineStatToggle::CreateUObject(this, &UEngine::ToggleStatFPS), bIsRHS));
 	EngineStats.Add(FEngineStatFuncs(TEXT("STAT_Summary"), TEXT("STATCAT_Engine"), FText::FromString(TEXT("Display UsedPhysical memory stat.")), FEngineStatRender::CreateUObject(this, &UEngine::RenderStatSummary), FEngineStatToggle(), bIsRHS));
 	EngineStats.Add(FEngineStatFuncs(TEXT("STAT_Unit"), TEXT("STATCAT_Engine"), FText::FromString(TEXT("Display frame times for Game Thread, Render Thread and GPU.  Also displays draw calls and primitive counts.")), FEngineStatRender::CreateUObject(this, &UEngine::RenderStatUnit), FEngineStatToggle::CreateUObject(this, &UEngine::ToggleStatUnit), bIsRHS));
+	EngineStats.Add(FEngineStatFuncs(TEXT("STAT_Thermals"), TEXT("STATCAT_Engine"), FText::FromString(TEXT("Display thermal state of the device if applicable.")), FEngineStatRender::CreateUObject(this, &UEngine::RenderStatThermals), FEngineStatToggle::CreateUObject(this, &UEngine::ToggleStatThermals), bIsRHS));
 	EngineStats.Add(FEngineStatFuncs(TEXT("STAT_DrawCount"), TEXT("STATCAT_Engine"), FText::FromString(TEXT("Display draw call counts per render pass.")), FEngineStatRender::CreateUObject(this, &UEngine::RenderStatDrawCount), FEngineStatToggle(), bIsRHS));
 	EngineStats.Add(FEngineStatFuncs(TEXT("STAT_Hitches"), TEXT("STATCAT_Engine"), FText::FromString(TEXT("Display a scrolling hitch time when frame hitches occur.")), FEngineStatRender::CreateUObject(this, &UEngine::RenderStatHitches), FEngineStatToggle::CreateUObject(this, &UEngine::ToggleStatHitches), bIsRHS));
 	EngineStats.Add(FEngineStatFuncs(TEXT("STAT_AI"), TEXT("STATCAT_Engine"), FText::FromString(TEXT("Display number of APlayerControllers and the number that were rendered to screen.")), FEngineStatRender::CreateUObject(this, &UEngine::RenderStatAI), FEngineStatToggle(), bIsRHS));
@@ -19200,6 +19201,30 @@ int32 UEngine::RenderStatUnit(UWorld* World, FViewport* Viewport, FCanvas* Canva
 	return Y;
 }
 
+// THERMALS
+bool UEngine::ToggleStatThermals(UWorld* World, FCommonViewportClient* ViewportClient, const TCHAR* Stream)
+{
+	if (ViewportClient == nullptr)
+	{
+		// Ignore if all Viewports are closed.
+		return false;
+	}
+
+	return true;
+}
+
+int32 UEngine::RenderStatThermals(UWorld* World, FViewport* Viewport, FCanvas* Canvas, int32 X, int32 Y, const FVector* ViewLocation, const FRotator* ViewRotation)
+{
+	// Forward this draw request to the viewport client
+	if (Viewport->GetClient())
+	{
+		checkf(Viewport->GetClient()->GetStatThermalsData(), TEXT("StatThermalsData must be allocated for this viewport if you wish to display stat."));
+		Y = Viewport->GetClient()->GetStatThermalsData()->DrawStat(Viewport, Canvas, X, Y);
+	}
+
+	return Y;
+}
+
 // DRAWCOUNT
 int32 UEngine::RenderStatDrawCount(UWorld* World, FViewport* Viewport, FCanvas* Canvas, int32 X, int32 Y, const FVector* ViewLocation, const FRotator* ViewRotation)
 {
diff --git a/Engine/Source/Runtime/Engine/Public/UnrealClient.h b/Engine/Source/Runtime/Engine/Public/UnrealClient.h
--- a/Engine/Source/Runtime/Engine/Public/UnrealClient.h
+++ b/Engine/Source/Runtime/Engine/Public/UnrealClient.h
@@ -385,6 +385,14 @@ struct FStatUnitData
 	int32 DrawStat(FViewport* InViewport, FCanvas* InCanvas, int32 InX, int32 InY);
 };
 
+struct FStatThermalsData
+{
+	FStatThermalsData() = default;
+
+	/** Render function to display the stat */
+	int32 DrawStat(FViewport* InViewport, FCanvas* InCanvas, int32 InX, int32 InY);
+};
+
 /** Data needed to display perframe stat tracking when STAT HITCHES is enabled */
 struct FStatHitchesData
 {
diff --git a/Engine/Source/Runtime/Engine/Public/ViewportClient.h b/Engine/Source/Runtime/Engine/Public/ViewportClient.h
--- a/Engine/Source/Runtime/Engine/Public/ViewportClient.h
+++ b/Engine/Source/Runtime/Engine/Public/ViewportClient.h
@@ -19,6 +19,7 @@ struct FInputKeyEventArgs;
 struct FKey;
 struct FStatHitchesData;
 struct FStatUnitData;
+struct FStatThermalsData;
 
 /**
  * An abstract interface to a viewport's client.
@@ -330,6 +331,11 @@ public:
 	 */
 	virtual FStatUnitData* GetStatUnitData() const { return NULL; }
 
+	/**
+	 * Get a ptr to the stat thermals data for this viewport
+	 */
+	virtual FStatThermalsData* GetStatThermalsData() const { return NULL; }
+
 	/**
 	* Get a ptr to the stat unit data for this viewport
 	*/
