World Partition: fix crash when loading using world travel in PIE.
Clear streaming data that could have inadvertently been serialized while saving in PIE when loading from disk (need to keep it when coming from PIE duplication).

#Synced-CL 49204220

diff --git a/Engine/Source/Runtime/Engine/Private/WorldPartition/RuntimeHashSet/WorldPartitionRuntimeHashSet.cpp b/Engine/Source/Runtime/Engine/Private/WorldPartition/RuntimeHashSet/WorldPartitionRuntimeHashSet.cpp
--- a/Engine/Source/Runtime/Engine/Private/WorldPartition/RuntimeHashSet/WorldPartitionRuntimeHashSet.cpp
+++ b/Engine/Source/Runtime/Engine/Private/WorldPartition/RuntimeHashSet/WorldPartitionRuntimeHashSet.cpp
@@ -1118,7 +1118,8 @@ void UWorldPartitionRuntimeHashSet::PostLoad()
 		FixupHLODSetup(RuntimePartition);
 	}
 
-	if (!GetWorld()->IsGameWorld())
+	// Clear streaming data that could have inadvertently been serialized while saving in PIE when loading from disk (need to keep it when coming from PIE duplication).
+	if (GetLinker())
 	{
 		FlushStreamingContent();
 	}
