Fix Wave Hot Reloading for MetaSounds using the new wave player node implementation

#Synced-CL 49208798

diff --git a/Engine/Source/Runtime/Engine/Private/SoundWaveProxyPlayer.cpp b/Engine/Source/Runtime/Engine/Private/SoundWaveProxyPlayer.cpp
--- a/Engine/Source/Runtime/Engine/Private/SoundWaveProxyPlayer.cpp
+++ b/Engine/Source/Runtime/Engine/Private/SoundWaveProxyPlayer.cpp
@@ -718,20 +718,12 @@ bool FSoundWaveProxyPlayer::InitializeDecoder(const FSoundWaveProxyRef& InWavePr
 	TRACE_CPUPROFILER_EVENT_SCOPE(FSoundWaveProxyPlayer::InitializeDecoder);
 
 	FName Format = InWaveProxy->GetRuntimeFormat();
-	IAudioInfoFactory* Factory = IAudioInfoFactoryRegistry::Get().Find(Format);
-	if (!ensure(Factory))
-	{
-		UE_LOG(LogAudio, Error, TEXT("FSoundWaveProxyPlayer::InitializeDecoder: Failed to create CompressedAudioInfo for wave (package: %s). Unable to find AudioInfoFactory for format: %s"),
-			*InWaveProxy->GetPackageName().ToString(), *Format.ToString());
-		return false;
-	}
-
 	TUniquePtr<ICompressedAudioInfo> InfoInstance;
-	InfoInstance.Reset(Factory->Create());
+	InfoInstance.Reset(IAudioInfoFactoryRegistry::Get().Create(Format, InWaveProxy->GetFName()));
 
 	if (!ensure(InfoInstance.IsValid()))
 	{
-		UE_LOG(LogAudio, Error, TEXT("FSoundWaveProxyPlayer::InitializeDecoder: Failed to create CompressedAudioInfo for wave (package: %s). Unable to create info from factory for for format: %s"), *InWaveProxy->GetPackageName().ToString(), *Format.ToString());
+		UE_LOG(LogAudio, Error, TEXT("FSoundWaveProxyPlayer::InitializeDecoder: Failed to create CompressedAudioInfo for wave (package: %s). Unable to create info from factory for format: %s"), *InWaveProxy->GetPackageName().ToString(), *Format.ToString());
 		return false;
 	}
 
diff --git a/Engine/Source/Runtime/Engine/Private/SoundWaveProxyReader.cpp b/Engine/Source/Runtime/Engine/Private/SoundWaveProxyReader.cpp
--- a/Engine/Source/Runtime/Engine/Private/SoundWaveProxyReader.cpp
+++ b/Engine/Source/Runtime/Engine/Private/SoundWaveProxyReader.cpp
@@ -395,7 +395,7 @@ bool FSoundWaveProxyReader::InitializeDecoder(float InStartTimeInSeconds)
 
 	if (!ensure(InfoInstance.IsValid()))
 	{
-		UE_LOG(LogAudio, Error, TEXT("FSoundWaveProxyReader::InitializeDecoder: Failed to create CompressedAudioInfo for wave (package: %s). Unable to create info from factory for for format: %s"), *WaveProxy->GetPackageName().ToString(), *Format.ToString());
+		UE_LOG(LogAudio, Error, TEXT("FSoundWaveProxyReader::InitializeDecoder: Failed to create CompressedAudioInfo for wave (package: %s). Unable to create info from factory for format: %s"), *WaveProxy->GetPackageName().ToString(), *Format.ToString());
 		return false;
 	}
 
@@ -618,4 +618,4 @@ FSoundWaveProxyReader::EDecodeResult FSoundWaveProxyReader::Decode()
 	}
 
 	return EDecodeResult::Finished;
-}
\ No newline at end of file
+}
