Move AudioSpeaker enum ahead of deprecation

#Synced-CL 48632724

diff --git a/Engine/Source/Runtime/Engine/Public/Audio.h b/Engine/Source/Runtime/Engine/Public/Audio.h
--- a/Engine/Source/Runtime/Engine/Public/Audio.h
+++ b/Engine/Source/Runtime/Engine/Public/Audio.h
@@ -8,6 +8,7 @@
 
 #include "CoreMinimal.h"
 #include "AudioDefines.h"
+#include "Audio/AudioSpeakers.h"
 #include "Stats/Stats.h"
 #include "HAL/ThreadSafeBool.h"
 #include "Sound/AudioOutputTarget.h"
@@ -67,25 +68,6 @@ struct FActiveSound;
 struct FWaveInstance;
 struct FSoundSourceBusSendInfo;
 
-/**
- * Channel definitions for multistream waves
- *
- * These are in the sample order OpenAL expects for a 7.1 sound
- * 
- */
-enum EAudioSpeakers : int
-{							//	4.0	5.1	6.1	7.1
-	SPEAKER_FrontLeft,		//	*	*	*	*
-	SPEAKER_FrontRight,		//	*	*	*	*
-	SPEAKER_FrontCenter,	//		*	*	*
-	SPEAKER_LowFrequency,	//		*	*	*
-	SPEAKER_LeftSurround,	//	*	*	*	*
-	SPEAKER_RightSurround,	//	*	*	*	*
-	SPEAKER_LeftBack,		//			*	*		If there is no BackRight channel, this is the BackCenter channel
-	SPEAKER_RightBack,		//				*
-	SPEAKER_Count
-};
-
 // Forward declarations.
 class UAudioComponent;
 class USoundNode;
diff --git a/Engine/Source/Runtime/Engine/Public/Audio/AudioSpeakers.h b/Engine/Source/Runtime/Engine/Public/Audio/AudioSpeakers.h
new file mode 100644
--- /dev/null
+++ b/Engine/Source/Runtime/Engine/Public/Audio/AudioSpeakers.h
@@ -0,0 +1,16 @@
+// Copyright Epic Games, Inc. All Rights Reserved.
+
+#pragma once
+
+enum EAudioSpeakers  : int
+{							//	4.0	5.1	6.1	7.1
+	SPEAKER_FrontLeft,		//	*	*	*	*
+	SPEAKER_FrontRight,		//	*	*	*	*
+	SPEAKER_FrontCenter,	//		*	*	*
+	SPEAKER_LowFrequency,	//		*	*	*
+	SPEAKER_LeftSurround,	//	*	*	*	*
+	SPEAKER_RightSurround,	//	*	*	*	*
+	SPEAKER_LeftBack,		//			*	*		If there is no BackRight channel, this is the BackCenter channel
+	SPEAKER_RightBack,		//				*
+	SPEAKER_Count
+};
\ No newline at end of file
