Add some destruction validation to cluster union and connectivity components. Speculative log and workaround for rare crash.

#rnx

#Synced-CL 49124308

diff --git a/Engine/Source/Runtime/Engine/Classes/PhysicsEngine/ClusterUnionComponent.h b/Engine/Source/Runtime/Engine/Classes/PhysicsEngine/ClusterUnionComponent.h
--- a/Engine/Source/Runtime/Engine/Classes/PhysicsEngine/ClusterUnionComponent.h
+++ b/Engine/Source/Runtime/Engine/Classes/PhysicsEngine/ClusterUnionComponent.h
@@ -215,6 +215,7 @@ public:
 	using FSpatialAcceleration = Chaos::ISpatialAcceleration<FExternalSpatialAccelerationPayload, Chaos::FReal, 3>;
 
 	ENGINE_API UClusterUnionComponent(const FObjectInitializer& ObjectInitializer);
+	ENGINE_API virtual ~UClusterUnionComponent();
 
 	UFUNCTION(BlueprintCallable, Category="Cluster Union")
 	ENGINE_API void AddComponentToCluster(UPrimitiveComponent* InComponent, const TArray<int32>& BoneIds, bool bRebuildGeometry = true);
diff --git a/Engine/Source/Runtime/Engine/Private/PhysicsEngine/ClusterUnionComponent.cpp b/Engine/Source/Runtime/Engine/Private/PhysicsEngine/ClusterUnionComponent.cpp
--- a/Engine/Source/Runtime/Engine/Private/PhysicsEngine/ClusterUnionComponent.cpp
+++ b/Engine/Source/Runtime/Engine/Private/PhysicsEngine/ClusterUnionComponent.cpp
@@ -116,6 +116,15 @@ UClusterUnionComponent::UClusterUnionComponent(const FObjectInitializer& ObjectI
 	bEnableDamageFromCollision = true;
 }
 
+UClusterUnionComponent::~UClusterUnionComponent()
+{
+	ensureMsgf(
+		!PhysicsProxy,
+		TEXT("ClusterUnionComponent %s destroyed while still attached to physics proxy. "
+			"BeginDestroy() might be overridden without a base class call."),
+			*GetName());
+}
+
 FPhysScene_Chaos* UClusterUnionComponent::GetChaosScene() const
 {
 	if (AActor* Owner = GetOwner(); Owner && Owner->GetWorld())
