Early eyes on this please - building servers to test  currently

Allow ping value displayed in clients to be toggleable between ICMP/UDP ping.

Add metrics to track ICMP and UDP pings so we can verify reported client side pings

#sh-12-10-high-latency-social-reports-ext

diff --git a/Engine/Source/Runtime/Engine/Private/NetConnection.cpp b/Engine/Source/Runtime/Engine/Private/NetConnection.cpp
--- a/Engine/Source/Runtime/Engine/Private/NetConnection.cpp
+++ b/Engine/Source/Runtime/Engine/Private/NetConnection.cpp
@@ -62,6 +62,9 @@ static TAutoConsoleVariable<int32> CVarPingExcludeFrameTime(TEXT("net.PingExclud
 static TAutoConsoleVariable<int32> CVarPingUsePacketRecvTime(TEXT("net.PingUsePacketRecvTime"), 0,
 	TEXT("Use OS or Receive Thread packet receive time, for calculating the ping. Excludes frame time."));
 
+static TAutoConsoleVariable<int32> CVarPingPlayerFacingPingType(TEXT("net.PlayerFacingPingType"), int32(EPingType::RoundTripExclFrame),
+	TEXT("The preferred ping type reported to the player controller. 1 = RoundTrip (including server frame time), 2 = RoundTripExclFrame, 4 = ICMP, 8 = UDPQoS"));
+
 #if !UE_BUILD_SHIPPING
 static TAutoConsoleVariable<int32> CVarPingDisplayServerTime(TEXT("net.PingDisplayServerTime"), 0,
 	TEXT("Show server frame time. Not available in shipping builds."));
@@ -3067,9 +3070,9 @@ bool UNetConnection::ReadPacketInfo(FBitReader& Reader, bool bHasPacketInfoPaylo
 		OutLagPacketId[AckIndex] = -1;
 
 		const double PacketSendTime = OutLagTime[AckIndex];
-		const double RTT			= (CurrentTime - PacketSendTime);
-		const double RTTExclFrame	= RTT - (bExcludeFrameTime ? ServerFrameTime : 0.0);
-		const double PacketPing		= FMath::Max(RTTExclFrame, 0.0);
+		const double RTT = (CurrentTime - PacketSendTime);
+		const double RTTExclFrame = RTT - (bExcludeFrameTime ? ServerFrameTime : 0.0);
+		const double PacketPing = FMath::Max(RTTExclFrame, 0.0);
 
 		LagAcc += PacketPing;
 		LagCount++;
@@ -3083,13 +3086,37 @@ bool UNetConnection::ReadPacketInfo(FBitReader& Reader, bool bHasPacketInfoPaylo
 			UE_CLOG(PlayerController.IsRemote(), LogNet, Warning, TEXT("%hs: PlayerController is unexpectedly a RemoteObject"), __func__);
 			if (!PlayerController.IsRemote())
 			{
-				PlayerController->UpdatePing(PacketPing);
+				double PlayerFacingPing = RTTExclFrame;
+				switch(EPingType(CVarPingPlayerFacingPingType.GetValueOnAnyThread()))
+				{
+					case EPingType::ICMP:
+						if (NetPing.IsValid())
+						{
+							PlayerFacingPing = NetPing->GetPingValues(EPingType::ICMP).Current;
+						}
+						break;
+					case EPingType::UDPQoS:
+						if (NetPing.IsValid())
+						{
+							PlayerFacingPing = NetPing->GetPingValues(EPingType::UDPQoS).Current;
+						}
+						break;
+					case EPingType::RoundTrip:
+						PlayerFacingPing = RTT;
+						break;
+					case EPingType::RoundTripExclFrame:
+					default:
+						PlayerFacingPing = RTTExclFrame;
+						break;
+				}
+				
+				PlayerController->UpdatePing(FMath::Max(PlayerFacingPing, 0.f));
 			}
 		}
 
 		if (NetPing.IsValid())
 		{
-			NetPing->UpdatePing(EPingType::RoundTrip, CurrentTime, FMath::Max(RTT, 0.0));
+			NetPing->UpdatePing(EPingType::RoundTrip, CurrentTime, FMath::Max(RTT, 0.0)); 
 			NetPing->UpdatePing(EPingType::RoundTripExclFrame, CurrentTime, PacketPing);
 		}
 
