PSOPrecache stat dump improvements

#Synced-CL 49265942

diff --git a/Engine/Source/Runtime/Engine/Private/PSOPrecacheMaterial.cpp b/Engine/Source/Runtime/Engine/Private/PSOPrecacheMaterial.cpp
--- a/Engine/Source/Runtime/Engine/Private/PSOPrecacheMaterial.cpp
+++ b/Engine/Source/Runtime/Engine/Private/PSOPrecacheMaterial.cpp
@@ -541,6 +541,30 @@ public:
 		return FindResult->PSOPrecachaData;
 	}
 
+	uint64 GetCompileTime(const TCHAR* InPassName, const TCHAR* InVertexFactoryName)
+	{
+		uint64 TotalCompileTime = 0;
+#if PSO_PRECACHING_VALIDATE
+		FRWScopeLock ReadLock(RWLock, SLT_ReadOnly);
+
+		EShadingPath ShadingPath = GetFeatureLevelShadingPath(GMaxRHIFeatureLevel);
+		for (auto Iter : TrackingMaterialPSORequestData)
+		{
+			uint64 PSOPrecacheHash = Iter.Key;
+			const FTrackingMaterialPSOPrecacheParams& PrecacheParams = Iter.Value;				
+			const TCHAR* PassName = FPSOCollectorCreateManager::GetName(ShadingPath, PrecacheParams.PSOCollectorIndex);
+			const TCHAR* VertexFactoryName = PrecacheParams.VertexFactoryData.VertexFactoryType ? PrecacheParams.VertexFactoryData.VertexFactoryType->GetName() : TEXT("None");
+
+			if ((InPassName == nullptr || FCString::Stricmp(InPassName, PassName) == 0) && (InVertexFactoryName == nullptr || FCString::Stricmp(InVertexFactoryName, VertexFactoryName) == 0))
+			{
+				uint64 PSOCompilationTime = PipelineStateCache::GetPrecachePSOCompilationTime(PrecacheParams.bGraphics, PSOPrecacheHash);
+				TotalCompileTime += PSOCompilationTime;
+			}
+		}
+#endif // PSO_PRECACHING_VALIDATE
+		return TotalCompileTime;
+	}
+
 	void DumpPSOPrecacheMaterialStats()
 	{		
 #if PSO_PRECACHING_VALIDATE
@@ -1044,6 +1068,11 @@ void DumpPSOPrecacheMaterialStats()
 	GMaterialPSORequestManager.DumpPSOPrecacheMaterialStats();
 }
 
+uint64 GetPSOPrecacheMaterialCompileTime(const TCHAR* InPassName, const TCHAR* InVertexFactoryName)
+{
+	return GMaterialPSORequestManager.GetCompileTime(InPassName, InVertexFactoryName);
+}
+
 static FAutoConsoleCommand GDumpPSOPrecacheMaterialStatsCmd(
 	TEXT("r.PSOPrecache.DumpMaterialStats"),
 	TEXT("Dumps all currently collected material based PSO precache data to csv file"),
diff --git a/Engine/Source/Runtime/Engine/Private/PSOPrecacheValidation.cpp b/Engine/Source/Runtime/Engine/Private/PSOPrecacheValidation.cpp
--- a/Engine/Source/Runtime/Engine/Private/PSOPrecacheValidation.cpp
+++ b/Engine/Source/Runtime/Engine/Private/PSOPrecacheValidation.cpp
@@ -9,6 +9,7 @@
 #include "PrimitiveSceneProxy.h"
 #include "UnrealEngine.h"
 #include "Materials/Material.h"
+#include "PSOPrecacheMaterial.h"
 
 #if PSO_PRECACHING_VALIDATE && UE_WITH_PSO_PRECACHING
 
@@ -387,9 +388,7 @@ uint64 PSOCollectorStats::GetPSOPrecacheHash(const FGraphicsPipelineStateInitial
 
 uint64 PSOCollectorStats::GetPSOPrecacheHash(const FComputePipelineStateInitializer& ComputeInitializer)
 {
-    FSHAHash ShaderHash = ComputeInitializer.ComputeShader->GetHash();
-	uint64 PSOHash = *reinterpret_cast<const uint64*>(&ShaderHash.Hash);
-	return PSOHash;
+	return ComputeInitializer.ComputeShader ? GetTypeHash(ComputeInitializer.ComputeShader->GetHash()) : 0;
 }
 
 void PSOCollectorStats::CheckFullPipelineStateInCache(
@@ -1156,6 +1155,7 @@ void PSOCollectorStats::DumpPSOPrecacheValidationStats()
 		}
 
 		FString Name;
+		uint64 TotalCompileTime = 0;		
 		int32 TrackedStats[(int32)EStatsType::Num][(int32)ETrackingType::Num];
 	};	
 	
@@ -1169,7 +1169,7 @@ void PSOCollectorStats::DumpPSOPrecacheValidationStats()
 				});
 		}
 
-		FString Header = FString::Printf(TEXT("%s,Full Precached,Full Used,Full Usage,Full Missed,Full TooLate,Full Untracked,ShaderOnly Precached,ShaderOnly Used,ShaderOnly Usage,ShaderOnly Missed,Minimal Precached,Minimal Used,Minimal Usage,Minimal Missed\n"), TypeName);		
+		FString Header = FString::Printf(TEXT("%s,Total Time (sec),Average Time (msec),Full Precached,Full Used,Full Usage,Full Missed,Full TooLate,Full Untracked,ShaderOnly Precached,ShaderOnly Used,ShaderOnly Usage,ShaderOnly Missed,Minimal Precached,Minimal Used,Minimal Usage,Minimal Missed\n"), TypeName);		
 		CSVFile->Serialize(TCHAR_TO_ANSI(*Header), Header.Len());
 				
 		for (auto Stats : TypeStats)
@@ -1194,6 +1194,11 @@ void PSOCollectorStats::DumpPSOPrecacheValidationStats()
 				FString NameEntry = FString::Printf(TEXT("%s"), *Stats.Name);
 				CSVFile->Serialize(TCHAR_TO_ANSI(*NameEntry), NameEntry.Len());
 
+				float TotalTimeInSeconds = FPlatformTime::ToSeconds64(Stats.TotalCompileTime);
+				float AverageTimeInMsec = FullPrecached > 0.0f ? TotalTimeInSeconds / FullPrecached * 1000.0f : 0.0f;
+				FString CompileTimeData = FString::Printf(TEXT(",%3.2f,%3.2f"), TotalTimeInSeconds, AverageTimeInMsec);
+				CSVFile->Serialize(TCHAR_TO_ANSI(*CompileTimeData), CompileTimeData.Len());
+
 				AddStats(EStatsType::Full);
 				
 				FString Row = FString::Printf(TEXT(",%d,%d"), Stats.GetStat(EStatsType::Full, ETrackingType::TooLate), Stats.GetStat(EStatsType::Full, ETrackingType::Untracked));
@@ -1221,6 +1226,7 @@ void PSOCollectorStats::DumpPSOPrecacheValidationStats()
 
 		FPerTypeStats PerTypeStats;
 		PerTypeStats.Name = TEXT("");
+		PerTypeStats.TotalCompileTime = GetPSOPrecacheMaterialCompileTime(nullptr, nullptr);
 		auto AddTotalStats = [&PerTypeStats](EStatsType StatsType, const PSOCollectorStats::FPrecacheStats& Stats)
 		{			
 			PerTypeStats.SetStat(StatsType, ETrackingType::Missed, Stats.MissData.GetTotalCount());
@@ -1266,6 +1272,7 @@ void PSOCollectorStats::DumpPSOPrecacheValidationStats()
 		{
 			FPerTypeStats MeshPassStats;
 			MeshPassStats.Name = FPSOCollectorCreateManager::GetName(ShadingPath, PSOCollectorIndex);
+			MeshPassStats.TotalCompileTime = GetPSOPrecacheMaterialCompileTime(*MeshPassStats.Name, nullptr);
 						
 			MeshPassStats.SetStat(EStatsType::Full, ETrackingType::Missed, FullPSOPrecacheStats.MissData.PerMeshPassCount[PSOCollectorIndex]);
 			MeshPassStats.SetStat(EStatsType::Full, ETrackingType::Precached, FullPSOPrecacheStats.PrecacheData.PerMeshPassCount[PSOCollectorIndex]);
@@ -1318,7 +1325,9 @@ void PSOCollectorStats::DumpPSOPrecacheValidationStats()
 		TArray<FPerTypeStats> VertexFactoryStats;
 		for (auto StatEntry : VertexFactoryStatMap)
 		{
-			VertexFactoryStats.Add(StatEntry.Value);
+			FPerTypeStats& Stat = StatEntry.Value;
+			Stat.TotalCompileTime = GetPSOPrecacheMaterialCompileTime(nullptr, *Stat.Name);
+			VertexFactoryStats.Add(Stat);
 		}
 
 		AddStatsToCSVFile(TEXT("VertexFactoryType"), true, VertexFactoryStats);
diff --git a/Engine/Source/Runtime/Engine/Public/PSOPrecacheMaterial.h b/Engine/Source/Runtime/Engine/Public/PSOPrecacheMaterial.h
--- a/Engine/Source/Runtime/Engine/Public/PSOPrecacheMaterial.h
+++ b/Engine/Source/Runtime/Engine/Public/PSOPrecacheMaterial.h
@@ -172,4 +172,9 @@ extern ENGINE_API FPSOPrecacheDataArray GetMaterialPSOPrecacheData(FMaterialPSOP
  */
 extern ENGINE_API void DumpPSOPrecacheMaterialStats();
 
+/**
+ * Get the total compile time for all material based PSOs for given pass and/or vertex factory name (if nullptr then collected for all passes and/or vertex factory types)
+ */
+extern ENGINE_API uint64 GetPSOPrecacheMaterialCompileTime(const TCHAR* InPassName, const TCHAR* InVertexFactoryName);
+
 
