[NetPhysics] Fix an issue where RewindData would not get created with RewindBeforeAdvance enabled

#Synced-CL 49271426

diff --git a/Engine/Source/Runtime/Engine/Private/PhysicsEngine/PhysicsReplication.cpp b/Engine/Source/Runtime/Engine/Private/PhysicsEngine/PhysicsReplication.cpp
--- a/Engine/Source/Runtime/Engine/Private/PhysicsEngine/PhysicsReplication.cpp
+++ b/Engine/Source/Runtime/Engine/Private/PhysicsEngine/PhysicsReplication.cpp
@@ -1037,17 +1037,17 @@ void FPhysicsReplicationAsync::OnPreSimulate_Internal()
 				continue;
 			}
 
-			if(Chaos::RewindBeforeAdvance == 0)
-			{ 
-				if (!bRewindDataExist && Input.RepMode == EPhysicsReplicationMode::Resimulation)
+			if (!bRewindDataExist && Input.RepMode == EPhysicsReplicationMode::Resimulation)
+			{
+				// We don't have rewind data but an actor is set to replicate using resimulation; we need to enable rewind capture.
+				if (ensure(Chaos::FPBDRigidsSolver::IsNetworkPhysicsPredictionEnabled() && RigidsSolver->IsUsingFixedDt()))
 				{
-					// We don't have rewind data but an actor is set to replicate using resimulation; we need to enable rewind capture.
-					if (ensure(Chaos::FPBDRigidsSolver::IsNetworkPhysicsPredictionEnabled() && RigidsSolver->IsUsingFixedDt()))
-					{
-						RigidsSolver->EnableRewindCapture();
-					}
+					RigidsSolver->EnableRewindCapture();
 				}
+			}
 
+			if (Chaos::RewindBeforeAdvance == 0)
+			{ 
 				UpdateRewindDataTarget(Input);
 			}
 			UpdateAsyncTarget(Input, RigidsSolver);
