[NetPhysics] When sim-proxies are not replicating with resimulation:
- Apply the inputs in the same timeline we receive states in. This plays out better together with Predictive Interpolation.
- ChaosMover will act on the particles state set by replication instead of setting the particle state itself from inputs.

#Synced-CL 49271426

diff --git a/Engine/Plugins/Experimental/ChaosMover/Source/ChaosMover/Private/ChaosMoverSimulation.cpp b/Engine/Plugins/Experimental/ChaosMover/Source/ChaosMover/Private/ChaosMoverSimulation.cpp
--- a/Engine/Plugins/Experimental/ChaosMover/Source/ChaosMover/Private/ChaosMoverSimulation.cpp
+++ b/Engine/Plugins/Experimental/ChaosMover/Source/ChaosMover/Private/ChaosMoverSimulation.cpp
@@ -30,6 +30,7 @@
 #include "MovementModeStateMachine.h"
 #include "MoverSimulationTypes.h"
 #include "PBDRigidsSolver.h"
+#include "Physics/NetworkPhysicsComponent.h"
 #include "Physics/Experimental/PhysScene_Chaos.h"
 #include "Physics/NetworkPhysicsComponent.h"
 #include "PhysicsEngine/PhysicsObjectExternalInterface.h"
@@ -452,9 +453,40 @@ void UChaosMoverSimulation::PostSimulationTickCharacter(const IChaosCharacterMov
 	FMoverDefaultSyncState* PostSimDefaultSyncState = OutputData.SyncState.SyncStateCollection.FindMutableDataByType<FMoverDefaultSyncState>();
 	check(PostSimDefaultSyncState);
 
+	// Is this a simulated proxy replicating with any other replication mode than physics resimulation?
+	bool bIsNonResimSimProxy = false;
+	if (const FChaosMoverSimulationDefaultInputs* SimInputs = LocalSimInput.FindDataByType<FChaosMoverSimulationDefaultInputs>())
+	{
+		if (SimInputs->AsyncNetworkPhysicsComponent && SimInputs->AsyncNetworkPhysicsComponent->IsNonResimSimProxy())
+		{
+			bIsNonResimSimProxy = true;
+		}
+	}
+
 	if (Chaos::FPBDRigidParticleHandle* ParticleHandle = GetControlledParticle())
 	{
-		// Constraint target
+		if (!TimeStep.bIsResimulating && bIsNonResimSimProxy)
+		{
+			const float DeltaSeconds = TimeStep.StepMs * 0.001f;
+
+			// Simulated proxies that replicate with other replication modes than resimulation should act on the particle data which is set by the replication
+			PostSimDefaultSyncState->SetTransforms_WorldSpace(
+				ParticleHandle->GetX() + (ParticleHandle->GetV() * DeltaSeconds),
+				FRotator(Chaos::FRotation3::IntegrateRotationWithAngularVelocity(ParticleHandle->GetR(), ParticleHandle->GetW(), DeltaSeconds)),
+				ParticleHandle->GetV(),
+				FMath::RadiansToDegrees(ParticleHandle->GetW()));
+		}
+		else
+		{
+			// Update velocities on particle from mover simulation
+			ParticleHandle->SetV(PostSimDefaultSyncState->GetVelocity_WorldSpace());
+			if (CharacterMode.ShouldCharacterRemainUpright())
+			{
+				ParticleHandle->SetW(FMath::DegreesToRadians(PostSimDefaultSyncState->GetAngularVelocityDegrees_WorldSpace()));
+			}
+		}
+
+		// Update constraint target
 		CharacterSimState.TargetDeltaPosition = PostSimDefaultSyncState->GetLocation_WorldSpace() - ParticleHandle->GetX();
 
 		FQuat TgtQuat = PostSimDefaultSyncState->GetOrientation_WorldSpace().Quaternion();
@@ -468,13 +500,6 @@ void UChaosMoverSimulation::PostSimulationTickCharacter(const IChaosCharacterMov
 			UpDir = SimInputs->UpDir;
 		}
 		CharacterSimState.TargetDeltaFacing = AngularDisplacement.Dot(UpDir);
-
-		// Velocity
-		ParticleHandle->SetV(PostSimDefaultSyncState->GetVelocity_WorldSpace());
-		if (CharacterMode.ShouldCharacterRemainUpright())
-		{
-			ParticleHandle->SetW(FMath::DegreesToRadians(PostSimDefaultSyncState->GetAngularVelocityDegrees_WorldSpace()));
-		}
 	}
 	else
 	{
@@ -537,6 +562,14 @@ void UChaosMoverSimulation::PostSimulationTickCharacterConstraint(const IChaosCh
 		CharacterConstraintMode.UpdateConstraintSettings(Settings);
 		Settings.VerticalAxis = SimInputs->UpDir;
 		Settings.TargetHeight = CharacterConstraintMode.GetTargetHeight();
+
+		if (SimInputs->AsyncNetworkPhysicsComponent && SimInputs->AsyncNetworkPhysicsComponent->IsNonResimSimProxy())
+		{
+			Settings.RadialForceLimit = 0.0f;
+			Settings.FrictionForceLimit = 0.0f;
+			Settings.TwistTorqueLimit = 0.0f;
+			Settings.SwingTorqueLimit = 0.0f;
+		}
 	}
 
 	check(Blackboard.Get());
@@ -841,11 +874,11 @@ void UChaosMoverSimulation::UpdateEvents(const FMoverTimeStep& TimeStep, UE::Cha
 			if (SimInputs->AsyncNetworkPhysicsComponent && !SimInputs->AsyncNetworkPhysicsComponent->IsServer())
 			{
 				// Clear events that are outside of the resimulation window
-				const int32 LastReceivedServerFrame = SimInputs->AsyncNetworkPhysicsComponent->GetLastReceivedServerFrame();
-				if (LastReceivedServerFrame > 0)
+				const int32 LatestReceivedStateServerFrame = SimInputs->AsyncNetworkPhysicsComponent->GetLatestReceivedStateFrame() + SimInputs->AsyncNetworkPhysicsComponent->GetNetworkPhysicsTickOffset();
+				if (LatestReceivedStateServerFrame > 0)
 				{
-					ProcessedEvents.RemoveAllSwap([&LastReceivedServerFrame](const TSharedPtr<FMoverSimulationEventData> Event) {
-						return Event->Context.ServerFrame <= LastReceivedServerFrame;
+					ProcessedEvents.RemoveAllSwap([&LatestReceivedStateServerFrame](const TSharedPtr<FMoverSimulationEventData> Event) {
+						return Event->Context.ServerFrame <= LatestReceivedStateServerFrame;
 						});
 				}
 
diff --git a/Engine/Plugins/Experimental/ChaosMover/Source/ChaosMover/Private/ChaosMoverSimulationTypes.cpp b/Engine/Plugins/Experimental/ChaosMover/Source/ChaosMover/Private/ChaosMoverSimulationTypes.cpp
--- a/Engine/Plugins/Experimental/ChaosMover/Source/ChaosMover/Private/ChaosMoverSimulationTypes.cpp
+++ b/Engine/Plugins/Experimental/ChaosMover/Source/ChaosMover/Private/ChaosMoverSimulationTypes.cpp
@@ -145,8 +145,7 @@ void FNetworkPhysicsDebugData::Set(const FAsyncNetworkPhysicsComponent* NetworkP
 		bIsLocallyControlled = NetworkPhysicsComponent->IsLocallyControlled();
 		PhysicsReplicationMode = NetworkPhysicsComponent->GetPhysicsReplicationMode();
 		NetworkPhysicsTickOffset = NetworkPhysicsComponent->GetNetworkPhysicsTickOffset();
-		LastReceivedServerFrame = NetworkPhysicsComponent->GetLastReceivedServerFrame();
-		LatestReceivedTargetStateFrame = NetworkPhysicsComponent->GetLatestReceivedTargetStateFrame();
+		LatestReceivedStateFrame = NetworkPhysicsComponent->GetLatestReceivedStateFrame();
 		ForwardPredictionTime = NetworkPhysicsComponent->GetForwardPredictionTime();
 		CurrentSimProxyInputDecayAtRuntime = NetworkPhysicsComponent->GetCurrentSimProxyInputDecayAtRuntime();
 		CurrentInputDecay = NetworkPhysicsComponent->GetCurrentInputDecay();
diff --git a/Engine/Plugins/Experimental/ChaosMover/Source/ChaosMover/Public/ChaosMover/ChaosMoverSimulationTypes.h b/Engine/Plugins/Experimental/ChaosMover/Source/ChaosMover/Public/ChaosMover/ChaosMoverSimulationTypes.h
--- a/Engine/Plugins/Experimental/ChaosMover/Source/ChaosMover/Public/ChaosMover/ChaosMoverSimulationTypes.h
+++ b/Engine/Plugins/Experimental/ChaosMover/Source/ChaosMover/Public/ChaosMover/ChaosMoverSimulationTypes.h
@@ -244,10 +244,7 @@ struct FNetworkPhysicsDebugData : public FMoverDataStructBase
 	int32 NetworkPhysicsTickOffset = INDEX_NONE;
 
 	UPROPERTY(VisibleAnywhere, Category = "Mover Info")
-	int32 LastReceivedServerFrame = INDEX_NONE;
-
-	UPROPERTY(VisibleAnywhere, Category = "Mover Info")
-	int32 LatestReceivedTargetStateFrame = INDEX_NONE;
+	int32 LatestReceivedStateFrame = INDEX_NONE;
 
 	UPROPERTY(VisibleAnywhere, Category = "Mover Info")
 	float ForwardPredictionTime = -1.0f;
diff --git a/Engine/Source/Runtime/Engine/Private/PhysicsEngine/NetworkPhysicsComponent.cpp b/Engine/Source/Runtime/Engine/Private/PhysicsEngine/NetworkPhysicsComponent.cpp
--- a/Engine/Source/Runtime/Engine/Private/PhysicsEngine/NetworkPhysicsComponent.cpp
+++ b/Engine/Source/Runtime/Engine/Private/PhysicsEngine/NetworkPhysicsComponent.cpp
@@ -2234,18 +2234,14 @@ void FAsyncNetworkPhysicsComponent::OnPreProcessInputs_Internal(const int32 Phys
 		}
 
 		// Apply replicated inputs on server and simulated proxies if setting is enabled, and on local player if we are resimulating
-		const bool bApplySimProxyInput = ComponentSettings.GetApplySimProxyInputAtRuntime() && !bIsServer && !IsLocallyControlled();
 		const bool bApplyServerInput = bIsServer && !IsLocallyControlled();
+		const bool bApplySimProxyInputAtRuntime = !bIsServer && !IsLocallyControlled() && !bIsSolverResim && ComponentSettings.GetApplySimProxyInputAtRuntime();
+		const bool bIsReplicatingWithResim = GetPhysicsReplicationMode() == EPhysicsReplicationMode::Resimulation;
 
-		if ((bApplyServerInput || bApplySimProxyInput || bIsSolverResim) && InputHistory && InputData && bCanApplyData)
+		if ((bApplyServerInput || bApplySimProxyInputAtRuntime || bIsSolverResim) && InputHistory && InputData && bCanApplyData)
 		{
 			FNetworkPhysicsPayload* PhysicsData = InputData.Get();
-			int32 NextExpectedInputFrame = PhysicsStep;
-			if (!bIsSolverResim)
-			{
-				NextExpectedInputFrame = bApplySimProxyInput ? LastInputFrameApplied : PhysicsData->LocalFrame;
-				NextExpectedInputFrame++; // Iterate one to get the next frame number we expect to be used
-			}
+			int32 NextExpectedInputFrame = bIsSolverResim ? PhysicsStep : LastInputFrameApplied + 1;
 
 			// There are important inputs earlier than upcoming input to apply
 			if (!bIsSolverResim && NewImportantInputFrame < NextExpectedInputFrame)
@@ -2283,10 +2279,32 @@ void FAsyncNetworkPhysicsComponent::OnPreProcessInputs_Internal(const int32 Phys
 				}
 			}
 
+			int32 InputFrameToApply = PhysicsStep;
+
+			// For sim-proxies applying inputs at runtime
+			if (bApplySimProxyInputAtRuntime)
+			{
+				// Get either the latest received input frame
+				if (InputHistory->GetLatestFrame() > 0)
+				{
+					InputFrameToApply = FMath::Min(InputFrameToApply, InputHistory->GetLatestFrame());
+				}
+
+				// Or if the sim-proxy should run in the interpolated timeline, get the last received state frame
+				if (IsNonResimSimProxy())
+				{
+					const int32 LatestStateFrame = GetLatestReceivedStateFrame();
+					if (LatestStateFrame > 0)
+					{
+						InputFrameToApply = FMath::Min(InputFrameToApply, LatestStateFrame);
+					}
+				}
+			}
+
 			const bool RequireExactFrame = bIsServer && !ComponentSettings.GetAllowInputExtrapolation();
-			if (InputHistory->ExtractData(PhysicsStep, bIsFirstResimFrame, PhysicsData, RequireExactFrame))
+			if (InputHistory->ExtractData(InputFrameToApply, bIsFirstResimFrame, PhysicsData, RequireExactFrame))
 			{
-				const int32 CurrentInputFrame = bApplySimProxyInput ? InputHistory->GetLatestFrame() : PhysicsData->LocalFrame;
+				const int32 CurrentInputFrame = PhysicsData->LocalFrame;
 
 				// Check if the extracted data was altered and if we have a hole in the buffer
 				if (bIsServer && PhysicsData->bDataAltered)
@@ -2315,7 +2333,7 @@ void FAsyncNetworkPhysicsComponent::OnPreProcessInputs_Internal(const int32 Phys
 					}
 				}
 				// Check if we have a gap between last used input and current input
-				else if (CurrentInputFrame > NextExpectedInputFrame)
+				else if (NextExpectedInputFrame > 0 && CurrentInputFrame > NextExpectedInputFrame)
 				{
 					if (ComponentSettings.GetApplyDataInsteadOfMergeData())
 					{
@@ -2330,8 +2348,8 @@ void FAsyncNetworkPhysicsComponent::OnPreProcessInputs_Internal(const int32 Phys
 							bool bLastInput = NextExpectedInputFrame == CurrentInputFrame;
 							if (InputHistory->ExtractData(NextExpectedInputFrame, bIsFirstResimFrame, PhysicsData, (RequireExactFrame || !bLastInput)) && !bLastInput)
 							{
-								// Perform input decay on sim-proxy runtime (non-resim) input
-								if (bApplySimProxyInput && ComponentSettings.GetApplyInputDecaySimProxyInputAtRuntime())
+								// Perform input decay on sim-proxy at runtime if replicated with resimulation
+								if (PhysicsReplicationCVars::ResimulationCVars::bEnableInputDecay && !IsNonResimSimProxy() && bApplySimProxyInputAtRuntime && ComponentSettings.GetApplyInputDecaySimProxyInputAtRuntime())
 								{
 									const float RuntimeInputDecay = GetCurrentSimProxyInputDecayAtRuntime();
 									PhysicsData->DecayData(RuntimeInputDecay);
@@ -2371,11 +2389,11 @@ void FAsyncNetworkPhysicsComponent::OnPreProcessInputs_Internal(const int32 Phys
 					// Explicitly say this input was not received, since it was altered by the server and when receiving the input for this frame it should get processed as altered but not received when calling ReceiveNewData()
 					PhysicsData->bReceivedData = false;
 					PhysicsData->bImportant = false;
-					InputHistory->RecordData(PhysicsStep, PhysicsData);
+					InputHistory->RecordData(InputFrameToApply, PhysicsData);
 				}
 
-				// Perform input decay on sim-proxy runtime (non-resim) input
-				if (PhysicsReplicationCVars::ResimulationCVars::bEnableInputDecay && !bIsSolverResim && bApplySimProxyInput && ComponentSettings.GetApplyInputDecaySimProxyInputAtRuntime())
+				// Perform input decay on sim-proxy at runtime if replicated with resimulation
+				if (PhysicsReplicationCVars::ResimulationCVars::bEnableInputDecay && !bIsSolverResim && bIsReplicatingWithResim && bApplySimProxyInputAtRuntime && ComponentSettings.GetApplyInputDecaySimProxyInputAtRuntime())
 				{
 					const float RuntimeInputDecay = GetCurrentSimProxyInputDecayAtRuntime();
 					PhysicsData->DecayData(RuntimeInputDecay);
@@ -2397,7 +2415,11 @@ void FAsyncNetworkPhysicsComponent::OnPreProcessInputs_Internal(const int32 Phys
 					}
 				}
 
-				LastInputFrameApplied = InputHistory->GetLatestFrame();
+				// Update LastInputFrameApplied for always for local player and server, but only during non-resim frames for sim-proxies since we predict inputs for them during resim
+				if (IsLocallyControlled() || IsServer() || !bIsSolverResim)
+				{
+					LastInputFrameApplied = PhysicsData->LocalFrame;
+				}
 
 #if DEBUG_NETWORK_PHYSICS
 				{
@@ -2407,7 +2429,7 @@ void FAsyncNetworkPhysicsComponent::OnPreProcessInputs_Internal(const int32 Phys
 #endif
 			}
 #if DEBUG_NETWORK_PHYSICS
-			else if (PhysicsStep <= InputHistory->GetLatestFrame())
+			else if (InputFrameToApply <= InputHistory->GetLatestFrame())
 			{
 				UE_LOG(LogChaos, Log, TEXT("		Non-Determinism: FAILED to extract and apply input from history | IsResim = %d | IsLocallyControlled = %d | -- Printing history --"), bIsSolverResim, IsLocallyControlled());
 				InputHistory->DebugData(FString::Printf(TEXT("InputHistory | Name = %s"), *GetActorName()));
@@ -2957,31 +2979,15 @@ const float FAsyncNetworkPhysicsComponent::GetCurrentSimProxyInputDecayAtRuntime
 	return FMath::Clamp(SimProxyRuntimeInputDecay, 0.0f, 1.0f);
 }
 
-const int32 FAsyncNetworkPhysicsComponent::GetLatestReceivedTargetStateFrame() const
-{
-	const Chaos::FPhysicsSolverBase* BaseSolver = GetSolver();
-	if (const Chaos::FRewindData* RewindData = BaseSolver ? BaseSolver->GetRewindData() : nullptr)
-	{
-		// Todo: we should not have to get a read lock to get this information but instead measure the latest lag upon reception
-		Chaos::FReadPhysicsObjectInterface_Internal Interface = Chaos::FPhysicsObjectInternalInterface::GetRead();
-		if (Chaos::FGeometryParticleHandle* Particle = Interface.GetParticle(PhysicsObject))
-		{
-			return RewindData->GetLatestReceivedTargetStateFrame(*Particle);
-		}
-	}
-
-	return INDEX_NONE;
-}
-
 const float FAsyncNetworkPhysicsComponent::GetForwardPredictionTime() const
 {
 	const Chaos::FPhysicsSolverBase* BaseSolver = GetSolver();
 	if (const Chaos::FRewindData* RewindData = BaseSolver ? BaseSolver->GetRewindData() : nullptr)
 	{
-		int32 LatestReceivedServerFrame = GetLastReceivedServerFrame();
-		if (LatestReceivedServerFrame != INDEX_NONE)
+		const int32 LatestStateFrame = GetLatestReceivedStateFrame();
+		if (LatestStateFrame != INDEX_NONE)
 		{
-			const int32 MeasuredInputDelayFrames = RewindData->GetLatestFrame() - LatestReceivedServerFrame;
+			const int32 MeasuredInputDelayFrames = RewindData->GetLatestFrame() - LatestStateFrame;
 			const float MeasuredInputDelaySeconds = FMath::Max(MeasuredInputDelayFrames, 0.0f) * BaseSolver->GetAsyncDeltaTime();
 			return MeasuredInputDelaySeconds;
 		}
@@ -3234,7 +3240,7 @@ const int32 FAsyncNetworkPhysicsComponent::SetupRewindData()
 	return NumFrames;
 }
 
-int32 FAsyncNetworkPhysicsComponent::GetLastReceivedServerFrame() const
+int32 FAsyncNetworkPhysicsComponent::GetLatestReceivedStateFrame() const
 {
 	if (LastReceivedStateFrame > INDEX_NONE)
 	{
diff --git a/Engine/Source/Runtime/Engine/Private/PhysicsEngine/PhysicsReplication.cpp b/Engine/Source/Runtime/Engine/Private/PhysicsEngine/PhysicsReplication.cpp
--- a/Engine/Source/Runtime/Engine/Private/PhysicsEngine/PhysicsReplication.cpp
+++ b/Engine/Source/Runtime/Engine/Private/PhysicsEngine/PhysicsReplication.cpp
@@ -2622,7 +2622,6 @@ void FPhysicsReplicationAsync::DebugDrawReplicationMode(const FPhysicsRepAsyncIn
 
 	FColor DebugColor = FColor::White;
 	FVector BoxExtent = FVector(10.0f, 10.0f, 10.0f);
-	FQuat Rotation = FQuat::Identity;
 
 	if (Input.PhysicsObject)
 	{
@@ -2632,7 +2631,6 @@ void FPhysicsReplicationAsync::DebugDrawReplicationMode(const FPhysicsRepAsyncIn
 			if (Chaos::FGeometryParticleHandle* Handle = Interface.GetParticle(Input.PhysicsObject))
 			{
 				BoxExtent = Handle->LocalBounds().Extents() * 0.5f;
-				Rotation = Handle->GetR();
 			}
 
 			const EPhysicsReplicationMode RepMode = Target->IsWaiting() ? Target->RepModeOverride : Target->RepMode;
@@ -2658,11 +2656,10 @@ void FPhysicsReplicationAsync::DebugDrawReplicationMode(const FPhysicsRepAsyncIn
 		Chaos::FSingleParticlePhysicsProxy* Proxy = Input.Proxy;
 		Chaos::FRigidBodyHandle_Internal* Handle = Proxy->GetPhysicsThreadAPI();
 
-		Rotation = Handle->GetR();
 		DebugColor = FColor::Green;
 	}
 
-	Chaos::FDebugDrawQueue::GetInstance().DrawDebugBox(Input.TargetState.Position, BoxExtent, Rotation, DebugColor, false, PhysicsReplicationCVars::DebugDrawLifeTime, 0, 1.0f);
+	Chaos::FDebugDrawQueue::GetInstance().DrawDebugBox(Input.TargetState.Position, BoxExtent, Input.TargetState.Quaternion, DebugColor, false, PhysicsReplicationCVars::DebugDrawLifeTime, 0, 1.0f);
 #endif
 }
 
diff --git a/Engine/Source/Runtime/Engine/Public/Physics/NetworkPhysicsComponent.h b/Engine/Source/Runtime/Engine/Public/Physics/NetworkPhysicsComponent.h
--- a/Engine/Source/Runtime/Engine/Public/Physics/NetworkPhysicsComponent.h
+++ b/Engine/Source/Runtime/Engine/Public/Physics/NetworkPhysicsComponent.h
@@ -1929,7 +1929,8 @@ public:
 	// Get the physics tick offset (add to clients physics tick to get the servers corresponding physics tick)
 	const int32 GetNetworkPhysicsTickOffset() const { return NetworkPhysicsTickOffset; }
 
-	// Get the physics replication mode used
+	/** Get the physics replication mode used
+	* TODO: FPhysicsReplication can change replication mode temporarily on the physics thread, make this correspond to the actual replication mode being used*/
 	const EPhysicsReplicationMode GetPhysicsReplicationMode() const { return PhysicsReplicationMode; }
 
 	// Add state/input history to rewind data
@@ -1941,8 +1942,11 @@ public:
 	// Enable RewindData history caching and return the history size
 	const int32 SetupRewindData();
 
+	// Return true if this is a simulated proxy and it's not currently replicating via physics resimulation
+	const bool IsNonResimSimProxy() const { return NetRole == ENetRole::ROLE_SimulatedProxy && PhysicsReplicationMode != EPhysicsReplicationMode::Resimulation; }
+
 	// Get the frame number of the last server update
-	ENGINE_API int32 GetLastReceivedServerFrame() const;
+	ENGINE_API int32 GetLatestReceivedStateFrame() const;
 
 	/** Returns the time in seconds between the frame of the latest received state and the latest frame we have reached.
 	  * This can be used as a quantity that is stable over the course of a resimulation (from rewind to last frame of resim) and grows with the amount of time being resimulated
diff --git a/Engine/Source/Runtime/Experimental/Chaos/Public/RewindData.h b/Engine/Source/Runtime/Experimental/Chaos/Public/RewindData.h
--- a/Engine/Source/Runtime/Experimental/Chaos/Public/RewindData.h
+++ b/Engine/Source/Runtime/Experimental/Chaos/Public/RewindData.h
@@ -253,7 +253,10 @@ public :
 			const int32 MinFrameIndex = ClosestData(ExtractFrame, true);
 			const int32 MaxFrameIndex = ClosestData(ExtractFrame, false);
 
-			if (MinFrameIndex != INDEX_NONE && MaxFrameIndex != INDEX_NONE)
+			if (MinFrameIndex != INDEX_NONE
+				&& MaxFrameIndex != INDEX_NONE 
+				&& (MinFrameIndex != MaxFrameIndex)
+				&& (DataHistory[MinFrameIndex].LocalFrame != DataHistory[MaxFrameIndex].LocalFrame))
 			{
 				DataType& ExtractedData = *static_cast<DataType*>(HistoryData);
 				ExtractedData = DataHistory[MinFrameIndex];
