[NetPhysics] Support Seamless Travel with the NetworkPhysicsComponent.

#Synced-CL 49278806

diff --git a/Engine/Plugins/Experimental/ChaosModularVehicle/Source/ChaosModularVehicleEngine/Private/ChaosModularVehicle/ModularVehicleBaseComponent.cpp b/Engine/Plugins/Experimental/ChaosModularVehicle/Source/ChaosModularVehicleEngine/Private/ChaosModularVehicle/ModularVehicleBaseComponent.cpp
--- a/Engine/Plugins/Experimental/ChaosModularVehicle/Source/ChaosModularVehicleEngine/Private/ChaosModularVehicle/ModularVehicleBaseComponent.cpp
+++ b/Engine/Plugins/Experimental/ChaosModularVehicle/Source/ChaosModularVehicleEngine/Private/ChaosModularVehicle/ModularVehicleBaseComponent.cpp
@@ -200,6 +200,7 @@ void UModularVehicleBaseComponent::OnDestroyPhysicsState()
 		NetworkPhysicsComponent->RemoveDataHistory();
 	}
 
+	CachedPhysicsProxy = nullptr;
 }
 
 int UModularVehicleBaseComponent::GenerateNewGuid()
diff --git a/Engine/Plugins/Experimental/ChaosMover/Source/ChaosMover/Private/Backends/ChaosMoverSubsystem.cpp b/Engine/Plugins/Experimental/ChaosMover/Source/ChaosMover/Private/Backends/ChaosMoverSubsystem.cpp
--- a/Engine/Plugins/Experimental/ChaosMover/Source/ChaosMover/Private/Backends/ChaosMoverSubsystem.cpp
+++ b/Engine/Plugins/Experimental/ChaosMover/Source/ChaosMover/Private/Backends/ChaosMoverSubsystem.cpp
@@ -64,6 +64,7 @@ void UChaosMoverSubsystem::Deinitialize()
 				if (AsyncCallback)
 				{
 					Solver->UnregisterAndFreeSimCallbackObject_External(AsyncCallback);
+					AsyncCallback = nullptr;
 				}
 			}
 		}
diff --git a/Engine/Source/Runtime/Engine/Private/PhysicsEngine/NetworkPhysicsComponent.cpp b/Engine/Source/Runtime/Engine/Private/PhysicsEngine/NetworkPhysicsComponent.cpp
--- a/Engine/Source/Runtime/Engine/Private/PhysicsEngine/NetworkPhysicsComponent.cpp
+++ b/Engine/Source/Runtime/Engine/Private/PhysicsEngine/NetworkPhysicsComponent.cpp
@@ -551,26 +551,17 @@ void UNetworkPhysicsComponent::InitPhysics()
 	}
 
 	/** NOTE:
-	* If the NetworkPhysicsComponent is added as a SubObject after the actor has processed bAutoActivate and bWantsInitializeComponent
-	* SetActive(true) and InitializeComponent() needs to be called manually for the component to function properly. */
+	* If the NetworkPhysicsComponent is added as a SubObject after the actor has processed bAutoActivate
+	* SetActive(true) and RegisterComponent() needs to be called manually for the component to function properly. */
 	PrimaryComponentTick.bCanEverTick = true;
 	PrimaryComponentTick.TickGroup = TG_PostPhysics;
 	bAutoActivate = true;
-	bWantsInitializeComponent = true;
 	SetIsReplicatedByDefault(true);
 }
 
-void UNetworkPhysicsComponent::BeginPlay()
+void UNetworkPhysicsComponent::OnRegister()
 {
-	Super::BeginPlay();
-
-	// Update async component with current component properties
-	UpdateAsyncComponent(true);
-}
-
-void UNetworkPhysicsComponent::InitializeComponent()
-{
-	Super::InitializeComponent();
+	Super::OnRegister();
 
 	InitializePhysicsReplication();
 }
@@ -655,32 +646,15 @@ void UNetworkPhysicsComponent::InitializePhysicsReplication()
 	}
 }
 
-#if UE_WITH_REMOTE_OBJECT_HANDLE
-void UNetworkPhysicsComponent::PostMigrate(const struct FUObjectMigrationContext& MigrationContext)
+void UNetworkPhysicsComponent::OnUnregister()
 {
-	Super::PostMigrate(MigrationContext);
+	Super::OnUnregister();
 
-	if (MigrationContext.MigrationSide == EObjectMigrationSide::Send)
-	{
-		// Need to uninitialize the component during migration, because the component is not GC'ed until next frame
-		// Then the call back on AsyncPhysicsNetworkComponent would cause a crash
-		UninitializeComponent();
-	}
-
-	if (MigrationContext.MigrationSide == EObjectMigrationSide::Receive)
-	{
-		// Need to reinitialize the component (excluding initialization of the base)
-		// The reason is ActorComponents do not get reinitialized after migrated to another server since all data on ActorComponent
-		// is serialized and migrated over. 
-		InitializePhysicsReplication();
-	}
+	UninitializePhysicsReplication();
 }
-#endif
 
-void UNetworkPhysicsComponent::UninitializeComponent()
+void UNetworkPhysicsComponent::UninitializePhysicsReplication()
 {
-	Super::UninitializeComponent();
-
 	if (NetworkPhysicsComponent_Internal)
 	{
 		if (FAsyncNetworkPhysicsComponentInput* AsyncInput = NetworkPhysicsComponent_Internal->GetProducerInputData_External())
@@ -689,35 +663,66 @@ void UNetworkPhysicsComponent::UninitializeComponent()
 			AsyncInput->PhysicsObject = nullptr;
 			AsyncInput->ImplementationInterface_Internal = nullptr;
 		}
-	}
 
-	if (UWorld* World = GetWorld())
-	{
-		if (FPhysScene* PhysScene = World->GetPhysicsScene())
+		if (UWorld* World = GetWorld())
 		{
-			if (Chaos::FPhysicsSolver* Solver = PhysScene->GetSolver())
+			if (FPhysScene* PhysScene = World->GetPhysicsScene())
 			{
-				/* Run OnUninitialize_Internal on the ISimCallbackObject as a way to unregister input / state history, unsubscribe from delegates etc.
-				* After UnregisterAndFreeSimCallbackObject_External the ISimCallbackObject will not get any callbacks anymore, use this as the last safe place to use the cached FPhysicsObject for example */
-				FAsyncNetworkPhysicsComponent* AsyncNetworkPhysicsComponent = NetworkPhysicsComponent_Internal;
-				Solver->EnqueueCommandImmediate(
-					[AsyncNetworkPhysicsComponent]()
-					{
-						if (AsyncNetworkPhysicsComponent)
+				if (Chaos::FPhysicsSolver* Solver = PhysScene->GetSolver())
+				{
+					/* Run OnUninitialize_Internal on the ISimCallbackObject as a way to unregister input / state history, unsubscribe from delegates etc.
+					* After UnregisterAndFreeSimCallbackObject_External the ISimCallbackObject will not get any callbacks anymore, use this as the last safe place to use the cached FPhysicsObject for example */
+					FAsyncNetworkPhysicsComponent* AsyncNetworkPhysicsComponent = NetworkPhysicsComponent_Internal;
+					Solver->EnqueueCommandImmediate(
+						[AsyncNetworkPhysicsComponent]()
 						{
-							AsyncNetworkPhysicsComponent->OnUninitialize_Internal();
+							if (AsyncNetworkPhysicsComponent)
+							{
+								AsyncNetworkPhysicsComponent->OnUninitialize_Internal();
+							}
 						}
-					}
-				);
+					);
 
-				// Clear async component from Physics Thread and memory
-				Solver->UnregisterAndFreeSimCallbackObject_External(NetworkPhysicsComponent_Internal);
+					// Clear async component from Physics Thread and memory
+					Solver->UnregisterAndFreeSimCallbackObject_External(NetworkPhysicsComponent_Internal);
+				}
 			}
 		}
 	}
 	NetworkPhysicsComponent_Internal = nullptr;
+	PhysicsObject = nullptr;
+}
+
+void UNetworkPhysicsComponent::BeginPlay()
+{
+	Super::BeginPlay();
+
+	// Update async component with current component properties
+	UpdateAsyncComponent(true);
 }
 
+#if UE_WITH_REMOTE_OBJECT_HANDLE
+void UNetworkPhysicsComponent::PostMigrate(const struct FUObjectMigrationContext& MigrationContext)
+{
+	Super::PostMigrate(MigrationContext);
+
+	if (MigrationContext.MigrationSide == EObjectMigrationSide::Send)
+	{
+		// Need to uninitialize the component during migration, because the component is not GC'ed until next frame
+		// Then the call back on AsyncPhysicsNetworkComponent would cause a crash
+		UninitializePhysicsReplication();
+	}
+
+	if (MigrationContext.MigrationSide == EObjectMigrationSide::Receive)
+	{
+		// Need to reinitialize the component (excluding initialization of the base)
+		// The reason is ActorComponents do not get reinitialized after migrated to another server since all data on ActorComponent
+		// is serialized and migrated over. 
+		InitializePhysicsReplication();
+	}
+}
+#endif
+
 void UNetworkPhysicsComponent::GetLifetimeReplicatedProps(TArray< FLifetimeProperty >& OutLifetimeProps) const
 {
 	Super::GetLifetimeReplicatedProps(OutLifetimeProps);
@@ -1713,7 +1718,7 @@ void UNetworkPhysicsComponent::UpdateAsyncComponent(const bool bFullUpdate)
 					AsyncInput->ActorName = AActor::GetDebugName(Owner);
 			
 					UNetworkPhysicsSettingsComponent* SettingsComponent = Owner->FindComponentByClass<UNetworkPhysicsSettingsComponent>();
-					if (SettingsComponent && SettingsComponent->GetNetworkPhysicsSettings_Internal() && SettingsComponent->GetSettings_Internal().IsValid())
+					if (SettingsComponent && SettingsComponent->GetSettings_Internal().IsValid())
 					{
 						AsyncInput->SettingsComponent = SettingsComponent->GetSettings_Internal();
 					}
@@ -2109,7 +2114,7 @@ void FAsyncNetworkPhysicsComponent::ConsumeAsyncInput_Internal(const int32 Physi
 			// Running average of how far behind the forward predicted timeline we receive inputs
 			if (Chaos::FPBDRigidsSolver* RigidsSolver = GetRigidSolver())
 			{
-				if (!bIsServer && !IsLocallyControlled())
+				if (!bIsServer && !IsLocallyControlled() && InputHistory)
 				{
 					const float Dt = RigidsSolver->GetAsyncDeltaTime();
 					const int32 InputPredictedFrames = FMath::Max(RigidsSolver->GetCurrentFrame() - InputHistory->GetLatestFrame(), 0);
diff --git a/Engine/Source/Runtime/Engine/Private/PhysicsEngine/NetworkPhysicsSettingsComponent.cpp b/Engine/Source/Runtime/Engine/Private/PhysicsEngine/NetworkPhysicsSettingsComponent.cpp
--- a/Engine/Source/Runtime/Engine/Private/PhysicsEngine/NetworkPhysicsSettingsComponent.cpp
+++ b/Engine/Source/Runtime/Engine/Private/PhysicsEngine/NetworkPhysicsSettingsComponent.cpp
@@ -49,28 +49,9 @@ void UNetworkPhysicsSettingsComponent::OnRegister()
 			RootPrimComp->OnComponentPhysicsStateChanged.AddUniqueDynamic(this, &ThisClass::OnComponentPhysicsStateChanged);
 		}
 	}
-}
-
-void UNetworkPhysicsSettingsComponent::OnUnregister()
-{
-	Super::OnUnregister();
-
-	if (AActor* Owner = GetOwner())
-	{
-		if (UPrimitiveComponent* RootPrimComp = Cast<UPrimitiveComponent>(Owner->GetRootComponent()))
-		{
-			RootPrimComp->OnComponentPhysicsStateChanged.RemoveDynamic(this, &ThisClass::OnComponentPhysicsStateChanged);
-		}
-	}
-}
-
-void UNetworkPhysicsSettingsComponent::InitializeComponent()
-{
-	Super::InitializeComponent();
 
 	if (!SettingsDataAsset)
 	{
-		UE_LOG(LogChaos, Warning, TEXT("UNetworkPhysicsSettingsComponent doesn't have a DataAsset assigned. Settings directly on the component are deprecated - %s"), *AActor::GetDebugName(GetOwner()));
 		return;
 	}
 
@@ -132,9 +113,17 @@ void UNetworkPhysicsSettingsComponent::InitializeComponent()
 	}
 }
 
-void UNetworkPhysicsSettingsComponent::UninitializeComponent()
+void UNetworkPhysicsSettingsComponent::OnUnregister()
 {
-	Super::UninitializeComponent();
+	Super::OnUnregister();
+
+	if (AActor* Owner = GetOwner())
+	{
+		if (UPrimitiveComponent* RootPrimComp = Cast<UPrimitiveComponent>(Owner->GetRootComponent()))
+		{
+			RootPrimComp->OnComponentPhysicsStateChanged.RemoveDynamic(this, &ThisClass::OnComponentPhysicsStateChanged);
+		}
+	}
 
 	if (SettingsDataAsset)
 	{
diff --git a/Engine/Source/Runtime/Engine/Public/Physics/NetworkPhysicsComponent.h b/Engine/Source/Runtime/Engine/Public/Physics/NetworkPhysicsComponent.h
--- a/Engine/Source/Runtime/Engine/Public/Physics/NetworkPhysicsComponent.h
+++ b/Engine/Source/Runtime/Engine/Public/Physics/NetworkPhysicsComponent.h
@@ -1353,10 +1353,10 @@ public:
 	ENGINE_API virtual void BeginPlay() override;
 
 	// Register the component into the network manager
-	ENGINE_API virtual void InitializeComponent() override;
+	virtual void OnRegister() override;
 
 	// Unregister the component from the network manager
-	ENGINE_API virtual void UninitializeComponent() override;
+	virtual void OnUnregister() override;
 
 	// Remove state/input history from rewind data
 	ENGINE_API void RemoveDataHistory();
@@ -1584,6 +1584,9 @@ private:
 	/** Initializes necessary data for NetworkPhysicsComponent_Internal and callbacks */
 	void InitializePhysicsReplication();
 
+	/** Uninitialize NetworkPhysicsComponent_Internal and callbacks */
+	void UninitializePhysicsReplication();
+
 public:
 	/** Get the delta source stored at index or with frame, @param bValueIsIndex switches between grabbing frame vs index
 	* Note, passing in @param Value as -1 returns the latest data
diff --git a/Engine/Source/Runtime/Engine/Public/Physics/NetworkPhysicsSettingsComponent.h b/Engine/Source/Runtime/Engine/Public/Physics/NetworkPhysicsSettingsComponent.h
--- a/Engine/Source/Runtime/Engine/Public/Physics/NetworkPhysicsSettingsComponent.h
+++ b/Engine/Source/Runtime/Engine/Public/Physics/NetworkPhysicsSettingsComponent.h
@@ -813,8 +813,6 @@ public:
 
 	virtual void OnRegister() override;
 	virtual void OnUnregister() override;
-	virtual void InitializeComponent() override;
-	virtual void UninitializeComponent() override;
 
 	virtual void BeginPlay() override;
 
