[Backout] - CL49284163 - CIS Compile Error
#rnx
Original CL Desc
-----------------------------------------------------------------
Split the focus summary out of the session summary, so it can be split into it's own event and handled seperately.

- Organize the Focus Summary data so we can make the events easier to parse and filter

#Synced-CL 49063813

diff --git a/Engine/Source/Editor/EditorToolEvents/EditorToolEvents.Build.cs b/Engine/Source/Editor/EditorToolEvents/EditorToolEvents.Build.cs
--- a/Engine/Source/Editor/EditorToolEvents/EditorToolEvents.Build.cs
+++ b/Engine/Source/Editor/EditorToolEvents/EditorToolEvents.Build.cs
@@ -26,7 +26,8 @@ public class EditorToolEvents : ModuleRules
 		PublicDependencyModuleNames.AddRange(
 			new string[]
 			{
-				"Core"
+				"Core",
+				"Engine"
 				// ... add other public dependencies that you statically link with here ...
 			}
 			);
@@ -35,9 +36,6 @@ public class EditorToolEvents : ModuleRules
 		PrivateDependencyModuleNames.AddRange(
 			new string[]
 			{
-				"SlateCore",
-				"Slate",
-				"Json"
 			}
 			);
 		
diff --git a/Engine/Source/Editor/EditorToolEvents/Private/EditorToolDelegates.cpp b/Engine/Source/Editor/EditorToolEvents/Private/EditorToolDelegates.cpp
--- a/Engine/Source/Editor/EditorToolEvents/Private/EditorToolDelegates.cpp
+++ b/Engine/Source/Editor/EditorToolEvents/Private/EditorToolDelegates.cpp
@@ -3,4 +3,3 @@
 #include "EditorToolDelegates.h"
 
 FEditorToolDelegates::FEditorToolStarted	FEditorToolDelegates::OnEditorToolStarted;
-FEditorToolDelegates::FSendEditorFocusSummary	FEditorToolDelegates::OnSendEditorFocusSummary;
diff --git a/Engine/Source/Editor/EditorToolEvents/Private/EditorUsageAnalyticsSummary.cpp b/Engine/Source/Editor/EditorToolEvents/Private/EditorUsageAnalyticsSummary.cpp
deleted file mode 100644
diff --git a/Engine/Source/Editor/EditorToolEvents/Public/EditorToolDelegates.h b/Engine/Source/Editor/EditorToolEvents/Public/EditorToolDelegates.h
--- a/Engine/Source/Editor/EditorToolEvents/Public/EditorToolDelegates.h
+++ b/Engine/Source/Editor/EditorToolEvents/Public/EditorToolDelegates.h
@@ -5,10 +5,7 @@
  */
 
 #pragma once
-
 #include "Delegates/Delegate.h"
-#include "Dom/JsonObject.h"
-#include "Templates/SharedPointer.h"
 
 struct FEditorToolDelegates
 {
@@ -18,7 +15,5 @@ public:
 	DECLARE_MULTICAST_DELEGATE_OneParam(FEditorToolStarted, const FString &/*class name*/);
 	static EDITORTOOLEVENTS_API FEditorToolStarted OnEditorToolStarted;
 	
-	/**Broadcasted when the editor is closed and sends a summary of widget focus data */
-	DECLARE_MULTICAST_DELEGATE_OneParam(FSendEditorFocusSummary, TSharedPtr<FJsonObject>/*json summary of focus data*/);
-	static EDITORTOOLEVENTS_API FSendEditorFocusSummary OnSendEditorFocusSummary;
+	/*More delegates can be added, starting minimal so only events we use are added*/
 };
diff --git a/Engine/Source/Editor/EditorToolEvents/Public/EditorUsageAnalyticsSummary.h b/Engine/Source/Editor/EditorToolEvents/Public/EditorUsageAnalyticsSummary.h
deleted file mode 100644
diff --git a/Engine/Source/Runtime/Engine/Engine.Build.cs b/Engine/Source/Runtime/Engine/Engine.Build.cs
--- a/Engine/Source/Runtime/Engine/Engine.Build.cs
+++ b/Engine/Source/Runtime/Engine/Engine.Build.cs
@@ -72,7 +72,7 @@ public class Engine : ModuleRules
 			PrivateDependencyModuleNames.AddRange(
 				new string[]
 				{
-					"EditorAnalyticsSession"
+					"EditorAnalyticsSession",
 				}
 			);
 		}
@@ -219,9 +219,6 @@ public class Engine : ModuleRules
 					"Zen"
 				}
 			);
-
-			// editor telemetry events
-			PrivateDependencyModuleNames.Add("EditorToolEvents");
 		}
 
 		PrivateIncludePathModuleNames.AddRange(
diff --git a/Engine/Source/Runtime/Engine/Private/Analytics/EditorAnalyticsSessionSummary.cpp b/Engine/Source/Runtime/Engine/Private/Analytics/EditorAnalyticsSessionSummary.cpp
--- a/Engine/Source/Runtime/Engine/Private/Analytics/EditorAnalyticsSessionSummary.cpp
+++ b/Engine/Source/Runtime/Engine/Private/Analytics/EditorAnalyticsSessionSummary.cpp
@@ -13,8 +13,6 @@
 #include "Framework/Docking/SDockingTabStack.h"
 #include "Widgets/Docking/SDockTab.h"
 
-#include "EditorUsageAnalyticsSummary.h"
-
 #include "Interfaces/IAnalyticsPropertyStore.h"
 #include "Kismet2/DebuggerCommands.h"
 #include "ProfilingDebugging/StallDetector.h"
@@ -56,11 +54,11 @@ namespace EditorAnalyticsCVars
 		TEXT("Disables collecting focused widget times for analytics summary."));
 }
 
-
 FEditorAnalyticsSessionSummary::FEditorAnalyticsSessionSummary(TSharedPtr<IAnalyticsPropertyStore> Store, uint32 MonitorProcessId)
 	: FEngineAnalyticsSessionSummary(MoveTemp(Store), MonitorProcessId)
 	, LastUserActivityTimeSecs(FPlatformTime::Seconds())
 	, AccountedUserIdleSecs(0)
+	, TimeOfLastUpdate(FPlatformTime::Seconds())
 {
 	EditorAnalyticsProperties::UserInteractionCount.Set(GetStore(), 0);
 	EditorAnalyticsProperties::Idle1Min.Set(GetStore(), 0);
@@ -100,21 +98,10 @@ FEditorAnalyticsSessionSummary::FEditorAnalyticsSessionSummary(TSharedPtr<IAnaly
 	FSlateApplication::Get().GetOnModalLoopTickEvent().AddRaw(this, &FEditorAnalyticsSessionSummary::Tick);
 	FSlateApplication::Get().GetLastUserInteractionTimeUpdateEvent().AddRaw(this, &FEditorAnalyticsSessionSummary::OnSlateUserInteraction);
 	UPackage::PackageDirtyStateChangedEvent.AddRaw(this, &FEditorAnalyticsSessionSummary::OnDirtyPackageStateChanged);
-
-	if (!EditorAnalyticsCVars::CVarDisableFocusTimes.GetValueOnAnyThread())
-	{
-		// we want to make sure an instance is created so the usage summary starts listening to events
-		FEditorUsageAnalyticsSummary::Get().Start();
-	}
 }
 
 void FEditorAnalyticsSessionSummary::ShutdownInternal()
 {
-	if (!EditorAnalyticsCVars::CVarDisableFocusTimes.GetValueOnAnyThread())
-	{
-		FEditorUsageAnalyticsSummary::Get().Shutdown();
-	}
-
 	FEditorDelegates::PreBeginPIE.RemoveAll(this);
 	FEditorDelegates::EndPIE.RemoveAll(this);
 	FSlateApplication::Get().GetOnModalLoopTickEvent().RemoveAll(this);
@@ -158,9 +145,106 @@ bool FEditorAnalyticsSessionSummary::UpdateSessionProgressInternal(bool bCrashin
 	EditorAnalyticsProperties::ShaderDDCHits.Set(GetStore(), GShaderCompilerStats->GetDDCHits());
 	EditorAnalyticsProperties::TimeShaderCompilationWasActive.Set(GetStore(), GShaderCompilerStats->GetTimeShaderCompilationWasActive());
 
+	if (!EditorAnalyticsCVars::CVarDisableFocusTimes.GetValueOnAnyThread())
+	{
+		UpdateUserFocusTimes(FPlatformTime::Seconds());
+	}
+
 	return bShouldPersist;
 }
 
+namespace EditorAnalyticsSessionSummaryHelpers
+{
+	bool GetFocusedTabIdentifierIfInFocusPath(TSharedPtr<const FSlateUser> SlateUser, FString& OutTabIdentifier)
+	{
+		if (!SlateUser.IsValid() || !SlateUser->HasValidFocusPath())
+		{
+			return false;
+		}
+
+		const FWeakWidgetPath& WeakFocusPath = SlateUser->GetWeakFocusPath();
+
+		// intentionally don't check 0 index as that'll be the window
+		for (int i = WeakFocusPath.Widgets.Num() - 1; i > 0; --i)
+		{
+			TWeakPtr<SWidget> WidgetInPathWeakPtr = WeakFocusPath.Widgets[i];
+
+			if (!WidgetInPathWeakPtr.IsValid())
+			{
+				continue;
+			}
+
+			TSharedPtr<SWidget> WidgetInPath = WidgetInPathWeakPtr.Pin();
+
+			// docking tab stack is the part of the docking stuff that shows up in focus paths
+			if (WidgetInPath->GetTypeAsString() == TEXT("SDockingTabStack"))
+			{
+				SDockingTabStack* TabStackWidget = reinterpret_cast<SDockingTabStack*>(WidgetInPath.Get());
+
+				if (TabStackWidget != nullptr)
+				{
+					TSharedPtr<SDockTab> TabWidget = TabStackWidget->GetSelectedTabFromWell();
+
+					if (TabWidget.IsValid())
+					{
+						OutTabIdentifier = TabWidget->GetLayoutIdentifier().ToString();
+
+						return true;
+					}
+
+				}
+			}
+		}
+
+		return false;
+	}
+}
+
+void FEditorAnalyticsSessionSummary::UpdateUserFocusTimes(double CurrTimeSecs)
+{
+	// This can get called from a console control handler thread. Getting the SlateApplication is only valid on certain threads.
+	if (!IsInGameThread() && !IsInSlateThread() && !IsInAsyncLoadingThread() && !IsInParallelLoadingThread())
+	{
+		return;
+	}
+
+	// Make sure we have a slate application to query
+	if (!FSlateApplication::IsInitialized())
+	{
+		return;
+	}
+
+	FSlateApplication& SlateApplication = FSlateApplication::Get();
+
+	TSharedPtr<const FSlateUser> CurrentUser = SlateApplication.GetCursorUser();
+
+	double PrevTimeOfLastUpdate = TimeOfLastUpdate.load();
+
+	// similar to our other timer updates, we want to ensure only a single thread performs this update
+	if (TimeOfLastUpdate.compare_exchange_strong(PrevTimeOfLastUpdate, CurrTimeSecs))
+	{
+		FString TabIdentifier;
+		if (EditorAnalyticsSessionSummaryHelpers::GetFocusedTabIdentifierIfInFocusPath(CurrentUser, TabIdentifier))
+		{
+			TabIdentifier = FString(TEXT("FocusTime")).Append(TabIdentifier);
+			float DeltaAsFloat = static_cast<float>(CurrTimeSecs - PrevTimeOfLastUpdate);
+
+			if (!GetStore()->Contains(TabIdentifier))
+			{
+				GetStore()->Set(TabIdentifier, DeltaAsFloat);
+			}
+			else
+			{
+				GetStore()->Update(TabIdentifier, [DeltaAsFloat](float& InOutValue)
+					{
+						InOutValue += DeltaAsFloat;
+						return true;
+					});
+			}
+		}
+	}
+}
+
 bool FEditorAnalyticsSessionSummary::UpdateUserIdleTime(double CurrTimeSecs, bool bReset)
 {
 	bool bSessionUpdated = false;
diff --git a/Engine/Source/Runtime/Engine/Private/Analytics/EditorAnalyticsSessionSummary.h b/Engine/Source/Runtime/Engine/Private/Analytics/EditorAnalyticsSessionSummary.h
--- a/Engine/Source/Runtime/Engine/Private/Analytics/EditorAnalyticsSessionSummary.h
+++ b/Engine/Source/Runtime/Engine/Private/Analytics/EditorAnalyticsSessionSummary.h
@@ -26,6 +26,7 @@ protected:
 
 private:
 	bool UpdateUserIdleTime(double CurrTimeSecs, bool bReset);
+	void UpdateUserFocusTimes(double CurrTimeSecs);
 	void OnSlateUserInteraction(double CurrSlateInteractionTime);
 	void OnEnterPIE(const bool /*bIsSimulating*/);
 	void OnExitPIE(const bool /*bIsSimulating*/);
@@ -37,6 +38,9 @@ private:
 
 	/** The number of idle seconds in the current idle sequence that were accounted (saved in the session) for the user idle counters. */
 	std::atomic<double> AccountedUserIdleSecs;
+
+	/** Time since we last updated our focus stats */
+	std::atomic<double> TimeOfLastUpdate;
 };
 
 #endif // WITH_EDITOR
