UWorldPartitionRuntimeCellTransformerISM: Reparent child components when the RootComponent is being transformed.
#Synced-CL 0

diff --git a/Engine/Source/Runtime/Engine/Private/WorldPartition/WorldPartitionRuntimeCellTransformerISM.cpp b/Engine/Source/Runtime/Engine/Private/WorldPartition/WorldPartitionRuntimeCellTransformerISM.cpp
--- a/Engine/Source/Runtime/Engine/Private/WorldPartition/WorldPartitionRuntimeCellTransformerISM.cpp
+++ b/Engine/Source/Runtime/Engine/Private/WorldPartition/WorldPartitionRuntimeCellTransformerISM.cpp
@@ -156,7 +156,16 @@ void UWorldPartitionRuntimeCellTransformerISM::Transform(ULevel* InLevel)
 			{
 				USceneComponent* NewRootComponent = NewObject<USceneComponent>(*Actor);
 				NewRootComponent->SetRelativeTransform(OldRootComponent->GetRelativeTransform());
+				NewRootComponent->SetMobility(OldRootComponent->GetMobility());
 				(*Actor)->SetRootComponent(NewRootComponent);
+
+				(*Actor)->ForEachComponent<USceneComponent>(false, [OldRootComponent, NewRootComponent](USceneComponent* Component)
+				{
+					if (Component->GetAttachParent() == OldRootComponent)
+					{
+						Component->SetupAttachment(NewRootComponent);
+					}
+				});
 			}
 		}
 	}
@@ -190,7 +199,6 @@ void UWorldPartitionRuntimeCellTransformerISM::Transform(ULevel* InLevel)
 		InLevel->Actors.Add(PackedActor);
 	}
 }
-
 bool UWorldPartitionRuntimeCellTransformerISM::CanAutoInstanceActor(AActor* InActor) const
 {
 	if (InActor->ActorHasTag(NAME_CellTransformerIgnoreActor))
