Exposed enum EStreamingSourceTargetState and UWorldPartitionStreamingSourceComponent's TargetState to Blueprint and made TargetState property public to allow changing state dynamically on a UWorldPartitionStreamingSourceComponent.

#Synced-CL 49040969

diff --git a/Engine/Source/Runtime/Engine/Classes/Components/WorldPartitionStreamingSourceComponent.h b/Engine/Source/Runtime/Engine/Classes/Components/WorldPartitionStreamingSourceComponent.h
--- a/Engine/Source/Runtime/Engine/Classes/Components/WorldPartitionStreamingSourceComponent.h
+++ b/Engine/Source/Runtime/Engine/Classes/Components/WorldPartitionStreamingSourceComponent.h
@@ -81,11 +81,11 @@ class UWorldPartitionStreamingSourceComponent : public UActorComponent, public I
 	UPROPERTY(BlueprintReadWrite, EditAnywhere, Category = "Streaming")
 	EStreamingSourcePriority Priority;
 
+	UPROPERTY(BlueprintReadWrite, EditAnywhere, Category = "Streaming")
+	EStreamingSourceTargetState TargetState;
+
 private:
 	/** Whether this component is enabled or not */
 	UPROPERTY(EditAnywhere, Interp, Category = "Streaming")
 	bool bStreamingSourceEnabled;
-
-	UPROPERTY(EditAnywhere, Category = "Streaming")
-	EStreamingSourceTargetState TargetState;
 };
diff --git a/Engine/Source/Runtime/Engine/Public/WorldPartition/WorldPartitionStreamingSource.h b/Engine/Source/Runtime/Engine/Public/WorldPartition/WorldPartitionStreamingSource.h
--- a/Engine/Source/Runtime/Engine/Public/WorldPartition/WorldPartitionStreamingSource.h
+++ b/Engine/Source/Runtime/Engine/Public/WorldPartition/WorldPartitionStreamingSource.h
@@ -213,7 +213,7 @@ public:
 /**
  * Streaming Source Target State
  */
-UENUM()
+UENUM(BlueprintType)
 enum class EStreamingSourceTargetState : uint8
 {
 	Loaded,
