Fixed invalid bounds for actors with property overrides inside a Level Instance.

#Synced-CL 49040969

diff --git a/Engine/Source/Runtime/Engine/Classes/Components/SceneComponent.h b/Engine/Source/Runtime/Engine/Classes/Components/SceneComponent.h
--- a/Engine/Source/Runtime/Engine/Classes/Components/SceneComponent.h
+++ b/Engine/Source/Runtime/Engine/Classes/Components/SceneComponent.h
@@ -105,11 +105,11 @@ private:
 	TWeakObjectPtr<class APhysicsVolume> PhysicsVolume;
 
 	/** What we are currently attached to. If valid, RelativeLocation etc. are used relative to this object */
-	UPROPERTY(ReplicatedUsing = OnRep_AttachParent)
+	UPROPERTY(ReplicatedUsing = OnRep_AttachParent, meta=(DisableLevelInstancePropertyOverride))
 	TObjectPtr<USceneComponent> AttachParent;
 
 	/** Optional socket name on AttachParent that we are attached to. */
-	UPROPERTY(ReplicatedUsing = OnRep_AttachSocketName)
+	UPROPERTY(ReplicatedUsing = OnRep_AttachSocketName, meta = (DisableLevelInstancePropertyOverride))
 	FName AttachSocketName;
 
 	FName NetOldAttachSocketName;
diff --git a/Engine/Source/Runtime/Engine/Private/WorldPartition/ActorDescContainer.cpp b/Engine/Source/Runtime/Engine/Private/WorldPartition/ActorDescContainer.cpp
--- a/Engine/Source/Runtime/Engine/Private/WorldPartition/ActorDescContainer.cpp
+++ b/Engine/Source/Runtime/Engine/Private/WorldPartition/ActorDescContainer.cpp
@@ -703,40 +703,51 @@ void UActorDescContainer::RemoveChildActorFromParentMap(FWorldPartitionActorDesc
 	}
 }
 
-void UActorDescContainer::PropagateActorToWorldUpdate(FWorldPartitionActorDesc* ActorDesc)
+void UActorDescContainer::ForEachChildActorDesc(const FGuid& InActorGuid, TFunctionRef<void(const FGuid&)> InFunc)
 {
 	TSet<FGuid> CycleDetector;
-	PropagateActorToWorldUpdateInternal(ActorDesc, CycleDetector);
-}
 
-void UActorDescContainer::PropagateActorToWorldUpdateInternal(FWorldPartitionActorDesc* ActorDesc, TSet<FGuid>& CycleDetector)
-{
-	bool bAlreadyVisited = false;
-	CycleDetector.Add(ActorDesc->GetGuid(), &bAlreadyVisited);
-	if (bAlreadyVisited)
+	TFunction<void(const FGuid&)> ForEachChildActorDescInternal = [&ForEachChildActorDescInternal, Func = InFunc, &CycleDetector, this](const FGuid& ActorGuid)
 	{
-		// Each actor can have at most one parent, which means that it can appear on at most one Children list in ParentActorToChildrenMap.
-		// Because of that, reaching the same actor twice, in one series of PropagateActorToWorldUpdate calls, means that we found a cycle.
-		// That can happen when the actor currently being saved becomes a descendant of one of its ex-descendants
-		// and that ex-descendant actor is not saved yet. In that case we can break the cycle here and stop propagating updates. 
-		// They'll be properly propagated once the "ex-descendant" actor mentioned above gets saved.
-		return;
-	}
+		bool bAlreadyVisited = false;
+		CycleDetector.Add(ActorGuid, &bAlreadyVisited);
+		if (bAlreadyVisited)
+		{
+			// Each actor can have at most one parent, which means that it can appear on at most one Children list in ParentActorToChildrenMap.
+			// Because of that, reaching the same actor twice, in one series of PropagateActorToWorldUpdate calls, means that we found a cycle.
+			// That can happen when the actor currently being saved becomes a descendant of one of its ex-descendants
+			// and that ex-descendant actor is not saved yet. In that case we can break the cycle here and stop propagating updates. 
+			// They'll be properly propagated once the "ex-descendant" actor mentioned above gets saved.
+			return;
+		}
 
-	if (const TSet<FGuid>* Children = ParentActorToChildrenMap.Find(ActorDesc->GetGuid()))
-	{
-		for (const FGuid& ChildGuid : *Children)
+		if (const TSet<FGuid>* Children = ParentActorToChildrenMap.Find(ActorGuid))
 		{
-			if (FWorldPartitionActorDesc* ChildActor = GetActorDesc(ChildGuid))
+			for (const FGuid& ChildGuid : *Children)
 			{
-				OnActorDescUpdating(ChildActor);
-				ChildActor->UpdateActorToWorld();
-				OnActorDescUpdated(ChildActor);
-
-				PropagateActorToWorldUpdateInternal(ChildActor, CycleDetector);
+				if (ensure(ChildGuid.IsValid()))
+				{
+					Func(ChildGuid);
+					ForEachChildActorDescInternal(ChildGuid);
+				}
 			}
 		}
-	}
+	};
+
+	ForEachChildActorDescInternal(InActorGuid);
+}
+
+void UActorDescContainer::PropagateActorToWorldUpdate(FWorldPartitionActorDesc* ActorDesc)
+{
+	ForEachChildActorDesc(ActorDesc->GetGuid(), [this](const FGuid& ChildActorGuid)
+	{
+		if (FWorldPartitionActorDesc* ChildActor = GetActorDesc(ChildActorGuid))
+		{
+			OnActorDescUpdating(ChildActor);
+			ChildActor->UpdateActorToWorld();
+			OnActorDescUpdated(ChildActor);
+		}
+	});
 }
 #endif
 
diff --git a/Engine/Source/Runtime/Engine/Private/WorldPartition/LevelInstance/LevelInstancePropertyOverrideContainer.cpp b/Engine/Source/Runtime/Engine/Private/WorldPartition/LevelInstance/LevelInstancePropertyOverrideContainer.cpp
--- a/Engine/Source/Runtime/Engine/Private/WorldPartition/LevelInstance/LevelInstancePropertyOverrideContainer.cpp
+++ b/Engine/Source/Runtime/Engine/Private/WorldPartition/LevelInstance/LevelInstancePropertyOverrideContainer.cpp
@@ -30,6 +30,108 @@ void ULevelInstancePropertyOverrideContainer::Initialize(const FInitializeParams
 	ContentBundleGuid = GetBaseContainer()->GetContentBundleGuid();
 	ExternalDataLayerAsset = GetBaseContainer()->GetExternalDataLayerAsset();
 
+	auto CloneActorDesc = [this](FWorldPartitionActorDesc* InActorDesc, const FActorContainerPath& InContainerPath)
+	{
+		if (!GetOverrideActorDesc(InActorDesc->GetGuid(), InContainerPath))
+		{
+			// Clone actordesc
+			UClass* ActorNativeClass = InActorDesc->GetActorNativeClass();
+			if (ensure(ActorNativeClass))
+			{
+				TUniquePtr<FWorldPartitionActorDesc> NewActorDesc(AActor::StaticCreateClassActorDesc(ActorNativeClass));
+				FWorldPartitionActorDescInitData ActorDescInitData = FWorldPartitionActorDescInitData()
+					.SetNativeClass(ActorNativeClass)
+					.SetPackageName(InActorDesc->GetActorPackage())
+					.SetActorPath(InActorDesc->GetActorSoftPath());
+				InActorDesc->SerializeTo(ActorDescInitData.GetSerializedData());
+				NewActorDesc->Init(ActorDescInitData);
+				check(NewActorDesc->Equals(InActorDesc));
+
+				// Set Container
+				check(!NewActorDesc->GetContainer());
+				NewActorDesc->SetContainer(this);
+
+				TMap<FGuid, TUniquePtr<FWorldPartitionActorDesc>>& TransientOverridesMap = PerContainerTransientOverrideActorDescs.FindOrAdd(InContainerPath);
+				TransientOverridesMap.Emplace(NewActorDesc->GetGuid(), MoveTemp(NewActorDesc));
+			}
+		}
+	};
+
+	TMap<FActorContainerPath, UActorDescContainer*> BaseContainers;
+	for (const auto& ActorDescsPerContainerElement : PropertyOverrideDesc->ActorDescsPerContainer)
+	{
+		const FActorContainerPath& ContainerPath = ActorDescsPerContainerElement.Key;
+		UActorDescContainer* BaseContainer = const_cast<UActorDescContainer *>(PropertyOverrideDesc->GetBaseContainer(GetBaseContainer(), ContainerPath));
+		if (ensure(BaseContainer))
+		{
+			check(BaseContainer->IsInitialized());
+			BaseContainers.Emplace(ContainerPath, BaseContainer);
+
+			for (const auto& ActorDescsMapElement : ActorDescsPerContainerElement.Value)
+			{
+				const TSharedPtr<FWorldPartitionActorDesc>& OverrideActorDesc = ActorDescsMapElement.Value;
+				FWorldPartitionActorDesc* BaseActorDesc = BaseContainer->GetActorDesc(OverrideActorDesc->GetGuid());
+				if (ensure(BaseActorDesc))
+				{
+					// Check if override actordesc actor transform is different
+					if (!OverrideActorDesc->GetActorTransformRelative().Equals(BaseActorDesc->GetActorTransformRelative()))
+					{
+						// Create transient override actordescs for all children of this override actordesc (if they are not already part of the override actordesc). 
+						// This will allow mutate their world editor and runtime bounds using UpdateActorToWorld. 
+						// Note that GetOverrideActorDesc now also searches in the transient overrides.
+						BaseContainer->ForEachChildActorDesc(OverrideActorDesc->GetGuid(), [this, &CloneActorDesc, &ContainerPath, BaseContainer](const FGuid& ChildActorGuid)
+						{
+							if (FWorldPartitionActorDesc* ChildActorDesc = BaseContainer->GetActorDesc(ChildActorGuid))
+							{
+								CloneActorDesc(ChildActorDesc, ContainerPath);
+							}
+						});
+					}
+				}
+			}
+		}
+	}
+
+	auto UpdateActorToWorld = [this](const FGuid& InActorGuid, const FActorContainerPath& InContainerPath)
+	{
+		if (FWorldPartitionActorDesc* MutableActorDesc = GetOverrideActorDesc(InActorGuid, InContainerPath))
+		{
+			FTransform ParentTransform = FTransform::Identity;
+			if (MutableActorDesc->GetParentActor().IsValid())
+			{
+				const FWorldPartitionActorDesc* ParentActorDesc = GetOverrideActorDesc(MutableActorDesc->GetParentActor(), InContainerPath);
+				if (ensure(ParentActorDesc))
+				{
+					ParentTransform = ParentActorDesc->GetActorTransform();
+				}
+			}
+			// Always pass parent transform (even when Identity) because we don't want UpdateActorToWorld to resolve
+			// the parent by itself, as it misses the context (InContainerPath) to properly resolve the actordesc.
+			OnActorDescUpdating(MutableActorDesc);
+			MutableActorDesc->UpdateActorToWorld(&ParentTransform);
+			OnActorDescUpdated(MutableActorDesc);
+		}
+	};
+
+	for (const auto& BaseContainerElement : BaseContainers)
+	{
+		const FActorContainerPath& ContainerPath = BaseContainerElement.Key;
+		UActorDescContainer* BaseContainer = BaseContainerElement.Value;
+
+		// Update the world bounds of all overrides and transient overrides starting from the root actordescs.
+		for (FActorDescList::TIterator<> ActorDescIterator(BaseContainer); ActorDescIterator; ++ActorDescIterator)
+		{
+			FWorldPartitionActorDesc* ActorDesc = *ActorDescIterator;
+			if (!ActorDesc->GetParentActor().IsValid())
+			{
+				UpdateActorToWorld(ActorDesc->GetGuid(), ContainerPath);
+				BaseContainer->ForEachChildActorDesc(ActorDesc->GetGuid(), [this, &UpdateActorToWorld, &ContainerPath](const FGuid& ChildActorGuid)
+				{
+					UpdateActorToWorld(ChildActorGuid, ContainerPath);
+				});
+			}
+		}
+	}
 	bContainerInitialized = true;
 }
 
@@ -41,6 +143,7 @@ void ULevelInstancePropertyOverrideContainer::Uninitialize()
 		// nothing to do except unregister delegates this class is a proxy to the PropertyOverrideDesc Base Container + Override Descs
 		UnregisterBaseContainerDelegates();
 		PropertyOverrideDesc = nullptr;
+		PerContainerTransientOverrideActorDescs.Reset();
 		bContainerInitialized = false;
 	}
 }
@@ -193,7 +296,24 @@ const FWorldPartitionActorDesc* ULevelInstancePropertyOverrideContainer::GetActo
 FWorldPartitionActorDesc* ULevelInstancePropertyOverrideContainer::GetOverrideActorDesc(const FGuid& InActorGuid, const FActorContainerPath& InContainerPath) const
 {
 	check(PropertyOverrideDesc);
-	return PropertyOverrideDesc->GetOverrideActorDesc(InActorGuid, InContainerPath);
+	FWorldPartitionActorDesc* Desc = PropertyOverrideDesc->GetOverrideActorDesc(InActorGuid, InContainerPath);
+	if (!Desc)
+	{
+		Desc = GetTransientOverrideActorDesc(InActorGuid, InContainerPath);
+	}
+	return Desc;
+}
+
+FWorldPartitionActorDesc* ULevelInstancePropertyOverrideContainer::GetTransientOverrideActorDesc(const FGuid& InActorGuid, const FActorContainerPath& InContainerPath) const
+{
+	if (const TMap<FGuid, TUniquePtr<FWorldPartitionActorDesc>>* ActorDescs = PerContainerTransientOverrideActorDescs.Find(InContainerPath))
+	{
+		if (const TUniquePtr<FWorldPartitionActorDesc>* ActorDesc = ActorDescs->Find(InActorGuid))
+		{
+			return ActorDesc->Get();
+		}
+	}
+	return nullptr;
 }
 
 #endif
diff --git a/Engine/Source/Runtime/Engine/Private/WorldPartition/LevelInstance/LevelInstancePropertyOverrideContainer.h b/Engine/Source/Runtime/Engine/Private/WorldPartition/LevelInstance/LevelInstancePropertyOverrideContainer.h
--- a/Engine/Source/Runtime/Engine/Private/WorldPartition/LevelInstance/LevelInstancePropertyOverrideContainer.h
+++ b/Engine/Source/Runtime/Engine/Private/WorldPartition/LevelInstance/LevelInstancePropertyOverrideContainer.h
@@ -5,6 +5,7 @@
 #include "CoreFwd.h"
 #include "Misc/Optional.h"
 #include "WorldPartition/ActorDescContainer.h"
+#include "WorldPartition/WorldPartitionActorContainerID.h"
 #include "LevelInstancePropertyOverrideContainer.generated.h"
 
 class FLevelInstancePropertyOverrideDesc;
@@ -48,6 +49,7 @@ private:
 	friend class ULevelInstanceContainerInstance;
 
 	FWorldPartitionActorDesc* GetOverrideActorDesc(const FGuid& InActorGuid, const FActorContainerPath& InContainerPath = FActorContainerPath()) const;
+	FWorldPartitionActorDesc* GetTransientOverrideActorDesc(const FGuid& InActorGuid, const FActorContainerPath& InContainerPath = FActorContainerPath()) const;
 
 	void SetPropertyOverrideDesc(const TSharedPtr<FLevelInstancePropertyOverrideDesc>& InPropertyOverrideDesc);
 	const FLevelInstancePropertyOverrideDesc* GetPropertyOverrideDesc() const { return PropertyOverrideDesc.Get(); }
@@ -67,5 +69,7 @@ private:
 	UActorDescContainer* GetBaseContainer() const;
 
 	TSharedPtr<FLevelInstancePropertyOverrideDesc> PropertyOverrideDesc;
+
+	TMap<FActorContainerPath, TMap<FGuid, TUniquePtr<FWorldPartitionActorDesc>>> PerContainerTransientOverrideActorDescs;
 #endif
 };
diff --git a/Engine/Source/Runtime/Engine/Private/WorldPartition/WorldPartitionActorDesc.cpp b/Engine/Source/Runtime/Engine/Private/WorldPartition/WorldPartitionActorDesc.cpp
--- a/Engine/Source/Runtime/Engine/Private/WorldPartition/WorldPartitionActorDesc.cpp
+++ b/Engine/Source/Runtime/Engine/Private/WorldPartition/WorldPartitionActorDesc.cpp
@@ -1065,10 +1065,11 @@ bool FWorldPartitionActorDesc::IsRuntimeRelevant(const FWorldPartitionActorDescI
 	return !IsMainWorldOnly() || InActorDescInstance->GetContainerInstance()->GetContainerID().IsMainContainer();
 }
 
-void FWorldPartitionActorDesc::UpdateActorToWorld()
+void FWorldPartitionActorDesc::UpdateActorToWorld(FTransform* InParentTransform)
 {
-	FTransform ParentTransform = FTransform::Identity;
-	if (ParentActor.IsValid())
+	// If set, use provided parent transform, else get it from ParentActor
+	FTransform ParentTransform = InParentTransform ? *InParentTransform : FTransform::Identity;
+	if (ParentActor.IsValid() && !InParentTransform)
 	{
 		if (FWorldPartitionActorDesc* ParentActorDesc = Container->GetActorDesc(ParentActor))
 		{
diff --git a/Engine/Source/Runtime/Engine/Public/WorldPartition/ActorDescContainer.h b/Engine/Source/Runtime/Engine/Public/WorldPartition/ActorDescContainer.h
--- a/Engine/Source/Runtime/Engine/Public/WorldPartition/ActorDescContainer.h
+++ b/Engine/Source/Runtime/Engine/Public/WorldPartition/ActorDescContainer.h
@@ -147,6 +147,8 @@ public:
 	void RegisterActorDescriptor(FWorldPartitionActorDesc* ActorDesc, UWorld* InWorldContext) {}
 	//~ End Deprecation
 
+	void ForEachChildActorDesc(const FGuid& InActorGuid, TFunctionRef<void(const FGuid&)> InFunc);
+
 protected:
 	FNameActorDescMap ActorsByName;
 
@@ -174,7 +176,6 @@ private:
 	void RemoveChildActorFromParentMap(FWorldPartitionActorDesc* ActorDesc);
 
 	void PropagateActorToWorldUpdate(FWorldPartitionActorDesc* ActorDesc);
-	void PropagateActorToWorldUpdateInternal(FWorldPartitionActorDesc* ActorDesc, TSet<FGuid>& CycleDetector);
 
 	static FUObjectAnnotationSparse<FDeletedObjectPlaceholderAnnotation, true> DeletedObjectPlaceholdersAnnotation;
 
diff --git a/Engine/Source/Runtime/Engine/Public/WorldPartition/WorldPartitionActorDesc.h b/Engine/Source/Runtime/Engine/Public/WorldPartition/WorldPartitionActorDesc.h
--- a/Engine/Source/Runtime/Engine/Public/WorldPartition/WorldPartitionActorDesc.h
+++ b/Engine/Source/Runtime/Engine/Public/WorldPartition/WorldPartitionActorDesc.h
@@ -342,6 +342,7 @@ public:
 	inline FName GetFolderPath() const { return FolderPath; }
 	inline const FGuid& GetFolderGuid() const { return FolderGuid; }
 	inline const FTransform& GetActorTransform() const { return ActorTransform; }
+	inline const FTransform& GetActorTransformRelative() const { return ActorTransformRelative; }
 
 	ENGINE_API FBox GetEditorBounds() const;
 	ENGINE_API FBox GetRuntimeBounds() const;
@@ -513,7 +514,7 @@ public:
 		Container = InContainer;
 	}
 
-	ENGINE_API virtual void UpdateActorToWorld();
+	ENGINE_API virtual void UpdateActorToWorld(FTransform* InParentTransform = nullptr);
 
 	enum class EToStringMode : uint8
 	{
