ISM ChangeSet Tracking Fix

Account for unused "padding" bits when clearing trailing elements.

#Synced-CL 49075355

diff --git a/Engine/Source/Runtime/Engine/Public/InstancedStaticMesh/InstanceAttributeTracker.h b/Engine/Source/Runtime/Engine/Public/InstancedStaticMesh/InstanceAttributeTracker.h
--- a/Engine/Source/Runtime/Engine/Public/InstancedStaticMesh/InstanceAttributeTracker.h
+++ b/Engine/Source/Runtime/Engine/Public/InstancedStaticMesh/InstanceAttributeTracker.h
@@ -88,6 +88,7 @@ public:
 	static constexpr uint32 BitsPerElement = sizeof(ElementType) * 8u;
 	static constexpr int32 MasksPerElement = BitsPerElement / uint32(EFlag::Num);
 	static constexpr ElementType AnyFlagMask = (ElementType(1U) << uint32(EFlag::Num)) - ElementType(1U);
+	static constexpr uint32 PaddingBits = BitsPerElement - MasksPerElement * uint32(EFlag::Num);
 
 	// Helper struct to convert the flag to its bitmask at comile time.
 	template <EFlag Flag>
@@ -571,7 +572,7 @@ private:
 		const int32 SubIndex = LastIndex % MasksPerElement;
 		if (SubIndex > 0)
 		{
-			const uint32 ElementSubShift = (MasksPerElement - (LastIndex % MasksPerElement)) * uint32(EFlag::Num);
+			const uint32 ElementSubShift = (MasksPerElement - (LastIndex % MasksPerElement)) * uint32(EFlag::Num) + PaddingBits;
 			// Load and clear the mask for this index
 			ElementType CurrentElement = Data[ElementIndex] & ((~ElementType(0)) >> ElementSubShift);
 			Data[ElementIndex] =  CurrentElement;
