Fix material parameter value vector cache reading from wrong indices array

#rnx
#Synced-CL 49324298

diff --git a/Engine/Source/Runtime/Engine/Classes/Components/MeshComponent.h b/Engine/Source/Runtime/Engine/Classes/Components/MeshComponent.h
--- a/Engine/Source/Runtime/Engine/Classes/Components/MeshComponent.h
+++ b/Engine/Source/Runtime/Engine/Classes/Components/MeshComponent.h
@@ -233,8 +233,6 @@ public:
 	virtual void GetDefaultMaterialSlotsOverlayMaterial(TArray<TObjectPtr<UMaterialInterface>>& OutMaterialSlotOverlayMaterials) const { return; }
 
 protected:
-	ENGINE_API void SetParameterValueOnMaterials(const FName ParameterName, const TFunction<void(UMaterialInstanceDynamic* DynamicMaterial)>& SetParameterFunction);
-
 	//~ Begin UObject Interface.
 	ENGINE_API virtual void BeginDestroy() override;
 	//~ End UObject Interface.
@@ -264,6 +262,10 @@ protected:
 		float ScalarParameterDefaultValue = 0.f;
 	};
 
+	ENGINE_API void SetParameterValueOnMaterials(const FName ParameterName
+		, const TFunction<void(UMaterialInstanceDynamic* DynamicMaterial)>& SetParameterFunction
+		, const TFunction<const TArray<int32>& (const FMaterialParameterCache* ParameterCache)>& GetMaterialIndicesFromCacheFunction);
+
 public:
 
 	/** Whether or not to cache material parameter to speed up setting scalar or vector value on materials */
diff --git a/Engine/Source/Runtime/Engine/Private/Components/MeshComponent.cpp b/Engine/Source/Runtime/Engine/Private/Components/MeshComponent.cpp
--- a/Engine/Source/Runtime/Engine/Private/Components/MeshComponent.cpp
+++ b/Engine/Source/Runtime/Engine/Private/Components/MeshComponent.cpp
@@ -479,7 +479,9 @@ TArray<class UMaterialInterface*> UMeshComponent::GetMaterials() const
 	return OutMaterials;
 }
 
-void UMeshComponent::SetParameterValueOnMaterials(const FName ParameterName, const TFunction<void(UMaterialInstanceDynamic* DynamicMaterial)>& SetParameterFunction)
+void UMeshComponent::SetParameterValueOnMaterials(const FName ParameterName
+												, const TFunction<void(UMaterialInstanceDynamic* DynamicMaterial)>& SetParameterFunction
+												, const TFunction<const TArray<int32>& (const FMaterialParameterCache* ParameterCache)>& GetMaterialIndicesFromCacheFunction)
 {
 	bool bDynamicSlotOverlayMaterialCreated = false;
 
@@ -539,7 +541,7 @@ void UMeshComponent::SetParameterValueOnMaterials(const FName ParameterName, con
 		auto GetMaterialAtIndex = [this](int32 MaterialIndex) { return GetMaterial(MaterialIndex); };
 		auto GetMaterialSlotOverlayMaterialAtIndex = [RetrievedMaterialSlotsOverlayMaterials](int32 MaterialIndex) { return RetrievedMaterialSlotsOverlayMaterials[MaterialIndex]; };
 
-		auto SetParameterValueOnMaterials_Internal = [&SetParameterFunction, ParameterName, this]
+		auto SetParameterValueOnMaterials_Internal = [&SetParameterFunction, &GetMaterialIndicesFromCacheFunction, ParameterName, this]
 			(FMaterialParameterCache* ParameterCache
 			, const TFunction<UMaterialInterface*(int32 MaterialIndex)>& GetMaterialFunction
 			, const TFunction<UMaterialInstanceDynamic* (UMaterialInterface* MaterialInterface, int32 MaterialIndex)>& CreateDynamicMaterialInstanceFunction
@@ -547,7 +549,7 @@ void UMeshComponent::SetParameterValueOnMaterials(const FName ParameterName, con
 			{
 				if (ParameterCache != nullptr)
 				{
-					const TArray<int32>& MaterialIndices = ParameterCache->ScalarParameterMaterialIndices;
+					const TArray<int32>& MaterialIndices = GetMaterialIndicesFromCacheFunction(ParameterCache);
 					// Loop over all the material indices and update set the parameter value on the corresponding materials		
 					for (int32 MaterialIndex : MaterialIndices)
 					{
@@ -582,7 +584,9 @@ void UMeshComponent::SetParameterValueOnMaterials(const FName ParameterName, con
 
 void UMeshComponent::SetScalarParameterValueOnMaterials(const FName ParameterName, const float ParameterValue)
 {
-	SetParameterValueOnMaterials(ParameterName, [ParameterName, ParameterValue](UMaterialInstanceDynamic* DynamicMaterial) { DynamicMaterial->SetScalarParameterValue(ParameterName, ParameterValue); });
+	SetParameterValueOnMaterials(ParameterName
+								, [ParameterName, ParameterValue](UMaterialInstanceDynamic* DynamicMaterial) { DynamicMaterial->SetScalarParameterValue(ParameterName, ParameterValue); }
+								, [](const FMaterialParameterCache* ParameterCache) -> const TArray<int32>& { return ParameterCache->ScalarParameterMaterialIndices; });
 }
 
 void UMeshComponent::SetVectorParameterValueOnMaterials(const FName ParameterName, const FVector ParameterValue)
@@ -592,7 +596,9 @@ void UMeshComponent::SetVectorParameterValueOnMaterials(const FName ParameterNam
 
 void UMeshComponent::SetColorParameterValueOnMaterials(const FName ParameterName, const FLinearColor ParameterValue)
 {
-	SetParameterValueOnMaterials(ParameterName, [ParameterName, ParameterValue](UMaterialInstanceDynamic* DynamicMaterial) { DynamicMaterial->SetVectorParameterValue(ParameterName, ParameterValue); });
+	SetParameterValueOnMaterials(ParameterName
+								, [ParameterName, ParameterValue](UMaterialInstanceDynamic* DynamicMaterial) { DynamicMaterial->SetVectorParameterValue(ParameterName, ParameterValue); }
+								, [](const FMaterialParameterCache* ParameterCache) -> const TArray<int32>& { return ParameterCache->VectorParameterMaterialIndices; });
 }
 
 void UMeshComponent::MarkCachedMaterialParameterNameIndicesDirty()
