Added SkyAtmosphereComponent race condition and stomp guards

diff --git a/Engine/Source/Runtime/Engine/Classes/Components/SkyAtmosphereComponent.h b/Engine/Source/Runtime/Engine/Classes/Components/SkyAtmosphereComponent.h
--- a/Engine/Source/Runtime/Engine/Classes/Components/SkyAtmosphereComponent.h
+++ b/Engine/Source/Runtime/Engine/Classes/Components/SkyAtmosphereComponent.h
@@ -12,6 +12,8 @@
 
 #include "SkyAtmosphereComponent.generated.h"
 
+DECLARE_LOG_CATEGORY_EXTERN(LogSkyAtmosphereComponent, Log, All)
+
 
 class FSkyAtmosphereSceneProxy;
 
@@ -39,6 +41,10 @@ enum class ESkyAtmosphereTransformMode : uint8
 	PlanetCenterAtComponentTransform,
 };
 
+// Temporary thread safety atomics and extra race condition logging to mitigate PLAY-108359
+// TODO: Remove once the root cause has been identified and fixed
+#define SKY_ATMOSPHERE_COMPONENT_CONCURRENCY_GUARDS 1
+
 /**
  * A component that represents a planet atmosphere material and simulates sky and light scattering within it.
  * @see https://docs.unrealengine.com/en-US/Engine/Actors/FogEffects/SkyAtmosphere/index.html
@@ -313,6 +319,11 @@ protected:
 	#if WITH_STATE_STREAM_ACTOR
 	FSkyAtmosphereHandle Handle;
 	#endif
+
+private:
+#if SKY_ATMOSPHERE_COMPONENT_CONCURRENCY_GUARDS
+	std::atomic<bool> bIsDestroying;
+#endif
 };
 
 
diff --git a/Engine/Source/Runtime/Engine/Private/Components/SkyAtmosphereComponent.cpp b/Engine/Source/Runtime/Engine/Private/Components/SkyAtmosphereComponent.cpp
--- a/Engine/Source/Runtime/Engine/Private/Components/SkyAtmosphereComponent.cpp
+++ b/Engine/Source/Runtime/Engine/Private/Components/SkyAtmosphereComponent.cpp
@@ -35,6 +35,17 @@
 constexpr bool UseSkyAtmosphereStateStream = true;
 #endif
 
+DEFINE_LOG_CATEGORY(LogSkyAtmosphereComponent)
+
+////////////////////////////////////////////////////////////////////////// Race Condition Diagnostics
+
+static bool GSkyAtmosphereRaceConditionFatalLog = true;
+FAutoConsoleVariableRef CVarSkyAtmosphereRaceConditionFatalLog(
+	TEXT("r.SkyAtmosphere.RaceConditionFatalLog"), 
+	GSkyAtmosphereRaceConditionFatalLog,
+	TEXT("Log race conditions as fatal errors for diagnostics purposes."),
+	ECVF_Default);
+
 
 /*=============================================================================
 	USkyAtmosphereComponent implementation.
@@ -89,6 +100,9 @@ FOverrideAtmosphericLight ConvertAtmosphericLightOverride(bool* Enabled, FVector
 USkyAtmosphereComponent::USkyAtmosphereComponent(const FObjectInitializer& ObjectInitializer)
 	: Super(ObjectInitializer)
 	, SkyAtmosphereSceneProxy(nullptr)
+#if SKY_ATMOSPHERE_COMPONENT_CONCURRENCY_GUARDS
+	, bIsDestroying(false)
+#endif
 {
 	// All distance here are in kilometer and scattering/absorptions coefficient in 1/kilometers.
 	const float EarthBottomRadius = 6360.0f;
@@ -204,6 +218,10 @@ void USkyAtmosphereComponent::CreateRenderState_Concurrent(FRegisterComponentCon
 	Super::CreateRenderState_Concurrent(Context);
 	// If one day we need to look up lightmass built data, lookup it up here using the guid from the correct MapBuildData.
 
+#if SKY_ATMOSPHERE_COMPONENT_CONCURRENCY_GUARDS
+	bIsDestroying.store(false);
+#endif
+
 	bool bHidden = false;
 #if WITH_EDITORONLY_DATA
 	bHidden = GetOwner() ? GetOwner()->bHiddenEdLevel : false;
@@ -241,12 +259,39 @@ void USkyAtmosphereComponent::CreateRenderState_Concurrent(FRegisterComponentCon
 
 void USkyAtmosphereComponent::SendRenderTransform_Concurrent()
 {
+#if SKY_ATMOSPHERE_COMPONENT_CONCURRENCY_GUARDS
+	if (bIsDestroying.load())
+	{
+		if (GSkyAtmosphereRaceConditionFatalLog)
+		{
+			UE_LOG(LogSkyAtmosphereComponent, Fatal, TEXT("SendRenderTransform_Concurrent attempted while DestroyRenderState_Concurrent is being executed"));
+		}
+		FDebug::DumpStackTraceToLog(ELogVerbosity::Warning);
+		ensureMsgf(0, TEXT("SendRenderTransform_Concurrent attempted while DestroyRenderState_Concurrent is being executed"));
+		return;
+	}
+#endif
+
 	Super::SendRenderTransform_Concurrent();
 	SendRenderTransformCommand();
 }
 
 void USkyAtmosphereComponent::DestroyRenderState_Concurrent()
 {
+#if SKY_ATMOSPHERE_COMPONENT_CONCURRENCY_GUARDS
+	bool Expected = false;
+	if (!bIsDestroying.compare_exchange_strong(Expected, true))
+	{
+		if (GSkyAtmosphereRaceConditionFatalLog)
+		{
+			UE_LOG(LogSkyAtmosphereComponent, Fatal, TEXT("DestroyRenderState_Concurrent called multiple times"));
+		}
+		FDebug::DumpStackTraceToLog(ELogVerbosity::Warning);
+		ensureMsgf(0, TEXT("DestroyRenderState_Concurrent called multiple times"));
+		return;
+	}
+#endif
+
 	Super::DestroyRenderState_Concurrent();
 
 	#if WITH_STATE_STREAM_ACTOR
diff --git a/Engine/Source/Runtime/Renderer/Private/SkyAtmosphereRendering.cpp b/Engine/Source/Runtime/Renderer/Private/SkyAtmosphereRendering.cpp
--- a/Engine/Source/Runtime/Renderer/Private/SkyAtmosphereRendering.cpp
+++ b/Engine/Source/Runtime/Renderer/Private/SkyAtmosphereRendering.cpp
@@ -25,6 +25,8 @@
 #include "PostProcess/PostProcessing.h" // IsPostProcessingWithAlphaChannelSupported
 #include "VirtualShadowMaps/VirtualShadowMapClipmap.h"
 
+DEFINE_LOG_CATEGORY(LogSkyAtmosphereRendering)
+
 
 //PRAGMA_DISABLE_OPTIMIZATION
 
@@ -231,6 +233,14 @@ static TAutoConsoleVariable<int32> CVarSkyAtmosphereAsyncCompute(
 	TEXT("SkyAtmosphere on async compute (default: false). When running on the async pipe, SkyAtmosphere LUT generation will overlap with the occlusion pass.\n When enabled, opaque shadow map sampling on atmosphere will be disabled to avoid race conditions because these textures can be created in an overlaping parallel thread."),
 	ECVF_RenderThreadSafe);
 
+////////////////////////////////////////////////////////////////////////// Memory Stomp Diagnostics
+
+static bool GSkyAtmosphereStompFatalLog = true;
+FAutoConsoleVariableRef CVarSkyAtmosphereStompFatalLog(
+	TEXT("r.SkyAtmosphere.StompFatalLog"), 
+	GSkyAtmosphereStompFatalLog,
+	TEXT("Log memory stomps as fatal errors for diagnostics purposes."),
+	ECVF_Default);
 
 DECLARE_GPU_STAT(SkyAtmosphereLUTs);
 DECLARE_GPU_STAT(SkyAtmosphere);
@@ -571,6 +581,9 @@ FSkyAtmosphereRenderSceneInfo::FSkyAtmosphereRenderSceneInfo(FSkyAtmosphereScene
 
 FSkyAtmosphereRenderSceneInfo::~FSkyAtmosphereRenderSceneInfo()
 {
+#if SKY_ATMOSPHERE_STOMP_GUARDS
+	MagicValue = DeletedMagicValue;
+#endif
 }
 
 void FSkyAtmosphereRenderSceneInfo::CreateDistantSkyLightLutBufferAndSRV(FRDGBuilder& GraphBuilder)
@@ -656,6 +669,28 @@ void FScene::RemoveSkyAtmosphere(FSkyAtmosphereSceneProxy* SkyAtmosphereScenePro
 			{
 				FPlatformAtomics::InterlockedDecrement(&Scene->NumUncachedStaticLightingInteractions);
 			}
+
+#if SKY_ATMOSPHERE_STOMP_GUARDS
+			if (SkyAtmosphereSceneProxy->RenderSceneInfo->WasDeleted())
+			{
+				if (GSkyAtmosphereStompFatalLog)
+				{
+					UE_LOG(LogSkyAtmosphereRendering, Fatal, TEXT("SkyAtmosphereSceneProxy's RenderSceneInfo was previously deleted. Use after free detected!"));
+				}
+				ensureMsgf(0, TEXT("SkyAtmosphereSceneProxy's RenderSceneInfo was previously deleted. Use after free detected!"));
+				return;
+			}
+			else if (!SkyAtmosphereSceneProxy->RenderSceneInfo->IsValid())
+			{
+				if (GSkyAtmosphereStompFatalLog)
+				{
+					UE_LOG(LogSkyAtmosphereRendering, Fatal, TEXT("SkyAtmosphereSceneProxy's RenderSceneInfo memory stomp detected!"));
+				}
+				ensureMsgf(0, TEXT("SkyAtmosphereSceneProxy's RenderSceneInfo memory stomp detected!"));
+				return;
+			}
+#endif
+
 			delete SkyAtmosphereSceneProxy->RenderSceneInfo;
 			Scene->SkyAtmosphereStack.RemoveSingle(SkyAtmosphereSceneProxy);
 
diff --git a/Engine/Source/Runtime/Renderer/Private/SkyAtmosphereRendering.h b/Engine/Source/Runtime/Renderer/Private/SkyAtmosphereRendering.h
--- a/Engine/Source/Runtime/Renderer/Private/SkyAtmosphereRendering.h
+++ b/Engine/Source/Runtime/Renderer/Private/SkyAtmosphereRendering.h
@@ -14,6 +14,12 @@
 #include "SceneView.h"
 #include "VirtualShadowMaps/VirtualShadowMapArray.h"
 
+// Temporary use-after-free detection and logging to mitigate PLAY-108359
+// TODO: Remove once the root cause has been identified and fixed
+#define SKY_ATMOSPHERE_STOMP_GUARDS 1
+
+DECLARE_LOG_CATEGORY_EXTERN(LogSkyAtmosphereRendering, Log, All)
+
 
 class FScene;
 class FViewInfo;
@@ -160,6 +166,20 @@ public:
 
 	TUniformBufferRef<FSkyAtmosphereInternalCommonParameters>& GetInternalCommonParametersUniformBuffer() { return InternalCommonParametersUniformBuffer; }
 
+#if SKY_ATMOSPHERE_STOMP_GUARDS
+private:
+	// stomp and use-after-free detection
+	static constexpr uint32 ValidMagicValue = 0xA7305FE1;
+	static constexpr uint32 DeletedMagicValue = 0xDEADBEEF;
+	uint32 MagicValue = ValidMagicValue;
+
+public:
+	// stomp detection
+	bool IsValid() const { return MagicValue == ValidMagicValue; }
+	// use-after-free detection
+	bool WasDeleted() const { return MagicValue == DeletedMagicValue; }
+#endif //SKY_ATMOSPHERE_STOMP_GUARDS
+
 private:
 
 	FSkyAtmosphereSceneProxy& SkyAtmosphereSceneProxy;
