Added slice index as an argument for canvases to support writing to different slices in a render target 2d array

#Synced-CL 49298335

diff --git a/Engine/Source/Runtime/Engine/Private/UserInterface/Canvas.cpp b/Engine/Source/Runtime/Engine/Private/UserInterface/Canvas.cpp
--- a/Engine/Source/Runtime/Engine/Private/UserInterface/Canvas.cpp
+++ b/Engine/Source/Runtime/Engine/Private/UserInterface/Canvas.cpp
@@ -576,11 +576,12 @@ bool FCanvasBatchedElementRenderItem::Render_GameThread(const FCanvas* Canvas, F
 	return bDirty;
 }
 
-FCanvasRenderContext::FCanvasRenderContext(FRDGBuilder& InGraphBuilder, const FRenderTarget* InRenderTarget, FIntRect InViewportRect, FIntRect InScissorRect, bool bScaledToRenderTarget)
+FCanvasRenderContext::FCanvasRenderContext(FRDGBuilder& InGraphBuilder, const FRenderTarget* InRenderTarget, FIntRect InViewportRect, FIntRect InScissorRect, bool bScaledToRenderTarget, int16 InSliceIndex)
 	: GraphBuilder(InGraphBuilder)
 	, RenderTarget(InRenderTarget->GetRenderTargetTexture(InGraphBuilder))
 	, ViewportRect(InViewportRect)
 	, ScissorRect(InScissorRect)
+	, SliceIndex(InSliceIndex)
 {
 	if (ViewportRect.Area() <= 0 || bScaledToRenderTarget)
 	{
@@ -599,11 +600,11 @@ FCanvasRenderThreadScope::~FCanvasRenderThreadScope()
 	RenderCommandFunctionArray* RenderCommandArray = RenderCommands;
 
 	ENQUEUE_RENDER_COMMAND(DispatchCanvasRenderCommands)(
-		[RenderCommandArray, RenderTarget = Canvas.RenderTarget, ViewRect = Canvas.ViewRect, ScissorRect = Canvas.ScissorRect, bScaledToRenderTarget = Canvas.bScaledToRenderTarget](FRHICommandListImmediate& RHICmdList) mutable
+		[RenderCommandArray, RenderTarget = Canvas.RenderTarget, ViewRect = Canvas.ViewRect, ScissorRect = Canvas.ScissorRect, bScaledToRenderTarget = Canvas.bScaledToRenderTarget, SliceIndex = Canvas.SliceIndex](FRHICommandListImmediate& RHICmdList) mutable
 	{
 		GetRendererModule().InitializeSystemTextures(RHICmdList);
 		FRDGBuilder GraphBuilder(RHICmdList, RDG_EVENT_NAME("CanvasRenderThreadScope"));
-		FCanvasRenderContext RenderContext(GraphBuilder, RenderTarget, ViewRect, ScissorRect, bScaledToRenderTarget);
+		FCanvasRenderContext RenderContext(GraphBuilder, RenderTarget, ViewRect, ScissorRect, bScaledToRenderTarget, SliceIndex);
 
 		for (uint32 Index = 0, Count = RenderCommandArray->Num(); Index < Count; ++Index)
 		{
@@ -618,7 +619,9 @@ FCanvasRenderThreadScope::~FCanvasRenderThreadScope()
 			PassParameters->RenderTargets[0] = FRenderTargetBinding(
 				RenderTarget->GetRenderTargetTexture(GraphBuilder),
 				ResolveTexture,
-				ERenderTargetLoadAction::ELoad
+				ERenderTargetLoadAction::ELoad,
+				/*MipIndex*/0,
+				SliceIndex
 			);
 
 			GraphBuilder.AddPass(RDG_EVENT_NAME("MSAAResolveCanvas"), PassParameters, ERDGPassFlags::Raster, [](FRDGAsyncTask, FRHICommandList&) {});
@@ -819,12 +822,18 @@ void FCanvas::Flush_RenderThread(FRDGBuilder& GraphBuilder, bool bForce)
 
 	GetRendererModule().InitializeSystemTextures(GraphBuilder.RHICmdList);
 
-	FCanvasRenderContext RenderContext(GraphBuilder, RenderTarget, ViewRect, ScissorRect, bScaledToRenderTarget);
+	FCanvasRenderContext RenderContext(GraphBuilder, RenderTarget, ViewRect, ScissorRect, bScaledToRenderTarget, SliceIndex);
 
 	// Set the RHI render target.
 	if (IsUsingInternalTexture())
 	{
-		AddClearRenderTargetPass(GraphBuilder, RenderContext.GetRenderTarget());
+		FRDGTextureClearInfo ClearInfo;
+		if (RenderContext.GetSliceIndex() != -1)
+		{
+			ClearInfo.FirstSliceIndex = RenderContext.GetSliceIndex();
+			ClearInfo.NumSlices = 1;
+		}
+		AddClearRenderTargetPass(GraphBuilder, RenderContext.GetRenderTarget(), ClearInfo);
 	}
 
 	// Disable depth test & writes
@@ -1067,7 +1076,7 @@ void FCanvas::Clear(const FLinearColor& ClearColor)
 {
 	FRenderTarget* CanvasRenderTarget = GetRenderTarget();
 	ENQUEUE_RENDER_COMMAND(ClearCommand)(
-		[ClearColor, CanvasRenderTarget](FRHICommandList& RHICmdList)
+		[ClearColor, CanvasRenderTarget, SliceIndex=SliceIndex](FRHICommandList& RHICmdList)
 		{
 			SCOPED_DRAW_EVENT(RHICmdList, CanvasClear);
 			if (CanvasRenderTarget)
@@ -1076,14 +1085,14 @@ void FCanvas::Clear(const FLinearColor& ClearColor)
 				if (CanvasRenderTarget->GetRenderTargetTexture() && CanvasRenderTarget->GetRenderTargetTexture()->GetClearBinding() == FClearValueBinding(ClearColor))
 				{
 					// do fast clear
-					FRHIRenderPassInfo RPInfo(CanvasRenderTarget->GetRenderTargetTexture(), ERenderTargetActions::Clear_Store);
+					FRHIRenderPassInfo RPInfo(CanvasRenderTarget->GetRenderTargetTexture(), ERenderTargetActions::Clear_Store, nullptr, /*MipIndex*/0, SliceIndex);
 					RHICmdList.Transition(FRHITransitionInfo(CanvasRenderTarget->GetRenderTargetTexture(), ERHIAccess::Unknown, ERHIAccess::RTV));
 					RHICmdList.BeginRenderPass(RPInfo, TEXT("ClearCanvas"));
 					RHICmdList.EndRenderPass();					
 				}
 				else if(CanvasRenderTarget->GetRenderTargetTexture())
 				{
-					FRHIRenderPassInfo RPInfo(CanvasRenderTarget->GetRenderTargetTexture(), ERenderTargetActions::DontLoad_Store);
+					FRHIRenderPassInfo RPInfo(CanvasRenderTarget->GetRenderTargetTexture(), ERenderTargetActions::DontLoad_Store, nullptr, /*MipIndex*/0, SliceIndex);
 					RHICmdList.Transition(FRHITransitionInfo(CanvasRenderTarget->GetRenderTargetTexture(), ERHIAccess::Unknown, ERHIAccess::RTV));
 					RHICmdList.BeginRenderPass(RPInfo, TEXT("ClearCanvas"));
 					RHICmdList.SetViewport(0, 0, 0.0f, CanvasRenderTarget->GetSizeXY().X, CanvasRenderTarget->GetSizeXY().Y, 1.0f);
@@ -2121,6 +2130,11 @@ EShaderPlatform FCanvas::GetShaderPlatform() const
 	return GShaderPlatformForFeatureLevel[FeatureLevel];
 }
 
+void FCanvas::SetSliceIndex(int16 InSliceIndex)
+{
+	SliceIndex = InSliceIndex;
+}
+
 void FScreenMessageWriter::DrawLine(const FText& Message, int32 X, const FLinearColor& Color)
 {
 	Canvas.DrawShadowedText((float)X, (float)Y, Message, GetStatsFont(), Color);
diff --git a/Engine/Source/Runtime/Engine/Public/CanvasRender.h b/Engine/Source/Runtime/Engine/Public/CanvasRender.h
--- a/Engine/Source/Runtime/Engine/Public/CanvasRender.h
+++ b/Engine/Source/Runtime/Engine/Public/CanvasRender.h
@@ -8,7 +8,7 @@
 class FCanvasRenderContext
 {
 public:
-	ENGINE_API FCanvasRenderContext(FRDGBuilder& InGraphBuilder, const FRenderTarget* RenderTarget, FIntRect InViewportRect, FIntRect InScissorRect, bool bScaledToRenderTarget);
+	ENGINE_API FCanvasRenderContext(FRDGBuilder& InGraphBuilder, const FRenderTarget* InRenderTarget, FIntRect InViewportRect, FIntRect InScissorRect, bool bScaledToRenderTarget, int16 InSliceIndex = -1);
 
 	template <typename ExecuteLambdaType, typename ParameterStructType>
 	void AddPass(FRDGEventName&& PassName, const ParameterStructType* PassParameters, ExecuteLambdaType&& ExecuteLambda)
@@ -31,7 +31,7 @@ public:
 	void AddPass(FRDGEventName&& PassName, ExecuteLambdaType&& ExecuteLambda)
 	{
 		FRenderTargetParameters* PassParameters = GraphBuilder.AllocParameters<FRenderTargetParameters>();
-		PassParameters->RenderTargets[0] = FRenderTargetBinding(RenderTarget, ERenderTargetLoadAction::ELoad);
+		PassParameters->RenderTargets[0] = FRenderTargetBinding(RenderTarget, ERenderTargetLoadAction::ELoad, /*MipIndex*/0, SliceIndex);
 		AddPass(Forward<FRDGEventName&&>(PassName), PassParameters, Forward<ExecuteLambdaType&&>(ExecuteLambda));
 	}
 
@@ -83,12 +83,18 @@ public:
 		return ScissorRect;
 	}
 
+	int16 GetSliceIndex() const
+	{
+		return SliceIndex;
+	}
+
 	FRDGBuilder& GraphBuilder;
 
 private:
 	FRDGTextureRef RenderTarget;
 	FIntRect ViewportRect;
 	FIntRect ScissorRect;
+	int16 SliceIndex;
 };
 
 class FCanvasRenderThreadScope
diff --git a/Engine/Source/Runtime/Engine/Public/CanvasTypes.h b/Engine/Source/Runtime/Engine/Public/CanvasTypes.h
--- a/Engine/Source/Runtime/Engine/Public/CanvasTypes.h
+++ b/Engine/Source/Runtime/Engine/Public/CanvasTypes.h
@@ -603,6 +603,14 @@ public:
 	inline FIntPoint GetParentCanvasSize() const { return ParentSize; }
 
 	float GetDPIScale() const { return  bStereoRendering ? 1.0f : DPIScale; }
+
+	int16 GetSliceIndex() const
+	{
+		return SliceIndex;
+	}
+
+	ENGINE_API void SetSliceIndex(int16 InSliceIndex);
+
 public:
 	/** Private class for handling word wrapping behavior. */
 	TSharedPtr<FCanvasWordWrapper> WordWrapper;
@@ -651,6 +659,8 @@ private:
 
 	float DPIScale;
 
+	int16 SliceIndex = -1;
+
 	/** 
 	* Shared construction function
 	*/
diff --git a/Engine/Source/Runtime/Renderer/Private/Renderer.cpp b/Engine/Source/Runtime/Renderer/Private/Renderer.cpp
--- a/Engine/Source/Runtime/Renderer/Private/Renderer.cpp
+++ b/Engine/Source/Runtime/Renderer/Private/Renderer.cpp
@@ -462,7 +462,7 @@ void FRendererModule::DrawTileMesh(FCanvasRenderContext& RenderContext, FMeshPas
 		RDG_EVENT_SCOPE(GraphBuilder, "DrawTileMesh");
 
 		auto* PassParameters = GraphBuilder.AllocParameters<FDrawTileMeshPassParameters>();
-		PassParameters->RenderTargets[0] = FRenderTargetBinding(RenderContext.GetRenderTarget(), ERenderTargetLoadAction::ELoad);
+		PassParameters->RenderTargets[0] = FRenderTargetBinding(RenderContext.GetRenderTarget(), ERenderTargetLoadAction::ELoad, /*MipIndex*/0, RenderContext.GetSliceIndex());
 		PassParameters->View = View.GetShaderParameters();
 		PassParameters->InstanceCullingDrawParams.Scene = SceneUniforms.GetBuffer(GraphBuilder);
 		PassParameters->InstanceCullingDrawParams.InstanceCulling = FInstanceCullingContext::CreateDummyInstanceCullingUniformBuffer(GraphBuilder);
