Keep the face remap on the chaos trimesh even if the corresponding setting is disabled, because some systems still rely on the mapping being there
#Synced-CL 49335258

diff --git a/Engine/Source/Runtime/Engine/Private/PhysicsEngine/Experimental/ChaosCooking.cpp b/Engine/Source/Runtime/Engine/Private/PhysicsEngine/Experimental/ChaosCooking.cpp
--- a/Engine/Source/Runtime/Engine/Private/PhysicsEngine/Experimental/ChaosCooking.cpp
+++ b/Engine/Source/Runtime/Engine/Private/PhysicsEngine/Experimental/ChaosCooking.cpp
@@ -324,7 +324,8 @@ namespace Chaos
 				}
 
 				
-				TUniquePtr<TArray<int32>> OutFaceRemapPtr = InParams.bSupportFaceRemap ? MakeUnique<TArray<int32>>(OutFaceRemap) : nullptr;
+				// TODO: If InParams.bSupportFaceRemap is false, we should in theory not need the remap array, but in practice things still rely on it being there
+				TUniquePtr<TArray<int32>> OutFaceRemapPtr = MakeUnique<TArray<int32>>(OutFaceRemap);
 				TUniquePtr<TArray<int32>> OutVertexRemapPtr = (InParams.bSupportVertexRemap && Chaos::TriMeshPerPolySupport) ? MakeUnique<TArray<int32>>(OutVertexRemap) : nullptr;
 				if (!OutFaceRemapPtr)
 				{
@@ -342,7 +343,7 @@ namespace Chaos
 				OutTriangleMeshes.Emplace(TriangleMesh);
 
 				// Propagate remapped indices from the FTriangleMeshImplicitObject back to the remap array
-				if (InParams.bSupportFaceRemap)
+				// TODO: If we use the bSupportFaceRemap variable to disable the face remap, then we can use this if-statement: if (InParams.bSupportFaceRemap)
 				{
 					CopyUpdatedFaceRemapFromTriangleMesh(*TriangleMesh, OutFaceRemap);
 				}
