Set "r.Material.StripUnusedDefaultTextures" to false by default. Try testing in Horde builds to see if it solves validation errors.

#Synced-CL 49241023

diff --git a/Engine/Source/Runtime/Engine/Private/Materials/MaterialInterface.cpp b/Engine/Source/Runtime/Engine/Private/Materials/MaterialInterface.cpp
--- a/Engine/Source/Runtime/Engine/Private/Materials/MaterialInterface.cpp
+++ b/Engine/Source/Runtime/Engine/Private/Materials/MaterialInterface.cpp
@@ -79,7 +79,7 @@ UEnum* UMaterialInterface::SamplerTypeEnum = nullptr;
 #if WITH_EDITOR
 static TAutoConsoleVariable<bool> CVarStripUnusedDefaultTextures(
 	TEXT("r.Material.StripUnusedDefaultTextures"),
-	true,
+	false,
 	TEXT("If enabled, unused Default textures will be discarded at cook time.\n"),
 	ECVF_ReadOnly);
 #endif
