Fixed regression caused by CL 49331033 causing an ensure when trying to find a parent actor descriptor for overridden actor descriptor.

#Synced-CL 49040969

diff --git a/Engine/Source/Runtime/Engine/Private/WorldPartition/LevelInstance/LevelInstancePropertyOverrideContainer.cpp b/Engine/Source/Runtime/Engine/Private/WorldPartition/LevelInstance/LevelInstancePropertyOverrideContainer.cpp
--- a/Engine/Source/Runtime/Engine/Private/WorldPartition/LevelInstance/LevelInstancePropertyOverrideContainer.cpp
+++ b/Engine/Source/Runtime/Engine/Private/WorldPartition/LevelInstance/LevelInstancePropertyOverrideContainer.cpp
@@ -30,7 +30,8 @@ void ULevelInstancePropertyOverrideContainer::Initialize(const FInitializeParams
 	ContentBundleGuid = GetBaseContainer()->GetContentBundleGuid();
 	ExternalDataLayerAsset = GetBaseContainer()->GetExternalDataLayerAsset();
 
-	auto CloneActorDesc = [this](FWorldPartitionActorDesc* InActorDesc, const FActorContainerPath& InContainerPath)
+	TMap<FWorldPartitionActorDesc*, UActorDescContainer*> OverrideDescToBaseContainer;
+	auto CloneActorDesc = [this, &OverrideDescToBaseContainer](FWorldPartitionActorDesc* InActorDesc, const FActorContainerPath& InContainerPath)
 	{
 		if (!GetOverrideActorDesc(InActorDesc->GetGuid(), InContainerPath))
 		{
@@ -51,6 +52,7 @@ void ULevelInstancePropertyOverrideContainer::Initialize(const FInitializeParams
 				check(!NewActorDesc->GetContainer());
 				NewActorDesc->SetContainer(this);
 
+				OverrideDescToBaseContainer.Emplace(NewActorDesc.Get(), InActorDesc->GetContainer());
 				TMap<FGuid, TUniquePtr<FWorldPartitionActorDesc>>& TransientOverridesMap = PerContainerTransientOverrideActorDescs.FindOrAdd(InContainerPath);
 				TransientOverridesMap.Emplace(NewActorDesc->GetGuid(), MoveTemp(NewActorDesc));
 			}
@@ -70,6 +72,8 @@ void ULevelInstancePropertyOverrideContainer::Initialize(const FInitializeParams
 			for (const auto& ActorDescsMapElement : ActorDescsPerContainerElement.Value)
 			{
 				const TSharedPtr<FWorldPartitionActorDesc>& OverrideActorDesc = ActorDescsMapElement.Value;
+				OverrideDescToBaseContainer.Emplace(OverrideActorDesc.Get(), BaseContainer);
+
 				FWorldPartitionActorDesc* BaseActorDesc = BaseContainer->GetActorDesc(OverrideActorDesc->GetGuid());
 				if (ensure(BaseActorDesc))
 				{
@@ -92,7 +96,7 @@ void ULevelInstancePropertyOverrideContainer::Initialize(const FInitializeParams
 		}
 	}
 
-	auto UpdateActorToWorld = [this](const FGuid& InActorGuid, const FActorContainerPath& InContainerPath)
+	auto UpdateActorToWorld = [this, &OverrideDescToBaseContainer](const FGuid& InActorGuid, const FActorContainerPath& InContainerPath)
 	{
 		if (FWorldPartitionActorDesc* MutableActorDesc = GetOverrideActorDesc(InActorGuid, InContainerPath))
 		{
@@ -100,6 +104,12 @@ void ULevelInstancePropertyOverrideContainer::Initialize(const FInitializeParams
 			if (MutableActorDesc->GetParentActor().IsValid())
 			{
 				const FWorldPartitionActorDesc* ParentActorDesc = GetOverrideActorDesc(MutableActorDesc->GetParentActor(), InContainerPath);
+				if (!ParentActorDesc)
+				{
+					// Parent actor is not part of overrides, find it in MutableActorDesc's associated container
+					UActorDescContainer* MutableActorDescContainer = OverrideDescToBaseContainer.FindChecked(MutableActorDesc);
+					ParentActorDesc = MutableActorDescContainer->GetActorDesc(MutableActorDesc->GetParentActor());
+				}
 				if (ensure(ParentActorDesc))
 				{
 					ParentTransform = ParentActorDesc->GetActorTransform();
