Uninitialize current anim instance when assigning a new single node instance

#Synced-CL 49335786

diff --git a/Engine/Source/Runtime/Engine/Private/Animation/AnimInstance.cpp b/Engine/Source/Runtime/Engine/Private/Animation/AnimInstance.cpp
--- a/Engine/Source/Runtime/Engine/Private/Animation/AnimInstance.cpp
+++ b/Engine/Source/Runtime/Engine/Private/Animation/AnimInstance.cpp
@@ -116,6 +116,12 @@ namespace MontageCVars
 		TEXT("a.Montage.FlushCompletedMontagesOnPlay"),
 		bFlushCompletedMontagesOnPlay,
 		TEXT("Whether we should flush all completed montages IMMEDIATELY when a new montage stomps the group. Use this to prevent accumulating montages when animation tick is paused, and firing all at once when we unpause."));
+
+	static bool bCleanupMontageInstancesOnDestroy = true;
+	static FAutoConsoleVariableRef CVarCleanupMontageInstancesOnDestroy(
+		TEXT("a.Montage.CleanupMontageInstancesOnDestroy"),
+		bCleanupMontageInstancesOnDestroy,
+		TEXT("Whether any montages still alive on BeginDestroy should be cleaned up."));
 }
 
 namespace UE::Anim::Private
@@ -1632,6 +1638,25 @@ void UAnimInstance::BeginDestroy()
 	}
 #endif // WITH_EDITOR
 
+	if (MontageCVars::bCleanupMontageInstancesOnDestroy)
+	{
+		if (MontageInstances.Num() > 0)
+		{
+			for (int32 Index = 0; Index < MontageInstances.Num(); ++Index)
+			{
+				FAnimMontageInstance* MontageInstance = MontageInstances[Index];
+				if (ensure(MontageInstance != nullptr))
+				{
+					ClearMontageInstanceReferences(*MontageInstance);
+					delete MontageInstance;
+				}
+			}
+
+			MontageInstances.Empty();
+			ActiveMontagesMap.Empty();
+		}
+	}
+
 	Super::BeginDestroy();
 	
 	TRACE_OBJECT_LIFETIME_END(this);
diff --git a/Engine/Source/Runtime/Engine/Private/Components/SkeletalMeshComponent.cpp b/Engine/Source/Runtime/Engine/Private/Components/SkeletalMeshComponent.cpp
--- a/Engine/Source/Runtime/Engine/Private/Components/SkeletalMeshComponent.cpp
+++ b/Engine/Source/Runtime/Engine/Private/Components/SkeletalMeshComponent.cpp
@@ -86,6 +86,12 @@ static TAutoConsoleVariable<int32> CVarHiPriSkinnedMeshesTicks(
 	1,
 	TEXT("If > 0, then schedule the skinned component ticks in a tick group before other ticks."));
 
+static bool bCVarUninitializeInstanceOnNewSingleNodeInstance = true;
+static FAutoConsoleVariableRef CVarUninitializeSingleNodeInstanceOnNewInstance(
+	TEXT("a.UninitializeSingleNodeInstanceOnNewInstance"),
+	bCVarUninitializeInstanceOnNewSingleNodeInstance,
+	TEXT("Whether the current anim instance (if any) should be uninitialized before being assigned a new instance."));
+
 // Deprecated. Please switch to ANIM_SKINNED_ASSET_ISPC_ENABLED_DEFAULT and "a.SkinnedAsset.ISPC", see SkinnedAsset.cpp.
 #if !defined(ANIM_SKELETAL_MESH_ISPC_ENABLED_DEFAULT)
 #define ANIM_SKELETAL_MESH_ISPC_ENABLED_DEFAULT 1
@@ -1240,6 +1246,11 @@ bool USkeletalMeshComponent::InitializeAnimScriptInstance(bool bForceReinit, boo
 					OldInstance = Cast<UAnimSingleNodeInstance>(AnimScriptInstance);
 				}
 
+				if (bCVarUninitializeInstanceOnNewSingleNodeInstance && AnimScriptInstance)
+				{
+					AnimScriptInstance->UninitializeAnimation();
+				}
+
 				AnimScriptInstance = NewObject<UAnimSingleNodeInstance>(this);
 
 				if (AnimScriptInstance)
