Fixed the Level Instance ActorDesc editor bounds calculation so that overridden transforms on child actors (via property overrides) now properly affect the parent Level Instance bounds.

#Synced-CL 49040969

diff --git a/Engine/Source/Runtime/Engine/Private/WorldPartition/ActorDescContainerInstance.cpp b/Engine/Source/Runtime/Engine/Private/WorldPartition/ActorDescContainerInstance.cpp
--- a/Engine/Source/Runtime/Engine/Private/WorldPartition/ActorDescContainerInstance.cpp
+++ b/Engine/Source/Runtime/Engine/Private/WorldPartition/ActorDescContainerInstance.cpp
@@ -584,21 +584,24 @@ void UActorDescContainerInstance::OnActorDescUpdating(FWorldPartitionActorDesc*
 	// It is possible for a container instance with an instancing context to listen to updates but not newly added actors
 	// so it might not find the actor if it was newly added (through the edit of a different instance of the same Level Instance)
 	// See UActorDescContainerInstance::RegisterDelegates
-	if (TUniquePtr<FWorldPartitionActorDescInstance>* ActorDescInstance = GetActorDescInstancePtr(InActorDesc->GetGuid()))
+	if (TUniquePtr<FWorldPartitionActorDescInstance>* ActorDescInstanceUniquePtr = GetActorDescInstancePtr(InActorDesc->GetGuid()))
 	{
-		check(ActorDescInstance->IsValid());
+		check(ActorDescInstanceUniquePtr->IsValid());
+		FWorldPartitionActorDescInstance* ActorDescInstance = ActorDescInstanceUniquePtr->Get();
 
-		if (bCreateChildContainerHierarchy && ActorDescInstance->Get()->IsChildContainerInstance())
+		if (bCreateChildContainerHierarchy && ActorDescInstance->IsChildContainerInstance())
 		{
-			ActorDescInstance->Get()->UnregisterChildContainerInstance();
+			ActorDescInstance->UnregisterChildContainerInstance();
 		}
 
 		if (UWorldPartition* WorldPartition = GetOuterWorldPartition())
 		{
-			WorldPartition->OnActorDescInstanceUpdating(ActorDescInstance->Get());
+			WorldPartition->OnActorDescInstanceUpdating(ActorDescInstance);
 		}
 
-		OnActorDescInstanceUpdatingEvent.Broadcast(ActorDescInstance->Get());
+		ActorDescInstance->OnUpdating();
+
+		OnActorDescInstanceUpdatingEvent.Broadcast(ActorDescInstance);
 	}
 }
 
diff --git a/Engine/Source/Runtime/Engine/Private/WorldPartition/WorldPartitionActorDescInstance.cpp b/Engine/Source/Runtime/Engine/Private/WorldPartition/WorldPartitionActorDescInstance.cpp
--- a/Engine/Source/Runtime/Engine/Private/WorldPartition/WorldPartitionActorDescInstance.cpp
+++ b/Engine/Source/Runtime/Engine/Private/WorldPartition/WorldPartitionActorDescInstance.cpp
@@ -7,6 +7,7 @@
 #include "WorldPartition/ActorDescContainerInstance.h"
 #include "WorldPartition/WorldPartition.h"
 #include "WorldPartition/WorldPartitionActorDesc.h"
+#include "WorldPartition/LevelInstance/LevelInstancePropertyOverrideContainer.h"
 #include "UObject/MetaData.h"
 
 #define LOCTEXT_NAMESPACE "FWorldPartitionActorDescInstance"
@@ -97,11 +98,133 @@ bool FWorldPartitionActorDescInstance::IsRuntimeRelevant() const
 
 FBox FWorldPartitionActorDescInstance::GetEditorBounds() const
 {
-	return ActorDesc->GetEditorBounds().TransformBy(GetContainerInstance()->GetTransform());
+	return GetLocalEditorBounds().TransformBy(GetContainerInstance()->GetTransform());
+}
+
+FBox FWorldPartitionActorDescInstance::GetLocalEditorBounds(bool bInForceComputeOnChildContainerInstance) const
+{ 
+	auto BuildActorContainerPath = [](const UActorDescContainerInstance* ContainerInstance)
+	{
+		FActorContainerPath ContainerPath;
+		while (ContainerInstance)
+		{
+			// Don't include top container
+			if (ContainerInstance->GetContainerActorGuid().IsValid())
+			{
+				ContainerPath.ContainerGuids.Add(ContainerInstance->GetContainerActorGuid());
+			}
+			ContainerInstance = ContainerInstance->GetParentContainerInstance();
+		}
+		Algo::Reverse(ContainerPath.ContainerGuids);
+		return ContainerPath;
+	};
+
+	auto GetTopMostContainerInstance = [](const UActorDescContainerInstance* ContainerInstance)
+	{
+		const UActorDescContainerInstance* CurrentContainerInstance = ContainerInstance;
+		while (CurrentContainerInstance && CurrentContainerInstance->GetParentContainerInstance())
+		{
+			CurrentContainerInstance = CurrentContainerInstance->GetParentContainerInstance();
+		}
+		return CurrentContainerInstance;
+	};
+
+	auto FindChildContainerInstance = [](const UActorDescContainerInstance* ContainerInstance, const FActorContainerPath ContainerPath) -> const UActorDescContainerInstance*
+	{
+		const UActorDescContainerInstance* CurrentContainerInstance = ContainerInstance;
+		for (const FGuid& Guid : ContainerPath.ContainerGuids)
+		{
+			if (const TObjectPtr<UActorDescContainerInstance>* FoundContainerInstance = CurrentContainerInstance->GetChildContainerInstances().Find(Guid))
+			{
+				CurrentContainerInstance = *FoundContainerInstance;
+			}
+			else
+			{
+				return nullptr;
+			}
+		}
+		return CurrentContainerInstance;
+	};
+
+	if (IsChildContainerInstance())
+	{
+		// Use cached value if valid
+		if (ChildContainerLocalEditorBoundsCache.IsSet())
+		{
+			return ChildContainerLocalEditorBoundsCache.GetValue();
+		}
+
+		// When container uses property overrides, compute bounds dynamically and cache the result.
+		const UActorDescContainerInstance* ChildContainer = GetChildContainerInstance();
+		const bool bForceComputeOnChildContainerInstance = bInForceComputeOnChildContainerInstance || (ChildContainer && ChildContainer->GetContainer()->IsA<ULevelInstancePropertyOverrideContainer>());
+		if (ChildContainer && bForceComputeOnChildContainerInstance)
+		{
+			FBox Bounds(ForceInitToZero);
+			for (UActorDescContainerInstance::TConstIterator<AActor> ActorDescIt(ChildContainer); ActorDescIt; ++ActorDescIt)
+			{
+				if (!ActorDescIt->IsMainWorldOnly())
+				{
+					Bounds += ActorDescIt->GetLocalEditorBounds(bForceComputeOnChildContainerInstance);
+				}
+			}
+			Bounds = Bounds.TransformBy(ActorDesc->GetActorTransformRelative());
+			ChildContainerLocalEditorBoundsCache = Bounds;
+			return Bounds;
+		}
+
+		// ChildContainer can be null if we loaded an instance of a world partition (most likely a level instance).
+		// In which case, the hierarchy of container instances is not initialized (to avoid massive duplication of container instances).
+		if (!ChildContainer)
+		{
+			// We can instead find the corresponding container instance in the main world partition using the container path.
+			// We are able to find the top most container instance only because this container instance has its parent 
+			// container instance set (which is part of the main world partition hierarchy).
+			const FActorContainerPath ContainerPath = BuildActorContainerPath(GetContainerInstance());
+			const UActorDescContainerInstance* TopMostContainerInstance = GetTopMostContainerInstance(GetContainerInstance());
+			const UActorDescContainerInstance* MatchingContainerInstance = FindChildContainerInstance(TopMostContainerInstance, ContainerPath);
+			// It's possible that we can't find a matching container instance when a newly created (unsaved) Level Instance is added to the world
+			if (MatchingContainerInstance && (MatchingContainerInstance != GetContainerInstance()))
+			{
+				FWorldPartitionActorDescInstance* MatchingActorDescInstance = MatchingContainerInstance->GetActorDescInstance(GetGuid());
+				if (ensure(MatchingActorDescInstance))
+				{
+					return MatchingActorDescInstance->GetLocalEditorBounds();
+				}
+			}
+		}
+	}
+
+	// Use ActorDesc editor bounds
+	return ActorDesc->GetEditorBounds();
+}
+
+void FWorldPartitionActorDescInstance::OnUpdating()
+{
+	// Invalidate cached bounds
+	if (IsChildContainerInstance())
+	{
+		UActorDescContainerInstance* CurrentContainerInstance = GetContainerInstance();
+		while (CurrentContainerInstance)
+		{
+			for (UActorDescContainerInstance::TIterator<> It(CurrentContainerInstance); It; ++It)
+			{
+				It->InvalidateCachedBounds();
+			}
+			// Recurse to parent as parent bounds are affected by children bounds
+			CurrentContainerInstance = const_cast<UActorDescContainerInstance*>(CurrentContainerInstance->GetParentContainerInstance());
+		}
+	}
+}
+
+void FWorldPartitionActorDescInstance::InvalidateCachedBounds()
+{
+	ChildContainerLocalEditorBoundsCache.Reset();
 }
 
 FBox FWorldPartitionActorDescInstance::GetRuntimeBounds() const
 {
+	// Override is not supported for streamed level instances.
+	// No need to compute dynamically runtime bounds per instance.
 	return ActorDesc->GetRuntimeBounds().TransformBy(GetContainerInstance()->GetTransform());
 }
 
diff --git a/Engine/Source/Runtime/Engine/Public/WorldPartition/WorldPartitionActorDescInstance.h b/Engine/Source/Runtime/Engine/Public/WorldPartition/WorldPartitionActorDescInstance.h
--- a/Engine/Source/Runtime/Engine/Public/WorldPartition/WorldPartitionActorDescInstance.h
+++ b/Engine/Source/Runtime/Engine/Public/WorldPartition/WorldPartitionActorDescInstance.h
@@ -124,9 +124,9 @@ public:
 	ENGINE_API const FString& GetActorNameString() const;
 	ENGINE_API const FString& GetActorLabelString() const;
 	ENGINE_API const FString& GetDisplayClassNameString() const;
-		
+
 protected:
-	virtual FBox GetLocalEditorBounds() const { return ActorDesc->GetEditorBounds(); }
+	virtual FBox GetLocalEditorBounds(bool bInForceComputeOnChildContainerInstance = false) const;
 
 	ENGINE_API bool StartAsyncLoad();
 	ENGINE_API void FlushAsyncLoad() const;
@@ -175,6 +175,11 @@ protected:
 	virtual ENGINE_API void UnregisterChildContainerInstance();
 	virtual ENGINE_API void UpdateChildContainerInstance();
 
+private:
+	/** Called by UActorDescContainerInstance::OnActorDescUpdating */
+	void OnUpdating();
+	void InvalidateCachedBounds();
+
 protected:
 	TObjectPtr<UActorDescContainerInstance>		ContainerInstance;
 
@@ -193,5 +198,7 @@ protected:
 	FWorldPartitionActorDesc*					ActorDesc;
 
 	TObjectPtr<UActorDescContainerInstance>		ChildContainerInstance;
+
+	mutable TOptional<FBox>						ChildContainerLocalEditorBoundsCache;
 #endif // WITH_EDITOR
 };
