[Sound Wave] - LogSlate warnings when expanding timecode that display incorrect symbols
#Synced-CL 49272210

diff --git a/Engine/Source/Runtime/Engine/Classes/Sound/SoundWave.h b/Engine/Source/Runtime/Engine/Classes/Sound/SoundWave.h
--- a/Engine/Source/Runtime/Engine/Classes/Sound/SoundWave.h
+++ b/Engine/Source/Runtime/Engine/Classes/Sound/SoundWave.h
@@ -1230,8 +1230,12 @@ public:
 #if WITH_EDITORONLY_DATA
 	ENGINE_API void SetTimecodeInfo(const FSoundWaveTimecodeInfo& InTimecode);
 	ENGINE_API TOptional<FSoundWaveTimecodeInfo> GetTimecodeInfo() const;
+
+protected:
+	ENGINE_API void FixTimeCodeInfo();
 #endif //WITH_EDITORONLY_DATA	
 
+public:
 	/**
 	* Overwrite sample rate. Used for procedural soundwaves, as well as sound waves that are resampled on compress/decompress.
 	*/
diff --git a/Engine/Source/Runtime/Engine/Private/SoundWave.cpp b/Engine/Source/Runtime/Engine/Private/SoundWave.cpp
--- a/Engine/Source/Runtime/Engine/Private/SoundWave.cpp
+++ b/Engine/Source/Runtime/Engine/Private/SoundWave.cpp
@@ -1080,6 +1080,32 @@ TOptional<FSoundWaveTimecodeInfo> USoundWave::GetTimecodeInfo() const
 	return TimecodeInfo;
 }
 
+void USoundWave::FixTimeCodeInfo()
+{
+	if (TimecodeInfo.OriginatorDate.IsEmpty() && TimecodeInfo.Description.IsEmpty() && TimecodeInfo.OriginatorReference.IsEmpty()
+		&& TimecodeInfo.OriginatorDescription.IsEmpty() && TimecodeInfo.OriginatorTime.IsEmpty())
+	{
+		// TimecodeInfo data is empty and does not need to be fixed
+		return;
+	}
+
+	FSharedBuffer BulkDataBuffer;
+	BulkDataBuffer = RawData.GetPayload().Get();
+
+	FWaveModInfo WaveInfo;
+	if (!WaveInfo.ReadWaveInfo((const uint8*)BulkDataBuffer.GetData(), BulkDataBuffer.GetSize()))
+	{
+		UE_LOG(LogAudio, Error, TEXT("Sound wave payload in asset %s failed to parse as a wave in FixTimeCodeInfo!"), *GetName());
+
+		return;
+	}
+
+	if (WaveInfo.TimecodeInfo.IsValid())
+	{
+		SetTimecodeInfo(*WaveInfo.TimecodeInfo);
+	}
+}
+
 #endif //WITH_EDITORONLY_DATA
 
 int32 USoundWave::GetResourceSizeForFormat(FName Format)
@@ -2214,18 +2240,24 @@ void USoundWave::PostLoad()
 	}
 
 
-	if (RawData.HasPayloadData() && ImportedSampleRate == 0)
+	if (RawData.HasPayloadData())
 	{
-		// update ImportedSampleRate ASAP so that it reflects the 
-		// will only be valid if this sound wave was from an imported .wav file. 
-		// otherwise we can safely ignore it
-		TArray<uint8> ImportedRawPCMData;
-		uint32 OriginalSampleRate = 0;
-		uint16 ImportedNumChannels = 0;
-		if (GetImportedSoundWaveData(ImportedRawPCMData, OriginalSampleRate, ImportedNumChannels))
+		if (ImportedSampleRate == 0)
 		{
-			ImportedSampleRate = OriginalSampleRate;
+			// update ImportedSampleRate ASAP so that it reflects the 
+			// will only be valid if this sound wave was from an imported .wav file. 
+			// otherwise we can safely ignore it
+			TArray<uint8> ImportedRawPCMData;
+			uint32 OriginalSampleRate = 0;
+			uint16 ImportedNumChannels = 0;
+			if (GetImportedSoundWaveData(ImportedRawPCMData, OriginalSampleRate, ImportedNumChannels))
+			{
+				ImportedSampleRate = OriginalSampleRate;
+			}
 		}
+
+		// Does not mark the asset dirty as it is EDITOR_ONLY_DATA but will always appear fixed to the user
+		FixTimeCodeInfo();
 	}
 	
 	if (TotalSamples == 0 && ImportedSampleRate != 0)
