Deprecated LigthWeightInstances

diff --git a/Engine/Plugins/Editor/EngineAssetDefinitions/Source/Private/AssetDefinition_LightWeightInstance.h b/Engine/Plugins/Editor/EngineAssetDefinitions/Source/Private/AssetDefinition_LightWeightInstance.h
--- a/Engine/Plugins/Editor/EngineAssetDefinitions/Source/Private/AssetDefinition_LightWeightInstance.h
+++ b/Engine/Plugins/Editor/EngineAssetDefinitions/Source/Private/AssetDefinition_LightWeightInstance.h
@@ -7,8 +7,8 @@
 
 #include "AssetDefinition_LightWeightInstance.generated.h"
 
-UCLASS()
-class UAssetDefinition_LightWeightInstance : public UAssetDefinitionDefault
+UCLASS(Deprecated, meta = (DeprecationMessage ="Deprecated in 5.8. Consider using InstancedActors for similar functionality"))
+class UDEPRECATED_AssetDefinition_LightWeightInstance : public UAssetDefinitionDefault
 {
 	GENERATED_BODY()
 
@@ -16,7 +16,7 @@ public:
 	// UAssetDefinition Begin
 	virtual FText GetAssetDisplayName() const override { return NSLOCTEXT("AssetTypeActions", "AssetTypeActions_LightWeightInstance", "Light Weight Instance"); }
 	virtual FLinearColor GetAssetColor() const override { return FLinearColor(FColor(255, 105, 180)); }
-	virtual TSoftClassPtr<UObject> GetAssetClass() const override { return ALightWeightInstanceManager::StaticClass(); }
+	virtual TSoftClassPtr<UObject> GetAssetClass() const override { return ADEPRECATED_LightWeightInstanceManager::StaticClass(); }
 	virtual TConstArrayView<FAssetCategoryPath> GetAssetCategories() const override
 	{
 		static const auto Categories = { EAssetCategoryPaths::Gameplay };
diff --git a/Engine/Plugins/LightWeightInstancesEditor/LightWeightInstancesEditor.uplugin b/Engine/Plugins/LightWeightInstancesEditor/LightWeightInstancesEditor.uplugin
--- a/Engine/Plugins/LightWeightInstancesEditor/LightWeightInstancesEditor.uplugin
+++ b/Engine/Plugins/LightWeightInstancesEditor/LightWeightInstancesEditor.uplugin
@@ -2,8 +2,8 @@
 	"FileVersion": 3,
 	"Version": 1,
 	"VersionName": "1.0",
-	"FriendlyName": "LightWeightInstancesEditor",
-	"Description": "Light Weight Instances provide the flexibility and interaction of actors while having performance similar to instanced meshes.",
+	"FriendlyName": "LightWeightInstancesEditor (DEPRECATED)",
+	"Description": "Light Weight Instances provide the flexibility and interaction of actors while having performance similar to instanced meshes. This plugin is now DEPRECATED. Use InstancedActors for similar functionality.",
 	"Category": "Other",
 	"CreatedBy": "Epic Games, Inc.",
 	"CreatedByURL": "https://epicgames.com",
@@ -14,6 +14,7 @@
 	"IsBetaVersion": true,
 	"IsExperimentalVersion": false,
 	"Installed": false,
+	"DeprecatedEngineVersion": "5.8",
 	"Modules": [
 		{
 			"Name": "LightWeightInstancesEditor",
diff --git a/Engine/Plugins/LightWeightInstancesEditor/Source/LightWeightInstancesEditor/Private/LightWeightInstancesEditor.cpp b/Engine/Plugins/LightWeightInstancesEditor/Source/LightWeightInstancesEditor/Private/LightWeightInstancesEditor.cpp
--- a/Engine/Plugins/LightWeightInstancesEditor/Source/LightWeightInstancesEditor/Private/LightWeightInstancesEditor.cpp
+++ b/Engine/Plugins/LightWeightInstancesEditor/Source/LightWeightInstancesEditor/Private/LightWeightInstancesEditor.cpp
@@ -14,6 +14,8 @@
 
 DEFINE_LOG_CATEGORY_STATIC(LogLWIEditor, Log, All);
 
+PRAGMA_DISABLE_DEPRECATION_WARNINGS
+
 #if WITH_EDITOR
 typedef FLevelEditorModule::FLevelViewportMenuExtender_SelectedActors DelegateType;
 #endif
@@ -83,11 +85,11 @@ TSharedRef<FExtender> FLightWeightInstancesEditorModule::CreateLevelViewportCont
 		Extender->AddMenuExtension("ActorConvert", EExtensionHook::After, LevelEditorCommandBindings, FMenuExtensionDelegate::CreateLambda(
 			[this, ActorName, InActors](FMenuBuilder& MenuBuilder) {
 				// We can only convert to an LWI if we didn't select an LWI
-				const bool bCanExecute = !InActors[0]->GetClass()->IsChildOf(ALightWeightInstanceManager::StaticClass());
+				const bool bCanExecute = !InActors[0]->GetClass()->IsChildOf(ADEPRECATED_LightWeightInstanceManager::StaticClass());
 
 				MenuBuilder.AddMenuEntry(
-					FText::Format(LOCTEXT("ConvertSelectedActorsToLWIsText", "Convert {0} To Light Weight Instances"), ActorName),
-					LOCTEXT("ConvertSelectedActorsToLWIsTooltip", "Convert the selected actors to light weight instances."),
+					FText::Format(LOCTEXT("ConvertSelectedActorsToLWIsText", "(DEPRECATED) Convert {0} To Light Weight Instances"), ActorName),
+					LOCTEXT("ConvertSelectedActorsToLWIsTooltip", "(DEPRECATED) Convert the selected actors to light weight instances."),
 					FSlateIcon(FAppStyle::GetAppStyleSetName(), "Icons.Convert"),
 					FUIAction(
 						FExecuteAction::CreateRaw(this, &FLightWeightInstancesEditorModule::ConvertActorsToLWIsUIAction, InActors),
@@ -113,7 +115,7 @@ void FLightWeightInstancesEditorModule::ConvertActorsToLWIsUIAction(const TArray
 		TArray<const UDataLayerInstance*> DataLayerInstances = Actor->GetDataLayerInstances();
 		const UDataLayerInstance* DataLayerInstance = DataLayerInstances.Num() > 0 ? DataLayerInstances[0] : nullptr;
 
-		ALightWeightInstanceManager* Manager = FLightWeightInstanceSubsystem::Get().FindOrAddLightWeightInstanceManager(*Actor->GetClass(), *Actor->GetWorld(), Actor->GetActorLocation(), DataLayerInstance);
+		ADEPRECATED_LightWeightInstanceManager* Manager = FLightWeightInstanceSubsystem::Get().FindOrAddLightWeightInstanceManager(*Actor->GetClass(), *Actor->GetWorld(), Actor->GetActorLocation(), DataLayerInstance);
 		check(Manager);
 		
 		if(UDataLayerEditorSubsystem* DataLayerEditor = UDataLayerEditorSubsystem::Get())
@@ -127,6 +129,9 @@ void FLightWeightInstancesEditorModule::ConvertActorsToLWIsUIAction(const TArray
 
 #endif // WITH_EDITOR
 
+PRAGMA_ENABLE_DEPRECATION_WARNINGS
+
 #undef LOCTEXT_NAMESPACE
 	
 IMPLEMENT_MODULE(FLightWeightInstancesEditorModule, LightWeightInstancesEditor)
+
diff --git a/Engine/Source/Editor/UnrealEd/Classes/Factories/LightWeightInstanceFactory.h b/Engine/Source/Editor/UnrealEd/Classes/Factories/LightWeightInstanceFactory.h
--- a/Engine/Source/Editor/UnrealEd/Classes/Factories/LightWeightInstanceFactory.h
+++ b/Engine/Source/Editor/UnrealEd/Classes/Factories/LightWeightInstanceFactory.h
@@ -7,10 +7,10 @@
 #include "Factories/Factory.h"
 #include "LightWeightInstanceFactory.generated.h"
 
-class ALightWeightInstanceManager;
+class ADEPRECATED_LightWeightInstanceManager;
 
-UCLASS(hidecategories=Object, MinimalAPI)
-class ULightWeightInstanceFactory : public UFactory
+UCLASS(hidecategories=Object, MinimalAPI, Deprecated, meta = (DeprecationMessage ="Deprecated in 5.8. Consider using InstancedActors for similar functionality"))
+class UDEPRECATED_LightWeightInstanceFactory : public UFactory
 {
 	GENERATED_UCLASS_BODY()
 
diff --git a/Engine/Source/Editor/UnrealEd/Private/Factories/LightWeightInstanceFactory.cpp b/Engine/Source/Editor/UnrealEd/Private/Factories/LightWeightInstanceFactory.cpp
--- a/Engine/Source/Editor/UnrealEd/Private/Factories/LightWeightInstanceFactory.cpp
+++ b/Engine/Source/Editor/UnrealEd/Private/Factories/LightWeightInstanceFactory.cpp
@@ -22,6 +22,8 @@
 
 #define LOCTEXT_NAMESPACE "LightWeightInstanceFactory"
 
+PRAGMA_DISABLE_DEPRECATION_WARNINGS
+
 class FLWIFilter : public IClassViewerFilter
 {
 public:
@@ -39,15 +41,15 @@ public:
 	}
 };
 
-ULightWeightInstanceFactory::ULightWeightInstanceFactory(const FObjectInitializer& ObjectInitializer)
+UDEPRECATED_LightWeightInstanceFactory::UDEPRECATED_LightWeightInstanceFactory(const FObjectInitializer& ObjectInitializer)
 	: Super(ObjectInitializer)
 {
-	SupportedClass = ALightWeightInstanceManager::StaticClass();
+	SupportedClass = ADEPRECATED_LightWeightInstanceManager::StaticClass();
 	bCreateNew = true;
 	bEditAfterNew = true;
 }
 
-bool ULightWeightInstanceFactory::ConfigureProperties()
+bool UDEPRECATED_LightWeightInstanceFactory::ConfigureProperties()
 {
 	// Null the parent class to ensure one is selected
 	ParentClass = nullptr;
@@ -71,11 +73,11 @@ bool ULightWeightInstanceFactory::ConfigureProperties()
 	TSharedPtr<FLWIFilter> Filter = MakeShareable(new FLWIFilter);
 	Options.ClassFilters.Add(Filter.ToSharedRef());
 
-	Filter->AllowedChildrenOfClasses.Add(ALightWeightInstanceManager::StaticClass());
+	Filter->AllowedChildrenOfClasses.Add(ADEPRECATED_LightWeightInstanceManager::StaticClass());
 
 	const FText TitleText = LOCTEXT("CreateBlueprintOptions", "Pick Parent Class");
 	UClass* ChosenClass = nullptr;
-	const bool bPressedOk = SClassPickerDialog::PickClass(TitleText, Options, ChosenClass, ALightWeightInstanceManager::StaticClass());
+	const bool bPressedOk = SClassPickerDialog::PickClass(TitleText, Options, ChosenClass, ADEPRECATED_LightWeightInstanceManager::StaticClass());
 
 	if (bPressedOk)
 	{
@@ -87,7 +89,7 @@ bool ULightWeightInstanceFactory::ConfigureProperties()
 	return bPressedOk;
 };
 
-UObject* ULightWeightInstanceFactory::FactoryCreateNew(UClass* Class, UObject* InParent, FName Name, EObjectFlags Flags, UObject* Context, FFeedbackContext* Warn)
+UObject* UDEPRECATED_LightWeightInstanceFactory::FactoryCreateNew(UClass* Class, UObject* InParent, FName Name, EObjectFlags Flags, UObject* Context, FFeedbackContext* Warn)
 {
 	UObject* LWIManager = nullptr;
 	if (ensure(SupportedClass == Class))
@@ -105,4 +107,6 @@ UObject* ULightWeightInstanceFactory::FactoryCreateNew(UClass* Class, UObject* I
 	return LWIManager;
 }
 
+PRAGMA_ENABLE_DEPRECATION_WARNINGS
+
 #undef LOCTEXT_NAMESPACE // "LightWeightInstanceFactory"
diff --git a/Engine/Source/Runtime/Engine/Classes/GameFramework/LightWeightInstanceBlueprintFunctionLibrary.h b/Engine/Source/Runtime/Engine/Classes/GameFramework/LightWeightInstanceBlueprintFunctionLibrary.h
--- a/Engine/Source/Runtime/Engine/Classes/GameFramework/LightWeightInstanceBlueprintFunctionLibrary.h
+++ b/Engine/Source/Runtime/Engine/Classes/GameFramework/LightWeightInstanceBlueprintFunctionLibrary.h
@@ -7,8 +7,8 @@
 
 #include "LightWeightInstanceBlueprintFunctionLibrary.generated.h"
 
-UCLASS(MinimalAPI)
-class ULightWeightInstanceBlueprintFunctionLibrary : public UBlueprintFunctionLibrary
+UCLASS(MinimalAPI, Deprecated, meta = (DeprecationMessage ="Deprecated in 5.8. Consider using InstancedActors for similar functionality"))
+class UDEPRECATED_LightWeightInstanceBlueprintFunctionLibrary : public UBlueprintFunctionLibrary
 {
 	GENERATED_BODY()
 
diff --git a/Engine/Source/Runtime/Engine/Classes/GameFramework/LightWeightInstanceManager.h b/Engine/Source/Runtime/Engine/Classes/GameFramework/LightWeightInstanceManager.h
--- a/Engine/Source/Runtime/Engine/Classes/GameFramework/LightWeightInstanceManager.h
+++ b/Engine/Source/Runtime/Engine/Classes/GameFramework/LightWeightInstanceManager.h
@@ -27,8 +27,8 @@ struct FLWIData
 };
 
 
-UCLASS(BlueprintType, Blueprintable, Experimental, MinimalAPI)
-class ALightWeightInstanceManager : public AActor, public IActorInstanceManagerInterface
+UCLASS(BlueprintType, Blueprintable, Experimental, MinimalAPI, Deprecated, meta = (DeprecationMessage ="Deprecated in 5.8. Consider using InstancedActors for similar functionality"))
+class ADEPRECATED_LightWeightInstanceManager : public AActor, public IActorInstanceManagerInterface
 {
 	GENERATED_UCLASS_BODY()
 
@@ -37,7 +37,7 @@ class ALightWeightInstanceManager : public AActor, public IActorInstanceManagerI
 	friend class ULightWeightInstanceBlueprintFunctionLibrary;
 
 public:
-	ENGINE_API virtual ~ALightWeightInstanceManager();
+	ENGINE_API virtual ~ADEPRECATED_LightWeightInstanceManager();
 
 	ENGINE_API virtual void Tick(float DeltaSeconds);
 
diff --git a/Engine/Source/Runtime/Engine/Classes/GameFramework/LightWeightInstanceStaticMeshManager.h b/Engine/Source/Runtime/Engine/Classes/GameFramework/LightWeightInstanceStaticMeshManager.h
--- a/Engine/Source/Runtime/Engine/Classes/GameFramework/LightWeightInstanceStaticMeshManager.h
+++ b/Engine/Source/Runtime/Engine/Classes/GameFramework/LightWeightInstanceStaticMeshManager.h
@@ -11,8 +11,8 @@
 
 DECLARE_DYNAMIC_DELEGATE_RetVal_OneParam(FActorInstanceHandle, FOnActorReady, FActorInstanceHandle, InHandle);
 
-UCLASS(BlueprintType, Blueprintable, Experimental, MinimalAPI)
-class ALightWeightInstanceStaticMeshManager : public ALightWeightInstanceManager, public ISMInstanceManager
+UCLASS(BlueprintType, Blueprintable, Experimental, MinimalAPI, Deprecated, meta = (DeprecationMessage ="Deprecated in 5.8. Consider using InstancedActors for similar functionality"))
+class ADEPRECATED_LightWeightInstanceStaticMeshManager : public ADEPRECATED_LightWeightInstanceManager, public ISMInstanceManager
 {
 	GENERATED_UCLASS_BODY()
 
diff --git a/Engine/Source/Runtime/Engine/Classes/GameFramework/LightWeightInstanceSubsystem.h b/Engine/Source/Runtime/Engine/Classes/GameFramework/LightWeightInstanceSubsystem.h
--- a/Engine/Source/Runtime/Engine/Classes/GameFramework/LightWeightInstanceSubsystem.h
+++ b/Engine/Source/Runtime/Engine/Classes/GameFramework/LightWeightInstanceSubsystem.h
@@ -27,21 +27,21 @@ struct FLightWeightInstanceSubsystem
 	ENGINE_API ~FLightWeightInstanceSubsystem();
 
 	// Returns the instance manager that handles the given handle
-	ENGINE_API ALightWeightInstanceManager* FindLightWeightInstanceManager(const FActorInstanceHandle& Handle) const;
+	ENGINE_API ADEPRECATED_LightWeightInstanceManager* FindLightWeightInstanceManager(const FActorInstanceHandle& Handle) const;
 
 	// Returns the instance manager that handles actors of type ActorClass in level Level
 	UE_DEPRECATED(5.3, "Use the version that takes in a position.")
-	ENGINE_API ALightWeightInstanceManager* FindLightWeightInstanceManager(UClass* ActorClass, const UDataLayerInstance* Layer, UWorld* World) const;
+	ENGINE_API ADEPRECATED_LightWeightInstanceManager* FindLightWeightInstanceManager(UClass* ActorClass, const UDataLayerInstance* Layer, UWorld* World) const;
 
 	// Returns the instance manager that handles instances of type Class that live in Level
 	UE_DEPRECATED(5.3, "Use the version that takes in a position.")
-	ENGINE_API ALightWeightInstanceManager* FindOrAddLightWeightInstanceManager(UClass* ActorClass, const UDataLayerInstance* DataLayer, UWorld* World);
+	ENGINE_API ADEPRECATED_LightWeightInstanceManager* FindOrAddLightWeightInstanceManager(UClass* ActorClass, const UDataLayerInstance* DataLayer, UWorld* World);
 
 	// Returns the instance manager that handles actors of type ActorClass in level Level
-	ENGINE_API ALightWeightInstanceManager* FindLightWeightInstanceManager(UClass& ActorClass, UWorld& World, const FVector& InPos, const UDataLayerInstance* DataLayer = nullptr) const;
+	ENGINE_API ADEPRECATED_LightWeightInstanceManager* FindLightWeightInstanceManager(UClass& ActorClass, UWorld& World, const FVector& InPos, const UDataLayerInstance* DataLayer = nullptr) const;
 
 	// Returns the instance manager that handles instances of type Class that live in Level
-	ENGINE_API ALightWeightInstanceManager* FindOrAddLightWeightInstanceManager(UClass& ActorClass, UWorld& World, const FVector& InPos, const UDataLayerInstance* DataLayer = nullptr);
+	ENGINE_API ADEPRECATED_LightWeightInstanceManager* FindOrAddLightWeightInstanceManager(UClass& ActorClass, UWorld& World, const FVector& InPos, const UDataLayerInstance* DataLayer = nullptr);
 
 	// Returns the actor specified by Handle if it exists. Returns nullptr if it doesn't
 	ENGINE_API AActor* GetActor_NoCreate(const FActorInstanceHandle& Handle) const;
@@ -69,24 +69,24 @@ struct FLightWeightInstanceSubsystem
 	static ENGINE_API FInt32Vector3 ConvertPositionToCoord(const FVector& InPosition);
 
 	// Add a manager to the subsystem, thread safe.
-	ENGINE_API bool AddManager(ALightWeightInstanceManager* Manager);
+	ENGINE_API bool AddManager(ADEPRECATED_LightWeightInstanceManager* Manager);
 
 	// Remove a manager from the subsystem, thread safe.
-	ENGINE_API bool RemoveManager(ALightWeightInstanceManager* Manager);
+	ENGINE_API bool RemoveManager(ADEPRECATED_LightWeightInstanceManager* Manager);
 
 protected:
 	// Returns the class of the instance manager best suited to support instances of type ActorClass
 	ENGINE_API UClass* FindBestInstanceManagerClass(const UClass* ActorClass);
 
 	// Returns the index associated with Manager
-	ENGINE_API int32 GetManagerIndex(const ALightWeightInstanceManager* Manager) const;
+	ENGINE_API int32 GetManagerIndex(const ADEPRECATED_LightWeightInstanceManager* Manager) const;
 
 	// Returns the light weight instance manager at index Index
-	ENGINE_API const ALightWeightInstanceManager* GetManagerAt(int32 Index) const;
+	ENGINE_API const ADEPRECATED_LightWeightInstanceManager* GetManagerAt(int32 Index) const;
 
 private:
 	// TODO: preallocate the size of this based on a config variable
-	TArray<ALightWeightInstanceManager*> LWInstanceManagers;
+	TArray<ADEPRECATED_LightWeightInstanceManager*> LWInstanceManagers;
 
 	// Mutex to make sure we don't change the LWInstanceManagers array while reading/writing it.
 	mutable FTransactionallySafeRWLock LWIManagersRWLock;
diff --git a/Engine/Source/Runtime/Engine/Private/LightWeightInstanceBlueprintFunctionLibrary.cpp b/Engine/Source/Runtime/Engine/Private/LightWeightInstanceBlueprintFunctionLibrary.cpp
--- a/Engine/Source/Runtime/Engine/Private/LightWeightInstanceBlueprintFunctionLibrary.cpp
+++ b/Engine/Source/Runtime/Engine/Private/LightWeightInstanceBlueprintFunctionLibrary.cpp
@@ -6,7 +6,9 @@
 
 #include UE_INLINE_GENERATED_CPP_BY_NAME(LightWeightInstanceBlueprintFunctionLibrary)
 
-FActorInstanceHandle ULightWeightInstanceBlueprintFunctionLibrary::CreateNewLightWeightInstance(UClass* InActorClass, FTransform InTransform, UDataLayerInstance* InLayer, UWorld* World)
+PRAGMA_DISABLE_DEPRECATION_WARNINGS
+
+FActorInstanceHandle UDEPRECATED_LightWeightInstanceBlueprintFunctionLibrary::CreateNewLightWeightInstance(UClass* InActorClass, FTransform InTransform, UDataLayerInstance* InLayer, UWorld* World)
 {
 	// set up initialization data
 	FLWIData PerInstanceData;
@@ -15,7 +17,7 @@ FActorInstanceHandle ULightWeightInstanceBlueprintFunctionLibrary::CreateNewLigh
 	return FLightWeightInstanceSubsystem::Get().CreateNewLightWeightInstance(InActorClass, &PerInstanceData, InLayer, World);
 }
 
-FActorInstanceHandle ULightWeightInstanceBlueprintFunctionLibrary::ConvertActorToLightWeightInstance(AActor* InActor)
+FActorInstanceHandle UDEPRECATED_LightWeightInstanceBlueprintFunctionLibrary::ConvertActorToLightWeightInstance(AActor* InActor)
 {
 	if (!ensure(InActor))
 	{
@@ -30,7 +32,7 @@ FActorInstanceHandle ULightWeightInstanceBlueprintFunctionLibrary::ConvertActorT
 #else
 	const UDataLayerInstance* DataLayerInstance = nullptr;
 #endif // WITH_EDITOR
-	if (ALightWeightInstanceManager* LWIManager = FLightWeightInstanceSubsystem::Get().FindOrAddLightWeightInstanceManager(*InActor->GetClass(), *InActor->GetWorld(), InActor->GetActorLocation(), DataLayerInstance))
+	if (ADEPRECATED_LightWeightInstanceManager* LWIManager = FLightWeightInstanceSubsystem::Get().FindOrAddLightWeightInstanceManager(*InActor->GetClass(), *InActor->GetWorld(), InActor->GetActorLocation(), DataLayerInstance))
 	{
 		return LWIManager->ConvertActorToLightWeightInstance(InActor);
 	}
@@ -38,3 +40,4 @@ FActorInstanceHandle ULightWeightInstanceBlueprintFunctionLibrary::ConvertActorT
 	return FActorInstanceHandle(InActor);
 }
 
+PRAGMA_ENABLE_DEPRECATION_WARNINGS
diff --git a/Engine/Source/Runtime/Engine/Private/LightWeightInstanceManager.cpp b/Engine/Source/Runtime/Engine/Private/LightWeightInstanceManager.cpp
--- a/Engine/Source/Runtime/Engine/Private/LightWeightInstanceManager.cpp
+++ b/Engine/Source/Runtime/Engine/Private/LightWeightInstanceManager.cpp
@@ -9,6 +9,8 @@
 
 #include UE_INLINE_GENERATED_CPP_BY_NAME(LightWeightInstanceManager)
 
+PRAGMA_DISABLE_DEPRECATION_WARNINGS
+
 // CVar variable that is the size of the grid managers are placed into.
 static int32 LWIGridSize = -1;
 static FAutoConsoleVariableRef CVarLWIGridSize
@@ -20,7 +22,7 @@ static FAutoConsoleVariableRef CVarLWIGridSize
 );
 
 
-ALightWeightInstanceManager::ALightWeightInstanceManager(const FObjectInitializer& ObjectInitializer)
+ADEPRECATED_LightWeightInstanceManager::ADEPRECATED_LightWeightInstanceManager(const FObjectInitializer& ObjectInitializer)
 	:Super(ObjectInitializer)
 {
 	bReplicates = true;
@@ -31,29 +33,29 @@ ALightWeightInstanceManager::ALightWeightInstanceManager(const FObjectInitialize
 	}
 }
 
-ALightWeightInstanceManager::~ALightWeightInstanceManager()
+ADEPRECATED_LightWeightInstanceManager::~ADEPRECATED_LightWeightInstanceManager()
 {
 	FLightWeightInstanceSubsystem::Get().RemoveManager(this);
 }
 
-void ALightWeightInstanceManager::SetRepresentedClass(UClass* ActorClass)
+void ADEPRECATED_LightWeightInstanceManager::SetRepresentedClass(UClass* ActorClass)
 {
 	RepresentedClass = ActorClass;
 }
 
-void ALightWeightInstanceManager::Tick(float DeltaSeconds)
+void ADEPRECATED_LightWeightInstanceManager::Tick(float DeltaSeconds)
 {
 	// do nothing
 }
 
-AActor* ALightWeightInstanceManager::FindActor(const FActorInstanceHandle& Handle)
+AActor* ADEPRECATED_LightWeightInstanceManager::FindActor(const FActorInstanceHandle& Handle)
 {
 	TObjectPtr<AActor>* FoundActor = Actors.Find(Handle.GetInstanceIndex());
 
 	return FoundActor ? *FoundActor : nullptr;
 }
 
-AActor* ALightWeightInstanceManager::FindOrCreateActor(const FActorInstanceHandle& Handle)
+AActor* ADEPRECATED_LightWeightInstanceManager::FindOrCreateActor(const FActorInstanceHandle& Handle)
 {
 	// make sure the handle doesn't have an actor already
 	if (ensure(!Handle.GetCachedActor()))
@@ -88,7 +90,7 @@ AActor* ALightWeightInstanceManager::FindOrCreateActor(const FActorInstanceHandl
 	return Handle.GetCachedActor();
 }
 
-AActor* ALightWeightInstanceManager::ConvertInstanceToActor(const FActorInstanceHandle& Handle)
+AActor* ADEPRECATED_LightWeightInstanceManager::ConvertInstanceToActor(const FActorInstanceHandle& Handle)
 {
 #if !UE_BUILD_SHIPPING
 	check(!InstanceToActorConversionsInProgress.Contains(Handle.GetInstanceIndex()));
@@ -115,7 +117,7 @@ AActor* ALightWeightInstanceManager::ConvertInstanceToActor(const FActorInstance
 	if (DestroyedActorIndices.Contains(Handle.GetInstanceIndex()))
 	{
 		UE_LOG(LogLightWeightInstance, Verbose,
-			TEXT("ALightWeightInstanceManager::ConvertInstanceToActor - manager [ %s ] had already spawned an actor for index [ %d ] that was destroyed"),
+			TEXT("ADEPRECATED_LightWeightInstanceManager::ConvertInstanceToActor - manager [ %s ] had already spawned an actor for index [ %d ] that was destroyed"),
 			*GetActorNameOrLabel(),
 			Handle.GetInstanceIndex());
 		return nullptr;
@@ -139,7 +141,7 @@ AActor* ALightWeightInstanceManager::ConvertInstanceToActor(const FActorInstance
 		check(Handle.GetCachedActor() == NewActor);
 		check(NewActor == Actors.FindRef(Handle.GetInstanceIndex()));
 
-		NewActor->OnDestroyed.AddUniqueDynamic(this, &ALightWeightInstanceManager::OnSpawnedActorDestroyed);
+		NewActor->OnDestroyed.AddUniqueDynamic(this, &ADEPRECATED_LightWeightInstanceManager::OnSpawnedActorDestroyed);
 
 		PostActorSpawn(Handle);
 	}
@@ -147,12 +149,12 @@ AActor* ALightWeightInstanceManager::ConvertInstanceToActor(const FActorInstance
 	return NewActor;
 }
 
-void ALightWeightInstanceManager::OnSpawnedActorDestroyed(AActor* DestroyedActor)
+void ADEPRECATED_LightWeightInstanceManager::OnSpawnedActorDestroyed(AActor* DestroyedActor)
 {
 	OnSpawnedActorDestroyed(DestroyedActor, FindIndexForActor(DestroyedActor));
 }
 
-void ALightWeightInstanceManager::OnSpawnedActorDestroyed(AActor* DestroyedActor, const int32 DestroyedActorInstanceIndex)
+void ADEPRECATED_LightWeightInstanceManager::OnSpawnedActorDestroyed(AActor* DestroyedActor, const int32 DestroyedActorInstanceIndex)
 {
 	check(DestroyedActor);
 	if (!ensure(DestroyedActorInstanceIndex != INDEX_NONE))
@@ -173,13 +175,13 @@ void ALightWeightInstanceManager::OnSpawnedActorDestroyed(AActor* DestroyedActor
 	else
 	{
 		UE_LOG(LogLightWeightInstance, Error,
-			TEXT("ALightWeightInstanceManager::OnSpawnedActorDestroyed - manager [ %s ] DestroyedActorIndices already has an entry for index [ %d ]"),
+			TEXT("ADEPRECATED_LightWeightInstanceManager::OnSpawnedActorDestroyed - manager [ %s ] DestroyedActorIndices already has an entry for index [ %d ]"),
 			*GetActorNameOrLabel(),
 			DestroyedActorInstanceIndex);
 	}
 }
 
-int32 ALightWeightInstanceManager::FindIndexForActor(const AActor* InActor) const
+int32 ADEPRECATED_LightWeightInstanceManager::FindIndexForActor(const AActor* InActor) const
 {
 	for (const TPair<int32, TObjectPtr<AActor>>& ActorPair : Actors)
 	{
@@ -192,7 +194,7 @@ int32 ALightWeightInstanceManager::FindIndexForActor(const AActor* InActor) cons
 	return INDEX_NONE;
 }
 
-FActorInstanceHandle ALightWeightInstanceManager::ConvertActorToLightWeightInstance(AActor* InActor)
+FActorInstanceHandle ADEPRECATED_LightWeightInstanceManager::ConvertActorToLightWeightInstance(AActor* InActor)
 {
 	if (!InActor)
 	{
@@ -237,13 +239,13 @@ FActorInstanceHandle ALightWeightInstanceManager::ConvertActorToLightWeightInsta
 	}
 }
 
-FLWIData* ALightWeightInstanceManager::AllocateInitData() const
+FLWIData* ADEPRECATED_LightWeightInstanceManager::AllocateInitData() const
 {
 	return new FLWIData;
 }
 
 
-bool ALightWeightInstanceManager::SetDataFromActor(FLWIData* InData, AActor* InActor) const
+bool ADEPRECATED_LightWeightInstanceManager::SetDataFromActor(FLWIData* InData, AActor* InActor) const
 {
 	if (!InData || !InActor)
 	{
@@ -255,17 +257,17 @@ bool ALightWeightInstanceManager::SetDataFromActor(FLWIData* InData, AActor* InA
 	return true;
 }
 
-int32 ALightWeightInstanceManager::ConvertCollisionIndexToInstanceIndex(int32 InIndex, const UPrimitiveComponent* RelevantComponent) const
+int32 ADEPRECATED_LightWeightInstanceManager::ConvertCollisionIndexToInstanceIndex(int32 InIndex, const UPrimitiveComponent* RelevantComponent) const
 {
 	return InIndex;
 }
 
-int32 ALightWeightInstanceManager::GetGridSize() const
+int32 ADEPRECATED_LightWeightInstanceManager::GetGridSize() const
 {
 	return LWIGridSize;	
 }
 
-FInt32Vector3 ALightWeightInstanceManager::ConvertPositionToCoord(const FVector& InPosition) const
+FInt32Vector3 ADEPRECATED_LightWeightInstanceManager::ConvertPositionToCoord(const FVector& InPosition) const
 {
 	float GridSize = GetGridSize();
 	if (GridSize > 0)
@@ -280,7 +282,7 @@ FInt32Vector3 ALightWeightInstanceManager::ConvertPositionToCoord(const FVector&
 	return FInt32Vector3::ZeroValue;
 }
 
-FBox ALightWeightInstanceManager::ConvertPositionToGridBounds(const FVector& InPosition) const
+FBox ADEPRECATED_LightWeightInstanceManager::ConvertPositionToGridBounds(const FVector& InPosition) const
 {
 	float GridSize = GetGridSize();
 	if (GridSize > 0)
@@ -296,30 +298,30 @@ FBox ALightWeightInstanceManager::ConvertPositionToGridBounds(const FVector& InP
 	return FBox(InPosition, InPosition);
 }
 
-void ALightWeightInstanceManager::SetSpawnParameters(FActorSpawnParameters& SpawnParams)
+void ADEPRECATED_LightWeightInstanceManager::SetSpawnParameters(FActorSpawnParameters& SpawnParams)
 {
 	SpawnParams.OverrideLevel = GetLevel();
 	SpawnParams.SpawnCollisionHandlingOverride = ESpawnActorCollisionHandlingMethod::AlwaysSpawn;
 	SpawnParams.ObjectFlags = RF_Transactional;
 }
 
-UClass* ALightWeightInstanceManager::GetActorClassToSpawn(const FActorInstanceHandle& Handle) const
+UClass* ADEPRECATED_LightWeightInstanceManager::GetActorClassToSpawn(const FActorInstanceHandle& Handle) const
 {
 	return RepresentedClass;
 }
 
-void ALightWeightInstanceManager::PreSpawnInitalization(const FActorInstanceHandle& Handle, AActor* SpawnedActor)
+void ADEPRECATED_LightWeightInstanceManager::PreSpawnInitalization(const FActorInstanceHandle& Handle, AActor* SpawnedActor)
 {
 	Handle.SetCachedActor(SpawnedActor);
 	Actors.Add(Handle.GetInstanceIndex(), SpawnedActor);
 }
 
-void ALightWeightInstanceManager::PostActorSpawn(const FActorInstanceHandle& Handle)
+void ADEPRECATED_LightWeightInstanceManager::PostActorSpawn(const FActorInstanceHandle& Handle)
 {
 	// do nothing
 }
 
-bool ALightWeightInstanceManager::IsIndexValid(int32 Index) const
+bool ADEPRECATED_LightWeightInstanceManager::IsIndexValid(int32 Index) const
 {
 	if (ValidIndices.IsValidIndex(Index))
 	{
@@ -329,7 +331,7 @@ bool ALightWeightInstanceManager::IsIndexValid(int32 Index) const
 	return false;
 }
 
-bool ALightWeightInstanceManager::FindActorForHandle(const FActorInstanceHandle& Handle) const
+bool ADEPRECATED_LightWeightInstanceManager::FindActorForHandle(const FActorInstanceHandle& Handle) const
 {
 	if (Handle.GetCachedActor())
 	{
@@ -342,17 +344,17 @@ bool ALightWeightInstanceManager::FindActorForHandle(const FActorInstanceHandle&
 	return ResolvedActor != nullptr;
 }
 
-FVector ALightWeightInstanceManager::GetLocation(const FActorInstanceHandle& Handle) const
+FVector ADEPRECATED_LightWeightInstanceManager::GetLocation(const FActorInstanceHandle& Handle) const
 {
 	return GetTransform(Handle).GetTranslation();
 }
 
-FRotator ALightWeightInstanceManager::GetRotation(const FActorInstanceHandle& Handle) const
+FRotator ADEPRECATED_LightWeightInstanceManager::GetRotation(const FActorInstanceHandle& Handle) const
 {
 	return GetTransform(Handle).Rotator();
 }
 
-FTransform ALightWeightInstanceManager::GetTransform(const FActorInstanceHandle& Handle) const
+FTransform ADEPRECATED_LightWeightInstanceManager::GetTransform(const FActorInstanceHandle& Handle) const
 {
 	if (FindActorForHandle(Handle))
 	{
@@ -369,7 +371,7 @@ FTransform ALightWeightInstanceManager::GetTransform(const FActorInstanceHandle&
 	return FTransform();
 }
 
-FString ALightWeightInstanceManager::GetName(const FActorInstanceHandle& Handle) const
+FString ADEPRECATED_LightWeightInstanceManager::GetName(const FActorInstanceHandle& Handle) const
 {
 	if (FindActorForHandle(Handle))
 	{
@@ -380,27 +382,27 @@ FString ALightWeightInstanceManager::GetName(const FActorInstanceHandle& Handle)
 	return FString::Printf(TEXT("%s_%u"), *BaseInstanceName, Handle.GetInstanceIndex());
 }
 
-bool ALightWeightInstanceManager::DoesRepresentClass(const UClass* OtherClass) const
+bool ADEPRECATED_LightWeightInstanceManager::DoesRepresentClass(const UClass* OtherClass) const
 {
 	return OtherClass ? OtherClass->IsChildOf(GetRepresentedClassInternal()) : false;
 }
 
-bool ALightWeightInstanceManager::DoesAcceptClass(const UClass* OtherClass) const
+bool ADEPRECATED_LightWeightInstanceManager::DoesAcceptClass(const UClass* OtherClass) const
 {
 	return OtherClass ? OtherClass->IsChildOf(GetAcceptedClass()) : false;
 }
 
-UClass* ALightWeightInstanceManager::GetRepresentedClassInternal() const
+UClass* ADEPRECATED_LightWeightInstanceManager::GetRepresentedClassInternal() const
 {
 	return IsValid(RepresentedClass) ? RepresentedClass : nullptr;
 }
 
-UClass* ALightWeightInstanceManager::GetAcceptedClass() const
+UClass* ADEPRECATED_LightWeightInstanceManager::GetAcceptedClass() const
 {
 	return IsValid(AcceptedClass) ? AcceptedClass : nullptr;
 }
 
-int32 ALightWeightInstanceManager::AddNewInstance(FLWIData* InitData)
+int32 ADEPRECATED_LightWeightInstanceManager::AddNewInstance(FLWIData* InitData)
 {
 	if (!InitData)
 	{
@@ -416,7 +418,7 @@ int32 ALightWeightInstanceManager::AddNewInstance(FLWIData* InitData)
 	return DataIdx;
 }
 
-void ALightWeightInstanceManager::AddNewInstanceAt(FLWIData* InitData, int32 Index)
+void ADEPRECATED_LightWeightInstanceManager::AddNewInstanceAt(FLWIData* InitData, int32 Index)
 {
 	// allocate space on the end of the array if we need to
 	if (Index >= ValidIndices.Num())
@@ -430,7 +432,7 @@ void ALightWeightInstanceManager::AddNewInstanceAt(FLWIData* InitData, int32 Ind
 	ValidIndices[Index] = true;
 }
 
-void ALightWeightInstanceManager::GrowDataArrays()
+void ADEPRECATED_LightWeightInstanceManager::GrowDataArrays()
 {
 	InstanceTransforms.AddUninitialized();
 	ValidIndices.AddUninitialized();
@@ -438,13 +440,13 @@ void ALightWeightInstanceManager::GrowDataArrays()
 	check(InstanceTransforms.Num() == ValidIndices.Num());
 }
 
-void ALightWeightInstanceManager::UpdateDataAtIndex(FLWIData* InData, int32 Index)
+void ADEPRECATED_LightWeightInstanceManager::UpdateDataAtIndex(FLWIData* InData, int32 Index)
 {
 	// Convert to local space.
 	InstanceTransforms[Index] = InData->Transform * GetActorTransform().Inverse();
 }
 
-void ALightWeightInstanceManager::RemoveInstance(const int32 Index)
+void ADEPRECATED_LightWeightInstanceManager::RemoveInstance(const int32 Index)
 {
 	if (ensure(IsIndexValid(Index)))
 	{
@@ -461,7 +463,7 @@ void ALightWeightInstanceManager::RemoveInstance(const int32 Index)
 	}
 }
 
-bool ALightWeightInstanceManager::HasAnyValidInstancesOrManagedActors() const
+bool ADEPRECATED_LightWeightInstanceManager::HasAnyValidInstancesOrManagedActors() const
 {
 	bool bHasAnyValidIndex = false;
 	for (const bool bIsIndexValid : ValidIndices)
@@ -475,30 +477,32 @@ bool ALightWeightInstanceManager::HasAnyValidInstancesOrManagedActors() const
 	return bHasAnyValidIndex;
 }
 
-void ALightWeightInstanceManager::GetLifetimeReplicatedProps(TArray<FLifetimeProperty>& OutLifetimeProps) const
+void ADEPRECATED_LightWeightInstanceManager::GetLifetimeReplicatedProps(TArray<FLifetimeProperty>& OutLifetimeProps) const
 {
 	Super::GetLifetimeReplicatedProps(OutLifetimeProps);
 
-	DOREPLIFETIME(ALightWeightInstanceManager, RepresentedClass);
-	DOREPLIFETIME(ALightWeightInstanceManager, InstanceTransforms);
-	DOREPLIFETIME(ALightWeightInstanceManager, FreeIndices);
-	DOREPLIFETIME(ALightWeightInstanceManager, ValidIndices);
+	DOREPLIFETIME(ADEPRECATED_LightWeightInstanceManager, RepresentedClass);
+	DOREPLIFETIME(ADEPRECATED_LightWeightInstanceManager, InstanceTransforms);
+	DOREPLIFETIME(ADEPRECATED_LightWeightInstanceManager, FreeIndices);
+	DOREPLIFETIME(ADEPRECATED_LightWeightInstanceManager, ValidIndices);
 }
 
-void ALightWeightInstanceManager::OnRep_Transforms()
+void ADEPRECATED_LightWeightInstanceManager::OnRep_Transforms()
 {
 	// do nothing
 }
 #if WITH_EDITOR
-void ALightWeightInstanceManager::PostEditChangeProperty(struct FPropertyChangedEvent& PropertyChangedEvent)
+void ADEPRECATED_LightWeightInstanceManager::PostEditChangeProperty(struct FPropertyChangedEvent& PropertyChangedEvent)
 {
 	Super::PostEditChangeProperty(PropertyChangedEvent);
 
 	const FName Name = (PropertyChangedEvent.Property != nullptr) ? PropertyChangedEvent.Property->GetFName() : NAME_None;
 
-	if (Name == GET_MEMBER_NAME_CHECKED(ALightWeightInstanceManager, RepresentedClass))
+	if (Name == GET_MEMBER_NAME_CHECKED(ADEPRECATED_LightWeightInstanceManager, RepresentedClass))
 	{
 		SetRepresentedClass(RepresentedClass);
 	}
 }
 #endif
+
+PRAGMA_ENABLE_DEPRECATION_WARNINGS
diff --git a/Engine/Source/Runtime/Engine/Private/LightWeightInstanceStaticMeshManager.cpp b/Engine/Source/Runtime/Engine/Private/LightWeightInstanceStaticMeshManager.cpp
--- a/Engine/Source/Runtime/Engine/Private/LightWeightInstanceStaticMeshManager.cpp
+++ b/Engine/Source/Runtime/Engine/Private/LightWeightInstanceStaticMeshManager.cpp
@@ -18,14 +18,14 @@
 #include "TimerManager.h"
 #endif // WITH_EDITOR
 
-ALightWeightInstanceStaticMeshManager::ALightWeightInstanceStaticMeshManager(const FObjectInitializer& ObjectInitializer)
+PRAGMA_DISABLE_DEPRECATION_WARNINGS
+
+ADEPRECATED_LightWeightInstanceStaticMeshManager::ADEPRECATED_LightWeightInstanceStaticMeshManager(const FObjectInitializer& ObjectInitializer)
 	:Super(ObjectInitializer)
 {
 #if WITH_EDITORONLY_DATA
-PRAGMA_DISABLE_DEPRECATION_WARNINGS
 	// Still create the deprecated InstancedStaticMeshComponent so the old data can be serialized during loading
 	InstancedStaticMeshComponent_DEPRECATED = CreateDefaultSubobject<UHierarchicalInstancedStaticMeshComponent>(TEXT("InstancedStaticMeshComponent0"));
-PRAGMA_ENABLE_DEPRECATION_WARNINGS
 #endif
 
 	ISMComponent = CreateDefaultSubobject<UInstancedStaticMeshComponent>(TEXT("ISMComponent0"));
@@ -42,11 +42,10 @@ PRAGMA_ENABLE_DEPRECATION_WARNINGS
 }
 
 #if WITH_EDITORONLY_DATA
-void ALightWeightInstanceStaticMeshManager::PostLoad()
+void ADEPRECATED_LightWeightInstanceStaticMeshManager::PostLoad()
 {
 	Super::PostLoad();
 
-PRAGMA_DISABLE_DEPRECATION_WARNINGS
 	if (InstancedStaticMeshComponent_DEPRECATED != nullptr && InstancedStaticMeshComponent_DEPRECATED->PerInstanceSMData.Num() > 0)
 	{
 		// assume new ISM component is still empty
@@ -88,11 +87,10 @@ PRAGMA_DISABLE_DEPRECATION_WARNINGS
 			RenderingIndicesToDataIndices[DataIndex] = DataIndex;
 		}
 	}
-PRAGMA_ENABLE_DEPRECATION_WARNINGS
 }
 #endif
 
-void ALightWeightInstanceStaticMeshManager::SetRepresentedClass(UClass* ActorClass)
+void ADEPRECATED_LightWeightInstanceStaticMeshManager::SetRepresentedClass(UClass* ActorClass)
 {
 	Super::SetRepresentedClass(ActorClass);
 
@@ -110,13 +108,11 @@ void ALightWeightInstanceStaticMeshManager::SetRepresentedClass(UClass* ActorCla
 
 	if (ISMComponent)
 	{
-PRAGMA_DISABLE_DEPRECATION_WARNINGS
 		ISMComponent->OnPostPopulatePerInstanceData();
-PRAGMA_ENABLE_DEPRECATION_WARNINGS
 	}
 }
 
-int32 ALightWeightInstanceStaticMeshManager::ConvertCollisionIndexToInstanceIndex(int32 InIndex, const UPrimitiveComponent* RelevantComponent) const
+int32 ADEPRECATED_LightWeightInstanceStaticMeshManager::ConvertCollisionIndexToInstanceIndex(int32 InIndex, const UPrimitiveComponent* RelevantComponent) const
 {
 	if (ensureMsgf(RenderingIndicesToDataIndices.IsValidIndex(InIndex), TEXT("Invalid index [ %d ]"), InIndex))
 	{
@@ -126,13 +122,13 @@ int32 ALightWeightInstanceStaticMeshManager::ConvertCollisionIndexToInstanceInde
 	return InIndex;
 }
 
-void ALightWeightInstanceStaticMeshManager::AddNewInstanceAt(FLWIData* InitData, int32 Index)
+void ADEPRECATED_LightWeightInstanceStaticMeshManager::AddNewInstanceAt(FLWIData* InitData, int32 Index)
 {
 	Super::AddNewInstanceAt(InitData, Index);
 	AddInstanceToRendering(Index);
 }
 
-void ALightWeightInstanceStaticMeshManager::RemoveInstance(const int32 Index)
+void ADEPRECATED_LightWeightInstanceStaticMeshManager::RemoveInstance(const int32 Index)
 {
 #if WITH_EDITOR
 	Modify();
@@ -142,7 +138,7 @@ void ALightWeightInstanceStaticMeshManager::RemoveInstance(const int32 Index)
 	Super::RemoveInstance(Index);
 }
 
-void ALightWeightInstanceStaticMeshManager::AddInstanceToRendering(int32 DataIndex)
+void ADEPRECATED_LightWeightInstanceStaticMeshManager::AddInstanceToRendering(int32 DataIndex)
 {
 	if (!ensureMsgf(RenderingIndicesToDataIndices.Contains(DataIndex) == false, TEXT("LWI rendering instance added more than once. Index: %d"), DataIndex))
 	{
@@ -156,7 +152,7 @@ void ALightWeightInstanceStaticMeshManager::AddInstanceToRendering(int32 DataInd
 	const int32 RenderingIdx = RenderingIndicesToDataIndices.Add(DataIndex);
 
 	UE_LOG(LogLightWeightInstance, Verbose,
-		TEXT("ALightWeightInstanceStaticMeshManager::AddInstanceToRendering - manager [ %s ] adding instance at data index [ %d ] rendering index [ %d ], did cancel pending delete? [ %s ]"),
+		TEXT("ADEPRECATED_LightWeightInstanceStaticMeshManager::AddInstanceToRendering - manager [ %s ] adding instance at data index [ %d ] rendering index [ %d ], did cancel pending delete? [ %s ]"),
 		*GetActorNameOrLabel(),
 		DataIndex,
 		RenderingIdx,
@@ -179,7 +175,7 @@ void ALightWeightInstanceStaticMeshManager::AddInstanceToRendering(int32 DataInd
 	}
 }
 
-void ALightWeightInstanceStaticMeshManager::RemoveInstanceFromRendering(int32 DataIndex)
+void ADEPRECATED_LightWeightInstanceStaticMeshManager::RemoveInstanceFromRendering(int32 DataIndex)
 {
 	if (IsIndexValid(DataIndex))
 	{
@@ -189,25 +185,25 @@ void ALightWeightInstanceStaticMeshManager::RemoveInstanceFromRendering(int32 Da
 #if WITH_EDITOR
 			if (GEditor)
 			{
-				GEditor->GetTimerManager()->SetTimerForNextTick(this, &ALightWeightInstanceStaticMeshManager::PostRemoveInstanceFromRendering);
+				GEditor->GetTimerManager()->SetTimerForNextTick(this, &ADEPRECATED_LightWeightInstanceStaticMeshManager::PostRemoveInstanceFromRendering);
 			}
 			else
 #endif // WITH_EDITOR
 			{
-				GetWorld()->GetTimerManager().SetTimerForNextTick(this, &ALightWeightInstanceStaticMeshManager::PostRemoveInstanceFromRendering);
+				GetWorld()->GetTimerManager().SetTimerForNextTick(this, &ADEPRECATED_LightWeightInstanceStaticMeshManager::PostRemoveInstanceFromRendering);
 			}
 		}
 
 		DataIndicesToBeDeleted.AddUnique(DataIndex);
 
 		UE_LOG(LogLightWeightInstance, Verbose,
-			TEXT("ALightWeightInstanceStaticMeshManager::RemoveInstanceFromRendering - manager [ %s ] removing instance at data index [ %d ]"),
+			TEXT("ADEPRECATED_LightWeightInstanceStaticMeshManager::RemoveInstanceFromRendering - manager [ %s ] removing instance at data index [ %d ]"),
 			*GetActorNameOrLabel(),
 			DataIndex);
 	}
 }
 
-void ALightWeightInstanceStaticMeshManager::PostRemoveInstanceFromRendering()
+void ADEPRECATED_LightWeightInstanceStaticMeshManager::PostRemoveInstanceFromRendering()
 {
 	TArray<int32> RenderingIndicesToBeDeleted;
 	for (int32 DataIndex : DataIndicesToBeDeleted)
@@ -215,7 +211,7 @@ void ALightWeightInstanceStaticMeshManager::PostRemoveInstanceFromRendering()
 		RenderingIndicesToBeDeleted.Add(DataIndicesToRenderingIndices[DataIndex]);
 
 		UE_LOG(LogLightWeightInstance, Verbose,
-			TEXT("ALightWeightInstanceStaticMeshManager::PostRemoveInstanceFromRendering - manager [ %s ] removing instance at data index [ %d ] rendering index [ %d ]"),
+			TEXT("ADEPRECATED_LightWeightInstanceStaticMeshManager::PostRemoveInstanceFromRendering - manager [ %s ] removing instance at data index [ %d ] rendering index [ %d ]"),
 			*GetActorNameOrLabel(),
 			DataIndex,
 			DataIndicesToRenderingIndices[DataIndex]);
@@ -249,21 +245,21 @@ void ALightWeightInstanceStaticMeshManager::PostRemoveInstanceFromRendering()
 	DataIndicesToBeDeleted.Reset();
 }
 
-void ALightWeightInstanceStaticMeshManager::GetLifetimeReplicatedProps(TArray<FLifetimeProperty>& OutLifetimeProps) const
+void ADEPRECATED_LightWeightInstanceStaticMeshManager::GetLifetimeReplicatedProps(TArray<FLifetimeProperty>& OutLifetimeProps) const
 {
 	Super::GetLifetimeReplicatedProps(OutLifetimeProps);
 
-	DOREPLIFETIME(ALightWeightInstanceStaticMeshManager, StaticMesh);
-	DOREPLIFETIME(ALightWeightInstanceStaticMeshManager, RenderingIndicesToDataIndices);
-	DOREPLIFETIME(ALightWeightInstanceStaticMeshManager, DataIndicesToRenderingIndices);
+	DOREPLIFETIME(ADEPRECATED_LightWeightInstanceStaticMeshManager, StaticMesh);
+	DOREPLIFETIME(ADEPRECATED_LightWeightInstanceStaticMeshManager, RenderingIndicesToDataIndices);
+	DOREPLIFETIME(ADEPRECATED_LightWeightInstanceStaticMeshManager, DataIndicesToRenderingIndices);
 }
 
-void ALightWeightInstanceStaticMeshManager::OnRep_StaticMesh()
+void ADEPRECATED_LightWeightInstanceStaticMeshManager::OnRep_StaticMesh()
 {
 	OnStaticMeshSet();
 }
 
-void ALightWeightInstanceStaticMeshManager::OnRep_Transforms()
+void ADEPRECATED_LightWeightInstanceStaticMeshManager::OnRep_Transforms()
 {
 	Super::OnRep_Transforms();
 
@@ -273,7 +269,7 @@ void ALightWeightInstanceStaticMeshManager::OnRep_Transforms()
 	}
 }
 
-void ALightWeightInstanceStaticMeshManager::PostActorSpawn(const FActorInstanceHandle& Handle)
+void ADEPRECATED_LightWeightInstanceStaticMeshManager::PostActorSpawn(const FActorInstanceHandle& Handle)
 {
 	Super::PostActorSpawn(Handle);
 
@@ -281,7 +277,7 @@ void ALightWeightInstanceStaticMeshManager::PostActorSpawn(const FActorInstanceH
 	RemoveInstanceFromRendering(Handle.GetInstanceIndex());
 }
 
-void ALightWeightInstanceStaticMeshManager::SetInstancedStaticMeshParams()
+void ADEPRECATED_LightWeightInstanceStaticMeshManager::SetInstancedStaticMeshParams()
 {
 	FName CollisionProfileName("LightWeightInstancedStaticMeshPhysics");
 	ISMComponent->SetCollisionProfileName(CollisionProfileName);
@@ -305,13 +301,13 @@ void ALightWeightInstanceStaticMeshManager::SetInstancedStaticMeshParams()
 	}
 } 
 
-void ALightWeightInstanceStaticMeshManager::SetStaticMeshFromActor(AActor* InActor)
+void ADEPRECATED_LightWeightInstanceStaticMeshManager::SetStaticMeshFromActor(AActor* InActor)
 {
 	ensureMsgf(false, TEXT("ALightWeightInstanceManager::SetStaticMeshFromActor was called on %s. Projects should override this function in subclasses."), *GetNameSafe(InActor));
 	ClearStaticMesh();
 }
 
-void ALightWeightInstanceStaticMeshManager::OnStaticMeshSet()
+void ADEPRECATED_LightWeightInstanceStaticMeshManager::OnStaticMeshSet()
 {
 #if WITH_EDITOR
 	Modify();
@@ -341,7 +337,7 @@ void ALightWeightInstanceStaticMeshManager::OnStaticMeshSet()
 	}
 }
 
-void ALightWeightInstanceStaticMeshManager::ClearStaticMesh()
+void ADEPRECATED_LightWeightInstanceStaticMeshManager::ClearStaticMesh()
 {
 #if WITH_EDITOR
 	Modify();
@@ -351,7 +347,7 @@ void ALightWeightInstanceStaticMeshManager::ClearStaticMesh()
 	OnStaticMeshSet();
 }
 
-FText ALightWeightInstanceStaticMeshManager::GetSMInstanceDisplayName(const FSMInstanceId& InstanceId) const
+FText ADEPRECATED_LightWeightInstanceStaticMeshManager::GetSMInstanceDisplayName(const FSMInstanceId& InstanceId) const
 {
 	check(InstanceId.ISMComponent == ISMComponent);
 
@@ -359,7 +355,7 @@ FText ALightWeightInstanceStaticMeshManager::GetSMInstanceDisplayName(const FSMI
 	return FText::Format(NSLOCTEXT("LightWeightInstanceStaticMeshManager", "DisplayNameFmt", "{0} - Instance {1}"), OwnerDisplayName, InstanceId.InstanceIndex);
 }
 
-FText ALightWeightInstanceStaticMeshManager::GetSMInstanceTooltip(const FSMInstanceId& InstanceId) const
+FText ADEPRECATED_LightWeightInstanceStaticMeshManager::GetSMInstanceTooltip(const FSMInstanceId& InstanceId) const
 {
 	check(InstanceId.ISMComponent == ISMComponent);
 	
@@ -367,7 +363,7 @@ FText ALightWeightInstanceStaticMeshManager::GetSMInstanceTooltip(const FSMInsta
 	return FText::Format(NSLOCTEXT("LightWeightInstanceStaticMeshManager", "TooltipFmt", "Instance {0} on {1}"), InstanceId.InstanceIndex, OwnerDisplayPath);
 }
 
-bool ALightWeightInstanceStaticMeshManager::CanEditSMInstance(const FSMInstanceId& InstanceId) const
+bool ADEPRECATED_LightWeightInstanceStaticMeshManager::CanEditSMInstance(const FSMInstanceId& InstanceId) const
 {
 	check(InstanceId.ISMComponent == ISMComponent);
 
@@ -375,7 +371,7 @@ bool ALightWeightInstanceStaticMeshManager::CanEditSMInstance(const FSMInstanceI
 	return InstanceManager->CanEditSMInstance(InstanceId);
 }
 
-bool ALightWeightInstanceStaticMeshManager::CanMoveSMInstance(const FSMInstanceId& InstanceId, const ETypedElementWorldType InWorldType) const
+bool ADEPRECATED_LightWeightInstanceStaticMeshManager::CanMoveSMInstance(const FSMInstanceId& InstanceId, const ETypedElementWorldType InWorldType) const
 {
 	check(InstanceId.ISMComponent == ISMComponent);
 
@@ -383,7 +379,7 @@ bool ALightWeightInstanceStaticMeshManager::CanMoveSMInstance(const FSMInstanceI
 	return InstanceManager->CanMoveSMInstance(InstanceId, InWorldType);
 }
 
-bool ALightWeightInstanceStaticMeshManager::GetSMInstanceTransform(const FSMInstanceId& InstanceId, FTransform& OutInstanceTransform, bool bWorldSpace) const
+bool ADEPRECATED_LightWeightInstanceStaticMeshManager::GetSMInstanceTransform(const FSMInstanceId& InstanceId, FTransform& OutInstanceTransform, bool bWorldSpace) const
 {
 	check(InstanceId.ISMComponent == ISMComponent);
 
@@ -397,7 +393,7 @@ bool ALightWeightInstanceStaticMeshManager::GetSMInstanceTransform(const FSMInst
 	return false;
 }
 
-bool ALightWeightInstanceStaticMeshManager::SetSMInstanceTransform(const FSMInstanceId& InstanceId, const FTransform& InstanceTransform, bool bWorldSpace, bool bMarkRenderStateDirty, bool bTeleport)
+bool ADEPRECATED_LightWeightInstanceStaticMeshManager::SetSMInstanceTransform(const FSMInstanceId& InstanceId, const FTransform& InstanceTransform, bool bWorldSpace, bool bMarkRenderStateDirty, bool bTeleport)
 {
 	check(InstanceId.ISMComponent == ISMComponent);
 
@@ -411,7 +407,7 @@ bool ALightWeightInstanceStaticMeshManager::SetSMInstanceTransform(const FSMInst
 	return false;
 }
 
-void ALightWeightInstanceStaticMeshManager::NotifySMInstanceMovementStarted(const FSMInstanceId& InstanceId)
+void ADEPRECATED_LightWeightInstanceStaticMeshManager::NotifySMInstanceMovementStarted(const FSMInstanceId& InstanceId)
 {
 	check(InstanceId.ISMComponent == ISMComponent);
 
@@ -419,7 +415,7 @@ void ALightWeightInstanceStaticMeshManager::NotifySMInstanceMovementStarted(cons
 	InstanceManager->NotifySMInstanceMovementStarted(InstanceId);
 }
 
-void ALightWeightInstanceStaticMeshManager::NotifySMInstanceMovementOngoing(const FSMInstanceId& InstanceId)
+void ADEPRECATED_LightWeightInstanceStaticMeshManager::NotifySMInstanceMovementOngoing(const FSMInstanceId& InstanceId)
 {
 	check(InstanceId.ISMComponent == ISMComponent);
 
@@ -427,7 +423,7 @@ void ALightWeightInstanceStaticMeshManager::NotifySMInstanceMovementOngoing(cons
 	InstanceManager->NotifySMInstanceMovementOngoing(InstanceId);
 }
 
-void ALightWeightInstanceStaticMeshManager::NotifySMInstanceMovementEnded(const FSMInstanceId& InstanceId)
+void ADEPRECATED_LightWeightInstanceStaticMeshManager::NotifySMInstanceMovementEnded(const FSMInstanceId& InstanceId)
 {
 	check(InstanceId.ISMComponent == ISMComponent);
 
@@ -435,7 +431,7 @@ void ALightWeightInstanceStaticMeshManager::NotifySMInstanceMovementEnded(const
 	InstanceManager->NotifySMInstanceMovementEnded(InstanceId);
 }
 
-void ALightWeightInstanceStaticMeshManager::NotifySMInstanceSelectionChanged(const FSMInstanceId& InstanceId, const bool bIsSelected)
+void ADEPRECATED_LightWeightInstanceStaticMeshManager::NotifySMInstanceSelectionChanged(const FSMInstanceId& InstanceId, const bool bIsSelected)
 {
 	check(InstanceId.ISMComponent == ISMComponent);
 
@@ -443,7 +439,7 @@ void ALightWeightInstanceStaticMeshManager::NotifySMInstanceSelectionChanged(con
 	InstanceManager->NotifySMInstanceSelectionChanged(InstanceId, bIsSelected);
 }
 
-bool ALightWeightInstanceStaticMeshManager::DeleteSMInstances(TArrayView<const FSMInstanceId> InstanceIds)
+bool ADEPRECATED_LightWeightInstanceStaticMeshManager::DeleteSMInstances(TArrayView<const FSMInstanceId> InstanceIds)
 {
 	TArray<int32> DataIndices;
 	GetLWIDataIndices(InstanceIds, DataIndices);
@@ -462,7 +458,7 @@ bool ALightWeightInstanceStaticMeshManager::DeleteSMInstances(TArrayView<const F
 	return true;
 }
 
-bool ALightWeightInstanceStaticMeshManager::DuplicateSMInstances(TArrayView<const FSMInstanceId> InstanceIds, TArray<FSMInstanceId>& OutNewInstanceIds)
+bool ADEPRECATED_LightWeightInstanceStaticMeshManager::DuplicateSMInstances(TArrayView<const FSMInstanceId> InstanceIds, TArray<FSMInstanceId>& OutNewInstanceIds)
 {
 	TArray<int32> DataIndices;
 	GetLWIDataIndices(InstanceIds, DataIndices);
@@ -482,7 +478,7 @@ bool ALightWeightInstanceStaticMeshManager::DuplicateSMInstances(TArrayView<cons
 	return true;
 }
 
-void ALightWeightInstanceStaticMeshManager::DuplicateLWIInstances(TArrayView<const int32> DataIndices, TArray<int32>& OutNewDataIndices)
+void ADEPRECATED_LightWeightInstanceStaticMeshManager::DuplicateLWIInstances(TArrayView<const int32> DataIndices, TArray<int32>& OutNewDataIndices)
 {
 	OutNewDataIndices.Reset(DataIndices.Num());
 	for (int32 DataIndex : DataIndices)
@@ -498,7 +494,7 @@ void ALightWeightInstanceStaticMeshManager::DuplicateLWIInstances(TArrayView<con
 	}
 }
 
-void ALightWeightInstanceStaticMeshManager::GetLWIDataIndices(TArrayView<const FSMInstanceId> InstanceIds, TArray<int32>& OutDataIndices) const
+void ADEPRECATED_LightWeightInstanceStaticMeshManager::GetLWIDataIndices(TArrayView<const FSMInstanceId> InstanceIds, TArray<int32>& OutDataIndices) const
 {
 	OutDataIndices.Reset(InstanceIds.Num());
 	for (const FSMInstanceId& InstanceId : InstanceIds)
@@ -511,13 +507,14 @@ void ALightWeightInstanceStaticMeshManager::GetLWIDataIndices(TArrayView<const F
 	}
 }
 
-int32 ALightWeightInstanceStaticMeshManager::ConvertInternalIndexToHandleIndex(int32 InInternalIndex) const
+int32 ADEPRECATED_LightWeightInstanceStaticMeshManager::ConvertInternalIndexToHandleIndex(int32 InInternalIndex) const
 {
 	return DataIndicesToRenderingIndices[InInternalIndex];
 }
 
-int32 ALightWeightInstanceStaticMeshManager::ConvertHandleIndexToInternalIndex(int32 InHandleIndex) const
+int32 ADEPRECATED_LightWeightInstanceStaticMeshManager::ConvertHandleIndexToInternalIndex(int32 InHandleIndex) const
 {
 	return RenderingIndicesToDataIndices[InHandleIndex];
 }
 
+PRAGMA_ENABLE_DEPRECATION_WARNINGS
diff --git a/Engine/Source/Runtime/Engine/Private/LightWeightInstanceSubsystem.cpp b/Engine/Source/Runtime/Engine/Private/LightWeightInstanceSubsystem.cpp
--- a/Engine/Source/Runtime/Engine/Private/LightWeightInstanceSubsystem.cpp
+++ b/Engine/Source/Runtime/Engine/Private/LightWeightInstanceSubsystem.cpp
@@ -8,6 +8,8 @@
 
 DEFINE_LOG_CATEGORY(LogLightWeightInstance);
 
+PRAGMA_DISABLE_DEPRECATION_WARNINGS
+
 FLightWeightInstanceSubsystem::FLightWeightInstanceSubsystem()
 {
 #if WITH_EDITOR
@@ -15,14 +17,14 @@ FLightWeightInstanceSubsystem::FLightWeightInstanceSubsystem()
 	{
 		OnLevelActorAddedHandle = GEngine->OnLevelActorAdded().AddLambda([this](AActor* Actor)
 			{
-				if (ALightWeightInstanceManager* LWIManager = Cast<ALightWeightInstanceManager>(Actor))
+				if (ADEPRECATED_LightWeightInstanceManager* LWIManager = Cast<ADEPRECATED_LightWeightInstanceManager>(Actor))
 				{
 					Get().AddManager(LWIManager);
 				}
 			});
 		OnLevelActorDeletedHandle = GEngine->OnLevelActorDeleted().AddLambda([this](AActor* Actor)
 			{
-				if (ALightWeightInstanceManager* LWIManager = Cast<ALightWeightInstanceManager>(Actor))
+				if (ADEPRECATED_LightWeightInstanceManager* LWIManager = Cast<ADEPRECATED_LightWeightInstanceManager>(Actor))
 				{
 					Get().RemoveManager(LWIManager);
 				}
@@ -42,7 +44,7 @@ FLightWeightInstanceSubsystem::~FLightWeightInstanceSubsystem()
 #endif // WITH_EDITOR
 }
 
-int32 FLightWeightInstanceSubsystem::GetManagerIndex(const ALightWeightInstanceManager* Manager) const
+int32 FLightWeightInstanceSubsystem::GetManagerIndex(const ADEPRECATED_LightWeightInstanceManager* Manager) const
 {
 	UE::TReadScopeLock Lock(LWIManagersRWLock);
 	for (int32 Idx = 0; Idx < LWInstanceManagers.Num(); ++Idx)
@@ -56,29 +58,29 @@ int32 FLightWeightInstanceSubsystem::GetManagerIndex(const ALightWeightInstanceM
 	return INDEX_NONE;
 }
 
-const ALightWeightInstanceManager* FLightWeightInstanceSubsystem::GetManagerAt(int32 Index) const
+const ADEPRECATED_LightWeightInstanceManager* FLightWeightInstanceSubsystem::GetManagerAt(int32 Index) const
 {
 	UE::TReadScopeLock Lock(LWIManagersRWLock);
 	return LWInstanceManagers.IsValidIndex(Index) ? LWInstanceManagers[Index] : nullptr;
 }
 
-bool FLightWeightInstanceSubsystem::AddManager(ALightWeightInstanceManager* Manager)
+bool FLightWeightInstanceSubsystem::AddManager(ADEPRECATED_LightWeightInstanceManager* Manager)
 {
 	UE::TWriteScopeLock Lock(LWIManagersRWLock);
 	return Get().LWInstanceManagers.Add(Manager) ? true : false;
 }
 
-bool FLightWeightInstanceSubsystem::RemoveManager(ALightWeightInstanceManager* Manager)
+bool FLightWeightInstanceSubsystem::RemoveManager(ADEPRECATED_LightWeightInstanceManager* Manager)
 {
 	UE::TWriteScopeLock Lock(LWIManagersRWLock);
 	return Get().LWInstanceManagers.Remove(Manager) ? true : false;
 }
 
-ALightWeightInstanceManager* FLightWeightInstanceSubsystem::FindLightWeightInstanceManager(const FActorInstanceHandle& Handle) const
+ADEPRECATED_LightWeightInstanceManager* FLightWeightInstanceSubsystem::FindLightWeightInstanceManager(const FActorInstanceHandle& Handle) const
 {
 	if (Handle.GetManagerInterface().IsValid())
 	{
-		return Handle.GetManager<ALightWeightInstanceManager>();
+		return Handle.GetManager<ADEPRECATED_LightWeightInstanceManager>();
 	}
 
 	if (AActor* CachedActor = Handle.GetCachedActor())
@@ -88,7 +90,7 @@ ALightWeightInstanceManager* FLightWeightInstanceSubsystem::FindLightWeightInsta
 
 		UE::TReadScopeLock Lock(LWIManagersRWLock);
 		// see if we already have a match
-		for (ALightWeightInstanceManager* LWInstance : LWInstanceManagers)
+		for (ADEPRECATED_LightWeightInstanceManager* LWInstance : LWInstanceManagers)
 		{
 			if (CachedActorClass == LWInstance->GetRepresentedClass(Handle.GetInstanceIndex()))
 			{
@@ -117,25 +119,25 @@ ALightWeightInstanceManager* FLightWeightInstanceSubsystem::FindLightWeightInsta
 	return nullptr;
 }
 
-ALightWeightInstanceManager* FLightWeightInstanceSubsystem::FindLightWeightInstanceManager(UClass* ActorClass, const UDataLayerInstance* DataLayer, UWorld* World) const
+ADEPRECATED_LightWeightInstanceManager* FLightWeightInstanceSubsystem::FindLightWeightInstanceManager(UClass* ActorClass, const UDataLayerInstance* DataLayer, UWorld* World) const
 {
 	check(World);
 	return FindLightWeightInstanceManager(*ActorClass, *World, FVector::ZeroVector, DataLayer);
 }
 
 
-ALightWeightInstanceManager* FLightWeightInstanceSubsystem::FindOrAddLightWeightInstanceManager(UClass* ActorClass, const UDataLayerInstance* DataLayer, UWorld* World)
+ADEPRECATED_LightWeightInstanceManager* FLightWeightInstanceSubsystem::FindOrAddLightWeightInstanceManager(UClass* ActorClass, const UDataLayerInstance* DataLayer, UWorld* World)
 {
 	check(World);
 	return FindOrAddLightWeightInstanceManager(*ActorClass, *World, FVector::ZeroVector, DataLayer);
 }
 
 
-ALightWeightInstanceManager* FLightWeightInstanceSubsystem::FindLightWeightInstanceManager(UClass& ActorClass, UWorld& World, const FVector& InPos, const UDataLayerInstance* DataLayer) const
+ADEPRECATED_LightWeightInstanceManager* FLightWeightInstanceSubsystem::FindLightWeightInstanceManager(UClass& ActorClass, UWorld& World, const FVector& InPos, const UDataLayerInstance* DataLayer) const
 {
 	UE::TReadScopeLock Lock(LWIManagersRWLock);
 
-	for (ALightWeightInstanceManager* InstanceManager : LWInstanceManagers)
+	for (ADEPRECATED_LightWeightInstanceManager* InstanceManager : LWInstanceManagers)
 	{
 		if (IsValid(InstanceManager))
 		{
@@ -160,9 +162,9 @@ ALightWeightInstanceManager* FLightWeightInstanceSubsystem::FindLightWeightInsta
 	return nullptr;
 }
 
-ALightWeightInstanceManager* FLightWeightInstanceSubsystem::FindOrAddLightWeightInstanceManager(UClass& ActorClass, UWorld& World, const FVector& InPos, const UDataLayerInstance* DataLayer)
+ADEPRECATED_LightWeightInstanceManager* FLightWeightInstanceSubsystem::FindOrAddLightWeightInstanceManager(UClass& ActorClass, UWorld& World, const FVector& InPos, const UDataLayerInstance* DataLayer)
 {
-	if (ALightWeightInstanceManager* FoundManager = FindLightWeightInstanceManager(ActorClass, World, InPos, DataLayer))
+	if (ADEPRECATED_LightWeightInstanceManager* FoundManager = FindLightWeightInstanceManager(ActorClass, World, InPos, DataLayer))
 	{
 		return FoundManager;
 	}
@@ -181,7 +183,7 @@ ALightWeightInstanceManager* FLightWeightInstanceSubsystem::FindOrAddLightWeight
 
 	FTransform ManagerTransform(FTransform::Identity);
 
-	ALightWeightInstanceManager* InstanceManagerCDO = CastChecked<ALightWeightInstanceManager>(BestMatchingClass->GetDefaultObject(true));
+	ADEPRECATED_LightWeightInstanceManager* InstanceManagerCDO = CastChecked<ADEPRECATED_LightWeightInstanceManager>(BestMatchingClass->GetDefaultObject(true));
 	int32 GridSize = InstanceManagerCDO->GetGridSize();
 	if (GridSize > 0)
 	{
@@ -191,7 +193,7 @@ ALightWeightInstanceManager* FLightWeightInstanceSubsystem::FindOrAddLightWeight
 		ManagerTransform.SetLocation(ManagerLocation);
 	}
 
-	if (ALightWeightInstanceManager* NewInstanceManager = World.SpawnActor<ALightWeightInstanceManager>(BestMatchingClass, ManagerTransform, SpawnParams))
+	if (ADEPRECATED_LightWeightInstanceManager* NewInstanceManager = World.SpawnActor<ADEPRECATED_LightWeightInstanceManager>(BestMatchingClass, ManagerTransform, SpawnParams))
 	{
 		NewInstanceManager->SetRepresentedClass(&ActorClass);
 
@@ -222,7 +224,7 @@ UClass* FLightWeightInstanceSubsystem::FindBestInstanceManagerClass(const UClass
 	TArray<UClass*> ManagerClasses;
 	for (TObjectIterator<UClass> It; It; ++It)
 	{
-		if (It->IsChildOf(ALightWeightInstanceManager::StaticClass()) && !It->HasAnyClassFlags(CLASS_Abstract))
+		if (It->IsChildOf(ADEPRECATED_LightWeightInstanceManager::StaticClass()) && !It->HasAnyClassFlags(CLASS_Abstract))
 		{
 			// Skip SKEL and REINST classes.
 			if (It->GetName().StartsWith(TEXT("SKEL_")) || It->GetName().StartsWith(TEXT("REINST_")))
@@ -238,12 +240,12 @@ UClass* FLightWeightInstanceSubsystem::FindBestInstanceManagerClass(const UClass
 	int32 BestDistance = INT_MAX;
 	for (UClass* ManagerClass : ManagerClasses)
 	{
-		if (ManagerClass->GetDefaultObject<ALightWeightInstanceManager>()->DoesAcceptClass(InActorClass))
+		if (ManagerClass->GetDefaultObject<ADEPRECATED_LightWeightInstanceManager>()->DoesAcceptClass(InActorClass))
 		{
-			const UClass* HandledClass = ManagerClass->GetDefaultObject<ALightWeightInstanceManager>()->GetRepresentedClassInternal();
+			const UClass* HandledClass = ManagerClass->GetDefaultObject<ADEPRECATED_LightWeightInstanceManager>()->GetRepresentedClassInternal();
 			if (!HandledClass)
 			{
-				HandledClass = ManagerClass->GetDefaultObject<ALightWeightInstanceManager>()->GetAcceptedClass();
+				HandledClass = ManagerClass->GetDefaultObject<ADEPRECATED_LightWeightInstanceManager>()->GetAcceptedClass();
 			}
 			if (InActorClass == HandledClass)
 			{
@@ -280,7 +282,7 @@ UClass* FLightWeightInstanceSubsystem::GetActorClass(const FActorInstanceHandle&
 		return CachedActor->StaticClass();
 	}
 
-	if (ALightWeightInstanceManager* LWIManager = FindLightWeightInstanceManager(Handle))
+	if (ADEPRECATED_LightWeightInstanceManager* LWIManager = FindLightWeightInstanceManager(Handle))
 	{
 		return LWIManager->GetRepresentedClass(Handle.GetInstanceIndex());
 	}
@@ -297,7 +299,7 @@ FVector FLightWeightInstanceSubsystem::GetLocation(const FActorInstanceHandle& H
 		return CachedActor->GetActorLocation();
 	}
 
-	if (ALightWeightInstanceManager* InstanceManager = FindLightWeightInstanceManager(Handle))
+	if (ADEPRECATED_LightWeightInstanceManager* InstanceManager = FindLightWeightInstanceManager(Handle))
 	{
 		return InstanceManager->GetLocation(Handle);
 	}
@@ -314,7 +316,7 @@ FString FLightWeightInstanceSubsystem::GetName(const FActorInstanceHandle& Handl
 		return CachedActor->GetName();
 	}
 
-	if (ALightWeightInstanceManager* InstanceManager = FindLightWeightInstanceManager(Handle))
+	if (ADEPRECATED_LightWeightInstanceManager* InstanceManager = FindLightWeightInstanceManager(Handle))
 	{
 		return InstanceManager->GetName(Handle);
 	}
@@ -331,7 +333,7 @@ ULevel* FLightWeightInstanceSubsystem::GetLevel(const FActorInstanceHandle& Hand
 		return CachedActor->GetLevel();
 	}
 
-	if (ALightWeightInstanceManager* InstanceManager = FindLightWeightInstanceManager(Handle))
+	if (ADEPRECATED_LightWeightInstanceManager* InstanceManager = FindLightWeightInstanceManager(Handle))
 	{
 		return InstanceManager->GetLevel();
 	}
@@ -348,7 +350,7 @@ bool FLightWeightInstanceSubsystem::IsInLevel(const FActorInstanceHandle& Handle
 		return CachedActor->IsInLevel(InLevel);
 	}
 
-	if (ALightWeightInstanceManager* InstanceManager = FindLightWeightInstanceManager(Handle))
+	if (ADEPRECATED_LightWeightInstanceManager* InstanceManager = FindLightWeightInstanceManager(Handle))
 	{
 		return InstanceManager->GetLevel() == InLevel;
 	}
@@ -361,7 +363,7 @@ FActorInstanceHandle FLightWeightInstanceSubsystem::CreateNewLightWeightInstance
 	// Get or create a light weight instance for this class and data layer
 	if (InitData)
 	{
-		if (ALightWeightInstanceManager* LWIManager = FindOrAddLightWeightInstanceManager(*InActorClass, *World, InitData->Transform.GetLocation(), InLayer))
+		if (ADEPRECATED_LightWeightInstanceManager* LWIManager = FindOrAddLightWeightInstanceManager(*InActorClass, *World, InitData->Transform.GetLocation(), InLayer))
 		{
 			// create an instance with the given data
 			int32 InstanceIdx = LWIManager->AddNewInstance(InitData);
@@ -375,8 +377,10 @@ FActorInstanceHandle FLightWeightInstanceSubsystem::CreateNewLightWeightInstance
 
 void FLightWeightInstanceSubsystem::DeleteInstance(const FActorInstanceHandle& Handle)
 {
-	if (ALightWeightInstanceManager* LWIManager = FindLightWeightInstanceManager(Handle))
+	if (ADEPRECATED_LightWeightInstanceManager* LWIManager = FindLightWeightInstanceManager(Handle))
 	{
 		LWIManager->RemoveInstance(Handle.GetInstanceIndex());
 	}
 }
+
+PRAGMA_ENABLE_DEPRECATION_WARNINGS
