[Backout] - CL49403711
#rnx
Original CL Desc
-----------------------------------------------------------------
Deprecated LigthWeightInstances

diff --git a/Engine/Plugins/Editor/EngineAssetDefinitions/Source/Private/AssetDefinition_LightWeightInstance.h b/Engine/Plugins/Editor/EngineAssetDefinitions/Source/Private/AssetDefinition_LightWeightInstance.h
--- a/Engine/Plugins/Editor/EngineAssetDefinitions/Source/Private/AssetDefinition_LightWeightInstance.h
+++ b/Engine/Plugins/Editor/EngineAssetDefinitions/Source/Private/AssetDefinition_LightWeightInstance.h
@@ -7,8 +7,8 @@
 
 #include "AssetDefinition_LightWeightInstance.generated.h"
 
-UCLASS(Deprecated, meta = (DeprecationMessage ="Deprecated in 5.8. Consider using InstancedActors for similar functionality"))
-class UDEPRECATED_AssetDefinition_LightWeightInstance : public UAssetDefinitionDefault
+UCLASS()
+class UAssetDefinition_LightWeightInstance : public UAssetDefinitionDefault
 {
 	GENERATED_BODY()
 
@@ -16,7 +16,7 @@ public:
 	// UAssetDefinition Begin
 	virtual FText GetAssetDisplayName() const override { return NSLOCTEXT("AssetTypeActions", "AssetTypeActions_LightWeightInstance", "Light Weight Instance"); }
 	virtual FLinearColor GetAssetColor() const override { return FLinearColor(FColor(255, 105, 180)); }
-	virtual TSoftClassPtr<UObject> GetAssetClass() const override { return ADEPRECATED_LightWeightInstanceManager::StaticClass(); }
+	virtual TSoftClassPtr<UObject> GetAssetClass() const override { return ALightWeightInstanceManager::StaticClass(); }
 	virtual TConstArrayView<FAssetCategoryPath> GetAssetCategories() const override
 	{
 		static const auto Categories = { EAssetCategoryPaths::Gameplay };
diff --git a/Engine/Plugins/LightWeightInstancesEditor/LightWeightInstancesEditor.uplugin b/Engine/Plugins/LightWeightInstancesEditor/LightWeightInstancesEditor.uplugin
--- a/Engine/Plugins/LightWeightInstancesEditor/LightWeightInstancesEditor.uplugin
+++ b/Engine/Plugins/LightWeightInstancesEditor/LightWeightInstancesEditor.uplugin
@@ -2,8 +2,8 @@
 	"FileVersion": 3,
 	"Version": 1,
 	"VersionName": "1.0",
-	"FriendlyName": "LightWeightInstancesEditor (DEPRECATED)",
-	"Description": "Light Weight Instances provide the flexibility and interaction of actors while having performance similar to instanced meshes. This plugin is now DEPRECATED. Use InstancedActors for similar functionality.",
+	"FriendlyName": "LightWeightInstancesEditor",
+	"Description": "Light Weight Instances provide the flexibility and interaction of actors while having performance similar to instanced meshes.",
 	"Category": "Other",
 	"CreatedBy": "Epic Games, Inc.",
 	"CreatedByURL": "https://epicgames.com",
@@ -14,7 +14,6 @@
 	"IsBetaVersion": true,
 	"IsExperimentalVersion": false,
 	"Installed": false,
-	"DeprecatedEngineVersion": "5.8",
 	"Modules": [
 		{
 			"Name": "LightWeightInstancesEditor",
diff --git a/Engine/Plugins/LightWeightInstancesEditor/Source/LightWeightInstancesEditor/Private/LightWeightInstancesEditor.cpp b/Engine/Plugins/LightWeightInstancesEditor/Source/LightWeightInstancesEditor/Private/LightWeightInstancesEditor.cpp
--- a/Engine/Plugins/LightWeightInstancesEditor/Source/LightWeightInstancesEditor/Private/LightWeightInstancesEditor.cpp
+++ b/Engine/Plugins/LightWeightInstancesEditor/Source/LightWeightInstancesEditor/Private/LightWeightInstancesEditor.cpp
@@ -14,8 +14,6 @@
 
 DEFINE_LOG_CATEGORY_STATIC(LogLWIEditor, Log, All);
 
-PRAGMA_DISABLE_DEPRECATION_WARNINGS
-
 #if WITH_EDITOR
 typedef FLevelEditorModule::FLevelViewportMenuExtender_SelectedActors DelegateType;
 #endif
@@ -85,11 +83,11 @@ TSharedRef<FExtender> FLightWeightInstancesEditorModule::CreateLevelViewportCont
 		Extender->AddMenuExtension("ActorConvert", EExtensionHook::After, LevelEditorCommandBindings, FMenuExtensionDelegate::CreateLambda(
 			[this, ActorName, InActors](FMenuBuilder& MenuBuilder) {
 				// We can only convert to an LWI if we didn't select an LWI
-				const bool bCanExecute = !InActors[0]->GetClass()->IsChildOf(ADEPRECATED_LightWeightInstanceManager::StaticClass());
+				const bool bCanExecute = !InActors[0]->GetClass()->IsChildOf(ALightWeightInstanceManager::StaticClass());
 
 				MenuBuilder.AddMenuEntry(
-					FText::Format(LOCTEXT("ConvertSelectedActorsToLWIsText", "(DEPRECATED) Convert {0} To Light Weight Instances"), ActorName),
-					LOCTEXT("ConvertSelectedActorsToLWIsTooltip", "(DEPRECATED) Convert the selected actors to light weight instances."),
+					FText::Format(LOCTEXT("ConvertSelectedActorsToLWIsText", "Convert {0} To Light Weight Instances"), ActorName),
+					LOCTEXT("ConvertSelectedActorsToLWIsTooltip", "Convert the selected actors to light weight instances."),
 					FSlateIcon(FAppStyle::GetAppStyleSetName(), "Icons.Convert"),
 					FUIAction(
 						FExecuteAction::CreateRaw(this, &FLightWeightInstancesEditorModule::ConvertActorsToLWIsUIAction, InActors),
@@ -115,7 +113,7 @@ void FLightWeightInstancesEditorModule::ConvertActorsToLWIsUIAction(const TArray
 		TArray<const UDataLayerInstance*> DataLayerInstances = Actor->GetDataLayerInstances();
 		const UDataLayerInstance* DataLayerInstance = DataLayerInstances.Num() > 0 ? DataLayerInstances[0] : nullptr;
 
-		ADEPRECATED_LightWeightInstanceManager* Manager = FLightWeightInstanceSubsystem::Get().FindOrAddLightWeightInstanceManager(*Actor->GetClass(), *Actor->GetWorld(), Actor->GetActorLocation(), DataLayerInstance);
+		ALightWeightInstanceManager* Manager = FLightWeightInstanceSubsystem::Get().FindOrAddLightWeightInstanceManager(*Actor->GetClass(), *Actor->GetWorld(), Actor->GetActorLocation(), DataLayerInstance);
 		check(Manager);
 		
 		if(UDataLayerEditorSubsystem* DataLayerEditor = UDataLayerEditorSubsystem::Get())
@@ -129,9 +127,6 @@ void FLightWeightInstancesEditorModule::ConvertActorsToLWIsUIAction(const TArray
 
 #endif // WITH_EDITOR
 
-PRAGMA_ENABLE_DEPRECATION_WARNINGS
-
 #undef LOCTEXT_NAMESPACE
 	
 IMPLEMENT_MODULE(FLightWeightInstancesEditorModule, LightWeightInstancesEditor)
-
diff --git a/Engine/Source/Editor/UnrealEd/Classes/Factories/LightWeightInstanceFactory.h b/Engine/Source/Editor/UnrealEd/Classes/Factories/LightWeightInstanceFactory.h
--- a/Engine/Source/Editor/UnrealEd/Classes/Factories/LightWeightInstanceFactory.h
+++ b/Engine/Source/Editor/UnrealEd/Classes/Factories/LightWeightInstanceFactory.h
@@ -7,10 +7,10 @@
 #include "Factories/Factory.h"
 #include "LightWeightInstanceFactory.generated.h"
 
-class ADEPRECATED_LightWeightInstanceManager;
+class ALightWeightInstanceManager;
 
-UCLASS(hidecategories=Object, MinimalAPI, Deprecated, meta = (DeprecationMessage ="Deprecated in 5.8. Consider using InstancedActors for similar functionality"))
-class UDEPRECATED_LightWeightInstanceFactory : public UFactory
+UCLASS(hidecategories=Object, MinimalAPI)
+class ULightWeightInstanceFactory : public UFactory
 {
 	GENERATED_UCLASS_BODY()
 
diff --git a/Engine/Source/Editor/UnrealEd/Private/Factories/LightWeightInstanceFactory.cpp b/Engine/Source/Editor/UnrealEd/Private/Factories/LightWeightInstanceFactory.cpp
--- a/Engine/Source/Editor/UnrealEd/Private/Factories/LightWeightInstanceFactory.cpp
+++ b/Engine/Source/Editor/UnrealEd/Private/Factories/LightWeightInstanceFactory.cpp
@@ -22,8 +22,6 @@
 
 #define LOCTEXT_NAMESPACE "LightWeightInstanceFactory"
 
-PRAGMA_DISABLE_DEPRECATION_WARNINGS
-
 class FLWIFilter : public IClassViewerFilter
 {
 public:
@@ -41,15 +39,15 @@ public:
 	}
 };
 
-UDEPRECATED_LightWeightInstanceFactory::UDEPRECATED_LightWeightInstanceFactory(const FObjectInitializer& ObjectInitializer)
+ULightWeightInstanceFactory::ULightWeightInstanceFactory(const FObjectInitializer& ObjectInitializer)
 	: Super(ObjectInitializer)
 {
-	SupportedClass = ADEPRECATED_LightWeightInstanceManager::StaticClass();
+	SupportedClass = ALightWeightInstanceManager::StaticClass();
 	bCreateNew = true;
 	bEditAfterNew = true;
 }
 
-bool UDEPRECATED_LightWeightInstanceFactory::ConfigureProperties()
+bool ULightWeightInstanceFactory::ConfigureProperties()
 {
 	// Null the parent class to ensure one is selected
 	ParentClass = nullptr;
@@ -73,11 +71,11 @@ bool UDEPRECATED_LightWeightInstanceFactory::ConfigureProperties()
 	TSharedPtr<FLWIFilter> Filter = MakeShareable(new FLWIFilter);
 	Options.ClassFilters.Add(Filter.ToSharedRef());
 
-	Filter->AllowedChildrenOfClasses.Add(ADEPRECATED_LightWeightInstanceManager::StaticClass());
+	Filter->AllowedChildrenOfClasses.Add(ALightWeightInstanceManager::StaticClass());
 
 	const FText TitleText = LOCTEXT("CreateBlueprintOptions", "Pick Parent Class");
 	UClass* ChosenClass = nullptr;
-	const bool bPressedOk = SClassPickerDialog::PickClass(TitleText, Options, ChosenClass, ADEPRECATED_LightWeightInstanceManager::StaticClass());
+	const bool bPressedOk = SClassPickerDialog::PickClass(TitleText, Options, ChosenClass, ALightWeightInstanceManager::StaticClass());
 
 	if (bPressedOk)
 	{
@@ -89,7 +87,7 @@ bool UDEPRECATED_LightWeightInstanceFactory::ConfigureProperties()
 	return bPressedOk;
 };
 
-UObject* UDEPRECATED_LightWeightInstanceFactory::FactoryCreateNew(UClass* Class, UObject* InParent, FName Name, EObjectFlags Flags, UObject* Context, FFeedbackContext* Warn)
+UObject* ULightWeightInstanceFactory::FactoryCreateNew(UClass* Class, UObject* InParent, FName Name, EObjectFlags Flags, UObject* Context, FFeedbackContext* Warn)
 {
 	UObject* LWIManager = nullptr;
 	if (ensure(SupportedClass == Class))
@@ -107,6 +105,4 @@ UObject* UDEPRECATED_LightWeightInstanceFactory::FactoryCreateNew(UClass* Class,
 	return LWIManager;
 }
 
-PRAGMA_ENABLE_DEPRECATION_WARNINGS
-
 #undef LOCTEXT_NAMESPACE // "LightWeightInstanceFactory"
diff --git a/Engine/Source/Runtime/Engine/Classes/GameFramework/LightWeightInstanceBlueprintFunctionLibrary.h b/Engine/Source/Runtime/Engine/Classes/GameFramework/LightWeightInstanceBlueprintFunctionLibrary.h
--- a/Engine/Source/Runtime/Engine/Classes/GameFramework/LightWeightInstanceBlueprintFunctionLibrary.h
+++ b/Engine/Source/Runtime/Engine/Classes/GameFramework/LightWeightInstanceBlueprintFunctionLibrary.h
@@ -7,8 +7,8 @@
 
 #include "LightWeightInstanceBlueprintFunctionLibrary.generated.h"
 
-UCLASS(MinimalAPI, Deprecated, meta = (DeprecationMessage ="Deprecated in 5.8. Consider using InstancedActors for similar functionality"))
-class UDEPRECATED_LightWeightInstanceBlueprintFunctionLibrary : public UBlueprintFunctionLibrary
+UCLASS(MinimalAPI)
+class ULightWeightInstanceBlueprintFunctionLibrary : public UBlueprintFunctionLibrary
 {
 	GENERATED_BODY()
 
diff --git a/Engine/Source/Runtime/Engine/Classes/GameFramework/LightWeightInstanceManager.h b/Engine/Source/Runtime/Engine/Classes/GameFramework/LightWeightInstanceManager.h
--- a/Engine/Source/Runtime/Engine/Classes/GameFramework/LightWeightInstanceManager.h
+++ b/Engine/Source/Runtime/Engine/Classes/GameFramework/LightWeightInstanceManager.h
@@ -27,8 +27,8 @@ struct FLWIData
 };
 
 
-UCLASS(BlueprintType, Blueprintable, Experimental, MinimalAPI, Deprecated, meta = (DeprecationMessage ="Deprecated in 5.8. Consider using InstancedActors for similar functionality"))
-class ADEPRECATED_LightWeightInstanceManager : public AActor, public IActorInstanceManagerInterface
+UCLASS(BlueprintType, Blueprintable, Experimental, MinimalAPI)
+class ALightWeightInstanceManager : public AActor, public IActorInstanceManagerInterface
 {
 	GENERATED_UCLASS_BODY()
 
@@ -37,7 +37,7 @@ class ADEPRECATED_LightWeightInstanceManager : public AActor, public IActorInsta
 	friend class ULightWeightInstanceBlueprintFunctionLibrary;
 
 public:
-	ENGINE_API virtual ~ADEPRECATED_LightWeightInstanceManager();
+	ENGINE_API virtual ~ALightWeightInstanceManager();
 
 	ENGINE_API virtual void Tick(float DeltaSeconds);
 
diff --git a/Engine/Source/Runtime/Engine/Classes/GameFramework/LightWeightInstanceStaticMeshManager.h b/Engine/Source/Runtime/Engine/Classes/GameFramework/LightWeightInstanceStaticMeshManager.h
--- a/Engine/Source/Runtime/Engine/Classes/GameFramework/LightWeightInstanceStaticMeshManager.h
+++ b/Engine/Source/Runtime/Engine/Classes/GameFramework/LightWeightInstanceStaticMeshManager.h
@@ -11,8 +11,8 @@
 
 DECLARE_DYNAMIC_DELEGATE_RetVal_OneParam(FActorInstanceHandle, FOnActorReady, FActorInstanceHandle, InHandle);
 
-UCLASS(BlueprintType, Blueprintable, Experimental, MinimalAPI, Deprecated, meta = (DeprecationMessage ="Deprecated in 5.8. Consider using InstancedActors for similar functionality"))
-class ADEPRECATED_LightWeightInstanceStaticMeshManager : public ADEPRECATED_LightWeightInstanceManager, public ISMInstanceManager
+UCLASS(BlueprintType, Blueprintable, Experimental, MinimalAPI)
+class ALightWeightInstanceStaticMeshManager : public ALightWeightInstanceManager, public ISMInstanceManager
 {
 	GENERATED_UCLASS_BODY()
 
diff --git a/Engine/Source/Runtime/Engine/Classes/GameFramework/LightWeightInstanceSubsystem.h b/Engine/Source/Runtime/Engine/Classes/GameFramework/LightWeightInstanceSubsystem.h
--- a/Engine/Source/Runtime/Engine/Classes/GameFramework/LightWeightInstanceSubsystem.h
+++ b/Engine/Source/Runtime/Engine/Classes/GameFramework/LightWeightInstanceSubsystem.h
@@ -27,21 +27,21 @@ struct FLightWeightInstanceSubsystem
 	ENGINE_API ~FLightWeightInstanceSubsystem();
 
 	// Returns the instance manager that handles the given handle
-	ENGINE_API ADEPRECATED_LightWeightInstanceManager* FindLightWeightInstanceManager(const FActorInstanceHandle& Handle) const;
+	ENGINE_API ALightWeightInstanceManager* FindLightWeightInstanceManager(const FActorInstanceHandle& Handle) const;
 
 	// Returns the instance manager that handles actors of type ActorClass in level Level
 	UE_DEPRECATED(5.3, "Use the version that takes in a position.")
-	ENGINE_API ADEPRECATED_LightWeightInstanceManager* FindLightWeightInstanceManager(UClass* ActorClass, const UDataLayerInstance* Layer, UWorld* World) const;
+	ENGINE_API ALightWeightInstanceManager* FindLightWeightInstanceManager(UClass* ActorClass, const UDataLayerInstance* Layer, UWorld* World) const;
 
 	// Returns the instance manager that handles instances of type Class that live in Level
 	UE_DEPRECATED(5.3, "Use the version that takes in a position.")
-	ENGINE_API ADEPRECATED_LightWeightInstanceManager* FindOrAddLightWeightInstanceManager(UClass* ActorClass, const UDataLayerInstance* DataLayer, UWorld* World);
+	ENGINE_API ALightWeightInstanceManager* FindOrAddLightWeightInstanceManager(UClass* ActorClass, const UDataLayerInstance* DataLayer, UWorld* World);
 
 	// Returns the instance manager that handles actors of type ActorClass in level Level
-	ENGINE_API ADEPRECATED_LightWeightInstanceManager* FindLightWeightInstanceManager(UClass& ActorClass, UWorld& World, const FVector& InPos, const UDataLayerInstance* DataLayer = nullptr) const;
+	ENGINE_API ALightWeightInstanceManager* FindLightWeightInstanceManager(UClass& ActorClass, UWorld& World, const FVector& InPos, const UDataLayerInstance* DataLayer = nullptr) const;
 
 	// Returns the instance manager that handles instances of type Class that live in Level
-	ENGINE_API ADEPRECATED_LightWeightInstanceManager* FindOrAddLightWeightInstanceManager(UClass& ActorClass, UWorld& World, const FVector& InPos, const UDataLayerInstance* DataLayer = nullptr);
+	ENGINE_API ALightWeightInstanceManager* FindOrAddLightWeightInstanceManager(UClass& ActorClass, UWorld& World, const FVector& InPos, const UDataLayerInstance* DataLayer = nullptr);
 
 	// Returns the actor specified by Handle if it exists. Returns nullptr if it doesn't
 	ENGINE_API AActor* GetActor_NoCreate(const FActorInstanceHandle& Handle) const;
@@ -69,24 +69,24 @@ struct FLightWeightInstanceSubsystem
 	static ENGINE_API FInt32Vector3 ConvertPositionToCoord(const FVector& InPosition);
 
 	// Add a manager to the subsystem, thread safe.
-	ENGINE_API bool AddManager(ADEPRECATED_LightWeightInstanceManager* Manager);
+	ENGINE_API bool AddManager(ALightWeightInstanceManager* Manager);
 
 	// Remove a manager from the subsystem, thread safe.
-	ENGINE_API bool RemoveManager(ADEPRECATED_LightWeightInstanceManager* Manager);
+	ENGINE_API bool RemoveManager(ALightWeightInstanceManager* Manager);
 
 protected:
 	// Returns the class of the instance manager best suited to support instances of type ActorClass
 	ENGINE_API UClass* FindBestInstanceManagerClass(const UClass* ActorClass);
 
 	// Returns the index associated with Manager
-	ENGINE_API int32 GetManagerIndex(const ADEPRECATED_LightWeightInstanceManager* Manager) const;
+	ENGINE_API int32 GetManagerIndex(const ALightWeightInstanceManager* Manager) const;
 
 	// Returns the light weight instance manager at index Index
-	ENGINE_API const ADEPRECATED_LightWeightInstanceManager* GetManagerAt(int32 Index) const;
+	ENGINE_API const ALightWeightInstanceManager* GetManagerAt(int32 Index) const;
 
 private:
 	// TODO: preallocate the size of this based on a config variable
-	TArray<ADEPRECATED_LightWeightInstanceManager*> LWInstanceManagers;
+	TArray<ALightWeightInstanceManager*> LWInstanceManagers;
 
 	// Mutex to make sure we don't change the LWInstanceManagers array while reading/writing it.
 	mutable FTransactionallySafeRWLock LWIManagersRWLock;
diff --git a/Engine/Source/Runtime/Engine/Private/LightWeightInstanceBlueprintFunctionLibrary.cpp b/Engine/Source/Runtime/Engine/Private/LightWeightInstanceBlueprintFunctionLibrary.cpp
--- a/Engine/Source/Runtime/Engine/Private/LightWeightInstanceBlueprintFunctionLibrary.cpp
+++ b/Engine/Source/Runtime/Engine/Private/LightWeightInstanceBlueprintFunctionLibrary.cpp
@@ -6,9 +6,7 @@
 
 #include UE_INLINE_GENERATED_CPP_BY_NAME(LightWeightInstanceBlueprintFunctionLibrary)
 
-PRAGMA_DISABLE_DEPRECATION_WARNINGS
-
-FActorInstanceHandle UDEPRECATED_LightWeightInstanceBlueprintFunctionLibrary::CreateNewLightWeightInstance(UClass* InActorClass, FTransform InTransform, UDataLayerInstance* InLayer, UWorld* World)
+FActorInstanceHandle ULightWeightInstanceBlueprintFunctionLibrary::CreateNewLightWeightInstance(UClass* InActorClass, FTransform InTransform, UDataLayerInstance* InLayer, UWorld* World)
 {
 	// set up initialization data
 	FLWIData PerInstanceData;
@@ -17,7 +15,7 @@ FActorInstanceHandle UDEPRECATED_LightWeightInstanceBlueprintFunctionLibrary::Cr
 	return FLightWeightInstanceSubsystem::Get().CreateNewLightWeightInstance(InActorClass, &PerInstanceData, InLayer, World);
 }
 
-FActorInstanceHandle UDEPRECATED_LightWeightInstanceBlueprintFunctionLibrary::ConvertActorToLightWeightInstance(AActor* InActor)
+FActorInstanceHandle ULightWeightInstanceBlueprintFunctionLibrary::ConvertActorToLightWeightInstance(AActor* InActor)
 {
 	if (!ensure(InActor))
 	{
@@ -32,7 +30,7 @@ FActorInstanceHandle UDEPRECATED_LightWeightInstanceBlueprintFunctionLibrary::Co
 #else
 	const UDataLayerInstance* DataLayerInstance = nullptr;
 #endif // WITH_EDITOR
-	if (ADEPRECATED_LightWeightInstanceManager* LWIManager = FLightWeightInstanceSubsystem::Get().FindOrAddLightWeightInstanceManager(*InActor->GetClass(), *InActor->GetWorld(), InActor->GetActorLocation(), DataLayerInstance))
+	if (ALightWeightInstanceManager* LWIManager = FLightWeightInstanceSubsystem::Get().FindOrAddLightWeightInstanceManager(*InActor->GetClass(), *InActor->GetWorld(), InActor->GetActorLocation(), DataLayerInstance))
 	{
 		return LWIManager->ConvertActorToLightWeightInstance(InActor);
 	}
@@ -40,4 +38,3 @@ FActorInstanceHandle UDEPRECATED_LightWeightInstanceBlueprintFunctionLibrary::Co
 	return FActorInstanceHandle(InActor);
 }
 
-PRAGMA_ENABLE_DEPRECATION_WARNINGS
diff --git a/Engine/Source/Runtime/Engine/Private/LightWeightInstanceManager.cpp b/Engine/Source/Runtime/Engine/Private/LightWeightInstanceManager.cpp
--- a/Engine/Source/Runtime/Engine/Private/LightWeightInstanceManager.cpp
+++ b/Engine/Source/Runtime/Engine/Private/LightWeightInstanceManager.cpp
@@ -9,8 +9,6 @@
 
 #include UE_INLINE_GENERATED_CPP_BY_NAME(LightWeightInstanceManager)
 
-PRAGMA_DISABLE_DEPRECATION_WARNINGS
-
 // CVar variable that is the size of the grid managers are placed into.
 static int32 LWIGridSize = -1;
 static FAutoConsoleVariableRef CVarLWIGridSize
@@ -22,7 +20,7 @@ static FAutoConsoleVariableRef CVarLWIGridSize
 );
 
 
-ADEPRECATED_LightWeightInstanceManager::ADEPRECATED_LightWeightInstanceManager(const FObjectInitializer& ObjectInitializer)
+ALightWeightInstanceManager::ALightWeightInstanceManager(const FObjectInitializer& ObjectInitializer)
 	:Super(ObjectInitializer)
 {
 	bReplicates = true;
@@ -33,29 +31,29 @@ ADEPRECATED_LightWeightInstanceManager::ADEPRECATED_LightWeightInstanceManager(c
 	}
 }
 
-ADEPRECATED_LightWeightInstanceManager::~ADEPRECATED_LightWeightInstanceManager()
+ALightWeightInstanceManager::~ALightWeightInstanceManager()
 {
 	FLightWeightInstanceSubsystem::Get().RemoveManager(this);
 }
 
-void ADEPRECATED_LightWeightInstanceManager::SetRepresentedClass(UClass* ActorClass)
+void ALightWeightInstanceManager::SetRepresentedClass(UClass* ActorClass)
 {
 	RepresentedClass = ActorClass;
 }
 
-void ADEPRECATED_LightWeightInstanceManager::Tick(float DeltaSeconds)
+void ALightWeightInstanceManager::Tick(float DeltaSeconds)
 {
 	// do nothing
 }
 
-AActor* ADEPRECATED_LightWeightInstanceManager::FindActor(const FActorInstanceHandle& Handle)
+AActor* ALightWeightInstanceManager::FindActor(const FActorInstanceHandle& Handle)
 {
 	TObjectPtr<AActor>* FoundActor = Actors.Find(Handle.GetInstanceIndex());
 
 	return FoundActor ? *FoundActor : nullptr;
 }
 
-AActor* ADEPRECATED_LightWeightInstanceManager::FindOrCreateActor(const FActorInstanceHandle& Handle)
+AActor* ALightWeightInstanceManager::FindOrCreateActor(const FActorInstanceHandle& Handle)
 {
 	// make sure the handle doesn't have an actor already
 	if (ensure(!Handle.GetCachedActor()))
@@ -90,7 +88,7 @@ AActor* ADEPRECATED_LightWeightInstanceManager::FindOrCreateActor(const FActorIn
 	return Handle.GetCachedActor();
 }
 
-AActor* ADEPRECATED_LightWeightInstanceManager::ConvertInstanceToActor(const FActorInstanceHandle& Handle)
+AActor* ALightWeightInstanceManager::ConvertInstanceToActor(const FActorInstanceHandle& Handle)
 {
 #if !UE_BUILD_SHIPPING
 	check(!InstanceToActorConversionsInProgress.Contains(Handle.GetInstanceIndex()));
@@ -117,7 +115,7 @@ AActor* ADEPRECATED_LightWeightInstanceManager::ConvertInstanceToActor(const FAc
 	if (DestroyedActorIndices.Contains(Handle.GetInstanceIndex()))
 	{
 		UE_LOG(LogLightWeightInstance, Verbose,
-			TEXT("ADEPRECATED_LightWeightInstanceManager::ConvertInstanceToActor - manager [ %s ] had already spawned an actor for index [ %d ] that was destroyed"),
+			TEXT("ALightWeightInstanceManager::ConvertInstanceToActor - manager [ %s ] had already spawned an actor for index [ %d ] that was destroyed"),
 			*GetActorNameOrLabel(),
 			Handle.GetInstanceIndex());
 		return nullptr;
@@ -141,7 +139,7 @@ AActor* ADEPRECATED_LightWeightInstanceManager::ConvertInstanceToActor(const FAc
 		check(Handle.GetCachedActor() == NewActor);
 		check(NewActor == Actors.FindRef(Handle.GetInstanceIndex()));
 
-		NewActor->OnDestroyed.AddUniqueDynamic(this, &ADEPRECATED_LightWeightInstanceManager::OnSpawnedActorDestroyed);
+		NewActor->OnDestroyed.AddUniqueDynamic(this, &ALightWeightInstanceManager::OnSpawnedActorDestroyed);
 
 		PostActorSpawn(Handle);
 	}
@@ -149,12 +147,12 @@ AActor* ADEPRECATED_LightWeightInstanceManager::ConvertInstanceToActor(const FAc
 	return NewActor;
 }
 
-void ADEPRECATED_LightWeightInstanceManager::OnSpawnedActorDestroyed(AActor* DestroyedActor)
+void ALightWeightInstanceManager::OnSpawnedActorDestroyed(AActor* DestroyedActor)
 {
 	OnSpawnedActorDestroyed(DestroyedActor, FindIndexForActor(DestroyedActor));
 }
 
-void ADEPRECATED_LightWeightInstanceManager::OnSpawnedActorDestroyed(AActor* DestroyedActor, const int32 DestroyedActorInstanceIndex)
+void ALightWeightInstanceManager::OnSpawnedActorDestroyed(AActor* DestroyedActor, const int32 DestroyedActorInstanceIndex)
 {
 	check(DestroyedActor);
 	if (!ensure(DestroyedActorInstanceIndex != INDEX_NONE))
@@ -175,13 +173,13 @@ void ADEPRECATED_LightWeightInstanceManager::OnSpawnedActorDestroyed(AActor* Des
 	else
 	{
 		UE_LOG(LogLightWeightInstance, Error,
-			TEXT("ADEPRECATED_LightWeightInstanceManager::OnSpawnedActorDestroyed - manager [ %s ] DestroyedActorIndices already has an entry for index [ %d ]"),
+			TEXT("ALightWeightInstanceManager::OnSpawnedActorDestroyed - manager [ %s ] DestroyedActorIndices already has an entry for index [ %d ]"),
 			*GetActorNameOrLabel(),
 			DestroyedActorInstanceIndex);
 	}
 }
 
-int32 ADEPRECATED_LightWeightInstanceManager::FindIndexForActor(const AActor* InActor) const
+int32 ALightWeightInstanceManager::FindIndexForActor(const AActor* InActor) const
 {
 	for (const TPair<int32, TObjectPtr<AActor>>& ActorPair : Actors)
 	{
@@ -194,7 +192,7 @@ int32 ADEPRECATED_LightWeightInstanceManager::FindIndexForActor(const AActor* In
 	return INDEX_NONE;
 }
 
-FActorInstanceHandle ADEPRECATED_LightWeightInstanceManager::ConvertActorToLightWeightInstance(AActor* InActor)
+FActorInstanceHandle ALightWeightInstanceManager::ConvertActorToLightWeightInstance(AActor* InActor)
 {
 	if (!InActor)
 	{
@@ -239,13 +237,13 @@ FActorInstanceHandle ADEPRECATED_LightWeightInstanceManager::ConvertActorToLight
 	}
 }
 
-FLWIData* ADEPRECATED_LightWeightInstanceManager::AllocateInitData() const
+FLWIData* ALightWeightInstanceManager::AllocateInitData() const
 {
 	return new FLWIData;
 }
 
 
-bool ADEPRECATED_LightWeightInstanceManager::SetDataFromActor(FLWIData* InData, AActor* InActor) const
+bool ALightWeightInstanceManager::SetDataFromActor(FLWIData* InData, AActor* InActor) const
 {
 	if (!InData || !InActor)
 	{
@@ -257,17 +255,17 @@ bool ADEPRECATED_LightWeightInstanceManager::SetDataFromActor(FLWIData* InData,
 	return true;
 }
 
-int32 ADEPRECATED_LightWeightInstanceManager::ConvertCollisionIndexToInstanceIndex(int32 InIndex, const UPrimitiveComponent* RelevantComponent) const
+int32 ALightWeightInstanceManager::ConvertCollisionIndexToInstanceIndex(int32 InIndex, const UPrimitiveComponent* RelevantComponent) const
 {
 	return InIndex;
 }
 
-int32 ADEPRECATED_LightWeightInstanceManager::GetGridSize() const
+int32 ALightWeightInstanceManager::GetGridSize() const
 {
 	return LWIGridSize;	
 }
 
-FInt32Vector3 ADEPRECATED_LightWeightInstanceManager::ConvertPositionToCoord(const FVector& InPosition) const
+FInt32Vector3 ALightWeightInstanceManager::ConvertPositionToCoord(const FVector& InPosition) const
 {
 	float GridSize = GetGridSize();
 	if (GridSize > 0)
@@ -282,7 +280,7 @@ FInt32Vector3 ADEPRECATED_LightWeightInstanceManager::ConvertPositionToCoord(con
 	return FInt32Vector3::ZeroValue;
 }
 
-FBox ADEPRECATED_LightWeightInstanceManager::ConvertPositionToGridBounds(const FVector& InPosition) const
+FBox ALightWeightInstanceManager::ConvertPositionToGridBounds(const FVector& InPosition) const
 {
 	float GridSize = GetGridSize();
 	if (GridSize > 0)
@@ -298,30 +296,30 @@ FBox ADEPRECATED_LightWeightInstanceManager::ConvertPositionToGridBounds(const F
 	return FBox(InPosition, InPosition);
 }
 
-void ADEPRECATED_LightWeightInstanceManager::SetSpawnParameters(FActorSpawnParameters& SpawnParams)
+void ALightWeightInstanceManager::SetSpawnParameters(FActorSpawnParameters& SpawnParams)
 {
 	SpawnParams.OverrideLevel = GetLevel();
 	SpawnParams.SpawnCollisionHandlingOverride = ESpawnActorCollisionHandlingMethod::AlwaysSpawn;
 	SpawnParams.ObjectFlags = RF_Transactional;
 }
 
-UClass* ADEPRECATED_LightWeightInstanceManager::GetActorClassToSpawn(const FActorInstanceHandle& Handle) const
+UClass* ALightWeightInstanceManager::GetActorClassToSpawn(const FActorInstanceHandle& Handle) const
 {
 	return RepresentedClass;
 }
 
-void ADEPRECATED_LightWeightInstanceManager::PreSpawnInitalization(const FActorInstanceHandle& Handle, AActor* SpawnedActor)
+void ALightWeightInstanceManager::PreSpawnInitalization(const FActorInstanceHandle& Handle, AActor* SpawnedActor)
 {
 	Handle.SetCachedActor(SpawnedActor);
 	Actors.Add(Handle.GetInstanceIndex(), SpawnedActor);
 }
 
-void ADEPRECATED_LightWeightInstanceManager::PostActorSpawn(const FActorInstanceHandle& Handle)
+void ALightWeightInstanceManager::PostActorSpawn(const FActorInstanceHandle& Handle)
 {
 	// do nothing
 }
 
-bool ADEPRECATED_LightWeightInstanceManager::IsIndexValid(int32 Index) const
+bool ALightWeightInstanceManager::IsIndexValid(int32 Index) const
 {
 	if (ValidIndices.IsValidIndex(Index))
 	{
@@ -331,7 +329,7 @@ bool ADEPRECATED_LightWeightInstanceManager::IsIndexValid(int32 Index) const
 	return false;
 }
 
-bool ADEPRECATED_LightWeightInstanceManager::FindActorForHandle(const FActorInstanceHandle& Handle) const
+bool ALightWeightInstanceManager::FindActorForHandle(const FActorInstanceHandle& Handle) const
 {
 	if (Handle.GetCachedActor())
 	{
@@ -344,17 +342,17 @@ bool ADEPRECATED_LightWeightInstanceManager::FindActorForHandle(const FActorInst
 	return ResolvedActor != nullptr;
 }
 
-FVector ADEPRECATED_LightWeightInstanceManager::GetLocation(const FActorInstanceHandle& Handle) const
+FVector ALightWeightInstanceManager::GetLocation(const FActorInstanceHandle& Handle) const
 {
 	return GetTransform(Handle).GetTranslation();
 }
 
-FRotator ADEPRECATED_LightWeightInstanceManager::GetRotation(const FActorInstanceHandle& Handle) const
+FRotator ALightWeightInstanceManager::GetRotation(const FActorInstanceHandle& Handle) const
 {
 	return GetTransform(Handle).Rotator();
 }
 
-FTransform ADEPRECATED_LightWeightInstanceManager::GetTransform(const FActorInstanceHandle& Handle) const
+FTransform ALightWeightInstanceManager::GetTransform(const FActorInstanceHandle& Handle) const
 {
 	if (FindActorForHandle(Handle))
 	{
@@ -371,7 +369,7 @@ FTransform ADEPRECATED_LightWeightInstanceManager::GetTransform(const FActorInst
 	return FTransform();
 }
 
-FString ADEPRECATED_LightWeightInstanceManager::GetName(const FActorInstanceHandle& Handle) const
+FString ALightWeightInstanceManager::GetName(const FActorInstanceHandle& Handle) const
 {
 	if (FindActorForHandle(Handle))
 	{
@@ -382,27 +380,27 @@ FString ADEPRECATED_LightWeightInstanceManager::GetName(const FActorInstanceHand
 	return FString::Printf(TEXT("%s_%u"), *BaseInstanceName, Handle.GetInstanceIndex());
 }
 
-bool ADEPRECATED_LightWeightInstanceManager::DoesRepresentClass(const UClass* OtherClass) const
+bool ALightWeightInstanceManager::DoesRepresentClass(const UClass* OtherClass) const
 {
 	return OtherClass ? OtherClass->IsChildOf(GetRepresentedClassInternal()) : false;
 }
 
-bool ADEPRECATED_LightWeightInstanceManager::DoesAcceptClass(const UClass* OtherClass) const
+bool ALightWeightInstanceManager::DoesAcceptClass(const UClass* OtherClass) const
 {
 	return OtherClass ? OtherClass->IsChildOf(GetAcceptedClass()) : false;
 }
 
-UClass* ADEPRECATED_LightWeightInstanceManager::GetRepresentedClassInternal() const
+UClass* ALightWeightInstanceManager::GetRepresentedClassInternal() const
 {
 	return IsValid(RepresentedClass) ? RepresentedClass : nullptr;
 }
 
-UClass* ADEPRECATED_LightWeightInstanceManager::GetAcceptedClass() const
+UClass* ALightWeightInstanceManager::GetAcceptedClass() const
 {
 	return IsValid(AcceptedClass) ? AcceptedClass : nullptr;
 }
 
-int32 ADEPRECATED_LightWeightInstanceManager::AddNewInstance(FLWIData* InitData)
+int32 ALightWeightInstanceManager::AddNewInstance(FLWIData* InitData)
 {
 	if (!InitData)
 	{
@@ -418,7 +416,7 @@ int32 ADEPRECATED_LightWeightInstanceManager::AddNewInstance(FLWIData* InitData)
 	return DataIdx;
 }
 
-void ADEPRECATED_LightWeightInstanceManager::AddNewInstanceAt(FLWIData* InitData, int32 Index)
+void ALightWeightInstanceManager::AddNewInstanceAt(FLWIData* InitData, int32 Index)
 {
 	// allocate space on the end of the array if we need to
 	if (Index >= ValidIndices.Num())
@@ -432,7 +430,7 @@ void ADEPRECATED_LightWeightInstanceManager::AddNewInstanceAt(FLWIData* InitData
 	ValidIndices[Index] = true;
 }
 
-void ADEPRECATED_LightWeightInstanceManager::GrowDataArrays()
+void ALightWeightInstanceManager::GrowDataArrays()
 {
 	InstanceTransforms.AddUninitialized();
 	ValidIndices.AddUninitialized();
@@ -440,13 +438,13 @@ void ADEPRECATED_LightWeightInstanceManager::GrowDataArrays()
 	check(InstanceTransforms.Num() == ValidIndices.Num());
 }
 
-void ADEPRECATED_LightWeightInstanceManager::UpdateDataAtIndex(FLWIData* InData, int32 Index)
+void ALightWeightInstanceManager::UpdateDataAtIndex(FLWIData* InData, int32 Index)
 {
 	// Convert to local space.
 	InstanceTransforms[Index] = InData->Transform * GetActorTransform().Inverse();
 }
 
-void ADEPRECATED_LightWeightInstanceManager::RemoveInstance(const int32 Index)
+void ALightWeightInstanceManager::RemoveInstance(const int32 Index)
 {
 	if (ensure(IsIndexValid(Index)))
 	{
@@ -463,7 +461,7 @@ void ADEPRECATED_LightWeightInstanceManager::RemoveInstance(const int32 Index)
 	}
 }
 
-bool ADEPRECATED_LightWeightInstanceManager::HasAnyValidInstancesOrManagedActors() const
+bool ALightWeightInstanceManager::HasAnyValidInstancesOrManagedActors() const
 {
 	bool bHasAnyValidIndex = false;
 	for (const bool bIsIndexValid : ValidIndices)
@@ -477,32 +475,30 @@ bool ADEPRECATED_LightWeightInstanceManager::HasAnyValidInstancesOrManagedActors
 	return bHasAnyValidIndex;
 }
 
-void ADEPRECATED_LightWeightInstanceManager::GetLifetimeReplicatedProps(TArray<FLifetimeProperty>& OutLifetimeProps) const
+void ALightWeightInstanceManager::GetLifetimeReplicatedProps(TArray<FLifetimeProperty>& OutLifetimeProps) const
 {
 	Super::GetLifetimeReplicatedProps(OutLifetimeProps);
 
-	DOREPLIFETIME(ADEPRECATED_LightWeightInstanceManager, RepresentedClass);
-	DOREPLIFETIME(ADEPRECATED_LightWeightInstanceManager, InstanceTransforms);
-	DOREPLIFETIME(ADEPRECATED_LightWeightInstanceManager, FreeIndices);
-	DOREPLIFETIME(ADEPRECATED_LightWeightInstanceManager, ValidIndices);
+	DOREPLIFETIME(ALightWeightInstanceManager, RepresentedClass);
+	DOREPLIFETIME(ALightWeightInstanceManager, InstanceTransforms);
+	DOREPLIFETIME(ALightWeightInstanceManager, FreeIndices);
+	DOREPLIFETIME(ALightWeightInstanceManager, ValidIndices);
 }
 
-void ADEPRECATED_LightWeightInstanceManager::OnRep_Transforms()
+void ALightWeightInstanceManager::OnRep_Transforms()
 {
 	// do nothing
 }
 #if WITH_EDITOR
-void ADEPRECATED_LightWeightInstanceManager::PostEditChangeProperty(struct FPropertyChangedEvent& PropertyChangedEvent)
+void ALightWeightInstanceManager::PostEditChangeProperty(struct FPropertyChangedEvent& PropertyChangedEvent)
 {
 	Super::PostEditChangeProperty(PropertyChangedEvent);
 
 	const FName Name = (PropertyChangedEvent.Property != nullptr) ? PropertyChangedEvent.Property->GetFName() : NAME_None;
 
-	if (Name == GET_MEMBER_NAME_CHECKED(ADEPRECATED_LightWeightInstanceManager, RepresentedClass))
+	if (Name == GET_MEMBER_NAME_CHECKED(ALightWeightInstanceManager, RepresentedClass))
 	{
 		SetRepresentedClass(RepresentedClass);
 	}
 }
 #endif
-
-PRAGMA_ENABLE_DEPRECATION_WARNINGS
diff --git a/Engine/Source/Runtime/Engine/Private/LightWeightInstanceStaticMeshManager.cpp b/Engine/Source/Runtime/Engine/Private/LightWeightInstanceStaticMeshManager.cpp
--- a/Engine/Source/Runtime/Engine/Private/LightWeightInstanceStaticMeshManager.cpp
+++ b/Engine/Source/Runtime/Engine/Private/LightWeightInstanceStaticMeshManager.cpp
@@ -18,14 +18,14 @@
 #include "TimerManager.h"
 #endif // WITH_EDITOR
 
-PRAGMA_DISABLE_DEPRECATION_WARNINGS
-
-ADEPRECATED_LightWeightInstanceStaticMeshManager::ADEPRECATED_LightWeightInstanceStaticMeshManager(const FObjectInitializer& ObjectInitializer)
+ALightWeightInstanceStaticMeshManager::ALightWeightInstanceStaticMeshManager(const FObjectInitializer& ObjectInitializer)
 	:Super(ObjectInitializer)
 {
 #if WITH_EDITORONLY_DATA
+PRAGMA_DISABLE_DEPRECATION_WARNINGS
 	// Still create the deprecated InstancedStaticMeshComponent so the old data can be serialized during loading
 	InstancedStaticMeshComponent_DEPRECATED = CreateDefaultSubobject<UHierarchicalInstancedStaticMeshComponent>(TEXT("InstancedStaticMeshComponent0"));
+PRAGMA_ENABLE_DEPRECATION_WARNINGS
 #endif
 
 	ISMComponent = CreateDefaultSubobject<UInstancedStaticMeshComponent>(TEXT("ISMComponent0"));
@@ -42,10 +42,11 @@ ADEPRECATED_LightWeightInstanceStaticMeshManager::ADEPRECATED_LightWeightInstanc
 }
 
 #if WITH_EDITORONLY_DATA
-void ADEPRECATED_LightWeightInstanceStaticMeshManager::PostLoad()
+void ALightWeightInstanceStaticMeshManager::PostLoad()
 {
 	Super::PostLoad();
 
+PRAGMA_DISABLE_DEPRECATION_WARNINGS
 	if (InstancedStaticMeshComponent_DEPRECATED != nullptr && InstancedStaticMeshComponent_DEPRECATED->PerInstanceSMData.Num() > 0)
 	{
 		// assume new ISM component is still empty
@@ -87,10 +88,11 @@ void ADEPRECATED_LightWeightInstanceStaticMeshManager::PostLoad()
 			RenderingIndicesToDataIndices[DataIndex] = DataIndex;
 		}
 	}
+PRAGMA_ENABLE_DEPRECATION_WARNINGS
 }
 #endif
 
-void ADEPRECATED_LightWeightInstanceStaticMeshManager::SetRepresentedClass(UClass* ActorClass)
+void ALightWeightInstanceStaticMeshManager::SetRepresentedClass(UClass* ActorClass)
 {
 	Super::SetRepresentedClass(ActorClass);
 
@@ -108,11 +110,13 @@ void ADEPRECATED_LightWeightInstanceStaticMeshManager::SetRepresentedClass(UClas
 
 	if (ISMComponent)
 	{
+PRAGMA_DISABLE_DEPRECATION_WARNINGS
 		ISMComponent->OnPostPopulatePerInstanceData();
+PRAGMA_ENABLE_DEPRECATION_WARNINGS
 	}
 }
 
-int32 ADEPRECATED_LightWeightInstanceStaticMeshManager::ConvertCollisionIndexToInstanceIndex(int32 InIndex, const UPrimitiveComponent* RelevantComponent) const
+int32 ALightWeightInstanceStaticMeshManager::ConvertCollisionIndexToInstanceIndex(int32 InIndex, const UPrimitiveComponent* RelevantComponent) const
 {
 	if (ensureMsgf(RenderingIndicesToDataIndices.IsValidIndex(InIndex), TEXT("Invalid index [ %d ]"), InIndex))
 	{
@@ -122,13 +126,13 @@ int32 ADEPRECATED_LightWeightInstanceStaticMeshManager::ConvertCollisionIndexToI
 	return InIndex;
 }
 
-void ADEPRECATED_LightWeightInstanceStaticMeshManager::AddNewInstanceAt(FLWIData* InitData, int32 Index)
+void ALightWeightInstanceStaticMeshManager::AddNewInstanceAt(FLWIData* InitData, int32 Index)
 {
 	Super::AddNewInstanceAt(InitData, Index);
 	AddInstanceToRendering(Index);
 }
 
-void ADEPRECATED_LightWeightInstanceStaticMeshManager::RemoveInstance(const int32 Index)
+void ALightWeightInstanceStaticMeshManager::RemoveInstance(const int32 Index)
 {
 #if WITH_EDITOR
 	Modify();
@@ -138,7 +142,7 @@ void ADEPRECATED_LightWeightInstanceStaticMeshManager::RemoveInstance(const int3
 	Super::RemoveInstance(Index);
 }
 
-void ADEPRECATED_LightWeightInstanceStaticMeshManager::AddInstanceToRendering(int32 DataIndex)
+void ALightWeightInstanceStaticMeshManager::AddInstanceToRendering(int32 DataIndex)
 {
 	if (!ensureMsgf(RenderingIndicesToDataIndices.Contains(DataIndex) == false, TEXT("LWI rendering instance added more than once. Index: %d"), DataIndex))
 	{
@@ -152,7 +156,7 @@ void ADEPRECATED_LightWeightInstanceStaticMeshManager::AddInstanceToRendering(in
 	const int32 RenderingIdx = RenderingIndicesToDataIndices.Add(DataIndex);
 
 	UE_LOG(LogLightWeightInstance, Verbose,
-		TEXT("ADEPRECATED_LightWeightInstanceStaticMeshManager::AddInstanceToRendering - manager [ %s ] adding instance at data index [ %d ] rendering index [ %d ], did cancel pending delete? [ %s ]"),
+		TEXT("ALightWeightInstanceStaticMeshManager::AddInstanceToRendering - manager [ %s ] adding instance at data index [ %d ] rendering index [ %d ], did cancel pending delete? [ %s ]"),
 		*GetActorNameOrLabel(),
 		DataIndex,
 		RenderingIdx,
@@ -175,7 +179,7 @@ void ADEPRECATED_LightWeightInstanceStaticMeshManager::AddInstanceToRendering(in
 	}
 }
 
-void ADEPRECATED_LightWeightInstanceStaticMeshManager::RemoveInstanceFromRendering(int32 DataIndex)
+void ALightWeightInstanceStaticMeshManager::RemoveInstanceFromRendering(int32 DataIndex)
 {
 	if (IsIndexValid(DataIndex))
 	{
@@ -185,25 +189,25 @@ void ADEPRECATED_LightWeightInstanceStaticMeshManager::RemoveInstanceFromRenderi
 #if WITH_EDITOR
 			if (GEditor)
 			{
-				GEditor->GetTimerManager()->SetTimerForNextTick(this, &ADEPRECATED_LightWeightInstanceStaticMeshManager::PostRemoveInstanceFromRendering);
+				GEditor->GetTimerManager()->SetTimerForNextTick(this, &ALightWeightInstanceStaticMeshManager::PostRemoveInstanceFromRendering);
 			}
 			else
 #endif // WITH_EDITOR
 			{
-				GetWorld()->GetTimerManager().SetTimerForNextTick(this, &ADEPRECATED_LightWeightInstanceStaticMeshManager::PostRemoveInstanceFromRendering);
+				GetWorld()->GetTimerManager().SetTimerForNextTick(this, &ALightWeightInstanceStaticMeshManager::PostRemoveInstanceFromRendering);
 			}
 		}
 
 		DataIndicesToBeDeleted.AddUnique(DataIndex);
 
 		UE_LOG(LogLightWeightInstance, Verbose,
-			TEXT("ADEPRECATED_LightWeightInstanceStaticMeshManager::RemoveInstanceFromRendering - manager [ %s ] removing instance at data index [ %d ]"),
+			TEXT("ALightWeightInstanceStaticMeshManager::RemoveInstanceFromRendering - manager [ %s ] removing instance at data index [ %d ]"),
 			*GetActorNameOrLabel(),
 			DataIndex);
 	}
 }
 
-void ADEPRECATED_LightWeightInstanceStaticMeshManager::PostRemoveInstanceFromRendering()
+void ALightWeightInstanceStaticMeshManager::PostRemoveInstanceFromRendering()
 {
 	TArray<int32> RenderingIndicesToBeDeleted;
 	for (int32 DataIndex : DataIndicesToBeDeleted)
@@ -211,7 +215,7 @@ void ADEPRECATED_LightWeightInstanceStaticMeshManager::PostRemoveInstanceFromRen
 		RenderingIndicesToBeDeleted.Add(DataIndicesToRenderingIndices[DataIndex]);
 
 		UE_LOG(LogLightWeightInstance, Verbose,
-			TEXT("ADEPRECATED_LightWeightInstanceStaticMeshManager::PostRemoveInstanceFromRendering - manager [ %s ] removing instance at data index [ %d ] rendering index [ %d ]"),
+			TEXT("ALightWeightInstanceStaticMeshManager::PostRemoveInstanceFromRendering - manager [ %s ] removing instance at data index [ %d ] rendering index [ %d ]"),
 			*GetActorNameOrLabel(),
 			DataIndex,
 			DataIndicesToRenderingIndices[DataIndex]);
@@ -245,21 +249,21 @@ void ADEPRECATED_LightWeightInstanceStaticMeshManager::PostRemoveInstanceFromRen
 	DataIndicesToBeDeleted.Reset();
 }
 
-void ADEPRECATED_LightWeightInstanceStaticMeshManager::GetLifetimeReplicatedProps(TArray<FLifetimeProperty>& OutLifetimeProps) const
+void ALightWeightInstanceStaticMeshManager::GetLifetimeReplicatedProps(TArray<FLifetimeProperty>& OutLifetimeProps) const
 {
 	Super::GetLifetimeReplicatedProps(OutLifetimeProps);
 
-	DOREPLIFETIME(ADEPRECATED_LightWeightInstanceStaticMeshManager, StaticMesh);
-	DOREPLIFETIME(ADEPRECATED_LightWeightInstanceStaticMeshManager, RenderingIndicesToDataIndices);
-	DOREPLIFETIME(ADEPRECATED_LightWeightInstanceStaticMeshManager, DataIndicesToRenderingIndices);
+	DOREPLIFETIME(ALightWeightInstanceStaticMeshManager, StaticMesh);
+	DOREPLIFETIME(ALightWeightInstanceStaticMeshManager, RenderingIndicesToDataIndices);
+	DOREPLIFETIME(ALightWeightInstanceStaticMeshManager, DataIndicesToRenderingIndices);
 }
 
-void ADEPRECATED_LightWeightInstanceStaticMeshManager::OnRep_StaticMesh()
+void ALightWeightInstanceStaticMeshManager::OnRep_StaticMesh()
 {
 	OnStaticMeshSet();
 }
 
-void ADEPRECATED_LightWeightInstanceStaticMeshManager::OnRep_Transforms()
+void ALightWeightInstanceStaticMeshManager::OnRep_Transforms()
 {
 	Super::OnRep_Transforms();
 
@@ -269,7 +273,7 @@ void ADEPRECATED_LightWeightInstanceStaticMeshManager::OnRep_Transforms()
 	}
 }
 
-void ADEPRECATED_LightWeightInstanceStaticMeshManager::PostActorSpawn(const FActorInstanceHandle& Handle)
+void ALightWeightInstanceStaticMeshManager::PostActorSpawn(const FActorInstanceHandle& Handle)
 {
 	Super::PostActorSpawn(Handle);
 
@@ -277,7 +281,7 @@ void ADEPRECATED_LightWeightInstanceStaticMeshManager::PostActorSpawn(const FAct
 	RemoveInstanceFromRendering(Handle.GetInstanceIndex());
 }
 
-void ADEPRECATED_LightWeightInstanceStaticMeshManager::SetInstancedStaticMeshParams()
+void ALightWeightInstanceStaticMeshManager::SetInstancedStaticMeshParams()
 {
 	FName CollisionProfileName("LightWeightInstancedStaticMeshPhysics");
 	ISMComponent->SetCollisionProfileName(CollisionProfileName);
@@ -301,13 +305,13 @@ void ADEPRECATED_LightWeightInstanceStaticMeshManager::SetInstancedStaticMeshPar
 	}
 } 
 
-void ADEPRECATED_LightWeightInstanceStaticMeshManager::SetStaticMeshFromActor(AActor* InActor)
+void ALightWeightInstanceStaticMeshManager::SetStaticMeshFromActor(AActor* InActor)
 {
 	ensureMsgf(false, TEXT("ALightWeightInstanceManager::SetStaticMeshFromActor was called on %s. Projects should override this function in subclasses."), *GetNameSafe(InActor));
 	ClearStaticMesh();
 }
 
-void ADEPRECATED_LightWeightInstanceStaticMeshManager::OnStaticMeshSet()
+void ALightWeightInstanceStaticMeshManager::OnStaticMeshSet()
 {
 #if WITH_EDITOR
 	Modify();
@@ -337,7 +341,7 @@ void ADEPRECATED_LightWeightInstanceStaticMeshManager::OnStaticMeshSet()
 	}
 }
 
-void ADEPRECATED_LightWeightInstanceStaticMeshManager::ClearStaticMesh()
+void ALightWeightInstanceStaticMeshManager::ClearStaticMesh()
 {
 #if WITH_EDITOR
 	Modify();
@@ -347,7 +351,7 @@ void ADEPRECATED_LightWeightInstanceStaticMeshManager::ClearStaticMesh()
 	OnStaticMeshSet();
 }
 
-FText ADEPRECATED_LightWeightInstanceStaticMeshManager::GetSMInstanceDisplayName(const FSMInstanceId& InstanceId) const
+FText ALightWeightInstanceStaticMeshManager::GetSMInstanceDisplayName(const FSMInstanceId& InstanceId) const
 {
 	check(InstanceId.ISMComponent == ISMComponent);
 
@@ -355,7 +359,7 @@ FText ADEPRECATED_LightWeightInstanceStaticMeshManager::GetSMInstanceDisplayName
 	return FText::Format(NSLOCTEXT("LightWeightInstanceStaticMeshManager", "DisplayNameFmt", "{0} - Instance {1}"), OwnerDisplayName, InstanceId.InstanceIndex);
 }
 
-FText ADEPRECATED_LightWeightInstanceStaticMeshManager::GetSMInstanceTooltip(const FSMInstanceId& InstanceId) const
+FText ALightWeightInstanceStaticMeshManager::GetSMInstanceTooltip(const FSMInstanceId& InstanceId) const
 {
 	check(InstanceId.ISMComponent == ISMComponent);
 	
@@ -363,7 +367,7 @@ FText ADEPRECATED_LightWeightInstanceStaticMeshManager::GetSMInstanceTooltip(con
 	return FText::Format(NSLOCTEXT("LightWeightInstanceStaticMeshManager", "TooltipFmt", "Instance {0} on {1}"), InstanceId.InstanceIndex, OwnerDisplayPath);
 }
 
-bool ADEPRECATED_LightWeightInstanceStaticMeshManager::CanEditSMInstance(const FSMInstanceId& InstanceId) const
+bool ALightWeightInstanceStaticMeshManager::CanEditSMInstance(const FSMInstanceId& InstanceId) const
 {
 	check(InstanceId.ISMComponent == ISMComponent);
 
@@ -371,7 +375,7 @@ bool ADEPRECATED_LightWeightInstanceStaticMeshManager::CanEditSMInstance(const F
 	return InstanceManager->CanEditSMInstance(InstanceId);
 }
 
-bool ADEPRECATED_LightWeightInstanceStaticMeshManager::CanMoveSMInstance(const FSMInstanceId& InstanceId, const ETypedElementWorldType InWorldType) const
+bool ALightWeightInstanceStaticMeshManager::CanMoveSMInstance(const FSMInstanceId& InstanceId, const ETypedElementWorldType InWorldType) const
 {
 	check(InstanceId.ISMComponent == ISMComponent);
 
@@ -379,7 +383,7 @@ bool ADEPRECATED_LightWeightInstanceStaticMeshManager::CanMoveSMInstance(const F
 	return InstanceManager->CanMoveSMInstance(InstanceId, InWorldType);
 }
 
-bool ADEPRECATED_LightWeightInstanceStaticMeshManager::GetSMInstanceTransform(const FSMInstanceId& InstanceId, FTransform& OutInstanceTransform, bool bWorldSpace) const
+bool ALightWeightInstanceStaticMeshManager::GetSMInstanceTransform(const FSMInstanceId& InstanceId, FTransform& OutInstanceTransform, bool bWorldSpace) const
 {
 	check(InstanceId.ISMComponent == ISMComponent);
 
@@ -393,7 +397,7 @@ bool ADEPRECATED_LightWeightInstanceStaticMeshManager::GetSMInstanceTransform(co
 	return false;
 }
 
-bool ADEPRECATED_LightWeightInstanceStaticMeshManager::SetSMInstanceTransform(const FSMInstanceId& InstanceId, const FTransform& InstanceTransform, bool bWorldSpace, bool bMarkRenderStateDirty, bool bTeleport)
+bool ALightWeightInstanceStaticMeshManager::SetSMInstanceTransform(const FSMInstanceId& InstanceId, const FTransform& InstanceTransform, bool bWorldSpace, bool bMarkRenderStateDirty, bool bTeleport)
 {
 	check(InstanceId.ISMComponent == ISMComponent);
 
@@ -407,7 +411,7 @@ bool ADEPRECATED_LightWeightInstanceStaticMeshManager::SetSMInstanceTransform(co
 	return false;
 }
 
-void ADEPRECATED_LightWeightInstanceStaticMeshManager::NotifySMInstanceMovementStarted(const FSMInstanceId& InstanceId)
+void ALightWeightInstanceStaticMeshManager::NotifySMInstanceMovementStarted(const FSMInstanceId& InstanceId)
 {
 	check(InstanceId.ISMComponent == ISMComponent);
 
@@ -415,7 +419,7 @@ void ADEPRECATED_LightWeightInstanceStaticMeshManager::NotifySMInstanceMovementS
 	InstanceManager->NotifySMInstanceMovementStarted(InstanceId);
 }
 
-void ADEPRECATED_LightWeightInstanceStaticMeshManager::NotifySMInstanceMovementOngoing(const FSMInstanceId& InstanceId)
+void ALightWeightInstanceStaticMeshManager::NotifySMInstanceMovementOngoing(const FSMInstanceId& InstanceId)
 {
 	check(InstanceId.ISMComponent == ISMComponent);
 
@@ -423,7 +427,7 @@ void ADEPRECATED_LightWeightInstanceStaticMeshManager::NotifySMInstanceMovementO
 	InstanceManager->NotifySMInstanceMovementOngoing(InstanceId);
 }
 
-void ADEPRECATED_LightWeightInstanceStaticMeshManager::NotifySMInstanceMovementEnded(const FSMInstanceId& InstanceId)
+void ALightWeightInstanceStaticMeshManager::NotifySMInstanceMovementEnded(const FSMInstanceId& InstanceId)
 {
 	check(InstanceId.ISMComponent == ISMComponent);
 
@@ -431,7 +435,7 @@ void ADEPRECATED_LightWeightInstanceStaticMeshManager::NotifySMInstanceMovementE
 	InstanceManager->NotifySMInstanceMovementEnded(InstanceId);
 }
 
-void ADEPRECATED_LightWeightInstanceStaticMeshManager::NotifySMInstanceSelectionChanged(const FSMInstanceId& InstanceId, const bool bIsSelected)
+void ALightWeightInstanceStaticMeshManager::NotifySMInstanceSelectionChanged(const FSMInstanceId& InstanceId, const bool bIsSelected)
 {
 	check(InstanceId.ISMComponent == ISMComponent);
 
@@ -439,7 +443,7 @@ void ADEPRECATED_LightWeightInstanceStaticMeshManager::NotifySMInstanceSelection
 	InstanceManager->NotifySMInstanceSelectionChanged(InstanceId, bIsSelected);
 }
 
-bool ADEPRECATED_LightWeightInstanceStaticMeshManager::DeleteSMInstances(TArrayView<const FSMInstanceId> InstanceIds)
+bool ALightWeightInstanceStaticMeshManager::DeleteSMInstances(TArrayView<const FSMInstanceId> InstanceIds)
 {
 	TArray<int32> DataIndices;
 	GetLWIDataIndices(InstanceIds, DataIndices);
@@ -458,7 +462,7 @@ bool ADEPRECATED_LightWeightInstanceStaticMeshManager::DeleteSMInstances(TArrayV
 	return true;
 }
 
-bool ADEPRECATED_LightWeightInstanceStaticMeshManager::DuplicateSMInstances(TArrayView<const FSMInstanceId> InstanceIds, TArray<FSMInstanceId>& OutNewInstanceIds)
+bool ALightWeightInstanceStaticMeshManager::DuplicateSMInstances(TArrayView<const FSMInstanceId> InstanceIds, TArray<FSMInstanceId>& OutNewInstanceIds)
 {
 	TArray<int32> DataIndices;
 	GetLWIDataIndices(InstanceIds, DataIndices);
@@ -478,7 +482,7 @@ bool ADEPRECATED_LightWeightInstanceStaticMeshManager::DuplicateSMInstances(TArr
 	return true;
 }
 
-void ADEPRECATED_LightWeightInstanceStaticMeshManager::DuplicateLWIInstances(TArrayView<const int32> DataIndices, TArray<int32>& OutNewDataIndices)
+void ALightWeightInstanceStaticMeshManager::DuplicateLWIInstances(TArrayView<const int32> DataIndices, TArray<int32>& OutNewDataIndices)
 {
 	OutNewDataIndices.Reset(DataIndices.Num());
 	for (int32 DataIndex : DataIndices)
@@ -494,7 +498,7 @@ void ADEPRECATED_LightWeightInstanceStaticMeshManager::DuplicateLWIInstances(TAr
 	}
 }
 
-void ADEPRECATED_LightWeightInstanceStaticMeshManager::GetLWIDataIndices(TArrayView<const FSMInstanceId> InstanceIds, TArray<int32>& OutDataIndices) const
+void ALightWeightInstanceStaticMeshManager::GetLWIDataIndices(TArrayView<const FSMInstanceId> InstanceIds, TArray<int32>& OutDataIndices) const
 {
 	OutDataIndices.Reset(InstanceIds.Num());
 	for (const FSMInstanceId& InstanceId : InstanceIds)
@@ -507,14 +511,13 @@ void ADEPRECATED_LightWeightInstanceStaticMeshManager::GetLWIDataIndices(TArrayV
 	}
 }
 
-int32 ADEPRECATED_LightWeightInstanceStaticMeshManager::ConvertInternalIndexToHandleIndex(int32 InInternalIndex) const
+int32 ALightWeightInstanceStaticMeshManager::ConvertInternalIndexToHandleIndex(int32 InInternalIndex) const
 {
 	return DataIndicesToRenderingIndices[InInternalIndex];
 }
 
-int32 ADEPRECATED_LightWeightInstanceStaticMeshManager::ConvertHandleIndexToInternalIndex(int32 InHandleIndex) const
+int32 ALightWeightInstanceStaticMeshManager::ConvertHandleIndexToInternalIndex(int32 InHandleIndex) const
 {
 	return RenderingIndicesToDataIndices[InHandleIndex];
 }
 
-PRAGMA_ENABLE_DEPRECATION_WARNINGS
diff --git a/Engine/Source/Runtime/Engine/Private/LightWeightInstanceSubsystem.cpp b/Engine/Source/Runtime/Engine/Private/LightWeightInstanceSubsystem.cpp
--- a/Engine/Source/Runtime/Engine/Private/LightWeightInstanceSubsystem.cpp
+++ b/Engine/Source/Runtime/Engine/Private/LightWeightInstanceSubsystem.cpp
@@ -8,8 +8,6 @@
 
 DEFINE_LOG_CATEGORY(LogLightWeightInstance);
 
-PRAGMA_DISABLE_DEPRECATION_WARNINGS
-
 FLightWeightInstanceSubsystem::FLightWeightInstanceSubsystem()
 {
 #if WITH_EDITOR
@@ -17,14 +15,14 @@ FLightWeightInstanceSubsystem::FLightWeightInstanceSubsystem()
 	{
 		OnLevelActorAddedHandle = GEngine->OnLevelActorAdded().AddLambda([this](AActor* Actor)
 			{
-				if (ADEPRECATED_LightWeightInstanceManager* LWIManager = Cast<ADEPRECATED_LightWeightInstanceManager>(Actor))
+				if (ALightWeightInstanceManager* LWIManager = Cast<ALightWeightInstanceManager>(Actor))
 				{
 					Get().AddManager(LWIManager);
 				}
 			});
 		OnLevelActorDeletedHandle = GEngine->OnLevelActorDeleted().AddLambda([this](AActor* Actor)
 			{
-				if (ADEPRECATED_LightWeightInstanceManager* LWIManager = Cast<ADEPRECATED_LightWeightInstanceManager>(Actor))
+				if (ALightWeightInstanceManager* LWIManager = Cast<ALightWeightInstanceManager>(Actor))
 				{
 					Get().RemoveManager(LWIManager);
 				}
@@ -44,7 +42,7 @@ FLightWeightInstanceSubsystem::~FLightWeightInstanceSubsystem()
 #endif // WITH_EDITOR
 }
 
-int32 FLightWeightInstanceSubsystem::GetManagerIndex(const ADEPRECATED_LightWeightInstanceManager* Manager) const
+int32 FLightWeightInstanceSubsystem::GetManagerIndex(const ALightWeightInstanceManager* Manager) const
 {
 	UE::TReadScopeLock Lock(LWIManagersRWLock);
 	for (int32 Idx = 0; Idx < LWInstanceManagers.Num(); ++Idx)
@@ -58,29 +56,29 @@ int32 FLightWeightInstanceSubsystem::GetManagerIndex(const ADEPRECATED_LightWeig
 	return INDEX_NONE;
 }
 
-const ADEPRECATED_LightWeightInstanceManager* FLightWeightInstanceSubsystem::GetManagerAt(int32 Index) const
+const ALightWeightInstanceManager* FLightWeightInstanceSubsystem::GetManagerAt(int32 Index) const
 {
 	UE::TReadScopeLock Lock(LWIManagersRWLock);
 	return LWInstanceManagers.IsValidIndex(Index) ? LWInstanceManagers[Index] : nullptr;
 }
 
-bool FLightWeightInstanceSubsystem::AddManager(ADEPRECATED_LightWeightInstanceManager* Manager)
+bool FLightWeightInstanceSubsystem::AddManager(ALightWeightInstanceManager* Manager)
 {
 	UE::TWriteScopeLock Lock(LWIManagersRWLock);
 	return Get().LWInstanceManagers.Add(Manager) ? true : false;
 }
 
-bool FLightWeightInstanceSubsystem::RemoveManager(ADEPRECATED_LightWeightInstanceManager* Manager)
+bool FLightWeightInstanceSubsystem::RemoveManager(ALightWeightInstanceManager* Manager)
 {
 	UE::TWriteScopeLock Lock(LWIManagersRWLock);
 	return Get().LWInstanceManagers.Remove(Manager) ? true : false;
 }
 
-ADEPRECATED_LightWeightInstanceManager* FLightWeightInstanceSubsystem::FindLightWeightInstanceManager(const FActorInstanceHandle& Handle) const
+ALightWeightInstanceManager* FLightWeightInstanceSubsystem::FindLightWeightInstanceManager(const FActorInstanceHandle& Handle) const
 {
 	if (Handle.GetManagerInterface().IsValid())
 	{
-		return Handle.GetManager<ADEPRECATED_LightWeightInstanceManager>();
+		return Handle.GetManager<ALightWeightInstanceManager>();
 	}
 
 	if (AActor* CachedActor = Handle.GetCachedActor())
@@ -90,7 +88,7 @@ ADEPRECATED_LightWeightInstanceManager* FLightWeightInstanceSubsystem::FindLight
 
 		UE::TReadScopeLock Lock(LWIManagersRWLock);
 		// see if we already have a match
-		for (ADEPRECATED_LightWeightInstanceManager* LWInstance : LWInstanceManagers)
+		for (ALightWeightInstanceManager* LWInstance : LWInstanceManagers)
 		{
 			if (CachedActorClass == LWInstance->GetRepresentedClass(Handle.GetInstanceIndex()))
 			{
@@ -119,25 +117,25 @@ ADEPRECATED_LightWeightInstanceManager* FLightWeightInstanceSubsystem::FindLight
 	return nullptr;
 }
 
-ADEPRECATED_LightWeightInstanceManager* FLightWeightInstanceSubsystem::FindLightWeightInstanceManager(UClass* ActorClass, const UDataLayerInstance* DataLayer, UWorld* World) const
+ALightWeightInstanceManager* FLightWeightInstanceSubsystem::FindLightWeightInstanceManager(UClass* ActorClass, const UDataLayerInstance* DataLayer, UWorld* World) const
 {
 	check(World);
 	return FindLightWeightInstanceManager(*ActorClass, *World, FVector::ZeroVector, DataLayer);
 }
 
 
-ADEPRECATED_LightWeightInstanceManager* FLightWeightInstanceSubsystem::FindOrAddLightWeightInstanceManager(UClass* ActorClass, const UDataLayerInstance* DataLayer, UWorld* World)
+ALightWeightInstanceManager* FLightWeightInstanceSubsystem::FindOrAddLightWeightInstanceManager(UClass* ActorClass, const UDataLayerInstance* DataLayer, UWorld* World)
 {
 	check(World);
 	return FindOrAddLightWeightInstanceManager(*ActorClass, *World, FVector::ZeroVector, DataLayer);
 }
 
 
-ADEPRECATED_LightWeightInstanceManager* FLightWeightInstanceSubsystem::FindLightWeightInstanceManager(UClass& ActorClass, UWorld& World, const FVector& InPos, const UDataLayerInstance* DataLayer) const
+ALightWeightInstanceManager* FLightWeightInstanceSubsystem::FindLightWeightInstanceManager(UClass& ActorClass, UWorld& World, const FVector& InPos, const UDataLayerInstance* DataLayer) const
 {
 	UE::TReadScopeLock Lock(LWIManagersRWLock);
 
-	for (ADEPRECATED_LightWeightInstanceManager* InstanceManager : LWInstanceManagers)
+	for (ALightWeightInstanceManager* InstanceManager : LWInstanceManagers)
 	{
 		if (IsValid(InstanceManager))
 		{
@@ -162,9 +160,9 @@ ADEPRECATED_LightWeightInstanceManager* FLightWeightInstanceSubsystem::FindLight
 	return nullptr;
 }
 
-ADEPRECATED_LightWeightInstanceManager* FLightWeightInstanceSubsystem::FindOrAddLightWeightInstanceManager(UClass& ActorClass, UWorld& World, const FVector& InPos, const UDataLayerInstance* DataLayer)
+ALightWeightInstanceManager* FLightWeightInstanceSubsystem::FindOrAddLightWeightInstanceManager(UClass& ActorClass, UWorld& World, const FVector& InPos, const UDataLayerInstance* DataLayer)
 {
-	if (ADEPRECATED_LightWeightInstanceManager* FoundManager = FindLightWeightInstanceManager(ActorClass, World, InPos, DataLayer))
+	if (ALightWeightInstanceManager* FoundManager = FindLightWeightInstanceManager(ActorClass, World, InPos, DataLayer))
 	{
 		return FoundManager;
 	}
@@ -183,7 +181,7 @@ ADEPRECATED_LightWeightInstanceManager* FLightWeightInstanceSubsystem::FindOrAdd
 
 	FTransform ManagerTransform(FTransform::Identity);
 
-	ADEPRECATED_LightWeightInstanceManager* InstanceManagerCDO = CastChecked<ADEPRECATED_LightWeightInstanceManager>(BestMatchingClass->GetDefaultObject(true));
+	ALightWeightInstanceManager* InstanceManagerCDO = CastChecked<ALightWeightInstanceManager>(BestMatchingClass->GetDefaultObject(true));
 	int32 GridSize = InstanceManagerCDO->GetGridSize();
 	if (GridSize > 0)
 	{
@@ -193,7 +191,7 @@ ADEPRECATED_LightWeightInstanceManager* FLightWeightInstanceSubsystem::FindOrAdd
 		ManagerTransform.SetLocation(ManagerLocation);
 	}
 
-	if (ADEPRECATED_LightWeightInstanceManager* NewInstanceManager = World.SpawnActor<ADEPRECATED_LightWeightInstanceManager>(BestMatchingClass, ManagerTransform, SpawnParams))
+	if (ALightWeightInstanceManager* NewInstanceManager = World.SpawnActor<ALightWeightInstanceManager>(BestMatchingClass, ManagerTransform, SpawnParams))
 	{
 		NewInstanceManager->SetRepresentedClass(&ActorClass);
 
@@ -224,7 +222,7 @@ UClass* FLightWeightInstanceSubsystem::FindBestInstanceManagerClass(const UClass
 	TArray<UClass*> ManagerClasses;
 	for (TObjectIterator<UClass> It; It; ++It)
 	{
-		if (It->IsChildOf(ADEPRECATED_LightWeightInstanceManager::StaticClass()) && !It->HasAnyClassFlags(CLASS_Abstract))
+		if (It->IsChildOf(ALightWeightInstanceManager::StaticClass()) && !It->HasAnyClassFlags(CLASS_Abstract))
 		{
 			// Skip SKEL and REINST classes.
 			if (It->GetName().StartsWith(TEXT("SKEL_")) || It->GetName().StartsWith(TEXT("REINST_")))
@@ -240,12 +238,12 @@ UClass* FLightWeightInstanceSubsystem::FindBestInstanceManagerClass(const UClass
 	int32 BestDistance = INT_MAX;
 	for (UClass* ManagerClass : ManagerClasses)
 	{
-		if (ManagerClass->GetDefaultObject<ADEPRECATED_LightWeightInstanceManager>()->DoesAcceptClass(InActorClass))
+		if (ManagerClass->GetDefaultObject<ALightWeightInstanceManager>()->DoesAcceptClass(InActorClass))
 		{
-			const UClass* HandledClass = ManagerClass->GetDefaultObject<ADEPRECATED_LightWeightInstanceManager>()->GetRepresentedClassInternal();
+			const UClass* HandledClass = ManagerClass->GetDefaultObject<ALightWeightInstanceManager>()->GetRepresentedClassInternal();
 			if (!HandledClass)
 			{
-				HandledClass = ManagerClass->GetDefaultObject<ADEPRECATED_LightWeightInstanceManager>()->GetAcceptedClass();
+				HandledClass = ManagerClass->GetDefaultObject<ALightWeightInstanceManager>()->GetAcceptedClass();
 			}
 			if (InActorClass == HandledClass)
 			{
@@ -282,7 +280,7 @@ UClass* FLightWeightInstanceSubsystem::GetActorClass(const FActorInstanceHandle&
 		return CachedActor->StaticClass();
 	}
 
-	if (ADEPRECATED_LightWeightInstanceManager* LWIManager = FindLightWeightInstanceManager(Handle))
+	if (ALightWeightInstanceManager* LWIManager = FindLightWeightInstanceManager(Handle))
 	{
 		return LWIManager->GetRepresentedClass(Handle.GetInstanceIndex());
 	}
@@ -299,7 +297,7 @@ FVector FLightWeightInstanceSubsystem::GetLocation(const FActorInstanceHandle& H
 		return CachedActor->GetActorLocation();
 	}
 
-	if (ADEPRECATED_LightWeightInstanceManager* InstanceManager = FindLightWeightInstanceManager(Handle))
+	if (ALightWeightInstanceManager* InstanceManager = FindLightWeightInstanceManager(Handle))
 	{
 		return InstanceManager->GetLocation(Handle);
 	}
@@ -316,7 +314,7 @@ FString FLightWeightInstanceSubsystem::GetName(const FActorInstanceHandle& Handl
 		return CachedActor->GetName();
 	}
 
-	if (ADEPRECATED_LightWeightInstanceManager* InstanceManager = FindLightWeightInstanceManager(Handle))
+	if (ALightWeightInstanceManager* InstanceManager = FindLightWeightInstanceManager(Handle))
 	{
 		return InstanceManager->GetName(Handle);
 	}
@@ -333,7 +331,7 @@ ULevel* FLightWeightInstanceSubsystem::GetLevel(const FActorInstanceHandle& Hand
 		return CachedActor->GetLevel();
 	}
 
-	if (ADEPRECATED_LightWeightInstanceManager* InstanceManager = FindLightWeightInstanceManager(Handle))
+	if (ALightWeightInstanceManager* InstanceManager = FindLightWeightInstanceManager(Handle))
 	{
 		return InstanceManager->GetLevel();
 	}
@@ -350,7 +348,7 @@ bool FLightWeightInstanceSubsystem::IsInLevel(const FActorInstanceHandle& Handle
 		return CachedActor->IsInLevel(InLevel);
 	}
 
-	if (ADEPRECATED_LightWeightInstanceManager* InstanceManager = FindLightWeightInstanceManager(Handle))
+	if (ALightWeightInstanceManager* InstanceManager = FindLightWeightInstanceManager(Handle))
 	{
 		return InstanceManager->GetLevel() == InLevel;
 	}
@@ -363,7 +361,7 @@ FActorInstanceHandle FLightWeightInstanceSubsystem::CreateNewLightWeightInstance
 	// Get or create a light weight instance for this class and data layer
 	if (InitData)
 	{
-		if (ADEPRECATED_LightWeightInstanceManager* LWIManager = FindOrAddLightWeightInstanceManager(*InActorClass, *World, InitData->Transform.GetLocation(), InLayer))
+		if (ALightWeightInstanceManager* LWIManager = FindOrAddLightWeightInstanceManager(*InActorClass, *World, InitData->Transform.GetLocation(), InLayer))
 		{
 			// create an instance with the given data
 			int32 InstanceIdx = LWIManager->AddNewInstance(InitData);
@@ -377,10 +375,8 @@ FActorInstanceHandle FLightWeightInstanceSubsystem::CreateNewLightWeightInstance
 
 void FLightWeightInstanceSubsystem::DeleteInstance(const FActorInstanceHandle& Handle)
 {
-	if (ADEPRECATED_LightWeightInstanceManager* LWIManager = FindLightWeightInstanceManager(Handle))
+	if (ALightWeightInstanceManager* LWIManager = FindLightWeightInstanceManager(Handle))
 	{
 		LWIManager->RemoveInstance(Handle.GetInstanceIndex());
 	}
 }
-
-PRAGMA_ENABLE_DEPRECATION_WARNINGS
