Added a filter function that takes an ActorDescInstance in FForEachActorWithLoadingParams (used by ForEachActorWithLoading).
Also made UDataLayerEditorSubsystem::PassDataLayersFilter public.
These changes allow WorldPartitionBuilders that use custom loading to call FWorldPartitionHelpers::ForEachActorWithLoading and apply the same data-layer filtering that is normally performed automatically when using the LoadingAdapter interface (i.e., the filtering that considers the Builder?s IncludedDataLayers and ExcludedDataLayers).

#Synced-CL 49040969

diff --git a/Engine/Source/Editor/DataLayerEditor/Public/DataLayer/DataLayerEditorSubsystem.h b/Engine/Source/Editor/DataLayerEditor/Public/DataLayer/DataLayerEditorSubsystem.h
--- a/Engine/Source/Editor/DataLayerEditor/Public/DataLayer/DataLayerEditorSubsystem.h
+++ b/Engine/Source/Editor/DataLayerEditor/Public/DataLayer/DataLayerEditorSubsystem.h
@@ -705,6 +705,9 @@ public:
 	 */
 	UE_API bool ApplyActorEditorContextDataLayersToActors(const TArray<AActor*>& InActors);
 
+	/** Returns whether the actor handle passes the data layers filter. */
+	UE_API bool PassDataLayersFilter(UWorld* World, const FWorldPartitionHandle& ActorHandle);
+
 	//~ Begin Deprecated
 
 	PRAGMA_DISABLE_DEPRECATION_WARNINGS
@@ -799,8 +802,6 @@ private:
 
 	UE_API bool SetDataLayerIsLoadedInEditorInternal(UDataLayerInstance* DataLayer, const bool bIsLoadedInEditor, const bool bIsFromUserChange);
 
-	UE_API bool PassDataLayersFilter(UWorld* World, const FWorldPartitionHandle& ActorHandle);
-
 	UE_API void BroadcastActorDataLayersChanged(const TWeakObjectPtr<AActor>& ChangedActor);
 	UE_API void BroadcastDataLayerChanged(const EDataLayerAction Action, const TWeakObjectPtr<const UDataLayerInstance>& ChangedDataLayer, const FName& ChangedProperty);
 	UE_API void BroadcastDataLayerEditorLoadingStateChanged(bool bIsFromUserChange);
diff --git a/Engine/Source/Runtime/Engine/Private/WorldPartition/WorldPartitionHelpers.cpp b/Engine/Source/Runtime/Engine/Private/WorldPartition/WorldPartitionHelpers.cpp
--- a/Engine/Source/Runtime/Engine/Private/WorldPartition/WorldPartitionHelpers.cpp
+++ b/Engine/Source/Runtime/Engine/Private/WorldPartition/WorldPartitionHelpers.cpp
@@ -189,19 +189,22 @@ void FWorldPartitionHelpers::ForEachActorWithLoading(UWorldPartition* WorldParti
 	{
 		if (Algo::AnyOf(Params.ActorClasses, [ActorDescInstance](UClass* ActorClass) { return IsActorDescClassCompatibleWith(ActorDescInstance->GetActorDesc(), ActorClass); }))
 		{
-			if (!Params.FilterActorDesc || Params.FilterActorDesc(ActorDescInstance->GetActorDesc()))
+			if (!Params.FilterActorDescInstance || Params.FilterActorDescInstance(ActorDescInstance))
 			{
-				WorldPartitionHelpers::LoadReferences(WorldPartition, ActorDescInstance->GetGuid(), Result.ActorReferences);
-
-				FWorldPartitionReference ActorReference(WorldPartition, ActorDescInstance->GetGuid());
-				if (!Func(*ActorReference))
-				{
-					return false;
-				}
-
-				if (!Params.bKeepReferences && (Params.bGCPerActor || FWorldPartitionHelpers::ShouldCollectGarbage()))
+				if (!Params.FilterActorDesc || Params.FilterActorDesc(ActorDescInstance->GetActorDesc()))
 				{
-					CallGarbageCollect();
+					WorldPartitionHelpers::LoadReferences(WorldPartition, ActorDescInstance->GetGuid(), Result.ActorReferences);
+
+					FWorldPartitionReference ActorReference(WorldPartition, ActorDescInstance->GetGuid());
+					if (!Func(*ActorReference))
+					{
+						return false;
+					}
+
+					if (!Params.bKeepReferences && (Params.bGCPerActor || FWorldPartitionHelpers::ShouldCollectGarbage()))
+					{
+						CallGarbageCollect();
+					}
 				}
 			}
 		}
diff --git a/Engine/Source/Runtime/Engine/Public/WorldPartition/WorldPartitionHelpers.h b/Engine/Source/Runtime/Engine/Public/WorldPartition/WorldPartitionHelpers.h
--- a/Engine/Source/Runtime/Engine/Public/WorldPartition/WorldPartitionHelpers.h
+++ b/Engine/Source/Runtime/Engine/Public/WorldPartition/WorldPartitionHelpers.h
@@ -138,9 +138,12 @@ public:
 		/* If not empty, iteration will be done only on those actors. */
 		TArray<FGuid> ActorGuids;
 
-		/* Custom filter function used to filter actors loading. */
+		/* Custom filter function used to filter actors loading based on actor desc. */
 		TUniqueFunction<bool(const FWorldPartitionActorDesc*)> FilterActorDesc;
 
+		/* Custom filter function used to filter actors loading based on actor desc instance. */
+		TUniqueFunction<bool(const FWorldPartitionActorDescInstance*)> FilterActorDescInstance;
+
 		/* Called right before releasing actor references and performing garbage collection. */
 		TUniqueFunction<void()> OnPreGarbageCollect;
 	};
