Avoid sending multiple calls to ServerUpdateLevelVisibility from the same client which would lead to triggering ensures.
This happens when a client uses split screen combined with wp.Runtime.UseMakingVisibleTransactionRequests / wp.Runtime.UseMakingInvisibleTransactionRequests.

#Synced-CL 49040969

diff --git a/Engine/Source/Runtime/Engine/Private/LevelStreaming.cpp b/Engine/Source/Runtime/Engine/Private/LevelStreaming.cpp
--- a/Engine/Source/Runtime/Engine/Private/LevelStreaming.cpp
+++ b/Engine/Source/Runtime/Engine/Private/LevelStreaming.cpp
@@ -803,6 +803,8 @@ void ULevelStreaming::ServerUpdateLevelVisibility(bool bIsVisible, bool bTryMake
 				LevelVisibility.PackageName = PlayerController->NetworkRemapPath(LevelVisibility.PackageName, false);
 				LevelVisibility.VisibilityRequestId = TransactionId;
 				PlayerController->ServerUpdateLevelVisibility(LevelVisibility);
+				// No need to loop on all player controllers as ServerUpdateLevelVisibility only needs to be done once per connection
+				break;
 			}
 		}
 	}
