Make sure the GC isn't running when we update the streaming state asynchronously.

#Synced-CL 0

diff --git a/Engine/Source/Runtime/Engine/Private/WorldPartition/WorldPartitionStreamingPolicy.cpp b/Engine/Source/Runtime/Engine/Private/WorldPartition/WorldPartitionStreamingPolicy.cpp
--- a/Engine/Source/Runtime/Engine/Private/WorldPartition/WorldPartitionStreamingPolicy.cpp
+++ b/Engine/Source/Runtime/Engine/Private/WorldPartition/WorldPartitionStreamingPolicy.cpp
@@ -785,7 +785,10 @@ void UWorldPartitionStreamingPolicy::OnStreamingStateUpdated()
 			AsyncUpdateStreamingStateTask = UE::Tasks::Launch(UE_SOURCE_LOCATION, [InputParams = MoveTemp(InputParams), this]()
 			{
 				QUICK_SCOPE_CYCLE_COUNTER(STAT_WorldPartitionStreamingPolicy_AsyncUpdateStreamingState);
+				
+				FGCScopeGuard GCGuard;
 				UWorldPartitionStreamingPolicy::UpdateStreamingStateInternal(InputParams, AsyncTaskTargetState);
+
 			}, UE::Tasks::ETaskPriority::Normal);
 			AsyncUpdateTaskState = EAsyncUpdateTaskState::Started;
 		}
