Fix typos in FActorSpawnParameters member: CustomPreSpawnInitalization -> CustomPreSpawnInitialization and others reported by the Unacceptable Word Validator.

#Synced-CL 49407983

diff --git a/Engine/Source/Runtime/Engine/Classes/Engine/World.h b/Engine/Source/Runtime/Engine/Classes/Engine/World.h
--- a/Engine/Source/Runtime/Engine/Classes/Engine/World.h
+++ b/Engine/Source/Runtime/Engine/Classes/Engine/World.h
@@ -513,7 +513,7 @@ public:
 	EObjectFlags ObjectFlags;
 
 	/* Custom function allowing the caller to specific a function to execute post actor construction but before other systems see this actor spawn. */
-	TFunction<void(AActor*)> CustomPreSpawnInitalization;
+	TFunction<void(AActor*)> CustomPreSpawnInitialization;
 };
 
 /* World actors spawmning helper functions */
diff --git a/Engine/Source/Runtime/Engine/Classes/GameFramework/LightWeightInstanceManager.h b/Engine/Source/Runtime/Engine/Classes/GameFramework/LightWeightInstanceManager.h
--- a/Engine/Source/Runtime/Engine/Classes/GameFramework/LightWeightInstanceManager.h
+++ b/Engine/Source/Runtime/Engine/Classes/GameFramework/LightWeightInstanceManager.h
@@ -116,7 +116,7 @@ protected:
 	ENGINE_API virtual UClass* GetActorClassToSpawn(const FActorInstanceHandle& Handle) const;
 
 	// Called after actor construction but before other systems see this actor spawn
-	ENGINE_API virtual void PreSpawnInitalization(const FActorInstanceHandle& Handle, AActor* SpawnedActor);
+	ENGINE_API virtual void PreSpawnInitialization(const FActorInstanceHandle& Handle, AActor* SpawnedActor);
 
 	// Called after spawning a new actor from a light weight instance
 	ENGINE_API virtual void PostActorSpawn(const FActorInstanceHandle& Handle);
diff --git a/Engine/Source/Runtime/Engine/Private/LevelActor.cpp b/Engine/Source/Runtime/Engine/Private/LevelActor.cpp
--- a/Engine/Source/Runtime/Engine/Private/LevelActor.cpp
+++ b/Engine/Source/Runtime/Engine/Private/LevelActor.cpp
@@ -272,7 +272,7 @@ void LineCheckTracker::CaptureLineCheck(int32 LineCheckFlags, const FVector* Ext
  *  - The name number bit 31 is reserved to preserve the fast path name generation (see
  *	  GFastPathUniqueNameGeneration).
  *	- On some environments, cooking happens without network, so we can't retrieve the 
- *	  MAC adress for spawned actors, assign a default one in this case.
+ *	  MAC address for spawned actors, assign a default one in this case.
  **/
 class FActorGUIDGenerator
 {
@@ -726,9 +726,9 @@ AActor* UWorld::SpawnActor( UClass* Class, FTransform const* UserTransformPtr, c
 		FActorParentComponentSetter::Set(Actor, SpawnParameters.OverrideParentComponent);
 	}
 
-	if (SpawnParameters.CustomPreSpawnInitalization)
+	if (SpawnParameters.CustomPreSpawnInitialization)
 	{
-		SpawnParameters.CustomPreSpawnInitalization(Actor);
+		SpawnParameters.CustomPreSpawnInitialization(Actor);
 	}
 
 	if ( GUndo )
diff --git a/Engine/Source/Runtime/Engine/Private/LightWeightInstanceManager.cpp b/Engine/Source/Runtime/Engine/Private/LightWeightInstanceManager.cpp
--- a/Engine/Source/Runtime/Engine/Private/LightWeightInstanceManager.cpp
+++ b/Engine/Source/Runtime/Engine/Private/LightWeightInstanceManager.cpp
@@ -129,9 +129,9 @@ AActor* ADEPRECATED_LightWeightInstanceManager::ConvertInstanceToActor(const FAc
 	{
 		FActorSpawnParameters SpawnParams;
 		SetSpawnParameters(SpawnParams);
-		SpawnParams.CustomPreSpawnInitalization = [this, &Handle](AActor* SpawnedActor)
+		SpawnParams.CustomPreSpawnInitialization = [this, &Handle](AActor* SpawnedActor)
 		{
-			PreSpawnInitalization(Handle, SpawnedActor);
+			PreSpawnInitialization(Handle, SpawnedActor);
 		};
 
 		NewActor = GetLevel()->GetWorld()->SpawnActor<AActor>(GetActorClassToSpawn(Handle), GetTransform(Handle), SpawnParams);
@@ -310,7 +310,7 @@ UClass* ADEPRECATED_LightWeightInstanceManager::GetActorClassToSpawn(const FActo
 	return RepresentedClass;
 }
 
-void ADEPRECATED_LightWeightInstanceManager::PreSpawnInitalization(const FActorInstanceHandle& Handle, AActor* SpawnedActor)
+void ADEPRECATED_LightWeightInstanceManager::PreSpawnInitialization(const FActorInstanceHandle& Handle, AActor* SpawnedActor)
 {
 	Handle.SetCachedActor(SpawnedActor);
 	Actors.Add(Handle.GetInstanceIndex(), SpawnedActor);
