Remove unused code from UGameUserSettings::UpdateResolutionQuality

#Synced-CL 49357083

diff --git a/Engine/Source/Runtime/Engine/Private/GameUserSettings.cpp b/Engine/Source/Runtime/Engine/Private/GameUserSettings.cpp
--- a/Engine/Source/Runtime/Engine/Private/GameUserSettings.cpp
+++ b/Engine/Source/Runtime/Engine/Private/GameUserSettings.cpp
@@ -277,9 +277,6 @@ void UGameUserSettings::UpdateVersion()
 
 void UGameUserSettings::UpdateResolutionQuality()
 {
-	const int32 MinHeight = UKismetSystemLibrary::GetMinYResolutionFor3DView();
-	const int32 ScreenWidth = (FullscreenMode == EWindowMode::WindowedFullscreen) ? GetDesktopResolution().X : ResolutionSizeX;
-	const int32 ScreenHeight = (FullscreenMode == EWindowMode::WindowedFullscreen) ? GetDesktopResolution().Y : ResolutionSizeY;
 	PRAGMA_DISABLE_DEPRECATION_WARNINGS
 	MinResolutionScale = Scalability::MinResolutionScale;
 	PRAGMA_ENABLE_DEPRECATION_WARNINGS
