NetDriver
* Removed log warning about calling SetReplicates on non-initialized actors but kept comment explaining why its not ideal

#rnx
#Synced-CL 49371252

diff --git a/Engine/Source/Runtime/Engine/Private/Actor.cpp b/Engine/Source/Runtime/Engine/Private/Actor.cpp
--- a/Engine/Source/Runtime/Engine/Private/Actor.cpp
+++ b/Engine/Source/Runtime/Engine/Private/Actor.cpp
@@ -4531,6 +4531,8 @@ void AActor::SetReplicates(bool bInReplicates)
 		RemoteRole = ExpectedRemoteRole;
 		bReplicates = bInReplicates;
 
+		// Initialized actors will immediately get added/removed from the ReplicationSystem
+		// Uninitialized actors should set bReplicates instead as SetReplicates() will change the RemoteRole earlier then normal
 		if (bActorInitialized)
 		{
 			if (bReplicates)
@@ -4567,10 +4569,6 @@ void AActor::SetReplicates(bool bInReplicates)
 
 			MARK_PROPERTY_DIRTY_FROM_NAME(AActor, RemoteRole, this);
 		}
-		else
-		{
-			UE_LOG(LogActor, Warning, TEXT("SetReplicates called on non-initialized actor %s. Directly setting bReplicates is the correct procedure for pre-init actors."), *GetName());
-		}
 	}
 }
 
