[World Asset Streaming] Include HLOD-cell flag in hash used for cell GUID generation.

When injecting an always-loaded HLOD cell, its runtime partition resolves to the main partition, which can cause a GUID collision with the corresponding main partition cell. Adding the HLOD-cell flag to the hash ensures that the injected HLOD cell generates a distinct GUID and prevents the collision.

#Synced-CL 49363040

diff --git a/Engine/Source/Runtime/Engine/Private/WorldPartition/RuntimeHashSet/WorldPartitionRuntimeHashSet.cpp b/Engine/Source/Runtime/Engine/Private/WorldPartition/RuntimeHashSet/WorldPartitionRuntimeHashSet.cpp
--- a/Engine/Source/Runtime/Engine/Private/WorldPartition/RuntimeHashSet/WorldPartitionRuntimeHashSet.cpp
+++ b/Engine/Source/Runtime/Engine/Private/WorldPartition/RuntimeHashSet/WorldPartitionRuntimeHashSet.cpp
@@ -824,7 +824,7 @@ FGuid UWorldPartitionRuntimeHashSet::RegisterWorldAssetStreaming(const UWorldPar
 		FName TargetPartitionName = TargetPartition->Name;
 		FGuid InstanceGuid = InParams.Guid;
 		FArchiveMD5 ArMD5;
-		ArMD5 << TargetPartitionName << InstanceGuid;
+		ArMD5 << TargetPartitionName << InstanceGuid << bIsHLODPass;
 		const FGuid CellGuid = ArMD5.GetGuidFromHash();
 		check(CellGuid.IsValid());
 		if (!bIsHLODPass)
