Editor: Outliner eye toggle fails to hide static meshes when Affect Indirect Lighting While Hidden is enabled

The logic should treat affect indirect the same as cast shadow with respect to the "while hidden" feature.

After this change, togglging the "Visible" property still works as expected (hidden objects can cast shadow or affect indirect lighting), while using the eye toggle in the editor will fully remove the object from rendering.

#Synced-CL 49422785

diff --git a/Engine/Source/Runtime/Engine/Private/Components/PrimitiveComponent.cpp b/Engine/Source/Runtime/Engine/Private/Components/PrimitiveComponent.cpp
--- a/Engine/Source/Runtime/Engine/Private/Components/PrimitiveComponent.cpp
+++ b/Engine/Source/Runtime/Engine/Private/Components/PrimitiveComponent.cpp
@@ -2037,7 +2037,10 @@ bool UPrimitiveComponent::ShouldComponentAddToScene() const
 	const bool bIsTemporarilyHiddenInEditor = false;
 #endif
 
-	return bSceneAdd && (ShouldRender() || (CastShadow && bCastHiddenShadow && !bIsHiddenInEditor) || bAffectIndirectLightingWhileHidden || bRayTracingFarField);
+	return bSceneAdd && (ShouldRender() ||
+						 (CastShadow && bCastHiddenShadow && !bIsHiddenInEditor) ||
+						 (bAffectDynamicIndirectLighting && bAffectIndirectLightingWhileHidden && !bIsHiddenInEditor) ||
+						 bRayTracingFarField);
 }
 
 bool UPrimitiveComponent::ShouldCreatePhysicsState() const
