[NetPhysics] Add functions to get and reset the NetworkPhysicsSettingsResimulationErrorCorrection of a PhysicObject

#Synced-CL 49424073

diff --git a/Engine/Source/Runtime/Engine/Private/PhysicsEngine/NetworkPhysicsSettingsComponent.cpp b/Engine/Source/Runtime/Engine/Private/PhysicsEngine/NetworkPhysicsSettingsComponent.cpp
--- a/Engine/Source/Runtime/Engine/Private/PhysicsEngine/NetworkPhysicsSettingsComponent.cpp
+++ b/Engine/Source/Runtime/Engine/Private/PhysicsEngine/NetworkPhysicsSettingsComponent.cpp
@@ -213,14 +213,14 @@ UNetworkPhysicsSettingsComponent* UNetworkPhysicsSettingsComponent::GetSettingsF
 	return Value ? *Value : nullptr;
 }
 
-void FNetworkPhysicsSettingsResimulationErrorCorrection::ApplySettings_External(Chaos::FPhysicsObjectHandle PhysicsObject) const
+void FNetworkPhysicsSettingsResimulationErrorCorrection::ApplySettings_External(Chaos::FPhysicsObjectHandle PhysicsObject, const bool bInForceOverride/* = false*/) const
 {
 	if (!PhysicsObject)
 	{
 		return;
 	}
 
-	if (bOverrideResimErrorInterpolationSettings)
+	if (bOverrideResimErrorInterpolationSettings || bInForceOverride)
 	{
 		FLockedWritePhysicsObjectExternalInterface Interface = FPhysicsObjectExternalInterface::LockWrite({ &PhysicsObject, 1 });
 
@@ -289,6 +289,146 @@ void FNetworkPhysicsSettingsResimulationErrorCorrection::ApplySettings_External(
 	}
 }
 
+void FNetworkPhysicsSettingsResimulationErrorCorrection::ResetSettings_External(Chaos::FPhysicsObjectHandle PhysicsObject)
+{
+	if (!PhysicsObject)
+	{
+		return;
+	}
+
+	FLockedWritePhysicsObjectExternalInterface Interface = FPhysicsObjectExternalInterface::LockWrite({ &PhysicsObject, 1 });
+
+	if (const Chaos::FGeometryParticle* Particle = Interface->GetParticle(PhysicsObject))
+	{
+		if (!Particle->GetProxy())
+		{
+			return;
+		}
+
+		auto ResetSettingsHelper = [](auto Proxy)
+			{
+				if (FProxyInterpolationError* InterpolationError = Proxy->template GetOrCreateErrorInterpolationData<FProxyInterpolationError>())
+				{
+					InterpolationError->ResetErrorInterpolationSettings();
+				}
+			};
+
+		switch (Particle->GetProxy()->GetType())
+		{
+			case EPhysicsProxyType::SingleParticleProxy:
+			{
+				Chaos::FSingleParticlePhysicsProxy* SPProxy = static_cast<Chaos::FSingleParticlePhysicsProxy*>(Particle->GetProxy());
+				if (ensure(SPProxy))
+				{
+					ResetSettingsHelper(SPProxy);
+				}
+				break;
+			}
+			case EPhysicsProxyType::ClusterUnionProxy:
+			{
+				Chaos::FClusterUnionPhysicsProxy* CUProxy = static_cast<Chaos::FClusterUnionPhysicsProxy*>(Particle->GetProxy());
+				if (ensure(CUProxy))
+				{
+					ResetSettingsHelper(CUProxy);
+				}
+				break;
+			}
+			case EPhysicsProxyType::GeometryCollectionType:
+			{
+				FGeometryCollectionPhysicsProxy* GCProxy = static_cast<FGeometryCollectionPhysicsProxy*>(Particle->GetProxy());
+				if (ensure(GCProxy))
+				{
+					ResetSettingsHelper(GCProxy);
+				}
+				break;
+			}
+			default:
+			{
+				ensure(false);
+				break;
+			}
+		};
+	}
+}
+
+bool FNetworkPhysicsSettingsResimulationErrorCorrection::FromPhysicObject(Chaos::FPhysicsObjectHandle PhysicsObject, FNetworkPhysicsSettingsResimulationErrorCorrection& OutValue)
+{
+	if (!PhysicsObject)
+	{
+		return false;
+	}
+
+	FLockedReadPhysicsObjectExternalInterface Interface = FPhysicsObjectExternalInterface::LockRead({ &PhysicsObject, 1 });
+
+	if (const Chaos::FGeometryParticle* Particle = Interface->GetParticle(PhysicsObject))
+	{
+		if (!Particle->GetProxy())
+		{
+			return false;
+		}
+
+		auto GetSettingsHelper = [&OutValue](auto Proxy)
+		{
+			// Get the proxies error interpolation, create one if not already present
+			if (const FProxyInterpolationError* InterpolationError = Proxy->template GetErrorInterpolationData<FProxyInterpolationError>())
+			{
+				// Apply the custom settings
+				if (const FErrorInterpolationSettings* Settings = InterpolationError->GetErrorInterpolationSettings())
+				{
+					OutValue.ResimErrorCorrectionDuration = Settings->ErrorCorrectionDuration;
+					OutValue.ResimErrorMaximumDistanceBeforeSnapping = Settings->MaximumErrorCorrectionBeforeSnapping;
+					OutValue.ResimErrorMaximumDesyncTimeBeforeSnapping = Settings->MaximumErrorCorrectionDesyncTimeBeforeSnapping;
+					OutValue.ResimErrorDirectionalDecayMultiplier = Settings->ErrorDirectionalDecayMultiplier;
+					OutValue.bRenderInterpApplyExponentialDecay = Settings->bApplyExponentialDecay;
+					OutValue.RenderInterpExponentialDecayLinearHalfLife = Settings->ExponentialDecayLinearHalfLife;
+					OutValue.RenderInterpExponentialDecayAngularHalfLife = Settings->ExponentialDecayAngularHalfLife;
+					OutValue.RenderInterpMinimumLinearThreshold = Settings->MinimumLinearThreshold;
+					OutValue.RenderInterpMinimumAngularThreshold = Settings->MinimumAngularThreshold;
+					return true;
+				}
+			}
+			return false;
+		};
+
+		switch (Particle->GetProxy()->GetType())
+		{
+			case EPhysicsProxyType::SingleParticleProxy:
+			{
+				const Chaos::FSingleParticlePhysicsProxy* SPProxy = static_cast<Chaos::FSingleParticlePhysicsProxy*>(Particle->GetProxy());
+				if (ensure(SPProxy))
+				{
+					return GetSettingsHelper(SPProxy);
+				}
+				break;
+			}
+			case EPhysicsProxyType::ClusterUnionProxy:
+			{
+				const Chaos::FClusterUnionPhysicsProxy* CUProxy = static_cast<Chaos::FClusterUnionPhysicsProxy*>(Particle->GetProxy());
+				if (ensure(CUProxy))
+				{
+					return GetSettingsHelper(CUProxy);
+				}
+				break;
+			}
+			case EPhysicsProxyType::GeometryCollectionType:
+			{
+				const FGeometryCollectionPhysicsProxy* GCProxy = static_cast<FGeometryCollectionPhysicsProxy*>(Particle->GetProxy());
+				if (ensure(GCProxy))
+				{
+					return GetSettingsHelper(GCProxy);
+				}
+				break;
+			}
+			default:
+			{
+				ensure(false);
+				return false;
+			}
+		};
+	}
+	return false;
+}
+
 #pragma region // FNetworkPhysicsSettingsComponentAsync
 
 FNetworkPhysicsSettingsComponentAsync::FNetworkPhysicsSettingsComponentAsync() : TSimCallbackObject()
diff --git a/Engine/Source/Runtime/Engine/Public/Physics/NetworkPhysicsSettingsComponent.h b/Engine/Source/Runtime/Engine/Public/Physics/NetworkPhysicsSettingsComponent.h
--- a/Engine/Source/Runtime/Engine/Public/Physics/NetworkPhysicsSettingsComponent.h
+++ b/Engine/Source/Runtime/Engine/Public/Physics/NetworkPhysicsSettingsComponent.h
@@ -331,7 +331,11 @@ struct FNetworkPhysicsSettingsResimulationErrorCorrection
 	GENERATED_BODY()
 
 	/** Apply these settings to Physics Object */
-	void ApplySettings_External(Chaos::FPhysicsObjectHandle PhysicsObject) const;
+	ENGINE_API void ApplySettings_External(Chaos::FPhysicsObjectHandle PhysicsObject, const bool bInForceOverride = false) const;
+
+	ENGINE_API static void ResetSettings_External(Chaos::FPhysicsObjectHandle PhysicsObject);
+
+	ENGINE_API static bool FromPhysicObject(Chaos::FPhysicsObjectHandle PhysicsObject, FNetworkPhysicsSettingsResimulationErrorCorrection& OutValue);
 
 	/** Enable override for post-resimulation error correction settings during render interpolation
 	* NOTE: This currently does not work if the experimental p.RenderInterp.ErrorVelocityCorrection CVar is set to true (false by default) */
diff --git a/Engine/Source/Runtime/Experimental/Chaos/Public/Chaos/Framework/PhysicsProxyBase.h b/Engine/Source/Runtime/Experimental/Chaos/Public/Chaos/Framework/PhysicsProxyBase.h
--- a/Engine/Source/Runtime/Experimental/Chaos/Public/Chaos/Framework/PhysicsProxyBase.h
+++ b/Engine/Source/Runtime/Experimental/Chaos/Public/Chaos/Framework/PhysicsProxyBase.h
@@ -319,6 +319,12 @@ public:
 
 	/** Get FErrorInterpolationSettings which stores custom settings for render interpolation error corrections, returns nullptr if there are no custom settings */
 	virtual FErrorInterpolationSettings* GetErrorInterpolationSettings() { return nullptr; }
+
+	/** Get FErrorInterpolationSettings which stores custom settings for render interpolation error corrections, returns nullptr if there are no custom settings */
+	virtual const FErrorInterpolationSettings* GetErrorInterpolationSettings() const { return nullptr; }
+
+	/** Reset FErrorInterpolationSettings */
+	virtual void ResetErrorInterpolationSettings() { }
 };
 
 // Render interpolation that can correct errors from resimulation / repositions through a linear decay over N simulation tick.
@@ -375,6 +381,16 @@ public:
 	{
 		return ErrorInterpolationSettings.Get();
 	}
+	
+	virtual const FErrorInterpolationSettings* GetErrorInterpolationSettings() const override
+    {
+		return ErrorInterpolationSettings.Get();
+	}
+
+	virtual void ResetErrorInterpolationSettings() override
+	{
+		ErrorInterpolationSettings.Reset();
+	}	
 
 protected:
 	int32 LastSimTick = 0;
diff --git a/Engine/Source/Runtime/Experimental/Chaos/Public/PhysicsProxy/ClusterUnionPhysicsProxy.h b/Engine/Source/Runtime/Experimental/Chaos/Public/PhysicsProxy/ClusterUnionPhysicsProxy.h
--- a/Engine/Source/Runtime/Experimental/Chaos/Public/PhysicsProxy/ClusterUnionPhysicsProxy.h
+++ b/Engine/Source/Runtime/Experimental/Chaos/Public/PhysicsProxy/ClusterUnionPhysicsProxy.h
@@ -196,6 +196,17 @@ namespace Chaos
 
 			return static_cast<ErrorDataType*>(InterpolationData.Get());
 		}
+		
+		template<typename ErrorDataType>
+		const ErrorDataType* GetErrorInterpolationData() const
+		{
+			return static_cast<ErrorDataType*>(InterpolationData.Get());
+		}
+
+		void ResetInterpolationData()
+		{
+			InterpolationData.Reset();
+		}
 
 	private:
 		// An array of a particles that exist in the external implicit object union.
diff --git a/Engine/Source/Runtime/Experimental/Chaos/Public/PhysicsProxy/GeometryCollectionPhysicsProxy.h b/Engine/Source/Runtime/Experimental/Chaos/Public/PhysicsProxy/GeometryCollectionPhysicsProxy.h
--- a/Engine/Source/Runtime/Experimental/Chaos/Public/PhysicsProxy/GeometryCollectionPhysicsProxy.h
+++ b/Engine/Source/Runtime/Experimental/Chaos/Public/PhysicsProxy/GeometryCollectionPhysicsProxy.h
@@ -518,6 +518,11 @@ public:
 	FProxyInterpolationBase* GetInterpolationData() { return InterpolationData.Get(); }
 	const FProxyInterpolationBase* GetInterpolationData() const { return InterpolationData.Get(); }
 
+	void ResetInterpolationData()
+	{
+		InterpolationData.Reset();
+	}
+	
 	enum class EReplicationMode: uint8
 	{
 		Unknown,
@@ -827,6 +832,12 @@ public:
 
 		return static_cast<ErrorDataType*>(InterpolationData.Get());
 	}
+	
+	template<typename ErrorDataType>
+	const ErrorDataType* GetErrorInterpolationData() const
+	{
+		return static_cast<ErrorDataType*>(InterpolationData.Get());
+	}
 
 private:
 #if WITH_EDITORONLY_DATA
diff --git a/Engine/Source/Runtime/Experimental/Chaos/Public/PhysicsProxy/SingleParticlePhysicsProxy.h b/Engine/Source/Runtime/Experimental/Chaos/Public/PhysicsProxy/SingleParticlePhysicsProxy.h
--- a/Engine/Source/Runtime/Experimental/Chaos/Public/PhysicsProxy/SingleParticlePhysicsProxy.h
+++ b/Engine/Source/Runtime/Experimental/Chaos/Public/PhysicsProxy/SingleParticlePhysicsProxy.h
@@ -70,6 +70,11 @@ public:
 	FProxyInterpolationBase* GetInterpolationData() { return InterpolationData.Get(); }
 	const FProxyInterpolationBase* GetInterpolationData() const { return InterpolationData.Get(); }
 
+	void ResetInterpolationData()
+	{
+		InterpolationData.Reset();
+	}
+	
 	FORCEINLINE FRigidBodyHandle_External& GetGameThreadAPI()
 	{
 		return (FRigidBodyHandle_External&)*this;
@@ -187,6 +192,12 @@ public:
 
 		return static_cast<ErrorDataType*>(InterpolationData.Get());
 	}
+	
+	template<typename ErrorDataType>
+	const ErrorDataType* GetErrorInterpolationData() const
+	{
+		return static_cast<ErrorDataType*>(InterpolationData.Get());
+	}
 
 protected:
 	TUniquePtr<PARTICLE_TYPE> Particle;
