[NetPhysics] Fix issue where a sleeping object would cause a desync when woken up by interaction while using resimulation replication.
- Instead of applying sleep state in the forward predicted timeline, compare it in the interpolated timeline and trigger a resimulation if desynced (same as for XRVW)

#Synced-CL 49558020

diff --git a/Engine/Source/Runtime/Engine/Private/PhysicsEngine/PhysicsReplication.cpp b/Engine/Source/Runtime/Engine/Private/PhysicsEngine/PhysicsReplication.cpp
--- a/Engine/Source/Runtime/Engine/Private/PhysicsEngine/PhysicsReplication.cpp
+++ b/Engine/Source/Runtime/Engine/Private/PhysicsEngine/PhysicsReplication.cpp
@@ -121,6 +121,9 @@ namespace PhysicsReplicationCVars
 	{
 		extern bool bApplyPredictiveInterpolationWhenBehindServer;
 
+		bool bResimulateSleepDesync = true;
+		static FAutoConsoleVariableRef CVarResimulateSleepDesync(TEXT("np2.Resim.ResimulateSleepDesync"), bResimulateSleepDesync, TEXT("Set if sleeping discrepancy between client and server should trigger a resimulation. Not enabling this can lead to client desyncing continuously if the server is sleeping and not sending data to client while the client object is still dynamic."));
+
 		bool bRuntimeCorrectionEnabled = false;
 		static FAutoConsoleVariableRef CVarResimRuntimeCorrectionEnabled(TEXT("np2.Resim.RuntimeCorrectionEnabled"), bRuntimeCorrectionEnabled, TEXT("Apply positional and rotational runtime corrections while within resim trigger distance."));
 
@@ -2372,6 +2375,7 @@ bool FPhysicsReplicationAsync::ShouldResimulateParticle(const Chaos::FPBDRigidPa
 		const bool bCompareR = Chaos::FPhysicsSolverBase::GetResimulationErrorRotationThresholdEnabled() || SettingsCurrent.ResimulationSettings.bOverrideResimulationErrorRotationThreshold;
 		const bool bCompareV = Chaos::FPhysicsSolverBase::GetResimulationErrorLinearVelocityThresholdEnabled() || SettingsCurrent.ResimulationSettings.bOverrideResimulationErrorLinearVelocityThreshold;
 		const bool bCompareW = Chaos::FPhysicsSolverBase::GetResimulationErrorAngularVelocityThresholdEnabled() || SettingsCurrent.ResimulationSettings.bOverrideResimulationErrorAngularVelocityThreshold;
+		const bool bCompareSleep = PhysicsReplicationCVars::ResimulationCVars::bResimulateSleepDesync;
 
 	#if !(UE_BUILD_SHIPPING || UE_BUILD_TEST)
 		// Debugging
@@ -2432,6 +2436,20 @@ bool FPhysicsReplicationAsync::ShouldResimulateParticle(const Chaos::FPBDRigidPa
 	#endif
 		}
 
+		// Check for sleeping discrepancy between client and server
+		if (!bShouldTriggerResim && bCompareSleep)
+		{
+			const bool bShouldSleep = (Target.TargetState.Flags & ERigidBodyFlags::Sleeping) != 0;
+			const bool bIsSleep = PastState.ObjectState() == Chaos::EObjectStateType::Sleeping;
+			bShouldTriggerResim = bIsSleep != bShouldSleep;
+	#if !(UE_BUILD_SHIPPING || UE_BUILD_TEST)
+			if (bShouldTriggerResim)
+			{
+				DebugColor = FColor::White;
+			}
+	#endif
+		}
+
 	#if !(UE_BUILD_SHIPPING || UE_BUILD_TEST)
 
 		if (Chaos::FPhysicsSolverBase::CanDebugNetworkPhysicsPrediction())
@@ -2516,15 +2534,8 @@ bool FPhysicsReplicationAsync::ResimulationReplication(Chaos::FPBDRigidParticleH
 				// Get state from locally cached history for frame corresponding to received data
 				const Chaos::FGeometryParticleState PastState = RewindData->GetPastStateAtFrame(*Handle, LocalFrame, RewindPhase);
 
-				const bool bShouldSleep = (Target.TargetState.Flags & ERigidBodyFlags::Sleeping) != 0;
 				const bool bShouldTriggerResim = (Chaos::RewindBeforeAdvance == 0) ? ShouldResimulateParticle(Handle, Target, DeltaSeconds) : false;
 
-				// Wake up if is sleeping and should not sleep
-				if (Handle->IsSleeping() && !bShouldSleep)
-				{
-					RigidsSolver->GetEvolution()->SetParticleObjectState(Handle, Chaos::EObjectStateType::Dynamic);
-				}
-
 				if (bShouldTriggerResim && Target.TickCount == 0 && LocalFrame > RewindData->GetBlockedResimFrame())
 				{
 					// Request resimulation
@@ -2582,19 +2593,8 @@ bool FPhysicsReplicationAsync::ResimulationReplication(Chaos::FPBDRigidParticleH
 					bClearTarget = Target.TickCount >= NumPredictedFrames;
 				}
 
-				// Set sleep state if we are about to clear the target from memory and the target is set to sleep
-				if (bClearTarget && bShouldSleep)
-				{
-					// Snap object into correct state, it should already be at that state or very close to it
-					RigidsSolver->GetEvolution()->ApplyParticleTransformCorrection(Handle, Target.TargetState.Position, Target.TargetState.Quaternion, /*bApplyToConnectedBodies*/true, /*bInRecalculateFrictionOnConnectedBodies*/true, ReplicatedParticleIDs);
-
-					RigidsSolver->GetEvolution()->SetParticleObjectState(Handle, Chaos::EObjectStateType::Sleeping);
-					if (PhysicsReplicationCVars::PredictiveInterpolationCVars::bSleepConnectedBodies)
-					{
-						RigidsSolver->GetEvolution()->ApplySleepOnConnectedParticles(Handle);
-					}
-				}
-				else if (Target.IsWaiting()) 
+				const bool bShouldSleep = (Target.TargetState.Flags & ERigidBodyFlags::Sleeping) != 0;
+				if (Target.IsWaiting() && !bShouldSleep)
 				{
 					// Don't clear the target if we are waiting for a specific target frame and not sleeping
 					bClearTarget = false;
