UMeshComponent should pass it's owner instead of a component in OnAssignedAsOverride

#rnx
#Synced-CL 49335702

diff --git a/Engine/Source/Runtime/Engine/Private/Components/MeshComponent.cpp b/Engine/Source/Runtime/Engine/Private/Components/MeshComponent.cpp
--- a/Engine/Source/Runtime/Engine/Private/Components/MeshComponent.cpp
+++ b/Engine/Source/Runtime/Engine/Private/Components/MeshComponent.cpp
@@ -97,7 +97,7 @@ void UMeshComponent::SetMaterial(int32 ElementIndex, UMaterialInterface* Materia
 
 			if (Material)
 			{
-				Material->OnAssignedAsOverride(this);
+				Material->OnAssignedAsOverride(GetOwner());
 			}
 
 			// Precache PSOs again
