Add bChecked support for invalid PersistentLevel member when calling UWorld::GetWorldSettings.

#Synced-CL 49558965

diff --git a/Engine/Source/Runtime/Engine/Private/World.cpp b/Engine/Source/Runtime/Engine/Private/World.cpp
--- a/Engine/Source/Runtime/Engine/Private/World.cpp
+++ b/Engine/Source/Runtime/Engine/Private/World.cpp
@@ -6857,8 +6857,13 @@ AWorldSettings* UWorld::K2_GetWorldSettings()
 
 AWorldSettings* UWorld::GetWorldSettings( const bool bCheckStreamingPersistent, const bool bChecked ) const
 {
-	//checkSlow(!IsInActualRenderingThread());
 	AWorldSettings* WorldSettings = nullptr;
+
+	if (bChecked)
+	{
+		checkf(PersistentLevel != nullptr, TEXT("%s"), *GetPathName());
+	}
+
 	if (PersistentLevel)
 	{
 		WorldSettings = PersistentLevel->GetWorldSettings(bChecked);
@@ -6877,6 +6882,7 @@ AWorldSettings* UWorld::GetWorldSettings( const bool bCheckStreamingPersistent,
 			}
 		}
 	}
+
 	return WorldSettings;
 }
 
