Make the FCurveTable::OnCurveTableChanged and UDataTable::FOnDataTableChanged delegates thread safe since they're bound during serialize from the ALT (which is potentially called from multiple threads if MT async loading is enabled). Fixes at least one soak crash with MT async loading.

#Synced-CL 49085100

diff --git a/Engine/Source/Runtime/Engine/Classes/Engine/CurveTable.h b/Engine/Source/Runtime/Engine/Classes/Engine/CurveTable.h
--- a/Engine/Source/Runtime/Engine/Classes/Engine/CurveTable.h
+++ b/Engine/Source/Runtime/Engine/Classes/Engine/CurveTable.h
@@ -43,7 +43,7 @@ class UCurveTable
 {
 	GENERATED_UCLASS_BODY()
 
-	DECLARE_MULTICAST_DELEGATE(FOnCurveTableChanged);
+	DECLARE_TS_MULTICAST_DELEGATE(FOnCurveTableChanged);
 
 	const TMap<FName, FRealCurve*>& GetRowMap() const { return RowMap; }
 	const TMap<FName, FRealCurve*>& GetRowMap() { return RowMap; }
diff --git a/Engine/Source/Runtime/Engine/Classes/Engine/DataTable.h b/Engine/Source/Runtime/Engine/Classes/Engine/DataTable.h
--- a/Engine/Source/Runtime/Engine/Classes/Engine/DataTable.h
+++ b/Engine/Source/Runtime/Engine/Classes/Engine/DataTable.h
@@ -81,7 +81,7 @@ class UDataTable
 {
 	GENERATED_UCLASS_BODY()
 
-	DECLARE_MULTICAST_DELEGATE(FOnDataTableChanged);
+	DECLARE_TS_MULTICAST_DELEGATE(FOnDataTableChanged);
 	DECLARE_MULTICAST_DELEGATE(FOnDataTableImport);
 	
 	friend FDataTableEditorUtils;
