Adapted PR #14216: Fix a deadlock in level streaming caused by an incorrect call to IncrementalUnregisterComponents.

#Synced-CL 49558965

diff --git a/Engine/Source/Runtime/Engine/Private/Level.cpp b/Engine/Source/Runtime/Engine/Private/Level.cpp
--- a/Engine/Source/Runtime/Engine/Private/Level.cpp
+++ b/Engine/Source/Runtime/Engine/Private/Level.cpp
@@ -2164,7 +2164,7 @@ bool ULevel::IncrementalUnregisterComponents(int32 NumComponentsToUnregister, FU
 		return NetMode == NM_DedicatedServer;
 	};
 
-	auto IncrementalUnregisterComponentsInternal = [this, NumComponentsToUnregister]()
+	auto IncrementalUnregisterComponentsInternal = [this, NumComponentsToUnregister, bIncrementalUnregisterComponents]()
 	{
 		// Find next valid actor to process components unregistration
 		int32 NumComponentsUnregistered = 0;
@@ -2182,7 +2182,7 @@ bool ULevel::IncrementalUnregisterComponents(int32 NumComponentsToUnregister, FU
 				Actor->RegisterAllActorTickFunctions(false, false);
 			}
 			CurrentActorIndexForUnregisterComponents++;
-			if (NumComponentsUnregistered > NumComponentsToUnregister)
+			if (bIncrementalUnregisterComponents && NumComponentsUnregistered > NumComponentsToUnregister)
 			{
 				break;
 			}
@@ -5163,4 +5163,4 @@ void ULevel::OnMultipleLSAsPopupDismissed()
 
 #endif // WITH_EDITOR
 
-#undef LOCTEXT_NAMESPACE
\ No newline at end of file
+#undef LOCTEXT_NAMESPACE
diff --git a/Engine/Source/Runtime/Engine/Private/World.cpp b/Engine/Source/Runtime/Engine/Private/World.cpp
--- a/Engine/Source/Runtime/Engine/Private/World.cpp
+++ b/Engine/Source/Runtime/Engine/Private/World.cpp
@@ -4170,7 +4170,7 @@ void UWorld::RemoveFromWorld(ULevel* Level, bool bAllowIncrementalRemoval, const
 		{
 			// Finish current level pending invisibility without time limit
 			check(Level->bIsBeingRemoved);
-			while (!Level->IncrementalUnregisterComponents(MAX_int32)) {}
+			while (!Level->IncrementalUnregisterComponents(0)) {}
 			bWaitForCompletion = true;
 			check(bFinishRemovingLevel);
 		}
