Seamless Travel:  fix for UDN issue where cached MIDs generated by post process materials go invalid across level loads via seamless travel, leading to a crash.  CleanupViewState call added to FSeamlessTravelHandler::Tick on load, mirroring the call in UEngine::LoadMap, which handles this for non-seamless level loading.  Cached MIDs are just a minor optimization, and the cleanup call is cheap and harmless, even if occasionally called unnecessarily.

#rnx

#Synced-CL 49564039

diff --git a/Engine/Source/Runtime/Engine/Private/World.cpp b/Engine/Source/Runtime/Engine/Private/World.cpp
--- a/Engine/Source/Runtime/Engine/Private/World.cpp
+++ b/Engine/Source/Runtime/Engine/Private/World.cpp
@@ -8981,6 +8981,7 @@ UWorld* FSeamlessTravelHandler::Tick()
 
 					const FName SendLoadedWorldName = It->PlayerController->NetworkRemapPath(LoadedWorldName, false);
 					It->PlayerController->ServerNotifyLoadedWorld(SendLoadedWorldName);
+					It->CleanupViewState();
 #if !UE_BUILD_SHIPPING
 					UE_LOG(LogNet, Verbose, TEXT("NotifyLoadedWorld End"));
 #endif
