Iris support for multiple DSTM migrations for the same actor per frame.

* Refactor and clean up of APlayerController::PostMigrate().
* Add deferred start and stop replication for actors and components in Iris.

#Synced-CL 49559917

diff --git a/Engine/Source/Runtime/Engine/Classes/Components/ActorComponent.h b/Engine/Source/Runtime/Engine/Classes/Components/ActorComponent.h
--- a/Engine/Source/Runtime/Engine/Classes/Components/ActorComponent.h
+++ b/Engine/Source/Runtime/Engine/Classes/Components/ActorComponent.h
@@ -405,6 +405,10 @@ private:
 
 #if UE_WITH_REMOTE_OBJECT_HANDLE
 	uint8 bComponentTickEnabled : 1 = false;
+
+	// When set to true the component will remove itself from Iris in EndPlay(). If set to false it's
+	// assumed that the component will be responsible for removing itself from Iris.
+	uint8 bCallStopReplicationInEndPlay : 1 = true;
 #endif
 
 public:
@@ -1433,7 +1437,6 @@ private:
 	friend class FComponentReregisterContextBase;
 	friend class FComponentRecreateRenderStateContext;
 	friend struct FActorComponentTickFunction;
-	friend struct UE::Net::FReplicationSystemUtil;
 #if WITH_EDITOR
 	friend struct SkeletalMeshUCSModifiedDataCleanupHelper;
 #endif
diff --git a/Engine/Source/Runtime/Engine/Classes/Engine/NetConnection.h b/Engine/Source/Runtime/Engine/Classes/Engine/NetConnection.h
--- a/Engine/Source/Runtime/Engine/Classes/Engine/NetConnection.h
+++ b/Engine/Source/Runtime/Engine/Classes/Engine/NetConnection.h
@@ -331,6 +331,15 @@ private:
 	UPROPERTY(transient)
 	TMap<int32, TObjectPtr<class UChildConnection>> ChildrenLookup;
 
+	/** 
+	 * Cache a pointer to an instance of ANoPawnPlayerController to avoid having to re-spawn it when changing this 
+	 * connection from a normal connection with a pawn and player controller to a no pawn connection.
+	 *
+	 * This is currently only used by APlayerController::PostMigrate() when migrating players between servers.
+	 */
+	UPROPERTY(transient)
+	TObjectPtr<class ANoPawnPlayerController> CachedNoPawnPlayerController;
+
 	int32 NextServerNetPlayerIndex = 1;
 
 public:
@@ -1888,6 +1897,11 @@ private:
 	ULevel* RepContextLevel;
 
 	friend struct FScopedRepContext;
+	
+#if UE_WITH_REMOTE_OBJECT_HANDLE
+	// Required by APlayerController::PostMigrate() to read and update CachedNoPawnPlayerController.
+	friend class APlayerController;
+#endif
 
 	bool bAutoFlush;
 
diff --git a/Engine/Source/Runtime/Engine/Classes/GameFramework/Actor.h b/Engine/Source/Runtime/Engine/Classes/GameFramework/Actor.h
--- a/Engine/Source/Runtime/Engine/Classes/GameFramework/Actor.h
+++ b/Engine/Source/Runtime/Engine/Classes/GameFramework/Actor.h
@@ -268,6 +268,13 @@ private:
 	/** Called from the constructor to initialize the class to its default settings */
 	ENGINE_API void InitializeDefaults();
 
+	enum class EAddComponentForReplicationBehavior : uint8
+	{
+		Immediate,
+		Deferred,
+		NeverAdd
+	};
+
 public:
 	/** Returns the properties used for network replication, this needs to be overridden by all actor classes with native replicated properties */
 	ENGINE_API virtual void GetLifetimeReplicatedProps(TArray<FLifetimeProperty>& OutLifetimeProps) const override;
@@ -3522,8 +3529,9 @@ protected:
 	 */
 	ENGINE_API void UpdateOwningNetConnection();
 
+	// Required by both these classes to call OnReplicationStartedForIris().
 	friend struct UE::Net::FReplicationSystemUtil;
-
+	friend class UEngineReplicationBridge;
 
 	/**
 	 * Updates the ReplicatePhysics condition. That information needs to be pushed to the ReplicationSystem.
@@ -4116,6 +4124,12 @@ public:
 	 */
 	ENGINE_API void ResetOwnedComponents();
 
+private:
+	
+	ENGINE_API void ResetOwnedComponents(EAddComponentForReplicationBehavior Type);
+
+public:
+
 	/** Called when the replicated state of a component changes to update the Actor's cached ReplicatedComponents array */
 	ENGINE_API void UpdateReplicatedComponent(UActorComponent* Component);
 
@@ -4312,6 +4326,7 @@ private:
 
 	/** Check if a new component is replicated by the actor and must be registered in the list */
 	ENGINE_API void AddComponentForReplication(UActorComponent* Component);
+	void AddComponentForReplication(UActorComponent* Component, EAddComponentForReplicationBehavior Type);
 
 	/** Remove a component from the replicated list */
 	ENGINE_API void RemoveReplicatedComponent(UActorComponent* Component);
diff --git a/Engine/Source/Runtime/Engine/Private/Actor.cpp b/Engine/Source/Runtime/Engine/Private/Actor.cpp
--- a/Engine/Source/Runtime/Engine/Private/Actor.cpp
+++ b/Engine/Source/Runtime/Engine/Private/Actor.cpp
@@ -478,6 +478,11 @@ bool AActor::CheckActorComponents() const
 }
 
 void AActor::ResetOwnedComponents()
+{
+	ResetOwnedComponents(EAddComponentForReplicationBehavior::Immediate);
+}
+
+void AActor::ResetOwnedComponents(EAddComponentForReplicationBehavior Type)
 {
 #if WITH_EDITOR
 	// Identify any natively-constructed components referenced by properties that either failed to serialize or came in as nullptr.
@@ -515,7 +520,7 @@ void AActor::ResetOwnedComponents()
 	ReplicatedComponents.Reset();
 	ReplicatedComponentsInfo.Reset();
 	
-	ForEachObjectWithOuter(this, [this](UObject* Child)
+	ForEachObjectWithOuter(this, [this, Type](UObject* Child)
 	{
 		UActorComponent* Component = Cast<UActorComponent>(Child);
 		if (Component && Component->GetOwner() == this)
@@ -526,7 +531,7 @@ void AActor::ResetOwnedComponents()
 			{
 				ReplicatedComponents.Add(Component);
 
-				AddComponentForReplication(Component);
+				AddComponentForReplication(Component, Type);
 			}
 		}
 	}, true, RF_NoFlags, EInternalObjectFlags::Garbage);
@@ -1321,11 +1326,6 @@ void AActor::PostMigrate(const FUObjectMigrationContext& MigrationContext)
 			// Migrated actors are unique to the server instance that currently owns their memory and are not replicated to other servers.
 			if (World)
 			{
-				// Notify replication systems that this actor has been destroyed.
-				FStopReplicatingActorParams StopRepParams(EEndPlayReason::Destroyed);
-				StopRepParams.StopReplicatingReason = UE::Net::EStopReplicatingReason::Migrated;
-				UE::Net::FReplicationSystemUtil::StopReplicatingActor(this, StopRepParams);
-
 				UE::Net::Private::FNetPropertyConditionManager::Get().NotifyObjectDestroyed(this);
 
 				if (FWorldContext* Context = GEngine->GetWorldContextFromWorld(World))
@@ -1356,6 +1356,13 @@ void AActor::PostMigrate(const FUObjectMigrationContext& MigrationContext)
 			// Unregister the actor's tick function
 			RegisterAllActorTickFunctions(/*bRegisterTickFunctions*/ false, /*bIncludeComponents*/ true);
 		}
+
+		if (MigrationSendType == EObjectMigrationSendType::ReassignOwnership)
+		{
+			FStopReplicatingActorParams StopRepParams(EEndPlayReason::Destroyed);
+			StopRepParams.StopReplicatingReason = UE::Net::EStopReplicatingReason::Migrated;
+			UE::Net::FDeferredReplicationSystemCalls::StopReplicatingActorDefer(this, StopRepParams);
+		}
 	}
 	else if (MigrationContext.MigrationSide == EObjectMigrationSide::Receive)
 	{
@@ -1363,7 +1370,15 @@ void AActor::PostMigrate(const FUObjectMigrationContext& MigrationContext)
 		UE_LOG(LogActor, Verbose, TEXT("%hs: Recv %s (LocalSrvId: %s. FromSrvId: %s. HasOwnership: %d. MigrationType: %s)"), __func__, *GetName(), *LocalServerId.ToString(), *MigrationContext.RemoteServerId.ToString(), bHasOwnership, ToString(MigrationRecvType));
 
 		// Build up the list of Components we have again (even if we previously owned this object, it could have been mutated)
-		ResetOwnedComponents();
+		EAddComponentForReplicationBehavior AddComponentForRepType = EAddComponentForReplicationBehavior::NeverAdd;
+		
+		// NOTE: This is unnecessary because  UActorComponent::PostMigrate() will explicitly defer start replication for itself.
+		//if (MigrationRecvType == EObjectMigrationRecvType::AssignedOwnership)
+		//{
+		//	AddComponentForRepType = EAddComponentForReplicationBehavior::Deferred;
+		//}
+
+		ResetOwnedComponents(AddComponentForRepType);
 
 		// We can get rid of this check further down the road...
 		// Set the actor's world transform if it has a native rootcomponent.
@@ -1386,7 +1401,7 @@ void AActor::PostMigrate(const FUObjectMigrationContext& MigrationContext)
 		{
 			ensureMsgf(!bActorIsBeingDestroyed, TEXT("We do not want to be assigned ownership of Actors that have already been destroyed. We do not yet know what to do with it."));
 
-			UE::Net::FReplicationSystemUtil::StartReplicatingActor(this);
+			UE::Net::FDeferredReplicationSystemCalls::StartReplicatingActorDefer(this);
 
 			// Each component will be in charge of registering itself in its own PostMigrate, but we must register
 			// the Actor (which we'll do be saying it's currently the opposite of the deserialized value, then flipping it)
diff --git a/Engine/Source/Runtime/Engine/Private/ActorReplication.cpp b/Engine/Source/Runtime/Engine/Private/ActorReplication.cpp
--- a/Engine/Source/Runtime/Engine/Private/ActorReplication.cpp
+++ b/Engine/Source/Runtime/Engine/Private/ActorReplication.cpp
@@ -627,6 +627,11 @@ void AActor::SetReplicatedComponentNetCondition(const UActorComponent* Replicate
 }
 
 void AActor::AddComponentForReplication(UActorComponent* Component)
+{
+	AddComponentForReplication(Component, EAddComponentForReplicationBehavior::Immediate);
+}
+
+void AActor::AddComponentForReplication(UActorComponent* Component, EAddComponentForReplicationBehavior Type)
 {
 	constexpr EObjectFlags CDOFlags = RF_ClassDefaultObject | RF_ArchetypeObject;
 	if (ActorHasBegunPlay == EActorBeginPlayState::HasNotBegunPlay || HasAnyFlags(CDOFlags))
@@ -668,7 +673,14 @@ void AActor::AddComponentForReplication(UActorComponent* Component)
 
     if (NetCondition != COND_Never)
     {
-		UE::Net::FReplicationSystemUtil::StartReplicatingActorComponent(this, Component);
+		if (Type == EAddComponentForReplicationBehavior::Immediate)
+		{
+			UE::Net::FReplicationSystemUtil::StartReplicatingActorComponent(this, Component);
+		}
+		else if (Type == EAddComponentForReplicationBehavior::Deferred)
+		{
+			UE::Net::FDeferredReplicationSystemCalls::StartReplicatingActorComponentDefer(this, Component);
+		}
     }
 }
 
diff --git a/Engine/Source/Runtime/Engine/Private/Components/ActorComponent.cpp b/Engine/Source/Runtime/Engine/Private/Components/ActorComponent.cpp
--- a/Engine/Source/Runtime/Engine/Private/Components/ActorComponent.cpp
+++ b/Engine/Source/Runtime/Engine/Private/Components/ActorComponent.cpp
@@ -755,12 +755,35 @@ void UActorComponent::PostMigrate(const struct FUObjectMigrationContext& Migrati
 			// Prevent tick functions from happening this frame (unregistering just prevents it from happening *next* frame)
 			SetComponentTickEnabled(false);
 		}
+
+		if (SendType == EObjectMigrationSendType::ReassignOwnership)
+		{
+			// It's not necessary to explicitly defer a call to StopReplicatingActorComponent(). It will automatically be
+			// deferred by the deferred call to StopReplicatingActor() for the parent actor.
+			//
+			// NOTE: This should be re-enabled if components ever migrate separately to actors
+			//
+			//UE::Net::FDeferredReplicationSystemCalls::StopReplicatingActorComponentDefer(this, EEndReplicationFlags::Migrated);
+
+			// Prevent calling StopReplicatingActorComponent() in EndPlay(). 
+			bCallStopReplicationInEndPlay = false;
+		}
 	}
 	else if (ensureMsgf(MigrationContext.MigrationSide == EObjectMigrationSide::Receive, TEXT("Invalid MigrationContext: Side was not Send or Receive")))
 	{
 		const EObjectMigrationRecvType RecvType = MigrationContext.GetObjectMigrationRecvType(this);
 		UE_LOG(LogActorComponent, VeryVerbose, TEXT("%hs: %s %s"), __func__, ToString(RecvType), *GetName());
 
+		// It's not necessary to explicitly defer a call to StartReplicatingActorComponent(). It will automatically be
+		// deferred by the deferred call to StartReplicatingActor() for the parent actor.
+		//
+		// NOTE: This should be re-enabled if components ever migrate separately to actors
+		// 
+		//if (RecvType == EObjectMigrationRecvType::AssignedOwnership)
+		//{
+		//	UE::Net::FDeferredReplicationSystemCalls::StartReplicatingActorComponentDefer(GetOwner(), this);
+		//}
+
 		// Add the actor and it's sub-objects to the network id cache since these objects are no longer considered remote and can be derferenced
 		// through the network id cache (e.g. by RPC argument and object deserialization).
 		if (HasBegunPlay())
@@ -1640,9 +1663,14 @@ void UActorComponent::EndPlay(const EEndPlayReason::Type EndPlayReason)
 
 	check(bHasBegunPlay);
 
-	if (EndPlayReason != EEndPlayReason::EndPlayInEditor && EndPlayReason != EEndPlayReason::Quit)
+#if UE_WITH_REMOTE_OBJECT_HANDLE
+	if (bCallStopReplicationInEndPlay)
+#endif
 	{
-		UE::Net::FReplicationSystemUtil::StopReplicatingActorComponent(this);
+		if (EndPlayReason != EEndPlayReason::EndPlayInEditor && EndPlayReason != EEndPlayReason::Quit)
+		{
+			UE::Net::FReplicationSystemUtil::StopReplicatingActorComponent(this);
+		}
 	}
 
 	// If we're in the process of being garbage collected it is unsafe to call out to blueprints
diff --git a/Engine/Source/Runtime/Engine/Private/Net/Iris/ReplicationSystem/EngineReplicationBridge.cpp b/Engine/Source/Runtime/Engine/Private/Net/Iris/ReplicationSystem/EngineReplicationBridge.cpp
--- a/Engine/Source/Runtime/Engine/Private/Net/Iris/ReplicationSystem/EngineReplicationBridge.cpp
+++ b/Engine/Source/Runtime/Engine/Private/Net/Iris/ReplicationSystem/EngineReplicationBridge.cpp
@@ -373,6 +373,8 @@ UE::Net::FNetRefHandle UEngineReplicationBridge::StartReplicatingActor(AActor* A
 	FNetRefHandle ExistingHandle = GetReplicatedRefHandle(Actor);
 	if (ExistingHandle.IsValid())
 	{
+		CancelDeferredActor(Actor);
+	
 		return ExistingHandle;
 	}
 
@@ -517,6 +519,8 @@ UE::Net::FNetRefHandle UEngineReplicationBridge::StartReplicatingActor(AActor* A
 			}
 		}
 	}
+	
+	CancelDeferredActor(Actor);
 
 	return ActorRefHandle;
 }
@@ -592,6 +596,8 @@ UE::Net::FNetRefHandle UEngineReplicationBridge::StartReplicatingComponent(FNetR
 		}
 	}
 
+	CancelDeferredComponent(SubObject);
+
 	return ReplicatedComponentHandle;
 }
 
@@ -603,6 +609,8 @@ UE::Net::FNetRefHandle UEngineReplicationBridge::StartReplicatingSubObject(UObje
 void UEngineReplicationBridge::StopReplicatingActor(AActor* Actor, const FStopReplicatingActorParams& Params)
 {
 	using namespace UE::Net;
+	
+	CancelDeferredActor(Actor);
 
 	FNetRefHandle RefHandle = GetReplicatedRefHandle(Actor, EGetRefHandleFlags::EvenIfGarbage);
 	if (!RefHandle.IsValid())
@@ -670,6 +678,8 @@ EEndReplicationFlags UEngineReplicationBridge::ConvertToEndReplicationFlags(UE::
 void UEngineReplicationBridge::StopReplicatingComponent(UActorComponent* ActorComponent, EEndReplicationFlags EndReplicationFlags)
 {
 	using namespace UE::Net;
+	
+	CancelDeferredComponent(ActorComponent);
 
 	FNetRefHandle ComponentHandle = GetReplicatedRefHandle(ActorComponent, EGetRefHandleFlags::EvenIfGarbage);
 	if (ComponentHandle.IsValid())
@@ -877,6 +887,274 @@ void UEngineReplicationBridge::WakeUpObjectInstantiatedFromRemote(AActor* Actor)
 	}
 }
 
+void UEngineReplicationBridge::StartReplicatingActorDefer(AActor* Instance)
+{
+	if (bIsFlushDeferredRunning)
+	{
+		UE_LOG(LogIris, Warning, TEXT("Cannot call UEngineReplicationBridge::StartReplicatingActorDefer() for %s while flushing."), *GetNameSafe(Instance));
+		return;
+	}
+
+	FDeferredStartStop* Def = DeferredStartStopCallsActors.Find(Instance);
+	if (Def == nullptr)
+	{
+		// If this is the first call for this actor make sure it hasn't already started replication so it can be started.
+		const bool bCanCallStart = !GetReplicatedRefHandle(Instance).IsValid();
+		if (!bCanCallStart)
+		{
+			UE_LOG(LogIris, Warning, TEXT("UEngineReplicationBridge::StartReplicatingActorDefer: First call for %s failed because it's already replicating."), *GetNameSafe(Instance));
+			return;
+		}
+
+		Def = &DeferredStartStopCallsActors.Add(Instance);
+	}
+
+	check(Def);
+
+	const int32 NewCount = (Def->Count + 1);
+	
+	// Prevent deferring multiple, consecutive calls to start replication for the same actor.
+	if (NewCount <= 1)
+	{
+		Def->Count = NewCount;
+		UE_LOG(LogIris, Verbose, TEXT("UEngineReplicationBridge::StartReplicatingActorDefer: %s"), *GetNameSafe(Instance));
+	}
+	else
+	{
+		UE_LOG(LogIris, Warning, TEXT("UEngineReplicationBridge::StartReplicatingActorDefer: Already deferred this call for %s"), *GetNameSafe(Instance));
+	}
+}
+
+void UEngineReplicationBridge::StopReplicatingActorDefer(AActor* Instance, const FStopReplicatingActorParams& Params)
+{
+	if (bIsFlushDeferredRunning)
+	{
+		UE_LOG(LogIris, Warning, TEXT("Cannot call UEngineReplicationBridge::StopReplicatingActorDefer() for %s while flushing."), *GetNameSafe(Instance));
+		return;
+	}
+
+	FDeferredStartStop* Def = DeferredStartStopCallsActors.Find(Instance);
+	if (Def == nullptr)
+	{
+		// If this is the first call for this actor make sure it's already started replication so it can be stopped.
+		const bool bCanCallStop = GetReplicatedRefHandle(Instance).IsValid();
+		if (!bCanCallStop)
+		{
+			UE_CLOG(Instance->GetIsReplicated(), LogIris, Warning, TEXT("UEngineReplicationBridge::StopReplicatingActorDefer: First call for %s failed because it's already not replicating."), *GetNameSafe(Instance));
+			return;
+		}
+
+		Def = &DeferredStartStopCallsActors.Add(Instance);
+	}
+
+	check(Def);
+
+	const int32 NewCount = (Def->Count - 1);
+
+	// Prevent deferring multiple, consecutive calls to stop replication for the same actor.
+	if (NewCount >= -1)
+	{
+		Def->Count = NewCount;
+		Def->StopReplicatingActorParams = Params;
+		UE_LOG(LogIris, Verbose, TEXT("UEngineReplicationBridge::StopReplicatingActorDefer: %s"), *GetNameSafe(Instance));
+	}
+	else
+	{
+		UE_LOG(LogIris, Warning, TEXT("UEngineReplicationBridge::StopReplicatingActorDefer: Already deferred this call for %s"), *GetNameSafe(Instance));
+	}
+}
+
+void UEngineReplicationBridge::StartReplicatingComponentDefer(const UObject* RootObject, UActorComponent* Component)
+{
+	if (bIsFlushDeferredRunning)
+	{
+		UE_LOG(LogIris, Warning, TEXT("Cannot call UEngineReplicationBridge::StartReplicatingComponentDefer() for %s while flushing."), *GetNameSafe(Component));
+		return;
+	}
+
+	FDeferredStartStop* Def = DeferredStartStopCallsComponents.Find(Component);
+	if (Def == nullptr)
+	{
+		// If this is the first call for this component, make sure it hasn't already started replication so it can be started.
+		const bool bCanCallStart = !GetReplicatedRefHandle(Component).IsValid();
+		if (!bCanCallStart)
+		{
+			UE_LOG(LogIris, Warning, TEXT("UEngineReplicationBridge::StartReplicatingComponentDefer: First call for %s failed because it's already replicating."), *GetNameSafe(Component));
+			return;
+		}
+
+		Def = &DeferredStartStopCallsComponents.Add(Component);
+	}
+
+	check(Def);
+
+	const int32 NewCount = (Def->Count + 1);
+
+	// Prevent deferring multiple, consecutive calls to start replication for the same component.
+	if (NewCount <= 1)
+	{
+		Def->Count = NewCount;
+		Def->RootObject = RootObject;
+		UE_LOG(LogIris, Verbose, TEXT("UEngineReplicationBridge::StartReplicatingComponentDefer: %s"), *GetNameSafe(Component));
+	}
+	else
+	{
+		UE_LOG(LogIris, Warning, TEXT("UEngineReplicationBridge::StartReplicatingComponentDefer: Already deferred this call for %s"), *GetNameSafe(Component));
+	}
+}
+
+void UEngineReplicationBridge::StopReplicatingComponentDefer(UActorComponent* Component, EEndReplicationFlags EndReplicationFlags)
+{
+	if (bIsFlushDeferredRunning)
+	{
+		UE_LOG(LogIris, Warning, TEXT("Cannot call UEngineReplicationBridge::StopReplicatingComponentDefer() for %s while flushing."), *GetNameSafe(Component));
+		return;
+	}
+
+	FDeferredStartStop* Def = DeferredStartStopCallsComponents.Find(Component);
+	if (Def == nullptr)
+	{
+		// If this is the first call for this component, make sure it's already started replication so it can be stopped.
+		const bool bCanCallStop = GetReplicatedRefHandle(Component).IsValid();
+		if (!bCanCallStop)
+		{
+			UE_CLOG(Component->GetIsReplicated(), LogIris, Warning, TEXT("UEngineReplicationBridge::StopReplicatingComponentDefer: First call for %s failed because it's already not replicating."), *GetNameSafe(Component));
+			return;
+		}
+
+		Def = &DeferredStartStopCallsComponents.Add(Component);
+	}
+
+	check(Def);
+
+	const int32 NewCount = (Def->Count - 1);
+
+	// Prevent deferring multiple, consecutive calls to stop replication for the same component.
+	if (NewCount >= -1)
+	{
+		Def->Count = NewCount;
+		Def->StopReplicatingComponentFlags = EndReplicationFlags;
+		UE_LOG(LogIris, Verbose, TEXT("UEngineReplicationBridge::StopReplicatingComponentDefer: %s"), *GetNameSafe(Component));
+	}
+	else
+	{
+		UE_LOG(LogIris, Warning, TEXT("UEngineReplicationBridge::StopReplicatingComponentDefer: Already deferred this call for %s"), *GetNameSafe(Component));
+	}
+}
+
+void UEngineReplicationBridge::FlushDeferred()
+{
+	bIsFlushDeferredRunning = true;
+
+	if (DeferredStartStopCallsActors.Num() > 0 || DeferredStartStopCallsComponents.Num() > 0)
+	{
+		UE_LOG(LogIris, Verbose, TEXT("UEngineReplicationBridge::FlushDeferred: Actors=%d Components=%d"), DeferredStartStopCallsActors.Num(), DeferredStartStopCallsComponents.Num());
+	}
+
+	// Actor replication needs to be processed before component replication in case the component depends
+	// on an owning actor that hasn't started replicating yet.
+
+	for (TPair<TObjectKey<AActor>, FDeferredStartStop>& Pair : DeferredStartStopCallsActors)
+	{
+		TObjectKey ActorKey = Pair.Key;
+		FDeferredStartStop& Def = Pair.Value;
+
+		// No start or stop was called, or the number of start and stop calls cancel out (i.e. do nothing).
+		if (Def.Count == 0)
+		{
+			continue;
+		}
+
+		AActor* Actor = ActorKey.ResolveObjectPtr();
+		if (!Actor)
+		{
+			// The actor no longer exists.
+			continue;
+		}
+
+		// Call UEngineReplicationBridge::StartReplicatingActor()
+		if (Def.Count > 0)
+		{
+			UE::Net::FNetRefHandle NewNetRefHandle = StartReplicatingActor(Actor);
+			if (NewNetRefHandle.IsValid())
+			{
+				Actor->OnReplicationStartedForIris(AActor::FOnReplicationStartedParams());
+			}
+				
+			UE_LOG(LogIris, Verbose, TEXT("UEngineReplicationBridge::FlushDeferred: Called StartReplicatingActor for %s (success=%d)"), *Actor->GetName(), NewNetRefHandle.IsValid());
+		}
+				
+		// Call UEngineReplicationBridge::StopReplicatingActor()
+		else if (Def.Count < 0)
+		{
+			// Cache locally in case the DeferredStartStopCallsActors map is changed as a side-effect of the remaining logic.
+			FStopReplicatingActorParams StopReplicatingActorParams = Def.StopReplicatingActorParams;
+
+			AActor::FOnStopReplicationParams OnStopReplicationParams = { .StopReplicatingParams = StopReplicatingActorParams };
+			Actor->OnStopReplicationForIris(OnStopReplicationParams);
+			StopReplicatingActor(Actor, StopReplicatingActorParams);
+
+			UE_LOG(LogIris, Verbose, TEXT("UEngineReplicationBridge::FlushDeferred: Called StopReplicatingActor for %s"), *Actor->GetName());
+		}
+	}
+
+	for (TPair<TObjectKey<UActorComponent>, FDeferredStartStop>& Pair : DeferredStartStopCallsComponents)
+	{
+		TObjectKey<UActorComponent> ActorComponentKey = Pair.Key;
+		FDeferredStartStop& Def = Pair.Value;
+
+		// No start or stop was called, or the number of start and stop calls cancel out (i.e. do nothing).
+		if (Def.Count == 0)
+		{
+			continue;
+		}
+
+		UActorComponent* ActorComponent = ActorComponentKey.ResolveObjectPtr();
+		if (!ActorComponent)
+		{
+			// The component no longer exists.
+			continue;
+		}
+
+		// Call UEngineReplicationBridge::StartReplicatingComponent()
+		if (Def.Count > 0)
+		{
+			UObject* RootObject = Def.RootObject.ResolveObjectPtr();
+			if (RootObject)
+			{
+				FNetRefHandle RootObjectHandle = GetReplicatedRefHandle(RootObject);
+				if (!RootObjectHandle.IsValid())
+				{
+					UE_LOG(LogIris, Verbose, TEXT("UEngineReplicationBridge::FlushDeferred: Cannot find handle for component's root object %s"), *RootObject->GetName());
+					continue;
+				}
+
+				UE::Net::FNetRefHandle NewNetRefHandle = StartReplicatingComponent(RootObjectHandle, ActorComponent);
+				if (NewNetRefHandle.IsValid())
+				{
+					ActorComponent->OnReplicationStartedForIris(UActorComponent::FOnReplicationStartedParams());
+				}
+
+				UE_LOG(LogIris, Verbose, TEXT("UEngineReplicationBridge::FlushDeferred: Called StartReplicatingComponent for %s (success=%d)"), *ActorComponent->GetName(), NewNetRefHandle.IsValid());
+			}
+		}
+
+		// Call UEngineReplicationBridge::StopReplicatingComponent()
+		else if (Def.Count < 0)
+		{
+			ActorComponent->OnStopReplicationForIris(UActorComponent::FOnStopReplicationParams());
+			StopReplicatingComponent(ActorComponent, Def.StopReplicatingComponentFlags);
+
+			UE_LOG(LogIris, Verbose, TEXT("UEngineReplicationBridge::FlushDeferred: Called StopReplicatingComponent for %s"), *ActorComponent->GetName());
+		}
+	}
+
+	bIsFlushDeferredRunning = false;
+
+	DeferredStartStopCallsActors.Reset();
+	DeferredStartStopCallsComponents.Reset();
+}
+
 void UEngineReplicationBridge::OnProtocolMismatchDetected(FNetRefHandle ObjectHandle)
 {
 	Super::OnProtocolMismatchDetected(ObjectHandle);
@@ -1144,3 +1422,25 @@ bool UEngineReplicationBridge::AllowReplicationOfActorsWithAnyRole() const
 {
 	return bAllowReplicationOfActorsWithAnyRole;
 }
+
+void UEngineReplicationBridge::CancelDeferredActor(AActor* Actor)
+{
+	if (bIsFlushDeferredRunning)
+	{
+		return;
+	}
+
+	UE_CLOG(DeferredStartStopCallsActors.Contains(Actor), LogIris, Verbose, TEXT("UEngineReplicationBridge::CancelDeferredActor: Cancelling deferred %s."), *GetNameSafe(Actor));
+	DeferredStartStopCallsActors.Remove(Actor);
+}
+
+void UEngineReplicationBridge::CancelDeferredComponent(UActorComponent* ActorComponent)
+{
+	if (bIsFlushDeferredRunning)
+	{
+		return;
+	}
+
+	UE_CLOG(DeferredStartStopCallsComponents.Contains(ActorComponent), LogIris, Verbose, TEXT("UEngineReplicationBridge::CancelDeferredComponent: Cancelling deferred %s."), *GetNameSafe(ActorComponent));
+	DeferredStartStopCallsComponents.Remove(ActorComponent);
+}
\ No newline at end of file
diff --git a/Engine/Source/Runtime/Engine/Private/Net/Iris/ReplicationSystem/ReplicationSystemUtil.cpp b/Engine/Source/Runtime/Engine/Private/Net/Iris/ReplicationSystem/ReplicationSystemUtil.cpp
--- a/Engine/Source/Runtime/Engine/Private/Net/Iris/ReplicationSystem/ReplicationSystemUtil.cpp
+++ b/Engine/Source/Runtime/Engine/Private/Net/Iris/ReplicationSystem/ReplicationSystemUtil.cpp
@@ -910,5 +910,78 @@ void FReplicationSystemUtil::SetPollFrequency(const UObject* Object, float PollF
 	});
 }
 
+void FDeferredReplicationSystemCalls::StartReplicatingActorDefer(AActor* Actor)
+{
+	if (UWorld* World = Actor->GetWorld())
+	{
+		if (UNetDriver* NetDriver = World->GetNetDriver())
+		{
+			if (UReplicationSystem* ReplicationSystem = NetDriver->GetReplicationSystem())
+			{
+				if (UEngineReplicationBridge* Bridge = ReplicationSystem->GetReplicationBridgeAs<UEngineReplicationBridge>())
+				{
+					Bridge->StartReplicatingActorDefer(Actor);
+				}
+			}
+		}
+	}
+}
+
+void FDeferredReplicationSystemCalls::StopReplicatingActorDefer(AActor* Actor, const FStopReplicatingActorParams& StopReplicatingActorParams)
+{
+	if (UWorld* World = Actor->GetWorld())
+	{
+		if (UNetDriver* NetDriver = World->GetNetDriver())
+		{
+			if (UReplicationSystem* ReplicationSystem = NetDriver->GetReplicationSystem())
+			{
+				if (UEngineReplicationBridge* Bridge = ReplicationSystem->GetReplicationBridgeAs<UEngineReplicationBridge>())
+				{
+					Bridge->StopReplicatingActorDefer(Actor, StopReplicatingActorParams);
+				}
+			}
+		}
+	}
+}
+
+void FDeferredReplicationSystemCalls::StartReplicatingActorComponentDefer(const AActor* Actor, UActorComponent* ActorComponent)
+{
+	if (UWorld* World = ActorComponent->GetWorld())
+	{
+		if (UNetDriver* NetDriver = World->GetNetDriver())
+		{
+			if (UReplicationSystem* ReplicationSystem = NetDriver->GetReplicationSystem())
+			{
+				if (UEngineReplicationBridge* Bridge = ReplicationSystem->GetReplicationBridgeAs<UEngineReplicationBridge>())
+				{
+					Bridge->StartReplicatingComponentDefer(Actor, ActorComponent);
+				}
+			}
+		}
+	}
+}
+
+void FDeferredReplicationSystemCalls::StopReplicatingActorComponentDefer(UActorComponent* ActorComponent, EEndReplicationFlags EndReplicationFlags)
+{
+	if (UWorld* World = ActorComponent->GetWorld())
+	{
+		if (UNetDriver* NetDriver = World->GetNetDriver())
+		{
+			if (UReplicationSystem* ReplicationSystem = NetDriver->GetReplicationSystem())
+			{
+				if (UEngineReplicationBridge* Bridge = ReplicationSystem->GetReplicationBridgeAs<UEngineReplicationBridge>())
+				{
+					Bridge->StopReplicatingComponentDefer(ActorComponent, EndReplicationFlags);
+				}
+			}
+		}
+	}
+}
+
+void FDeferredReplicationSystemCalls::FlushDeferred(UEngineReplicationBridge* Bridge)
+{
+	Bridge->FlushDeferred();
+}
+
 } // end namespace UE::Net
 
diff --git a/Engine/Source/Runtime/Engine/Private/NetConnection.cpp b/Engine/Source/Runtime/Engine/Private/NetConnection.cpp
--- a/Engine/Source/Runtime/Engine/Private/NetConnection.cpp
+++ b/Engine/Source/Runtime/Engine/Private/NetConnection.cpp
@@ -1484,6 +1484,12 @@ void UNetConnection::CleanUp()
 	InTraceCollector = nullptr;
 	OutTraceCollector = nullptr;
 #endif
+
+	if (CachedNoPawnPlayerController)
+	{
+		CachedNoPawnPlayerController->Destroy();
+		CachedNoPawnPlayerController = nullptr;
+	}
 }
 
 void UNetConnection::DestroyOwningActor()
diff --git a/Engine/Source/Runtime/Engine/Private/NetDriver.cpp b/Engine/Source/Runtime/Engine/Private/NetDriver.cpp
--- a/Engine/Source/Runtime/Engine/Private/NetDriver.cpp
+++ b/Engine/Source/Runtime/Engine/Private/NetDriver.cpp
@@ -1185,6 +1185,26 @@ void UNetDriver::TickFlush(float DeltaSeconds)
 
 	if (IsServer() && (ClientConnections.Num() > 0 || UE::Net::Private::bUpdateReplicationSystemWithNoConnections) && !bSkipServerReplicateActors)
 	{
+		if (ReplicationSystem)
+		{
+			if (UEngineReplicationBridge* Bridge = ReplicationSystem->GetReplicationBridgeAs<UEngineReplicationBridge>())
+			{
+#if UE_WITH_REMOTE_OBJECT_HANDLE && UE_AUTORTFM
+				if (bReplicateTransactionally)
+				{
+					static const FName FlushDeferredName("UEngineReplicationBridge::FlushDeferred");
+					UE::RemoteExecutor::ExecuteTransactional(FlushDeferredName, [Bridge]() {
+						UE::Net::FDeferredReplicationSystemCalls::FlushDeferred(Bridge);
+					});
+				}
+				else
+#endif
+				{
+					UE::Net::FDeferredReplicationSystemCalls::FlushDeferred(Bridge);
+				}
+			}
+		}
+
 		// Update all clients.
 #if WITH_SERVER_CODE
 		CSV_SCOPED_TIMING_STAT_EXCLUSIVE(ServerReplicateActors);
diff --git a/Engine/Source/Runtime/Engine/Private/PlayerController.cpp b/Engine/Source/Runtime/Engine/Private/PlayerController.cpp
--- a/Engine/Source/Runtime/Engine/Private/PlayerController.cpp
+++ b/Engine/Source/Runtime/Engine/Private/PlayerController.cpp
@@ -5077,44 +5077,46 @@ static void UpdateNetConnectionControllerPosition(UNetConnection* NetConnection)
 
 void APlayerController::PostMigrate(const FUObjectMigrationContext& MigrationContext)
 {
-	// The ANoPawnPlayerController is server specific and should never be migrated between servers.
-	ensure(!IsA(ANoPawnPlayerController::StaticClass()));
-
-	auto GetNetPlayerIndex = [](UNetConnection* Connection)
+	auto RemoveFromReplicationSystem = [](UNetDriver* NetDriver, AActor* Actor)
 	{
-		if (UChildConnection* ChildConnection = Cast<UChildConnection>(Connection))
+		if (NetDriver->GetReplicationSystem())
 		{
-			if (UNetConnection* ParentConnection = ChildConnection->GetParentConnection())
-			{
-				int32 NetPlayerIndex = 0;
-				const bool bFoundNetPlayerIndex = ParentConnection->GetNetPlayerIndex(ChildConnection, &NetPlayerIndex);
-				ensure(bFoundNetPlayerIndex);
-				return NetPlayerIndex;
-			}
+			FStopReplicatingActorParams StopRepParams(EEndPlayReason::Destroyed);
+			StopRepParams.StopReplicatingReason = UE::Net::EStopReplicatingReason::Migrated;
+			UE::Net::FDeferredReplicationSystemCalls::StopReplicatingActorDefer(Actor, StopRepParams);
+		}
+		else
+		{
+			NetDriver->RemoveNetworkActor(Actor);
 		}
-
-		return 0;
 	};
 
-	auto GetConnectionWithPlayerController = [](UNetDriver* Driver, APlayerController* PlayerController)
+	auto AddToReplicationSystem = [](UNetDriver* NetDriver, AActor* Actor)
 	{
-		UNetConnection* OutConnection = nullptr;
-		Driver->ForEachClientConnection([&OutConnection, PlayerController](UNetConnection* Connection)
+		if (NetDriver->GetReplicationSystem())
 		{
-			if (Connection->PlayerController == PlayerController)
-			{
-				OutConnection = Connection;
-				return true;
-			}
+			UE::Net::FDeferredReplicationSystemCalls::StartReplicatingActorDefer(Actor);
+		}
+		else
+		{
+			NetDriver->AddNetworkActor(Actor);
+		}
+	};
 
-			return false;
-		});
+	// The ANoPawnPlayerController is server specific and should never be migrated between servers.
+	ensure(!IsA(ANoPawnPlayerController::StaticClass()));
+
+	// APlayerController::NetConnection is set to null when serialized between servers so it should be null here. This function
+	// is expected to set a new value of NetConnection if and when required.
+	ensure(NetConnection == nullptr);
 
-		return OutConnection;
-	};
-	
 	Super::PostMigrate(MigrationContext);
 	
+	if (!ensure(GetWorld()) || !ensure(GetWorld()->GetNetDriver()))
+	{
+		return;
+	}
+	
 	const bool bHasOwnership = UE::RemoteObject::Handle::IsOwned(this);
 	
 	if (MigrationContext.MigrationSide == EObjectMigrationSide::Receive)
@@ -5124,118 +5126,199 @@ void APlayerController::PostMigrate(const FUObjectMigrationContext& MigrationCon
 		UE_LOG(LogPlayerController, Verbose, TEXT("[%s] APlayerController::PostMigrate [%s]: %s bHasOwnership=%d MigrationSide=Receive RecvType=%s PlayerId=%s"), 
 			   *FRemoteServerId::GetLocalServerId().ToString(), *FRemoteObjectId(this).ToString(), *GetName(), bHasOwnership, ToString(MigrationRecvType), *CachedConnectionPlayerId.ToDebugString());
 
-		ensure(NetConnection == nullptr);
+		const bool bAttachControllerToConnection = (MigrationRecvType == EObjectMigrationRecvType::AssignedOwnership);
 
-		const bool bAttachControllerToConnection = true;
-		const bool bRegisterWithWorld = (MigrationRecvType == EObjectMigrationRecvType::AssignedOwnership);
+		// Make sure to update NetConnection of a player controller that was loaned and is now returning, because the pointer will 
+		// be set to null when serialized.
+		const bool bUpdateNetConnection = (bAttachControllerToConnection || MigrationRecvType == EObjectMigrationRecvType::ReturnedLoan);
 
-		if (bAttachControllerToConnection)
+		if (bUpdateNetConnection)
 		{
+			// Set NetConnection to the local connection that matches the player id in the migrated player controller.
 			NetConnection = GetWorld()->GetNetDriver()->FindConnection(CachedConnectionPlayerId);
-			if (ensure(NetConnection))
+			
+			if (!ensureMsgf(NetConnection, TEXT("APlayerController::PostMigrate: Failed to find a NetConnection associated with player id %s"), *CachedConnectionPlayerId.ToDebugString()))
 			{
-				ensure(NetConnection->PlayerController);
-				ensure(NetConnection->PlayerController->IsA(ANoPawnPlayerController::StaticClass()));
-				
-				// The ANoPawnPlayerController actor has a low NetUpdateFrequency, so ForceNetUpdate() must be called because
-				// it's going to be detached from a connection.
-				NetConnection->PlayerController->ForceNetUpdate();
-				
-				// Detach the previous player controller from the network connection so that any attempts to send an RPC
-				// on that controller to the client will not succeed.
-				NetConnection->PlayerController->NetConnection = nullptr;
-				
-				const int32 PrevNetPlayerIndex = NetConnection->PlayerController->NetPlayerIndex;
-				const int32 PrevHandshakeTransactionId = NetConnection->PlayerController->GetClientHandshakeId();
-
-				// Destroy the previous player controller used by the connection.
-				UE_LOG(LogPlayerController, Verbose, TEXT("[%s] APLayerController::PostMigrate [%s]: Destroying previous controller %s"), 
-					   *FRemoteServerId::GetLocalServerId().ToString(), *FRemoteObjectId(this).ToString(), *GetNameSafe(NetConnection->PlayerController));
-				GetWorld()->RemoveController(NetConnection->PlayerController);
-				GetWorld()->DestroyActor(NetConnection->PlayerController, true);
-				
+				return;
+			}
+
+			// Set Player at the same time as NetConnection because it's used in APlayerController::GetConnection().
+			Player = NetConnection;
+		}
+
+		if (bAttachControllerToConnection)
+		{	
+			// It's assumed the current connection for the player id should be using an ANoPawnPlayerController.
+			ensure(NetConnection->PlayerController);
+			ensure(NetConnection->PlayerController.IsA<ANoPawnPlayerController>());
+			
+			ANoPawnPlayerController* PrevNoPawnPlayerController = Cast<ANoPawnPlayerController>(NetConnection->PlayerController);
+			if (PrevNoPawnPlayerController)
+			{
+				// Detach the network connection and its player controller.
+				NetConnection->PlayerController = nullptr;
+				NetConnection->OwningActor = nullptr;
+				PrevNoPawnPlayerController->NetConnection = nullptr;
+				PrevNoPawnPlayerController->Player = nullptr;
+
+				// Update the cached ANoPawnPlayerController to account for cases where this connection was established with an
+				// instance of ANoPawnPlayerController outside of APlayerController::PostMigrate() (e.g. a connection was opened
+				// with NMT_JoinNoPawn or NMT_JoinNoPawnSplit). Without this code we might lose reference to an instance of 
+				// ANoPawnPlayerController and create a controller leak.
+				NetConnection->CachedNoPawnPlayerController = PrevNoPawnPlayerController;
+
+				// Attach this migrated player controller to the network connection.
 				NetConnection->PlayerController = this;
 				NetConnection->OwningActor = this;
-				Player = NetConnection;
-				NetPlayerIndex = GetNetPlayerIndex(NetConnection);
+
+				// Update the values of this migrated player controller that are associated with the connection.
+				NetPlayerIndex = PrevNoPawnPlayerController->NetPlayerIndex;
 				SetClientHandshakeId(NetConnection->GetClientHandshakeId());
-				ensure(PrevNetPlayerIndex == NetPlayerIndex);
-				ensure(PrevHandshakeTransactionId == GetClientHandshakeId());
+				LocalPlayerConnectionIdentifier = PrevNoPawnPlayerController->LocalPlayerConnectionIdentifier;
+				CachedConnectionPlayerId = PrevNoPawnPlayerController->CachedConnectionPlayerId;
+
+				// Stop replicating the old ANoPawnPlayerController since it's no longer used.
+				RemoveFromReplicationSystem(GetWorld()->GetNetDriver(), PrevNoPawnPlayerController);
 
-				// Update Iris with the new Player/NetConnection
 				UpdateOwningNetConnection();
 
-				UE_LOG(LogPlayerController, Verbose, TEXT("[%s] APlayerController::PostMigrate [%s]: Binding %s to connection %s"), 
-					   *FRemoteServerId::GetLocalServerId().ToString(), *FRemoteObjectId(this).ToString(), 
+				UE_LOG(LogPlayerController, Verbose, TEXT("[%s] APlayerController::PostMigrate [%s]: Binding %s to connection %s"),
+					   *FRemoteServerId::GetLocalServerId().ToString(), *FRemoteObjectId(this).ToString(),
 					   *GetName(), *NetConnection->GetName());
-			}
-		}
 
-		if (bRegisterWithWorld)
-		{
-			GetWorld()->AddController(this);
+				// There should only be one registered controller in the world for this connection and associated player.
+				GetWorld()->RemoveController(PrevNoPawnPlayerController);
+				GetWorld()->AddController(this);
+			}
 		}
 	}
 	else if (MigrationContext.MigrationSide == EObjectMigrationSide::Send)
 	{
 		EObjectMigrationSendType MigrationSendType = MigrationContext.GetObjectMigrationSendType(this);
 
-		UE_LOG(LogPlayerController, Verbose, TEXT("[%s] APlayerController::PostMigrate [%s]: %s bHasOwnership=%d MigrationSide=Send SendType=%s PlayerId=%s"), 
+		UE_LOG(LogPlayerController, Verbose, TEXT("[%s] APlayerController::PostMigrate [%s]: %s bHasOwnership=%d MigrationSide=Send SendType=%s PlayerId=%s"),
 			   *FRemoteServerId::GetLocalServerId().ToString(), *FRemoteObjectId(this).ToString(), *GetName(), bHasOwnership, ToString(MigrationSendType), *CachedConnectionPlayerId.ToDebugString());
 
-		const bool bDetachControllerFromConnection = true;
-		const bool bDeregisterFromWorld = (MigrationSendType == EObjectMigrationSendType::ReassignOwnership);
+		const bool bDetachControllerToConnection = (MigrationSendType == EObjectMigrationSendType::ReassignOwnership);
 
-		if (bDetachControllerFromConnection)
+		if (bDetachControllerToConnection)
 		{
-			UNetConnection* PrevNetConnection = GetConnectionWithPlayerController(GetWorld()->GetNetDriver(), this);
-			if (ensure(PrevNetConnection))
+			// Lookup the local connection that matches the player id in the player controller being migrated away.
+			UNetConnection* PrevNetConnection = GetWorld()->GetNetDriver()->FindConnection(CachedConnectionPlayerId);
+		
+			if (!ensureMsgf(PrevNetConnection, TEXT("APlayerController::PostMigrate: Failed to find a connection associated with player id %s"), *CachedConnectionPlayerId.ToDebugString()))
 			{
-				ensure(PrevNetConnection->PlayerController == this);
-				ensure(PrevNetConnection->PlayerController);
-				ensure(PrevNetConnection->OwningActor);
+				return;
+			}
+
+			if (!ensureMsgf(PrevNetConnection->PlayerController == this, TEXT("APlayerController::PostMigrate: Connection doesn't match player controller being migrated away.")))
+			{
+				return;
+			}
 
+			if (ensure(PrevNetConnection->PlayerController))
+			{
+				// It's assumed this connection has a player controller associated that is not an instance of ANoPawnPlayerController.
+				ensure(!PrevNetConnection->PlayerController.IsA<ANoPawnPlayerController>());
+				ensure(PrevNetConnection->PlayerController != PrevNetConnection->CachedNoPawnPlayerController);
+
+				// Detach the network connection and its player controller.
+				PrevNetConnection->PlayerController->NetConnection = nullptr;
+				PrevNetConnection->PlayerController->Player = nullptr;
 				PrevNetConnection->PlayerController = nullptr;
 				PrevNetConnection->OwningActor = nullptr;
 
-				// Create a new instance of ANoPawnPlayerController for the connection to use while the original instance of APlayerController is migrated
-				// on another server. This controller will be destroyed when it's replaced by APlayerController when it returns to the server.
-				FActorSpawnParameters SpawnInfo;
-				SpawnInfo.ObjectFlags |= RF_Transient;
-				SpawnInfo.bDeferConstruction = true;
-				PrevNetConnection->PlayerController = GetWorld()->SpawnActor<APlayerController>(ANoPawnPlayerController::StaticClass(), FVector(), FRotator(), SpawnInfo);
-				PrevNetConnection->PlayerController->SetRole(ROLE_Authority);
-				PrevNetConnection->PlayerController->SetReplicates(true);
-				PrevNetConnection->PlayerController->SetAutonomousProxy(true);
-				// Set the NetConnection before FinishSpawning so that the AActor::UpdateOwningNetConnection call properly updates the Iris conection owner.
-				PrevNetConnection->PlayerController->NetConnection = PrevNetConnection;
-				PrevNetConnection->PlayerController->FinishSpawning(FTransform());
-				PrevNetConnection->PlayerController->CachedConnectionPlayerId = CachedConnectionPlayerId;
-				PrevNetConnection->PlayerController->NetPlayerIndex = NetPlayerIndex;
-				PrevNetConnection->PlayerController->SetClientHandshakeId(PrevNetConnection->GetClientHandshakeId());
-
-				PrevNetConnection->OwningActor = PrevNetConnection->PlayerController;
+				auto BindAndInitializeConnectionAndController = [](UNetConnection* Connection, ANoPawnPlayerController* NoPawnPlayerController, const APlayerController* OldPlayerController)
+				{
+					// Bind the connection with the ANoPawnPlayerController.
+					Connection->PlayerController = NoPawnPlayerController;
+					Connection->OwningActor = NoPawnPlayerController;
+					NoPawnPlayerController->NetConnection = Connection;
+					NoPawnPlayerController->Player = Connection;
+
+					// Make sure that the ANoPawnPlayerController has the correct connection related state.
+					NoPawnPlayerController->CachedConnectionPlayerId = OldPlayerController->CachedConnectionPlayerId;
+					NoPawnPlayerController->NetPlayerIndex = OldPlayerController->NetPlayerIndex;
+					NoPawnPlayerController->SetClientHandshakeId(Connection->GetClientHandshakeId());
+					NoPawnPlayerController->LocalPlayerConnectionIdentifier = OldPlayerController->LocalPlayerConnectionIdentifier;
+				};
+
+				// If the cached no pawn player controller was created in a different world, destroy it so that it will be re-created. This
+				// can happen during seamless travel and the world is destroyed but the connection remains open.
+				if (PrevNetConnection->CachedNoPawnPlayerController)
+				{
+					if (PrevNetConnection->CachedNoPawnPlayerController->GetWorld() != GetWorld())
+					{
+						UE_LOG(LogPlayerController, Verbose, TEXT("APlayerController::PostMigrate(): Destroying cached instance ANoPawnPlayerController because it was spawned in a different world."));
+
+						PrevNetConnection->CachedNoPawnPlayerController->Destroy();
+						PrevNetConnection->CachedNoPawnPlayerController = nullptr;
+					}
+				}
+
+				// If there is no cached instances of ANoPawnPlayerController, make sure one is spawned.
+				if (PrevNetConnection->CachedNoPawnPlayerController == nullptr)
+				{
+					UE_LOG(LogPlayerController, Verbose, TEXT("APlayerController::PostMigrate(): Spawning a new instance of ANoPawnPlayerController for connection %s"), *PrevNetConnection->GetName());
+
+					FActorSpawnParameters SpawnInfo;
+					SpawnInfo.ObjectFlags |= RF_Transient;
+					SpawnInfo.bDeferConstruction = true;
+					ANoPawnPlayerController* NewNoPawnPlayerController = GetWorld()->SpawnActor<ANoPawnPlayerController>(ANoPawnPlayerController::StaticClass(), FVector(), FRotator(), SpawnInfo);
+					if (ensure(NewNoPawnPlayerController))
+					{
+						// Attach this new instance of ANoPawnPlayerController to the network connection.
+						BindAndInitializeConnectionAndController(PrevNetConnection, NewNoPawnPlayerController, this);
+
+						// Initial configuration for the ANoPawnPlayerController.
+						NewNoPawnPlayerController->SetRole(ROLE_Authority);
+						NewNoPawnPlayerController->SetReplicates(true);
+						NewNoPawnPlayerController->SetAutonomousProxy(true);
+
+						// Finalizing spawning of the controller.
+						// When using Iris this will call AActor::UpdateOwningNetConnection() and StartReplication() (immediately and not deferred).
+						NewNoPawnPlayerController->FinishSpawning(FTransform());
+
+						// Cache the instances of ANoPawnPlayerController to avoid respawning it on player migration.
+						PrevNetConnection->CachedNoPawnPlayerController = NewNoPawnPlayerController;
+					}
+				}
+				else
+				{
+					ensure(PrevNetConnection->CachedNoPawnPlayerController->IsA<ANoPawnPlayerController>());
 				
-				// Update the position of NoPawnPlayerController during migration, so it will be correct when used for relevancy check. 
-				UpdateNetConnectionControllerPosition(PrevNetConnection);
+					ANoPawnPlayerController* CachedNoPawnPlayerController = PrevNetConnection->CachedNoPawnPlayerController;
 
-				UE_LOG(LogPlayerController, Verbose, TEXT("[%s] APlayerController::PostMigrate [%s]: Replacing %s with new player controller %s on connection %s"), 
-					*FRemoteServerId::GetLocalServerId().ToString(), *FRemoteObjectId(this).ToString(), 
-					*GetName(), *PrevNetConnection->PlayerController->GetName(), *PrevNetConnection->GetName());
+					// Attach this cached instance of ANoPawnPlayerController to the network connection.
+					BindAndInitializeConnectionAndController(PrevNetConnection, CachedNoPawnPlayerController, this);
+				
+					// Start replicating the ANoPawnPlayerController.
+					AddToReplicationSystem(GetWorld()->GetNetDriver(), CachedNoPawnPlayerController);
+
+					// Since the player controller is being re-used, ensure that its initial state is the same as if it were just spawned.
+					//
+					// NOTE: This currently commented out because PostMigrate() can be called while Iris is replicating and Iris will ensure
+					//		 when calling AActor::UpdateOwningNetConnection() during replication. This seems to be okay for now since
+					//		 none of the state from ANoPawnPlayerController is changing between migrations.
+					//CachedNoPawnPlayerController->SetRole(ROLE_Authority);
+					//CachedNoPawnPlayerController->SetReplicates(true); // Calls AActor::UpdateOwningNetConnection().
+					//CachedNoPawnPlayerController->SetAutonomousProxy(true);
+				}
 
-				GetWorld()->AddController(PrevNetConnection->PlayerController);
+				// Update the position of the ANoPawnPlayerController during migration, so it will be correct when used for relevancy check. 
+				UpdateNetConnectionControllerPosition(PrevNetConnection);
+
+				UE_LOG(LogPlayerController, Verbose, TEXT("[%s] APlayerController::PostMigrate [%s]: Replacing %s with new player controller %s on connection %s"),
+					   *FRemoteServerId::GetLocalServerId().ToString(), *FRemoteObjectId(this).ToString(),
+					   *GetName(), *PrevNetConnection->PlayerController->GetName(), *PrevNetConnection->GetName());
 
 				ensure(PrevNetConnection->PlayerController);
 				ensure(PrevNetConnection->OwningActor);
-			}
-		}
-
-		if (bDeregisterFromWorld)
-		{
-			UE_LOG(LogPlayerController, Verbose, TEXT("[%s] APlayerController::PostMigrate [%s]: Removing %s from the world"), 
-				   *FRemoteServerId::GetLocalServerId().ToString(), *FRemoteObjectId(this).ToString(), *GetName());
 
-			GetWorld()->RemoveController(this);
+				// There should only be one registered controller in the world for this connection and associated player.
+				GetWorld()->RemoveController(this);
+				GetWorld()->AddController(PrevNetConnection->PlayerController);
+			}
 		}
 	}
 }
diff --git a/Engine/Source/Runtime/Engine/Public/Net/Iris/ReplicationSystem/EngineReplicationBridge.h b/Engine/Source/Runtime/Engine/Public/Net/Iris/ReplicationSystem/EngineReplicationBridge.h
--- a/Engine/Source/Runtime/Engine/Public/Net/Iris/ReplicationSystem/EngineReplicationBridge.h
+++ b/Engine/Source/Runtime/Engine/Public/Net/Iris/ReplicationSystem/EngineReplicationBridge.h
@@ -17,6 +17,8 @@ class UWorld;
 namespace UE::Net
 {
 	enum class ENetRefHandleError : uint32;
+
+	class FDeferredReplicationSystemCalls;
 }
 
 namespace UE::Net
@@ -185,8 +187,58 @@ private:
 	/** Return true if actors that have a local role other than ROLE_Authority can be replicated to clients (default: false). */
 	bool AllowReplicationOfActorsWithAnyRole() const;
 
+	/** Cancel a start or stop replicating request for an actor. */
+	void CancelDeferredActor(AActor* Actor);
+
+	/** Cancel a start or stop replicating request for a component. */
+	void CancelDeferredComponent(UActorComponent* Component);
+
+	/** 
+	 * Deferred versions of the following functions:
+	 *	StartReplicatingActor()
+	 *	StopReplicatingActor()
+	 *	StartReplicatingComponent()
+	 *	StopReplicatingComponent()
+	 *
+	 * These functions will be deferred until FlushDeferred() is called.
+	 * 
+	 * For any given actor or component, start can only be called if the object isn't already registered
+	 * in the replication system or stop has been previously deferred; and stop can only be called if the
+	 * object is already registered in the replication system or start has been previously deferred.
+	 *
+	 * When FlushDeferred() is called, two calls to start and stop for the same actor or component will
+	 * cancel out so neither action is performed.
+	 * 
+	 * If StopReplicatingActorDefer() gets called in FlushDeferred() then it will use the value of 
+	 * FStopReplicatingActorParams from the last call to StopReplicatingActorDefer().
+	 * 
+	 * If StartReplicatingComponentDefer() gets called in FlushDeferred() then it will use the value of
+	 * AActor* from the last call to StartReplicatingComponentDefer().
+	 * 
+	 * E.g. 
+	 * 	Assume these functions are called in order:
+	 *		StopReplicatingActorDefer(MyActor, StopParams1)
+	 *		StartReplicatingActorDefer(MyActor)
+	 *		StopReplicatingActorDefer(MyActor, StopParams2)
+	 * 
+	 * When calling FlushDeferred() then StopReplicatingActor() will be called for MyActor since it 
+	 * has the greater call count, and it will use the argument StopParams2 because that was the argument
+	 * from the last call to StopReplicatingActorDefer().
+	 *
+	 * The order of deferred calls is not currently guaranteed when processed in FlushDeferred().
+	 */
+	void StartReplicatingActorDefer(AActor* Instance);
+	void StopReplicatingActorDefer(AActor* Instance, const FStopReplicatingActorParams& Params);
+	void StartReplicatingComponentDefer(const UObject* RootObject, UActorComponent* Component);
+	void StopReplicatingComponentDefer(UActorComponent* Component, EEndReplicationFlags EndReplicationFlags);
+	
+	/** Process all deferred function calls. */
+	void FlushDeferred();
+
 private:
 
+	friend class UE::Net::FDeferredReplicationSystemCalls;
+
 	UE::Net::FNetObjectFactoryId ActorFactoryId;
 	UE::Net::FNetObjectFactoryId SubObjectFactoryId;
 
@@ -205,4 +257,31 @@ private:
 
 	bool bAllowExecutePreReplication = true;
 	bool bAllowReplicationOfActorsWithAnyRole = false;
+
+	struct FDeferredStartStop
+	{
+		/** 
+		 * Used to determine if start, stop or nothing should be called for an actor or component.
+		 * 
+		 *  < 0		Stop is called
+		 * == 0		Nothing happens
+		 *  > 0		Start is called
+		 */
+		int32 Count = 0;
+
+		/** The parameters passed into the last call to StopReplicatingActorDefer(). */
+		FStopReplicatingActorParams StopReplicatingActorParams = FStopReplicatingActorParams(EEndPlayReason::Type::Destroyed);
+
+		/** The end replication flags passed into the last call to StopReplicatingComponentDefer(). */
+		EEndReplicationFlags StopReplicatingComponentFlags = EEndReplicationFlags::None;
+
+		/** The root/parent UObject used by the component passed into the last call to StartReplicatingComponentDefer(). */
+		TObjectKey<UObject> RootObject;
+	};
+
+	/** A flag used to prevent adding or removing new deferred start and stops while flushing. */
+	bool bIsFlushDeferredRunning = false;
+
+	TMap<TObjectKey<AActor>, FDeferredStartStop> DeferredStartStopCallsActors;
+	TMap<TObjectKey<UActorComponent>, FDeferredStartStop> DeferredStartStopCallsComponents;
 };
diff --git a/Engine/Source/Runtime/Engine/Public/Net/Iris/ReplicationSystem/ReplicationSystemUtil.h b/Engine/Source/Runtime/Engine/Public/Net/Iris/ReplicationSystem/ReplicationSystemUtil.h
--- a/Engine/Source/Runtime/Engine/Public/Net/Iris/ReplicationSystem/ReplicationSystemUtil.h
+++ b/Engine/Source/Runtime/Engine/Public/Net/Iris/ReplicationSystem/ReplicationSystemUtil.h
@@ -5,6 +5,7 @@
 // HEADER_UNIT_SKIP - Can't find includes if UE_WITH_IRIS is not set
 
 
+#include "Iris/ReplicationSystem/ReplicationBridgeTypes.h"
 #include "Iris/ReplicationSystem/Conditionals/ReplicationCondition.h"
 #include "Net/Core/NetHandle/NetHandle.h"
 #include "Engine/EngineTypes.h"
@@ -93,7 +94,7 @@ struct FReplicationSystemUtil
 	/** Returns the UEngineReplicationBridge of the UReplicationSystem belonging to the main NetDriver of this World. May return null. */
 	ENGINE_API static UEngineReplicationBridge* GetEngineReplicationBridge(const UWorld* World);
 
-		/** Returns the NetHandle for any replicated object. The returned handle may be invalid. */
+	/** Returns the NetHandle for any replicated object. The returned handle may be invalid. */
 	ENGINE_API static FNetHandle GetNetHandle(const UObject* ReplicatedObject);
 
 	/** Start replicating actor and all of its registered components and subobjects. If any ReplicationSystem wants to replicate the actor a NetHandle will be created. */
@@ -183,5 +184,33 @@ struct FReplicationSystemUtil
 	ENGINE_API static void SetPollFrequency(const UObject* Object, float PollFrequency);
 };
 
+/** Expose internal deferred replication system calls to the rest of the engine. */
+class FDeferredReplicationSystemCalls
+{
+private:
+
+	/** 
+	 * Deferred version of the following functions:
+	 *	FReplicationSystemUtil::StartReplicatingActor(AActor*)
+	 *	FReplicationSystemUtil::StopReplicatingActor(AActor*, FStopReplicatingActorParams*)
+	 * 	FReplicationSystemUtil::StartReplicatingActorComponent(const AActor*, UActorComponent*)
+	 *	FReplicationSystemUtil::StopReplicatingActorComponent(UActorComponent*)
+	 *
+	 * The functions below are wrappers for equivalent functions in UEngineReplicationBridge.
+	 * 
+	 * All deferred functions will be processed when FlushDeferred() is called.
+	 */
+	static void StartReplicatingActorDefer(AActor* Actor);
+	static void StopReplicatingActorDefer(AActor* Actor, const FStopReplicatingActorParams& StopReplicatingActorParams);
+	static void StartReplicatingActorComponentDefer(const AActor* Actor, UActorComponent* ActorComponent);
+	static void StopReplicatingActorComponentDefer(UActorComponent* ActorComponent, EEndReplicationFlags EndReplicationFlags);
+	static void FlushDeferred(UEngineReplicationBridge* Bridge);
+
+	friend class ::AActor;
+	friend class ::UActorComponent;
+	friend class ::APlayerController;
+	friend class ::UNetDriver;
+};
+
 }
 
