Fixed an edge case where RemoveFromWorld could be called in non-incremental mode while the level was already in the pre-unregistering phase (with pending asynchronous physics destruction).
This change complements fix done in CL 49563267.

#Synced-CL 49565614

diff --git a/Engine/Source/Runtime/Engine/Private/Level.cpp b/Engine/Source/Runtime/Engine/Private/Level.cpp
--- a/Engine/Source/Runtime/Engine/Private/Level.cpp
+++ b/Engine/Source/Runtime/Engine/Private/Level.cpp
@@ -2099,13 +2099,15 @@ bool ULevel::IncrementalRunConstructionScripts(bool bProcessAllActors)
 
 bool ULevel::IncrementalPreUnregisterComponents(FUnregisterComponentContext& Context)
 {
-	auto ProcessPreUnregisteringComponents = [this]()
+	const bool bWaitForCompletion = Context.GetIncrementalRemoveTimeout().WillNeverExpire();
+
+	auto ProcessPreUnregisteringComponents = [this, bWaitForCompletion]()
 	{
 		if (HasPreUnregisteringComponents())
 		{
 			if (FPhysScene* PhysScene = OwningWorld->GetPhysicsScene())
 			{
-				PhysScene->ProcessAsyncPhysicsStateJobs();
+				PhysScene->ProcessAsyncPhysicsStateJobs(bWaitForCompletion);
 			}
 		}
 	};
@@ -2198,6 +2200,11 @@ bool ULevel::IncrementalUnregisterComponents(int32 NumComponentsToUnregister, FU
 
 	FUnregisterComponentContext LocalContext(OwningWorld);
 	FUnregisterComponentContext& Context = InContext ? *InContext : LocalContext;
+	if (!bIncrementalUnregisterComponents)
+	{
+		Context.SetIncrementalRemoveTimeout(UE::FTimeout::Never());
+	}
+
 	bool bAllowLoop;
 	do
 	{
