Fixes a crash in SubsystemCollectionBase::Deinitialize during exit when FWorldContext::ThisCurrentWorld is null.

Example repro:
1. Launch with an invalid world path (eg LyraGame-Win64-Debug.exe c:\badpath).
2. UGameInstance::StartGameInstance() calls UEngine::Browse() to attempt the load and fails, leaving FWorldContext::ThisCurrentWorld == nullptr .
3. UGameInstance::StartGameInstance() shows a message dialog where selecting 'Cancel' results in a call to FPlatformMisc::RequestExit().

This later crashes in SubsystemCollectionBase::Deinitialize() due to subsystems with InternalIndex == -1. The core issue is that UGameEngine::PreExit() walks over the list of UGameEngine::WorldList (FWorldContext instances) conditionally performs cleans up based on FWorldContext::ThisCurrentWorld. When the world is null, this skips calling UGameInstance::Shutdown(), resulting in an atypical shutdown order.

This change fixes the issue by using FWorldContext to retrieve the owning UGameInstance to call Shutdown() when the World is null. CancelPending() is now called in both branches for consistency. The existing 'World != NULL' path is updated to use the context, making it clearer that the logic is the same in both cases.

#Synced-CL 49566552

diff --git a/Engine/Source/Runtime/Engine/Private/GameEngine.cpp b/Engine/Source/Runtime/Engine/Private/GameEngine.cpp
--- a/Engine/Source/Runtime/Engine/Private/GameEngine.cpp
+++ b/Engine/Source/Runtime/Engine/Private/GameEngine.cpp
@@ -1332,13 +1332,15 @@ void UGameEngine::PreExit()
 	// Clean up all worlds
 	for (int32 WorldIndex = 0; WorldIndex < WorldList.Num(); ++WorldIndex)
 	{
-		UWorld* const World = WorldList[WorldIndex].World();
+		FWorldContext& WorldContext = WorldList[WorldIndex];
+		
+		UWorld* const World = WorldContext.World();
 		if ( World != NULL )
 		{
 			World->BeginTearingDown();
 
 			// Cancel any pending connection to a server
-			CancelPending(World);
+			CancelPending(WorldContext);
 
 			// Shut down any existing game connections
 			ShutdownWorldNetDriver(World);
@@ -1352,13 +1354,27 @@ void UGameEngine::PreExit()
 			
 			World->EndPlay(EEndPlayReason::Quit);
 
-			if (World->GetGameInstance() != nullptr)
+			if (WorldContext.OwningGameInstance != nullptr)
 			{
-				World->GetGameInstance()->Shutdown();
+				WorldContext.OwningGameInstance->Shutdown();
 			}
 
 			World->CleanupWorld();
 		}
+		else 
+		{
+			// A world may not be bound to the context in an early shutdown case such as the world failing to load. 
+			// Preserve the call to Shutdown to avoid introducing a different destruction order in an exceptional error 
+			// case.
+			
+			// Cancel any pending connection to a server
+			CancelPending(WorldContext);
+			
+			if (WorldContext.OwningGameInstance != nullptr)
+			{
+				WorldContext.OwningGameInstance->Shutdown();
+			}
+		}
 	}
 
 	Super::PreExit();
