Removed invalid check now happening following CL 49563263.

#Synced-CL 49565614

diff --git a/Engine/Source/Runtime/Engine/Private/Level.cpp b/Engine/Source/Runtime/Engine/Private/Level.cpp
--- a/Engine/Source/Runtime/Engine/Private/Level.cpp
+++ b/Engine/Source/Runtime/Engine/Private/Level.cpp
@@ -2155,7 +2155,6 @@ bool ULevel::IncrementalUnregisterComponents(int32 NumComponentsToUnregister, FU
 
 	if (bIncrementalUnregisterComponentsCompleted)
 	{
-		check(bIncrementalUnregisterComponents);
 		check(IncrementalUnregisterComponentState == EIncrementalUnregisterComponentState::Init);
 		return true;
 	}
