API Changes to Wave Proxy Player

Remove "Block Rate". Unused.

Add functions to get Frame based times

#rnx
#Synced-CL 49561705

diff --git a/Engine/Plugins/Runtime/Metasound/Source/MetasoundEngine/Private/MetasoundWavePlayerNodeV2.cpp b/Engine/Plugins/Runtime/Metasound/Source/MetasoundEngine/Private/MetasoundWavePlayerNodeV2.cpp
--- a/Engine/Plugins/Runtime/Metasound/Source/MetasoundEngine/Private/MetasoundWavePlayerNodeV2.cpp
+++ b/Engine/Plugins/Runtime/Metasound/Source/MetasoundEngine/Private/MetasoundWavePlayerNodeV2.cpp
@@ -242,10 +242,9 @@ namespace Metasound
 			, Inputs(MoveTemp(InInputs))
 			, Outputs(MoveTemp(InOutputs))
 		{
-			const int32 NumFramesPerBlock = InParams.OperatorSettings.GetNumFramesPerBlock();
 			const float SampleRate = InParams.OperatorSettings.GetSampleRate();
 			const int32 NumChannels = Outputs.Audio.Num();
-			FSoundWaveProxyPlayer::FSettings PlayerSettings(NumFramesPerBlock, SampleRate, NumChannels);
+			FSoundWaveProxyPlayer::FSettings PlayerSettings(SampleRate, NumChannels);
 			PlayerSettings.bMaintainAudioSync = Inputs.bMaintainAudioSync;
 			ProxyPlayer = FSoundWaveProxyPlayer::Create(PlayerSettings);
 
@@ -263,10 +262,9 @@ namespace Metasound
 			OperatorSettings = ResetParams.OperatorSettings;
 			Outputs.Reset(ResetParams, Outputs.Audio.Num());
 			
-			const int32 NumFramesPerBlock = ResetParams.OperatorSettings.GetNumFramesPerBlock();
             const float SampleRate = ResetParams.OperatorSettings.GetSampleRate();
             const int32 NumChannels = Outputs.Audio.Num();
-			FSoundWaveProxyPlayer::FSettings PlayerSettings(NumFramesPerBlock, SampleRate, NumChannels);
+			FSoundWaveProxyPlayer::FSettings PlayerSettings(SampleRate, NumChannels);
 			PlayerSettings.bMaintainAudioSync = Inputs.bMaintainAudioSync;
 			ProxyPlayer->Reset(PlayerSettings);
 
diff --git a/Engine/Source/Runtime/Engine/Classes/Sound/SoundWaveProxyPlayer.h b/Engine/Source/Runtime/Engine/Classes/Sound/SoundWaveProxyPlayer.h
--- a/Engine/Source/Runtime/Engine/Classes/Sound/SoundWaveProxyPlayer.h
+++ b/Engine/Source/Runtime/Engine/Classes/Sound/SoundWaveProxyPlayer.h
@@ -23,12 +23,10 @@ public:
 
 	struct FSettings
 	{
-		FSettings(int32 InNumFramesPerBlock, float InOutputSampleRate, int32 InOutputNumChannels)
+		FSettings(float InOutputSampleRate, int32 InOutputNumChannels)
 			: OutputSampleRate(InOutputSampleRate)
 			, OutputNumChannels(InOutputNumChannels)
-			, NumFramesPerBlock(InNumFramesPerBlock)
 		{
-			check(NumFramesPerBlock > 0);
 			check(OutputSampleRate > 0.0f);
 			check(OutputNumChannels > 0);
 			check(OutputNumChannels <= 8);
@@ -38,7 +36,6 @@ public:
 		int32 OutputNumChannels = 0;
 		// Method for upmixing mono audio to multichannel output
 		Audio::EChannelMapMonoUpmixMethod MonoUpmixMethod = Audio::EChannelMapMonoUpmixMethod::FullVolume;
-		int32 NumFramesPerBlock = 0;
 		int32 MaxDecodeSizeFrames = 1024;
 		bool bMaintainAudioSync = false;
 	};
@@ -109,6 +106,30 @@ public:
 	 */
 	UE_API float GetOutputSampleRate() const;
 
+	/**
+	 * 
+	 * @return the number of channels of the source sound wave
+	 */
+	UE_API int32 GetSourceNumChannels() const;
+
+	/**
+	 * 
+	 * @return Number of frames in the source sound wave
+	 */
+	UE_API int32 GetSourceNumFrames() const;
+
+	/**
+	 * 
+	 * @return Duration (in seconds) of the source sound wave 
+	 */
+	UE_API float GetSourceDuration() const;
+
+	/**
+	 * 
+	 * @return the sample rate of the source sound wave
+	 */
+	UE_API float GetSourceSampleRate() const;
+
 	/**
 	 * 
 	 * @return whether this player has a valid Sound Wave and Decoder
@@ -221,6 +242,13 @@ public:
 	 */
 	UE_API float GetCurrentPlaybackTimeSeconds() const;
 
+	/**
+	 * The Current playback position into the Source 
+	 * 
+	 * @return Current Playback Position In Frames
+	 */
+	UE_API int32 GetCurrentPlaybackFrame() const;
+
 	/**
 	 * Returns the current playback progress based on the playback time into the source file
 	 *
diff --git a/Engine/Source/Runtime/Engine/Private/SoundWaveProxyPlayer.cpp b/Engine/Source/Runtime/Engine/Private/SoundWaveProxyPlayer.cpp
--- a/Engine/Source/Runtime/Engine/Private/SoundWaveProxyPlayer.cpp
+++ b/Engine/Source/Runtime/Engine/Private/SoundWaveProxyPlayer.cpp
@@ -99,6 +99,30 @@ float FSoundWaveProxyPlayer::GetOutputSampleRate() const
 	return Settings.OutputSampleRate;
 }
 
+int32 FSoundWaveProxyPlayer::GetSourceNumChannels() const
+{
+	check(IsPlayerValid());
+	return NumChannels;
+}
+
+float FSoundWaveProxyPlayer::GetSourceSampleRate() const
+{
+	check(IsPlayerValid())
+	return SampleRate;
+}
+
+int32 FSoundWaveProxyPlayer::GetSourceNumFrames() const
+{
+	check(IsPlayerValid());
+	return NumFramesInWave;	
+}
+
+float FSoundWaveProxyPlayer::GetSourceDuration() const
+{
+	check(IsPlayerValid())
+	return DurationSeconds;
+}
+
 bool FSoundWaveProxyPlayer::IsPlayerValid() const
 {
 	return SoundWave.IsValid() && CompressedAudioInfo.IsValid();
@@ -595,6 +619,17 @@ float FSoundWaveProxyPlayer::GetCurrentPlaybackTimeSeconds() const
 	return SourceFrameIndex / SampleRate;	
 }
 
+int32 FSoundWaveProxyPlayer::GetCurrentPlaybackFrame() const
+{
+	if (!SoundWave)
+	{
+		return 0;
+	}
+	
+	return SourceFrameIndex;
+}
+
+
 float FSoundWaveProxyPlayer::GetPlaybackProgress() const
 {
 	if (!SoundWave)
