Added utility functions for converting soft paths into UObject and UClass pointers.

#Synced-CL 49462178

diff --git a/Engine/Source/Runtime/Engine/Classes/Kismet/KismetSystemLibrary.h b/Engine/Source/Runtime/Engine/Classes/Kismet/KismetSystemLibrary.h
--- a/Engine/Source/Runtime/Engine/Classes/Kismet/KismetSystemLibrary.h
+++ b/Engine/Source/Runtime/Engine/Classes/Kismet/KismetSystemLibrary.h
@@ -149,6 +149,16 @@ class UKismetSystemLibrary : public UBlueprintFunctionLibrary
 	UFUNCTION(BlueprintPure, Category = "Utilities")
 	static ENGINE_API FSoftObjectPath GetSoftObjectPath(const UObject* Object);
 
+#if WITH_EDITOR
+	// Gets the UObject associated with a SoftPath, loading if necessary.
+	UFUNCTION(BlueprintPure, Category = "Utilities")
+	static UObject* GetObjectFromSoftPath(const FSoftObjectPath& Path);
+
+	// Gets the UClass associated with a SoftPath, loading if necessary.
+	UFUNCTION(BlueprintPure, Category = "Utilities")
+	static UClass* GetClassFromSoftPath(const FSoftClassPath& Path);
+#endif
+
 	// Returns the full file system path to a UObject
 	// If given a non-asset UObject, it will return an empty string
 	UFUNCTION(BlueprintPure, Category = "Utilities|Paths")
diff --git a/Engine/Source/Runtime/Engine/Private/KismetSystemLibrary.cpp b/Engine/Source/Runtime/Engine/Private/KismetSystemLibrary.cpp
--- a/Engine/Source/Runtime/Engine/Private/KismetSystemLibrary.cpp
+++ b/Engine/Source/Runtime/Engine/Private/KismetSystemLibrary.cpp
@@ -162,6 +162,30 @@ FSoftObjectPath UKismetSystemLibrary::GetSoftObjectPath(const UObject* Object)
 	return FSoftObjectPath(Object);
 }
 
+#if WITH_EDITOR
+UObject* UKismetSystemLibrary::GetObjectFromSoftPath(const FSoftObjectPath& Path)
+{
+	UObject* Object = nullptr;
+	Object = Path.ResolveObject();
+	if (!Object)
+	{
+		Object = Path.TryLoad();
+	}
+	return Object;
+}
+
+UClass* UKismetSystemLibrary::GetClassFromSoftPath(const FSoftClassPath& Path)
+{
+	UClass* Class = nullptr;
+	Class = Path.ResolveClass();
+	if (!Class)
+	{
+		Class = Path.TryLoadClass<UObject>();
+	}
+	return Class;
+}
+#endif
+
 FString UKismetSystemLibrary::GetSystemPath(const UObject* Object)
 {
 	if (Object && Object->IsAsset())
