Reduced the scope of the fix applied in 49572827 by removing changes to GameEngine::PreExit() when a world is present.

The original code in GameEngine::PreExit() when World != NULL made no assumptions about the stability of WorldList during iteration; it?s possible for code to add or remove elements. The original fix in 49572827 made behavior consistent with the no-world path, but introduced an assumption that WorldList does not change while iterating.

This change undoes the modifications to the if (World != NULL) path and makes the else case more explicit about its lifetime expectations by removing local aliases.

#Synced-CL 49566552

diff --git a/Engine/Source/Runtime/Engine/Private/GameEngine.cpp b/Engine/Source/Runtime/Engine/Private/GameEngine.cpp
--- a/Engine/Source/Runtime/Engine/Private/GameEngine.cpp
+++ b/Engine/Source/Runtime/Engine/Private/GameEngine.cpp
@@ -1332,15 +1332,13 @@ void UGameEngine::PreExit()
 	// Clean up all worlds
 	for (int32 WorldIndex = 0; WorldIndex < WorldList.Num(); ++WorldIndex)
 	{
-		FWorldContext& WorldContext = WorldList[WorldIndex];
-		
-		UWorld* const World = WorldContext.World();
+		UWorld* const World = WorldList[WorldIndex].World();
 		if ( World != NULL )
 		{
 			World->BeginTearingDown();
 
 			// Cancel any pending connection to a server
-			CancelPending(WorldContext);
+			CancelPending(World);
 
 			// Shut down any existing game connections
 			ShutdownWorldNetDriver(World);
@@ -1354,9 +1352,9 @@ void UGameEngine::PreExit()
 			
 			World->EndPlay(EEndPlayReason::Quit);
 
-			if (WorldContext.OwningGameInstance != nullptr)
+			if (World->GetGameInstance() != nullptr)
 			{
-				WorldContext.OwningGameInstance->Shutdown();
+				World->GetGameInstance()->Shutdown();
 			}
 
 			World->CleanupWorld();
@@ -1366,13 +1364,11 @@ void UGameEngine::PreExit()
 			// A world may not be bound to the context in an early shutdown case such as the world failing to load. 
 			// Preserve the call to Shutdown to avoid introducing a different destruction order in an exceptional error 
 			// case.
+			CancelPending(WorldList[WorldIndex]);
 			
-			// Cancel any pending connection to a server
-			CancelPending(WorldContext);
-			
-			if (WorldContext.OwningGameInstance != nullptr)
+			if (WorldList[WorldIndex].OwningGameInstance != nullptr)
 			{
-				WorldContext.OwningGameInstance->Shutdown();
+				WorldList[WorldIndex].OwningGameInstance->Shutdown();
 			}
 		}
 	}
