PR #14120: Prevent crash in DebugDrawSkeleton when using 'show bones' command

#Synced-CL 49587318

diff --git a/Engine/Source/Runtime/Engine/Private/SkeletalMeshSceneProxy.cpp b/Engine/Source/Runtime/Engine/Private/SkeletalMeshSceneProxy.cpp
--- a/Engine/Source/Runtime/Engine/Private/SkeletalMeshSceneProxy.cpp
+++ b/Engine/Source/Runtime/Engine/Private/SkeletalMeshSceneProxy.cpp
@@ -1444,7 +1444,7 @@ void FSkeletalMeshSceneProxy::DebugDrawSkeleton(int32 ViewIndex, FMeshElementCol
 	const FMatrix& ProxyLocalToWorld = GetLocalToWorld();
 
 	// Can't draw this, don't have ComponentSpaceTransforms. This happens with sk meshes rendered with FSkeletalMeshObjectStatic.
-	if (!MeshObject->GetComponentSpaceTransforms())
+	if (!MeshObject || !MeshObject->GetComponentSpaceTransforms())
 	{
 		return;
 	}
@@ -1460,10 +1460,12 @@ void FSkeletalMeshSceneProxy::DebugDrawSkeleton(int32 ViewIndex, FMeshElementCol
 
 	FPrimitiveDrawInterface* PDI = Collector.GetPDI(ViewIndex);
 	TArray<FTransform>& ComponentSpaceTransforms = *MeshObject->GetComponentSpaceTransforms();
-
-	for (int32 Index = 0; Index < ComponentSpaceTransforms.Num(); ++Index)
+	const FReferenceSkeleton& RefSkeleton = SkinnedAsset->GetRefSkeleton();
+	const int32 NumBones = FMath::Min(ComponentSpaceTransforms.Num(), RefSkeleton.GetNum());
+	
+	for (int32 Index = 0; Index < NumBones; ++Index)
 	{
-		const int32 ParentIndex = SkinnedAsset->GetRefSkeleton().GetParentIndex(Index);
+		const int32 ParentIndex = RefSkeleton.GetParentIndex(Index);
 		FVector Start, End;
 		
 		FLinearColor LineColor = DebugDrawColor.Get(MakeRandomColorForSkeleton(GetPrimitiveComponentId().PrimIDValue));
diff --git a/Engine/Source/Runtime/Engine/Public/ReferenceSkeleton.h b/Engine/Source/Runtime/Engine/Public/ReferenceSkeleton.h
--- a/Engine/Source/Runtime/Engine/Public/ReferenceSkeleton.h
+++ b/Engine/Source/Runtime/Engine/Public/ReferenceSkeleton.h
@@ -178,6 +178,10 @@ private:
 
 	int32 GetParentIndexInternal(const int32 BoneIndex, const TArray<FMeshBoneInfo>& BoneInfo) const
 	{
+		if (!BoneInfo.IsValidIndex(BoneIndex))
+		{
+			return INDEX_NONE;
+		}
 		const int32 ParentIndex = BoneInfo[BoneIndex].ParentIndex;
 
 		// Parent must be valid. Either INDEX_NONE for Root, or before children for non root bones.
