[WorldPartition] Fix crash when 2 partitions are named the same, because RuntimePartitionDesc.Class is no longer a valid property
- Make the rename algorithms more resilient
- Finish deprecating FRuntimePartitionDesc::Class

#Synced-CL 49450672

diff --git a/Engine/Source/Runtime/Engine/Private/WorldPartition/RuntimeHashSet/WorldPartitionRuntimeHashSet.cpp b/Engine/Source/Runtime/Engine/Private/WorldPartition/RuntimeHashSet/WorldPartitionRuntimeHashSet.cpp
--- a/Engine/Source/Runtime/Engine/Private/WorldPartition/RuntimeHashSet/WorldPartitionRuntimeHashSet.cpp
+++ b/Engine/Source/Runtime/Engine/Private/WorldPartition/RuntimeHashSet/WorldPartitionRuntimeHashSet.cpp
@@ -260,7 +260,6 @@ void UWorldPartitionRuntimeHashSet::SetDefaultValues()
 	check(RuntimePartitions.IsEmpty());
 
 	FRuntimePartitionDesc& RuntimePartitionDesc = RuntimePartitions.AddDefaulted_GetRef();
-	RuntimePartitionDesc.Class = URuntimePartitionLHGrid::StaticClass();
 	RuntimePartitionDesc.Name = TEXT("MainPartition");
 
 	RuntimePartitionDesc.MainLayer = NewObject<URuntimePartitionLHGrid>(this, NAME_None);
@@ -981,35 +980,59 @@ void UWorldPartitionRuntimeHashSet::PostEditChangeChainProperty(FPropertyChanged
 		{
 			FRuntimePartitionHLODSetup& RuntimePartitionHLODSetup = RuntimePartitionDesc.HLODSetups[HLODSetupsIndex];
 
-			for (int32 CurHLODSetupsIndex = 0; CurHLODSetupsIndex < RuntimePartitionDesc.HLODSetups.Num(); CurHLODSetupsIndex++)
+			// Rename HLOD if it clashes with another one
+			auto FindUniqueHLODName = [this, &RuntimePartitionDesc](FName TestedName, int32 SkipIndex)->FName
 			{
-				if (CurHLODSetupsIndex != HLODSetupsIndex)
+				bool FoundDuplicate;
+				do
 				{
-					if (RuntimePartitionHLODSetup.Name == RuntimePartitionDesc.HLODSetups[CurHLODSetupsIndex].Name)
+					FoundDuplicate = false;
+					for (int32 CurHLODSetupsIndex = 0; CurHLODSetupsIndex < RuntimePartitionDesc.HLODSetups.Num(); CurHLODSetupsIndex++)
 					{
-						RuntimePartitionHLODSetup.Name = *FString::Printf(TEXT("HLOD_%d"), RuntimePartitionHLODSetup.PartitionLayer->HLODIndex);
-						break;
+						if (CurHLODSetupsIndex != SkipIndex)
+						{
+							if (TestedName == RuntimePartitionDesc.HLODSetups[CurHLODSetupsIndex].Name)
+							{
+								FoundDuplicate = true;
+								TestedName.SetNumber(TestedName.GetNumber() + 1);
+								break;
+							}
+						}
 					}
-				}
-			}
+				} while (FoundDuplicate);
+				return TestedName;
+			};
 
+			RuntimePartitionHLODSetup.Name = FindUniqueHLODName(RuntimePartitionHLODSetup.Name, HLODSetupsIndex);
 			RuntimePartitionHLODSetup.PartitionLayer->Name = RuntimePartitionHLODSetup.Name;
 			RuntimePartitionHLODSetup.RowDisplayName = RuntimePartitionHLODSetup.PartitionLayer->Name;
 		}
 		else
 		{
-			for (int32 CurRuntimePartitionIndex = 0; CurRuntimePartitionIndex < RuntimePartitions.Num(); CurRuntimePartitionIndex++)
+			// Rename layer if it clashes with another one
+			auto FindUniquePartitionName = [this](FName TestedName, int32 SkipIndex)->FName
 			{
-				if (CurRuntimePartitionIndex != RuntimePartitionIndex)
+				bool FoundDuplicate;
+				do
 				{
-					if (RuntimePartitionDesc.Name == RuntimePartitions[CurRuntimePartitionIndex].Name)
+					FoundDuplicate = false;
+					for (int32 CurRuntimePartitionIndex = 0; CurRuntimePartitionIndex < RuntimePartitions.Num(); CurRuntimePartitionIndex++)
 					{
-						RuntimePartitionDesc.Name = RuntimePartitionDesc.Class->GetFName();
-						break;
+						if (CurRuntimePartitionIndex != SkipIndex)
+						{
+							if (TestedName == RuntimePartitions[CurRuntimePartitionIndex].Name)
+							{
+								FoundDuplicate = true;
+								TestedName.SetNumber(TestedName.GetNumber() + 1);
+								break;
+							}
+						}
 					}
-				}
-			}
+				} while (FoundDuplicate);
+				return TestedName;
+			};
 
+			RuntimePartitionDesc.Name = FindUniquePartitionName(RuntimePartitionDesc.Name, RuntimePartitionIndex);
 			RuntimePartitionDesc.MainLayer->Name = RuntimePartitionDesc.Name;
 		}
 	}
@@ -1356,4 +1379,5 @@ const FGuid* UWorldPartitionRuntimeHashSet::GetStandaloneHLODActorSourceCellOver
 const FGuid* UWorldPartitionRuntimeHashSet::GetCustomHLODActorSourceCellOverride(const FGuid& InActorGuid) const
 {
 	return CustomHLODActorToSourceCellsMap.Find(InActorGuid);
-}
\ No newline at end of file
+}
+
diff --git a/Engine/Source/Runtime/Engine/Private/WorldPartition/RuntimeHashSet/WorldPartitionRuntimeHashSetConversions.cpp b/Engine/Source/Runtime/Engine/Private/WorldPartition/RuntimeHashSet/WorldPartitionRuntimeHashSetConversions.cpp
--- a/Engine/Source/Runtime/Engine/Private/WorldPartition/RuntimeHashSet/WorldPartitionRuntimeHashSetConversions.cpp
+++ b/Engine/Source/Runtime/Engine/Private/WorldPartition/RuntimeHashSet/WorldPartitionRuntimeHashSetConversions.cpp
@@ -73,7 +73,6 @@ UWorldPartitionRuntimeHashSet* UWorldPartitionRuntimeHashSet::CreateFrom(const U
 			if (!Grid.HLODLayer)
 			{
 				FRuntimePartitionDesc& RuntimePartitionDesc = HashSet->RuntimePartitions.AddDefaulted_GetRef();
-				RuntimePartitionDesc.Class = URuntimePartitionLHGrid::StaticClass();
 				RuntimePartitionDesc.Name = Grid.GridName;
 
 				URuntimePartitionLHGrid* LHGrid = NewObject<URuntimePartitionLHGrid>(HashSet, NAME_None);
diff --git a/Engine/Source/Runtime/Engine/Private/WorldPartition/RuntimeHashSet/WorldPartitionRuntimeHashSetHLODGeneration.cpp b/Engine/Source/Runtime/Engine/Private/WorldPartition/RuntimeHashSet/WorldPartitionRuntimeHashSetHLODGeneration.cpp
--- a/Engine/Source/Runtime/Engine/Private/WorldPartition/RuntimeHashSet/WorldPartitionRuntimeHashSetHLODGeneration.cpp
+++ b/Engine/Source/Runtime/Engine/Private/WorldPartition/RuntimeHashSet/WorldPartitionRuntimeHashSetHLODGeneration.cpp
@@ -267,7 +267,6 @@ bool UWorldPartitionRuntimeHashSet::SetupHLODActors(const IStreamingGenerationCo
 						{
 							FRuntimePartitionDesc& PartitionDesc = StandaloneHLODWorldRuntimePartitions.AddDefaulted_GetRef();
 							PartitionDesc.Name = SourcePartitionDesc.Name;
-							PartitionDesc.Class = SourcePartitionDesc.Class;
 							PartitionDesc.MainLayer = DuplicateObject<URuntimePartition>(SourcePartitionDesc.MainLayer, StandaloneHLODWorldHash);
 
 							for (FRuntimePartitionHLODSetup& SourceHLODSetup : SourcePartitionDesc.HLODSetups)
diff --git a/Engine/Source/Runtime/Engine/Public/WorldPartition/RuntimeHashSet/WorldPartitionRuntimeHashSet.h b/Engine/Source/Runtime/Engine/Public/WorldPartition/RuntimeHashSet/WorldPartitionRuntimeHashSet.h
--- a/Engine/Source/Runtime/Engine/Public/WorldPartition/RuntimeHashSet/WorldPartitionRuntimeHashSet.h
+++ b/Engine/Source/Runtime/Engine/Public/WorldPartition/RuntimeHashSet/WorldPartitionRuntimeHashSet.h
@@ -58,12 +58,21 @@ struct FRuntimePartitionDesc
 {
 	GENERATED_USTRUCT_BODY()
 
+	PRAGMA_DISABLE_DEPRECATION_WARNINGS
+    FRuntimePartitionDesc() = default;
+    ~FRuntimePartitionDesc() = default;
+    FRuntimePartitionDesc(const FRuntimePartitionDesc&) = default;
+    FRuntimePartitionDesc(FRuntimePartitionDesc&&) = default;
+    FRuntimePartitionDesc& operator=(const FRuntimePartitionDesc&) = default;
+    FRuntimePartitionDesc& operator=(FRuntimePartitionDesc&&) = default;
+    PRAGMA_ENABLE_DEPRECATION_WARNINGS
+
 	/** Name for this partition, used to map actors to it through the Actor.RuntimeGrid property  */
 	UPROPERTY(EditAnywhere, Category = RuntimeSettings, Meta = (EditCondition = "MainLayer != nullptr", HideEditConditionToggle))
 	FName Name;
 
 	/** Partition class */
-	//UE_DEPRECATED(5.7, "This property isn't used anymore")
+	UE_DEPRECATED(5.8, "This property is replaced by MainLayer")
 	UPROPERTY()
 	TSubclassOf<URuntimePartition> Class;
 
