Added an array index guard for accessing the MontageInstances array. Somehow on the server this array was being modified during the loop iteration, resulting in an arrray out of index assertion failure.

#rnx
#tests PIE in Editor

diff --git a/Engine/Source/Runtime/Engine/Private/Animation/AnimInstance.cpp b/Engine/Source/Runtime/Engine/Private/Animation/AnimInstance.cpp
--- a/Engine/Source/Runtime/Engine/Private/Animation/AnimInstance.cpp
+++ b/Engine/Source/Runtime/Engine/Private/Animation/AnimInstance.cpp
@@ -3526,6 +3526,13 @@ void UAnimInstance::StopAllMontages(float BlendOut)
 {
 	for (int32 Index = MontageInstances.Num() - 1; Index >= 0; Index--)
 	{
+		if (!MontageInstances.IsValidIndex(Index))
+		{
+			// Adding a check here to avoid indexing past the end of the array. This could potentially occur if a
+			// montage has been removed from the array during this loop iteration.
+			continue;
+		}
+
 		MontageInstances[Index]->Stop(FAlphaBlend(BlendOut), true);
 
 		// UninitializeAnimation() empties the MontagesInstances array thus we need to exit early.
@@ -3547,6 +3554,13 @@ void UAnimInstance::StopAllMontagesByGroupName(FName InGroupName, const FMontage
 {
 	for (int32 InstanceIndex = MontageInstances.Num() - 1; InstanceIndex >= 0; InstanceIndex--)
 	{
+		if (!MontageInstances.IsValidIndex(InstanceIndex))
+		{
+			// Adding a check here to avoid indexing past the end of the array. This could potentially occur if a
+			// montage has been removed from the array during this loop iteration.
+			continue;
+		}
+
 		FAnimMontageInstance* MontageInstance = MontageInstances[InstanceIndex];
 		if (MontageInstance && MontageInstance->Montage && (MontageInstance->Montage->GetGroupName() == InGroupName))
 		{
