[SimpleStreamableAssetManager] Fix unsafe usage of UStreamableRenderAsset in SimpleStreamableAssetManager async task.
Previously, RenderAsset (a UStreamableRenderAsset object) was being dereferenced by the SimpleStreamableAssetManager async task thread. RenderAsset may have already been garbage collected by this point.
This change avoids dereferencing RenderAsset in the async task thread by moving the asset validation logic out of `FSimpleStreamableAssetManager::RegisterRecord` and into the `FSimpleStreamableAssetManager::FRegister::FRegister` constructor. This constructor is executed by a ParallelFor thread initiated by the game thread. This change also adds a `FRegisteredAsset` struct. Similar to `FRemovedAssetRecord`, this struct holds a TSharedPtr to the asset's SimpleStreamableAssetManagerIndex. This way we can update the asset's SimpleStreamableAssetManagerIndex without needing to dereference the RenderAsset in the async task.

#Synced-CL 49565614

diff --git a/Engine/Source/Runtime/Engine/Private/Streaming/SimpleStreamableAssetManager.cpp b/Engine/Source/Runtime/Engine/Private/Streaming/SimpleStreamableAssetManager.cpp
--- a/Engine/Source/Runtime/Engine/Private/Streaming/SimpleStreamableAssetManager.cpp
+++ b/Engine/Source/Runtime/Engine/Private/Streaming/SimpleStreamableAssetManager.cpp
@@ -272,7 +272,7 @@ void FSimpleStreamableAssetManager::UpdateTask_Async()
 		for (const FRegister& Record : Pending_RegisterRecords)
 		{
 			// Lets Skip reserving indices if no streamable assets
-			if (Record.Assets.Num() > 0)
+			if (Record.RegisteredAssets.Num() > 0)
 			{
 				// Cheap index assignment for proxy
 				const int32 ObjectIndex = ObjectUsedIndices.FindAndSetFirstZeroBit(FreeObjectIndexHint);
@@ -374,69 +374,47 @@ void FSimpleStreamableAssetManager::RegisterRecord(FRegister& Record)
 		return;
 	}
 
-	TArray<FStreamingRenderAssetPrimitiveInfo>& Assets = Record.Assets;
-	
-	//* Remove not needed */
-	for (int32 Index = 0; Index < Assets.Num(); Index++)
-	{
-		const FStreamingRenderAssetPrimitiveInfo& Info = Assets[Index];
-		if (!Info.RenderAsset || !Info.RenderAsset->IsStreamable())
-		{
-			Assets.RemoveAtSwap(Index--, EAllowShrinking::No);
-		}
-		else
-		{
-#if DO_CHECK
-			ensure(Info.TexelFactor >= 0.f
-				|| Info.RenderAsset->IsA<UStaticMesh>()
-				|| Info.RenderAsset->IsA<USkeletalMesh>()
-				|| (Info.RenderAsset->IsA<UTexture>() && Info.RenderAsset->GetLODGroupForStreaming() == TEXTUREGROUP_Terrain_Heightmap));
-#endif
-
-			// Other wise check that everything is setup right. If the component is not yet registered, then the bound data is irrelevant.
-			const bool bCanBeStreamedByDistance = Info.CanBeStreamedByDistance(true);
-
-			if (!Record.bForceMipStreaming && !bCanBeStreamedByDistance && Info.TexelFactor >= 0.f)
-			{
-				Assets.RemoveAtSwap(Index--, EAllowShrinking::No);
-			}
-		}
-	}
-
-	int32 AssetCount = Assets.Num();
+	TArray<FRegisteredAsset>& RegisteredAssets = Record.RegisteredAssets;
+	int32 AssetCount = RegisteredAssets.Num();
 
 	if (GSimpleStreamableAssetManagerEnsureAssetUniqueOnRegistration && AssetCount > 0)
 	{
 		// Sort by Texture to merge duplicate texture entries.
 		// Then sort by TexelFactor
-		Assets.Sort([](const FStreamingRenderAssetPrimitiveInfo& Lhs, const FStreamingRenderAssetPrimitiveInfo& Rhs)
+		RegisteredAssets.Sort([](const FRegisteredAsset& Lhs, const FRegisteredAsset& Rhs)
 		{
-				if (Lhs.RenderAsset == Rhs.RenderAsset)
-				{
-					return Lhs.TexelFactor > Rhs.TexelFactor;
-				}
-				return Lhs.RenderAsset < Rhs.RenderAsset;
+			const FStreamingRenderAssetPrimitiveInfo& LhsInfo = Lhs.AssetPrimitiveInfo;
+			const FStreamingRenderAssetPrimitiveInfo& RhsInfo = Rhs.AssetPrimitiveInfo;
+
+			if (LhsInfo.RenderAsset == RhsInfo.RenderAsset)
+			{
+				return LhsInfo.TexelFactor > RhsInfo.TexelFactor;
+			}
+			return LhsInfo.RenderAsset < RhsInfo.RenderAsset;
 		});
 
 		int32 ProcessedIndex = 0;
 		int32 EmplaceIndex = 1;
-		for (int32 Index = 1; Index < Assets.Num(); Index++)
+
+		for (int32 Index = 1; Index < RegisteredAssets.Num(); Index++)
 		{
+			const FStreamingRenderAssetPrimitiveInfo& CurInfo = RegisteredAssets[Index].AssetPrimitiveInfo;
+			FStreamingRenderAssetPrimitiveInfo& ProcessedInfo = RegisteredAssets[ProcessedIndex].AssetPrimitiveInfo;
+
 			// We found new asset 
-			if (Assets[ProcessedIndex].RenderAsset != Assets[Index].RenderAsset)
+			if (ProcessedInfo.RenderAsset != CurInfo.RenderAsset)
 			{
-				
 				if (EmplaceIndex != Index) // We need to fill the gap in array
 				{
-					Assets[EmplaceIndex] = MoveTemp(Assets[Index]);
+					RegisteredAssets[EmplaceIndex] = MoveTemp(RegisteredAssets[Index]);
 				}
 				++ProcessedIndex;
 				++EmplaceIndex;
 			}
 			// For landscape entries negative TexelFactor values are used we need to ensure we get min while sorted max 
-			else if (Assets[ProcessedIndex].TexelFactor < 0 && Assets[ProcessedIndex].RenderAsset == Assets[Index].RenderAsset)
+			else if (ProcessedInfo.TexelFactor < 0 && ProcessedInfo.RenderAsset == CurInfo.RenderAsset)
 			{
-				Assets[ProcessedIndex].TexelFactor = FMath::Min(Assets[ProcessedIndex].TexelFactor, Assets[Index].TexelFactor);
+				ProcessedInfo.TexelFactor = FMath::Min(ProcessedInfo.TexelFactor, CurInfo.TexelFactor);
 			}
 		}
 
@@ -448,9 +426,9 @@ void FSimpleStreamableAssetManager::RegisterRecord(FRegister& Record)
 		float MinDistanceSq = 0, MinRangeSq = 0, MaxRangeSq = FLT_MAX;
         GetDistanceAndRange(Record, MinDistanceSq, MinRangeSq, MaxRangeSq);
         SetBounds(ObjectIndex, Record.ObjectBounds, PackedRelativeBox_Identity, Record.StreamingScaleFactor, Record.LastRenderedTime, Record.ObjectBounds.Origin, MinDistanceSq, MinRangeSq, MaxRangeSq);
-       	
+
 		const bool bForceLOD = Record.bForceMipStreaming;
-		AddRenderAssetElements(TArrayView<FStreamingRenderAssetPrimitiveInfo>(Assets.GetData(), AssetCount), ObjectIndex, bForceLOD);
+		AddRenderAssetElements(TArrayView<FRegisteredAsset>(RegisteredAssets.GetData(), AssetCount), ObjectIndex, bForceLOD);
 	}
 }
 
@@ -529,13 +507,13 @@ FBoxSphereBounds FSimpleStreamableAssetManager::GetBounds(int32 BoundsIndex) con
 	return Bounds;
 }
 
-void FSimpleStreamableAssetManager::AddRenderAssetElements(const TArrayView<FStreamingRenderAssetPrimitiveInfo>& RenderAssetInstanceInfos, int32 ObjectRegistrationIndex, bool bForceMipStreaming)
+void FSimpleStreamableAssetManager::AddRenderAssetElements(const TArrayView<FRegisteredAsset>& RegisteredAssets, int32 ObjectRegistrationIndex, bool bForceMipStreaming)
 {
 	TSet<FAssetRecord> ObjectAssets;
-	ObjectAssets.Reserve(RenderAssetInstanceInfos.Num());
-	for (const FStreamingRenderAssetPrimitiveInfo& AssetInfo : RenderAssetInstanceInfos)
+	ObjectAssets.Reserve(RegisteredAssets.Num());
+	for (const FRegisteredAsset& RegisteredAsset : RegisteredAssets)
 	{
-		UStreamableRenderAsset* RenderAsset = AssetInfo.RenderAsset;
+		UStreamableRenderAsset* RenderAsset = RegisteredAsset.AssetPrimitiveInfo.RenderAsset;
 		if (RenderAsset == nullptr)
 		{
 			continue;
@@ -543,7 +521,7 @@ void FSimpleStreamableAssetManager::AddRenderAssetElements(const TArrayView<FStr
 		
 		const int32 AssetIndex = [&]
 		{
-			int32& SimpleStreamableAssetManagerIndex = *RenderAsset->SimpleStreamableAssetManagerIndex;
+			int32& SimpleStreamableAssetManagerIndex = *RegisteredAsset.SimpleStreamableAssetManagerIndex;
 			if (SimpleStreamableAssetManagerIndex != INDEX_NONE)
 			{
 				return SimpleStreamableAssetManagerIndex;
@@ -573,7 +551,10 @@ void FSimpleStreamableAssetManager::AddRenderAssetElements(const TArrayView<FStr
 			AssetIndexToBounds4Index.Add({});
 		}
 		
-		const int32 AssetElementIndex = AssetIndexToBounds4Index[AssetIndex].Add(FAssetBoundElement{ .ObjectRegistrationIndex = ObjectRegistrationIndex, .TexelFactor =  AssetInfo.TexelFactor, .bForceLOD = bForceMipStreaming});
+		const int32 AssetElementIndex = AssetIndexToBounds4Index[AssetIndex].Add(FAssetBoundElement{
+						.ObjectRegistrationIndex = ObjectRegistrationIndex,
+						.TexelFactor = RegisteredAsset.AssetPrimitiveInfo.TexelFactor,
+						.bForceLOD = bForceMipStreaming });
 
 		ObjectAssets.Add(FAssetRecord{ .AssetRegistrationIndex = AssetIndex, .AssetElementIndex = AssetElementIndex
 #if SIMPLE_STREAMABLE_ASSET_MANAGER_ALLOW_DEBUG_POINTER
@@ -675,4 +656,4 @@ void FSimpleStreamableAssetManager::UpdateBoundSizes_Impl(
 		BoundsViewInfos,
 		MaxLevelRenderAssetScreenSize
 	);
-}
+}
\ No newline at end of file
diff --git a/Engine/Source/Runtime/Engine/Public/Streaming/SimpleStreamableAssetManager.h b/Engine/Source/Runtime/Engine/Public/Streaming/SimpleStreamableAssetManager.h
--- a/Engine/Source/Runtime/Engine/Public/Streaming/SimpleStreamableAssetManager.h
+++ b/Engine/Source/Runtime/Engine/Public/Streaming/SimpleStreamableAssetManager.h
@@ -5,6 +5,9 @@
 #include "Components/PrimitiveComponent.h"
 #include "Misc/SpinLock.h"
 #include "Engine/StreamableRenderAsset.h"
+#include "Engine/StaticMesh.h"
+#include "Engine/SkeletalMesh.h"
+#include "Engine/TextureDefines.h"
 #include "Containers/ConsumeAllMpmcQueue.h"
 
 class FPrimitiveSceneProxy;
@@ -164,10 +167,16 @@ public:
 				, bForceMipStreaming(InForceMipStreaming)
 		{}
 	};
-	
+
+	struct FRegisteredAsset
+	{
+		TSharedPtr<int32, ESPMode::ThreadSafe> SimpleStreamableAssetManagerIndex = nullptr;
+		FStreamingRenderAssetPrimitiveInfo AssetPrimitiveInfo;
+	};
+
 	struct FRegister : public FUpdate
 	{
-		TArray<FStreamingRenderAssetPrimitiveInfo> Assets{};
+		TArray<FRegisteredAsset> RegisteredAssets{};
 
 		FRegister() = default;
 		FRegister(const FRegister&) = default;
@@ -319,7 +328,7 @@ private:
 	int32 UsedAssetIndices = 0;
 	TBitArray<> AssetUsedIndices;
 	TArray<TSimpleSparseArray<FAssetBoundElement>> AssetIndexToBounds4Index;
-	void AddRenderAssetElements(const TArrayView<FStreamingRenderAssetPrimitiveInfo>& RenderAssetInstanceInfos, int32 ObjectRegistrationIndex, bool bForceMipStreaming);
+	void AddRenderAssetElements(const TArrayView<FRegisteredAsset>& RegisteredAssets, int32 ObjectRegistrationIndex, bool bForceMipStreaming);
 	void RemoveRenderAssetElements(int32 ObjectRegistrationIndex);
 
 	// Used to update last render time
@@ -410,6 +419,7 @@ FSimpleStreamableAssetManager::FRegister::FRegister(const TObject* InObject, con
 		InObject->GetPrimitiveSceneInfo()->LastRenderTime,
 		InObject->IsForceMipStreaming())
 {
+	TArray<FStreamingRenderAssetPrimitiveInfo> Assets;
 	if (InObject->IsSupportingStreamableRenderAssetsGathering())
 	{
 		InObject->GetStreamableRenderAssetInfo(InPrimitive->Bounds, Assets);
@@ -423,6 +433,28 @@ FSimpleStreamableAssetManager::FRegister::FRegister(const TObject* InObject, con
 		// Primitive has to support at least one of the two ways of providing assets
 		checkNoEntry(); 
 	}
+
+	// Ignore not needed assets
+	RegisteredAssets.Reserve(Assets.Num());
+	for (FStreamingRenderAssetPrimitiveInfo& Info : Assets)
+	{
+		if (Info.RenderAsset && Info.RenderAsset->IsStreamable())
+		{
+#if DO_CHECK
+			ensure(Info.TexelFactor >= 0.f
+				|| Info.RenderAsset->IsA<UStaticMesh>()
+				|| Info.RenderAsset->IsA<USkeletalMesh>()
+				|| (Info.RenderAsset->IsA<UTexture>() && Info.RenderAsset->GetLODGroupForStreaming() == TEXTUREGROUP_Terrain_Heightmap));
+#endif
+			// Otherwise check that everything is setup right. If the component is not yet registered, then the bound data is irrelevant.
+			if (bForceMipStreaming || Info.CanBeStreamedByDistance(true) || Info.TexelFactor < 0.f)
+			{
+				RegisteredAssets.Add(FRegisteredAsset{
+					.SimpleStreamableAssetManagerIndex = Info.RenderAsset->SimpleStreamableAssetManagerIndex,
+					.AssetPrimitiveInfo = MoveTemp(Info)});
+			}
+		}
+	}
 }
 
 inline int32 GetTypeHash(const FSimpleStreamableAssetManager::FAssetRecord& Object) { return Object.AssetRegistrationIndex; }
\ No newline at end of file
