Fix check fail for VT stacks
Even though all textures in a VT stack are sampled at the same frequency, we may have one of these also sampled at the VS level, using a different stack.
Since ShaderFrequencyMask is per VT stack anyway, we may keep track of the shader frequency when we decide to allocate a new one. We just need to ensure that we do not re-use a stack with the wrong frequency. Entry.ScopeID == CurrentScopeID should already take care of this, but we might as well test the added member instead

#rnx

#Synced-CL 49554335

diff --git a/Engine/Source/Runtime/Core/Private/UObject/FortniteMainBranchObjectVersion.cpp b/Engine/Source/Runtime/Core/Private/UObject/FortniteMainBranchObjectVersion.cpp
--- a/Engine/Source/Runtime/Core/Private/UObject/FortniteMainBranchObjectVersion.cpp
+++ b/Engine/Source/Runtime/Core/Private/UObject/FortniteMainBranchObjectVersion.cpp
@@ -13,7 +13,7 @@ TMap<FGuid, FGuid> FFortniteMainBranchObjectVersion::GetSystemGuids()
 	SystemGuids.Add(DevGuids.MATERIALSHADERMAP_DERIVEDDATA_VER, FGuid("1A77CBCDFBDB4C12BF8563EDBB7D8888"));
 	SystemGuids.Add(DevGuids.NIAGARASHADERMAP_DERIVEDDATA_VER, FGuid("CC229D919A214E59B79D01AC2E5BF981"));
 	SystemGuids.Add(DevGuids.Niagara_LatestScriptCompileVersion, FGuid("FB15C2870A4E4AE8B5F58566F0E3825E"));
-	SystemGuids.Add(DevGuids.MaterialTranslationDDCVersion, FGuid("6FC7EA8335E5415E89A56113F8347F54"));
+	SystemGuids.Add(DevGuids.MaterialTranslationDDCVersion, FGuid("88DCA9D21A884B7FBB944CAB57C34EC3"));
 
 	// These GUIDs are in AppendRestrictedFortniteMainStreamObjectVersionGuids as they should rarely be changed.
 	//SystemGuids.Add(DevGuids.NANITE_DERIVEDDATA_VER, FGuid("ACF98184978742A592537DA223A1E6BE"));
diff --git a/Engine/Source/Runtime/Engine/Private/Materials/HLSLMaterialTranslator.cpp b/Engine/Source/Runtime/Engine/Private/Materials/HLSLMaterialTranslator.cpp
--- a/Engine/Source/Runtime/Engine/Private/Materials/HLSLMaterialTranslator.cpp
+++ b/Engine/Source/Runtime/Engine/Private/Materials/HLSLMaterialTranslator.cpp
@@ -2807,30 +2807,10 @@ void FHLSLMaterialTranslator::TranslateMaterial()
 
 	MaterialCompilationOutput.UniformExpressionSet.SetParameterCollections(ParameterCollections);
 
+	check(VTStacks.Num() == MaterialCompilationOutput.UniformExpressionSet.VTStacks.Num());
 	for (int VTStackIndex = 0; VTStackIndex < MaterialCompilationOutput.UniformExpressionSet.VTStacks.Num(); ++VTStackIndex)
 	{
-		FMaterialVirtualTextureStack& CurrentVTStack = MaterialCompilationOutput.UniformExpressionSet.VTStacks[VTStackIndex];
-		const uint32 NumLayers = CurrentVTStack.GetNumLayers();
-		uint32 NumLayersWithVSFrequency = 0;
-		uint32 NumLayersWithUniformExpression = 0;
-
-		for (uint32 LayerIndex = 0u; LayerIndex < NumLayers; ++LayerIndex)
-		{
-			const int32 ParameterIndex = CurrentVTStack.LayerUniformExpressionIndices[LayerIndex];
-			if (ParameterIndex != INDEX_NONE)
-			{
-				const FMaterialTextureParameterInfo& Parameter = MaterialCompilationOutput.UniformExpressionSet.GetTextureParameter(EMaterialTextureParameterType::Virtual, ParameterIndex);
-				++NumLayersWithUniformExpression;
-				if (HasShaderFrequency(Parameter.ShaderFrequencyMask, SF_Vertex))
-				{
-					++NumLayersWithVSFrequency;
-				}
-			}
-		}
-
-		// ensure that all layers share the same frequency
-		check(NumLayersWithVSFrequency == 0 || (NumLayersWithVSFrequency == NumLayersWithUniformExpression));
-		CurrentVTStack.ShaderFrequencyMask = (NumLayersWithVSFrequency > 0) ? (1<<SF_Vertex) : 0;
+		MaterialCompilationOutput.UniformExpressionSet.VTStacks[VTStackIndex].ShaderFrequencyMask = (1 << VTStacks[VTStackIndex].ShaderFrequency);
 	}
 
 	// Store the number of unique VT samples
@@ -7445,6 +7425,7 @@ uint32 FHLSLMaterialTranslator::AcquireVTStackIndex(
 		
 		if (!bStackIsFull &&
 			Entry.ScopeID == CurrentScopeID &&
+			Entry.ShaderFrequency == ShaderFrequency &&
 			Entry.CoordinateHash == CoordinateHash &&
 			Entry.MipValue0Hash == MipValue0Hash &&
 			Entry.MipValue1Hash == MipValue1Hash &&
@@ -7475,6 +7456,7 @@ uint32 FHLSLMaterialTranslator::AcquireVTStackIndex(
 	Entry.MipValue1Hash = MipValue1Hash;
 	Entry.PageTableHash = PageTableHash;
 	Entry.PageTableUniformHash = PageTableUniformHash;
+	Entry.ShaderFrequency = ShaderFrequency;
 
 	MaterialCompilationOutput.UniformExpressionSet.VTStacks.Add(FMaterialVirtualTextureStack(SampleInfo.PreallocatedStackTextureIndex));
 
diff --git a/Engine/Source/Runtime/Engine/Private/Materials/HLSLMaterialTranslator.h b/Engine/Source/Runtime/Engine/Private/Materials/HLSLMaterialTranslator.h
--- a/Engine/Source/Runtime/Engine/Private/Materials/HLSLMaterialTranslator.h
+++ b/Engine/Source/Runtime/Engine/Private/Materials/HLSLMaterialTranslator.h
@@ -275,6 +275,7 @@ protected:
 		uint64 PageTableHash;
 		uint64 PageTableUniformHash;
 		int32 CodeIndex;
+		EShaderFrequency ShaderFrequency;
 	};
 
 	/** Discovered tags during translation */
