Skeletal mesh component no refreshes the skeletal mesh asset at PostEditImport, so the skeletal mesh is not lost after a duplicate

#Synced-CL 49562108

diff --git a/Engine/Source/Runtime/Engine/Classes/Components/SkeletalMeshComponent.h b/Engine/Source/Runtime/Engine/Classes/Components/SkeletalMeshComponent.h
--- a/Engine/Source/Runtime/Engine/Classes/Components/SkeletalMeshComponent.h
+++ b/Engine/Source/Runtime/Engine/Classes/Components/SkeletalMeshComponent.h
@@ -1808,6 +1808,11 @@ public:
 	ENGINE_API virtual void Serialize(FArchive& Ar) override;
 	ENGINE_API virtual void PostLoad() override;
 	ENGINE_API virtual void PostInitProperties() override;
+
+#if WITH_EDITOR
+	ENGINE_API virtual void PostEditImport() override;
+#endif
+
 #if UE_WITH_REMOTE_OBJECT_HANDLE
 	ENGINE_API virtual void PostMigrate(const struct FUObjectMigrationContext& MigrationContext) override;
 #endif
diff --git a/Engine/Source/Runtime/Engine/Private/Components/SkeletalMeshComponent.cpp b/Engine/Source/Runtime/Engine/Private/Components/SkeletalMeshComponent.cpp
--- a/Engine/Source/Runtime/Engine/Private/Components/SkeletalMeshComponent.cpp
+++ b/Engine/Source/Runtime/Engine/Private/Components/SkeletalMeshComponent.cpp
@@ -671,6 +671,16 @@ void USkeletalMeshComponent::PostInitProperties()
 	}
 }
 
+#if WITH_EDITOR
+void USkeletalMeshComponent::PostEditImport()
+{
+	Super::PostEditImport();
+
+	// Update property alias after a duplicate (PostDuplicate not called on actor components)
+	RefreshSkeletalMeshAsset();
+}
+#endif
+
 #if UE_WITH_REMOTE_OBJECT_HANDLE
 void USkeletalMeshComponent::PostMigrate(const struct FUObjectMigrationContext& MigrationContext)
 {
