Fixed missing uniform access on VTCs

#Synced-CL 49558267

diff --git a/Engine/Source/Runtime/Engine/Private/Materials/HLSLMaterialTranslator.cpp b/Engine/Source/Runtime/Engine/Private/Materials/HLSLMaterialTranslator.cpp
--- a/Engine/Source/Runtime/Engine/Private/Materials/HLSLMaterialTranslator.cpp
+++ b/Engine/Source/Runtime/Engine/Private/Materials/HLSLMaterialTranslator.cpp
@@ -8506,6 +8506,8 @@ int32 FHLSLMaterialTranslator::TextureFromCollection(int32 TextureCollectionCode
 		{
 			return Errorf(TEXT("Texture collection virtual sampling only supports 2d textures"));
 		}
+
+		AccessUniformExpression(TextureCollectionCodeIndex);
 		
 		FString Uniforms = GetParameterCode(AddInternalCodeChunk(
 			MCT_Unexposed, TEXT("FIndirectVirtualTextureUniform"),
