Substrate - only rename the opacity input fro the root node when the front mateiral is used.

#tests Pc editor
#rnx
#Synced-CL 49615688

diff --git a/Engine/Source/Runtime/Engine/Private/Materials/MaterialAttributeDefinitionMap.cpp b/Engine/Source/Runtime/Engine/Private/Materials/MaterialAttributeDefinitionMap.cpp
--- a/Engine/Source/Runtime/Engine/Private/Materials/MaterialAttributeDefinitionMap.cpp
+++ b/Engine/Source/Runtime/Engine/Private/Materials/MaterialAttributeDefinitionMap.cpp
@@ -442,13 +442,14 @@ FText FMaterialAttributeDefinitionMap::GetAttributeOverrideForMaterial(const FGu
 	EMaterialProperty Property = GMaterialPropertyAttributesMap.Find(AttributeID)->Property;
 
 	const bool bSubstrateEnabled = Substrate::IsSubstrateEnabled();
+	const bool bHasSubstrateFrontMaterialConnected = Material->HasSubstrateFrontMaterialConnected();
 
 	switch (Property)
 	{
 	case MP_EmissiveColor:
 		return Material->IsUIMaterial() ? LOCTEXT("UIOutputColor", "Final Color") : LOCTEXT("EmissiveColor", "Emissive Color");
 	case MP_Opacity:
-		return bSubstrateEnabled ? LOCTEXT("OpacityOverride", "Opacity Override") : LOCTEXT("Opacity", "Opacity");
+		return bSubstrateEnabled && bHasSubstrateFrontMaterialConnected ? LOCTEXT("OpacityOverride", "Opacity Override") : LOCTEXT("Opacity", "Opacity");
 	case MP_OpacityMask:
 		return LOCTEXT("OpacityMask", "Opacity Mask");
 	case MP_DiffuseColor:
