ObjectTrace: Allow to pass user defined outers when creating the outer chain

#rnx

#Synced-CL 49623288

diff --git a/Engine/Source/Runtime/Engine/Private/ObjectTrace.cpp b/Engine/Source/Runtime/Engine/Private/ObjectTrace.cpp
--- a/Engine/Source/Runtime/Engine/Private/ObjectTrace.cpp
+++ b/Engine/Source/Runtime/Engine/Private/ObjectTrace.cpp
@@ -475,7 +475,7 @@ void FObjectTrace::OutputInstance(const UObject* InOuterObject, uint64 InInstanc
 	OutputInstanceInternal(InOuterObject, InInstanceId, InOuterId, InType, InName, InPathName, Flags);
 }
 
-bool FObjectTrace::OutputObject(const UObject* InObject)
+bool FObjectTrace::OutputObject(const UObject* InObject, UObject* InUserDefinedOuter)
 {
 	if (UE::ObjectTrace::IsTracingDisabledForObject(InObject))
 	{
@@ -494,8 +494,9 @@ bool FObjectTrace::OutputObject(const UObject* InObject)
 		Flags |= UE::ObjectTrace::ObjectInfoFlag_Transient;
 	}
 
-	OutputObject(InObject->GetOuter());
-	OutputInstanceInternal(InObject, InstanceId, GetObjectId(InObject->GetOuter()), InObject->GetClass(), InObject->GetFName().ToString(), InObject->GetPathName(), Flags);
+	UObject* Outer = InUserDefinedOuter ? InUserDefinedOuter : InObject->GetOuter();	
+	OutputObject(Outer);
+	OutputInstanceInternal(InObject, InstanceId, GetObjectId(Outer), InObject->GetClass(), InObject->GetFName().ToString(), InObject->GetPathName(), Flags);
 	return true;
 }
 
@@ -531,9 +532,9 @@ void FObjectTrace::OutputObjectEvent(const UObject* InObject, const TCHAR* InEve
 		<< ObjectEvent.Event(InEvent);
 }
 
-void FObjectTrace::OutputObjectLifetimeBegin(const UObject* InObject)
+void FObjectTrace::OutputObjectLifetimeBegin(const UObject* InObject, UObject* InUserDefinedOuter)
 {
-	TRACE_OBJECT(InObject);
+	TRACE_OBJECT_WITH_OUTER(InObject, InUserDefinedOuter);
 }
 
 void FObjectTrace::OutputInstanceLifetimeEnd(const UObject* InOuterObject, uint64 InInstanceId)
diff --git a/Engine/Source/Runtime/Engine/Public/ObjectTrace.h b/Engine/Source/Runtime/Engine/Public/ObjectTrace.h
--- a/Engine/Source/Runtime/Engine/Public/ObjectTrace.h
+++ b/Engine/Source/Runtime/Engine/Public/ObjectTrace.h
@@ -67,7 +67,7 @@ struct FObjectTrace
 	};
 
 	/** Helper function to output an object */
-	ENGINE_API static bool OutputObject(const UObject* InObject);
+	ENGINE_API static bool OutputObject(const UObject* InObject, UObject* InUserDefinedOuter = nullptr);
 	
 	/** Helper function to output an "object instance" which is not an actual UObject */
 	ENGINE_API static void OutputInstance(const UObject* InOuterObject, uint64 InInstanceId, uint64 InOuterId, UStruct* InType, const FString& InName, const FString& InPathName = FString());
@@ -76,7 +76,7 @@ struct FObjectTrace
 	ENGINE_API static void OutputObjectTransform(const UObject* InObject, const FTransform& InTransform, ETeleportType InTeleportType);
 
 	/** Helper function to output object creation event */
-	ENGINE_API static void OutputObjectLifetimeBegin(const UObject* InObject);
+	ENGINE_API static void OutputObjectLifetimeBegin(const UObject* InObject, UObject* InUserDefinedOuter = nullptr);
 	
 	/** Helper function to output object destruction event */
 	ENGINE_API static void OutputInstanceLifetimeEnd(const UObject* InOuterObject, uint64 InInstanceId);
@@ -136,6 +136,9 @@ struct FObjectTrace
 #define TRACE_OBJECT(Object) \
 	FObjectTrace::OutputObject(Object);
 
+#define TRACE_OBJECT_WITH_OUTER(Object, UserDefinedOuter) \
+	FObjectTrace::OutputObject(Object, UserDefinedOuter);
+
 #define TRACE_INSTANCE(OuterObject, InstanceId, OuterId, Type, Name) \
 	FObjectTrace::OutputInstance(OuterObject, InstanceId, OuterId, Type, Name);
 
@@ -153,6 +156,9 @@ struct FObjectTrace
 #define TRACE_OBJECT_LIFETIME_BEGIN(Object) \
 	if (CAN_TRACE_OBJECT(Object)) { FObjectTrace::OutputObjectLifetimeBegin(Object); }
 
+#define TRACE_OBJECT_LIFETIME_BEGIN_WITH_OUTER(Object, UserDefinedOuter) \
+	if (CAN_TRACE_OBJECT(Object)) { FObjectTrace::OutputObjectLifetimeBegin(Object, UserDefinedOuter); }
+
 #define TRACE_OBJECT_LIFETIME_END(Object) \
 	if (CAN_TRACE_OBJECT(Object)) { FObjectTrace::OutputObjectLifetimeEnd(Object); }
 
