Collision Filter Refactor Part 5.5: Final part of 64 bit extensions. Updating the the collision channels, editor, serialization, and CVD to work with 64 bit.

#Synced-CL 49621724

diff --git a/Engine/Plugins/ChaosVD/Source/ChaosVD/Private/DetailsCustomizations/ChaosVDDetailsCustomizationUtils.cpp b/Engine/Plugins/ChaosVD/Source/ChaosVD/Private/DetailsCustomizations/ChaosVDDetailsCustomizationUtils.cpp
--- a/Engine/Plugins/ChaosVD/Source/ChaosVD/Private/DetailsCustomizations/ChaosVDDetailsCustomizationUtils.cpp
+++ b/Engine/Plugins/ChaosVD/Source/ChaosVD/Private/DetailsCustomizations/ChaosVDDetailsCustomizationUtils.cpp
@@ -4,6 +4,7 @@
 
 #include "ChaosVDModule.h"
 #include "ChaosVDScene.h"
+#include "CollisionQueryParams.h"
 #include "DetailLayoutBuilder.h"
 #include "DetailWidgetRow.h"
 #include "DataWrappers/ChaosVDParticleDataWrapper.h"
@@ -111,6 +112,7 @@ TSharedPtr<FChaosVDCollisionChannelsInfoContainer> FChaosVDDetailsCustomizationU
 
 	const int32 NumEnum = Enum->NumEnums();
 	constexpr int32 ExpectedChannels = FChaosVDDetailsCustomizationUtils::GetMaxCollisionChannelIndex();
+	NewCollisionChannelsInfoContainer->CollisionChannelInfos.SetNum(ExpectedChannels);
 
 	if (NumEnum < ExpectedChannels)
 	{
@@ -128,7 +130,7 @@ TSharedPtr<FChaosVDCollisionChannelsInfoContainer> FChaosVDDetailsCustomizationU
 		Info.CollisionChannel = static_cast<ECollisionChannel>(ChannelIndex);
 		Info.bIsTraceType = Enum->GetMetaData(*TraceType, ChannelIndex) == TEXT("1");
 
-		NewCollisionChannelsInfoContainer->CustomChannelsNames[ChannelIndex] = Info;
+		NewCollisionChannelsInfoContainer->CollisionChannelInfos[ChannelIndex] = Info;
 	}
 
 	return NewCollisionChannelsInfoContainer;
@@ -221,9 +223,7 @@ void FChaosVDDetailsCustomizationUtils::BuildCollisionChannelMatrix(const FChaos
 	IDetailGroup& TraceResponsesGroup = ParentCategoryGroup.AddGroup(TEXT("CollisionTraceResponses"), LOCTEXT("CollisionTraceResponsesLabel", "Trace Responses"), bStartExpanded);
 	TraceResponsesGroup.EnableReset(false);
 
-	constexpr int32 ExpectedChannels = FChaosVDDetailsCustomizationUtils::GetMaxCollisionChannelIndex();
-
-	for (int32 ChannelIndex = 0; ChannelIndex < ExpectedChannels; ++ChannelIndex)
+	for (int32 ChannelIndex = 0; ChannelIndex < CollisionChannelsInfo.Num(); ++ChannelIndex)
 	{
 		const FChaosVDCollisionChannelInfo& ChannelInfo = CollisionChannelsInfo[ChannelIndex];
 		if (ChannelInfo.bIsTraceType)
@@ -235,7 +235,7 @@ void FChaosVDDetailsCustomizationUtils::BuildCollisionChannelMatrix(const FChaos
 	IDetailGroup& CollisionResponsesGroup = ParentCategoryGroup.AddGroup(TEXT("CollisionObjectResponses"), LOCTEXT("CollisionObjectResponses", "Object Responses"), bStartExpanded);
 	TraceResponsesGroup.EnableReset(false);
 
-	for (int32 ChannelIndex = 0; ChannelIndex < ExpectedChannels; ++ChannelIndex)
+	for (int32 ChannelIndex = 0; ChannelIndex < CollisionChannelsInfo.Num(); ++ChannelIndex)
 	{
 		const FChaosVDCollisionChannelInfo& ChannelInfo = CollisionChannelsInfo[ChannelIndex];
 		if (!ChannelInfo.bIsTraceType)
@@ -263,6 +263,27 @@ TSharedRef<SWidget> FChaosVDDetailsCustomizationUtils::CreateCollisionResponseMa
 		];
 }
 
+void FChaosVDDetailsCustomizationUtils::ExtractResponseContainer(const uint64 BlockMask, const uint64 OverlapMask, FCollisionResponseContainer& OutResponses)
+{
+	constexpr uint32 ChannelCount = UE_ARRAY_COUNT(FCollisionResponseContainer::EnumArray);
+	static_assert(ChannelCount == 64);
+
+	for (uint32 ChannelIndex = 0; ChannelIndex < ChannelCount; ++ChannelIndex)
+	{
+		const uint64 ChannelBit = ECC_TO_BITFIELD(ChannelIndex);
+		uint8 Response = ECR_Ignore;
+		if (BlockMask & ChannelBit)
+		{
+			Response = ECR_Block;
+		}
+		else if (OverlapMask & ChannelBit)
+		{
+			Response = ECR_Overlap;
+		}
+		OutResponses.EnumArray[ChannelIndex] = Response;
+	}
+}
+
 void FChaosVDDetailsCustomizationUtils::AddWidgetRowForCheckboxValue(TAttribute<ECheckBoxState>&& State, const FText& InValueName, IDetailGroup& DetailGroup)
 {
 	DetailGroup.AddWidgetRow()
diff --git a/Engine/Plugins/ChaosVD/Source/ChaosVD/Private/DetailsCustomizations/ChaosVDDetailsCustomizationUtils.h b/Engine/Plugins/ChaosVD/Source/ChaosVD/Private/DetailsCustomizations/ChaosVDDetailsCustomizationUtils.h
--- a/Engine/Plugins/ChaosVD/Source/ChaosVD/Private/DetailsCustomizations/ChaosVDDetailsCustomizationUtils.h
+++ b/Engine/Plugins/ChaosVD/Source/ChaosVD/Private/DetailsCustomizations/ChaosVDDetailsCustomizationUtils.h
@@ -113,26 +113,13 @@ public:
 
 	static constexpr int32 GetMaxCollisionChannelIndex();
 
-	template<typename CVDCollisionFilteringData>
-	static FCollisionFilterData ConvertToEngineFilteringData(const CVDCollisionFilteringData& InCVDFilteringData);
+	UE_INTERNAL static void ExtractResponseContainer(const uint64 BlockMask, const uint64 OverlapMask, FCollisionResponseContainer& OutResponses);
 
 	static void AddWidgetRowForCheckboxValue(TAttribute<ECheckBoxState>&& State, const FText& InValueName, IDetailGroup& DetailGroup);
 
 	static FText GetDefaultCollisionChannelsUseWarningMessage();
 };
 
-template <typename CVDCollisionFilteringData>
-FCollisionFilterData FChaosVDDetailsCustomizationUtils::ConvertToEngineFilteringData(const CVDCollisionFilteringData& InCVDFilteringData)
-{
-	FCollisionFilterData EngineFilteringData;
-	EngineFilteringData.Word0 = InCVDFilteringData.Word0;
-	EngineFilteringData.Word1 = InCVDFilteringData.Word1;
-	EngineFilteringData.Word2 = InCVDFilteringData.Word2;
-	EngineFilteringData.Word3 = InCVDFilteringData.Word3;
-
-	return EngineFilteringData;
-}
-
 constexpr int32 FChaosVDDetailsCustomizationUtils::GetMaxCollisionChannelIndex()
 {
 	// Skip the deprecated channel which is last
diff --git a/Engine/Plugins/ChaosVD/Source/ChaosVD/Private/DetailsCustomizations/ChaosVDQueryDataWrappersCustomizationDetails.cpp b/Engine/Plugins/ChaosVD/Source/ChaosVD/Private/DetailsCustomizations/ChaosVDQueryDataWrappersCustomizationDetails.cpp
--- a/Engine/Plugins/ChaosVD/Source/ChaosVD/Private/DetailsCustomizations/ChaosVDQueryDataWrappersCustomizationDetails.cpp
+++ b/Engine/Plugins/ChaosVD/Source/ChaosVD/Private/DetailsCustomizations/ChaosVDQueryDataWrappersCustomizationDetails.cpp
@@ -286,7 +286,10 @@ ECollisionResponse FChaosVDCollisionResponseParamsCustomization::GetCurrentColli
 
 	if (FChaosVDCollisionResponseParams* Data = CurrentPropertyDataHandle ? CurrentPropertyDataHandle->GetDataInstance() : nullptr)
 	{
-		return static_cast<ECollisionResponse>(Data->FlagsPerChannel[ChannelIndex]);
+		if (ensure(ChannelIndex < Data->FlagsPerChannel.Num()))
+		{
+			return static_cast<ECollisionResponse>(Data->FlagsPerChannel[ChannelIndex]);
+		}
 	}
 
 	return ECollisionResponse::ECR_MAX;
@@ -344,7 +347,7 @@ void FChaosVDCollisionResponseParamsCustomization::CustomizeChildren(TSharedRef<
 	CollisionChannelStateGetter.BindSP(this, &FChaosVDCollisionResponseParamsCustomization::GetCurrentCollisionResponseForChannel);
 	
 	FChaosVDDetailsCustomizationUtils::BuildCollisionChannelMatrix(CollisionChannelStateGetter,
-																	{ CachedCollisionChannelInfos->CustomChannelsNames, UE_ARRAY_COUNT(CachedCollisionChannelInfos->CustomChannelsNames)},
+																	CachedCollisionChannelInfos->CollisionChannelInfos,
 																	CollisionGroup);
 }
 
@@ -374,10 +377,9 @@ void FChaosVDCollisionObjectParamsCustomization::CustomizeChildren(TSharedRef<IP
 	IDetailGroup& CollisionGroup = StructBuilder.AddGroup(TEXT("CollisionObjectResponseParams"), LOCTEXT("CollisionObjectResponseQueryParamsLabel", "Collision Response Query Params") );
 	CollisionGroup.EnableReset(false);
 
-	constexpr int32 MaxChannels = FChaosVDDetailsCustomizationUtils::GetMaxCollisionChannelIndex();
-	for (int32 ChannelIndex = 0 ; ChannelIndex < MaxChannels; ChannelIndex++)
+	for (int32 ChannelIndex = 0 ; ChannelIndex < CachedCollisionChannelInfos->CollisionChannelInfos.Num(); ChannelIndex++)
 	{
-		FChaosVDCollisionChannelInfo& ChannelInfo = CachedCollisionChannelInfos->CustomChannelsNames[ChannelIndex];
+		FChaosVDCollisionChannelInfo& ChannelInfo = CachedCollisionChannelInfos->CollisionChannelInfos[ChannelIndex];
 
 		if (ChannelInfo.bIsTraceType)
 		{
@@ -416,7 +418,7 @@ ECheckBoxState FChaosVDCollisionObjectParamsCustomization::GetCurrentObjectFlag(
 		return ECheckBoxState::Undetermined;
 	}
 
-	return (CollisionObjectResponseParams->ObjectTypesToQuery & ECC_TO_BITFIELD(ChannelIndex)) != 0 ? ECheckBoxState::Checked : ECheckBoxState::Unchecked;
+	return (CollisionObjectResponseParams->ObjectTypesToQuery64 & ECC_TO_BITFIELD(ChannelIndex)) != 0 ? ECheckBoxState::Checked : ECheckBoxState::Unchecked;
 }
 
 #undef LOCTEXT_NAMESPACE
diff --git a/Engine/Plugins/ChaosVD/Source/ChaosVD/Private/DetailsCustomizations/ChaosVDShapeDataCustomization.cpp b/Engine/Plugins/ChaosVD/Source/ChaosVD/Private/DetailsCustomizations/ChaosVDShapeDataCustomization.cpp
--- a/Engine/Plugins/ChaosVD/Source/ChaosVD/Private/DetailsCustomizations/ChaosVDShapeDataCustomization.cpp
+++ b/Engine/Plugins/ChaosVD/Source/ChaosVD/Private/DetailsCustomizations/ChaosVDShapeDataCustomization.cpp
@@ -92,7 +92,7 @@ void FChaosVDShapeDataCustomization::CustomizeChildren(TSharedRef<IPropertyHandl
 		return;
 	}
 
-	CacheCurrentCollisionChannelData();
+	CacheCollisionFilterData();
 	
 	TSharedPtr<IPropertyHandle> CollisionTraceTypeHandle = StructPropertyHandle->GetChildHandle(GET_MEMBER_NAME_CHECKED(FChaosVDShapeCollisionData, CollisionTraceType));
 	CollisionCategoryBuilder.AddProperty(CollisionTraceTypeHandle.ToSharedRef());
@@ -139,9 +139,15 @@ void FChaosVDShapeDataCustomization::CustomizeChildren(TSharedRef<IPropertyHandl
 		]
 	];
 
-	BuildCollisionFilerDataLayout(ECollisionFilterDataType::Query, LOCTEXT("QueryFilterDataRowLabel", "Query Filter Data"), ParentLayoutBuilder, ECollisionFilterLayoutFlags::StartExpanded);
-
-	BuildCollisionFilerDataLayout(ECollisionFilterDataType::Sim, LOCTEXT("SimFilterDataRowLabel", "Sim Filter Data"), ParentLayoutBuilder, ECollisionFilterLayoutFlags::StartExpanded);
+	if (!bFiltersRequireSplitView)
+	{
+		BuildCollisionFilerDataLayout(ECollisionFilterDataType::Query, LOCTEXT("QueryFilterDataRowLabel", "Filter Data"), ParentLayoutBuilder, ECollisionFilterLayoutFlags::StartExpanded);
+	}
+	else
+	{
+		BuildCollisionFilerDataLayout(ECollisionFilterDataType::Query, LOCTEXT("QueryFilterDataRowLabel", "Query Filter Data"), ParentLayoutBuilder, ECollisionFilterLayoutFlags::StartExpanded);
+		BuildCollisionFilerDataLayout(ECollisionFilterDataType::Sim, LOCTEXT("SimFilterDataRowLabel", "Sim Filter Data"), ParentLayoutBuilder, ECollisionFilterLayoutFlags::StartExpanded);
+	}
 }
 
 void FChaosVDShapeDataCustomization::BuildCollisionFilerDataLayout(ECollisionFilterDataType FilterDataType, const FText& InDetailsGroupLabel, IDetailLayoutBuilder& ParentLayoutBuilder, ECollisionFilterLayoutFlags LayoutFlags)
@@ -205,7 +211,7 @@ void FChaosVDShapeDataCustomization::BuildCollisionFilerDataLayout(ECollisionFil
 	});
 
 	FChaosVDDetailsCustomizationUtils::BuildCollisionChannelMatrix(CollisionSimChannelStateGetter,
-	                                                               { CachedCollisionChannelInfos->CustomChannelsNames, UE_ARRAY_COUNT(CachedCollisionChannelInfos->CustomChannelsNames)},
+	                                                               CachedCollisionChannelInfos->CollisionChannelInfos,
 	                                                               CollisionResponseGroup);
 }
 
@@ -249,11 +255,8 @@ FText FChaosVDShapeDataCustomization::GetCurrentCollisionChannelAsText(ECollisio
 {
 	if (FChaosVDShapeCollisionData* Data = CurrentShapeDataHandle ? CurrentShapeDataHandle->GetDataInstance() : nullptr)
 	{
-		FChaosVDCollisionFilterData& FilterDataRef = FilterDataType == ECollisionFilterDataType::Query ? Data->QueryData : Data->SimData;
-		uint8 QueryExtraFilter = 0;
-		PRAGMA_DISABLE_DEPRECATION_WARNINGS
-		return UEnum::GetDisplayValueAsText(GetCollisionChannelAndExtraFilter(FilterDataRef.Word3, QueryExtraFilter));
-		PRAGMA_ENABLE_DEPRECATION_WARNINGS
+		const FCachedFilterData& FilterData = CachedFilterData[FilterDataType];
+		return UEnum::GetDisplayValueAsText((ECollisionChannel)FilterData.ChannelIndex);
 	}
 
 	return FText::GetEmpty();
@@ -263,13 +266,8 @@ uint8 FChaosVDShapeDataCustomization::GetCurrentExtraFilter(ECollisionFilterData
 {
 	if (FChaosVDShapeCollisionData* Data = CurrentShapeDataHandle ? CurrentShapeDataHandle->GetDataInstance() : nullptr)
 	{
-		FChaosVDCollisionFilterData& FilterDataRef = FilterDataType== ECollisionFilterDataType::Query ? Data->QueryData : Data->SimData;
-		uint8 QueryExtraFilter = 0;
-		PRAGMA_DISABLE_DEPRECATION_WARNINGS
-		GetCollisionChannelAndExtraFilter(FilterDataRef.Word3, QueryExtraFilter);
-		PRAGMA_ENABLE_DEPRECATION_WARNINGS
-	
-		return QueryExtraFilter;
+		const FCachedFilterData& FilterData = CachedFilterData[FilterDataType];
+		return FilterData.MaskFilter;
 	}
 
 	return 0;
@@ -277,37 +275,25 @@ uint8 FChaosVDShapeDataCustomization::GetCurrentExtraFilter(ECollisionFilterData
 
 ECollisionResponse FChaosVDShapeDataCustomization::GetCurrentCollisionResponseForChannel(int32 ChannelIndex, ECollisionFilterDataType Type) const
 {
-	constexpr int32 MaxChannels = FChaosVDDetailsCustomizationUtils::GetMaxCollisionChannelIndex();
-
-	if (ChannelIndex >= MaxChannels)
+	if (CurrentShapeDataHandle == nullptr || CurrentShapeDataHandle->GetDataInstance() == nullptr)
 	{
 		return ECollisionResponse::ECR_MAX;
 	}
 
-	if (!CurrentShapeDataHandle->GetDataInstance())
+	const FCachedFilterData& FilterData = CachedFilterData[Type];
+	if (ChannelIndex < 0 || ChannelIndex >= UE_ARRAY_COUNT(FilterData.Responses.EnumArray))
 	{
 		return ECollisionResponse::ECR_MAX;
 	}
-
-	switch (Type)
-	{
-		case Query:
-			return static_cast<ECollisionResponse>(CachedQueryCollisionResponsesFlagsPerChannel.EnumArray[ChannelIndex]);
-		case Sim:
-			return static_cast<ECollisionResponse>(CachedSimCollisionResponsesFlagsPerChannel.EnumArray[ChannelIndex]);
-		default:
-			break;
-	}
-
-	return ECollisionResponse::ECR_MAX;
+	return static_cast<ECollisionResponse>(FilterData.Responses.EnumArray[ChannelIndex]);
 }
 
 ECheckBoxState FChaosVDShapeDataCustomization::GetCurrentStateForFilteringFlag(ECollisionFilterDataType Type, Chaos::EFilterFlags FilteringFlag) const
 {
 	if (FChaosVDShapeCollisionData* Data = CurrentShapeDataHandle ? CurrentShapeDataHandle->GetDataInstance() : nullptr)
 	{
-		FChaosVDCollisionFilterData& FilterDataRef = Type == Query ? Data->QueryData : Data->SimData;
-		return FilterDataRef.Word3 & static_cast<uint8>(FilteringFlag) ? ECheckBoxState::Checked : ECheckBoxState::Unchecked;
+		const FCachedFilterData& FilterData = CachedFilterData[Type];
+		return EnumHasAnyFlags(FilterData.Flags, FilteringFlag) ? ECheckBoxState::Checked : ECheckBoxState::Unchecked;
 	}
 	
 	return ECheckBoxState::Undetermined;
@@ -337,17 +323,48 @@ void FChaosVDShapeDataCustomization::AddCollisionFilterFlagRow(ECollisionFilterD
 	FChaosVDDetailsCustomizationUtils::AddWidgetRowForCheckboxValue(MoveTemp(CheckBoxStateAttribute), FText::FromString(Chaos::LexToString(Flag)), DetailGroup);
 }
 
-void FChaosVDShapeDataCustomization::CacheCurrentCollisionChannelData()
+void FChaosVDShapeDataCustomization::CacheCollisionFilterData()
 {
-	if (FChaosVDShapeCollisionData* Data = CurrentShapeDataHandle ? CurrentShapeDataHandle->GetDataInstance() : nullptr)
+	const FChaosVDShapeCollisionData* Data = CurrentShapeDataHandle ? CurrentShapeDataHandle->GetDataInstance() : nullptr;
+	if (Data == nullptr)
+	{
+		return;
+	}
+
+	const FChaosVDCollisionFilterDataV2& FilterData = Data->FilterData;
+	constexpr uint32 ChannelCount = ECollisionChannel::ECC_OverlapAll_Deprecated;
+	static_assert(ChannelCount == 64);
+	static_assert(UE_ARRAY_COUNT(FCollisionResponseContainer::EnumArray) == ChannelCount);
+
+	FCachedFilterData& QueryFilterData = CachedFilterData[ECollisionFilterDataType::Query];
+	FCachedFilterData& SimFilterData = CachedFilterData[ECollisionFilterDataType::Sim];
+	QueryFilterData = FCachedFilterData();
+	SimFilterData = FCachedFilterData();
+	bFiltersRequireSplitView = false;
+
+	// Convert back to channel index for now
+	if (FilterData.CollisionChannelMask != 0)
 	{
-		PRAGMA_DISABLE_DEPRECATION_WARNINGS
-		FCollisionFilterData QueryFilterData = FChaosVDDetailsCustomizationUtils::ConvertToEngineFilteringData(Data->QueryData);
-		CachedQueryCollisionResponsesFlagsPerChannel = ExtractQueryCollisionResponseContainer(QueryFilterData);
+		QueryFilterData.ChannelIndex = (uint8)FMath::CountTrailingZeros64(FilterData.CollisionChannelMask);
+	}
+	QueryFilterData.Flags = (Chaos::EFilterFlags)FilterData.GetFlags();
+	QueryFilterData.MaskFilter = FilterData.GetMaskFilter();
+	FChaosVDDetailsCustomizationUtils::ExtractResponseContainer(FilterData.BlockMask, FilterData.OverlapMask, QueryFilterData.Responses);
 
-		FCollisionFilterData SimFilterData = FChaosVDDetailsCustomizationUtils::ConvertToEngineFilteringData(Data->SimData);
-		CachedSimCollisionResponsesFlagsPerChannel = ExtractSimCollisionResponseContainer(SimFilterData);
-		PRAGMA_ENABLE_DEPRECATION_WARNINGS
+	if (FilterData.bLoadedFromLegacy)
+	{
+		uint32 SimFlags = 0;
+		FilterData.ExtractWord3(FilterData.LegacySimWord3, SimFilterData.MaskFilter, SimFilterData.ChannelIndex, SimFlags);
+		SimFilterData.Flags = (Chaos::EFilterFlags)SimFlags;
+		FChaosVDDetailsCustomizationUtils::ExtractResponseContainer(FilterData.LegacySimBlockMask, 0, SimFilterData.Responses);
+		bFiltersRequireSplitView = SimFilterData.Flags != QueryFilterData.Flags
+			|| SimFilterData.ChannelIndex != QueryFilterData.ChannelIndex
+			|| SimFilterData.MaskFilter != QueryFilterData.MaskFilter
+			|| FilterData.BlockMask != FilterData.LegacySimBlockMask;
+	}
+	else
+	{
+		SimFilterData = QueryFilterData;
 	}
 }
 
@@ -374,7 +391,7 @@ void FChaosVDShapeDataCustomization::HandleSceneUpdated()
 {
 	// To avoid reconstructing the data channel flags every time a row in the collision channel response matrix is updated, we process it once and cache it
 	// but we need to update it each time the CVD scene is updated as the source data of the currents elected particle might have changed.
-	CacheCurrentCollisionChannelData();
+	CacheCollisionFilterData();
 
 	// Note: Intentionally not calling UpdateCollisionChannelsInfoCache, that is only updated once when a recording is loaded, therefore we don't nee to do it again
 }
diff --git a/Engine/Plugins/ChaosVD/Source/ChaosVD/Private/DetailsCustomizations/ChaosVDShapeDataCustomization.h b/Engine/Plugins/ChaosVD/Source/ChaosVD/Private/DetailsCustomizations/ChaosVDShapeDataCustomization.h
--- a/Engine/Plugins/ChaosVD/Source/ChaosVD/Private/DetailsCustomizations/ChaosVDShapeDataCustomization.h
+++ b/Engine/Plugins/ChaosVD/Source/ChaosVD/Private/DetailsCustomizations/ChaosVDShapeDataCustomization.h
@@ -36,7 +36,8 @@ private:
 	enum ECollisionFilterDataType
 	{
 		Query,
-		Sim
+		Sim,
+		Count,
 	};
 
 	enum class ECollisionFilterLayoutFlags : uint8
@@ -48,6 +49,14 @@ private:
 	};
 	FRIEND_ENUM_CLASS_FLAGS(ECollisionFilterLayoutFlags)
 
+	struct FCachedFilterData
+	{
+		uint8 ChannelIndex = 0;
+		uint8 MaskFilter = 0;
+		Chaos::EFilterFlags Flags = Chaos::EFilterFlags::None;
+		FCollisionResponseContainer Responses;
+	};
+
 	void BuildCollisionFilerDataLayout(ECollisionFilterDataType FilterDataType, const FText& InDetailsGroupLabel, IDetailLayoutBuilder& ParentLayoutBuilder, ECollisionFilterLayoutFlags LayoutFlags);
 
 	FText GetCurrentCollisionChannelAsText(ECollisionFilterDataType FilterDataType) const;
@@ -60,7 +69,7 @@ private:
 
 	void AddExtraCollisionFilterFlags(ECollisionFilterDataType Type, const FText& IntSectionLabel, IDetailCategoryBuilder& CategoryBuilder);
 
-	void CacheCurrentCollisionChannelData();
+	void CacheCollisionFilterData();
 
 	void RegisterCVDScene(const TSharedPtr<FChaosVDScene>& InScene);
 
@@ -69,8 +78,8 @@ private:
 	TSharedPtr<FChaosVDCollisionChannelsInfoContainer> CachedCollisionChannelInfos;
 	bool bChannelInfoBuiltFromDefaults = true;
 	
-	FCollisionResponseContainer CachedSimCollisionResponsesFlagsPerChannel;
-	FCollisionResponseContainer CachedQueryCollisionResponsesFlagsPerChannel;
+	TStaticArray<FCachedFilterData, ECollisionFilterDataType::Count> CachedFilterData;
+	bool bFiltersRequireSplitView = false;
 
 	TSharedPtr<FChaosVDDetailsPropertyDataHandle<FChaosVDShapeCollisionData>> CurrentShapeDataHandle = nullptr;
 
diff --git a/Engine/Source/Editor/DetailCustomizations/Private/CollisionProfileDetails.cpp b/Engine/Source/Editor/DetailCustomizations/Private/CollisionProfileDetails.cpp
--- a/Engine/Source/Editor/DetailCustomizations/Private/CollisionProfileDetails.cpp
+++ b/Engine/Source/Editor/DetailCustomizations/Private/CollisionProfileDetails.cpp
@@ -60,10 +60,10 @@ DECLARE_DELEGATE_RetVal_OneParam(bool, FOnValidateChannel, const FCustomChannelS
 DECLARE_DELEGATE_RetVal_TwoParams(bool, FOnValidateProfile, const FCollisionResponseTemplate*, int32)
 
 constexpr ECollisionChannel MinCustomCollisionChannel = ECC_GameTraceChannel1;
-constexpr ECollisionChannel MaxCustomCollisionChannel = ECC_GameTraceChannel18;
+constexpr ECollisionChannel MaxCustomCollisionChannel = ECC_GameTraceChannel50;
 constexpr uint32 MaxNumCustomCollisionChannels = MaxCustomCollisionChannel - MinCustomCollisionChannel + 1;
 constexpr uint32 MaxNumCollisionChannels = ECC_OverlapAll_Deprecated;
-static_assert(MaxNumCollisionChannels == 32);
+static_assert(MaxNumCollisionChannels == 64);
 
 #define COLLIISION_COLUMN_WIDTH		50.f
 
@@ -1404,7 +1404,7 @@ void FCollisionProfileDetails::CustomizeDetails( IDetailLayoutBuilder& DetailBui
 				.AutoWrapText(true)
 				.Font(IDetailLayoutBuilder::GetDetailFont())
 				.ToolTip(ObjectChannelTooltip)
-				.Text(LOCTEXT("ObjectChannel_Menu_Description", "You can have up to 18 custom channels including object and trace channels. This is the list of object types for your project. If you delete an object type that is being used by the game, any uses of that type will revert to WorldStatic."))
+				.Text(LOCTEXT("ObjectChannel_Menu_Description", "You can have up to 50 custom channels including object and trace channels. This is the list of object types for your project. If you delete an object type that is being used by the game, any uses of that type will revert to WorldStatic."))
 			]
 
 			+SHorizontalBox::Slot()
@@ -1501,7 +1501,7 @@ void FCollisionProfileDetails::CustomizeDetails( IDetailLayoutBuilder& DetailBui
 				.Font(IDetailLayoutBuilder::GetDetailFont())
 				.AutoWrapText(true)
 				.ToolTip(TraceChannelTooltip)
-				.Text(LOCTEXT("TraceChannel_Menu_Description", "You can have up to 18 custom channels including object and trace channels. This is the list of trace channels for your project. If you delete a trace channel that is being used by the game, the behavior of the trace is undefined."))
+				.Text(LOCTEXT("TraceChannel_Menu_Description", "You can have up to 50 custom channels including object and trace channels. This is the list of trace channels for your project. If you delete a trace channel that is being used by the game, the behavior of the trace is undefined."))
 			]
 
 			+SHorizontalBox::Slot()
diff --git a/Engine/Source/Runtime/Core/Public/UObject/FortniteMainBranchObjectVersions.inl b/Engine/Source/Runtime/Core/Public/UObject/FortniteMainBranchObjectVersions.inl
--- a/Engine/Source/Runtime/Core/Public/UObject/FortniteMainBranchObjectVersions.inl
+++ b/Engine/Source/Runtime/Core/Public/UObject/FortniteMainBranchObjectVersions.inl
@@ -730,5 +730,8 @@ FORTNITE_MAIN_VERSION(LandscapeTexturePatchAlphaOverride, 238)
 // For SpawnActor/AddComponent, allow template to be used for the spawned actor.
 FORTNITE_MAIN_VERSION(PCGAllowTemplateWhenSpawningFromAttribute, 239)
 
+// Chaos collision filters extension from 32 to 64 bit
+FORTNITE_MAIN_VERSION(CollisionFilter64Bit, 240)
+
 // -----<new versions can be added above this line>-------------------------------------------------
 #undef FORTNITE_MAIN_VERSION
diff --git a/Engine/Source/Runtime/Engine/Classes/Engine/EngineTypes.h b/Engine/Source/Runtime/Engine/Classes/Engine/EngineTypes.h
--- a/Engine/Source/Runtime/Engine/Classes/Engine/EngineTypes.h
+++ b/Engine/Source/Runtime/Engine/Classes/Engine/EngineTypes.h
@@ -1132,6 +1132,38 @@ enum ECollisionChannel : int
 	ECC_GameTraceChannel16 UMETA(Hidden),
 	ECC_GameTraceChannel17 UMETA(Hidden),
 	ECC_GameTraceChannel18 UMETA(Hidden),
+	ECC_GameTraceChannel19 UMETA(Hidden),
+	ECC_GameTraceChannel20 UMETA(Hidden),
+	ECC_GameTraceChannel21 UMETA(Hidden),
+	ECC_GameTraceChannel22 UMETA(Hidden),
+	ECC_GameTraceChannel23 UMETA(Hidden),
+	ECC_GameTraceChannel24 UMETA(Hidden),
+	ECC_GameTraceChannel25 UMETA(Hidden),
+	ECC_GameTraceChannel26 UMETA(Hidden),
+	ECC_GameTraceChannel27 UMETA(Hidden),
+	ECC_GameTraceChannel28 UMETA(Hidden),
+	ECC_GameTraceChannel29 UMETA(Hidden),
+	ECC_GameTraceChannel30 UMETA(Hidden),
+	ECC_GameTraceChannel31 UMETA(Hidden),
+	ECC_GameTraceChannel32 UMETA(Hidden),
+	ECC_GameTraceChannel33 UMETA(Hidden),
+	ECC_GameTraceChannel34 UMETA(Hidden),
+	ECC_GameTraceChannel35 UMETA(Hidden),
+	ECC_GameTraceChannel36 UMETA(Hidden),
+	ECC_GameTraceChannel37 UMETA(Hidden),
+	ECC_GameTraceChannel38 UMETA(Hidden),
+	ECC_GameTraceChannel39 UMETA(Hidden),
+	ECC_GameTraceChannel40 UMETA(Hidden),
+	ECC_GameTraceChannel41 UMETA(Hidden),
+	ECC_GameTraceChannel42 UMETA(Hidden),
+	ECC_GameTraceChannel43 UMETA(Hidden),
+	ECC_GameTraceChannel44 UMETA(Hidden),
+	ECC_GameTraceChannel45 UMETA(Hidden),
+	ECC_GameTraceChannel46 UMETA(Hidden),
+	ECC_GameTraceChannel47 UMETA(Hidden),
+	ECC_GameTraceChannel48 UMETA(Hidden),
+	ECC_GameTraceChannel49 UMETA(Hidden),
+	ECC_GameTraceChannel50 UMETA(Hidden),
 	
 	/** Add new serializeable channels above here (i.e. entries that exist in FCollisionResponseContainer) */
 	/** Add only nonserialized/transient flags below */
@@ -1199,6 +1231,38 @@ enum EObjectTypeQuery : int
 	ObjectTypeQuery30 UMETA(Hidden), 
 	ObjectTypeQuery31 UMETA(Hidden), 
 	ObjectTypeQuery32 UMETA(Hidden),
+	ObjectTypeQuery33 UMETA(Hidden), 
+	ObjectTypeQuery34 UMETA(Hidden), 
+	ObjectTypeQuery35 UMETA(Hidden), 
+	ObjectTypeQuery36 UMETA(Hidden), 
+	ObjectTypeQuery37 UMETA(Hidden), 
+	ObjectTypeQuery38 UMETA(Hidden), 
+	ObjectTypeQuery39 UMETA(Hidden), 
+	ObjectTypeQuery40 UMETA(Hidden), 
+	ObjectTypeQuery41 UMETA(Hidden), 
+	ObjectTypeQuery42 UMETA(Hidden),
+	ObjectTypeQuery43 UMETA(Hidden), 
+	ObjectTypeQuery44 UMETA(Hidden), 
+	ObjectTypeQuery45 UMETA(Hidden), 
+	ObjectTypeQuery46 UMETA(Hidden), 
+	ObjectTypeQuery47 UMETA(Hidden), 
+	ObjectTypeQuery48 UMETA(Hidden), 
+	ObjectTypeQuery49 UMETA(Hidden), 
+	ObjectTypeQuery50 UMETA(Hidden),
+	ObjectTypeQuery51 UMETA(Hidden),
+	ObjectTypeQuery52 UMETA(Hidden),
+	ObjectTypeQuery53 UMETA(Hidden),
+	ObjectTypeQuery54 UMETA(Hidden),
+	ObjectTypeQuery55 UMETA(Hidden),
+	ObjectTypeQuery56 UMETA(Hidden),
+	ObjectTypeQuery57 UMETA(Hidden),
+	ObjectTypeQuery58 UMETA(Hidden),
+	ObjectTypeQuery59 UMETA(Hidden),
+	ObjectTypeQuery60 UMETA(Hidden),
+	ObjectTypeQuery61 UMETA(Hidden),
+	ObjectTypeQuery62 UMETA(Hidden),
+	ObjectTypeQuery63 UMETA(Hidden),
+	ObjectTypeQuery64 UMETA(Hidden),
 
 	ObjectTypeQuery_MAX	UMETA(Hidden)
 };
@@ -1239,6 +1303,38 @@ enum ETraceTypeQuery : int
 	TraceTypeQuery30 UMETA(Hidden), 
 	TraceTypeQuery31 UMETA(Hidden), 
 	TraceTypeQuery32 UMETA(Hidden),
+	TraceTypeQuery33 UMETA(Hidden), 
+	TraceTypeQuery34 UMETA(Hidden), 
+	TraceTypeQuery35 UMETA(Hidden), 
+	TraceTypeQuery36 UMETA(Hidden), 
+	TraceTypeQuery37 UMETA(Hidden), 
+	TraceTypeQuery38 UMETA(Hidden), 
+	TraceTypeQuery39 UMETA(Hidden), 
+	TraceTypeQuery40 UMETA(Hidden), 
+	TraceTypeQuery41 UMETA(Hidden), 
+	TraceTypeQuery42 UMETA(Hidden),
+	TraceTypeQuery43 UMETA(Hidden), 
+	TraceTypeQuery44 UMETA(Hidden), 
+	TraceTypeQuery45 UMETA(Hidden), 
+	TraceTypeQuery46 UMETA(Hidden), 
+	TraceTypeQuery47 UMETA(Hidden), 
+	TraceTypeQuery48 UMETA(Hidden), 
+	TraceTypeQuery49 UMETA(Hidden), 
+	TraceTypeQuery50 UMETA(Hidden),
+	TraceTypeQuery51 UMETA(Hidden),
+	TraceTypeQuery52 UMETA(Hidden),
+	TraceTypeQuery53 UMETA(Hidden),
+	TraceTypeQuery54 UMETA(Hidden),
+	TraceTypeQuery55 UMETA(Hidden),
+	TraceTypeQuery56 UMETA(Hidden),
+	TraceTypeQuery57 UMETA(Hidden),
+	TraceTypeQuery58 UMETA(Hidden),
+	TraceTypeQuery59 UMETA(Hidden),
+	TraceTypeQuery60 UMETA(Hidden),
+	TraceTypeQuery61 UMETA(Hidden),
+	TraceTypeQuery62 UMETA(Hidden),
+	TraceTypeQuery63 UMETA(Hidden),
+	TraceTypeQuery64 UMETA(Hidden),
 
 	TraceTypeQuery_MAX	UMETA(Hidden)
 };
@@ -1464,7 +1560,7 @@ struct FCollisionResponseContainer
 	TEnumAsByte<enum ECollisionResponse> GameTraceChannel15;    // 28
 
 	UPROPERTY(EditAnywhere, BlueprintReadOnly, Category=CollisionResponseContainer)
-	TEnumAsByte<enum ECollisionResponse> GameTraceChannel16;    // 28
+	TEnumAsByte<enum ECollisionResponse> GameTraceChannel16;    // 29
 
 	UPROPERTY(EditAnywhere, BlueprintReadOnly, Category=CollisionResponseContainer)
 	TEnumAsByte<enum ECollisionResponse> GameTraceChannel17;    // 30
@@ -1472,6 +1568,102 @@ struct FCollisionResponseContainer
 	UPROPERTY(EditAnywhere, BlueprintReadOnly, Category=CollisionResponseContainer)
 	TEnumAsByte<enum ECollisionResponse> GameTraceChannel18;    // 31
 
+	UPROPERTY(EditAnywhere, BlueprintReadOnly, Category=CollisionResponseContainer)
+	TEnumAsByte<enum ECollisionResponse> GameTraceChannel19;    // 32
+
+	UPROPERTY(EditAnywhere, BlueprintReadOnly, Category=CollisionResponseContainer)
+	TEnumAsByte<enum ECollisionResponse> GameTraceChannel20;    // 33
+
+	UPROPERTY(EditAnywhere, BlueprintReadOnly, Category=CollisionResponseContainer)
+	TEnumAsByte<enum ECollisionResponse> GameTraceChannel21;    // 34
+
+	UPROPERTY(EditAnywhere, BlueprintReadOnly, Category=CollisionResponseContainer)
+	TEnumAsByte<enum ECollisionResponse> GameTraceChannel22;    // 35
+
+	UPROPERTY(EditAnywhere, BlueprintReadOnly, Category=CollisionResponseContainer)
+	TEnumAsByte<enum ECollisionResponse> GameTraceChannel23;    // 36
+
+	UPROPERTY(EditAnywhere, BlueprintReadOnly, Category=CollisionResponseContainer)
+	TEnumAsByte<enum ECollisionResponse> GameTraceChannel24;    // 37
+
+	UPROPERTY(EditAnywhere, BlueprintReadOnly, Category=CollisionResponseContainer)
+	TEnumAsByte<enum ECollisionResponse> GameTraceChannel25;    // 38
+
+	UPROPERTY(EditAnywhere, BlueprintReadOnly, Category=CollisionResponseContainer)
+	TEnumAsByte<enum ECollisionResponse> GameTraceChannel26;    // 39
+
+	UPROPERTY(EditAnywhere, BlueprintReadOnly, Category=CollisionResponseContainer)
+	TEnumAsByte<enum ECollisionResponse> GameTraceChannel27;    // 40
+
+	UPROPERTY(EditAnywhere, BlueprintReadOnly, Category=CollisionResponseContainer)
+	TEnumAsByte<enum ECollisionResponse> GameTraceChannel28;    // 41
+
+	UPROPERTY(EditAnywhere, BlueprintReadOnly, Category=CollisionResponseContainer)
+	TEnumAsByte<enum ECollisionResponse> GameTraceChannel29;    // 42
+
+	UPROPERTY(EditAnywhere, BlueprintReadOnly, Category=CollisionResponseContainer)
+	TEnumAsByte<enum ECollisionResponse> GameTraceChannel30;    // 43
+
+	UPROPERTY(EditAnywhere, BlueprintReadOnly, Category=CollisionResponseContainer)
+	TEnumAsByte<enum ECollisionResponse> GameTraceChannel31;    // 44
+
+	UPROPERTY(EditAnywhere, BlueprintReadOnly, Category=CollisionResponseContainer)
+	TEnumAsByte<enum ECollisionResponse> GameTraceChannel32;    // 45
+
+	UPROPERTY(EditAnywhere, BlueprintReadOnly, Category=CollisionResponseContainer)
+	TEnumAsByte<enum ECollisionResponse> GameTraceChannel33;    // 46
+
+	UPROPERTY(EditAnywhere, BlueprintReadOnly, Category=CollisionResponseContainer)
+	TEnumAsByte<enum ECollisionResponse> GameTraceChannel34;    // 47
+
+	UPROPERTY(EditAnywhere, BlueprintReadOnly, Category=CollisionResponseContainer)
+	TEnumAsByte<enum ECollisionResponse> GameTraceChannel35;    // 48
+
+	UPROPERTY(EditAnywhere, BlueprintReadOnly, Category=CollisionResponseContainer)
+	TEnumAsByte<enum ECollisionResponse> GameTraceChannel36;    // 49
+
+	UPROPERTY(EditAnywhere, BlueprintReadOnly, Category=CollisionResponseContainer)
+	TEnumAsByte<enum ECollisionResponse> GameTraceChannel37;    // 50
+
+	UPROPERTY(EditAnywhere, BlueprintReadOnly, Category=CollisionResponseContainer)
+	TEnumAsByte<enum ECollisionResponse> GameTraceChannel38;    // 51
+
+	UPROPERTY(EditAnywhere, BlueprintReadOnly, Category=CollisionResponseContainer)
+	TEnumAsByte<enum ECollisionResponse> GameTraceChannel39;    // 52
+
+	UPROPERTY(EditAnywhere, BlueprintReadOnly, Category=CollisionResponseContainer)
+	TEnumAsByte<enum ECollisionResponse> GameTraceChannel40;    // 53
+
+	UPROPERTY(EditAnywhere, BlueprintReadOnly, Category=CollisionResponseContainer)
+	TEnumAsByte<enum ECollisionResponse> GameTraceChannel41;    // 54
+
+	UPROPERTY(EditAnywhere, BlueprintReadOnly, Category=CollisionResponseContainer)
+	TEnumAsByte<enum ECollisionResponse> GameTraceChannel42;    // 55
+
+	UPROPERTY(EditAnywhere, BlueprintReadOnly, Category=CollisionResponseContainer)
+	TEnumAsByte<enum ECollisionResponse> GameTraceChannel43;    // 56
+
+	UPROPERTY(EditAnywhere, BlueprintReadOnly, Category=CollisionResponseContainer)
+	TEnumAsByte<enum ECollisionResponse> GameTraceChannel44;    // 57
+
+	UPROPERTY(EditAnywhere, BlueprintReadOnly, Category=CollisionResponseContainer)
+	TEnumAsByte<enum ECollisionResponse> GameTraceChannel45;    // 58
+
+	UPROPERTY(EditAnywhere, BlueprintReadOnly, Category=CollisionResponseContainer)
+	TEnumAsByte<enum ECollisionResponse> GameTraceChannel46;    // 59
+
+	UPROPERTY(EditAnywhere, BlueprintReadOnly, Category=CollisionResponseContainer)
+	TEnumAsByte<enum ECollisionResponse> GameTraceChannel47;    // 60
+
+	UPROPERTY(EditAnywhere, BlueprintReadOnly, Category=CollisionResponseContainer)
+	TEnumAsByte<enum ECollisionResponse> GameTraceChannel48;    // 61
+
+	UPROPERTY(EditAnywhere, BlueprintReadOnly, Category=CollisionResponseContainer)
+	TEnumAsByte<enum ECollisionResponse> GameTraceChannel49;    // 62
+
+	UPROPERTY(EditAnywhere, BlueprintReadOnly, Category=CollisionResponseContainer)
+	TEnumAsByte<enum ECollisionResponse> GameTraceChannel50;    // 63
+
 #endif
 
 	union
@@ -1515,8 +1707,40 @@ struct FCollisionResponseContainer
 			uint8 GameTraceChannel16;    // 29 
 			uint8 GameTraceChannel17;    // 30
 			uint8 GameTraceChannel18;    // 31
+			uint8 GameTraceChannel19;    // 32
+			uint8 GameTraceChannel20;    // 33
+			uint8 GameTraceChannel21;    // 34
+			uint8 GameTraceChannel22;    // 35
+			uint8 GameTraceChannel23;    // 36
+			uint8 GameTraceChannel24;    // 37
+			uint8 GameTraceChannel25;    // 38
+			uint8 GameTraceChannel26;    // 39
+			uint8 GameTraceChannel27;    // 40
+			uint8 GameTraceChannel28;    // 41
+			uint8 GameTraceChannel29;    // 42
+			uint8 GameTraceChannel30;    // 43
+			uint8 GameTraceChannel31;    // 44
+			uint8 GameTraceChannel32;    // 45
+			uint8 GameTraceChannel33;    // 46
+			uint8 GameTraceChannel34;    // 47
+			uint8 GameTraceChannel35;    // 48
+			uint8 GameTraceChannel36;    // 49
+			uint8 GameTraceChannel37;    // 50
+			uint8 GameTraceChannel38;    // 51
+			uint8 GameTraceChannel39;    // 52
+			uint8 GameTraceChannel40;    // 53
+			uint8 GameTraceChannel41;    // 54
+			uint8 GameTraceChannel42;    // 55
+			uint8 GameTraceChannel43;    // 56
+			uint8 GameTraceChannel44;    // 57
+			uint8 GameTraceChannel45;    // 58
+			uint8 GameTraceChannel46;    // 59
+			uint8 GameTraceChannel47;    // 60
+			uint8 GameTraceChannel48;    // 61
+			uint8 GameTraceChannel49;    // 62
+			uint8 GameTraceChannel50;    // 63
 		};
-		uint8 EnumArray[32];
+		uint8 EnumArray[64];
 	};
 
 	/** This constructor will set all channels to ECR_Block */
diff --git a/Engine/Source/Runtime/Engine/Private/Collision/CollisionProfile.cpp b/Engine/Source/Runtime/Engine/Private/Collision/CollisionProfile.cpp
--- a/Engine/Source/Runtime/Engine/Private/Collision/CollisionProfile.cpp
+++ b/Engine/Source/Runtime/Engine/Private/Collision/CollisionProfile.cpp
@@ -10,7 +10,7 @@
 
 DEFINE_LOG_CATEGORY_STATIC(LogCollisionProfile, Warning, All)
 
-constexpr uint32 MaxNumCollisionChannels = 32;
+constexpr uint32 MaxNumCollisionChannels = 64;
 #define IS_VALID_COLLISIONCHANNEL(x) ((x)> ECC_Destructible && (x) < ECC_OverlapAll_Deprecated)
 
 // do not chnage this name. This value is serialize to other object, if you change, it will mess up serialization, you'll need to fix up name by versioning
diff --git a/Engine/Source/Runtime/Engine/Private/Collision/PhysicsFiltering.cpp b/Engine/Source/Runtime/Engine/Private/Collision/PhysicsFiltering.cpp
--- a/Engine/Source/Runtime/Engine/Private/Collision/PhysicsFiltering.cpp
+++ b/Engine/Source/Runtime/Engine/Private/Collision/PhysicsFiltering.cpp
@@ -3,8 +3,11 @@
 #include "Physics/PhysicsFiltering.h"
 #include "CollisionQueryParams.h"
 
+#include "Chaos/Vector32Int8.h"
+
 void ExtractResponses(const FCollisionResponseContainer& InResponses, uint64& BlockingBits, uint64& TouchingBits)
 {
+#if !PLATFORM_ENABLE_VECTORINTRINSICS
 	BlockingBits = 0;
 	TouchingBits = 0;
 	constexpr uint64 Count = UE_ARRAY_COUNT(InResponses.EnumArray);
@@ -20,6 +23,27 @@ void ExtractResponses(const FCollisionResponseContainer& InResponses, uint64& Bl
 			TouchingBits |= ChannelBit;
 		}
 	}
+#else
+	using namespace Chaos;
+	BlockingBits = 0;
+	TouchingBits = 0;
+	static_assert(UE_ARRAY_COUNT(InResponses.EnumArray) == 64);
+
+	const VectorRegister32Int8 OverlapRegister = Vector32Int8Set1(ECR_Overlap);
+	const VectorRegister32Int8 BlockRegister = Vector32Int8Set1(ECR_Block);
+	const int8* Data = (const int8*)(InResponses.EnumArray);
+
+	const VectorRegister32Int8 LowChunk = VectorLoad32Int8(Data);
+	const uint64 OverlapLow = VectorMaskBits(VectorCompareEQ(LowChunk, OverlapRegister));
+	const uint64 BlockLow = VectorMaskBits(VectorCompareEQ(LowChunk, BlockRegister));
+
+	const VectorRegister32Int8 HighChunk = VectorLoad32Int8(Data + 32);
+	const uint64 OverlapHigh = VectorMaskBits(VectorCompareEQ(HighChunk, OverlapRegister));
+	const uint64 BlockHigh = VectorMaskBits(VectorCompareEQ(HighChunk, BlockRegister));
+
+	TouchingBits = (OverlapHigh << 32) | OverlapLow;
+	BlockingBits = (BlockHigh << 32) | BlockLow;
+#endif
 }
 
 //////////////////////////////////////////////////////////////////////////
diff --git a/Engine/Source/Runtime/Engine/Public/CollisionQueryParams.h b/Engine/Source/Runtime/Engine/Public/CollisionQueryParams.h
--- a/Engine/Source/Runtime/Engine/Public/CollisionQueryParams.h
+++ b/Engine/Source/Runtime/Engine/Public/CollisionQueryParams.h
@@ -10,9 +10,9 @@
 #include "UObject/RemoteObjectTransfer.h"
 
 /** Macro to convert ECollisionChannels to bit flag **/
-#define ECC_TO_BITFIELD(x)	(1<<(x))
+#define ECC_TO_BITFIELD(x)	(1LLU<<(x))
 /** Macro to convert from CollisionResponseContainer to bit flag **/
-#define CRC_TO_BITFIELD(x)	(1<<(x))
+#define CRC_TO_BITFIELD(x)	(1LLU<<(x))
 
 class AActor;
 class UPrimitiveComponent;
@@ -353,8 +353,7 @@ private:
 
 	static int64 ToBitMask(ECollisionChannel QueryChannel)
 	{
-		const int32 Result = ECC_TO_BITFIELD(QueryChannel);
-		return Result;
+		return ECC_TO_BITFIELD(QueryChannel);
 	}
 
 	FCollisionQueryFlag()
@@ -465,8 +464,7 @@ private:
 
 	static int64 ToBitMask(ECollisionChannel QueryChannel)
 	{
-		const int32 Result = ECC_TO_BITFIELD(QueryChannel);
-		return Result;
+		return ECC_TO_BITFIELD(QueryChannel);
 	}
 public:
 
diff --git a/Engine/Source/Runtime/Experimental/Chaos/Private/Chaos/CollisionFilterData.cpp b/Engine/Source/Runtime/Experimental/Chaos/Private/Chaos/CollisionFilterData.cpp
--- a/Engine/Source/Runtime/Experimental/Chaos/Private/Chaos/CollisionFilterData.cpp
+++ b/Engine/Source/Runtime/Experimental/Chaos/Private/Chaos/CollisionFilterData.cpp
@@ -70,6 +70,15 @@ namespace Chaos::Filter
 		return (uint64)1 << (uint64)GetCollisionChannelIndex(Word3);
 	}
 
+	uint8 GetFirstCollisionChannelIndex(uint64 Mask)
+	{
+		if (Mask == 0)
+		{
+			return 0;
+		}
+		return (uint8)FMath::CountTrailingZeros64(Mask);
+	}
+
 	void SetCollisionChannel(uint32& Word3, uint32 Value)
 	{
 		const uint32 Word3Cleared = Word3 & ~CollisionChannelMask;
@@ -160,6 +169,12 @@ namespace Chaos::Filter
 		ComponentId = InComponentId;
 	}
 
+	void FInstanceData::Serialize(FChaosArchive& Ar)
+	{
+		Ar << OwnerId;
+		Ar << ComponentId;
+	}
+
 	FString FInstanceData::ToString() const
 	{
 		return FString::Format(TEXT("OwnerId(%d) ComponentId(%d)"), { GetOwnerId(), GetComponentId() });
@@ -256,6 +271,30 @@ namespace Chaos::Filter
 		return bBlocking ? ENarrowFilterResult::Block : ENarrowFilterResult::None;
 	}
 
+	void FShapeFilterData::Serialize(FChaosArchive& Ar)
+	{
+		Ar << BlockChannels;
+		Ar << OverlapChannels;
+
+		// Pre-emptively serialize the channel as a separate mask. This will allow easy extending to multiple channels later
+		if (Ar.IsLoading())
+		{
+			uint64 ChannelMask;
+			Ar << ChannelMask;
+			Ar << Word3;
+			uint8 ChannelIndex = GetFirstCollisionChannelIndex(ChannelMask);
+			SetCollisionChannel(Word3, ChannelIndex);
+		}
+		else
+		{
+			uint32 TempWord3 = Word3;
+			TempWord3 &= ~CollisionChannelMask;
+			uint64 ChannelMask = GetCollisionChannelMask64(Word3);
+			Ar << ChannelMask;
+			Ar << TempWord3;
+		}
+	}
+
 	FString FShapeFilterData::ToString() const
 	{
 		return FString::Format(TEXT("Block(%x) Overlap(%x) CollisionChannel(%x) MaskFilter(%x) Flags(%x)"), { GetBlockChannels(), GetOverlapChannels(), GetCollisionChannelIndex(), GetMaskFilter(), (uint8)GetFlags() });
@@ -659,6 +698,41 @@ namespace Chaos::Filter
 		return SimWord3;
 	}
 
+	FCombinedShapeFilterData FShapeFilterBuilder::LoadFromLegacySerialization(const FCollisionFilterData& QueryFilterData, const FCollisionFilterData& SimFilterData)
+	{
+		uint32 Word3 = QueryFilterData.Word3;
+		if (QueryFilterData.Word3 != SimFilterData.Word3)
+		{
+			const uint32 QueryMaskFilter = QueryFilterData.Word3 & MaskFilterMask;
+			const uint32 QueryChannelIndex = QueryFilterData.Word3 & LegacyCollisionChannelMask;
+			const uint32 QueryFlags = QueryFilterData.Word3 & LegacyQueryFlagsMask;
+			const uint32 SimMaskFilter = SimFilterData.Word3 & MaskFilterMask;
+			const uint32 SimChannelIndex = SimFilterData.Word3 & LegacyCollisionChannelMask;
+			const uint32 SimFlags = SimFilterData.Word3 & LegacySimFlagsMask;
+			UE_CLOG(QueryMaskFilter != SimMaskFilter, LogChaos, Warning, TEXT("Deserializing collision filters with different mask filters: Query(0x%X), Sim(0x%X). Query mask filter will be used."), QueryMaskFilter >> MaskFilterBitsOffset, SimMaskFilter >> MaskFilterBitsOffset);
+			UE_CLOG(QueryChannelIndex != SimChannelIndex, LogChaos, Warning, TEXT("Deserializing collision filters with different collision channels: Query(%d), Sim(%d). Query collision channel will be used."), QueryChannelIndex >> LegacyCollisionChannelBitsOffset, SimChannelIndex >> LegacyCollisionChannelBitsOffset);
+
+			const uint32 MaskFilter = QueryMaskFilter != 0 ? QueryMaskFilter : SimMaskFilter;
+			const uint32 ChannelIndex = QueryChannelIndex != 0 ? QueryChannelIndex : SimChannelIndex;
+			const uint32 Flags = QueryFlags | SimFlags;
+			Word3 = MaskFilter | ChannelIndex | Flags;
+		}
+		const uint32 QueryBlockMask = QueryFilterData.Word1;
+		const uint32 SimBlockMask = SimFilterData.Word1;
+		uint32 BlockMask = QueryBlockMask;
+		UE_CLOG(QueryBlockMask != SimBlockMask, LogChaos, Warning, TEXT("Deserializing collision filters with different block masks: Query(0x%X), Sim(0x%X). Query block mask will be used."), QueryBlockMask, SimBlockMask);
+
+		FCombinedShapeFilterData CombinedShapeFilterData;
+		CombinedShapeFilterData.InstanceData.OwnerId = QueryFilterData.Word0;
+		CombinedShapeFilterData.InstanceData.ComponentId = SimFilterData.Word2;
+
+		CombinedShapeFilterData.ShapeFilterData.BlockChannels = BlockMask;
+		CombinedShapeFilterData.ShapeFilterData.OverlapChannels = QueryFilterData.Word2;
+		CombinedShapeFilterData.ShapeFilterData.Word3 = LegacyToNewWord3(Word3);
+
+		return CombinedShapeFilterData;
+	}
+
 	void FShapeFilterBuilder::BuildFromLegacyShapeFilter(const FCollisionFilterData& QueryFilterData, const FCollisionFilterData& SimFilterData, const uint32 BlockChannels, const uint32 Word3, FCombinedShapeFilterData& OutCombinedShapeFilterData)
 	{
 		OutCombinedShapeFilterData.InstanceData.OwnerId = QueryFilterData.Word0;
diff --git a/Engine/Source/Runtime/Experimental/Chaos/Private/Chaos/ParticleDirtyFlags.cpp b/Engine/Source/Runtime/Experimental/Chaos/Private/Chaos/ParticleDirtyFlags.cpp
--- a/Engine/Source/Runtime/Experimental/Chaos/Private/Chaos/ParticleDirtyFlags.cpp
+++ b/Engine/Source/Runtime/Experimental/Chaos/Private/Chaos/ParticleDirtyFlags.cpp
@@ -4,69 +4,90 @@
 
 namespace Chaos
 {
+	PRAGMA_DISABLE_DEPRECATION_WARNINGS
+	constexpr uint32 QueryDataStart = offsetof(FCollisionData, QueryData);
+	constexpr uint32 SimDataStart = offsetof(FCollisionData, SimData);
+	constexpr uint32 SimDataEnd = SimDataStart + sizeof(FCollisionFilterData);
+	constexpr uint32 ShapeFilterDataStart = QueryDataStart + offsetof(FCollisionFilterData, Word2);
+	// Validate the total size is sufficient to store the two new structs
+	static_assert(sizeof(Filter::FInstanceData) + sizeof(Filter::FShapeFilterData) <= 2 * sizeof(FCollisionFilterData));
+	// Validate the instance data and shape filters fit in the correct spots
+	static_assert(sizeof(Filter::FInstanceData) <= ShapeFilterDataStart - QueryDataStart);
+	static_assert(sizeof(Filter::FShapeFilterData) <= SimDataEnd - ShapeFilterDataStart);
+	PRAGMA_ENABLE_DEPRECATION_WARNINGS
+
 	FCollisionFilterData FCollisionData::GetQueryData() const
 	{
-		PRAGMA_DISABLE_DEPRECATION_WARNINGS
-		return QueryData;
-		PRAGMA_ENABLE_DEPRECATION_WARNINGS
+		return Filter::FShapeFilterBuilder::GetLegacyShapeQueryFilter(GetCombinedShapeFilterData());
 	}
 
 	void FCollisionData::SetQueryData(const FCollisionFilterData& InQueryData)
 	{
-		PRAGMA_DISABLE_DEPRECATION_WARNINGS
-		QueryData = InQueryData;
-		SimData.Word3 = QueryData.Word3 = Chaos::Filter::FShapeFilterBuilder::GetWord3FromNewQueryFilter(QueryData.Word3, SimData.Word3);
-		PRAGMA_ENABLE_DEPRECATION_WARNINGS
+		Filter::FCombinedShapeFilterData CombinedFilterData = GetCombinedShapeFilterData();
+		Filter::FShapeFilterBuilder::SetLegacyShapeQueryFilter(CombinedFilterData, InQueryData);
+		SetCombinedShapeFilterData(CombinedFilterData);
 	}
 
 	FCollisionFilterData FCollisionData::GetSimData() const
 	{
-		PRAGMA_DISABLE_DEPRECATION_WARNINGS
-		return SimData;
-		PRAGMA_ENABLE_DEPRECATION_WARNINGS
+		return Filter::FShapeFilterBuilder::GetLegacyShapeSimFilter(GetCombinedShapeFilterData());
 	}
 
 	void FCollisionData::SetSimData(const FCollisionFilterData& InSimData)
 	{
-		PRAGMA_DISABLE_DEPRECATION_WARNINGS
-		SimData = InSimData;
-		SimData.Word3 = QueryData.Word3 = Chaos::Filter::FShapeFilterBuilder::GetWord3FromNewSimFilter(SimData.Word3, QueryData.Word3);
-		PRAGMA_ENABLE_DEPRECATION_WARNINGS
+		Filter::FCombinedShapeFilterData CombinedFilterData = GetCombinedShapeFilterData();
+		Filter::FShapeFilterBuilder::SetLegacyShapeSimFilter(CombinedFilterData, InSimData);
+		SetCombinedShapeFilterData(CombinedFilterData);
 	}
 
 	Chaos::Filter::FShapeFilterData FCollisionData::GetShapeFilterData() const
 	{
-		return GetCombinedShapeFilterData().GetShapeFilterData();
+		PRAGMA_DISABLE_DEPRECATION_WARNINGS
+		Filter::FShapeFilterData ShapeFilter;
+		memcpy(&ShapeFilter, &QueryData.Word2, sizeof(Filter::FShapeFilterData));
+		return ShapeFilter;
+		PRAGMA_ENABLE_DEPRECATION_WARNINGS
 	}
 
 	void FCollisionData::SetShapeFilterData(const Chaos::Filter::FShapeFilterData& ShapeFilter)
 	{
-		Chaos::Filter::FCombinedShapeFilterData CombinedShapeFilter = GetCombinedShapeFilterData();
-		CombinedShapeFilter.SetShapeFilterData(ShapeFilter);
-		SetCombinedShapeFilterData(CombinedShapeFilter);
+		PRAGMA_DISABLE_DEPRECATION_WARNINGS
+		memcpy(&QueryData.Word2, &ShapeFilter, sizeof(Filter::FShapeFilterData));
+		PRAGMA_ENABLE_DEPRECATION_WARNINGS
 	}
 
 	Chaos::Filter::FInstanceData FCollisionData::GetFilterInstanceData() const
 	{
-		return GetCombinedShapeFilterData().GetInstanceData();
+		PRAGMA_DISABLE_DEPRECATION_WARNINGS
+		Filter::FInstanceData InstanceData;
+		memcpy(&InstanceData, &QueryData, sizeof(Filter::FInstanceData));
+		return InstanceData;
+		PRAGMA_ENABLE_DEPRECATION_WARNINGS
 	}
 
 	void FCollisionData::SetFilterInstanceData(const Chaos::Filter::FInstanceData& InstanceData)
 	{
-		Chaos::Filter::FCombinedShapeFilterData CombinedShapeFilter = GetCombinedShapeFilterData();
-		CombinedShapeFilter.SetInstanceData(InstanceData);
-		SetCombinedShapeFilterData(CombinedShapeFilter);
+		PRAGMA_DISABLE_DEPRECATION_WARNINGS
+		memcpy(&QueryData, &InstanceData, sizeof(Filter::FInstanceData));
+		PRAGMA_ENABLE_DEPRECATION_WARNINGS
 	}
 
 	Chaos::Filter::FCombinedShapeFilterData FCollisionData::GetCombinedShapeFilterData() const
 	{
-		return Chaos::Filter::FShapeFilterBuilder::BuildFromLegacyShapeFilter(GetQueryData(), GetSimData());
+		PRAGMA_DISABLE_DEPRECATION_WARNINGS
+		Filter::FInstanceData InstanceData;
+		Filter::FShapeFilterData ShapeFilter;
+		memcpy(&InstanceData, &QueryData, sizeof(Filter::FInstanceData));
+		memcpy(&ShapeFilter, &QueryData.Word2, sizeof(Filter::FShapeFilterData));
+		return Filter::FCombinedShapeFilterData(ShapeFilter, InstanceData);
+		PRAGMA_ENABLE_DEPRECATION_WARNINGS
 	}
 
 	void FCollisionData::SetCombinedShapeFilterData(const Chaos::Filter::FCombinedShapeFilterData& CombinedShapeFilter)
 	{
 		PRAGMA_DISABLE_DEPRECATION_WARNINGS
-		Chaos::Filter::FShapeFilterBuilder::GetLegacyShapeFilter(CombinedShapeFilter, QueryData, SimData);
+		memcpy(&QueryData, &CombinedShapeFilter.GetInstanceData(), sizeof(Filter::FInstanceData));
+		memcpy(&QueryData.Word2, &CombinedShapeFilter.GetShapeFilterData(), sizeof(Filter::FShapeFilterData));
 		PRAGMA_ENABLE_DEPRECATION_WARNINGS
 	}
 } // namespace Chaos
diff --git a/Engine/Source/Runtime/Experimental/Chaos/Private/ChaosVisualDebugger/ChaosVDDataWrapperUtils.cpp b/Engine/Source/Runtime/Experimental/Chaos/Private/ChaosVisualDebugger/ChaosVDDataWrapperUtils.cpp
--- a/Engine/Source/Runtime/Experimental/Chaos/Private/ChaosVisualDebugger/ChaosVDDataWrapperUtils.cpp
+++ b/Engine/Source/Runtime/Experimental/Chaos/Private/ChaosVisualDebugger/ChaosVDDataWrapperUtils.cpp
@@ -671,20 +671,15 @@ void FChaosVDDataWrapperUtils::CopyShapeDataToWrapper(const Chaos::FShapeInstanc
 	OutCopyTo.bSimCollision = CollisionData.bSimCollision;
 	OutCopyTo.CollisionTraceType = static_cast<EChaosVDCollisionTraceFlag>(CollisionData.CollisionTraceType);
 
-	// TODO @ JoshD: Update
-	const Chaos::Filter::FCombinedShapeFilterData CombinedShapeFilter = CollisionData.GetCombinedShapeFilterData();
-	FCollisionFilterData QueryData, SimData;
-	Chaos::Filter::FShapeFilterBuilder::GetLegacyShapeFilter(CombinedShapeFilter, QueryData, SimData);
-
-	OutCopyTo.SimData.Word0 = SimData.Word0;
-	OutCopyTo.SimData.Word1 = SimData.Word1;
-	OutCopyTo.SimData.Word2 = SimData.Word2;
-	OutCopyTo.SimData.Word3 = SimData.Word3;
-
-	OutCopyTo.QueryData.Word0 = QueryData.Word0;
-	OutCopyTo.QueryData.Word1 = QueryData.Word1;
-	OutCopyTo.QueryData.Word2 = QueryData.Word2;
-	OutCopyTo.QueryData.Word3 = QueryData.Word3;
+	const Chaos::Filter::FInstanceData InstanceData = CollisionData.GetFilterInstanceData();
+	const Chaos::Filter::FShapeFilterData ShapeFilter = CollisionData.GetShapeFilterData();
+
+	OutCopyTo.FilterData.OwnerId = InstanceData.GetOwnerId();
+	OutCopyTo.FilterData.ComponentId = InstanceData.GetComponentId();
+	OutCopyTo.FilterData.OverlapMask = ShapeFilter.GetOverlapChannels();
+	OutCopyTo.FilterData.BlockMask = ShapeFilter.GetBlockChannels();
+	OutCopyTo.FilterData.CollisionChannelMask = ShapeFilter.GetCollisionChannelMask();
+	OutCopyTo.FilterData.SetMaskFilterAndFlags(ShapeFilter.GetMaskFilter(), (uint32)ShapeFilter.GetFlags());
 }
 
 void FChaosVDDataWrapperUtils::BuildDataWrapperFromAABBStructure(const Chaos::ISpatialAccelerationCollection<Chaos::FAccelerationStructureHandle, Chaos::FReal, 3>* SceneAccelerationStructures, int32 OwnerSolverID, TArray<FChaosVDAABBTreeDataWrapper>& OutAABBTrees)
diff --git a/Engine/Source/Runtime/Experimental/Chaos/Public/Chaos/CollisionFilterData.h b/Engine/Source/Runtime/Experimental/Chaos/Public/Chaos/CollisionFilterData.h
--- a/Engine/Source/Runtime/Experimental/Chaos/Public/Chaos/CollisionFilterData.h
+++ b/Engine/Source/Runtime/Experimental/Chaos/Public/Chaos/CollisionFilterData.h
@@ -109,6 +109,7 @@ namespace Chaos::Filter
 		UE_INTERNAL CHAOS_API uint32 GetComponentId() const;
 		UE_INTERNAL CHAOS_API void SetComponentId(const uint32 InComponentId);
 
+		CHAOS_API void Serialize(FChaosArchive& Ar);
 		CHAOS_API FString ToString() const;
 
 	private:
@@ -152,6 +153,7 @@ namespace Chaos::Filter
 		// Filter API
 		CHAOS_API ENarrowFilterResult NarrowFilter(const FShapeFilterData& OtherShape) const;
 
+		CHAOS_API void Serialize(FChaosArchive& Ar);
 		CHAOS_API FString ToString() const;
 
 		UE_INTERNAL uint32 GetLegacyWord3() const;
@@ -295,6 +297,8 @@ namespace Chaos::Filter
 		UE_INTERNAL static uint32 GetWord3FromNewSimFilter(const uint32 SimWord3, const uint32 QueryWord3);
 		UE_INTERNAL static uint32 GetWord3FromNewQueryFilter(const uint32 QueryWord3, const uint32 SimWord3);
 
+		UE_INTERNAL CHAOS_API static FCombinedShapeFilterData LoadFromLegacySerialization(const FCollisionFilterData& QueryFilterData, const FCollisionFilterData& SimFilterData);
+
 	private:
 		static void BuildFromLegacyShapeFilter(const FCollisionFilterData& QueryFilterData, const FCollisionFilterData& SimFilterData, const uint32 BlockChannels, const uint32 Word3, FCombinedShapeFilterData& OutCombinedShapeFilterData);
 
@@ -349,3 +353,15 @@ namespace Chaos::Filter
 		FQueryFilterData QueryFilterData;
 	};
 } // namespace Chaos::Filter
+
+inline Chaos::FChaosArchive& operator<<(Chaos::FChaosArchive& Ar, Chaos::Filter::FInstanceData& Data)
+{
+	Data.Serialize(Ar);
+	return Ar;
+}
+
+inline Chaos::FChaosArchive& operator<<(Chaos::FChaosArchive& Ar, Chaos::Filter::FShapeFilterData& Data)
+{
+	Data.Serialize(Ar);
+	return Ar;
+}
diff --git a/Engine/Source/Runtime/Experimental/Chaos/Public/Chaos/ParticleDirtyFlags.h b/Engine/Source/Runtime/Experimental/Chaos/Public/Chaos/ParticleDirtyFlags.h
--- a/Engine/Source/Runtime/Experimental/Chaos/Public/Chaos/ParticleDirtyFlags.h
+++ b/Engine/Source/Runtime/Experimental/Chaos/Public/Chaos/ParticleDirtyFlags.h
@@ -15,6 +15,7 @@
 #include "Chaos/RigidParticleControlFlags.h"
 #include "UObject/ExternalPhysicsCustomObjectVersion.h"
 #include "UObject/ExternalPhysicsMaterialCustomObjectVersion.h"
+#include "UObject/FortniteMainBranchObjectVersion.h"
 #include "UObject/FortniteSeasonBranchObjectVersion.h"
 #include "UObject/PhysicsObjectVersion.h"
 #include "UObject/UE5MainStreamObjectVersion.h"
@@ -770,10 +771,46 @@ struct FCollisionData
 		Ar.UsingCustomVersion(FExternalPhysicsCustomObjectVersion::GUID);
 		Ar.UsingCustomVersion(FExternalPhysicsMaterialCustomObjectVersion::GUID);
 		Ar.UsingCustomVersion(FFortniteSeasonBranchObjectVersion::GUID);
+		Ar.UsingCustomVersion(FFortniteMainBranchObjectVersion::GUID);
 
 		PRAGMA_DISABLE_DEPRECATION_WARNINGS
-		Ar << QueryData;
-		Ar << SimData;
+		if (Ar.CustomVer(FFortniteMainBranchObjectVersion::GUID) < FFortniteMainBranchObjectVersion::CollisionFilter64Bit)
+		{
+			FCollisionFilterData OldQueryData, OldSimData;
+			if(Ar.IsLoading())
+			{
+				Ar << OldQueryData;
+				Ar << OldSimData;
+				const Chaos::Filter::FCombinedShapeFilterData CombinedFilter = Chaos::Filter::FShapeFilterBuilder::LoadFromLegacySerialization(OldQueryData, OldSimData);
+				SetCombinedShapeFilterData(CombinedFilter);
+			}
+			else
+			{
+				Chaos::Filter::FCombinedShapeFilterData CombinedFilter = GetCombinedShapeFilterData();
+				Chaos::Filter::FShapeFilterBuilder::GetLegacyShapeFilter(CombinedFilter, OldQueryData, OldSimData);
+				Ar << OldQueryData;
+				Ar << OldSimData;
+			}
+		}
+		else
+		{
+			Chaos::Filter::FInstanceData InstanceData;
+			Chaos::Filter::FShapeFilterData FilterData;
+			if (Ar.IsLoading())
+			{
+				Ar << InstanceData;
+				Ar << FilterData;
+				SetFilterInstanceData(InstanceData);
+				SetShapeFilterData(FilterData);
+			}
+			else
+			{
+				InstanceData = GetFilterInstanceData();
+				FilterData = GetShapeFilterData();
+				Ar << InstanceData;
+				Ar << FilterData;
+			}
+		}
 		PRAGMA_ENABLE_DEPRECATION_WARNINGS
 
 		if (Ar.CustomVer(FExternalPhysicsCustomObjectVersion::GUID) >= FExternalPhysicsCustomObjectVersion::AddShapeSimAndQueryCollisionEnabled)
diff --git a/Engine/Source/Runtime/Experimental/ChaosSolverEngine/Private/Chaos/ChaosVDEngineEditorBridge.cpp b/Engine/Source/Runtime/Experimental/ChaosSolverEngine/Private/Chaos/ChaosVDEngineEditorBridge.cpp
--- a/Engine/Source/Runtime/Experimental/ChaosSolverEngine/Private/Chaos/ChaosVDEngineEditorBridge.cpp
+++ b/Engine/Source/Runtime/Experimental/ChaosSolverEngine/Private/Chaos/ChaosVDEngineEditorBridge.cpp
@@ -168,7 +168,7 @@ void FChaosVDEngineEditorBridge::SerializeCollisionChannelsNames()
 			Info.DisplayName = CollisionProfileData->ReturnChannelNameFromContainerIndex(ChannelIndex).ToString();
 			Info.CollisionChannel = ChannelIndex;
 			Info.bIsTraceType = CollisionProfileData->ConvertToTraceType(static_cast<ECollisionChannel>(ChannelIndex)) != TraceTypeQuery_MAX;
-			CollisionChannelInfoContainer.CustomChannelsNames[ChannelIndex] = Info;
+			CollisionChannelInfoContainer.CollisionChannelInfos.Add(Info);
 		}
 	}
 
diff --git a/Engine/Source/Runtime/Experimental/ChaosVisualDebugger/Private/DataWrappers/ChaosVDCollisionDataWrappers.cpp b/Engine/Source/Runtime/Experimental/ChaosVisualDebugger/Private/DataWrappers/ChaosVDCollisionDataWrappers.cpp
--- a/Engine/Source/Runtime/Experimental/ChaosVisualDebugger/Private/DataWrappers/ChaosVDCollisionDataWrappers.cpp
+++ b/Engine/Source/Runtime/Experimental/ChaosVisualDebugger/Private/DataWrappers/ChaosVDCollisionDataWrappers.cpp
@@ -201,6 +201,62 @@ bool FChaosVDCollisionFilterData::Serialize(FArchive& Ar)
 	return !Ar.IsError();
 }
 
+bool FChaosVDCollisionFilterDataV2::Serialize(FArchive& Ar)
+{
+	Ar << OwnerId;
+	Ar << ComponentId;
+	Ar << OverlapMask;
+	Ar << BlockMask;
+	Ar << CollisionChannelMask;
+	Ar << MaskFilterAndFlags;
+
+	return !Ar.IsError();
+}
+
+uint8 FChaosVDCollisionFilterDataV2::GetMaskFilter() const
+{
+	const uint64 Result = (MaskFilterAndFlags & MaskFilterMask) >> MaskFilterOffset;
+	return static_cast<uint8>(Result);
+}
+
+uint32 FChaosVDCollisionFilterDataV2::GetFlags() const
+{
+	return MaskFilterAndFlags & FlagsMask;
+}
+
+void FChaosVDCollisionFilterDataV2::SetMaskFilterAndFlags(uint32 InMaskFilter, uint32 InFlags)
+{
+	const uint64 MaskFilterPart = (InMaskFilter << MaskFilterOffset) & MaskFilterMask;
+	const uint64 FlagsPart = InFlags & FlagsMask;
+	MaskFilterAndFlags = MaskFilterPart | FlagsPart;
+}
+
+PRAGMA_DISABLE_DEPRECATION_WARNINGS
+void FChaosVDCollisionFilterDataV2::LoadLegacy(const FChaosVDCollisionFilterData& QueryData, const FChaosVDCollisionFilterData& SimData)
+{
+	bLoadedFromLegacy = true;
+	OwnerId = QueryData.Word0;
+	BlockMask = QueryData.Word1;
+	OverlapMask = QueryData.Word2;
+	ComponentId = SimData.Word2;
+
+	// Pull the collision channel index out and clear the old bits
+	const uint32 CollisionChannelIndex = (QueryData.Word3 & CollisionChannelIndexMask) >> CollisionChannelIndexOffset;
+	MaskFilterAndFlags = QueryData.Word3 & ~CollisionChannelIndexMask;
+	CollisionChannelMask = 1LLU << CollisionChannelIndex;
+
+	LegacySimBlockMask = SimData.Word1;
+	LegacySimWord3 = SimData.Word3;
+}
+PRAGMA_ENABLE_DEPRECATION_WARNINGS
+
+void FChaosVDCollisionFilterDataV2::ExtractWord3(const uint32 Word3, uint8& OutMaskFilter, uint8& OutChannelIndex, uint32& OutFlags)
+{
+	OutFlags = (Word3 & FlagsMask);
+	OutChannelIndex = (Word3 & CollisionChannelIndexMask) >> CollisionChannelIndexOffset;
+	OutMaskFilter = (Word3 & MaskFilterMask) >> MaskFilterOffset;
+}
+
 bool FChaosVDShapeCollisionData::Serialize(FArchive& Ar)
 {
 	Ar.UsingCustomVersion(FFortniteMainBranchObjectVersion::GUID);
@@ -229,8 +285,25 @@ bool FChaosVDShapeCollisionData::Serialize(FArchive& Ar)
 
 	if (Ar.CustomVer(FFortniteMainBranchObjectVersion::GUID) >= FFortniteMainBranchObjectVersion::SimAndQueryDataSupportInChaosVisualDebugger)
 	{
-		Ar << SimData;
-		Ar << QueryData;
+		if (Ar.CustomVer(FFortniteMainBranchObjectVersion::GUID) < FFortniteMainBranchObjectVersion::CollisionFilter64Bit)
+		{
+			PRAGMA_DISABLE_DEPRECATION_WARNINGS
+#if !WITH_EDITORONLY_DATA
+			FChaosVDCollisionFilterData SimData_DEPRECATED, QueryData_DEPRECATED;
+#endif // WITH_EDITORONLY_DATA
+
+			Ar << SimData_DEPRECATED;
+			Ar << QueryData_DEPRECATED;
+			if (Ar.IsLoading())
+			{
+				FilterData.LoadLegacy(QueryData_DEPRECATED, SimData_DEPRECATED);
+			}
+			PRAGMA_ENABLE_DEPRECATION_WARNINGS
+		}
+		else
+		{
+			Ar << FilterData;
+		}
 	}
 
 	return !Ar.IsError();
@@ -242,8 +315,7 @@ bool FChaosVDShapeCollisionData::operator==(const FChaosVDShapeCollisionData& Ot
 			&& bSimCollision == Other.bSimCollision
 			&& bQueryCollision == Other.bQueryCollision
 			&& bIsProbe == Other.bIsProbe
-			&& QueryData == Other.QueryData
-			&& SimData == Other.SimData;
+			&& FilterData == Other.FilterData;
 }
 
 bool FChaosVDCollisionChannelInfo::Serialize(FArchive& Ar)
@@ -259,7 +331,20 @@ bool FChaosVDCollisionChannelInfo::Serialize(FArchive& Ar)
 
 bool FChaosVDCollisionChannelsInfoContainer::Serialize(FArchive& Ar)
 {
-	CVD_SERIALIZE_STATIC_ARRAY(Ar, CustomChannelsNames);
+	Ar.UsingCustomVersion(FFortniteMainBranchObjectVersion::GUID);
+	if (Ar.CustomVer(FFortniteMainBranchObjectVersion::GUID) < FFortniteMainBranchObjectVersion::CollisionFilter64Bit)
+	{
+		constexpr uint32 Size = 32;
+		CollisionChannelInfos.SetNum(Size);
+		for (uint32 I = 0; I < Size; ++I)
+		{
+			Ar << CollisionChannelInfos[I];
+		}
+	}
+	else
+	{
+		Ar << CollisionChannelInfos;
+	}
 
 	return !Ar.IsError();
 }
diff --git a/Engine/Source/Runtime/Experimental/ChaosVisualDebugger/Private/DataWrappers/ChaosVDQueryDataWrappers.cpp b/Engine/Source/Runtime/Experimental/ChaosVisualDebugger/Private/DataWrappers/ChaosVDQueryDataWrappers.cpp
--- a/Engine/Source/Runtime/Experimental/ChaosVisualDebugger/Private/DataWrappers/ChaosVDQueryDataWrappers.cpp
+++ b/Engine/Source/Runtime/Experimental/ChaosVisualDebugger/Private/DataWrappers/ChaosVDQueryDataWrappers.cpp
@@ -25,13 +25,29 @@ bool FChaosVDCollisionResponseParams::Serialize(FArchive& Ar)
 
 bool FChaosVDCollisionObjectQueryParams::Serialize(FArchive& Ar)
 {
+	Ar.UsingCustomVersion(FFortniteMainBranchObjectVersion::GUID);
 	Ar << bHasValidData;
 	if (!bHasValidData)
 	{
 		return !Ar.IsError();
 	}
 
-	Ar << ObjectTypesToQuery;
+	if (Ar.CustomVer(FFortniteMainBranchObjectVersion::GUID) < FFortniteMainBranchObjectVersion::CollisionFilter64Bit)
+	{
+		uint8 OldObjectTypesToQuery;
+		Ar << OldObjectTypesToQuery;
+		ObjectTypesToQuery64 = OldObjectTypesToQuery;
+
+#if WITH_EDITORONLY_DATA
+		PRAGMA_DISABLE_DEPRECATION_WARNINGS
+		ObjectTypesToQuery_DEPRECATED = OldObjectTypesToQuery;
+		PRAGMA_ENABLE_DEPRECATION_WARNINGS
+#endif // WITH_EDITORONLY_DATA
+	}
+	else
+	{
+		Ar << ObjectTypesToQuery64;
+	}
 	Ar << IgnoreMask;
 
 	return !Ar.IsError();
diff --git a/Engine/Source/Runtime/Experimental/ChaosVisualDebugger/Public/DataWrappers/ChaosVDCollisionDataWrappers.h b/Engine/Source/Runtime/Experimental/ChaosVisualDebugger/Public/DataWrappers/ChaosVDCollisionDataWrappers.h
--- a/Engine/Source/Runtime/Experimental/ChaosVisualDebugger/Public/DataWrappers/ChaosVDCollisionDataWrappers.h
+++ b/Engine/Source/Runtime/Experimental/ChaosVisualDebugger/Public/DataWrappers/ChaosVDCollisionDataWrappers.h
@@ -394,7 +394,7 @@ enum class EChaosVDCollisionTraceFlag
 };
 
 USTRUCT()
-struct FChaosVDCollisionFilterData
+struct UE_DEPRECATED(5.8, "This struct is no longer used.") FChaosVDCollisionFilterData
 {
 	GENERATED_BODY()
 
@@ -412,6 +412,8 @@ struct FChaosVDCollisionFilterData
 	bool operator==(const FChaosVDCollisionFilterData& Other) const = default;
 };
 
+
+PRAGMA_DISABLE_DEPRECATION_WARNINGS
 template <>
 struct TTypeTraits<FChaosVDCollisionFilterData> : public TTypeTraitsBase <FChaosVDCollisionFilterData>
 {
@@ -419,6 +421,51 @@ struct TTypeTraits<FChaosVDCollisionFilterData> : public TTypeTraitsBase <FChaos
 };
 
 CVD_IMPLEMENT_SERIALIZER(FChaosVDCollisionFilterData)
+PRAGMA_ENABLE_DEPRECATION_WARNINGS
+
+struct FChaosVDCollisionFilterDataV2
+{
+	bool operator==(const FChaosVDCollisionFilterDataV2& Other) const = default;
+
+	CHAOSVDRUNTIME_API bool Serialize(FArchive& Ar);
+
+	UE_INTERNAL CHAOSVDRUNTIME_API uint8 GetMaskFilter() const;
+	UE_INTERNAL CHAOSVDRUNTIME_API uint32 GetFlags() const;
+	UE_INTERNAL CHAOSVDRUNTIME_API void SetMaskFilterAndFlags(uint32 InMaskFilter, uint32 InFlags);
+
+	PRAGMA_DISABLE_DEPRECATION_WARNINGS
+	UE_INTERNAL CHAOSVDRUNTIME_API void LoadLegacy(const FChaosVDCollisionFilterData& QueryData, const FChaosVDCollisionFilterData& SimData);
+	PRAGMA_ENABLE_DEPRECATION_WARNINGS
+
+	UE_INTERNAL CHAOSVDRUNTIME_API static void ExtractWord3(const uint32 Word3, uint8& OutMaskFilter, uint8& OutChannelIndex, uint32& OutFlags);
+
+	UE_INTERNAL uint32 OwnerId = 0;
+	UE_INTERNAL uint32 ComponentId = 0;
+	UE_INTERNAL uint64 OverlapMask = 0;
+	UE_INTERNAL uint64 BlockMask = 0;
+	UE_INTERNAL uint64 CollisionChannelMask = 0;
+	UE_INTERNAL uint32 MaskFilterAndFlags = 0;
+
+	UE_INTERNAL bool bLoadedFromLegacy = false;
+	UE_INTERNAL uint32 LegacySimBlockMask = 0;
+	UE_INTERNAL uint32 LegacySimWord3 = 0;
+
+private:
+	UE_INTERNAL static constexpr uint32 NumFlagsBits = 21;
+	UE_INTERNAL static constexpr uint32 NumCollisionChannelIndexBits = 5;
+	UE_INTERNAL static constexpr uint32 NumMaskFilterBits = 6;
+
+	PRAGMA_DISABLE_INTERNAL_WARNINGS
+	UE_INTERNAL static constexpr uint32 CollisionChannelIndexOffset = NumFlagsBits;
+	UE_INTERNAL static constexpr uint32 MaskFilterOffset = CollisionChannelIndexOffset + NumCollisionChannelIndexBits;
+
+	UE_INTERNAL static constexpr uint32 FlagsMask = ((1 << NumFlagsBits) - 1);
+	UE_INTERNAL static constexpr uint32 CollisionChannelIndexMask = ((1 << NumCollisionChannelIndexBits) - 1) << CollisionChannelIndexOffset;
+	UE_INTERNAL static constexpr uint32 MaskFilterMask = ((1 << NumMaskFilterBits) - 1) << MaskFilterOffset;
+	PRAGMA_ENABLE_INTERNAL_WARNINGS
+};
+
+CVD_IMPLEMENT_SERIALIZER(FChaosVDCollisionFilterDataV2)
 
 UENUM()
 enum class EChaosVDCollisionShapeDataFlags : uint8
@@ -433,6 +480,14 @@ USTRUCT()
 struct FChaosVDShapeCollisionData
 {
 	GENERATED_BODY()
+	
+	PRAGMA_DISABLE_DEPRECATION_WARNINGS
+	FChaosVDShapeCollisionData() = default;
+	FChaosVDShapeCollisionData(const FChaosVDShapeCollisionData&) = default;
+	FChaosVDShapeCollisionData(FChaosVDShapeCollisionData&&) = default;
+	FChaosVDShapeCollisionData& operator=(const FChaosVDShapeCollisionData&) = default;
+	FChaosVDShapeCollisionData& operator=(FChaosVDShapeCollisionData&&) = default;
+	PRAGMA_ENABLE_DEPRECATION_WARNINGS
 
 	UPROPERTY(VisibleAnywhere, Category=CollisionData)
 	EChaosVDCollisionTraceFlag CollisionTraceType = EChaosVDCollisionTraceFlag::UseDefault;
@@ -444,11 +499,19 @@ struct FChaosVDShapeCollisionData
 	UPROPERTY()
 	uint8 bIsProbe : 1 = false;
 
-	UPROPERTY(VisibleAnywhere, Category=FilterData)
-	FChaosVDCollisionFilterData QueryData;
+	#if WITH_EDITORONLY_DATA
+	PRAGMA_DISABLE_DEPRECATION_WARNINGS
+	UE_DEPRECATED(5.8, "QueryData has been deprecated. Use FilterData instead.")
+	UPROPERTY(meta = (DeprecatedProperty, DeprecationMessage = "QueryData has been deprecated"))
+	FChaosVDCollisionFilterData QueryData_DEPRECATED;
+
+	UE_DEPRECATED(5.8, "SimData has been deprecated. Use FilterData instead.")
+	UPROPERTY(meta = (DeprecatedProperty, DeprecationMessage = "SimData has been deprecated"))
+	FChaosVDCollisionFilterData SimData_DEPRECATED;
+	PRAGMA_ENABLE_DEPRECATION_WARNINGS
+	#endif // WITH_EDITORONLY_DATA
 
-	UPROPERTY(VisibleAnywhere, Category=SimData)
-	FChaosVDCollisionFilterData SimData;
+	FChaosVDCollisionFilterDataV2 FilterData;
 
 	UPROPERTY(VisibleAnywhere, Category="CVD Data")
 	bool bIsComplex = false;
@@ -497,12 +560,25 @@ struct FChaosVDCollisionChannelsInfoContainer
 {
 	GENERATED_BODY()
 	
-	UE_INTERNAL static constexpr int32 MaxSupportedChannels = 32;
+	UE_INTERNAL static constexpr int32 MaxSupportedChannels = 64;
+
+	PRAGMA_DISABLE_DEPRECATION_WARNINGS
+	FChaosVDCollisionChannelsInfoContainer() = default;
+	FChaosVDCollisionChannelsInfoContainer(const FChaosVDCollisionChannelsInfoContainer&) = default;
+	FChaosVDCollisionChannelsInfoContainer(FChaosVDCollisionChannelsInfoContainer&&) = default;
+	FChaosVDCollisionChannelsInfoContainer& operator=(const FChaosVDCollisionChannelsInfoContainer&) = default;
+	FChaosVDCollisionChannelsInfoContainer& operator=(FChaosVDCollisionChannelsInfoContainer&&) = default;
+	PRAGMA_ENABLE_DEPRECATION_WARNINGS
 
 	CHAOSVDRUNTIME_API static FStringView WrapperTypeName;
 
-	UPROPERTY(VisibleAnywhere, Category=CollisionData)
-	FChaosVDCollisionChannelInfo CustomChannelsNames[32] = {};
+	#if WITH_EDITORONLY_DATA
+	UE_DEPRECATED(5.8, "CustomChannelsNames has been deprecated")
+	UPROPERTY(meta = (DeprecatedProperty, DeprecationMessage = "CustomChannelsNames has been deprecated"))
+	FChaosVDCollisionChannelInfo CustomChannelsNames_DEPRECATED[32] = {};
+	#endif // WITH_EDITORONLY_DATA
+
+	UE_INTERNAL TArray<FChaosVDCollisionChannelInfo> CollisionChannelInfos;
 
 	CHAOSVDRUNTIME_API bool Serialize(FArchive& Ar);
 };
diff --git a/Engine/Source/Runtime/Experimental/ChaosVisualDebugger/Public/DataWrappers/ChaosVDQueryDataWrappers.h b/Engine/Source/Runtime/Experimental/ChaosVisualDebugger/Public/DataWrappers/ChaosVDQueryDataWrappers.h
--- a/Engine/Source/Runtime/Experimental/ChaosVisualDebugger/Public/DataWrappers/ChaosVDQueryDataWrappers.h
+++ b/Engine/Source/Runtime/Experimental/ChaosVisualDebugger/Public/DataWrappers/ChaosVDQueryDataWrappers.h
@@ -66,18 +66,31 @@ struct FChaosVDCollisionObjectQueryParams : public FChaosVDWrapperDataBase
 {
 	GENERATED_BODY()
 
+	PRAGMA_DISABLE_DEPRECATION_WARNINGS
+	FChaosVDCollisionObjectQueryParams() = default;
+	FChaosVDCollisionObjectQueryParams(const FChaosVDCollisionObjectQueryParams&) = default;
+	FChaosVDCollisionObjectQueryParams(FChaosVDCollisionObjectQueryParams&&) = default;
+	FChaosVDCollisionObjectQueryParams& operator=(const FChaosVDCollisionObjectQueryParams&) = default;
+	FChaosVDCollisionObjectQueryParams& operator=(FChaosVDCollisionObjectQueryParams&&) = default;
+	PRAGMA_ENABLE_DEPRECATION_WARNINGS
+
 	CHAOSVDRUNTIME_API bool Serialize(FArchive& Ar);
 
 	template <typename TOther>
 	void CopyFrom(const TOther& Other)
 	{
-		ObjectTypesToQuery = Other.GetObjectTypesToQuery();
+		ObjectTypesToQuery64 = Other.GetObjectTypesToQuery();
 		IgnoreMask = Other.IgnoreMask;
 		bHasValidData = true;
 	}
 	
-	UPROPERTY(VisibleAnywhere, Category=QueryData)
-	uint8 ObjectTypesToQuery = 0;
+	#if WITH_EDITORONLY_DATA
+	UE_DEPRECATED(5.8, "ObjectTypesToQuery has been deprecated")
+	UPROPERTY(meta = (DeprecatedProperty, DeprecationMessage = "ObjectTypesToQuery has been deprecated"))
+	uint8 ObjectTypesToQuery_DEPRECATED = 0;
+	#endif // WITH_EDITORONLY_DATA
+
+	UE_INTERNAL uint64 ObjectTypesToQuery64 = 0;
 
 	UPROPERTY(VisibleAnywhere, Category=QueryData)
 	uint8 IgnoreMask = 0;
