Fixes a crash in UEngine::Browse() after the user selects to load the default map after failing to load a command line provided map.

UGameInstance::StartGameInstance attempts to call UEngine::Browse() a second time to load the default map if the initial call fails. This initial call clears WorldContext as part of the failure. The second call to UEngine::Browse()  dereferences the null WorldContext resulting in a crash.

Example repro:
1. Launch with an invalid world path (eg LyraGame-Win64-Debug.exe c:\badpath).
2. UGameInstance::StartGameInstance() calls UEngine::Browse() to attempt the load and fails, leaving UGameInstance::WorldContext == nullptr .
3. UGameInstance::StartGameInstance() shows a message dialog where selecting 'OK' results in a second call to UEngine::Browse() which dereferences

This change caches a stable identifier to the WorldContext and restores it before the second call to UEngine::Browse().

#Synced-CL 49566552

diff --git a/Engine/Source/Runtime/Engine/Private/GameInstance.cpp b/Engine/Source/Runtime/Engine/Private/GameInstance.cpp
--- a/Engine/Source/Runtime/Engine/Private/GameInstance.cpp
+++ b/Engine/Source/Runtime/Engine/Private/GameInstance.cpp
@@ -675,6 +675,11 @@ void UGameInstance::StartGameInstance()
 		PackageName = DefaultMap + GameMapsSettings->LocalMapOptions;
 	}
 
+	// Cache a stable identifier to the WorldContext to enable WorldContext to be restored if the initial call to Browse fails.
+	// This may be browsing to a user provided map which may not exist. On failure, this enable the user to decide to load to 
+	// the default map.
+	FName InitialWorldContextHandle = WorldContext->ContextHandle;
+
 	FURL URL(&DefaultURL, *PackageName, TRAVEL_Partial);
 	if (URL.Valid)
 	{
@@ -689,6 +694,9 @@ void UGameInstance::StartGameInstance()
 		// the map specified on the command-line couldn't be loaded.  ask the user if we should load the default map instead
 		if (FCString::Stricmp(*PackageName, *DefaultMap) != 0)
 		{
+			// Restore the WorldContext used in the initial call to Browse or create a new WorldContext if the WorldContext was destroyed. 
+			WorldContext = &Engine->GetWorldContextFromHandleChecked(InitialWorldContextHandle);
+			
 			const FText Message = FText::Format(NSLOCTEXT("Engine", "MapNotFound", "The map specified on the commandline '{0}' could not be found. Would you like to load the default map instead?"), FText::FromString(URL.Map));
 			if (   FCString::Stricmp(*URL.Map, *DefaultMap) != 0  
 				&& FMessageDialog::Open(EAppMsgType::OkCancel, Message) != EAppReturnType::Ok)
