Fixed: Race conditioon on creating reflection capture encoded HDR texture resource
#Synced-CL 49641827

diff --git a/Engine/Source/Runtime/Engine/Classes/Components/ReflectionCaptureComponent.h b/Engine/Source/Runtime/Engine/Classes/Components/ReflectionCaptureComponent.h
--- a/Engine/Source/Runtime/Engine/Classes/Components/ReflectionCaptureComponent.h
+++ b/Engine/Source/Runtime/Engine/Classes/Components/ReflectionCaptureComponent.h
@@ -95,7 +95,9 @@ class UReflectionCaptureComponent : public USceneComponent
 
 	ENGINE_API class FReflectionCaptureMapBuildData* GetMapBuildData() const;
 
-	ENGINE_API bool IsEncodedHDRCubemapTextureRequired(EShaderPlatform ShaderPlatform);
+	ENGINE_API static bool IsEncodedHDRCubemapTextureRequired(EShaderPlatform ShaderPlatform);
+	
+	ENGINE_API void AllocateEncodedHDRCubemapTexture(FStaticShaderPlatform ShaderPlatform);
 	
 	virtual void PropagateLightingScenarioChange() override;
 
@@ -105,7 +107,6 @@ class UReflectionCaptureComponent : public USceneComponent
 	virtual void CreateRenderState_Concurrent(FRegisterComponentContext* Context) override;
 	virtual void DestroyRenderState_Concurrent() override;
 	virtual void SendRenderTransform_Concurrent() override;
-	virtual void OnRegister() override;
 	virtual void InvalidateLightingCacheDetailed(bool bInvalidateBuildEnqueuedLighting, bool bTranslationOnly) override;
 	//~ End UActorComponent Interface
 
diff --git a/Engine/Source/Runtime/Engine/Private/Components/ReflectionCaptureComponent.cpp b/Engine/Source/Runtime/Engine/Private/Components/ReflectionCaptureComponent.cpp
--- a/Engine/Source/Runtime/Engine/Private/Components/ReflectionCaptureComponent.cpp
+++ b/Engine/Source/Runtime/Engine/Private/Components/ReflectionCaptureComponent.cpp
@@ -747,59 +747,48 @@ void UReflectionCaptureComponent::SafeReleaseEncodedHDRCubemapTexture()
 	}
 }
 
-void UReflectionCaptureComponent::OnRegister()
+void UReflectionCaptureComponent::AllocateEncodedHDRCubemapTexture(FStaticShaderPlatform ShaderPlatform)
 {
-	const FSceneInterface* Scene = GetWorld()->Scene;
-	const EShaderPlatform  ShaderPlatform = Scene ? Scene->GetShaderPlatform() : GMaxRHIShaderPlatform;
-	const bool bEncodedHDRCubemapTextureRequired = IsEncodedHDRCubemapTextureRequired(ShaderPlatform);
-
-	if (bEncodedHDRCubemapTextureRequired)
+	check(IsEncodedHDRCubemapTextureRequired(ShaderPlatform));
+	
+	FReflectionCaptureMapBuildData* MapBuildData = GetMapBuildData();
+	// If the MapBuildData is valid, update it. If it is not we will use the cached values, if there are any
+	if (MapBuildData && MapBuildData->CubemapSize > 0)
 	{
-		FReflectionCaptureMapBuildData* MapBuildData = GetMapBuildData();
-
-		// If the MapBuildData is valid, update it. If it is not we will use the cached values, if there are any
-		if (MapBuildData && MapBuildData->CubemapSize > 0)
+		if (!EncodedHDRCubemapTexture)
 		{
-			if (!EncodedHDRCubemapTexture)
+			// Make sure data is not discarded yet
+			check(!MapBuildData->HasBeenUploadedFinal());
+			
+			EncodedHDRCubemapTexture = new FReflectionTextureCubeResource();
+			TArray<uint8> EncodedHDRCapturedData;
+			EPixelFormat EncodedHDRCubemapTextureFormat = PF_Unknown;
+			if (IsMobileDeferredShadingEnabled(ShaderPlatform))
 			{
-				EncodedHDRCubemapTexture = new FReflectionTextureCubeResource();
-				TArray<uint8> EncodedHDRCapturedData;
-				EPixelFormat EncodedHDRCubemapTextureFormat = PF_Unknown;
-				if (IsMobileDeferredShadingEnabled(ShaderPlatform))
+				// make a copy, FullHDR data will still be needed later
+				EncodedHDRCapturedData = MapBuildData->FullHDRCapturedData;
+				EncodedHDRCubemapTextureFormat = PF_FloatRGBA;
+			}
+			else
+			{
+				if (GIsEditor)
 				{
-					// make a copy, FullHDR data will still be needed later
-					EncodedHDRCapturedData = MapBuildData->FullHDRCapturedData;
-					EncodedHDRCubemapTextureFormat = PF_FloatRGBA;
+					EncodedHDRCapturedData = MapBuildData->EncodedHDRCapturedData;
 				}
 				else
 				{
-					if (GIsEditor)
-					{ 
-						EncodedHDRCapturedData = MapBuildData->EncodedHDRCapturedData;
-					}
-					else
-					{
-						// EncodedHDR data is no longer needed
-						Swap(EncodedHDRCapturedData, MapBuildData->EncodedHDRCapturedData);
-					}
-					EncodedHDRCubemapTextureFormat = PF_FloatR11G11B10;
+					// EncodedHDR data is no longer needed
+					Swap(EncodedHDRCapturedData, MapBuildData->EncodedHDRCapturedData);
 				}
-
-				EncodedHDRCubemapTexture->SetupParameters(MapBuildData->CubemapSize, FMath::CeilLogTwo(MapBuildData->CubemapSize) + 1, EncodedHDRCubemapTextureFormat, MoveTemp(EncodedHDRCapturedData));
-				BeginInitResource(EncodedHDRCubemapTexture);
+				EncodedHDRCubemapTextureFormat = PF_FloatR11G11B10;
 			}
 
-			CachedAverageBrightness = MapBuildData->AverageBrightness;
+			EncodedHDRCubemapTexture->SetupParameters(MapBuildData->CubemapSize, FMath::CeilLogTwo(MapBuildData->CubemapSize) + 1, EncodedHDRCubemapTextureFormat, MoveTemp(EncodedHDRCapturedData));
+			BeginInitResource(EncodedHDRCubemapTexture);
 		}
-	}
-	else
-	{
-		// SM5 doesn't require cached values
-		SafeReleaseEncodedHDRCubemapTexture();
-		CachedAverageBrightness = 0;
-	}
 
-	Super::OnRegister();
+		CachedAverageBrightness = MapBuildData->AverageBrightness;
+	}
 }
 
 void UReflectionCaptureComponent::DestroyRenderState_Concurrent()
diff --git a/Engine/Source/Runtime/Renderer/Private/ReflectionEnvironmentCapture.cpp b/Engine/Source/Runtime/Renderer/Private/ReflectionEnvironmentCapture.cpp
--- a/Engine/Source/Runtime/Renderer/Private/ReflectionEnvironmentCapture.cpp
+++ b/Engine/Source/Runtime/Renderer/Private/ReflectionEnvironmentCapture.cpp
@@ -77,12 +77,6 @@ static FAutoConsoleVariableRef CVarSkylightCaptureLODDistanceScale(
 	TEXT("Negative values will be clamped to 1"),
 	ECVF_Scalability);
 
-static bool GTossReflectionCaptureDataOnGT = false;
-static FAutoConsoleVariableRef CVarTossReflectionCaptureDataOnGT(
-	TEXT("r.TossReflectionCaptureDataOnGT"),
-	GTossReflectionCaptureDataOnGT,
-	TEXT("Wether or not toss reflection capture data on GT if needed\n"));
-
 /** 
  * Mip map used by a Roughness of 0, counting down from the lowest resolution mip (MipCount - 1).  
  * This has been tweaked along with ReflectionCaptureRoughnessMipScale to make good use of the resolution in each mip, especially the highest resolution mips.
@@ -849,6 +843,20 @@ void FScene::AllocateReflectionCaptures(const TArray<UReflectionCaptureComponent
 {
 	if (NewCaptures.Num() > 0)
 	{
+		// Create encoded HDR texture resources if platform requires it
+		// This has to be done before capture data is uploaded to GPU and discarded, see UpdateAllReflectionCaptures
+		EShaderPlatform Platform = GetShaderPlatform();
+		if (UReflectionCaptureComponent::IsEncodedHDRCubemapTextureRequired(Platform))
+		{
+			for (UReflectionCaptureComponent* Component : NewCaptures)
+			{
+				if (Component)
+				{
+					Component->AllocateEncodedHDRCubemapTexture(Platform);
+				}
+			}
+		}
+				
 		if (SupportsTextureCubeArray(GetFeatureLevel()))
 		{
 			int32_t PlatformMaxNumReflectionCaptures = FMath::Min(FMath::FloorToInt(GMaxTextureArrayLayers / 6.0f), GetMaxNumReflectionCaptures(GetShaderPlatform()));
@@ -1470,11 +1478,8 @@ void FScene::CaptureOrUploadReflectionCapture(UReflectionCaptureComponent* Captu
 
 			FScene* Scene = this;
 
-			bool bTossCaptureDataOnGT = CaptureComponent->IsEncodedHDRCubemapTextureRequired(GetShaderPlatform()) && GTossReflectionCaptureDataOnGT;
-			TWeakObjectPtr WeakCapture = CaptureComponent;
-
 			ENQUEUE_RENDER_COMMAND(UploadCaptureCommand)
-				([Scene, CaptureData, CaptureComponent, WeakCapture, bTossCaptureDataOnGT](FRHICommandListImmediate& RHICmdList)
+				([Scene, CaptureData, CaptureComponent](FRHICommandListImmediate& RHICmdList)
 			{
 				// After the final upload we cannot upload again because we tossed the source MapBuildData,
 				// After uploading it into the scene's texture array, to guaratee there's only one copy in memory.
@@ -1483,25 +1488,7 @@ void FScene::CaptureOrUploadReflectionCapture(UReflectionCaptureComponent* Captu
 				{
 					UploadReflectionCapture_RenderingThread(RHICmdList, Scene, CaptureData, CaptureComponent);
 
-					if (bTossCaptureDataOnGT)
-					{
-						AsyncTask(ENamedThreads::GameThread, [CaptureData, WeakCapture]()
-								  {
-									  if (WeakCapture.IsValid())
-									  {
-										  FReflectionCaptureData* CaptureDataLocal = WeakCapture.Get()->GetMapBuildData();
-										  check(CaptureDataLocal == CaptureData);
-										  if (CaptureDataLocal != nullptr)
-										  {
-										      CaptureDataLocal->OnDataUploadedToGPUFinal();
-										  }
-									  }
-								  });
-					}
-					else
-					{
-						CaptureData->OnDataUploadedToGPUFinal();
-					}
+					CaptureData->OnDataUploadedToGPUFinal();
 				}
 				else
 				{
