Change MorphTargetVertexInfoBuffers::Compress and InitMorphResourcesStreaming signature to use a TArrayView instead of an TArray reference so it able to accept non const pointer types.

#rnx

#Synced-CL 49618898

diff --git a/Engine/Source/Runtime/Engine/Private/Rendering/MorphTargetVertexInfoBuffers.cpp b/Engine/Source/Runtime/Engine/Private/Rendering/MorphTargetVertexInfoBuffers.cpp
--- a/Engine/Source/Runtime/Engine/Private/Rendering/MorphTargetVertexInfoBuffers.cpp
+++ b/Engine/Source/Runtime/Engine/Private/Rendering/MorphTargetVertexInfoBuffers.cpp
@@ -165,7 +165,7 @@ void FMorphTargetVertexInfoBuffers::InitMorphResources(EShaderPlatform ShaderPla
 	ValidateVertexBuffers(true);
 }
 
-void FMorphTargetVertexInfoBuffers::InitMorphResourcesStreaming(EShaderPlatform ShaderPlatform, const TArray<FSkelMeshRenderSection>& RenderSections, const TArray<const FMorphTargetLODModel*>& MorphTargets, int32 NumVertices, float TargetPositionErrorTolerance)
+void FMorphTargetVertexInfoBuffers::InitMorphResourcesStreaming(EShaderPlatform ShaderPlatform, const TArray<FSkelMeshRenderSection>& RenderSections, TConstArrayView<const FMorphTargetLODModel*> MorphTargets, int32 NumVertices, float TargetPositionErrorTolerance)
 {
 	ResetCPUData();
 
@@ -335,7 +335,7 @@ void FMorphTargetVertexInfoBuffers::InitDataFromCompressedMorphs(
 }
 #endif
 
-void FMorphTargetVertexInfoBuffers::Compress(const TArray<FSkelMeshRenderSection>& RenderSections, const TArray<const FMorphTargetLODModel*>& MorphTargets, const TBitArray<>& UsesBuiltinMorphTargetCompression, int32 NumVertices, float TargetPositionErrorTolerance)
+void FMorphTargetVertexInfoBuffers::Compress(const TArray<FSkelMeshRenderSection>& RenderSections, TConstArrayView<const FMorphTargetLODModel*> MorphTargets, const TBitArray<>& UsesBuiltinMorphTargetCompression, int32 NumVertices, float TargetPositionErrorTolerance)
 {
 	using namespace UE::MorphTargetVertexCodec;
 	
diff --git a/Engine/Source/Runtime/Engine/Public/Rendering/MorphTargetVertexInfoBuffers.h b/Engine/Source/Runtime/Engine/Public/Rendering/MorphTargetVertexInfoBuffers.h
--- a/Engine/Source/Runtime/Engine/Public/Rendering/MorphTargetVertexInfoBuffers.h
+++ b/Engine/Source/Runtime/Engine/Public/Rendering/MorphTargetVertexInfoBuffers.h
@@ -22,7 +22,7 @@ public:
 
 	ENGINE_API void InitMorphResources(EShaderPlatform ShaderPlatform, const TArray<FSkelMeshRenderSection>& RenderSections, const TArray<UMorphTarget*>& MorphTargets, int32 NumVertices, int32 LODIndex, float TargetPositionErrorTolerance);
 
-	ENGINE_API void InitMorphResourcesStreaming(EShaderPlatform ShaderPlatform, const TArray<FSkelMeshRenderSection>& RenderSections, const TArray<const FMorphTargetLODModel*>& MorphTargets, int32 NumVertices, float TargetPositionErrorTolerance);
+	ENGINE_API void InitMorphResourcesStreaming(EShaderPlatform ShaderPlatform, const TArray<FSkelMeshRenderSection>& RenderSections, TConstArrayView<const FMorphTargetLODModel*> MorphTargets, int32 NumVertices, float TargetPositionErrorTolerance);
 
 #if !WITH_EDITOR
 	ENGINE_API void InitMorphResourcesStreaming(EShaderPlatform ShaderPlatform, TConstArrayView<const FMorphTargetCompressedLODModel*> MorphTargets, TConstArrayView<FVector4f> InMinimumValuePerMorph, TConstArrayView<FVector4f> InMaximumValuesPerMorph);
@@ -102,7 +102,7 @@ protected:
 	ENGINE_API void ValidateVertexBuffers(bool bMorphTargetsShouldBeValid);
 	ENGINE_API void Serialize(FArchive& Ar);
 
-	void Compress(const TArray<FSkelMeshRenderSection>& RenderSections, const TArray<const FMorphTargetLODModel*>& MorphTargets, const TBitArray<>& UsesBuiltinMorphTargetCompression, int32 NumVertices, float TargetPositionErrorTolerance);
+	void Compress(const TArray<FSkelMeshRenderSection>& RenderSections, TConstArrayView<const FMorphTargetLODModel*> MorphTargets, const TBitArray<>& UsesBuiltinMorphTargetCompression, int32 NumVertices, float TargetPositionErrorTolerance);
 
 #if !WITH_EDITOR
 	void InitDataFromCompressedMorphs(TConstArrayView<const FMorphTargetCompressedLODModel*> MorphTargets, TConstArrayView<FVector4f> InMinimumValuePerMorph, TConstArrayView<FVector4f> MaximumValuesPerMorph);
