Fix shutdown crash. The UObject destuciton order is not deterministic on shutdown.

#Synced-CL 49562939

diff --git a/Engine/Source/Runtime/Engine/Classes/PhysicsEngine/PhysicsSettings.h b/Engine/Source/Runtime/Engine/Classes/PhysicsEngine/PhysicsSettings.h
--- a/Engine/Source/Runtime/Engine/Classes/PhysicsEngine/PhysicsSettings.h
+++ b/Engine/Source/Runtime/Engine/Classes/PhysicsEngine/PhysicsSettings.h
@@ -416,8 +416,10 @@ class UPhysicsSettings : public UPhysicsSettingsCore
 
 
 public:
-
-	static UPhysicsSettings* Get() { return CastChecked<UPhysicsSettings>(UPhysicsSettings::StaticClass()->GetDefaultObject()); }
+	static UPhysicsSettings* Get()
+	{
+		return GetMutableDefault<UPhysicsSettings>();
+	}
 
 	UFUNCTION(BlueprintCallable, Category = "Physics")
 	int32 GetPhysicsHistoryCount() const
