Iris
* NetRootObjectAdapter::RelevantWith now supports receiving a subobject
* RelevantWith will also work if called multiple times for the same root object or for one of it's subobjects.
* The root object will remove it's dependency link only when all relevantwith calls have been removed.
* Ex: If RelevantWith is called three times with the same actor , you need to remove it three times too.
* Added EEndReplicationFlags' Destroy, DestroyNetHandle, ClearNetPushId when stopping replication.
* Added unittest for this functionnality

#rnx

#Synced-CL 49546171

diff --git a/Engine/Source/Runtime/Engine/Private/Net/Iris/ReplicationSystem/NetRootObjectAdapter.cpp b/Engine/Source/Runtime/Engine/Private/Net/Iris/ReplicationSystem/NetRootObjectAdapter.cpp
--- a/Engine/Source/Runtime/Engine/Private/Net/Iris/ReplicationSystem/NetRootObjectAdapter.cpp
+++ b/Engine/Source/Runtime/Engine/Private/Net/Iris/ReplicationSystem/NetRootObjectAdapter.cpp
@@ -245,13 +245,13 @@ void FNetRootObjectAdapter::StopReplication()
 
 	UE_LOG(LogIris, Verbose, TEXT("%hs for: %s (level: %s)"), __FUNCTION__, *GetNameSafe(ReplicatedObject), *GetNameSafe(Level));
 
-	// If we were still attached to dependents
-	if (DependentRefCount > 0)
+	// Warn if we were still attached to dependents
+	if (!DependentRootObjects.IsEmpty())
 	{
-		UE_LOG(LogIris, Verbose, TEXT("%hs for: %s while still dependent to %d root objects. Dependencies will be cleared."), __FUNCTION__, *GetNameSafe(ReplicatedObject), DependentRefCount);
+		UE_LOG(LogIris, Warning, TEXT("%hs for: %s while still dependent to %d root objects. Dependencies will be cleared."), __FUNCTION__, *GetNameSafe(ReplicatedObject), DependentRootObjects.Num());
 
-		// Iris Dependencies will be destroyed by StopReplicatingNetRefHandle so let's clear the refcount here
-		DependentRefCount = 0;
+		// Iris dependencies will be destroyed by StopReplicatingNetRefHandle so let's clear the refcounts here
+		DependentRootObjects.Empty();
 	}
 
 	FReplicationSystemUtil::ForEachServerBridge(World, [&](UEngineReplicationBridge* Bridge)
@@ -263,23 +263,19 @@ void FNetRootObjectAdapter::StopReplication()
 			return;
 		}
 
-		//TODO: Figure out what EndReplication flags to use in our different cases.
-		// 		Might be time to refactor this enum it's both too implicit and too explicit.
-		//		Ex: Do Dynamic objects need to pass Destroy or not ? 
-		//			Does Destroy do anything on stables ?
-		//			Why should the caller need to think about all these million use cases; the UEngineReplicationBridge::StopReplicatingActor does abstract most of it. Should we just add a Bridge one for RootObjects too ?
-		//			Should flush be an object trait instead ?
+		//TODO: These default flags should probably be automatic or given by the bridge.
+		const EEndReplicationFlags DefaultFlags = EEndReplicationFlags::Destroy | EEndReplicationFlags::DestroyNetHandle | EEndReplicationFlags::ClearNetPushId;
 
-		// Stop replicating the entity
-		Bridge->StopReplicatingNetRefHandle(RefHandle, EEndReplicationFlags::None);
+		// Stop replicating the object
+		Bridge->StopReplicatingNetRefHandle(RefHandle, DefaultFlags);
 	});
 
 	bIsReplicating = false;
 }
 
-void FNetRootObjectAdapter::RelevantWith(const UObject* OtherRootObject)
+void FNetRootObjectAdapter::RelevantWith(const UObject* OtherReplicatedObject)
 {
-	check(OtherRootObject);
+	check(OtherReplicatedObject);
 
 	if (!bIsInitialized)
 	{
@@ -294,82 +290,86 @@ void FNetRootObjectAdapter::RelevantWith(const UObject* OtherRootObject)
 		return;
 	}
 
-	UObject* ReplicatedObject = WeakReplicatedObject.Get();
-	if (!ReplicatedObject)
+	UObject* MyReplicatedObject = WeakReplicatedObject.Get();
+	if (!MyReplicatedObject)
 	{
-		ensureMsgf(ReplicatedObject, TEXT("RootObjectAdapter assigned to a destroyed object. Something went wrong"));
+		ensureMsgf(MyReplicatedObject, TEXT("RootObjectAdapter assigned to a destroyed object. Something went wrong"));
 		return;
 	}
 
-	if (ReplicatedObject == OtherRootObject)
+	if (MyReplicatedObject == OtherReplicatedObject)
 	{
-		ensureMsgf(ReplicatedObject != OtherRootObject, TEXT("Cannot call RelevantWith to yourself: %s"), *GetNameSafe(OtherRootObject));
+		ensureMsgf(MyReplicatedObject != OtherReplicatedObject, TEXT("Cannot call RelevantWith on yourself: %s"), *GetNameSafe(MyReplicatedObject));
 		return;
 	}
 
-	UE_LOG(LogIris, Verbose, TEXT("%hs for: %s | Attached to %s | (level: %s)"), __FUNCTION__, *GetNameSafe(ReplicatedObject), *GetNameSafe(OtherRootObject), *GetNameSafe(Level));
-
-	DependentRefCount++;
-
-	// Make sure our object is replicating when it first gets associated with another root object.
-	if (!bIsReplicating && DependentRefCount == 1)
-	{
-		StartReplication(Level);
-	}
-
-	if (!bIsReplicating)
-	{
-		// Something went wrong
-		return;
-	}
+	UE_LOG(LogIris, Verbose, TEXT("%hs for: %s | Attached to %s | (level: %s)"), __FUNCTION__, *GetNameSafe(MyReplicatedObject), *GetNameSafe(OtherReplicatedObject), *GetNameSafe(Level));
 
 	UWorld* World = Level->GetWorld();
 	check(World);
 
-	FReplicationSystemUtil::ForEachServerBridge(World, [&](UObjectReplicationBridge* Bridge)
+	FReplicationSystemUtil::ForEachServerBridge(World, [&](UEngineReplicationBridge* Bridge)
 	{
-		// Find our handle
-		const FNetRefHandle MyRefHandle = Bridge->GetReplicatedRefHandle(ReplicatedObject);
-		if (!MyRefHandle.IsValid())
+		// Find the handle of the object we will be dependent to
+		const FNetRefHandle OtherRefHandle = Bridge->GetReplicatedRefHandle(OtherReplicatedObject);
+		if (!OtherRefHandle.IsValid())
 		{
-			ensureMsgf(false, TEXT("No netrefhandle assigned to %s"), *GetNameSafe(ReplicatedObject));
+			UE_LOG(LogIris, Error, TEXT("%hs Cannot replicate: %s with: %s because the other object is not replicating"), __FUNCTION__, *GetNameSafe(MyReplicatedObject), *GetNameSafe(OtherReplicatedObject));
+			ensure(false);
 			return;
 		}
 
-		// Find the handle of the object we will be dependent to
-		const FNetRefHandle OtherRefHandle = Bridge->GetReplicatedRefHandle(OtherRootObject);
-		if (!OtherRefHandle.IsValid())
+		// Find the handle of it's root object
+		const FNetRefHandle ParentHandle = Bridge->GetRootObjectOfAnyObject(OtherRefHandle);
+		if (!ParentHandle.IsValid())
 		{
-			UE_LOG(LogIris, Error, TEXT("%hs Cannot replicate: %s with: %s because the other root object is not replicating"), __FUNCTION__, *GetNameSafe(ReplicatedObject), *GetNameSafe(OtherRootObject));
+			UE_LOG(LogIris, Error, TEXT("%hs Could not find the root object handle to replicate %s with %s"), __FUNCTION__, *GetNameSafe(MyReplicatedObject), *GetNameSafe(OtherReplicatedObject));
 			ensure(false);
 			return;
 		}
 
-		// Make sure it's a root object
-		const bool bIsRootObject = Bridge->IsRootObject(OtherRefHandle);
-		if (!bIsRootObject)
+		// Make sure our object is replicating when it first gets associated with another root object.
+		if (!bIsReplicating)
+		{
+			StartReplication(Level);
+		}
+
+		if (!bIsReplicating)
+		{
+			// Something went wrong
+			return;
+		}
+
+		// Find our handle
+		const FNetRefHandle MyRefHandle = Bridge->GetReplicatedRefHandle(MyReplicatedObject);
+		if (!MyRefHandle.IsValid())
 		{
-			UE_LOG(LogIris, Error, TEXT("%hs Cannot replicate %s with %s because it's a subobject"), __FUNCTION__, *GetNameSafe(ReplicatedObject), *GetNameSafe(OtherRootObject));
-			ensure(false);
 			return;
 		}
 
-		Bridge->AddDependentObject(OtherRefHandle, MyRefHandle);
+		uint32& RootObjectRefCount = DependentRootObjects.FindOrAdd(ParentHandle);
+		++RootObjectRefCount;
+
+		// Add the Iris dependency the first time we are linked with that root object
+		if (RootObjectRefCount == 1)
+		{
+			Bridge->AddDependentObject(ParentHandle, MyRefHandle);
+		}
 	});
 }
 
-void FNetRootObjectAdapter::RemoveRelevantWith(const UObject* OtherRootObject)
+void FNetRootObjectAdapter::RemoveRelevantWith(const UObject* OtherReplicatedObject)
 {
-	UObject* ReplicatedObject = WeakReplicatedObject.Get();
-	if (!ReplicatedObject)
+	UObject* MyReplicatedObject = WeakReplicatedObject.Get();
+	if (!MyReplicatedObject)
 	{
-		ensureMsgf(ReplicatedObject, TEXT("RootObjectAdapter assigned to a destroyed object. Something went wrong"));
+		ensureMsgf(MyReplicatedObject, TEXT("RootObjectAdapter assigned to a destroyed object. Something went wrong"));
 		return;
 	}
 
-	if (ReplicatedObject == OtherRootObject)
+	if (MyReplicatedObject == OtherReplicatedObject)
 	{
-		ensureMsgf(ReplicatedObject != OtherRootObject, TEXT("Cannot StopReplicatingWith yourself: %s"), *GetNameSafe(OtherRootObject));
+		ensureMsgf(MyReplicatedObject != OtherReplicatedObject, TEXT("Cannot StopReplicatingWith yourself: %s"), *GetNameSafe(OtherReplicatedObject));
 		return;
 	}
 
@@ -380,52 +380,61 @@ void FNetRootObjectAdapter::RemoveRelevantWith(const UObject* OtherRootObject)
 		return;
 	}
 
-	UE_LOG(LogIris, Verbose, TEXT("%hs for: %s | Dettaching from: %s | (level: %s)"), __FUNCTION__, *GetNameSafe(ReplicatedObject), *GetNameSafe(OtherRootObject), *GetNameSafe(Level));
+	UE_LOG(LogIris, Verbose, TEXT("%hs for: %s | Dettaching from: %s | (level: %s)"), __FUNCTION__, *GetNameSafe(MyReplicatedObject), *GetNameSafe(OtherReplicatedObject), *GetNameSafe(Level));
 
 	UWorld* World = Level->GetWorld();
 	check(World);
 
-	if (DependentRefCount > 0)
-	{
-		DependentRefCount--;
-	}
-	else
-	{
-		ensureMsgf(false, TEXT("Error with mismatched RelevantWith/RemoveRelevantWith calls. DependentRefCount is already at 0 when removing relevancy between %s and %s."), *GetNameSafe(ReplicatedObject), *GetNameSafe(OtherRootObject));
-		// Let's try to remove this dependency anyway in case the error was in a previous call.
-	}
-
-	
-
 	if (!bIsReplicating)
 	{
 		// Nothing more to do if we stopped replicating already
 		return;
 	}
 
-	FReplicationSystemUtil::ForEachServerBridge(World, [&](UObjectReplicationBridge* Bridge)
+	FReplicationSystemUtil::ForEachServerBridge(World, [&](UEngineReplicationBridge* Bridge)
 	{
-		// Find our handle
-		const FNetRefHandle MyRefHandle = Bridge->GetReplicatedRefHandle(ReplicatedObject);
-		if (!MyRefHandle.IsValid())
+		// Find the handle of the object we will be dependent to
+		const FNetRefHandle OtherRefHandle = Bridge->GetReplicatedRefHandle(OtherReplicatedObject);
+		if (!OtherRefHandle.IsValid())
 		{
-			ensureMsgf(false, TEXT("No netrefhandle assigned to %s"), *GetNameSafe(ReplicatedObject));
+			UE_LOG(LogIris, Error, TEXT("%hs ignored for: %s with: %s because the other replicated object is not replicating"), __FUNCTION__, *GetNameSafe(MyReplicatedObject), *GetNameSafe(OtherReplicatedObject));
+			ensure(false);
 			return;
 		}
 
-		// Find the handle of the object we will be dependent to
-		const FNetRefHandle OtherRefHandle = Bridge->GetReplicatedRefHandle(OtherRootObject);
-		if (!OtherRefHandle.IsValid())
+		// Find the handle of it's root object
+		const FNetRefHandle ParentHandle = Bridge->GetRootObjectOfAnyObject(OtherRefHandle);
+		if (!ParentHandle.IsValid())
 		{
-			UE_LOG(LogIris, Warning, TEXT("%hs ignored for: %s with: %s because the other root object is not replicating"), __FUNCTION__, *GetNameSafe(ReplicatedObject), *GetNameSafe(OtherRootObject));
+			UE_LOG(LogIris, Error, TEXT("%hs Could not find root object handle to replicate %s with %s"), __FUNCTION__, *GetNameSafe(MyReplicatedObject), *GetNameSafe(OtherReplicatedObject));
+			ensure(false);
 			return;
 		}
 
-		Bridge->RemoveDependentObject(OtherRefHandle, MyRefHandle);
+		{
+			uint32* RootObjectRefCount = DependentRootObjects.Find(ParentHandle);
+	
+			// Remove the Iris dependency if it's the last reference
+			if (!RootObjectRefCount || (*RootObjectRefCount) <= 1)
+			{
+				// Find our handle
+				const FNetRefHandle MyRefHandle = Bridge->GetReplicatedRefHandle(MyReplicatedObject);
+				if (ensureMsgf(MyRefHandle.IsValid(), TEXT("No netrefhandle assigned to %s"), *GetNameSafe(MyReplicatedObject)))
+				{
+					Bridge->RemoveDependentObject(ParentHandle, MyRefHandle);
+				}
+				
+				DependentRootObjects.Remove(ParentHandle);
+			}
+			else
+			{
+				--RootObjectRefCount;
+			}
+		}
 	});
 
-	// If we are no longer replicated via another object, remove ourself from the replication system.
-	if (DependentRefCount == 0 && bOnlyReplicateWhenLinked)
+	// If we are no longer replicated via anybody else, remove ourself from the replication system.
+	if (DependentRootObjects.IsEmpty() && bOnlyReplicateWhenLinked)
 	{
 		StopReplication();
 	}
diff --git a/Engine/Source/Runtime/Engine/Public/Net/Iris/ReplicationSystem/NetRootObjectAdapter.h b/Engine/Source/Runtime/Engine/Public/Net/Iris/ReplicationSystem/NetRootObjectAdapter.h
--- a/Engine/Source/Runtime/Engine/Public/Net/Iris/ReplicationSystem/NetRootObjectAdapter.h
+++ b/Engine/Source/Runtime/Engine/Public/Net/Iris/ReplicationSystem/NetRootObjectAdapter.h
@@ -4,6 +4,7 @@
 
 #include "UObject/WeakObjectPtr.h"
 #include "Misc/NotNull.h"
+#include "Iris/ReplicationSystem/NetRefHandle.h"
 
 class UWorld;
 class ULevel;
@@ -36,6 +37,11 @@ struct FRootObjectSettings
  * 2) Configure the replication settings via Configure() prior to starting replication
  * 3) Provide the ULevel the UObject is a part of
  * 4) Make the object replicated via StartReplication or RelevantWith
+ * 
+ * This class is meant to provide an easy to use solution for quickly transforming any non-actor into an autonomous replicated object.
+ * But be aware it is not made to be the optimal solution if you intend to create alot of instances of the class using it.
+ * The memory overheaded of the variables in this class could be optimized away entirely if your class directly interacted with the Iris API in a similar way to what the adapter does.
+ * So for a more optimal solution, it is recommended to look at this class for inspiration and instead add the minimum necessary features you need to your own base class.
  */
 class FNetRootObjectAdapter
 {
@@ -74,17 +80,17 @@ public:
 	ENGINE_API void StopReplication();
 
 	/**
-	 * Start replicating the object but make it relevant only via another root object.
+	 * Start replicating the object but make it relevant only via another replicated object.
 	 * You can have relevancy links to multiple different objects. As long as of those is relevant to a client then this object will be too.
-	 * @OtherRootObject:  A replicated root object we want to be tie our relevancy status to.
+	 * @OtherReplicatedObject:  A replicated object we want to be tie our relevancy status to. If it's a subobject, we will use it's root object instead.
 	 */
-	ENGINE_API void RelevantWith(const UObject* OtherRootObject);
+	ENGINE_API void RelevantWith(const UObject* OtherReplicatedObject);
 
 	/**
-	 * Remove the relevancy dependency between our object and another root object
-	 * @OtherRootObject: A replicated root object that was previously in a RelevantWith call.
+	 * Remove the relevancy dependency between our object and another replicated object.
+	 * @OtherReplicatedObject: The replicated object this object was set to be relevant with
 	 */
-	ENGINE_API void RemoveRelevantWith(const UObject* OtherRootObject);
+	ENGINE_API void RemoveRelevantWith(const UObject* OtherReplicatedObject);
 
 	/** The netobject factory to use with this rootobject. By default it will use the UNetRootObjectFactory. */
 	ENGINE_API FName GetNetFactoryName() const;
@@ -123,14 +129,14 @@ public:
 
 private:
 
+	/** Map of root objects and how many times we are dependent to each. */
+	TMap<FNetRefHandle /* RootObject */, uint32 /* DependentCount */> DependentRootObjects;
+
 	FName ExplicitNetFactoryName;	
 
 	FWeakObjectPtr WeakReplicatedObject;
 	TWeakObjectPtr<ULevel> WeakLevel;
 
-	/** Number of objects we are dependent with. */
-	uint32 DependentRefCount = 0;
-
 	// Status flags
 
 	/** Is the adapter correctly setup */
