[NetPhysics] Improve NetworkPhysicsTickOffset API to easily get the value on the physics thread and game thread.

#Synced-CL 49558237

diff --git a/Engine/Source/Runtime/Engine/Private/PhysicsEngine/Experimental/PhysScene_Chaos.cpp b/Engine/Source/Runtime/Engine/Private/PhysicsEngine/Experimental/PhysScene_Chaos.cpp
--- a/Engine/Source/Runtime/Engine/Private/PhysicsEngine/Experimental/PhysScene_Chaos.cpp
+++ b/Engine/Source/Runtime/Engine/Private/PhysicsEngine/Experimental/PhysScene_Chaos.cpp
@@ -893,6 +893,11 @@ IPhysicsReplication* FPhysScene_Chaos::GetPhysicsReplication()
 	return PhysicsReplication.Get();
 }
 
+const IPhysicsReplication* FPhysScene_Chaos::GetPhysicsReplication() const
+{
+	return PhysicsReplication.Get();
+}
+
 IPhysicsReplication* FPhysScene_Chaos::CreatePhysicsReplication()
 {
 	// Create replication manager
diff --git a/Engine/Source/Runtime/Engine/Private/PhysicsEngine/NetworkPhysicsComponent.cpp b/Engine/Source/Runtime/Engine/Private/PhysicsEngine/NetworkPhysicsComponent.cpp
--- a/Engine/Source/Runtime/Engine/Private/PhysicsEngine/NetworkPhysicsComponent.cpp
+++ b/Engine/Source/Runtime/Engine/Private/PhysicsEngine/NetworkPhysicsComponent.cpp
@@ -1691,13 +1691,6 @@ void UNetworkPhysicsComponent::UpdateAsyncComponent(const bool bFullUpdate)
 		{
 			if (!HasServerWorld())
 			{
-				if (UWorld* World = GetWorld())
-				{
-					if (APlayerController* PC = World->GetFirstPlayerController())
-					{
-						AsyncInput->NetworkPhysicsTickOffset = PC->GetNetworkPhysicsTickOffset();
-					}
-				}
 				AsyncInput->InputsToNetwork_Owner = InputsToNetwork_Owner;
 			}
 
@@ -1819,7 +1812,6 @@ FAsyncNetworkPhysicsComponent::FAsyncNetworkPhysicsComponent() : TSimCallbackObj
 	, bIsLocallyControlled(true)
 	, NetMode(ENetMode::NM_Standalone)
 	, NetRole(ENetRole::ROLE_Authority)
-	, NetworkPhysicsTickOffset(0)
 	, PhysicsReplicationMode(EPhysicsReplicationMode::Default)
 	, bIsUsingLegacyData(false)
 	, SettingsComponent(nullptr)
@@ -1912,21 +1904,6 @@ void FAsyncNetworkPhysicsComponent::ConsumeAsyncInput_Internal(const int32 Physi
 			{
 				NetRole = *AsyncInput->NetRole;
 			}
-			if (AsyncInput->NetworkPhysicsTickOffset.IsSet())
-			{
-				if (NetworkPhysicsTickOffset != *AsyncInput->NetworkPhysicsTickOffset)
-				{
-					NetworkPhysicsTickOffset = *AsyncInput->NetworkPhysicsTickOffset;
-
-					if (ComponentSettings.GetDynamicInputBufferScalingEnabled())
-					{
-						// Also reset dynamic input buffer if tick offset was altered
-						const float MinInputBufferMs = PhysicsReplicationCVars::ResimulationCVars::DynamicInputBufferScalingMinBufferMs;
-						const float MaxInputBufferMs = PhysicsReplicationCVars::ResimulationCVars::DynamicInputBufferScalingMaxBufferMs;
-						InputBufferTargetSizeMs = FMath::Min(InputBufferTargetSizeMs, MaxInputBufferMs - ((MaxInputBufferMs - MinInputBufferMs) * 0.5f));
-					}
-				}
-			}
 			if (AsyncInput->InputsToNetwork_Owner.IsSet())
 			{
 				// Only marshaled from GT to PT on the client
@@ -2008,6 +1985,7 @@ void FAsyncNetworkPhysicsComponent::ConsumeAsyncInput_Internal(const int32 Physi
 		/** Continuously marshaled data */
 		{
 			const bool bIsServer = IsServer();
+			const int32 NetworkPhysicsTickOffset = UE::NetworkPhysicsUtils::GetNetworkPhysicsTickOffset_Internal(GetSolver());
 
 			/** Receive data helper */
 			auto ReceiveHelper = [&](Chaos::FBaseRewindHistory* History, Chaos::FBaseRewindHistory* ReceiveData, const bool bImportant, const bool bCompareData)
@@ -2473,7 +2451,7 @@ void FAsyncNetworkPhysicsComponent::OnPostProcessInputs_Internal(const int32 Phy
 	{
 		// Prepare to gather input data
 		FNetworkPhysicsPayload* PhysicsData = InputData.Get();
-		PhysicsData->PrepareFrame(PhysicsStep, bIsServer, GetNetworkPhysicsTickOffset());
+		PhysicsData->PrepareFrame(PhysicsStep, bIsServer, UE::NetworkPhysicsUtils::GetNetworkPhysicsTickOffset_Internal(GetSolver()));
 
 		if (bIsUsingLegacyData)
 		{
@@ -2514,7 +2492,7 @@ void FAsyncNetworkPhysicsComponent::OnPostProcessInputs_Internal(const int32 Phy
 		{
 			// Prepare to gather state data
 			FNetworkPhysicsPayload* PhysicsData = StateData.Get();
-			PhysicsData->PrepareFrame(PhysicsStep, bIsServer, GetNetworkPhysicsTickOffset());
+			PhysicsData->PrepareFrame(PhysicsStep, bIsServer, UE::NetworkPhysicsUtils::GetNetworkPhysicsTickOffset_Internal(GetSolver()));
 
 			if (bIsUsingLegacyData)
 			{
@@ -2658,7 +2636,7 @@ void FAsyncNetworkPhysicsComponent::OnPostSolve_Internal()
 		{
 			// Prepare to gather state data
 			FNetworkPhysicsPayload* PhysicsData = StateData.Get();
-			PhysicsData->PrepareFrame(NextPhysicsStep, bIsServer, GetNetworkPhysicsTickOffset());
+			PhysicsData->PrepareFrame(NextPhysicsStep, bIsServer, UE::NetworkPhysicsUtils::GetNetworkPhysicsTickOffset_Internal(GetSolver()));
 
 			if (bIsUsingLegacyData)
 			{
@@ -2736,7 +2714,7 @@ void FAsyncNetworkPhysicsComponent::SendInputData_Internal(FAsyncNetworkPhysicsC
 #if DEBUG_NETWORK_PHYSICS
 				{
 					const int32 LocalFrame = GetRigidSolver()->GetCurrentFrame();
-					const int32 ServerFrame = bIsServer ? LocalFrame : LocalFrame + GetNetworkPhysicsTickOffset();
+					const int32 ServerFrame = bIsServer ? LocalFrame : LocalFrame + UE::NetworkPhysicsUtils::GetNetworkPhysicsTickOffset_Internal(GetSolver());
 					FString NetRoleString = bIsServer ? FString("SERVER") : (IsLocallyControlled() ? FString("AUTONO") : FString("PROXY "));
 					AsyncOutput.InputData->DebugData(FString::Printf(TEXT("%s | PT | SendInputData_Internal | UNRELIABLE | CurrentLocalFrame = %d | CurrentServerFrame = %d | Name: %s"), *NetRoleString, LocalFrame, ServerFrame, *GetActorName()));
 				}
@@ -2766,7 +2744,7 @@ void FAsyncNetworkPhysicsComponent::SendInputData_Internal(FAsyncNetworkPhysicsC
 #if DEBUG_NETWORK_PHYSICS
 					{
 						const int32 LocalFrame = GetRigidSolver()->GetCurrentFrame();
-						const int32 ServerFrame = bIsServer ? LocalFrame : LocalFrame + GetNetworkPhysicsTickOffset();
+						const int32 ServerFrame = bIsServer ? LocalFrame : LocalFrame + UE::NetworkPhysicsUtils::GetNetworkPhysicsTickOffset_Internal(GetSolver());
 						FString NetRoleString = bIsServer ? FString("SERVER") : (IsLocallyControlled() ? FString("AUTONO") : FString("PROXY "));
 						AsyncOutput.InputDataImportant[Idx]->DebugData(FString::Printf(TEXT("%s | PT | SendInputData_Internal | RELIABLE | CurrentLocalFrame = %d | CurrentServerFrame = %d | Name: %s"), *NetRoleString, LocalFrame, ServerFrame, *GetActorName()));
 					}
@@ -2801,7 +2779,7 @@ void FAsyncNetworkPhysicsComponent::SendStateData_Internal(FAsyncNetworkPhysicsC
 #if DEBUG_NETWORK_PHYSICS
 				{
 					const int32 LocalFrame = GetRigidSolver()->GetCurrentFrame();
-					const int32 ServerFrame = IsServer() ? LocalFrame : LocalFrame + GetNetworkPhysicsTickOffset();
+					const int32 ServerFrame = IsServer() ? LocalFrame : LocalFrame + UE::NetworkPhysicsUtils::GetNetworkPhysicsTickOffset_Internal(GetSolver());
 					AsyncOutput.StateData->DebugData(FString::Printf(TEXT("SERVER | PT | SendStateData_Internal | UNRELIABLE | CurrentLocalFrame = %d | CurrentServerFrame = %d | Name: %s"), LocalFrame, ServerFrame, *GetActorName()));
 				}
 #endif
@@ -2830,7 +2808,7 @@ void FAsyncNetworkPhysicsComponent::SendStateData_Internal(FAsyncNetworkPhysicsC
 #if DEBUG_NETWORK_PHYSICS
 					{
 						const int32 LocalFrame = GetRigidSolver()->GetCurrentFrame();
-						const int32 ServerFrame = IsServer() ? LocalFrame : LocalFrame + GetNetworkPhysicsTickOffset();
+						const int32 ServerFrame = IsServer() ? LocalFrame : LocalFrame + UE::NetworkPhysicsUtils::GetNetworkPhysicsTickOffset_Internal(GetSolver());
 						AsyncOutput.StateDataImportant[Idx]->DebugData(FString::Printf(TEXT("SERVER | PT | SendStateData_Internal | RELIABLE | CurrentLocalFrame = %d | CurrentServerFrame = %d | Name: %s"), LocalFrame, ServerFrame, *GetActorName()));
 					}
 #endif
@@ -3290,4 +3268,35 @@ namespace UE::NetworkPhysicsUtils
 
 		return INDEX_NONE;
 	}
+
+	const int32 GetNetworkPhysicsTickOffset_External(const UWorld* World)
+	{
+		Chaos::EnsureIsInGameThreadContext();
+		if (World)
+		{
+			if (const FPhysScene* Scene = World->GetPhysicsScene())
+			{
+				if (const IPhysicsReplication* ReplicationSystem = Scene->GetPhysicsReplication())
+				{
+					return ReplicationSystem->GetNetworkPhysicsTickOffset();
+				}
+			}
+		}
+
+		return 0;
+	}
+
+	const int32 GetNetworkPhysicsTickOffset_Internal(const Chaos::FPhysicsSolverBase* Solver)
+	{
+		Chaos::EnsureIsInPhysicsThreadContext();
+		if (Solver)
+		{
+			if (const IPhysicsReplicationAsync* ReplicationSystem = Solver->GetPhysicsReplication_Internal())
+			{
+				return ReplicationSystem->GetNetworkPhysicsTickOffset_Internal();
+			}
+		}
+
+		return 0;
+	}
 }
diff --git a/Engine/Source/Runtime/Engine/Private/PhysicsEngine/PhysicsReplication.cpp b/Engine/Source/Runtime/Engine/Private/PhysicsEngine/PhysicsReplication.cpp
--- a/Engine/Source/Runtime/Engine/Private/PhysicsEngine/PhysicsReplication.cpp
+++ b/Engine/Source/Runtime/Engine/Private/PhysicsEngine/PhysicsReplication.cpp
@@ -434,20 +434,17 @@ void FPhysicsReplication::OnTick(float DeltaSeconds, TMap<TWeakObjectPtr<UPrimit
 		return;
 	}
 
-	int32 LocalFrameOffset = 0; // LocalFrame = ServerFrame + LocalFrameOffset;
-	bool LocalFrameOffsetAssigned = false;
+	NetworkPhysicsTickOffsetAssigned = false;
+	NetworkPhysicsTickOffset = 0; // LocalFrame = ServerFrame - NetworkPhysicsTickOffset;
 
 	if (UPhysicsSettings::Get()->PhysicsPrediction.bEnablePhysicsPrediction)
 	{
 		if (UWorld* World = GetOwningWorld())
 		{
-			if (World->GetNetMode() == NM_Client)
+			if (APlayerController* PlayerController = World->GetFirstPlayerController())
 			{
-				if (APlayerController* PlayerController = World->GetFirstPlayerController())
-				{
-					LocalFrameOffsetAssigned = PlayerController->GetNetworkPhysicsTickOffsetAssigned();
-					LocalFrameOffset = PlayerController->GetNetworkPhysicsTickOffset();
-				}
+				NetworkPhysicsTickOffsetAssigned = PlayerController->GetNetworkPhysicsTickOffsetAssigned();
+				NetworkPhysicsTickOffset = PlayerController->GetNetworkPhysicsTickOffset();
 			}
 		}
 	}
@@ -456,6 +453,11 @@ void FPhysicsReplication::OnTick(float DeltaSeconds, TMap<TWeakObjectPtr<UPrimit
 	if (PhysicsReplicationAsync)
 	{
 		PrepareAsyncData_External(PhysicErrorCorrection);
+		if (ensure(AsyncInput))
+		{
+			AsyncInput->NetworkPhysicsTickOffsetAssigned = NetworkPhysicsTickOffsetAssigned;
+			AsyncInput->NetworkPhysicsTickOffset = NetworkPhysicsTickOffset;
+		}
 	}
 
 	// Get the ping between this PC & the server
@@ -495,7 +497,7 @@ void FPhysicsReplication::OnTick(float DeltaSeconds, TMap<TWeakObjectPtr<UPrimit
 
 							if (UpdatedState.Flags & ERigidBodyFlags::NeedsUpdate)
 							{
-								const int32 LocalFrame = PhysicsTarget.ServerFrame - LocalFrameOffset;
+								const int32 LocalFrame = PhysicsTarget.ServerFrame - NetworkPhysicsTickOffset;
 								const bool bRestoredState = ApplyRigidBodyState(DeltaSeconds, BI, PhysicsTarget, PhysicErrorCorrection, PingSecondsOneWay, LocalFrame, 0);
 
 								// Need to update the component to match new position.
@@ -527,25 +529,23 @@ void FPhysicsReplication::OnTick(float DeltaSeconds, TMap<TWeakObjectPtr<UPrimit
 	}
 	PendingDeleteFromComponentsToTargets.Reset();
 
-	// PhysicsObject replication flow
-	for (FReplicatedPhysicsTarget& PhysicsTarget : ReplicatedTargetsQueue)
+	if (AsyncInput)
 	{
-		const float PingSecondsOneWay = LocalPing * 0.5f * 0.001f;
-
-		// Queue up the target state for async replication
-		FPhysicsRepAsyncInputData AsyncInputData(PhysicsTarget.PhysicsObject);
-		AsyncInputData.TargetState = PhysicsTarget.TargetState;
-		AsyncInputData.Proxy = nullptr;
-		AsyncInputData.RepMode = PhysicsTarget.ReplicationMode;
-		AsyncInputData.ServerFrame = PhysicsTarget.ServerFrame;
-		AsyncInputData.LatencyOneWay = PingSecondsOneWay;
-		
-		if (LocalFrameOffsetAssigned)
+		// PhysicsObject replication flow
+		for (FReplicatedPhysicsTarget& PhysicsTarget : ReplicatedTargetsQueue)
 		{
-			AsyncInputData.FrameOffset = LocalFrameOffset;
+			const float PingSecondsOneWay = LocalPing * 0.5f * 0.001f;
+
+			// Queue up the target state for async replication
+			FPhysicsRepAsyncInputData AsyncInputData(PhysicsTarget.PhysicsObject);
+			AsyncInputData.TargetState = PhysicsTarget.TargetState;
+			AsyncInputData.Proxy = nullptr;
+			AsyncInputData.RepMode = PhysicsTarget.ReplicationMode;
+			AsyncInputData.ServerFrame = PhysicsTarget.ServerFrame;
+			AsyncInputData.LatencyOneWay = PingSecondsOneWay;
+		
+			AsyncInput->InputData.Add(AsyncInputData);
 		}
-
-		AsyncInput->InputData.Add(AsyncInputData);
 	}
 	ReplicatedTargetsQueue.Reset();
 
@@ -955,6 +955,9 @@ void FPhysicsReplicationAsync::AddResimulationRequest_Internal(const float Delta
 
 	if (const FPhysicsReplicationAsyncInput* AsyncInput = GetConsumerInput_Internal())
 	{
+		NetworkPhysicsTickOffsetAssigned = AsyncInput->NetworkPhysicsTickOffsetAssigned;
+		NetworkPhysicsTickOffset = AsyncInput->NetworkPhysicsTickOffset;
+
 		Chaos::FPBDRigidsSolver* RigidsSolver = static_cast<Chaos::FPBDRigidsSolver*>(GetSolver());
 		check(RigidsSolver);
 
@@ -988,7 +991,7 @@ void FPhysicsReplicationAsync::AddResimulationRequest_Internal(const float Delta
 								// Update target data required to trigger resim from input
 								Target.ServerFrame = InputData.ServerFrame;
 								Target.TargetState = InputData.TargetState;
-								Target.FrameOffset = InputData.FrameOffset.IsSet() ? *InputData.FrameOffset : 0;
+								Target.FrameOffset = NetworkPhysicsTickOffset;
 
 								if (IsTargetValidForResim(Target))
 								{
@@ -1010,6 +1013,14 @@ void FPhysicsReplicationAsync::AddResimulationRequest_Internal(const float Delta
 
 void FPhysicsReplicationAsync::OnPreSimulate_Internal()
 {
+	const FPhysicsReplicationAsyncInput* AsyncInput = GetConsumerInput_Internal();
+
+	if (AsyncInput)
+	{
+		NetworkPhysicsTickOffsetAssigned = AsyncInput->NetworkPhysicsTickOffsetAssigned;
+		NetworkPhysicsTickOffset = AsyncInput->NetworkPhysicsTickOffset;
+	}
+
 	if (FPhysicsReplication::ShouldSkipPhysicsReplication())
 	{
 		return;
@@ -1042,7 +1053,7 @@ void FPhysicsReplicationAsync::OnPreSimulate_Internal()
 		return;
 	}
 
-	if (const FPhysicsReplicationAsyncInput* AsyncInput = GetConsumerInput_Internal())
+	if (AsyncInput)
 	{
 		// Update async targets with target input
 		for (const FPhysicsRepAsyncInputData& Input : AsyncInput->InputData)
@@ -1187,7 +1198,7 @@ void FPhysicsReplicationAsync::UpdateRewindDataTarget(const FPhysicsRepAsyncInpu
 	}
 
 	// If there is no FrameOffset set then we have not synced up physics ticks with the server yet so don't cache this data
-	if (Input.FrameOffset.IsSet() == false)
+	if (NetworkPhysicsTickOffsetAssigned == false)
 	{
 		return;
 	}
@@ -1208,7 +1219,7 @@ void FPhysicsReplicationAsync::UpdateRewindDataTarget(const FPhysicsRepAsyncInpu
 	if (Chaos::FGeometryParticleHandle* Handle = Interface.GetParticle(Input.PhysicsObject))
 	{
 		// Cache all target states inside RewindData
-		const int32 LocalFrame = Input.ServerFrame - *Input.FrameOffset;
+		const int32 LocalFrame = Input.ServerFrame - NetworkPhysicsTickOffset;
 		RewindData->SetTargetStateAtFrame(*Handle, LocalFrame, Chaos::FFrameAndPhase::EParticleHistoryPhase::PrePushData,
 			Input.TargetState.Position, Input.TargetState.Quaternion,
 			Input.TargetState.LinVel, FMath::DegreesToRadians(Input.TargetState.AngVel), (Input.TargetState.Flags & ERigidBodyFlags::Sleeping));
@@ -1279,7 +1290,7 @@ void FPhysicsReplicationAsync::UpdateAsyncTarget(const FPhysicsRepAsyncInputData
 		const bool bPrevAllowTargetAltering = Target->bAllowTargetAltering;
 
 		// Cache if the physics frame offset has changed since last target
-		const bool bFrameOffsetCorrected = Target->FrameOffset != Input.FrameOffset;
+		const bool bFrameOffsetCorrected = Target->FrameOffset != NetworkPhysicsTickOffset;
 
 		// Set if the target is allowed to be altered after this update
 		Target->bAllowTargetAltering = !(Target->TargetState.Flags & ERigidBodyFlags::Sleeping) && !(Input.TargetState.Flags & ERigidBodyFlags::Sleeping);
@@ -1303,7 +1314,7 @@ void FPhysicsReplicationAsync::UpdateAsyncTarget(const FPhysicsRepAsyncInputData
 		Target->ReceiveFrame = CurrentFrame;
 		Target->TargetState = Input.TargetState;
 		Target->RepMode = Input.RepMode;
-		Target->FrameOffset = Input.FrameOffset.IsSet() ? *Input.FrameOffset : 0;
+		Target->FrameOffset = NetworkPhysicsTickOffset;
 		Target->TickCount = 0;
 		Target->AccumulatedSleepSeconds = 0.0f;
 
@@ -1538,7 +1549,7 @@ void FPhysicsReplicationAsync::CheckTargetResimValidity(FReplicatedPhysicsTarget
 		return;
 	}
 
-	const int32 LocalFrame = Target.ServerFrame - Target.FrameOffset;
+	const int32 LocalFrame = Target.ServerFrame - NetworkPhysicsTickOffset;
 	if (!RewindData->IsFrameWithinRewindHistory(LocalFrame))
 	{
 		if (LocalFrame < RewindData->GetEarliestFrame_Internal())
@@ -1598,7 +1609,7 @@ void FPhysicsReplicationAsync::ApplyPhysicsReplicationLOD(Chaos::FConstPhysicsOb
 		if (Target.RepMode == EPhysicsReplicationMode::PredictiveInterpolation)
 		{
 			const bool bShouldSleep = (Target.TargetState.Flags & ERigidBodyFlags::Sleeping) != 0;
-			int32 TargetClientFrame = (Target.ServerFrame - Target.FrameOffset);
+			int32 TargetClientFrame = (Target.ServerFrame - NetworkPhysicsTickOffset);
 
 			// If we use Predicitve Interpolation and we should not sleep and the aligned frame from LOD is ahead of the target, perform LOD aligment extrapolation
 			if (!bShouldSleep && LodData->AlignedFrame > TargetClientFrame)
@@ -1935,7 +1946,7 @@ bool FPhysicsReplicationAsync::PredictiveInterpolation(Chaos::FPBDRigidParticleH
 	if (PhysicsReplicationCVars::ResimulationCVars::bDisableReplicationOnInteraction && ParticlesInResimIslands.Contains(Handle->GetHandleIdx()))
 	{
 		// If particle is in an island with a resim object, don't run replication and wait for an up to date target (after leaving the island)
-		Target.SetWaiting(RigidsSolver->GetCurrentFrame() + Target.FrameOffset, EPhysicsReplicationMode::Resimulation);
+		Target.SetWaiting(RigidsSolver->GetCurrentFrame() + NetworkPhysicsTickOffset, EPhysicsReplicationMode::Resimulation);
 		return false;
 	}
 
@@ -2369,7 +2380,7 @@ bool FPhysicsReplicationAsync::IsTargetValidForResim(const FReplicatedPhysicsTar
 		return false;
 	}
 
-	const int32 LocalFrame = Target.ServerFrame - Target.FrameOffset;
+	const int32 LocalFrame = Target.ServerFrame - NetworkPhysicsTickOffset;
 
 	if (!RewindData->IsFrameWithinRewindHistory(LocalFrame))
 	{
@@ -2384,7 +2395,7 @@ bool FPhysicsReplicationAsync::ShouldResimulateParticle(const Chaos::FPBDRigidPa
 	if (IsTargetValidForResim(Target))
 	{
 		const Chaos::FRewindData* RewindData = static_cast<const Chaos::FPBDRigidsSolver*>(GetSolver())->GetRewindData();
-		const int32 LocalFrame = Target.ServerFrame - Target.FrameOffset;
+		const int32 LocalFrame = Target.ServerFrame - NetworkPhysicsTickOffset;
 		static constexpr Chaos::FFrameAndPhase::EParticleHistoryPhase RewindPhase = Chaos::FFrameAndPhase::EParticleHistoryPhase::PrePushData;
 
 		// Get state from locally cached history for frame corresponding to received data
@@ -2474,7 +2485,7 @@ bool FPhysicsReplicationAsync::ShouldResimulateParticle(const Chaos::FPBDRigidPa
 
 		if (Chaos::FPhysicsSolverBase::CanDebugNetworkPhysicsPrediction())
 		{
-			UE_LOG(LogPhysics, Log, TEXT("Apply Rigid body state at local frame %d with offset = %d"), LocalFrame, Target.FrameOffset);
+			UE_LOG(LogPhysics, Log, TEXT("Apply Rigid body state at local frame %d with offset = %d"), LocalFrame, NetworkPhysicsTickOffset);
 			UE_LOG(LogPhysics, Log, TEXT("Should Trigger Resim = %s | Server Frame = %d | Client Frame = %d"), (bShouldTriggerResim ? TEXT("True") : TEXT("False")), Target.ServerFrame, LocalFrame);
 			UE_LOG(LogPhysics, Log, TEXT("Particle Target Position = %s | Current Position = %s"), *Target.TargetState.Position.ToString(), *PastState.GetX().ToString());
 			UE_LOG(LogPhysics, Log, TEXT("Particle Target Velocity = %s | Current Velocity = %s"), *Target.TargetState.LinVel.ToString(), *PastState.GetV().ToString());
@@ -2548,7 +2559,7 @@ bool FPhysicsReplicationAsync::ResimulationReplication(Chaos::FPBDRigidParticleH
 		{ 
 			if(RigidsSolver->GetEvolution())
 			{ 
-				const int32 LocalFrame = Target.ServerFrame - Target.FrameOffset;
+				const int32 LocalFrame = Target.ServerFrame - NetworkPhysicsTickOffset;
 				static constexpr Chaos::FFrameAndPhase::EParticleHistoryPhase RewindPhase = Chaos::FFrameAndPhase::EParticleHistoryPhase::PrePushData;
 
 				// Get state from locally cached history for frame corresponding to received data
diff --git a/Engine/Source/Runtime/Engine/Public/Physics/Experimental/PhysScene_Chaos.h b/Engine/Source/Runtime/Engine/Public/Physics/Experimental/PhysScene_Chaos.h
--- a/Engine/Source/Runtime/Engine/Public/Physics/Experimental/PhysScene_Chaos.h
+++ b/Engine/Source/Runtime/Engine/Public/Physics/Experimental/PhysScene_Chaos.h
@@ -179,6 +179,8 @@ public:
 
 	ENGINE_API IPhysicsReplication* GetPhysicsReplication();
 
+	ENGINE_API const IPhysicsReplication* GetPhysicsReplication() const;
+
 	ENGINE_API IPhysicsReplication* CreatePhysicsReplication();
 
 	ENGINE_API IPhysicsReplicationLOD* GetPhysicsReplicationLOD();
diff --git a/Engine/Source/Runtime/Engine/Public/Physics/NetworkPhysicsComponent.h b/Engine/Source/Runtime/Engine/Public/Physics/NetworkPhysicsComponent.h
--- a/Engine/Source/Runtime/Engine/Public/Physics/NetworkPhysicsComponent.h
+++ b/Engine/Source/Runtime/Engine/Public/Physics/NetworkPhysicsComponent.h
@@ -28,6 +28,18 @@ DECLARE_MULTICAST_DELEGATE_OneParam(FOnPostProcessInputsInternal, const int32);
 DECLARE_MULTICAST_DELEGATE_TwoParams(FOnInjectInputsExternal, const int32 /* PhysicsStep */, const int32 /* NumSteps */);
 DECLARE_MULTICAST_DELEGATE_OneParam(FAddResimulationRequestInternal, const int32 /* PhysicsStep */);
 
+namespace UE::NetworkPhysicsUtils
+{
+	// Returns the index of the frame about to be simulated, in the server timeline
+	ENGINE_API int32 GetUpcomingServerFrame_External(UWorld* World);
+
+	// Returns the network physics frame offset on the game thread
+	ENGINE_API const int32 GetNetworkPhysicsTickOffset_External(const UWorld* World);
+
+	// Returns the network physics frame offset on the physics thread
+	ENGINE_API const int32 GetNetworkPhysicsTickOffset_Internal(const Chaos::FPhysicsSolverBase* Solver);
+}
+
 /** Templated data history, holding a data buffer */
 template<typename DataType, bool bLegacyData = false>
 struct TNetRewindHistory : public Chaos::TDataRewindHistory<DataType>
@@ -1812,7 +1824,6 @@ struct FAsyncNetworkPhysicsComponentInput : public Chaos::FSimCallbackInput
 	TOptional<bool> bIsLocallyControlled;
 	TOptional<ENetMode> NetMode;
 	TOptional<ENetRole> NetRole;
-	TOptional<int32> NetworkPhysicsTickOffset;
 	TOptional<uint16> InputsToNetwork_Owner;
 	TOptional<EPhysicsReplicationMode> PhysicsReplicationMode;
 	TOptional<TWeakObjectPtr<UActorComponent>> ActorComponent;
@@ -1839,7 +1850,6 @@ struct FAsyncNetworkPhysicsComponentInput : public Chaos::FSimCallbackInput
 		bIsLocallyControlled.Reset();
 		NetMode.Reset();
 		NetRole.Reset();
-		NetworkPhysicsTickOffset.Reset();
 		InputsToNetwork_Owner.Reset();
 		PhysicsReplicationMode.Reset();
 		ActorComponent.Reset();
@@ -1930,7 +1940,10 @@ public:
 	const FString GetActorName() const { return ActorName; }
 
 	// Get the physics tick offset (add to clients physics tick to get the servers corresponding physics tick)
-	const int32 GetNetworkPhysicsTickOffset() const { return NetworkPhysicsTickOffset; }
+	const int32 GetNetworkPhysicsTickOffset() const
+	{
+		return UE::NetworkPhysicsUtils::GetNetworkPhysicsTickOffset_Internal(GetSolver());
+	}
 
 	/** Get the physics replication mode used
 	* TODO: FPhysicsReplication can change replication mode temporarily on the physics thread, make this correspond to the actual replication mode being used*/
@@ -2024,7 +2037,6 @@ private:
 	bool bIsLocallyControlled;
 	ENetMode NetMode;
 	ENetRole NetRole;
-	int32 NetworkPhysicsTickOffset;
 	EPhysicsReplicationMode PhysicsReplicationMode;
 	FString ActorName;
 	bool bIsUsingLegacyData;
@@ -2106,9 +2118,3 @@ private:
 	FDelegateHandle DelegateOnPostProcessInputs_Internal;
 	FDelegateHandle DelegateOnAddResimulationRequest_Internal;
 };
-
-namespace UE::NetworkPhysicsUtils
-{
-	// Returns the index of the frame about to be simulated, in the server timeline
-	ENGINE_API int32 GetUpcomingServerFrame_External(UWorld* World);
-}
diff --git a/Engine/Source/Runtime/Engine/Public/PhysicsReplication.h b/Engine/Source/Runtime/Engine/Public/PhysicsReplication.h
--- a/Engine/Source/Runtime/Engine/Public/PhysicsReplication.h
+++ b/Engine/Source/Runtime/Engine/Public/PhysicsReplication.h
@@ -73,7 +73,6 @@ struct FPhysicsRepAsyncInputData
 	TOptional<FPhysicsRepErrorCorrectionData> ErrorCorrection;
 	EPhysicsReplicationMode RepMode;
 	int32 ServerFrame;
-	TOptional<int32> FrameOffset;
 	float LatencyOneWay;
 
 	FPhysicsRepAsyncInputData(Chaos::FConstPhysicsObjectHandle POHandle)
@@ -87,9 +86,14 @@ struct FPhysicsReplicationAsyncInput : public Chaos::FSimCallbackInput
 {
 	FPhysicsRepErrorCorrectionData ErrorCorrection;
 	TArray<FPhysicsRepAsyncInputData> InputData;
+	bool NetworkPhysicsTickOffsetAssigned;
+	int32 NetworkPhysicsTickOffset;
+
 	void Reset()
 	{
 		InputData.Reset();
+		NetworkPhysicsTickOffsetAssigned = false;
+		NetworkPhysicsTickOffset = 0;
 	}
 };
 
@@ -211,6 +215,8 @@ public:
 	//~ Begin IPhysicsReplicationAsync interface
 	virtual void RegisterSettings(Chaos::FConstPhysicsObjectHandle PhysicsObject, TWeakPtr<const FNetworkPhysicsSettingsData> InSettings) override;
 	virtual void AddResimulationRequest_Internal(const float DeltaSeconds) override;
+
+	virtual int32 GetNetworkPhysicsTickOffset_Internal() const override { return NetworkPhysicsTickOffset; }
 	//~ End IPhysicsReplicationAsync interface
 
 private:
@@ -229,6 +235,9 @@ private:
 	int32 ResimOutOfBoundsCounter = 0;
 	float ResimErrorLogTimer = 0.0f;
 
+	bool NetworkPhysicsTickOffsetAssigned = false;
+	int32 NetworkPhysicsTickOffset = 0;
+
 private:
 	FReplicatedPhysicsTargetAsync* AddObjectToReplication(Chaos::FConstPhysicsObjectHandle PhysicsObject);
 	void RemoveObjectFromReplication(Chaos::FConstPhysicsObjectHandle PhysicsObject);
@@ -324,6 +333,8 @@ public:
 	/** Remove the replicated target*/
 	ENGINE_API virtual void RemoveReplicatedTarget(UPrimitiveComponent* Component) override;
 
+	virtual int32 GetNetworkPhysicsTickOffset() const override { return NetworkPhysicsTickOffset; }
+
 protected:
 
 	/** Update the physics body state given a set of replicated targets */
@@ -357,5 +368,8 @@ private:
 	FPhysicsReplicationAsync* PhysicsReplicationAsync;
 	FPhysicsReplicationAsyncInput* AsyncInput;	//async data being written into before we push into callback
 
+	bool NetworkPhysicsTickOffsetAssigned = false;
+	int32 NetworkPhysicsTickOffset = 0;
+
 	ENGINE_API void PrepareAsyncData_External(const FRigidBodyErrorCorrection& ErrorCorrection);	//prepare async data for writing. Call on external thread (i.e. game thread)
 };
diff --git a/Engine/Source/Runtime/Engine/Public/PhysicsReplicationInterface.h b/Engine/Source/Runtime/Engine/Public/PhysicsReplicationInterface.h
--- a/Engine/Source/Runtime/Engine/Public/PhysicsReplicationInterface.h
+++ b/Engine/Source/Runtime/Engine/Public/PhysicsReplicationInterface.h
@@ -20,6 +20,9 @@ public:
 	virtual void SetReplicatedTarget(UPrimitiveComponent* Component, FName BoneName, const FRigidBodyState& ReplicatedTarget, int32 ServerFrame) = 0;
 
 	virtual void RemoveReplicatedTarget(UPrimitiveComponent* Component) = 0;
+
+	/** Get the network physics tick offset on the game thread */
+	virtual int32 GetNetworkPhysicsTickOffset() const { return 0; }
 };
 
 class IPhysicsReplicationAsync // Physics Thread API
@@ -32,6 +35,9 @@ public:
 
 	virtual void RegisterSettings(Chaos::FConstPhysicsObjectHandle PhysicsObject, TWeakPtr<const FNetworkPhysicsSettingsData> InSettings) = 0;
 
-	/** Add potential resimulation reques from the physics replication */
+	/** Add potential resimulation request from the physics replication */
 	virtual void AddResimulationRequest_Internal( const float DeltaSeconds) = 0;
+
+	/** Get the network physics tick offset on the physics thread */
+	virtual int32 GetNetworkPhysicsTickOffset_Internal() const { return 0; }
 };
diff --git a/Engine/Source/Runtime/Experimental/Chaos/Public/Chaos/Framework/PhysicsSolverBase.h b/Engine/Source/Runtime/Experimental/Chaos/Public/Chaos/Framework/PhysicsSolverBase.h
--- a/Engine/Source/Runtime/Experimental/Chaos/Public/Chaos/Framework/PhysicsSolverBase.h
+++ b/Engine/Source/Runtime/Experimental/Chaos/Public/Chaos/Framework/PhysicsSolverBase.h
@@ -530,6 +530,11 @@ namespace Chaos
 			return PhysicsReplicationAsync;
 		}
 
+		const IPhysicsReplicationAsync* GetPhysicsReplication_Internal() const
+		{
+			return PhysicsReplicationAsync;
+		}
+
 		void SetPhysicsReplicationLOD_Internal(IPhysicsReplicationLODAsync* InPhysicsReplicationLOD)
 		{
 			PhysicsReplicationLODAsync = InPhysicsReplicationLOD;
