Add input validation to StaticSwitchParameter.
Modify IsMaterialAttributes so it returns if both are MAs, not just one, to fix issues with existing invalid inputs.

[RN] minor | Rendering Architecture

#Synced-CL 49554335

diff --git a/Engine/Source/Runtime/Engine/Private/Materials/MaterialExpressions.cpp b/Engine/Source/Runtime/Engine/Private/Materials/MaterialExpressions.cpp
--- a/Engine/Source/Runtime/Engine/Private/Materials/MaterialExpressions.cpp
+++ b/Engine/Source/Runtime/Engine/Private/Materials/MaterialExpressions.cpp
@@ -8664,18 +8664,50 @@ bool UMaterialExpressionScalarParameter::HasClassAndNameCollision(UMaterialExpre
 bool UMaterialExpressionStaticSwitchParameter::IsResultMaterialAttributes(int32 OutputIndex)
 {
 	check(OutputIndex == 0);
-	// This one is a little tricky. Since we are treating a dangling reroute as an empty expression, this
-	// should early out, whereas IsResultMaterialAttributes on a reroute node will return false as the 
-	// reroute node's input is dangling and therefore its type is unknown.
-	if ((A.GetTracedInput().Expression && A.Expression->IsResultMaterialAttributes(A.OutputIndex)) ||
-		(B.GetTracedInput().Expression && B.Expression->IsResultMaterialAttributes(B.OutputIndex)))
+
+	// Both A and B must be connected in a parameter. 
+	// Because of this, MA output will be valid only if both are MAs
+	return ((A.GetTracedInput().Expression && A.Expression->IsResultMaterialAttributes(A.OutputIndex)) &&
+		(B.GetTracedInput().Expression && B.Expression->IsResultMaterialAttributes(B.OutputIndex)));
+}
+
+EMaterialValueType GetRealInputValueType(const FExpressionInput& InExpression)
+{
+	const FExpressionInput& ExpressionInputRef = InExpression.GetTracedInput();
+	if (ExpressionInputRef.Expression)
 	{
-		return true;
+		if (ExpressionInputRef.Expression->IsResultMaterialAttributes(ExpressionInputRef.OutputIndex))
+		{
+			return MCT_MaterialAttributes;
+		}
+		if (ExpressionInputRef.Expression->IsResultSubstrateMaterial(ExpressionInputRef.OutputIndex))
+		{
+			return MCT_Substrate;
+		}
+		return ExpressionInputRef.Expression->GetOutputValueType(ExpressionInputRef.OutputIndex);
 	}
-	else
+
+	return MCT_Unknown;
+}
+
+EMaterialValueType UMaterialExpressionStaticSwitchParameter::GetInputValueType(int32 InputIndex)
+{
+	if (InputIndex == 0)
 	{
-		return false;
+		if (B.IsConnected())
+		{
+			return GetRealInputValueType(B);
+		}
 	}
+	else if (InputIndex == 1)
+	{
+		if (A.IsConnected())
+		{
+			return GetRealInputValueType(A);
+		}
+	}
+
+	return MCT_Unknown;
 }
 
 EMaterialValueType UMaterialExpressionStaticSwitchParameter::GetOutputValueType(int32 OutputIndex)
@@ -8965,6 +8997,21 @@ EMaterialValueType UMaterialExpressionStaticSwitch::GetInputValueType(int32 Inpu
 {
 	if (InputIndex == 0 || InputIndex == 1)
 	{
+		if (InputIndex == 0)
+		{
+			if (B.IsConnected())
+			{
+				return GetRealInputValueType(B);
+			}
+		}
+		else if (InputIndex == 1)
+		{
+			if (A.IsConnected())
+			{
+				return GetRealInputValueType(A);
+			}
+		}
+
 		return MCT_Unknown;
 	}
 	else
@@ -16913,19 +16960,13 @@ EMaterialValueType UMaterialExpressionFunctionOutput::GetInputValueType(int32 In
 
 EMaterialValueType UMaterialExpressionFunctionOutput::GetOutputValueType(int32 OutputIndex)
 {
-	const FExpressionInput& ExpressionInputRef = A.GetTracedInput();
-	if (ExpressionInputRef.Expression)
+	const EMaterialValueType RealInputValueType = GetRealInputValueType(A);
+
+	if (RealInputValueType != MCT_Unknown)
 	{
-		if (ExpressionInputRef.Expression->IsResultMaterialAttributes(ExpressionInputRef.OutputIndex))
-		{
-			return MCT_MaterialAttributes;
-		}
-		if (ExpressionInputRef.Expression->IsResultSubstrateMaterial(ExpressionInputRef.OutputIndex))
-		{
-			return MCT_Substrate;
-		}
-		return ExpressionInputRef.Expression->GetOutputValueType(ExpressionInputRef.OutputIndex);
+		return RealInputValueType;
 	}
+
 	return Super::GetOutputValueType(OutputIndex);
 }
 
diff --git a/Engine/Source/Runtime/Engine/Public/Materials/MaterialExpressionStaticSwitchParameter.h b/Engine/Source/Runtime/Engine/Public/Materials/MaterialExpressionStaticSwitchParameter.h
--- a/Engine/Source/Runtime/Engine/Public/Materials/MaterialExpressionStaticSwitchParameter.h
+++ b/Engine/Source/Runtime/Engine/Public/Materials/MaterialExpressionStaticSwitchParameter.h
@@ -30,7 +30,7 @@ public:
 	virtual int32 Compile(class FMaterialCompiler* Compiler, int32 OutputIndex) override;
 	virtual void GetCaption(TArray<FString>& OutCaptions) const override;
 	virtual bool IsResultMaterialAttributes(int32 OutputIndex) override;
-	virtual EMaterialValueType GetInputValueType(int32 InputIndex) override {return MCT_Unknown;}
+	virtual EMaterialValueType GetInputValueType(int32 InputIndex) override;
 	virtual EMaterialValueType GetOutputValueType(int32 OutputIndex) override;
 	virtual bool IsStaticExpression() const override { return true; }
 
