Add local vertex PSO precache entry when RT dynamic geometry update is set when nanite vertex factory is used on a static mesh (nanite vertex factory doesn't support RT dynamic geometry update but fallback mesh is used for that)

#Synced-CL 49557757

diff --git a/Engine/Source/Runtime/Engine/Public/StaticMeshComponentHelper.h b/Engine/Source/Runtime/Engine/Public/StaticMeshComponentHelper.h
--- a/Engine/Source/Runtime/Engine/Public/StaticMeshComponentHelper.h
+++ b/Engine/Source/Runtime/Engine/Public/StaticMeshComponentHelper.h
@@ -293,6 +293,28 @@ void FStaticMeshComponentHelper::CollectPSOPrecacheData(const T& Component, cons
 	if (Nanite::FNaniteResourcesHelper::ShouldCreateNaniteProxy(Component, &NaniteMaterials))
 	{
 		FStaticMeshComponentHelper::CollectPSOPrecacheDataImpl(Component, &FNaniteVertexFactory::StaticType, BasePrecachePSOParams, SMC_GetElements, OutParams);
+
+		// If bRequireRTDynamicGeometryUpdate is set then add another entry with local vertex factory because nanite can't be used for dynamic RT geometry updates
+		FMaterialInterfacePSOPrecacheParamsList RTDynamicGeometryUpdateParams;
+		for (FMaterialInterfacePSOPrecacheParams& Params : OutParams)
+		{
+			if (Params.PSOPrecacheParams.bRequireRTDynamicGeometryUpdate)
+			{
+				// Add new entry with local vertex factory and disabling rendering to all passes to reduce number of PSOs to compile
+				FMaterialInterfacePSOPrecacheParams NewEntry = Params;
+				NewEntry.VertexFactoryDataList.Empty(1);
+				NewEntry.VertexFactoryDataList.Add(FPSOPrecacheVertexFactoryData(&FLocalVertexFactory::StaticType));
+				NewEntry.PSOPrecacheParams.bRenderInMainPass = false;
+				NewEntry.PSOPrecacheParams.bRenderInDepthPass = false;
+				NewEntry.PSOPrecacheParams.bCastShadow = false;
+				NewEntry.PSOPrecacheParams.bRenderCustomDepth = false;
+				RTDynamicGeometryUpdateParams.Add(NewEntry);
+
+				// Nanite VF entry doesn't require RT dynamic geometry update anymore
+				Params.PSOPrecacheParams.bRequireRTDynamicGeometryUpdate = false;
+			}
+		}
+		OutParams.Append(RTDynamicGeometryUpdateParams);
 	}
 	else
 	{
