[World] Old DataLayer deprecation continuation part 2 - Remove UDEPRECATED_DataLayer related APIs

#tests Tested data layers in gyms

#Synced-CL 49650926

diff --git a/Engine/Source/Editor/DataLayerEditor/Private/DataLayer/DataLayerEditorSubsystem.cpp b/Engine/Source/Editor/DataLayerEditor/Private/DataLayer/DataLayerEditorSubsystem.cpp
--- a/Engine/Source/Editor/DataLayerEditor/Private/DataLayer/DataLayerEditorSubsystem.cpp
+++ b/Engine/Source/Editor/DataLayerEditor/Private/DataLayer/DataLayerEditorSubsystem.cpp
@@ -1705,16 +1705,6 @@ bool UDataLayerEditorSubsystem::ResetUserSettings()
 	return true;
 }
 
-bool UDataLayerEditorSubsystem::HasDeprecatedDataLayers() const
-{
-	UWorld* World = GetWorld();
-	if (AWorldDataLayers* WorldDataLayers = World ? World->GetWorldDataLayers() : nullptr)
-	{
-		return WorldDataLayers->HasDeprecatedDataLayers();
-	}
-	return false;
-}
-
 bool UDataLayerEditorSubsystem::PassDataLayersFilter(UWorld* World, const FWorldPartitionHandle& ActorHandle)
 {
 	UWorld* OwningWorld = World->PersistentLevel->GetWorld();
@@ -1807,58 +1797,55 @@ UDataLayerInstance* UDataLayerEditorSubsystem::CreateDataLayerInstance(const FDa
 
 	if (AWorldDataLayers* WorldDataLayers = Parameters.WorldDataLayers != nullptr ? Parameters.WorldDataLayers.Get() : (GetWorld() ? GetWorld()->GetWorldDataLayers() : nullptr))
 	{
-		if (!WorldDataLayers->HasDeprecatedDataLayers())
+		if (Parameters.bIsPrivate)
 		{
-			if (Parameters.bIsPrivate)
-			{
-				NewDataLayerInstance = CreateDataLayerInstance<UDataLayerInstancePrivate>(WorldDataLayers);
-			}
-			// Don't create an instance if no valid asset is provided
-			else if (Parameters.DataLayerAsset)
+			NewDataLayerInstance = CreateDataLayerInstance<UDataLayerInstancePrivate>(WorldDataLayers);
+		}
+		// Don't create an instance if no valid asset is provided
+		else if (Parameters.DataLayerAsset)
+		{
+			if (UExternalDataLayerAsset* ExternalDataLayerAsset = Cast<UExternalDataLayerAsset>(Parameters.DataLayerAsset))
 			{
-				if (UExternalDataLayerAsset* ExternalDataLayerAsset = Cast<UExternalDataLayerAsset>(Parameters.DataLayerAsset))
+				UWorld* OuterWorld = WorldDataLayers->GetTypedOuter<UWorld>();
+				if (UExternalDataLayerManager* ExternalDataLayerManager = UExternalDataLayerManager::GetExternalDataLayerManager(OuterWorld))
 				{
-					UWorld* OuterWorld = WorldDataLayers->GetTypedOuter<UWorld>();
-					if (UExternalDataLayerManager* ExternalDataLayerManager = UExternalDataLayerManager::GetExternalDataLayerManager(OuterWorld))
+					FText FailureReason;
+					if (ExternalDataLayerManager->CanInjectExternalDataLayerAsset(ExternalDataLayerAsset, &FailureReason))
 					{
-						FText FailureReason;
-						if (ExternalDataLayerManager->CanInjectExternalDataLayerAsset(ExternalDataLayerAsset, &FailureReason))
+						const bool bAllowCreate = true;
+						if (AWorldDataLayers* EDLWorldDataLayers = ExternalDataLayerManager->GetWorldDataLayers(ExternalDataLayerAsset, bAllowCreate))
 						{
-							const bool bAllowCreate = true;
-							if (AWorldDataLayers* EDLWorldDataLayers = ExternalDataLayerManager->GetWorldDataLayers(ExternalDataLayerAsset, bAllowCreate))
+							UExternalDataLayerInstance* ExternalDataLayerInstance = EDLWorldDataLayers->GetExternalDataLayerInstance(ExternalDataLayerAsset);
+							if (ensure(!ExternalDataLayerInstance))
 							{
-								UExternalDataLayerInstance* ExternalDataLayerInstance = EDLWorldDataLayers->GetExternalDataLayerInstance(ExternalDataLayerAsset);
-								if (ensure(!ExternalDataLayerInstance))
+								// Create External Data Layer Instance
+								ExternalDataLayerInstance = EDLWorldDataLayers->CreateDataLayer<UExternalDataLayerInstance>(ExternalDataLayerAsset);
+								if (ensure(ExternalDataLayerInstance))
 								{
-									// Create External Data Layer Instance
-									ExternalDataLayerInstance = EDLWorldDataLayers->CreateDataLayer<UExternalDataLayerInstance>(ExternalDataLayerAsset);
-									if (ensure(ExternalDataLayerInstance))
+									if (ExternalDataLayerManager->InjectExternalDataLayer(ExternalDataLayerAsset))
 									{
-										if (ExternalDataLayerManager->InjectExternalDataLayer(ExternalDataLayerAsset))
-										{
-											NewDataLayerInstance = const_cast<UExternalDataLayerInstance*>(ExternalDataLayerManager->GetExternalDataLayerInstance(ExternalDataLayerAsset));
-											ensure(NewDataLayerInstance == ExternalDataLayerInstance);
-										}
-										else
-										{
-											UE_LOG(LogDataLayerEditorSubsystem, Warning, TEXT("[EDL %s] Failed to inject External Data Layer."), *ExternalDataLayerAsset->GetName());
-										}
+										NewDataLayerInstance = const_cast<UExternalDataLayerInstance*>(ExternalDataLayerManager->GetExternalDataLayerInstance(ExternalDataLayerAsset));
+										ensure(NewDataLayerInstance == ExternalDataLayerInstance);
+									}
+									else
+									{
+										UE_LOG(LogDataLayerEditorSubsystem, Warning, TEXT("[EDL %s] Failed to inject External Data Layer."), *ExternalDataLayerAsset->GetName());
 									}
 								}
 							}
 						}
-						else
-						{
-							UE_LOG(LogDataLayerEditorSubsystem, Warning, TEXT("%s"), *FailureReason.ToString());
-							LastWarningNotification = FailureReason;
-						}
 					}
-					UE_CLOG(!NewDataLayerInstance, LogDataLayerEditorSubsystem, Error, TEXT("[EDL %s] Failed to create External Data Layer Instance."), *ExternalDataLayerAsset->GetName());
-				}
-				else
-				{
-					NewDataLayerInstance = CreateDataLayerInstance<UDataLayerInstanceWithAsset>(WorldDataLayers, Parameters.DataLayerAsset);
+					else
+					{
+						UE_LOG(LogDataLayerEditorSubsystem, Warning, TEXT("%s"), *FailureReason.ToString());
+						LastWarningNotification = FailureReason;
+					}
 				}
+				UE_CLOG(!NewDataLayerInstance, LogDataLayerEditorSubsystem, Error, TEXT("[EDL %s] Failed to create External Data Layer Instance."), *ExternalDataLayerAsset->GetName());
+			}
+			else
+			{
+				NewDataLayerInstance = CreateDataLayerInstance<UDataLayerInstanceWithAsset>(WorldDataLayers, Parameters.DataLayerAsset);
 			}
 		}
 		else
diff --git a/Engine/Source/Editor/DataLayerEditor/Private/DataLayer/DataLayerMode.cpp b/Engine/Source/Editor/DataLayerEditor/Private/DataLayer/DataLayerMode.cpp
--- a/Engine/Source/Editor/DataLayerEditor/Private/DataLayer/DataLayerMode.cpp
+++ b/Engine/Source/Editor/DataLayerEditor/Private/DataLayer/DataLayerMode.cpp
@@ -801,13 +801,6 @@ FSceneOutlinerDragValidationInfo FDataLayerMode::ValidateDataLayerAssetDrop(cons
 	check(!DataLayerAssetsToDrop.IsEmpty());
 	const FString CreateOrMoveString(bIsMovingDataLayerInstances ? "move" : "create");
 
-	if (DataLayerEditorSubsystem->HasDeprecatedDataLayers())
-	{
-		return FSceneOutlinerDragValidationInfo(ESceneOutlinerDropCompatibility::Incompatible, FText::Format(LOCTEXT("CantCreateOrMoveInstanceWorldHasDeprecatedDataLayers", "Cannot {0} Data Layers Instance from assets since \"{1}\" has deprecated data layers."),
-			FText::FromString(CreateOrMoveString), 
-			FText::FromString(GetOwningWorld()->GetName())));
-	}
-
 	const AWorldDataLayers* DropTargetWorldDataLayers = GetWorldDataLayersFromTreeItem(DropTarget);
 	const UDataLayerInstance* DropTargetDataLayerInstance = GetDataLayerInstanceFromTreeItem(DropTarget);
 	const UDataLayerInstanceWithAsset* DropTargetDataLayerWithAsset = DropTargetDataLayerInstance ? Cast<UDataLayerInstanceWithAsset>(DropTargetDataLayerInstance) : nullptr;
@@ -1476,15 +1469,7 @@ void FDataLayerMode::RegisterContextMenu()
 
 				auto CreateNewDataLayer = [SceneOutliner, Mode, CreateNewDataLayerInternal](UDataLayerInstance* InParentDataLayer = nullptr, const UDataLayerAsset* InDataLayerAsset = nullptr, bool bInIsPrivate = false) -> UDataLayerInstance*
 				{
-					if (UDataLayerEditorSubsystem::Get()->HasDeprecatedDataLayers())
-					{
-						check(InDataLayerAsset == nullptr);
-						return CreateNewDataLayerInternal(nullptr, false);
-					}
-					else
-					{
-						return CreateNewDataLayerInternal(InDataLayerAsset, bInIsPrivate);
-					}
+					return CreateNewDataLayerInternal(InDataLayerAsset, bInIsPrivate);
 				};
 
 				auto CreateNewEmptyDataLayer = [SceneOutliner, Mode, CreateNewDataLayerInternal](bool bInIsPrivate = false) -> UDataLayerInstance*
@@ -1493,7 +1478,7 @@ void FDataLayerMode::RegisterContextMenu()
 				};
 
 				const AWorldDataLayers* OwningWorldAWorldDataLayers = Mode->GetOwningWorldAWorldDataLayers();
-				if (OwningWorldAWorldDataLayers && !OwningWorldAWorldDataLayers->HasDeprecatedDataLayers())
+				if (OwningWorldAWorldDataLayers)
 				{
 					const AWorldDataLayers* TargetWorldDataLayers = GetBestCandidateWorldDataLayersFromSelection(SceneOutliner, Mode);
 					Section.AddSubMenu("CreateNewDataLayerWithAsset", LOCTEXT("CreateNewDataLayerWithAssetSubMenu", "Create New Data Layer With Asset"), LOCTEXT("CreateNewDataLayerWithAssetSubMenu_ToolTip", "Create New Data Layer With Asset"),
@@ -1553,11 +1538,7 @@ void FDataLayerMode::RegisterContextMenu()
 				}
 
 				Section.AddMenuEntry("CreateNewDataLayerPrivate", LOCTEXT("CreateNewDataLayerPrivate", "Create New Private Data Layer"), LOCTEXT("CreateNewDataLayerPrivateToolTip", "Creates an Editor Only Data Layer that cannot be used by other Worlds."), FSlateIcon(),
-					FUIAction(
-						FExecuteAction::CreateLambda([CreateNewEmptyDataLayer]() { CreateNewEmptyDataLayer(true); }),
-						FCanExecuteAction(),
-						FIsActionChecked(),
-						FIsActionButtonVisible::CreateLambda([]() { return !UDataLayerEditorSubsystem::Get()->HasDeprecatedDataLayers(); })));
+					FUIAction(FExecuteAction::CreateLambda([CreateNewEmptyDataLayer]() { CreateNewEmptyDataLayer(true); } )));
 
 				const bool bAllSelectedDataLayersCanBeRemoved = Algo::AllOf(SelectedDataLayers, [](const UDataLayerInstance* DataLayerInstance) { return DataLayerInstance->CanBeRemoved(); });
 				Section.AddMenuEntry("DeleteSelectedDataLayers", LOCTEXT("DeleteSelectedDataLayers", "Delete Selected Data Layer(s)"), FText(), FSlateIcon(),
@@ -1695,24 +1676,6 @@ void FDataLayerMode::RegisterContextMenu()
 						})
 					));
 
-				if (UDataLayerEditorSubsystem::Get()->HasDeprecatedDataLayers())
-				{
-					Section.AddMenuEntry("RenameSelectedDataLayer", LOCTEXT("RenameSelectedDataLayer", "Rename Selected Data Layer"), FText(), FSlateIcon(),
-						FUIAction(
-							FExecuteAction::CreateLambda([Mode, SceneOutliner, SelectedDataLayers]() {
-							if (SelectedDataLayers.Num() == 1)
-							{
-								FSceneOutlinerTreeItemPtr ItemToRename = SceneOutliner->GetTreeItem(SelectedDataLayers[0]);
-								if (ItemToRename.IsValid() && Mode->CanRenameItem(*ItemToRename) && ItemToRename->CanInteract())
-								{
-									SceneOutliner->SetPendingRenameItem(ItemToRename);
-									SceneOutliner->ScrollItemIntoView(ItemToRename);
-								}
-							}}),
-							FCanExecuteAction::CreateLambda([SelectedDataLayers] { return (SelectedDataLayers.Num() == 1) && !SelectedDataLayers[0]->IsReadOnly(); })
-						));
-				}
-
 				Section.AddSeparator("SectionsSeparator");
 			}
 
@@ -1887,33 +1850,30 @@ void FDataLayerMode::RegisterContextMenu()
 					));
 			}
 
-			if (!UDataLayerEditorSubsystem::Get()->HasDeprecatedDataLayers())
+			FToolMenuSection& Section = InMenu->AddSection("AssetOptionsSection", LOCTEXT("AssetOptionsText", "Asset Options"));
+			Section.AddMenuEntryWithCommandList(FGlobalEditorCommonCommands::Get().FindInContentBrowser, Mode->Commands);
+			const UExternalDataLayerInstance* ExternalDataLayerInstance = SelectedDataLayers.Num() == 1 ? Cast<UExternalDataLayerInstance>(SelectedDataLayers[0]) : nullptr;
+			if (const UExternalDataLayerAsset* ExternalDataLayerAsset = ExternalDataLayerInstance ? ExternalDataLayerInstance->GetExternalDataLayerAsset() : nullptr)
 			{
-				FToolMenuSection& Section = InMenu->AddSection("AssetOptionsSection", LOCTEXT("AssetOptionsText", "Asset Options"));
-				Section.AddMenuEntryWithCommandList(FGlobalEditorCommonCommands::Get().FindInContentBrowser, Mode->Commands);
-				const UExternalDataLayerInstance* ExternalDataLayerInstance = SelectedDataLayers.Num() == 1 ? Cast<UExternalDataLayerInstance>(SelectedDataLayers[0]) : nullptr;
-				if (const UExternalDataLayerAsset* ExternalDataLayerAsset = ExternalDataLayerInstance ? ExternalDataLayerInstance->GetExternalDataLayerAsset() : nullptr)
-				{
-					Section.AddMenuEntry("BrowseToPluginAsset", LOCTEXT("BrowseToPluginAsset", "Browse to Plugin Asset"), FText(), FSlateIcon(FAppStyle::GetAppStyleSetName(), "SystemWideCommands.FindInContentBrowser"),
-						FUIAction(
-							FExecuteAction::CreateLambda([ExternalDataLayerAsset]()
+				Section.AddMenuEntry("BrowseToPluginAsset", LOCTEXT("BrowseToPluginAsset", "Browse to Plugin Asset"), FText(), FSlateIcon(FAppStyle::GetAppStyleSetName(), "SystemWideCommands.FindInContentBrowser"),
+					FUIAction(
+						FExecuteAction::CreateLambda([ExternalDataLayerAsset]()
+						{
+							UExternalDataLayerEngineSubsystem& ExternalDataLayerEngineSubsystem = UExternalDataLayerEngineSubsystem::Get();
+							TArray<UObject*> Objects;;
+							for (UObject* Client : ExternalDataLayerEngineSubsystem.GetClientsForExternalDataLayerAsset(ExternalDataLayerAsset))
 							{
-								UExternalDataLayerEngineSubsystem& ExternalDataLayerEngineSubsystem = UExternalDataLayerEngineSubsystem::Get();
-								TArray<UObject*> Objects;;
-								for (UObject* Client : ExternalDataLayerEngineSubsystem.GetClientsForExternalDataLayerAsset(ExternalDataLayerAsset))
-								{
-									if (UDataAsset* DataAsset = Client->GetTypedOuter<UDataAsset>())
-									{
-										Objects.Add(DataAsset);
-									}
-								}
-								if (!Objects.IsEmpty())
+								if (UDataAsset* DataAsset = Client->GetTypedOuter<UDataAsset>())
 								{
-									GEditor->SyncBrowserToObjects(Objects);
+									Objects.Add(DataAsset);
 								}
-							})
-						));
-				}
+							}
+							if (!Objects.IsEmpty())
+							{
+								GEditor->SyncBrowserToObjects(Objects);
+							}
+						})
+					));
 			}
 		}));
 	}
diff --git a/Engine/Source/Editor/DataLayerEditor/Private/DataLayer/SDataLayerBrowser.cpp b/Engine/Source/Editor/DataLayerEditor/Private/DataLayer/SDataLayerBrowser.cpp
--- a/Engine/Source/Editor/DataLayerEditor/Private/DataLayer/SDataLayerBrowser.cpp
+++ b/Engine/Source/Editor/DataLayerEditor/Private/DataLayer/SDataLayerBrowser.cpp
@@ -138,14 +138,6 @@ void SDataLayerBrowser::Construct(const FArguments& InArgs)
 			return FString();
 	});
 
-	SAssignNew(DeprecatedDataLayerWarningBox, SMultiLineEditableTextBox)
-		.IsReadOnly(true)
-		.Font(IDetailLayoutBuilder::GetDetailFontBold())
-		.BackgroundColor(FAppStyle::GetColor("ErrorReporting.WarningBackgroundColor"))
-		.Text(LOCTEXT("Deprecated_DataLayers", "Some data within DataLayers is deprecated. Run DataLayerToAssetCommandlet to create DataLayerInstances and DataLayer Assets for this level."))
-		.AutoWrapText(true)
-		.Visibility_Lambda([]() { return UDataLayerEditorSubsystem::Get()->HasDeprecatedDataLayers() ? EVisibility::Visible : EVisibility::Collapsed; });
-
 	FSceneOutlinerInitializationOptions InitOptions;
 	InitOptions.OutlinerIdentifier = TEXT("DataLayerEditorOutliner");
 	InitOptions.bShowHeaderRow = true;
@@ -171,11 +163,6 @@ void SDataLayerBrowser::Construct(const FArguments& InArgs)
 	[
 		SNew(SVerticalBox)
 		+ SVerticalBox::Slot()
-		.AutoHeight()
-		[
-			DeprecatedDataLayerWarningBox.ToSharedRef()
-		]
-		+ SVerticalBox::Slot()
 		[
 			// Data Layer Outliner
 			SNew(SSplitter)
diff --git a/Engine/Source/Editor/DataLayerEditor/Private/DataLayer/SDataLayerBrowser.h b/Engine/Source/Editor/DataLayerEditor/Private/DataLayer/SDataLayerBrowser.h
--- a/Engine/Source/Editor/DataLayerEditor/Private/DataLayer/SDataLayerBrowser.h
+++ b/Engine/Source/Editor/DataLayerEditor/Private/DataLayer/SDataLayerBrowser.h
@@ -44,5 +44,4 @@ private:
 	TSharedPtr<SVerticalBox> ContentAreaBox;
 	TSharedPtr<SBorder> DataLayerContentsSection;
 	TSharedPtr<SBorder> DataLayerContentsHeader;
-	TSharedPtr<SMultiLineEditableTextBox> DeprecatedDataLayerWarningBox;
 };
diff --git a/Engine/Source/Editor/DataLayerEditor/Public/DataLayer/DataLayerEditorSubsystem.h b/Engine/Source/Editor/DataLayerEditor/Public/DataLayer/DataLayerEditorSubsystem.h
--- a/Engine/Source/Editor/DataLayerEditor/Public/DataLayer/DataLayerEditorSubsystem.h
+++ b/Engine/Source/Editor/DataLayerEditor/Public/DataLayer/DataLayerEditorSubsystem.h
@@ -40,7 +40,6 @@ class FDataLayersBroadcast;
 class FLevelEditorViewportClient;
 class FSubsystemCollectionBase;
 class SWidget;
-class UDEPRECATED_DataLayer;
 class UDataLayerAsset;
 class UEditorEngine;
 class ULevel;
@@ -691,7 +690,8 @@ public:
 	/**
 	* Whether there are any deprecated DataLayerInstance found
 	*/
-	UE_API bool HasDeprecatedDataLayers() const;
+	UE_DEPRECATED(5.8, "Deprecated data layers should have been removed by now.")
+	bool HasDeprecatedDataLayers() const { return false; }
 
 	/**
 	 * Assign new unique short name to DataLayerInstance if it supports it. 
@@ -712,38 +712,32 @@ public:
 
 	PRAGMA_DISABLE_DEPRECATION_WARNINGS
 
-	UE_DEPRECATED(5.0, "Per-view Data Layer visibility was removed.")
+	UE_DEPRECATED(all, "Per-view Data Layer visibility was removed.")
 	UFUNCTION(BlueprintCallable, Category = DataLayers, meta = (DeprecatedFunction, DeprecationMessage = "Per-view Data Layer visibility was removed."))
 	void UpdateAllViewVisibility(UDEPRECATED_DataLayer* DataLayerThatChanged) {}
 
-	UE_DEPRECATED(5.0, "Per-view Data Layer visibility was removed.")
-	void UpdatePerViewVisibility(FLevelEditorViewportClient* ViewportClient, UDEPRECATED_DataLayer* DataLayerThatChanged = nullptr) {}
-
-	UE_DEPRECATED(5.0, "Per-view Data Layer visibility was removed.")
-	void UpdateActorViewVisibility(FLevelEditorViewportClient* ViewportClient, AActor* Actor, const bool bReregisterIfDirty = true) {}
-
-	UE_DEPRECATED(5.0, "Per-view Data Layer visibility was removed.")
+	UE_DEPRECATED(all, "Per-view Data Layer visibility was removed.")
 	UFUNCTION(BlueprintCallable, Category = DataLayers, meta = (DeprecatedFunction, DeprecationMessage = "Per-view Data Layer visibility was removed."))
 	void UpdateActorAllViewsVisibility(AActor* Actor) {}
 		
-	UE_DEPRECATED(5.0, "Use SetDataLayerIsLoadedInEditor() instead.")
+	UE_DEPRECATED(all, "Use SetDataLayerIsLoadedInEditor() instead.")
 	UFUNCTION(BlueprintCallable, Category = DataLayers, meta = (DeprecatedFunction, DeprecationMessage = "Use SetDataLayerIsLoadedInEditor instead"))
 	bool SetDataLayerIsDynamicallyLoadedInEditor(UDEPRECATED_DataLayer* DataLayer, const bool bIsLoadedInEditor, const bool bIsFromUserChange) { return false; }
-
-	UE_DEPRECATED(5.0, "Use SetDataLayersIsLoadedInEditor() instead.")
+	
+	UE_DEPRECATED(all, "Use SetDataLayersIsLoadedInEditor() instead.")
 	UFUNCTION(BlueprintCallable, Category = DataLayers, meta = (DeprecatedFunction, DeprecationMessage = "Use SetDataLayersIsLoadedInEditor instead"))
 	bool SetDataLayersIsDynamicallyLoadedInEditor(const TArray<UDEPRECATED_DataLayer*>& DataLayers, const bool bIsLoadedInEditor, const bool bIsFromUserChange) { return false;  }
-
-	UE_DEPRECATED(5.0, "Use ToggleDataLayerIsLoadedInEditor() instead.")
+	
+	UE_DEPRECATED(all, "Use ToggleDataLayerIsLoadedInEditor() instead.")
 	UFUNCTION(BlueprintCallable, Category = DataLayers, meta = (DeprecatedFunction, DeprecationMessage = "Use ToggleDataLayerIsLoadedInEditor instead"))
 	bool ToggleDataLayerIsDynamicallyLoadedInEditor(UDEPRECATED_DataLayer* DataLayer, const bool bIsFromUserChange) { return false; }
-
-	UE_DEPRECATED(5.0, "Use ToggleDataLayersIsLoadedInEditor() instead.")
+	
+	UE_DEPRECATED(all, "Use ToggleDataLayersIsLoadedInEditor() instead.")
 	UFUNCTION(BlueprintCallable, Category = DataLayers, meta = (DeprecatedFunction, DeprecationMessage = "Use ToggleDataLayersIsLoadedInEditor instead"))
 	bool ToggleDataLayersIsDynamicallyLoadedInEditor(const TArray<UDEPRECATED_DataLayer*>& DataLayers, const bool bIsFromUserChange) { return false; }
-
+	
 	PRAGMA_ENABLE_DEPRECATION_WARNINGS
-			
+
 	UE_DEPRECATED(5.1, "Use GetDataLayerInstance instead")
 	UFUNCTION(BlueprintCallable, Category = DataLayers)
 	UE_API UDataLayerInstance* GetDataLayer(const FActorDataLayer& ActorDataLayer) const;
diff --git a/Engine/Source/Runtime/Engine/Classes/GameFramework/Actor.h b/Engine/Source/Runtime/Engine/Classes/GameFramework/Actor.h
--- a/Engine/Source/Runtime/Engine/Classes/GameFramework/Actor.h
+++ b/Engine/Source/Runtime/Engine/Classes/GameFramework/Actor.h
@@ -41,7 +41,6 @@ struct FNetworkObjectInfo;
 class FActorTransactionAnnotation;
 class FComponentInstanceDataCache;
 struct FWorldRenameFromRootContext;
-class UDEPRECATED_DataLayer;
 class UDataLayerAsset;
 class UExternalDataLayerAsset;
 class UDataLayerInstance;
@@ -1493,16 +1492,6 @@ public:
 
 	//~ Begin Deprecated
 
-	PRAGMA_DISABLE_DEPRECATION_WARNINGS
-
-	UE_DEPRECATED(5.1, "Convert DataLayer using UDataLayerToAssetCommandlet and use AddDataLayer(UDataLayerInstance*)")
-	ENGINE_API bool AddDataLayer(const UDEPRECATED_DataLayer* DataLayer);
-
-	UE_DEPRECATED(5.1, "Convert DataLayer using UDataLayerToAssetCommandlet and use RemoveDataLayer(UDataLayerInstance*)")
-	ENGINE_API bool RemoveDataLayer(const UDEPRECATED_DataLayer* DataLayer);
-
-	PRAGMA_ENABLE_DEPRECATION_WARNINGS
-
 	UE_DEPRECATED(5.3, "Use AActor::SupportsDataLayerType(TSubclassOf<UDataLayerInstance>) instead.")
 	ENGINE_API bool SupportsDataLayer() const;
 
diff --git a/Engine/Source/Runtime/Engine/Private/ActorEditor.cpp b/Engine/Source/Runtime/Engine/Private/ActorEditor.cpp
--- a/Engine/Source/Runtime/Engine/Private/ActorEditor.cpp
+++ b/Engine/Source/Runtime/Engine/Private/ActorEditor.cpp
@@ -11,7 +11,6 @@
 #include "Engine/BlueprintGeneratedClass.h"
 #include "WorldPartition/DataLayer/WorldDataLayers.h"
 #include "WorldPartition/WorldPartition.h"
-#include "WorldPartition/DataLayer/DataLayer.h"
 #include "WorldPartition/DataLayer/DataLayerManager.h"
 #include "EditorSupportDelegates.h"
 #include "Logging/MessageLog.h"
@@ -2194,24 +2193,6 @@ bool AActor::RemoveDataLayer(const FActorDataLayer& ActorDataLayer)
 	return false;
 }
 
-bool AActor::AddDataLayer(const UDEPRECATED_DataLayer* DataLayer)
-{
-	if (SupportsDataLayerType(UDataLayerInstance::StaticClass()) && DataLayer)
-	{
-		return AddDataLayer(FActorDataLayer(DataLayer->GetFName()));
-	}
-	return false;
-}
-
-bool AActor::RemoveDataLayer(const UDEPRECATED_DataLayer* DataLayer)
-{
-	if (DataLayer)
-	{
-		return RemoveDataLayer(FActorDataLayer(DataLayer->GetFName()));
-	}
-	return false;
-}
-
 PRAGMA_ENABLE_DEPRECATION_WARNINGS
 
 bool AActor::ContainsDataLayer(const FActorDataLayer& ActorDataLayer) const
diff --git a/Engine/Source/Runtime/Engine/Private/WorldPartition/DataLayer/DataLayer.cpp b/Engine/Source/Runtime/Engine/Private/WorldPartition/DataLayer/DataLayer.cpp
deleted file mode 100644
diff --git a/Engine/Source/Runtime/Engine/Private/WorldPartition/DataLayer/DataLayerInstancePrivate.cpp b/Engine/Source/Runtime/Engine/Private/WorldPartition/DataLayer/DataLayerInstancePrivate.cpp
--- a/Engine/Source/Runtime/Engine/Private/WorldPartition/DataLayer/DataLayerInstancePrivate.cpp
+++ b/Engine/Source/Runtime/Engine/Private/WorldPartition/DataLayer/DataLayerInstancePrivate.cpp
@@ -31,8 +31,6 @@ FName UDataLayerInstancePrivate::MakeName()
 
 void UDataLayerInstancePrivate::OnCreated()
 {
-	check(!GetDirectOuterWorldDataLayers()->HasDeprecatedDataLayers() || IsRunningCommandlet());
-
 	Modify(/*bAlwaysMarkDirty*/false);
 	FDataLayerUtils::SetDataLayerShortName(this, TEXT("DataLayer"));
 		
diff --git a/Engine/Source/Runtime/Engine/Private/WorldPartition/DataLayer/DataLayerInstanceWithAsset.cpp b/Engine/Source/Runtime/Engine/Private/WorldPartition/DataLayer/DataLayerInstanceWithAsset.cpp
--- a/Engine/Source/Runtime/Engine/Private/WorldPartition/DataLayer/DataLayerInstanceWithAsset.cpp
+++ b/Engine/Source/Runtime/Engine/Private/WorldPartition/DataLayer/DataLayerInstanceWithAsset.cpp
@@ -61,7 +61,6 @@ TSubclassOf<UDataLayerInstanceWithAsset> UDataLayerInstanceWithAsset::GetDataLay
 void UDataLayerInstanceWithAsset::OnCreated(const UDataLayerAsset* Asset)
 {
 	checkf(Asset->IsA<UExternalDataLayerAsset>() == this->IsA<UExternalDataLayerInstance>(), TEXT("Only ExternalDataLayerInstance can reference ExternalDataLayerAsset"));
-	check(!GetDirectOuterWorldDataLayers()->HasDeprecatedDataLayers() || IsRunningCommandlet());
 
 	Modify(/*bAlwaysMarkDirty*/false);
 
diff --git a/Engine/Source/Runtime/Engine/Private/WorldPartition/DataLayer/DataLayerManager.cpp b/Engine/Source/Runtime/Engine/Private/WorldPartition/DataLayer/DataLayerManager.cpp
--- a/Engine/Source/Runtime/Engine/Private/WorldPartition/DataLayer/DataLayerManager.cpp
+++ b/Engine/Source/Runtime/Engine/Private/WorldPartition/DataLayer/DataLayerManager.cpp
@@ -103,6 +103,26 @@ FAutoConsoleCommand UDataLayerManager::ToggleDataLayerActivation(
 	})
 );
 
+inline bool GetDataLayerRuntimeStateFromName(const FString& InStateName, EDataLayerRuntimeState& OutState)
+{
+	if (InStateName.Equals(GetDataLayerRuntimeStateName(EDataLayerRuntimeState::Unloaded), ESearchCase::IgnoreCase))
+	{
+		OutState = EDataLayerRuntimeState::Unloaded;
+		return true;
+	}
+	else if (InStateName.Equals(GetDataLayerRuntimeStateName(EDataLayerRuntimeState::Loaded), ESearchCase::IgnoreCase))
+	{
+		OutState = EDataLayerRuntimeState::Loaded;
+		return true;
+	}
+	else if (InStateName.Equals(GetDataLayerRuntimeStateName(EDataLayerRuntimeState::Activated), ESearchCase::IgnoreCase))
+	{
+		OutState = EDataLayerRuntimeState::Activated;
+		return true;
+	}
+	return false;
+}
+
 static const FString GSetDataLayerRuntimeStateCommandCommandName("wp.Runtime.SetDataLayerRuntimeState");
 FAutoConsoleCommand UDataLayerManager::SetDataLayerRuntimeStateCommand(
 	*GSetDataLayerRuntimeStateCommandCommandName,
diff --git a/Engine/Source/Runtime/Engine/Private/WorldPartition/DataLayer/DataLayerSubsystem.cpp b/Engine/Source/Runtime/Engine/Private/WorldPartition/DataLayer/DataLayerSubsystem.cpp
--- a/Engine/Source/Runtime/Engine/Private/WorldPartition/DataLayer/DataLayerSubsystem.cpp
+++ b/Engine/Source/Runtime/Engine/Private/WorldPartition/DataLayer/DataLayerSubsystem.cpp
@@ -2,7 +2,6 @@
 
 #include "WorldPartition/DataLayer/DataLayerSubsystem.h"
 #include "WorldPartition/ActorDescContainer.h"
-#include "WorldPartition/DataLayer/DataLayer.h"
 #include "WorldPartition/DataLayer/WorldDataLayers.h"
 #include "WorldPartition/DataLayer/DataLayerEditorContext.h"
 #include "WorldPartition/DataLayer/WorldDataLayersActorDesc.h"
diff --git a/Engine/Source/Runtime/Engine/Private/WorldPartition/DataLayer/DeprecatedDataLayerInstance.cpp b/Engine/Source/Runtime/Engine/Private/WorldPartition/DataLayer/DeprecatedDataLayerInstance.cpp
--- a/Engine/Source/Runtime/Engine/Private/WorldPartition/DataLayer/DeprecatedDataLayerInstance.cpp
+++ b/Engine/Source/Runtime/Engine/Private/WorldPartition/DataLayer/DeprecatedDataLayerInstance.cpp
@@ -3,7 +3,6 @@
 #include "WorldPartition/DataLayer/DeprecatedDataLayerInstance.h"
 #include "Misc/StringFormatArg.h"
 #include "WorldPartition/DataLayer/DataLayerUtils.h"
-#include "WorldPartition/DataLayer/DataLayer.h"
 #include "WorldPartition/DataLayer/WorldDataLayers.h"
 
 #include UE_INLINE_GENERATED_CPP_BY_NAME(DeprecatedDataLayerInstance)
@@ -26,8 +25,6 @@ FName UDeprecatedDataLayerInstance::MakeName()
 
 void UDeprecatedDataLayerInstance::OnCreated()
 {
-	check(GetOuterAWorldDataLayers()->HasDeprecatedDataLayers());
-
 	Modify(/*bAlwaysMarkDirty*/false);
 
 	FDataLayerUtils::SetDataLayerShortName(this, TEXT("DataLayer"));
@@ -47,31 +44,6 @@ void UDeprecatedDataLayerInstance::PostLoad()
 	Super::PostLoad();
 }
 
-FName UDeprecatedDataLayerInstance::MakeName(const UDEPRECATED_DataLayer* DeprecatedDataLayer)
-{
-	return FName(FString::Format(TEXT("{0}_{1}"), { *DeprecatedDataLayer->GetDataLayerLabel().ToString(), FGuid::NewGuid().ToString() }));
-}
-
-void UDeprecatedDataLayerInstance::OnCreated(const UDEPRECATED_DataLayer* DeprecatedDataLayer)
-{
-	check(GetOuterAWorldDataLayers()->HasDeprecatedDataLayers());
-
-	Modify(/*bAlwaysMarkDirty*/false);
-
-	Label = DeprecatedDataLayer->GetDataLayerLabel();
-
-	DeprecatedDataLayerFName = DeprecatedDataLayer->GetFName();
-
-	DataLayerType = DeprecatedDataLayer->IsRuntime() ? EDataLayerType::Runtime : EDataLayerType::Editor;
-	DebugColor = DeprecatedDataLayer->GetDebugColor();
-
-	bIsVisible = DeprecatedDataLayer->bIsVisible;
-	bIsInitiallyVisible = DeprecatedDataLayer->bIsInitiallyVisible;
-	bIsInitiallyLoadedInEditor = DeprecatedDataLayer->bIsInitiallyLoadedInEditor;
-	bIsLocked = DeprecatedDataLayer->bIsLocked;
-	InitialRuntimeState = DeprecatedDataLayer->InitialRuntimeState;
-}
-
 FActorDataLayer UDeprecatedDataLayerInstance::GetActorDataLayer() const
 {
 	return FActorDataLayer(GetDataLayerFName());
diff --git a/Engine/Source/Runtime/Engine/Private/WorldPartition/DataLayer/WorldDataLayers.cpp b/Engine/Source/Runtime/Engine/Private/WorldPartition/DataLayer/WorldDataLayers.cpp
--- a/Engine/Source/Runtime/Engine/Private/WorldPartition/DataLayer/WorldDataLayers.cpp
+++ b/Engine/Source/Runtime/Engine/Private/WorldPartition/DataLayer/WorldDataLayers.cpp
@@ -6,7 +6,6 @@
 
 #include "WorldPartition/DataLayer/WorldDataLayers.h"
 #include "WorldPartition/DataLayer/DataLayerManager.h"
-#include "WorldPartition/DataLayer/DataLayer.h"
 #include "WorldPartition/DataLayer/DataLayerInstanceWithAsset.h"
 #include "WorldPartition/DataLayer/DeprecatedDataLayerInstance.h"
 #include "WorldPartition/DataLayer/DataLayerUtils.h"
@@ -737,7 +736,6 @@ int32 AWorldDataLayers::RemoveDataLayers(const TArray<UDataLayerInstance*>& InDa
 
 	if (RemovedCount > 0)
 	{
-		UpdateContainsDeprecatedDataLayers();
 		if (bInResolveActorDescContainers)
 		{
 			ResolveActorDescContainers();
@@ -1189,8 +1187,6 @@ void AWorldDataLayers::OnDataLayerManagerInitialized()
 		RemoveDataLayers(ToDelete);
 	}
 
-	ConvertDataLayerToInstances();
-
 	// Remove all Editor Data Layers when cooking or when in a game world
 	if (IsRunningCookCommandlet() || GetWorld()->IsGameWorld())
 	{
@@ -1416,7 +1412,7 @@ bool AWorldDataLayers::IsReadOnly(FText* OutReason) const
 bool AWorldDataLayers::SupportsExternalPackageDataLayerInstances() const
 {
 	ULevel* Level = GetLevel();
-	return Level && Level->bIsPartitioned && Level->IsUsingExternalObjects() && !HasDeprecatedDataLayers();
+	return Level && Level->bIsPartitioned && Level->IsUsingExternalObjects();
 }
 
 bool AWorldDataLayers::SetUseExternalPackageDataLayerInstances(bool bInNewValue, bool bInInteractiveMode)
@@ -1431,12 +1427,6 @@ bool AWorldDataLayers::SetUseExternalPackageDataLayerInstances(bool bInNewValue,
 		return false;
 	}
 
-	if (HasDeprecatedDataLayers())
-	{
-		UE_LOG(LogWorldPartition, Warning, TEXT("Changing external packaging of data layer instances is not supported with deprecated data layers. Convert your data to use Data Layer Assets."));
-		return false;
-	}
-
 	check(SupportsExternalPackageDataLayerInstances());
 	const bool bInteractiveMode = bInInteractiveMode && !IsRunningCommandlet();
 
@@ -1486,58 +1476,6 @@ bool AWorldDataLayers::SetUseExternalPackageDataLayerInstances(bool bInNewValue,
 	return true;
 }
 
-void AWorldDataLayers::ConvertDataLayerToInstances()
-{
-	static_assert(DATALAYER_TO_INSTANCE_RUNTIME_CONVERSION_ENABLED, "AWorldDataLayers::ConvertDataLayerToInstances function is deprecated and needs to be deleted.");
-	
-	PRAGMA_DISABLE_DEPRECATION_WARNINGS
-	bHasDeprecatedDataLayers = !WorldDataLayers_DEPRECATED.IsEmpty();
-
-	for (UDEPRECATED_DataLayer* DeprecatedDataLayer : WorldDataLayers_DEPRECATED)
-	{
-		UDeprecatedDataLayerInstance* DataLayerInstance = CreateDataLayer<UDeprecatedDataLayerInstance>(DeprecatedDataLayer);
-		DeprecatedDataLayerNameToDataLayerInstance.Add(DeprecatedDataLayer->GetFName(), DataLayerInstance);
-	}
-
-	for (UDEPRECATED_DataLayer* DeprecatedDataLayer : WorldDataLayers_DEPRECATED)
-	{
-		if (DeprecatedDataLayer->GetParent() != nullptr)
-		{
-			UDataLayerInstance* ParentInstance = const_cast<UDataLayerInstance*>(GetDataLayerInstance(DeprecatedDataLayer->GetParent()->GetFName()));
-			UDataLayerInstance* ChildInstance = const_cast<UDataLayerInstance*>(GetDataLayerInstance(DeprecatedDataLayer->GetFName()));
-			if (!ChildInstance->SetParent(ParentInstance))
-			{
-				UE_LOG(LogWorldPartition, Error, TEXT("Failed to Convert DataLayer %s' hieararchy to DataLayerInstances. Run DataLayerToAsset Commandlet or fix the hierarchy manually."), *DeprecatedDataLayer->GetDataLayerLabel().ToString());
-			}
-		}
-	}
-
-	WorldDataLayers_DEPRECATED.Empty();
-
-	UpdateContainsDeprecatedDataLayers();
-
-	PRAGMA_ENABLE_DEPRECATION_WARNINGS
-}
-
-void AWorldDataLayers::UpdateContainsDeprecatedDataLayers()
-{
-	static_assert(DATALAYER_TO_INSTANCE_RUNTIME_CONVERSION_ENABLED, "AWorldDataLayers::UpdateContainsDeprecatedDataLayers function is deprecated and needs to be deleted.");
-
-	bHasDeprecatedDataLayers = !WorldDataLayers_DEPRECATED.IsEmpty();
-
-	if (!bHasDeprecatedDataLayers)
-	{
-		for (const UDataLayerInstance* DataLayerInstance : DataLayerInstances)
-		{
-			if (DataLayerInstance->IsA<UDeprecatedDataLayerInstance>())
-			{
-				bHasDeprecatedDataLayers = true;
-				break;
-			}
-		}
-	}
-}
-
 void AWorldDataLayers::ResolveActorDescContainers()
 {
 	// Always use actor owning world to find DataLayerManager for resolving of data layers as partitioned 
@@ -1594,29 +1532,6 @@ bool AWorldDataLayers::IsEditorOnly() const
 
 PRAGMA_DISABLE_DEPRECATION_WARNINGS
 
-#if WITH_EDITOR
-
-bool AWorldDataLayers::RemoveDataLayer(const UDEPRECATED_DataLayer* InDataLayer)
-{
-	if (ContainsDataLayer(InDataLayer))
-	{
-		Modify();
-		WorldDataLayers_DEPRECATED.Remove(InDataLayer);
-
-		UpdateContainsDeprecatedDataLayers();
-
-		return true;
-	}
-	return false;
-}
-
-#endif
-
-bool AWorldDataLayers::ContainsDataLayer(const UDEPRECATED_DataLayer* InDataLayer) const
-{
-	return WorldDataLayers_DEPRECATED.Contains(InDataLayer);
-}
-
 const UDataLayerInstance* AWorldDataLayers::GetDataLayerFromLabel(const FName& InDataLayerLabel) const
 {
 	const FString DataLayerLabelSanitized = *FDataLayerUtils::GetSanitizedDataLayerShortName(InDataLayerLabel.ToString());
diff --git a/Engine/Source/Runtime/Engine/Public/WorldPartition/DataLayer/DataLayer.h b/Engine/Source/Runtime/Engine/Public/WorldPartition/DataLayer/DataLayer.h
--- a/Engine/Source/Runtime/Engine/Public/WorldPartition/DataLayer/DataLayer.h
+++ b/Engine/Source/Runtime/Engine/Public/WorldPartition/DataLayer/DataLayer.h
@@ -2,186 +2,8 @@
 
 #pragma once
 
-#include "CoreMinimal.h"
-#include "UObject/ObjectMacros.h"
-#include "UObject/UObjectGlobals.h"
-#include "UObject/Object.h"
-#include "ActorDataLayer.h"
-#include "Math/Color.h"
-#include "WorldPartition/DataLayer/DataLayerInstance.h"
+#include "Misc/CoreMiscDefines.h"
 
-#include "DataLayer.generated.h"
+UE_DEPRECATED_HEADER(5.8, "The contents of this header have been fully deprecated. See DeprecatedDataLayer.h for details on particular cases.") 
 
-UENUM(BlueprintType, meta = (ScriptName = "DataLayerStateType"))
-enum class UE_DEPRECATED(5.0, "Use EDataLayerRuntimeState instead.") EDataLayerState : uint8
-{
-	Unloaded,
-	Loaded,
-	Activated
-};
-
-// Used for debugging
-bool inline GetDataLayerRuntimeStateFromName(const FString& InStateName, EDataLayerRuntimeState& OutState)
-{
-	if (InStateName.Equals(GetDataLayerRuntimeStateName(EDataLayerRuntimeState::Unloaded), ESearchCase::IgnoreCase))
-	{
-		OutState = EDataLayerRuntimeState::Unloaded;
-		return true;
-	}
-	else if (InStateName.Equals(GetDataLayerRuntimeStateName(EDataLayerRuntimeState::Loaded), ESearchCase::IgnoreCase))
-	{
-		OutState = EDataLayerRuntimeState::Loaded;
-		return true;
-	}
-	else if (InStateName.Equals(GetDataLayerRuntimeStateName(EDataLayerRuntimeState::Activated), ESearchCase::IgnoreCase))
-	{
-		OutState = EDataLayerRuntimeState::Activated;
-		return true;
-	}
-	return false;
-}
-
-static_assert(EDataLayerRuntimeState::Unloaded < EDataLayerRuntimeState::Loaded && EDataLayerRuntimeState::Loaded < EDataLayerRuntimeState::Activated, "Streaming Query code is dependent on this being true");
-
-
-static_assert(DATALAYER_TO_INSTANCE_RUNTIME_CONVERSION_ENABLED, "UDEPRECATED_DataLayer class is deprecated and needs to be deleted.");
-class UE_DEPRECATED(5.1, "Use UDataLayerInstance & UDataLayerAsset to create DataLayers") UDEPRECATED_DataLayer;
-
-UCLASS(Config = Engine, PerObjectConfig, Within = WorldDataLayers, BlueprintType, Deprecated, MinimalAPI)
-class UDEPRECATED_DataLayer : public UObject
-{
-	GENERATED_UCLASS_BODY()
-
-	friend class UDeprecatedDataLayerInstance;
-
-public:
-#if WITH_EDITOR
-	ENGINE_API virtual bool CanEditChange(const FProperty* Property) const;
-	ENGINE_API virtual void PostEditChangeProperty(FPropertyChangedEvent& PropertyChangedEvent) override;
-
-	ENGINE_API void SetVisible(bool bIsVisible);
-	ENGINE_API void SetIsInitiallyVisible(bool bIsInitiallyVisible);
-	ENGINE_API void SetIsRuntime(bool bIsRuntime);
-	ENGINE_API void SetIsLoadedInEditor(bool bIsLoadedInEditor, bool bFromUserChange);
-	void SetIsLocked(bool bInIsLocked) { bIsLocked = bInIsLocked; }
-
-	ENGINE_API bool IsLocked() const;
-	bool IsInitiallyLoadedInEditor() const { return bIsInitiallyLoadedInEditor; }
-	bool IsLoadedInEditor() const { return bIsLoadedInEditor; }
-	ENGINE_API bool IsEffectiveLoadedInEditor() const;
-	bool IsLoadedInEditorChangedByUserOperation() const { return bIsLoadedInEditorChangedByUserOperation; }
-	void ClearLoadedInEditorChangedByUserOperation() { bIsLoadedInEditorChangedByUserOperation = false; }
-
-	ENGINE_API bool CanParent(const UDEPRECATED_DataLayer* InParent) const;
-	ENGINE_API void SetParent(UDEPRECATED_DataLayer* InParent);
-	ENGINE_API void SetChildParent(UDEPRECATED_DataLayer* InParent);
-
-	static ENGINE_API FText GetDataLayerText(const UDEPRECATED_DataLayer* InDataLayer);
-	ENGINE_API const TCHAR* GetDataLayerIconName() const;
-#endif
-	const TArray<TObjectPtr<UDEPRECATED_DataLayer>>& GetChildren() const { return Children_DEPRECATED; }
-	ENGINE_API void ForEachChild(TFunctionRef<bool(const UDEPRECATED_DataLayer*)> Operation) const;
-
-	const UDEPRECATED_DataLayer* GetParent() const { return Parent_DEPRECATED; }
-	UDEPRECATED_DataLayer* GetParent() { return Parent_DEPRECATED; }
-
-	ENGINE_API virtual void PostLoad() override;
-
-	UFUNCTION(Category = "Data Layer", BlueprintCallable)
-	bool Equals(const FActorDataLayer& ActorDataLayer) const { return ActorDataLayer.Name == GetFName(); }
-
-	UFUNCTION(Category = "Data Layer", BlueprintCallable)
-	FName GetDataLayerLabel() const { return DataLayerLabel; }
-
-	UFUNCTION(Category = "Data Layer|Editor", BlueprintCallable)
-	ENGINE_API bool IsInitiallyVisible() const;
-
-	UFUNCTION(Category = "Data Layer|Editor", BlueprintCallable)
-	ENGINE_API bool IsVisible() const;
-
-	UFUNCTION(Category = "Data Layer|Editor", BlueprintCallable)
-	ENGINE_API bool IsEffectiveVisible() const;
-
-	UFUNCTION(Category = "Data Layer|Editor", BlueprintCallable)
-	FColor GetDebugColor() const { return DebugColor; }
-
-	UFUNCTION(Category = "Data Layer|Runtime", BlueprintCallable)
-	bool IsRuntime() const { return bIsRuntime; }
-	
-	UFUNCTION(Category = "Data Layer|Runtime", BlueprintCallable)
-	EDataLayerRuntimeState GetInitialRuntimeState() const { return IsRuntime() ? InitialRuntimeState : EDataLayerRuntimeState::Unloaded; }
-
-	//~ Begin Deprecated
-
-	UE_DEPRECATED(5.0, "Use IsRuntime() instead.")
-	UFUNCTION(Category = "Data Layer|Runtime", BlueprintCallable, meta = (DeprecatedFunction, DeprecationMessage = "Use IsRuntime instead"))
-	bool IsDynamicallyLoaded() const { return IsRuntime(); }
-
-	UE_DEPRECATED(5.0, "Use GetInitialRuntimeState() instead.")
-	UFUNCTION(Category = "Data Layer|Runtime", BlueprintCallable, meta = (DeprecatedFunction, DeprecationMessage = "Use GetInitialRuntimeState instead"))
-	bool IsInitiallyActive() const { return IsRuntime() && GetInitialRuntimeState() == EDataLayerRuntimeState::Activated; }
-
-	PRAGMA_DISABLE_DEPRECATION_WARNINGS // Suppress compiler warning on override of deprecated function
-	UE_DEPRECATED(5.0, "Use GetInitialRuntimeState() instead.")
-	UFUNCTION(Category = "Data Layer|Runtime", BlueprintCallable, meta = (DeprecatedFunction, DeprecationMessage = "Use GetInitialRuntimeState instead"))
-	EDataLayerState GetInitialState() const { return (EDataLayerState)GetInitialRuntimeState(); }
-	PRAGMA_ENABLE_DEPRECATION_WARNINGS
-
-	//~ End Deprecated
-private:
-
-	void AddChild(UDEPRECATED_DataLayer* DataLayer);
-#if WITH_EDITOR
-	ENGINE_API void RemoveChild(UDEPRECATED_DataLayer* DataLayer);
-	ENGINE_API void PropagateIsRuntime();
-#endif
-
-#if WITH_EDITORONLY_DATA
-	UPROPERTY()
-	uint32 bIsInitiallyActive_DEPRECATED : 1;
-
-	/** Whether actors associated with the DataLayer are visible in the viewport */
-	UPROPERTY(Transient)
-	uint32 bIsVisible : 1;
-
-	/** Whether actors associated with the Data Layer should be initially visible in the viewport when loading the map */
-	UPROPERTY(Category = "Data Layer|Editor", EditAnywhere)
-	uint32 bIsInitiallyVisible : 1;
-
-	/** Determines the default value of the data layer's loaded state in editor if it hasn't been changed in data layer outliner by the user */
-	UPROPERTY(Category = "Data Layer|Editor", EditAnywhere, meta = (DisplayName = "Is Initially Loaded"))
-	uint32 bIsInitiallyLoadedInEditor : 1;
-
-	/** Wheter the data layer is loaded in editor (user setting) */
-	UPROPERTY(Transient)
-	uint32 bIsLoadedInEditor : 1;
-
-	/** Whether this data layer editor visibility was changed by a user operation */
-	UPROPERTY(Transient)
-	uint32 bIsLoadedInEditorChangedByUserOperation : 1;
-
-	/** Whether this data layer is locked, which means the user can't change actors assignation, remove or rename it */
-	UPROPERTY()
-	uint32 bIsLocked : 1;
-#endif
-
-	/** The display name of the Data Layer */
-	UPROPERTY()
-	FName DataLayerLabel;
-
-	/** Whether the Data Layer affects actor runtime loading */
-	UPROPERTY(Category = "Data Layer|Advanced", EditAnywhere)
-	uint32 bIsRuntime : 1;
-
-	UPROPERTY(Category = "Data Layer|Advanced|Runtime", EditAnywhere, meta = (EditConditionHides, EditCondition = "bIsRuntime"))
-	EDataLayerRuntimeState InitialRuntimeState;
-
-	UPROPERTY(Category = "Data Layer|Editor", EditAnywhere)
-	FColor DebugColor;
-
-	UPROPERTY()
-	TObjectPtr<UDEPRECATED_DataLayer> Parent_DEPRECATED;
-
-	UPROPERTY(Transient)
-	TArray<TObjectPtr<UDEPRECATED_DataLayer>> Children_DEPRECATED;
-};
+#include "WorldPartition/DataLayer/Deprecated/DeprecatedDataLayer.h"
diff --git a/Engine/Source/Runtime/Engine/Public/WorldPartition/DataLayer/Deprecated/DeprecatedDataLayer.h b/Engine/Source/Runtime/Engine/Public/WorldPartition/DataLayer/Deprecated/DeprecatedDataLayer.h
new file mode 100644
--- /dev/null
+++ b/Engine/Source/Runtime/Engine/Public/WorldPartition/DataLayer/Deprecated/DeprecatedDataLayer.h
@@ -0,0 +1,112 @@
+// Copyright Epic Games, Inc. All Rights Reserved.
+
+#pragma once
+
+#include "CoreMinimal.h"
+#include "UObject/ObjectMacros.h"
+#include "UObject/UObjectGlobals.h"
+#include "UObject/Object.h"
+#include "Math/Color.h"
+#include "WorldPartition/DataLayer/DataLayerInstance.h"
+
+#include "DeprecatedDataLayer.generated.h"
+
+// These deprecated reflected types and members must not be removed for the duration of UE5, even after 2 versions post-deprecation, to allow loading and upgrading old assets.
+
+UENUM(BlueprintType)
+enum class UE_DEPRECATED(all, "Use EDataLayerRuntimeState instead.") EDataLayerState : uint8
+{
+	Unloaded,
+	Loaded,
+	Activated
+};
+
+static_assert(EDataLayerRuntimeState::Unloaded < EDataLayerRuntimeState::Loaded && EDataLayerRuntimeState::Loaded < EDataLayerRuntimeState::Activated, "Streaming Query code is dependent on this being true");
+
+static_assert(DATALAYER_TO_INSTANCE_RUNTIME_CONVERSION_ENABLED, "UDEPRECATED_DataLayer class is deprecated and needs to be deleted.");
+class UE_DEPRECATED(all, "Use UDataLayerInstance & UDataLayerAsset to create DataLayers") UDEPRECATED_DataLayer;
+
+UCLASS(Config = Engine, PerObjectConfig, Within = WorldDataLayers, BlueprintType, Deprecated, MinimalAPI)
+class UDEPRECATED_DataLayer : public UObject
+{
+	GENERATED_BODY()
+
+	friend class UDeprecatedDataLayerInstance;
+
+	UFUNCTION(BlueprintCallable, meta = (DeprecatedFunction, DeprecationMessage = "Use UDataLayerInstance & UDataLayerAsset to create DataLayers"))
+	bool Equals(const FActorDataLayer& ActorDataLayer) const { return false; }
+
+	UFUNCTION(BlueprintCallable, meta = (DeprecatedFunction, DeprecationMessage = "Use UDataLayerInstance & UDataLayerAsset to create DataLayers"))
+	FName GetDataLayerLabel() const { return FName{}; }
+
+	UFUNCTION(BlueprintCallable, meta = (DeprecatedFunction, DeprecationMessage = "Use UDataLayerInstance & UDataLayerAsset to create DataLayers"))
+	bool IsInitiallyVisible() const { return false; }
+
+	UFUNCTION(BlueprintCallable, meta = (DeprecatedFunction, DeprecationMessage = "Use UDataLayerInstance & UDataLayerAsset to create DataLayers"))
+	bool IsVisible() const { return false; }
+
+	UFUNCTION(BlueprintCallable, meta = (DeprecatedFunction, DeprecationMessage = "Use UDataLayerInstance & UDataLayerAsset to create DataLayers"))
+	bool IsEffectiveVisible() const { return false; }
+
+	UFUNCTION(BlueprintCallable, meta = (DeprecatedFunction, DeprecationMessage = "Use UDataLayerInstance & UDataLayerAsset to create DataLayers"))
+	FColor GetDebugColor() const { return FColor{}; }
+
+	UFUNCTION(BlueprintCallable, meta = (DeprecatedFunction, DeprecationMessage = "Use UDataLayerInstance & UDataLayerAsset to create DataLayers"))
+	bool IsRuntime() const { return false; }
+	
+	UFUNCTION(BlueprintCallable, meta = (DeprecatedFunction, DeprecationMessage = "Use UDataLayerInstance & UDataLayerAsset to create DataLayers"))
+	EDataLayerRuntimeState GetInitialRuntimeState() const { return EDataLayerRuntimeState::Unloaded; }
+
+	UFUNCTION(BlueprintCallable, meta = (DeprecatedFunction, DeprecationMessage = "Use UDataLayerInstance & UDataLayerAsset to create DataLayers"))
+	bool IsDynamicallyLoaded() const { return false; }
+
+	UFUNCTION(BlueprintCallable, meta = (DeprecatedFunction, DeprecationMessage = "Use UDataLayerInstance & UDataLayerAsset to create DataLayers"))
+	bool IsInitiallyActive() const { return IsRuntime() && GetInitialRuntimeState() == EDataLayerRuntimeState::Activated; }
+
+	PRAGMA_DISABLE_DEPRECATION_WARNINGS
+	UFUNCTION(BlueprintCallable, meta = (DeprecatedFunction, DeprecationMessage = "Use UDataLayerInstance & UDataLayerAsset to create DataLayers"))
+	EDataLayerState GetInitialState() const { return EDataLayerState::Unloaded; }
+	PRAGMA_ENABLE_DEPRECATION_WARNINGS
+
+private:
+#if WITH_EDITORONLY_DATA
+	UPROPERTY(meta = (DeprecatedProperty, DeprecationMessage = "Use UDataLayerInstance & UDataLayerAsset to create DataLayers"))
+	uint32 bIsInitiallyActive_DEPRECATED : 1;
+
+	UPROPERTY(Transient, meta = (DeprecatedProperty, DeprecationMessage = "Use UDataLayerInstance & UDataLayerAsset to create DataLayers"))
+	uint32 bIsVisible : 1;
+
+	UPROPERTY(meta = (DeprecatedProperty, DeprecationMessage = "Use UDataLayerInstance & UDataLayerAsset to create DataLayers"))
+	uint32 bIsInitiallyVisible : 1;
+
+	UPROPERTY(meta = (DeprecatedProperty, DeprecationMessage = "Use UDataLayerInstance & UDataLayerAsset to create DataLayers"))
+	uint32 bIsInitiallyLoadedInEditor : 1;
+
+	UPROPERTY(Transient, meta = (DeprecatedProperty, DeprecationMessage = "Use UDataLayerInstance & UDataLayerAsset to create DataLayers"))
+	uint32 bIsLoadedInEditor : 1;
+
+	UPROPERTY(Transient, meta = (DeprecatedProperty, DeprecationMessage = "Use UDataLayerInstance & UDataLayerAsset to create DataLayers"))
+	uint32 bIsLoadedInEditorChangedByUserOperation : 1;
+
+	UPROPERTY(meta = (DeprecatedProperty, DeprecationMessage = "Use UDataLayerInstance & UDataLayerAsset to create DataLayers"))
+	uint32 bIsLocked : 1;
+#endif
+
+	UPROPERTY(meta = (DeprecatedProperty, DeprecationMessage = "Use UDataLayerInstance & UDataLayerAsset to create DataLayers"))
+	FName DataLayerLabel;
+
+	UPROPERTY(meta = (DeprecatedProperty, DeprecationMessage = "Use UDataLayerInstance & UDataLayerAsset to create DataLayers"))
+	uint32 bIsRuntime : 1;
+
+	UPROPERTY(meta = (DeprecatedProperty, DeprecationMessage = "Use UDataLayerInstance & UDataLayerAsset to create DataLayers"))
+	EDataLayerRuntimeState InitialRuntimeState;
+
+	UPROPERTY(meta = (DeprecatedProperty, DeprecationMessage = "Use UDataLayerInstance & UDataLayerAsset to create DataLayers"))
+	FColor DebugColor;
+
+	UPROPERTY(meta = (DeprecatedProperty, DeprecationMessage = "Use UDataLayerInstance & UDataLayerAsset to create DataLayers"))
+	TObjectPtr<UDEPRECATED_DataLayer> Parent_DEPRECATED;
+
+	UPROPERTY(meta = (DeprecatedProperty, DeprecationMessage = "Use UDataLayerInstance & UDataLayerAsset to create DataLayers"))
+	TArray<TObjectPtr<UDEPRECATED_DataLayer>> Children_DEPRECATED;
+};
\ No newline at end of file
diff --git a/Engine/Source/Runtime/Engine/Public/WorldPartition/DataLayer/DeprecatedDataLayerInstance.h b/Engine/Source/Runtime/Engine/Public/WorldPartition/DataLayer/DeprecatedDataLayerInstance.h
--- a/Engine/Source/Runtime/Engine/Public/WorldPartition/DataLayer/DeprecatedDataLayerInstance.h
+++ b/Engine/Source/Runtime/Engine/Public/WorldPartition/DataLayer/DeprecatedDataLayerInstance.h
@@ -5,8 +5,6 @@
 #include "WorldPartition/DataLayer/ActorDataLayer.h"
 #include "DeprecatedDataLayerInstance.generated.h"
 
-class UDEPRECATED_DataLayer;
-
 // Class used for Runtime Conversion of the Deprecated UDataLayer Class to UDataLayerInstance + UDataLayerAsset.
 // This class is not to be inherited. It is solely used by AWorldDatalayers to convert UDataLayers to UDataLayerInstances on Level Boot.
 // You will need to run the DataLayerToAsset CommandLet to convert the deprecated datalayers to UDataLayerAssets and UDataLayerInstanceWithAsset.
@@ -30,8 +28,11 @@ public:
 	FActorDataLayer GetActorDataLayer() const;
 
 	PRAGMA_DISABLE_DEPRECATION_WARNINGS
-	ENGINE_API static FName MakeName(const UDEPRECATED_DataLayer* DeprecatedDataLayer);
-	ENGINE_API void OnCreated(const UDEPRECATED_DataLayer* DeprecatedDataLayer);
+	UE_DEPRECATED(5.8, "UDEPRECATED_DataLayer are no longer supported")
+	static FName MakeName(const class UDEPRECATED_DataLayer* DeprecatedDataLayer) { return FName{}; }
+
+	UE_DEPRECATED(5.8, "UDEPRECATED_DataLayer are no longer supported")
+	void OnCreated(const class UDEPRECATED_DataLayer* DeprecatedDataLayer) {}
 	PRAGMA_ENABLE_DEPRECATION_WARNINGS
 #endif
 	
diff --git a/Engine/Source/Runtime/Engine/Public/WorldPartition/DataLayer/WorldDataLayers.h b/Engine/Source/Runtime/Engine/Public/WorldPartition/DataLayer/WorldDataLayers.h
--- a/Engine/Source/Runtime/Engine/Public/WorldPartition/DataLayer/WorldDataLayers.h
+++ b/Engine/Source/Runtime/Engine/Public/WorldPartition/DataLayer/WorldDataLayers.h
@@ -7,7 +7,6 @@
 #include "GameFramework/Info.h"
 #include "WorldPartition/DataLayer/ActorDataLayer.h"
 #include "WorldPartition/DataLayer/DataLayerInstance.h"
-#include "WorldPartition/DataLayer/DataLayer.h"
 #include "WorldPartition/DataLayer/DataLayerInstanceProviderInterface.h"
 #include "WorldPartition/DataLayer/ExternalDataLayerAsset.h"
 #include "Net/Core/PushModel/PushModel.h"
@@ -16,7 +15,6 @@
 #include "ExternalPackageHelper.h"
 #include "WorldDataLayers.generated.h"
 
-class UDEPRECATED_DataLayer;
 class UDataLayerInstance;
 class UDataLayerInstanceWithAsset;
 class UDataLayerAsset;
@@ -110,7 +108,6 @@ public:
 	static ENGINE_API AWorldDataLayers* Create(const FActorSpawnParameters& SpawnParameters);
 
 	ENGINE_API bool IsEmpty() const;
-	bool HasDeprecatedDataLayers() const { return bHasDeprecatedDataLayers; }
 
 	template<class DataLayerInstanceType, typename ...CreationsArgs>
 	DataLayerInstanceType* CreateDataLayer(CreationsArgs... InCreationArgs);
@@ -186,30 +183,6 @@ public:
 
 	//~ Begin Deprecated
 
-	PRAGMA_DISABLE_DEPRECATION_WARNINGS // Suppress compiler warning on override of deprecated function
-
-	UE_DEPRECATED(5.0, "Use SetDataLayerRuntimeState() instead.")
-	void SetDataLayerState(FActorDataLayer InDataLayer, EDataLayerState InState) { SetDataLayerRuntimeState(GetDataLayerInstance(InDataLayer.Name), (EDataLayerRuntimeState)InState); }
-
-	UE_DEPRECATED(5.0, "Use GetDataLayerRuntimeStateByName() instead.")
-	EDataLayerState GetDataLayerStateByName(FName InDataLayerName) const { return (EDataLayerState)GetDataLayerRuntimeStateByName(InDataLayerName); }
-
-	UE_DEPRECATED(5.0, "Use GetEffectiveActiveDataLayerNames() instead.")
-	const TSet<FName>& GetActiveDataLayerNames() const { return GetEffectiveActiveDataLayerNames(); }
-
-	UE_DEPRECATED(5.0, "Use GetEffectiveLoadedDataLayerNames() instead.")
-	const TSet<FName>& GetLoadedDataLayerNames() const { return GetEffectiveLoadedDataLayerNames(); }
-
-#if WITH_EDITOR
-	UE_DEPRECATED(5.1, "Convert DataLayer using UDataLayerToAssetCommandlet and use UDataLayerInstance* overload instead.")
-	ENGINE_API bool RemoveDataLayer(const UDEPRECATED_DataLayer* InDataLayer);
-#endif
-
-	UE_DEPRECATED(5.1, "Convert DataLayer using UDataLayerToAssetCommandlet and use UDataLayerInstance* overload instead.")
-	ENGINE_API bool ContainsDataLayer(const UDEPRECATED_DataLayer* InDataLayer) const;
-
-	PRAGMA_ENABLE_DEPRECATION_WARNINGS
-
 	UE_DEPRECATED(5.1, "Use GetDataLayerInstance with FName or UDataLayerAsset instead")
 	ENGINE_API const UDataLayerInstance* GetDataLayerInstance(const FActorDataLayer& InActorDataLayer) const;
 
@@ -284,8 +257,6 @@ private:
 	virtual bool ActorTypeIsMainWorldOnly() const override { return true; }
 
 	ENGINE_API void AddDataLayerInstance(UDataLayerInstance* InDataLayerInstance);
-	ENGINE_API void ConvertDataLayerToInstances();
-	ENGINE_API void UpdateContainsDeprecatedDataLayers();
 	ENGINE_API void ResolveActorDescContainers();
 	ENGINE_API void RemoveEditorDataLayers();
 	void UpdateCurrentColorizedDataLayerInstance();
@@ -338,8 +309,9 @@ private:
 	TMap<FName, TWeakObjectPtr<UDataLayerInstance>> DeprecatedDataLayerNameToDataLayerInstance;
 
 	PRAGMA_DISABLE_DEPRECATION_WARNINGS
+	UE_DEPRECATED(all, "Convert Data layers to DataLayerInstances and DataLayerAssets using DataLayerToAsset CommandLet and use DataLayerInstances instead.")
 	UPROPERTY(meta = (DeprecatedProperty, DeprecationMessage = "Convert Data layers to DataLayerInstances and DataLayerAssets using DataLayerToAsset CommandLet and use DataLayerInstances instead."))
-	TSet<TObjectPtr<UDEPRECATED_DataLayer>> WorldDataLayers_DEPRECATED;
+	TSet<TObjectPtr<class UDEPRECATED_DataLayer>> WorldDataLayers_DEPRECATED;
 	PRAGMA_ENABLE_DEPRECATION_WARNINGS
 
 	UPROPERTY(Transient, Replicated, ReplicatedUsing=OnRep_ActiveDataLayerNames)
@@ -365,9 +337,6 @@ private:
 
 	int32 DataLayersStateEpoch;
 
-	static_assert(DATALAYER_TO_INSTANCE_RUNTIME_CONVERSION_ENABLED, "bHasDeprecatedDataLayers is deprecated and needs to be deleted.");
-	bool bHasDeprecatedDataLayers;
-
 	friend class UWorldPartition;
 	friend class UDataLayerManager;
 	friend class UExternalDataLayerManager;
@@ -438,7 +407,6 @@ DataLayerInstanceType* AWorldDataLayers::CreateDataLayer(CreationsArgs... InCrea
 	NewDataLayer->OnCreated(Forward<CreationsArgs>(InCreationArgs)...);
 	AddDataLayerInstance(NewDataLayer);
 
-	UpdateContainsDeprecatedDataLayers();
 	ResolveActorDescContainers();
 	return NewDataLayer;
 }
