PR #14228: [Sky][Performance] Disable ticking on ASkyAtmosphere since it isn't used

#tests PC editor
#rnx
#Synced-CL 49680783

diff --git a/Engine/Source/Runtime/Engine/Private/Components/SkyAtmosphereComponent.cpp b/Engine/Source/Runtime/Engine/Private/Components/SkyAtmosphereComponent.cpp
--- a/Engine/Source/Runtime/Engine/Private/Components/SkyAtmosphereComponent.cpp
+++ b/Engine/Source/Runtime/Engine/Private/Components/SkyAtmosphereComponent.cpp
@@ -627,7 +627,7 @@ ASkyAtmosphere::ASkyAtmosphere(const FObjectInitializer& ObjectInitializer)
 	}
 #endif // WITH_EDITORONLY_DATA
 
-	PrimaryActorTick.bCanEverTick = true;
+	PrimaryActorTick.bCanEverTick = false;
 	SetHidden(false);
 }
 
