- Fixed UE::Net::Connection::ValidateLevelVisibility which could call DoesPackageExist on a memory package and trigger a warning.
- Fixed FUpdateLevelVisibilityLevelInfo::NetSerialize where a None LevelVisibility.CookPackageName deserialized as LevelVisibility.PackageName; it now correctly remains None (no functional impact, but misleading).

#Synced-CL 49593614

diff --git a/Engine/Source/Runtime/Engine/Private/GameFramework/UpdateLevelVisibilityLevelInfo.cpp b/Engine/Source/Runtime/Engine/Private/GameFramework/UpdateLevelVisibilityLevelInfo.cpp
--- a/Engine/Source/Runtime/Engine/Private/GameFramework/UpdateLevelVisibilityLevelInfo.cpp
+++ b/Engine/Source/Runtime/Engine/Private/GameFramework/UpdateLevelVisibilityLevelInfo.cpp
@@ -43,12 +43,14 @@ bool FUpdateLevelVisibilityLevelInfo::NetSerialize(FArchive& Ar, UPackageMap* Pa
 	bool bArePackageAndFileTheSame = !!((PlayerControllerCVars::LevelVisibilityDontSerializeFileName) || (FileName == PackageName) || (FileName == NAME_None));
 	bool bLocalIsVisible = !!bIsVisible;
 	bool bLocalTryMakeVisible = !!bTryMakeVisible;
-	bool bArePackageAndCookPackageTheSame = (CookPackageName == PackageName) || (CookPackageName.IsNone());
+	bool bArePackageAndCookPackageTheSame = (CookPackageName == PackageName);
+	bool bIsCookPackageNameNone = CookPackageName.IsNone();
 
 	Ar.SerializeBits(&bArePackageAndFileTheSame, 1);
 	Ar.SerializeBits(&bLocalIsVisible, 1);
 	Ar.SerializeBits(&bLocalTryMakeVisible, 1);
 	Ar.SerializeBits(&bArePackageAndCookPackageTheSame, 1);
+	Ar.SerializeBits(&bIsCookPackageNameNone, 1);
 	Ar << PackageName;
 
 	if (!bArePackageAndFileTheSame)
@@ -60,13 +62,20 @@ bool FUpdateLevelVisibilityLevelInfo::NetSerialize(FArchive& Ar, UPackageMap* Pa
 		FileName = PackageName;
 	}
 
-	if (!bArePackageAndCookPackageTheSame)
+	if (!bIsCookPackageNameNone)
 	{
-		Ar << CookPackageName;
+		if (!bArePackageAndCookPackageTheSame)
+		{
+			Ar << CookPackageName;
+		}
+		else if (Ar.IsLoading())
+		{
+			CookPackageName = PackageName;
+		}
 	}
 	else if (Ar.IsLoading())
 	{
-		CookPackageName = PackageName;
+		CookPackageName = NAME_None;
 	}
 
 	VisibilityRequestId.NetSerialize(Ar, PackageMap, bOutSuccess);
diff --git a/Engine/Source/Runtime/Engine/Private/NetConnection.cpp b/Engine/Source/Runtime/Engine/Private/NetConnection.cpp
--- a/Engine/Source/Runtime/Engine/Private/NetConnection.cpp
+++ b/Engine/Source/Runtime/Engine/Private/NetConnection.cpp
@@ -163,7 +163,7 @@ namespace UE::Net::Connection::Private
 
 		FValidateLevelVisibilityResult Result;
 
-		auto FindStreamingLevel = [World, &LevelVisibility, &Result](FName InPackageName)
+		auto FindStreamingLevel = [World, &Result](FName InPackageName)
 		{
 			FString PackageNameStr = InPackageName.ToString();
 			Result.Package = FindPackage(nullptr, *PackageNameStr);
@@ -182,16 +182,16 @@ namespace UE::Net::Connection::Private
 			if (!Result.bLevelExists)
 			{
 				Result.bLevelExists = FLevelUtils::IsValidStreamingLevel(World, *PackageNameStr);
-
-				if (!Result.bLevelExists)
-				{
-					Result.bLevelExists = (!FPackageName::IsMemoryPackage(PackageNameStr) && FPackageName::DoesPackageExist(LevelVisibility.FileName.ToString()));
-				}
 			}
 		};
 
 		FindStreamingLevel(LevelVisibility.PackageName);
 
+		if (!Result.bLevelExists)
+		{
+			Result.bLevelExists = (!FPackageName::IsMemoryPackage(LevelVisibility.PackageName.ToString()) && FPackageName::DoesPackageExist(LevelVisibility.FileName.ToString()));
+		}
+
 		// If the level is not found, try resolving it using the optional LevelVisibility.CookPackageName.
 		// This is useful when the in-memory PackageName differs from the cooked package name.
 		// A common example is when a PIE client sends PIE memory-generated packages to a cooked server.
@@ -200,7 +200,6 @@ namespace UE::Net::Connection::Private
 			FindStreamingLevel(LevelVisibility.CookPackageName);
 		}
 
-
 		return Result;
 	}
 
