Add missing anim stats delcaration for verbose anim profiling to compile

#Synced-CL 49559469

diff --git a/Engine/Source/Runtime/Engine/Public/Animation/AnimMTStats.h b/Engine/Source/Runtime/Engine/Public/Animation/AnimMTStats.h
--- a/Engine/Source/Runtime/Engine/Public/Animation/AnimMTStats.h
+++ b/Engine/Source/Runtime/Engine/Public/Animation/AnimMTStats.h
@@ -65,4 +65,6 @@ DO_ANIMSTAT_PROCESSING(ResetRoot);
 DO_ANIMSTAT_PROCESSING(SplineIK);
 DO_ANIMSTAT_PROCESSING(SpringBone);
 DO_ANIMSTAT_PROCESSING(RotationMultiplier);
+DO_ANIMSTAT_PROCESSING(RemoveCurve);
+DO_ANIMSTAT_PROCESSING(ModularVehicleController);
 #endif
