add new flags to actually fully remove the face remap array from the chaos trimesh data, as the existing flag only removes one copy of the mesh (and existing code / tests rely on it not being fully stripped in cases where the previous remap flag was false)

#Synced-CL 49560847

diff --git a/Engine/Source/Runtime/Engine/Classes/PhysicsEngine/BodySetup.h b/Engine/Source/Runtime/Engine/Classes/PhysicsEngine/BodySetup.h
--- a/Engine/Source/Runtime/Engine/Classes/PhysicsEngine/BodySetup.h
+++ b/Engine/Source/Runtime/Engine/Classes/PhysicsEngine/BodySetup.h
@@ -112,9 +112,12 @@ struct FCookBodySetupInfo
 	/** Whether to support UV from hit results */
 	bool bSupportUVFromHitResults;
 
-	/** Whether to support face remap, needed for physical material masks */
+	/** Whether to support face remap, needed for physical material masks. Note: Also set bSupportFaceRemapOnMeshBVH to false to fully remove the face remap data. */
 	bool bSupportFaceRemap;
 
+	/** Whether to support face remap on the mesh BVH specifically. To fully remove face remap data, both this and the above bSupportFaceRemap must be false. */
+	bool bSupportFaceRemapOnMeshBVH;
+
 	/** Whether to support vertex remap, needed for skeletal mesh per-poly collision and any other FBodyInstance::UpdateTriMeshVertices calls */
 	bool bSupportVertexRemap;
 
@@ -192,11 +195,19 @@ class UBodySetup : public UBodySetupCore
 
 	/** 
 	 * If true, the physics triangle mesh will store UVs and the face remap table. This is needed
-	 * to support physical material masks in scene queries. 
+	 * to support physical material masks in scene queries.
+	 * Note: Also set bSupportFaceRemapOnMeshBVH to false to fully remove the remap table.
 	 */
 	UPROPERTY()
 	uint8 bSupportUVsAndFaceRemap:1;
 
+	/**
+	 * If true, or the above bSupportUVsAndFaceRemap is true, the physics triangle mesh will store the face remap table.
+	 * If both are false, the physics triangle mesh will use less memory, and queries for the original face index will return INDEX_NONE
+	 */
+	UPROPERTY()
+	uint8 bSupportFaceRemapOnMeshBVH : 1;
+
 	/**
 	 * If true, the physics triangle mesh will store the vertex remap table. This is needed
 	 * to support skeletal mesh per-poly collision, and any other FBodyInstance::UpdateTriMeshVertices calls.
diff --git a/Engine/Source/Runtime/Engine/Private/PhysicsEngine/BodySetup.cpp b/Engine/Source/Runtime/Engine/Private/PhysicsEngine/BodySetup.cpp
--- a/Engine/Source/Runtime/Engine/Private/PhysicsEngine/BodySetup.cpp
+++ b/Engine/Source/Runtime/Engine/Private/PhysicsEngine/BodySetup.cpp
@@ -62,6 +62,7 @@ FCookBodySetupInfo::FCookBodySetupInfo() :
 	bCookTriMesh(false),
 	bSupportUVFromHitResults(false),
 	bSupportFaceRemap(false),
+	bSupportFaceRemapOnMeshBVH(true),
 	bSupportVertexRemap(true),
 	bTriMeshError(false)
 {
@@ -140,6 +141,7 @@ UBodySetup::UBodySetup(const FObjectInitializer& ObjectInitializer)
 	bGenerateMirroredCollision = true;
 	bGenerateNonMirroredCollision = true;
 	bSupportUVsAndFaceRemap = false;
+	bSupportFaceRemapOnMeshBVH = true;
 	bSupportVertexRemap = true;
 	DefaultInstance.SetObjectType(ECC_PhysicsBody);
 #if WITH_EDITORONLY_DATA
@@ -328,6 +330,7 @@ void UBodySetup::GetCookInfo(FCookBodySetupInfo& OutCookInfo, EPhysXMeshCookFlag
 			OutCookInfo.TriMeshCookFlags = CookFlags;
 
 			OutCookInfo.bSupportFaceRemap = bSupportUVsAndFaceRemap;
+			OutCookInfo.bSupportFaceRemapOnMeshBVH = bSupportFaceRemapOnMeshBVH;
 			OutCookInfo.bSupportVertexRemap = bSupportVertexRemap;
 		}
 		else
@@ -1527,6 +1530,10 @@ void UBodySetup::GetGeometryDDCKey(FString& OutString) const
 		OutString += FString("_1");
 	}
 
+	if (!bSupportFaceRemapOnMeshBVH)
+	{
+		OutString += FString("_F");
+	}
 	if (!bSupportVertexRemap)
 	{
 		OutString += FString("_V");
diff --git a/Engine/Source/Runtime/Engine/Private/PhysicsEngine/Experimental/ChaosCooking.cpp b/Engine/Source/Runtime/Engine/Private/PhysicsEngine/Experimental/ChaosCooking.cpp
--- a/Engine/Source/Runtime/Engine/Private/PhysicsEngine/Experimental/ChaosCooking.cpp
+++ b/Engine/Source/Runtime/Engine/Private/PhysicsEngine/Experimental/ChaosCooking.cpp
@@ -324,8 +324,9 @@ namespace Chaos
 				}
 
 				
-				// TODO: If InParams.bSupportFaceRemap is false, we should in theory not need the remap array, but in practice things still rely on it being there
-				TUniquePtr<TArray<int32>> OutFaceRemapPtr = MakeUnique<TArray<int32>>(OutFaceRemap);
+				// Keep the face remap if either related flag is true, to carry it through initial processing. (It will be stripped later from the cook helper's remap array if bSupportFaceRemap is false.)
+				bool bNeedFaceRemapOnBVH = InParams.bSupportFaceRemap || InParams.bSupportFaceRemapOnMeshBVH;
+				TUniquePtr<TArray<int32>> OutFaceRemapPtr = bNeedFaceRemapOnBVH ? MakeUnique<TArray<int32>>(OutFaceRemap) : nullptr;
 				TUniquePtr<TArray<int32>> OutVertexRemapPtr = (InParams.bSupportVertexRemap && Chaos::TriMeshPerPolySupport) ? MakeUnique<TArray<int32>>(OutVertexRemap) : nullptr;
 				if (!OutFaceRemapPtr)
 				{
@@ -343,7 +344,7 @@ namespace Chaos
 				OutTriangleMeshes.Emplace(TriangleMesh);
 
 				// Propagate remapped indices from the FTriangleMeshImplicitObject back to the remap array
-				// TODO: If we use the bSupportFaceRemap variable to disable the face remap, then we can use this if-statement: if (InParams.bSupportFaceRemap)
+				if (bNeedFaceRemapOnBVH)
 				{
 					CopyUpdatedFaceRemapFromTriangleMesh(*TriangleMesh, OutFaceRemap);
 				}
