Fix crash when building runtime virtual texture streaming mips if there is more than one RVT component referencing and rendering with the streaming mip texture.
#rnx
#Synced-CL 49674283

diff --git a/Engine/Source/Runtime/Engine/Private/Components/RuntimeVirtualTextureComponent.cpp b/Engine/Source/Runtime/Engine/Private/Components/RuntimeVirtualTextureComponent.cpp
--- a/Engine/Source/Runtime/Engine/Private/Components/RuntimeVirtualTextureComponent.cpp
+++ b/Engine/Source/Runtime/Engine/Private/Components/RuntimeVirtualTextureComponent.cpp
@@ -482,10 +482,6 @@ void URuntimeVirtualTextureComponent::InitializeStreamingTexture(EShadingPath Sh
 	// We need an existing StreamingTexture object to update.
 	if (IsActiveInWorld() && VirtualTexture != nullptr && StreamingTexture != nullptr)
 	{
-		// Call Unregister()/Register() around this to release and recreate runtime virtual texture producers associated with a UVirtualTextureBuilder.
-		// Required around modifications of a UVirtualTextureBuilder because virtual producers hold pointers to the internal data.
-		UnregisterComponent();
-
 		FVirtualTextureBuildDesc BuildDesc;
 		BuildDesc.bSinglePhysicalSpace = VirtualTexture->GetSinglePhysicalSpace();
 
@@ -518,8 +514,6 @@ void URuntimeVirtualTextureComponent::InitializeStreamingTexture(EShadingPath Sh
 		constexpr bool bWaitForCompilation = true;
 		StreamingTexture->BuildTexture(ShadingPath, BuildDesc, bWaitForCompilation);
 		StreamingTexture->Modify();
-
-		RegisterComponent();
 	}
 }
 
diff --git a/Engine/Source/Runtime/Engine/Private/VT/VirtualTextureBuilder.cpp b/Engine/Source/Runtime/Engine/Private/VT/VirtualTextureBuilder.cpp
--- a/Engine/Source/Runtime/Engine/Private/VT/VirtualTextureBuilder.cpp
+++ b/Engine/Source/Runtime/Engine/Private/VT/VirtualTextureBuilder.cpp
@@ -3,10 +3,12 @@
 #include "VT/VirtualTextureBuilder.h"
 
 #include "RenderUtils.h"
-#include "VT/VirtualTexture.h"
+#include "Components/RuntimeVirtualTextureComponent.h"
 #include "SceneInterface.h"
 #include "SceneUtils.h"
 #include "TextureCompiler.h"
+#include "UObject/UObjectIterator.h"
+#include "VT/VirtualTexture.h"
 
 #if WITH_EDITOR
 #include "Interfaces/ITargetPlatform.h"
@@ -177,6 +179,15 @@ void UVirtualTextureBuilder::BuildTexture(EShadingPath ShadingPath, FVirtualText
 		Texture = NewObject<UVirtualTexture2D>(this, TEXT("Texture"));
 		BuildVirtualTexture2D(Texture, BuildDesc, bWaitForCompilation);
 	}
+
+	// Reinit associated runtime virtual textures.
+	for (TObjectIterator<URuntimeVirtualTextureComponent> It(RF_ClassDefaultObject, true, EInternalObjectFlags::Garbage); It; ++It)
+	{
+		if (It->GetStreamingTexture() == this)
+		{
+			It->MarkRenderStateDirty();
+		}
+	}
 }
 
 #endif
