Fix localization warning

#Synced-CL 49689149

diff --git a/Engine/Source/Runtime/Engine/Private/ShaderCompiler/ShaderCompilerThreadRunnable.cpp b/Engine/Source/Runtime/Engine/Private/ShaderCompiler/ShaderCompilerThreadRunnable.cpp
--- a/Engine/Source/Runtime/Engine/Private/ShaderCompiler/ShaderCompilerThreadRunnable.cpp
+++ b/Engine/Source/Runtime/Engine/Private/ShaderCompiler/ShaderCompilerThreadRunnable.cpp
@@ -1403,7 +1403,7 @@ void FShaderCompileThreadRunnable::TracePerWorkerSummary(const TArray<FWorkerMem
 	else if (bInHasMemoryViolation)
 	{
 		UE_LOGFMT_NSLOC(
-			LogShaderCompilers, Display, "ShaderCompilers", "PerWorkerMemoryViolation",
+			LogShaderCompilers, Display, "ShaderCompilers", "FailedRetrievePerWorkerMemoryViolation",
 			"Shader Compiler Worker memory violation: Failed to retrieve per-worker summary breakdown."
 		);
 	}
