Iris
- Only call SetActorIsPendingPostNetInit for actors that are exected to call PostNetInit
- Fixed issue where a static actor would skip their call to BeginPlay if they are in pending dispatch list and BeginPlay was issues on world triggered from OnRep of another actor.

#Synced-CL 49669520

diff --git a/Engine/Source/Runtime/Engine/Private/Net/Iris/ReplicationSystem/NetActorFactory.cpp b/Engine/Source/Runtime/Engine/Private/Net/Iris/ReplicationSystem/NetActorFactory.cpp
--- a/Engine/Source/Runtime/Engine/Private/Net/Iris/ReplicationSystem/NetActorFactory.cpp
+++ b/Engine/Source/Runtime/Engine/Private/Net/Iris/ReplicationSystem/NetActorFactory.cpp
@@ -318,9 +318,6 @@ UNetObjectFactory::FInstantiateResult UNetActorFactory::InstantiateReplicatedObj
 			Result.Flags |= EReplicationBridgeCreateNetRefHandleResultFlags::AllowDestroyInstanceFromRemote;
 		}
 
-		// Mark actor as pending its PostNetInit
-		Actor->SetActorIsPendingPostNetInit(true);
-
 		return Result;
 	}
 
@@ -343,9 +340,6 @@ UNetObjectFactory::FInstantiateResult UNetActorFactory::InstantiateReplicatedObj
 			Result.Flags |= EReplicationBridgeCreateNetRefHandleResultFlags::AllowDestroyInstanceFromRemote;
 		}
 
-		// Mark actor as pending its PostNetInit
-		FoundActor->SetActorIsPendingPostNetInit(true);
-
 		return Result;
 	}
 
