- Make sure padding on texture coordinate buffer with uneven amount of TC sets is taking into account because ResourceCount in CreateRHIBuffer (it doesn't have the element size so it also can't really use it)
- Restore shared setup for TC elements in input layout and force all entries back to half4/float4 to reduce amount of (unneeded) extra PSOs compile due to permutations on different amount of TCs used

#Synced-CL 49719249

diff --git a/Engine/Source/Runtime/Engine/Private/LocalVertexFactory.cpp b/Engine/Source/Runtime/Engine/Private/LocalVertexFactory.cpp
--- a/Engine/Source/Runtime/Engine/Private/LocalVertexFactory.cpp
+++ b/Engine/Source/Runtime/Engine/Private/LocalVertexFactory.cpp
@@ -558,22 +558,45 @@ void AddStaticMeshTextureCoordinateElements(
 	FVertexDeclarationElementList& Elements,
 	FVertexFactory::FVertexStreamList& InOutStreams)
 {
-	for (int32 CoordinateIndex = 0; CoordinateIndex < TextureCoordinates.Num(); ++CoordinateIndex)
+	if (TextureCoordinates.IsEmpty())
 	{
-		Elements.Add(FVertexFactory::AccessStreamComponent(
-			TextureCoordinates[CoordinateIndex],
-			BaseTexCoordAttribute + CoordinateIndex,
-			InOutStreams
-		));
+		return;
 	}
-
-	for (int32 CoordinateIndex = TextureCoordinates.Num(); CoordinateIndex < MAX_STATIC_TEXCOORDS / 2; ++CoordinateIndex)
+	
+	// Validate all texture coordinates share the same stream data
+	for (int32 TextureCoordIndex = 1; TextureCoordIndex < TextureCoordinates.Num(); ++TextureCoordIndex)
+	{
+		check(TextureCoordinates[0].VertexBuffer == TextureCoordinates[TextureCoordIndex].VertexBuffer);
+		check(TextureCoordinates[0].Stride == TextureCoordinates[TextureCoordIndex].Stride);
+		check(TextureCoordinates[0].StreamOffset == TextureCoordinates[TextureCoordIndex].StreamOffset);
+		check(TextureCoordinates[0].VertexStreamUsage == TextureCoordinates[TextureCoordIndex].VertexStreamUsage);
+	}
+	
+	EVertexElementType FirstElementType = TextureCoordinates[0].Type;
+	
+	EVertexElementType UVDoubleWideVertexElementType = VET_None;
+	uint32 DoubleUVSizeInBytes = 0;
+	if (FirstElementType == VET_Float2 || FirstElementType == VET_Float4)
+	{
+		DoubleUVSizeInBytes = sizeof(TStaticMeshVertexUVsTypeSelector<EStaticMeshVertexUVType::HighPrecision>::UVsTypeT) * 2;
+		UVDoubleWideVertexElementType = VET_Float4;
+	}
+	else
+	{
+		DoubleUVSizeInBytes = sizeof(TStaticMeshVertexUVsTypeSelector<EStaticMeshVertexUVType::Default>::UVsTypeT) * 2;
+		UVDoubleWideVertexElementType = VET_Half4;
+	}
+	
+	uint8 TexCoordStreamIndex = FVertexFactory::AddVertexStream(TextureCoordinates[0], InOutStreams);
+	
+	for (int32 CoordinateIndex = 0; CoordinateIndex < MAX_STATIC_TEXCOORDS / 2; ++CoordinateIndex)
 	{
-		Elements.Add(FVertexFactory::AccessStreamComponent(
-			TextureCoordinates[TextureCoordinates.Num() - 1],
+		Elements.Add(FVertexElement(
+			TexCoordStreamIndex,
+			DoubleUVSizeInBytes * CoordinateIndex,
+			UVDoubleWideVertexElementType,
 			BaseTexCoordAttribute + CoordinateIndex,
-			InOutStreams
-		));
+			TextureCoordinates[0].Stride));
 	}
 }
 
diff --git a/Engine/Source/Runtime/Engine/Private/Rendering/StaticMeshVertexBuffer.cpp b/Engine/Source/Runtime/Engine/Private/Rendering/StaticMeshVertexBuffer.cpp
--- a/Engine/Source/Runtime/Engine/Private/Rendering/StaticMeshVertexBuffer.cpp
+++ b/Engine/Source/Runtime/Engine/Private/Rendering/StaticMeshVertexBuffer.cpp
@@ -11,6 +11,60 @@
 #include "RHIResourceReplace.h"
 #include "StaticMeshVertexData.h"
 
+	
+static uint32 GetNumTexCoordInVertexBuffer(uint32 NumTexCoords, uint32 NumVertices)
+{
+	// Increase size with single vertex in case we have odd number of texcoords because the vertex declaration
+	// will declare the stream as half4/float4 while it only contains a half2/float2 and the last element can in theory
+	// be out of bounds then
+	// Vertex streams are all created with half4/float4 to reduce the number of PSO permutations coming from different texcoords in the vertex declaration
+	uint32 Padding = (NumTexCoords % 2 == 0) ? 0 : 1;
+	return (NumVertices + Padding) * NumTexCoords;
+}
+
+/** The implementation of the static mesh texture coordingate vertex data storage type which takes extra padding into account. */
+template<typename VertexDataType>
+class TStaticTexCoordVertexData :
+	public TStaticMeshVertexData<VertexDataType>
+{
+protected:
+
+	uint32 NumTexCoords;
+
+public:
+
+	TStaticTexCoordVertexData(uint32 InNumTexCoords, bool InNeedsCPUAccess=false)
+		: TStaticMeshVertexData<VertexDataType>(InNeedsCPUAccess), NumTexCoords(InNumTexCoords)
+	{
+	}
+
+	void ResizeBuffer(uint32 NumVertices, EResizeBufferFlags BufferFlags = EResizeBufferFlags::None) override
+	{
+		uint32 NumTexCoordData = GetNumTexCoordInVertexBuffer(NumTexCoords, NumVertices);
+		if ((uint32)TStaticMeshVertexData<VertexDataType>::Data.Num() < NumTexCoordData)
+		{
+			// Enlarge the array.
+			if (!EnumHasAnyFlags(BufferFlags, EResizeBufferFlags::AllowSlackOnGrow))
+			{
+				TStaticMeshVertexData<VertexDataType>::Data.Reserve(NumTexCoordData);
+			}
+
+			TStaticMeshVertexData<VertexDataType>::Data.AddUninitialized(NumTexCoordData - TStaticMeshVertexData<VertexDataType>::Data.Num());
+		}
+		else if ((uint32)TStaticMeshVertexData<VertexDataType>::Data.Num() > NumTexCoordData)
+		{
+			// Shrink the array.
+			EAllowShrinking AllowShinking = EnumHasAnyFlags(BufferFlags, EResizeBufferFlags::AllowSlackOnReduce) ? EAllowShrinking::No : EAllowShrinking::Yes;
+			TStaticMeshVertexData<VertexDataType>::Data.RemoveAt(NumTexCoordData, TStaticMeshVertexData<VertexDataType>::Data.Num() - NumTexCoordData, AllowShinking);
+		}
+	}
+
+	void Empty(uint32 NumVertices) override
+	{
+		TStaticMeshVertexData<VertexDataType>::Data.Empty(GetNumTexCoordInVertexBuffer(NumTexCoords, NumVertices));
+	}
+};
+
 FStaticMeshVertexBuffer::FStaticMeshVertexBuffer() :
 	TangentsData(nullptr),
 	TexcoordData(nullptr),
@@ -54,7 +108,7 @@ void FStaticMeshVertexBuffer::Init(uint32 InNumVertices, uint32 InNumTexCoords,
 	// Allocate the vertex data buffer.
 	TangentsData->ResizeBuffer(NumVertices);
 	TangentsDataPtr = NumVertices ? TangentsData->GetDataPointer() : nullptr;
-	TexcoordData->ResizeBuffer(NumVertices * GetNumTexCoords());
+	TexcoordData->ResizeBuffer(NumVertices);
 	TexcoordDataPtr = NumVertices ? TexcoordData->GetDataPointer() : nullptr;
 }
 
@@ -118,7 +172,7 @@ void FStaticMeshVertexBuffer::Init(const FStaticMeshVertexBuffer& InVertexBuffer
 			check(GetNumTexCoords() == InVertexBuffer.GetNumTexCoords());
 			const uint8* InData = InVertexBuffer.TexcoordDataPtr;
 
-			TexcoordData->ResizeBuffer(NumVertices * GetNumTexCoords());
+			TexcoordData->ResizeBuffer(NumVertices);
 			TexcoordDataPtr = TexcoordData->GetDataPointer();
 			FMemory::Memcpy(TexcoordDataPtr, InData, TexcoordData->GetStride() * NumVertices * GetNumTexCoords());
 		}
@@ -133,8 +187,8 @@ void FStaticMeshVertexBuffer::ConvertHalfTexcoordsToFloat(const uint8* InData)
 	FStaticMeshVertexDataInterface* OriginalTexcoordData = TexcoordData;
 
 	typedef TStaticMeshVertexUVsDatum<typename TStaticMeshVertexUVsTypeSelector<EStaticMeshVertexUVType::HighPrecision>::UVsTypeT> UVType;
-	TexcoordData = new TStaticMeshVertexData<UVType>(OriginalTexcoordData->GetAllowCPUAccess());
-	TexcoordData->ResizeBuffer(NumVertices * GetNumTexCoords());
+	TexcoordData = new TStaticTexCoordVertexData<UVType>(NumTexCoords, OriginalTexcoordData->GetAllowCPUAccess());
+	TexcoordData->ResizeBuffer(NumVertices);
 	TexcoordDataPtr = TexcoordData->GetDataPointer();
 	TexcoordStride = sizeof(UVType);
 
@@ -169,7 +223,7 @@ void FStaticMeshVertexBuffer::AppendVertices( const FStaticMeshBuildVertex* Vert
 		NumVertices += NumVerticesToAppend;
 
 		TangentsData->ResizeBuffer(NumVertices);
-		TexcoordData->ResizeBuffer(NumVertices * GetNumTexCoords());
+		TexcoordData->ResizeBuffer(NumVertices);
 
 		if( NumVertices > 0 )
 		{
@@ -284,12 +338,7 @@ FBufferRHIRef FStaticMeshVertexBuffer::CreateTangentsRHIBuffer(FRHICommandListBa
 
 FBufferRHIRef FStaticMeshVertexBuffer::CreateTexCoordRHIBuffer(FRHICommandListBase& RHICmdList)
 {
-	// Increase size with single vertex in case we have odd number of texcoords because the vertex declaration
-	// will declare the stream as half4/float4 while it only contains a half2/float2 and the last element can in theory
-	// be out of bounds then
-	// Vertex streams are all created with half4/float4 to reduce the number of PSO permutations coming from different texcoords in the vertex declaration
-	uint32 Padding = (GetNumTexCoords() % 2 == 0) ? 0 : 1;
-	uint32 ResourceCount = (GetNumVertices() + Padding) * GetNumTexCoords();
+	uint32 ResourceCount = GetNumTexCoordInVertexBuffer(GetNumTexCoords(), GetNumVertices());
 	return CreateRHIBuffer(RHICmdList, TexcoordData, ResourceCount, BUF_Static | BUF_ShaderResource, TEXT("TexCoordRHIBuffer"));
 }
 
@@ -389,13 +438,13 @@ void FStaticMeshVertexBuffer::AllocateData(bool bNeedsCPUAccess /*= true*/)
 	{
 		typedef TStaticMeshVertexUVsDatum<typename TStaticMeshVertexUVsTypeSelector<EStaticMeshVertexUVType::HighPrecision>::UVsTypeT> UVType;
 		TexcoordStride = sizeof(UVType);
-		TexcoordData = new TStaticMeshVertexData<UVType>(bNeedsCPUAccess);
+		TexcoordData = new TStaticTexCoordVertexData<UVType>(NumTexCoords, bNeedsCPUAccess);
 	}
 	else
 	{
 		typedef TStaticMeshVertexUVsDatum<typename TStaticMeshVertexUVsTypeSelector<EStaticMeshVertexUVType::Default>::UVsTypeT> UVType;
 		TexcoordStride = sizeof(UVType);
-		TexcoordData = new TStaticMeshVertexData<UVType>(bNeedsCPUAccess);
+		TexcoordData = new TStaticTexCoordVertexData<UVType>(NumTexCoords, bNeedsCPUAccess);
 	}
 }
 
diff --git a/Engine/Source/Runtime/Engine/Public/StaticMeshVertexData.h b/Engine/Source/Runtime/Engine/Public/StaticMeshVertexData.h
--- a/Engine/Source/Runtime/Engine/Public/StaticMeshVertexData.h
+++ b/Engine/Source/Runtime/Engine/Public/StaticMeshVertexData.h
@@ -11,8 +11,11 @@ template<typename VertexDataType>
 class TStaticMeshVertexData :
 	public FStaticMeshVertexDataInterface
 {
+protected:
+
 	using FVertexResourceArray = TResourceArray<VertexDataType, VERTEXBUFFER_ALIGNMENT>;
 	FVertexResourceArray Data;
+
 public:
 
 	/**
