Fix refcount decreasing pointer instead of counter

#rnx
#Synced-CL 49546171

diff --git a/Engine/Source/Runtime/Engine/Private/Net/Iris/ReplicationSystem/NetRootObjectAdapter.cpp b/Engine/Source/Runtime/Engine/Private/Net/Iris/ReplicationSystem/NetRootObjectAdapter.cpp
--- a/Engine/Source/Runtime/Engine/Private/Net/Iris/ReplicationSystem/NetRootObjectAdapter.cpp
+++ b/Engine/Source/Runtime/Engine/Private/Net/Iris/ReplicationSystem/NetRootObjectAdapter.cpp
@@ -428,7 +428,7 @@ void FNetRootObjectAdapter::RemoveRelevantWith(const UObject* OtherReplicatedObj
 			}
 			else
 			{
-				--RootObjectRefCount;
+				--(*RootObjectRefCount);
 			}
 		}
 	});
