[World Partition]

Fixed missing cell transformer invocation in UWorldPartitionRuntimeSpatialHash path.

#Synced-CL 49660775

diff --git a/Engine/Source/Runtime/Engine/Private/WorldPartition/WorldPartitionRuntimeSpatialHash.cpp b/Engine/Source/Runtime/Engine/Private/WorldPartition/WorldPartitionRuntimeSpatialHash.cpp
--- a/Engine/Source/Runtime/Engine/Private/WorldPartition/WorldPartitionRuntimeSpatialHash.cpp
+++ b/Engine/Source/Runtime/Engine/Private/WorldPartition/WorldPartitionRuntimeSpatialHash.cpp
@@ -1673,6 +1673,8 @@ bool UWorldPartitionRuntimeSpatialHash::CreateStreamingGrid(const FSpatialHashRu
 
 					PopulateRuntimeCell(StreamingCell, FilteredActors, OutPackagesToGenerate);
 
+					WorldPartition->ApplyRuntimeCellsTransformerStack(StreamingCell);
+
 					int32 LayerCellIndex;
 					int32* LayerCellIndexPtr = GridLevel.LayerCellsMapping.Find(CellIndex);
 					if (LayerCellIndexPtr)
