Replace CVar r.TextureGroup.OptionalQualityLevel with new device profile property UTextureLODSettings::TextureGroupOptionalQualityLevel

This allows us to specify the value for Windows without propogating to other platforms just because the cooking editor runs on Windows.

#Synced-CL 49717515

diff --git a/Engine/Source/Runtime/Engine/Classes/Engine/TextureLODSettings.h b/Engine/Source/Runtime/Engine/Classes/Engine/TextureLODSettings.h
--- a/Engine/Source/Runtime/Engine/Classes/Engine/TextureLODSettings.h
+++ b/Engine/Source/Runtime/Engine/Classes/Engine/TextureLODSettings.h
@@ -73,19 +73,19 @@ struct FTextureLODGroup
 	UPROPERTY()
 	int32 MaxLODSize_Smallest = -1;
 
-	/** Cap the number of mips such that the largest mip is this big. Has no effect for textures with no mip chain. Used in combination with r.TextureGroup.QualityLevel, r.TextureGroup.MaxQualityLevel and r.TextureGroup.OptionalQualityLevel. Not used for virtual textures. */
+	/** Cap the number of mips such that the largest mip is this big. Has no effect for textures with no mip chain. Used in combination with r.TextureGroup.QualityLevel, r.TextureGroup.MaxQualityLevel and TextureGroupOptionalQualityLevel. Not used for virtual textures. */
 	UPROPERTY()
 	int32 MaxLODSize_Low = -1;
 
-	/** Cap the number of mips such that the largest mip is this big. Has no effect for textures with no mip chain. Used in combination with r.TextureGroup.QualityLevel, r.TextureGroup.MaxQualityLevel and r.TextureGroup.OptionalQualityLevel. Not used for virtual textures. */
+	/** Cap the number of mips such that the largest mip is this big. Has no effect for textures with no mip chain. Used in combination with r.TextureGroup.QualityLevel, r.TextureGroup.MaxQualityLevel and TextureGroupOptionalQualityLevel. Not used for virtual textures. */
 	UPROPERTY()
 	int32 MaxLODSize_Medium = -1;
 
-	/** Cap the number of mips such that the largest mip is this big. Has no effect for textures with no mip chain. Used in combination with r.TextureGroup.QualityLevel, r.TextureGroup.MaxQualityLevel and r.TextureGroup.OptionalQualityLevel. Not used for virtual textures. */
+	/** Cap the number of mips such that the largest mip is this big. Has no effect for textures with no mip chain. Used in combination with r.TextureGroup.QualityLevel, r.TextureGroup.MaxQualityLevel and TextureGroupOptionalQualityLevel. Not used for virtual textures. */
 	UPROPERTY()
 	int32 MaxLODSize_High = -1;
 
-	/** Cap the number of mips such that the largest mip is this big. Has no effect for textures with no mip chain. Used in combination with r.TextureGroup.QualityLevel, r.TextureGroup.MaxQualityLevel and r.TextureGroup.OptionalQualityLevel. Not used for virtual textures. */
+	/** Cap the number of mips such that the largest mip is this big. Has no effect for textures with no mip chain. Used in combination with r.TextureGroup.QualityLevel, r.TextureGroup.MaxQualityLevel and TextureGroupOptionalQualityLevel. Not used for virtual textures. */
 	UPROPERTY()
 	int32 MaxLODSize_Epic = -1;
 
@@ -253,6 +253,10 @@ protected:
 
 public:
 
+	/** If set to a valid value, use MaxLODSize for the given quality instead of OptionalMaxLODSize during cook time. */
+	UPROPERTY(EditAnywhere, config, Category="Texture LOD Settings")
+	int32 TextureGroupOptionalQualityLevel = -1;
+
 	/** Array of LOD settings with entries per group. */
 	UPROPERTY(EditAnywhere, config, Category="Texture LOD Settings")
 	TArray<FTextureLODGroup> TextureLODGroups;
diff --git a/Engine/Source/Runtime/Engine/Private/DeviceProfiles/DeviceProfile.cpp b/Engine/Source/Runtime/Engine/Private/DeviceProfiles/DeviceProfile.cpp
--- a/Engine/Source/Runtime/Engine/Private/DeviceProfiles/DeviceProfile.cpp
+++ b/Engine/Source/Runtime/Engine/Private/DeviceProfiles/DeviceProfile.cpp
@@ -152,6 +152,12 @@ void UDeviceProfile::ValidateTextureLODGroups()
 #define SETUPLODGROUP(GroupId) SetupLODGroup(GroupId);
 	FOREACH_ENUM_TEXTUREGROUP(SETUPLODGROUP)
 #undef SETUPLODGROUP
+
+	// Propagate value from parent if the property isn't set.
+	if (ParentProfile && (TextureGroupOptionalQualityLevel == -1))
+	{
+		TextureGroupOptionalQualityLevel = ParentProfile->TextureGroupOptionalQualityLevel;
+	}
 }
 
 #if WITH_EDITOR
diff --git a/Engine/Source/Runtime/Engine/Private/TextureLODSettings.cpp b/Engine/Source/Runtime/Engine/Private/TextureLODSettings.cpp
--- a/Engine/Source/Runtime/Engine/Private/TextureLODSettings.cpp
+++ b/Engine/Source/Runtime/Engine/Private/TextureLODSettings.cpp
@@ -52,14 +52,6 @@ FAutoConsoleVariableRef CVarTextureGroupQualityLevel(
 	ECVF_Default
 );
 
-int32 TextureGroupOptionalQualityLevel = -1;
-FAutoConsoleVariableRef CVarTextureGroupOptionalQualityLevel(
-	TEXT("r.TextureGroup.OptionalQualityLevel"),
-	TextureGroupOptionalQualityLevel,
-	TEXT("If set, it will use MaxLODSize for the given quality instead of OptionalMaxLODSize during cook time. (default=-1)"),
-	ECVF_Default
-);
-
 void FTextureLODGroup::SetupGroup()
 {
 	// editor would never want to use smaller mips based on memory (could affect cooking, etc!)
