Substrate

diff --git a/Engine/Source/Runtime/Engine/Private/Materials/HLSLMaterialTranslator.cpp b/Engine/Source/Runtime/Engine/Private/Materials/HLSLMaterialTranslator.cpp
--- a/Engine/Source/Runtime/Engine/Private/Materials/HLSLMaterialTranslator.cpp
+++ b/Engine/Source/Runtime/Engine/Private/Materials/HLSLMaterialTranslator.cpp
@@ -145,6 +145,8 @@ static FCookStatsManager::FAutoRegisterCallback RegisterCookStats([](FCookStatsM
 
 #if WITH_EDITORONLY_DATA
 
+using namespace Substrate;
+
 static int32 GAnalyticDerivEnabled = 1;
 static FAutoConsoleVariableRef CVarAnalyticDerivEnabled(
 	TEXT("r.MaterialEditor.AnalyticDeriv"),
@@ -264,7 +266,6 @@ struct FHLSLMaterialTranslator::FEnvironmentDefines
 	bool bUsesDistortion;
 	bool bUsesExplicitDerivatives;
 	bool bUsesFirstPersonInterpolation;
-	bool bDistortionAccountForCoverage;
 	bool bMaterialEnableTranslucencyFogging;
 	bool bMaterialEnableTranslucencyCloudFogging;
 	bool bMaterialIsSky;
@@ -299,31 +300,13 @@ struct FHLSLMaterialTranslator::FEnvironmentDefines
 	bool bMaterialVolumetricAdvancedOverrideAmbientOcclusion;
 	bool bMaterialVolumetricAdvancedAdvancedGroundContribution;
 	bool bMaterialVolumetricCloudEmptySpaceSkippingOutput;
-	bool bMaterialIsSubstrate;
-	bool bMaterialUsesRootNodeToSubstrateHiddenConversion;
 	bool bDualSourceColorBlendingEnabled;
-	bool bSubstratePremultipliedAlphaOpacityOverridden;
-	bool bSubstrateUsesConversionFromLegacy;
-	bool bSubstrateMaterialOutputOpaqueRoughRefractions;
-	int32 SubstrateMaterialExportType;
-	int32 SubstrateMaterialExportContext;
-	int32 SubstrateMaterialExportLegacyBlendMode;
-	bool bSubstrateOptimizedUnlit;
-	bool bSubstrateSinglePath;
-	bool bSubstrateFastPath;
-	uint32 SubstrateClampedClosureCount;
-	bool SubstrateIsComplexSpecialPath;
-	bool bSubstrateComplexSpecialPath;
-	bool bSubstrateLegacyIrisNormal;
-	bool bSubstrateLegacyIrisTangent;
+	TArray<TPair<FString, int32>> SubstrateDefines;
 	uint8 NumCustomizedUVs;
 	bool bTextureSampleDebug;
-	bool bSubstrateRoughnessTracking;
-	ESubstrateBsdfFeature SubstrateMaterialBsdfFeatures;
 	bool bMaterialEnableTranslucentLocalLightShadow;
 	bool bMaterialEnableTranslucentHighQualityLocalLightShadow;
 	bool bMaterialEnableTranslucentHighQualityDirectionalLightShadow;
-	TArray<TPair<FString, int32>> SubstrateDefines;
 	TArray<TObjectPtr<UMaterialParameterCollection>> ParameterCollections;
 
 	bool HasShadingModel(EMaterialShadingModel model) const
@@ -373,7 +356,6 @@ struct FHLSLMaterialTranslator::FEnvironmentDefines
 		Ar << bUsesDistortion;
 		Ar << bUsesExplicitDerivatives;
 		Ar << bUsesFirstPersonInterpolation;
-		Ar << bDistortionAccountForCoverage;
 		Ar << bMaterialEnableTranslucencyFogging;
 		Ar << bMaterialEnableTranslucencyCloudFogging;
 		Ar << bMaterialIsSky;
@@ -408,27 +390,9 @@ struct FHLSLMaterialTranslator::FEnvironmentDefines
 		Ar << bMaterialVolumetricAdvancedOverrideAmbientOcclusion;
 		Ar << bMaterialVolumetricAdvancedAdvancedGroundContribution;
 		Ar << bMaterialVolumetricCloudEmptySpaceSkippingOutput;
-		Ar << bMaterialIsSubstrate;
-		Ar << bMaterialUsesRootNodeToSubstrateHiddenConversion;
 		Ar << bDualSourceColorBlendingEnabled;
-		Ar << bSubstratePremultipliedAlphaOpacityOverridden;
-		Ar << bSubstrateUsesConversionFromLegacy;
-		Ar << bSubstrateMaterialOutputOpaqueRoughRefractions;
-		Ar << SubstrateMaterialExportType;
-		Ar << SubstrateMaterialExportContext;
-		Ar << SubstrateMaterialExportLegacyBlendMode;
-		Ar << bSubstrateOptimizedUnlit;
-		Ar << bSubstrateSinglePath;
-		Ar << bSubstrateFastPath;
-		Ar << SubstrateClampedClosureCount;
-		Ar << bSubstrateComplexSpecialPath;
-		Ar << bSubstrateLegacyIrisNormal;
-		Ar << bSubstrateLegacyIrisTangent;
 		Ar << NumCustomizedUVs;
 		Ar << bTextureSampleDebug;
-		Ar << SubstrateDefines;
-		Ar << bSubstrateRoughnessTracking;
-		Ar << SubstrateMaterialBsdfFeatures;
 		Ar << bMaterialEnableTranslucentLocalLightShadow;
 		Ar << bMaterialEnableTranslucentHighQualityLocalLightShadow;
 		Ar << bMaterialEnableTranslucentHighQualityDirectionalLightShadow;
@@ -480,8 +444,6 @@ static inline bool IsDebugTextureSampleEnabled()
 	return IsAnalyticDerivEnabled() && (GDebugTextureSampleEnabled != 0);
 }
 
-#define DEBUG_SUBSTRATE_TREE_STACK 0
-
 /** @return the vector type containing a given number of components. */
 static inline EMaterialValueType GetVectorType(uint32 NumComponents)
 {
@@ -581,80 +543,6 @@ static inline const TCHAR * GetFloatZeroVector(uint32 NumComponents)
 	}
 }
 
-static inline const TCHAR* GetSubstrateOperatorStr(int32 OperatorType)
-{
-	switch (OperatorType)
-	{
-	case SUBSTRATE_OPERATOR_WEIGHT:
-	{
-		return TEXT("WEIGHT    ");
-	}
-	case SUBSTRATE_OPERATOR_VERTICAL:
-	{
-		return TEXT("VERTICAL  ");
-	}
-	case SUBSTRATE_OPERATOR_HORIZONTAL:
-	{
-		return TEXT("HORIZONTAL");
-	}
-	case SUBSTRATE_OPERATOR_ADD:
-	{
-		return TEXT("ADD       ");
-	}
-	case SUBSTRATE_OPERATOR_SELECT:
-	{
-		return TEXT("SELECT    ");
-	}
-	case SUBSTRATE_OPERATOR_BSDF:
-	{
-		return TEXT("BSDF      ");
-	}
-	case SUBSTRATE_OPERATOR_BSDF_LEGACY:
-	{
-		return TEXT("BSDFLEGACY");
-	}
-	}
-	return TEXT("UNKNOWN   ");
-};
-
-FHLSLMaterialTranslator::FSubstrateCompilationContext::FSubstrateCompilationContext()
-{
-	Initialise();
-}
-
-FHLSLMaterialTranslator::FSubstrateCompilationContext::FSubstrateCompilationContext(ESubstrateCompilationContext InCompilationContext)
-{
-	Initialise();
-	CompilationContextIndex = InCompilationContext;
-}
-
-// This limitation is required so that the operator array is never reallocated, invalidating references to it while parsing the Substrate tree within SubstrateGenerateMaterialTopologyTree for instance.
-#define SUBSTRATE_MAX_COMPILER_REGISTERED_OPERATOR_COUNT 128
-
-void FHLSLMaterialTranslator::FSubstrateCompilationContext::Initialise()
-{
-	CompilationContextIndex = ESubstrateCompilationContext::SCC_MAX;
-
-	NextFreeSubstrateShaderNormalIndex = 0;
-	FinalUsedSharedLocalBasesCount = 0;
-	SubstrateMaterialRootOperator = nullptr;
-	SubstrateMaterialExpressionRegisteredOperators.Reserve(SUBSTRATE_MAX_COMPILER_REGISTERED_OPERATOR_COUNT);
-	SubstrateMaterialExpressionToOperatorIndex.Reserve(SUBSTRATE_MAX_COMPILER_REGISTERED_OPERATOR_COUNT);
-	SubstrateMaterialEffectiveClosureCount = 0;
-	SubstrateMaterialRequestedSizeByte = 0;
-	SubstrateMaterialComplexity.Reset();
-	bSubstrateMaterialIsUnlitNode = false;
-	bSubstrateWritesEmissive = false;
-	bSubstrateWritesAmbientOcclusion = false;
-	bSubstrateTreeOutOfStackDepthOccurred = false;
-
-	// Default value used as the root of the tree for the first path (when a node parent==nullptr).
-	SubstrateNodeIdentifierStack.Push(FGuid(0x7AEE, 0xBAD, 0xDEAD, 0xBEEF));
-
-	SubstrateThicknessIndexToExpressionInput.SetNum(0);
-	SubstrateThicknessStack.SetNum(0);
-}
-
 void FHLSLMaterialTranslator::AppendVersion(FShaderKeyGenerator& KeyGen, EShaderPlatform Platform)
 {
 	static const FGuid DDCVersion = FDevSystemGuids::GetSystemGuid(FDevSystemGuids::Get().MaterialTranslationDDCVersion);
@@ -728,10 +616,6 @@ FHLSLMaterialTranslator::FHLSLMaterialTranslator(FMaterial* InMaterial,
 ,	bIsFullyRough(0)
 ,	bAllowCodeChunkGeneration(true)
 ,	bUsesAnisotropy(false)
-,	bSubstrateEnabled(false)
-,	bMaterialIsSubstrate(false)
-,	bMaterialUsesRootNodeToSubstrateHiddenConversion(false)
-,	bProjectSubstrateHiddenMaterialAssetConversionEnabled(false)
 ,	bUsesCurvature(false)
 ,	bUsesPerInstanceFadeAmount(false)
 ,	bCullIntermediateUniformExpressions(GCullIntermediateUniformExpressions)
@@ -745,10 +629,9 @@ FHLSLMaterialTranslator::FHLSLMaterialTranslator(FMaterial* InMaterial,
 ,	DynamicParticleParameterMask(0)
 ,	NumVtSamples(0)
 ,	TargetPlatform(InTargetPlatform)
-,	SubstrateCompilationConfig()
-,   MaterialCompilationOutput(InMaterialCompilationOutput)
+,	MaterialCompilationOutput(InMaterialCompilationOutput)
 {
-	check((sizeof(SubstrateCompilationContext)/sizeof(FSubstrateCompilationContext)) == ESubstrateCompilationContext::SCC_MAX);
+	check((sizeof(SubstrateTranslatorData.SubstrateCompilationContext)/sizeof(FSubstrateCompilationContext)) == ESubstrateCompilationContext::SCC_MAX);
 
 	FMemory::Memzero(SharedPixelProperties);
 
@@ -800,16 +683,11 @@ FHLSLMaterialTranslator::FHLSLMaterialTranslator(FMaterial* InMaterial,
 		}
 	}
 
-	bSubstrateWritesEmissive = false;
-	bSubstrateWritesAmbientOcclusion = false;
-
-	bSubstrateUsesConversionFromLegacy = false;
-	bSubstrateOutputsOpaqueRoughRefractions = false;
-
 	if (InSubstrateCompilationConfig)
 	{
-		SubstrateCompilationConfig = *InSubstrateCompilationConfig;
+		SubstrateTranslatorData.SubstrateCompilationConfig = *InSubstrateCompilationConfig;
 	}
+	SubstrateTranslatorData.Initialize(Material, MaterialCompilationOutput, this, nullptr);
 
 	EnvironmentDefines.Reset(new FEnvironmentDefines);
 
@@ -1479,295 +1357,35 @@ void FHLSLMaterialTranslator::TranslateMaterial()
 		}
 	}
 
+	CHECK_DDC_QUERY_FINISHED_ELSE_RETURN();
 	//
 	// Process the Substrate tree representing the material topology.
 	//
-	bSubstrateEnabled = Substrate::IsSubstrateEnabled();
-	bProjectSubstrateHiddenMaterialAssetConversionEnabled = Substrate::IsHiddenMaterialAssetConversionEnabled();
-	UMaterialExpression* FrontMaterialExpr = nullptr;
-	int32 FrontMaterialOutputIndex = INDEX_NONE;
-
-	bSubstrateWritesEmissive = false;
-	bSubstrateWritesAmbientOcclusion = false;
-
-	CHECK_DDC_QUERY_FINISHED_ELSE_RETURN();
-	UMaterialExpression* ExpressionToPreview = Material->GetMaterialGraphNodePreviewExpression();
-	if (ExpressionToPreview)
-	{
-		// If this is a compilation used for a preview of a node from am aterial graph, then we need ot use that node as the root of the Substrate tree.
-		// Then the resulting SubstrateData will be converted to a color using SubstrateCompilePreview translator function when compiling the emissive color channel the node has been short circuited into.
-		const uint32 ExpressionPreviewOutputIndex = 0;
-		if (ExpressionToPreview && ExpressionToPreview->IsResultSubstrateMaterial(ExpressionPreviewOutputIndex))
-		{
-			FrontMaterialExpr = ExpressionToPreview;
-			FrontMaterialOutputIndex = ExpressionPreviewOutputIndex;
-		}
-		else
-		{
-			FrontMaterialExpr = nullptr;	// This is not a Substrate input that is connected there so we cannot create a Substrate tree.
-		}
-	}
-	else
-	{
-		FSubstrateMaterialInput* FrontMaterialInput = Material->GetMaterialInterface() ? &Material->GetMaterialInterface()->GetMaterial()->GetEditorOnlyData()->FrontMaterial : nullptr;
-		FrontMaterialExpr = FrontMaterialInput ? FrontMaterialInput->GetTracedInput().Expression : nullptr;
-		FrontMaterialOutputIndex = FrontMaterialInput ? FrontMaterialInput->OutputIndex : INDEX_NONE;
-	}
-
-
-	FExpressionInput NullInput;
-	FExpressionInput* ThinTranslucentTransmittanceColor = &NullInput;
-	FExpressionInput* ThinTranslucentSurfaceCoverage = &NullInput;
-	FExpressionInput* WaterScatteringCoefficients = &NullInput;
-	FExpressionInput* WaterAbsorptionCoefficients = &NullInput;
-	FExpressionInput* WaterPhaseG = &NullInput;
-	FExpressionInput* ColorScaleBehindWater = &NullInput;
-	FExpressionInput* ClearCoatNormal = &NullInput;
-	FExpressionInput* CustomTangent = &NullInput;
-
-	bMaterialUsesRootNodeToSubstrateHiddenConversion = bSubstrateEnabled && bProjectSubstrateHiddenMaterialAssetConversionEnabled && !FrontMaterialExpr;
-
-	// If bSubstrateMaterialUsesLegacyMaterialCompilation logic is changed, then FMaterialResource::IsSubstrateMaterial must be updated.
-	const bool bSubstrateMaterialUsesLegacyMaterialCompilation = bMaterialUsesRootNodeToSubstrateHiddenConversion && Substrate::IsSubstrateBlendableGBufferEnabled(Platform);
-	if (bSubstrateMaterialUsesLegacyMaterialCompilation)
-	{
-		// In this case we use the legacy material shader: faster to compile and less instructions (no substrate conversion)
-		bSubstrateEnabled = false;
-		bProjectSubstrateHiddenMaterialAssetConversionEnabled = false;
-		bMaterialUsesRootNodeToSubstrateHiddenConversion = false;
-
-		// Still set some Substrate material compilation output that are needed at runtime.
-		FMaterialShadingModelField MaterialShadingModels = GetCompiledShadingModels();
-		MaterialCompilationOutput.SubstrateMaterialCompilationOutput.SubstrateClosureCount = 1;
-		MaterialCompilationOutput.SubstrateMaterialCompilationOutput.SubstrateUintPerPixel = 0;
-		MaterialCompilationOutput.SubstrateMaterialCompilationOutput.bMaterialUsesLegacyGBufferDataPassThrough = 1;
-		if (MaterialShadingModels.HasOnlyShadingModel(MSM_DefaultLit))
-		{
-			MaterialCompilationOutput.SubstrateMaterialCompilationOutput.SubstrateMaterialType = SUBSTRATE_MATERIAL_TYPE_SIMPLE;
-		}
-		else
-		{
-			if (MaterialShadingModels.HasShadingModel(MSM_Cloth))
-			{
-				MaterialCompilationOutput.SubstrateMaterialCompilationOutput.SubstrateMaterialType = SUBSTRATE_MATERIAL_TYPE_SINGLE;
-				MaterialCompilationOutput.SubstrateMaterialCompilationOutput.SubstrateMaterialBsdfFeatures |= ESubstrateBsdfFeature::Fuzz;
-			}
-			if (MaterialShadingModels.HasShadingModel(MSM_ClearCoat))
-			{
-				MaterialCompilationOutput.SubstrateMaterialCompilationOutput.SubstrateMaterialType = SUBSTRATE_MATERIAL_TYPE_SINGLE;
-				MaterialCompilationOutput.SubstrateMaterialCompilationOutput.SubstrateMaterialBsdfFeatures |= ESubstrateBsdfFeature::SecondRoughnessOrSimpleClearCoat;
-			}
-			if (MaterialShadingModels.HasAnyShadingModel({ MSM_Subsurface,MSM_SubsurfaceProfile,MSM_TwoSidedFoliage,MSM_Eye,MSM_PreintegratedSkin }))
-			{
-				MaterialCompilationOutput.SubstrateMaterialCompilationOutput.SubstrateMaterialType = SUBSTRATE_MATERIAL_TYPE_SINGLE;
-				MaterialCompilationOutput.SubstrateMaterialCompilationOutput.SubstrateMaterialBsdfFeatures |= ESubstrateBsdfFeature::SSS;
-			}
-			if (MaterialShadingModels.HasShadingModel(MSM_Eye))
-			{
-				MaterialCompilationOutput.SubstrateMaterialCompilationOutput.SubstrateMaterialType = SUBSTRATE_MATERIAL_TYPE_COMPLEX;
-				MaterialCompilationOutput.SubstrateMaterialCompilationOutput.SubstrateMaterialBsdfFeatures |= ESubstrateBsdfFeature::Eye;
-			}
-			if (MaterialShadingModels.HasShadingModel(MSM_Hair))
-			{
-				MaterialCompilationOutput.SubstrateMaterialCompilationOutput.SubstrateMaterialType = SUBSTRATE_MATERIAL_TYPE_COMPLEX;
-				MaterialCompilationOutput.SubstrateMaterialCompilationOutput.SubstrateMaterialBsdfFeatures |= ESubstrateBsdfFeature::Hair;
-			}
-
-		}
-	}
-
-	if (bMaterialUsesRootNodeToSubstrateHiddenConversion)
-	{
-		UMaterialInterface* MaterialInterface = Material->GetMaterialInterface();
-		UMaterial* BaseMaterial = MaterialInterface ? MaterialInterface->GetMaterial() : nullptr;
-		UMaterialEditorOnlyData* EditorOnlyData = BaseMaterial ? BaseMaterial->GetEditorOnlyData() : nullptr;
-
-		{
-			// Gather input for custom output
-			UMaterialExpressionThinTranslucentMaterialOutput* ThinTranslucentOutput = nullptr;
-			UMaterialExpressionSingleLayerWaterMaterialOutput* SingleLayerWaterOutput = nullptr;
-			UMaterialExpressionClearCoatNormalCustomOutput* ClearCoatBottomNormalOutput = nullptr;
-			UMaterialExpressionTangentOutput* TangentOutput = nullptr;
-			TArray<class UMaterialExpressionCustomOutput*> CustomOutputExpressions;
-			BaseMaterial->GetAllCustomOutputExpressions(CustomOutputExpressions);
-			for (UMaterialExpressionCustomOutput* Expression : CustomOutputExpressions)
-			{
-				// Gather custom output for thin translucency
-				if (ThinTranslucentOutput == nullptr && Cast<UMaterialExpressionThinTranslucentMaterialOutput>(Expression))
-				{
-					ThinTranslucentOutput = Cast<UMaterialExpressionThinTranslucentMaterialOutput>(Expression);
-
-					ThinTranslucentTransmittanceColor = &ThinTranslucentOutput->TransmittanceColor;
-					ThinTranslucentSurfaceCoverage    = &ThinTranslucentOutput->SurfaceCoverage;
-				}
-
-				// Gather custom output for single layer water
-				if (SingleLayerWaterOutput == nullptr && Cast<UMaterialExpressionSingleLayerWaterMaterialOutput>(Expression))
-				{
-					SingleLayerWaterOutput = Cast<UMaterialExpressionSingleLayerWaterMaterialOutput>(Expression);
-
-					WaterScatteringCoefficients = &SingleLayerWaterOutput->ScatteringCoefficients;
-					WaterAbsorptionCoefficients = &SingleLayerWaterOutput->AbsorptionCoefficients;
-					WaterPhaseG = &SingleLayerWaterOutput->PhaseG;
-					ColorScaleBehindWater = &SingleLayerWaterOutput->ColorScaleBehindWater;
-				}
-
-				// Gather custom output for clear coat
-				if (ClearCoatBottomNormalOutput == nullptr && Cast<UMaterialExpressionClearCoatNormalCustomOutput>(Expression))
-				{
-					ClearCoatBottomNormalOutput = Cast<UMaterialExpressionClearCoatNormalCustomOutput>(Expression);
-					ClearCoatNormal = &ClearCoatBottomNormalOutput->Input;
-				}
-
-				// Gather custom output for tangent (unused atm)
-				if (TangentOutput == nullptr && Cast<UMaterialExpressionTangentOutput>(Expression))
-				{
-					TangentOutput = Cast<UMaterialExpressionTangentOutput>(Expression);
-					CustomTangent = &TangentOutput->Input;
-				}
-
-				if (ThinTranslucentOutput && SingleLayerWaterOutput && ClearCoatBottomNormalOutput && TangentOutput)
-				{
-					break;
-				}
-			}
-		}
-
-		for (uint32 SubstrateCompilationContextIndex = 0; SubstrateCompilationContextIndex < ESubstrateCompilationContext::SCC_MAX; ++SubstrateCompilationContextIndex)
-		{
-			// This is needed because SubstrateGenerateMaterialTopologyTree will call some compiler context though material expressions.
-			CurrentSubstrateCompilationContext = ESubstrateCompilationContext(SubstrateCompilationContextIndex);
-
-			FSubstrateCompilationContext& SubstrateCtx = SubstrateCompilationContext[SubstrateCompilationContextIndex];
-			SubstrateCtx = FSubstrateCompilationContext(ESubstrateCompilationContext(SubstrateCompilationContextIndex));
-
-			if (EditorOnlyData)
-			{
-				FGuid RootNodeConversionGUID = FGuid(0, 1, 2, 3); // This dummy GUID is the root node.
-				bSubstrateWritesEmissive |= EditorOnlyData->EmissiveColor.IsConnected();
-				bSubstrateWritesAmbientOcclusion |= EditorOnlyData->AmbientOcclusion.IsConnected();
-
-				FExpressionInput* SurfaceThickness = nullptr;
-				SubstrateThicknessStackPush(nullptr, SurfaceThickness);
-
-				if (BaseMaterial && BaseMaterial->bUseMaterialAttributes)
-				{
-					UMaterialExpressionSubstrateConvertMaterialAttributes::SubstrateGenerateMaterialTopologyTreeCommon(
-						this, RootNodeConversionGUID, nullptr /*Parent*/, 0,
-						BaseMaterial ? FMaterialAttributeDefinitionMap::GetConnectedMaterialAttributesBitmask(BaseMaterial->GetExpressions()) : 0,
-						EditorOnlyData->ShadingModelFromMaterialExpression.IsConnected(),
-						EditorOnlyData->MaterialAttributes.IsConnected(MP_EmissiveColor));
-				}
-				else if ((Material->GetMaterialDomain() == MD_Surface || Material->GetMaterialDomain() == MD_RuntimeVirtualTexture || Material->GetMaterialDomain() == MD_DeferredDecal) && EditorOnlyData)
-				{
-					FSubstrateOperator* Op = UMaterialExpressionSubstrateShadingModels::SubstrateGenerateMaterialTopologyTreeCommon(
-						this, RootNodeConversionGUID, nullptr /*Parent*/, 0,
-						EditorOnlyData->EmissiveColor,
-						EditorOnlyData->Anisotropy,
-						*ClearCoatNormal,
-						*CustomTangent,
-						EditorOnlyData->ShadingModelFromMaterialExpression);
-
-					if (Material->GetMaterialDomain() == MD_DeferredDecal)
-					{
-						// We notify that this is for decal using this usage flag.
-						// Otherwise, we do not need to do anything else, the ConvertToDecal node is only here to enforce parameter blending.
-						// This in order to get a signle BSDF as input to convert as a decal.
-						Op->SubUsage |= SUBSTRATE_OPERATOR_SUBUSAGE_DECAL;
-					}
-				}
-				else if (Material->GetMaterialDomain() == MD_Volume && EditorOnlyData)
-				{
-					UMaterialExpressionSubstrateVolumetricFogCloudBSDF::SubstrateGenerateMaterialTopologyTreeCommon(
-						this, RootNodeConversionGUID, nullptr /*Parent*/, 0,
-						EditorOnlyData->EmissiveColor,
-						EditorOnlyData->AmbientOcclusion);
-				}
-				else if (Material->GetMaterialDomain() == MD_LightFunction && EditorOnlyData)
-				{
-					UMaterialExpressionSubstrateLightFunction::SubstrateGenerateMaterialTopologyTreeCommon(
-						this, RootNodeConversionGUID, nullptr /*Parent*/, 0);
-				}
-				else if (Material->GetMaterialDomain() == MD_PostProcess && EditorOnlyData)
-				{
-					UMaterialExpressionSubstratePostProcess::SubstrateGenerateMaterialTopologyTreeCommon(
-						this, RootNodeConversionGUID, nullptr /*Parent*/, 0);
-				}
-				else if (Material->GetMaterialDomain() == MD_UI && EditorOnlyData)
-				{
-					UMaterialExpressionSubstrateUI::SubstrateGenerateMaterialTopologyTreeCommon(
-						this, RootNodeConversionGUID, nullptr /*Parent*/, 0);
-				}
-
-				SubstrateThicknessStackPop();
 
-				check(SubstrateCtx.SubstrateThicknessStack.Num() == 0);
-
-				if (SubstrateCtx.bSubstrateTreeOutOfStackDepthOccurred)
-				{
-					Errorf(TEXT(" %s [%s]: Substrate - Cyclic graph detected when we only support acyclic graph."), *Material->GetDebugName(), *Material->GetAssetPath().ToString());
-					return;
-				}
-				if (!SubstrateCtx.SubstrateGenerateDerivedMaterialOperatorData(this))
-				{
-					Errorf(TEXT("Substrate material errors encountered."));
-					return;
-				}
-			}
+	// Get shading models from compilation (or material).
+	FMaterialShadingModelField MaterialShadingModels = GetCompiledShadingModels();
 
-		}
-		CurrentSubstrateCompilationContext = ESubstrateCompilationContext::SCC_Default;
-	}
-	else if (bSubstrateEnabled && FrontMaterialExpr)
 	{
 		// Temp code chunk scope (e.g.needed for the creation of static booleans from static switch parameter node, see UMaterialExpressionStaticSwitch::GetEffectiveInput).
 		TArray<FShaderCodeChunk> TempChunks;
 		AssignTempScope(TempChunks);
 
-	#if DEBUG_SUBSTRATE_TREE_STACK
-		UE_LOG(LogMaterial, Display, TEXT(" SubstrateTreeStack: SubstrateGenerateMaterialTopologyTree"));
-	#endif
-		for (uint32 SubstrateCompilationContextIndex = 0; SubstrateCompilationContextIndex < ESubstrateCompilationContext::SCC_MAX; ++SubstrateCompilationContextIndex)
-		{
-			// This is needed because SubstrateGenerateMaterialTopologyTree will call some compiler context though material expressions.
-			CurrentSubstrateCompilationContext = ESubstrateCompilationContext(SubstrateCompilationContextIndex);
-
-			FSubstrateCompilationContext& SubstrateCtx = SubstrateCompilationContext[SubstrateCompilationContextIndex];
-			SubstrateCtx = FSubstrateCompilationContext(ESubstrateCompilationContext(SubstrateCompilationContextIndex));
-
-			FScalarMaterialInput* SurfaceThickness = Material->IsThinSurface() && Material->GetMaterialInterface() ? &Material->GetMaterialInterface()->GetMaterial()->GetEditorOnlyData()->SurfaceThickness : nullptr;
-			SubstrateThicknessStackPush(nullptr, SurfaceThickness);
-			FrontMaterialExpr->SubstrateGenerateMaterialTopologyTree(this, nullptr, FrontMaterialOutputIndex);
-			SubstrateThicknessStackPop();
-
-			check(SubstrateCtx.SubstrateThicknessStack.Num() == 0);
+		bool bDoClearFunctionStack = false;
+		SubstrateTranslatorData.ProcessSubstrateTreeAndTopology(MaterialShadingModels, Material->GetMaterialGraphNodePreviewExpression(), bDoClearFunctionStack);
 
-			if (SubstrateCtx.bSubstrateTreeOutOfStackDepthOccurred)
-			{
-				Errorf(TEXT(" %s [%s]: Substrate - Cyclic graph detected when we only support acyclic graph."), *Material->GetDebugName(), *Material->GetAssetPath().ToString());
-				return;
-			}
-			if (!SubstrateCtx.SubstrateGenerateDerivedMaterialOperatorData(this))
-			{
-				Errorf(TEXT("Substrate material errors encountered."));
-				return;
-			}
-
-			bSubstrateWritesEmissive |= SubstrateCtx.bSubstrateWritesEmissive;
-			bSubstrateWritesAmbientOcclusion |= SubstrateCtx.bSubstrateWritesAmbientOcclusion;
-		}
-		CurrentSubstrateCompilationContext = ESubstrateCompilationContext::SCC_Default;
-
-		//Clear the stacks as they may be referencing scoped chunks that are lost before the compile step.
-		//Skips the Vertex shader because this is handled whilst resetting custom vertex interpolation below.
-		for (int iFrequency = EShaderFrequency::SF_Vertex + 1; iFrequency < EShaderFrequency::SF_NumFrequencies; iFrequency++)
+		if (bDoClearFunctionStack)
 		{
-			while (FunctionStacks[iFrequency].Num() > 1)
+			//Clear the stacks as they may be referencing scoped chunks that are lost before the compile step.
+			//Skips the Vertex shader because this is handled whilst resetting custom vertex interpolation below.
+			for (int iFrequency = EShaderFrequency::SF_Vertex + 1; iFrequency < EShaderFrequency::SF_NumFrequencies; iFrequency++)
 			{
-				FMaterialFunctionCompileState* Stack = FunctionStacks[iFrequency].Pop(EAllowShrinking::No);
-				delete Stack;
+				while (FunctionStacks[iFrequency].Num() > 1)
+				{
+					FMaterialFunctionCompileState* Stack = FunctionStacks[iFrequency].Pop(EAllowShrinking::No);
+					delete Stack;
+				}
+				FunctionStacks[iFrequency][0]->Reset();
 			}
-			FunctionStacks[iFrequency][0]->Reset();
 		}
 	}
 
@@ -1886,27 +1504,23 @@ void FHLSLMaterialTranslator::TranslateMaterial()
 	Chunk[MP_OpacityMask]		= Material->CompilePropertyAndSetMaterialProperty(MP_OpacityMask,	this);
 	Chunk[MP_WorldPositionOffset]= Material->CompilePropertyAndSetMaterialProperty(MP_WorldPositionOffset, this);
 
-
-#if DEBUG_SUBSTRATE_TREE_STACK
-	UE_LOG(LogMaterial, Display, TEXT(" SubstrateTreeStack: Material->CompilePropertyAndSetMaterialProperty(MP_FrontMaterial)"));
-#endif
 	Chunk[MP_SurfaceThickness] = Material->CompilePropertyAndSetMaterialProperty(MP_SurfaceThickness		,this);
 
 	// This is what is run when hidden conversion is not executed
-	if (!bMaterialUsesRootNodeToSubstrateHiddenConversion)
+	if (!SubstrateTranslatorData.bMaterialUsesRootNodeToSubstrateHiddenConversion)
 	{
 		Chunk[MP_FrontMaterial] = Material->CompilePropertyAndSetMaterialProperty(MP_FrontMaterial, this);
 
 		CHECK_DDC_QUERY_FINISHED_ELSE_RETURN();
 
 		// Now generate the code for the fully simplified MP_FrontMaterial right after MP_FrontMaterial
-		FullySimplifiedSubstrateFrontMaterialCodeChunk = SubstrateCreateAndRegisterNullMaterial();
+		SubstrateTranslatorData.FullySimplifiedSubstrateFrontMaterialCodeChunk = SubstrateCreateAndRegisterNullMaterial();
 		{																							// This causes issues, material look different
 			// Then generate the code
 			FullySimplifiedFrontMaterialCodeChunkStart = SharedPropertyCodeChunks[FrontMaterialShaderFrequency].Num();
-			CurrentSubstrateCompilationContext = ESubstrateCompilationContext::SCC_FullySimplified;
-			FullySimplifiedSubstrateFrontMaterialCodeChunk = Material->CompilePropertyAndSetMaterialProperty(MP_FrontMaterial, this);
-			CurrentSubstrateCompilationContext = ESubstrateCompilationContext::SCC_Default;
+			SubstrateTranslatorData.CurrentSubstrateCompilationContext = ESubstrateCompilationContext::SCC_FullySimplified;
+			SubstrateTranslatorData.FullySimplifiedSubstrateFrontMaterialCodeChunk = Material->CompilePropertyAndSetMaterialProperty(MP_FrontMaterial, this);
+			SubstrateTranslatorData.CurrentSubstrateCompilationContext = ESubstrateCompilationContext::SCC_Default;
 			FullySimplifiedFrontMaterialCodeChunkEnd = SharedPropertyCodeChunks[FrontMaterialShaderFrequency].Num();
 		}
 	}
@@ -1926,9 +1540,9 @@ void FHLSLMaterialTranslator::TranslateMaterial()
 					
 					return UMaterialExpressionSubstrateConvertMaterialAttributes::CompileCommon(this, 0,
 						FMaterialAttributeDefinitionMap::GetConnectedMaterialAttributesBitmask(BaseMaterial->GetExpressions()), EditorOnlyData->MaterialAttributes, MaterialInterface->GetShadingModels().GetFirstShadingModel(),
-						*WaterScatteringCoefficients, *WaterAbsorptionCoefficients, *WaterPhaseG, *ColorScaleBehindWater,
+						*SubstrateTranslatorData.GetWaterScatteringCoefficients(), *SubstrateTranslatorData.GetWaterAbsorptionCoefficients(), *SubstrateTranslatorData.GetWaterPhaseG(), *SubstrateTranslatorData.GetColorScaleBehindWater(),
 						bHasSSS, SSSProfile,
-						*ClearCoatNormal, *CustomTangent);
+						*SubstrateTranslatorData.GetClearCoatNormal(), *SubstrateTranslatorData.GetCustomTangent());
 
 				}
 				else if ((Material->GetMaterialDomain() == MD_Surface || Material->GetMaterialDomain() == MD_RuntimeVirtualTexture || Material->GetMaterialDomain() == MD_DeferredDecal) && EditorOnlyData)
@@ -1941,11 +1555,11 @@ void FHLSLMaterialTranslator::TranslateMaterial()
 						EditorOnlyData->BaseColor, EditorOnlyData->Specular, EditorOnlyData->Metallic, EditorOnlyData->Roughness, EditorOnlyData->EmissiveColor,
 						EditorOnlyData->Opacity, EditorOnlyData->SubsurfaceColor, EditorOnlyData->ClearCoat, EditorOnlyData->ClearCoatRoughness,
 						EditorOnlyData->ShadingModelFromMaterialExpression, MaterialInterface->GetShadingModels().GetFirstShadingModel(),
-						*ThinTranslucentTransmittanceColor, *ThinTranslucentSurfaceCoverage,
-						*WaterScatteringCoefficients, *WaterAbsorptionCoefficients, *WaterPhaseG, *ColorScaleBehindWater,
+						*SubstrateTranslatorData.GetThinTranslucentTransmittanceColor(), *SubstrateTranslatorData.GetThinTranslucentSurfaceCoverage(),
+						*SubstrateTranslatorData.GetWaterScatteringCoefficients(), *SubstrateTranslatorData.GetWaterAbsorptionCoefficients(), *SubstrateTranslatorData.GetWaterPhaseG(), *SubstrateTranslatorData.GetColorScaleBehindWater(),
 						bHasAnisotropy, EditorOnlyData->Anisotropy,
 						EditorOnlyData->Normal, EditorOnlyData->Tangent,
-						*ClearCoatNormal, *CustomTangent,
+						*SubstrateTranslatorData.GetClearCoatNormal(), *SubstrateTranslatorData.GetCustomTangent(),
 						bHasSSS, SSSProfile,
 						EditorOnlyData);
 				}
@@ -1982,17 +1596,14 @@ void FHLSLMaterialTranslator::TranslateMaterial()
 		Chunk[MP_FrontMaterial] = CompileFrontMaterial();
 
 		FullySimplifiedFrontMaterialCodeChunkStart = SharedPropertyCodeChunks[FrontMaterialShaderFrequency].Num();
-		CurrentSubstrateCompilationContext = ESubstrateCompilationContext::SCC_FullySimplified;
-		FullySimplifiedSubstrateFrontMaterialCodeChunk = CompileFrontMaterial();
-		CurrentSubstrateCompilationContext = ESubstrateCompilationContext::SCC_Default;
+		SubstrateTranslatorData.CurrentSubstrateCompilationContext = ESubstrateCompilationContext::SCC_FullySimplified;
+		SubstrateTranslatorData.FullySimplifiedSubstrateFrontMaterialCodeChunk = CompileFrontMaterial();
+		SubstrateTranslatorData.CurrentSubstrateCompilationContext = ESubstrateCompilationContext::SCC_Default;
 		FullySimplifiedFrontMaterialCodeChunkEnd = SharedPropertyCodeChunks[FrontMaterialShaderFrequency].Num();
 	}
 
 	CHECK_DDC_QUERY_FINISHED_ELSE_RETURN();
 
-	// Get shading models from compilation (or material).
-	FMaterialShadingModelField MaterialShadingModels = GetCompiledShadingModels();
-
 	ValidateShadingModelsForFeatureLevel(MaterialShadingModels);
 
 	CHECK_DDC_QUERY_FINISHED_ELSE_RETURN();
@@ -2084,7 +1695,7 @@ void FHLSLMaterialTranslator::TranslateMaterial()
 	// At this point we mark the code chunk generation complete.
 	bAllowCodeChunkGeneration = false;
 
-	bUsesEmissiveColor = bSubstrateWritesEmissive || IsMaterialPropertyUsed(MP_EmissiveColor, Chunk[MP_EmissiveColor], FLinearColor(0, 0, 0, 0), 3);
+	bUsesEmissiveColor = SubstrateTranslatorData.bSubstrateWritesEmissive || IsMaterialPropertyUsed(MP_EmissiveColor, Chunk[MP_EmissiveColor], FLinearColor(0, 0, 0, 0), 3);
 	bUsesPixelDepthOffset = (AllowPixelDepthOffset(Platform) && Material->AllowPixelDepthOffset() && IsMaterialPropertyUsed(MP_PixelDepthOffset, Chunk[MP_PixelDepthOffset], FLinearColor(0, 0, 0, 0), 1));
 
 	bool bUsesWorldPositionOffsetCurrent = IsMaterialPropertyUsed(MP_WorldPositionOffset, Chunk[MP_WorldPositionOffset], FLinearColor(0, 0, 0, 0), 3);
@@ -2117,7 +1728,7 @@ void FHLSLMaterialTranslator::TranslateMaterial()
 
 	bUsesAnisotropy = IsMaterialPropertyUsed(MP_Anisotropy, Chunk[MP_Anisotropy], FLinearColor(0, 0, 0, 0), 1);
 	MaterialCompilationOutput.bUsesAnisotropy = bUsesAnisotropy;
-	if (bSubstrateMaterialUsesLegacyMaterialCompilation && bUsesAnisotropy)
+	if (SubstrateTranslatorData.bSubstrateMaterialUsesLegacyMaterialCompilation && bUsesAnisotropy)
 	{
 		MaterialCompilationOutput.SubstrateMaterialCompilationOutput.SubstrateMaterialType = SUBSTRATE_MATERIAL_TYPE_COMPLEX;
 		MaterialCompilationOutput.SubstrateMaterialCompilationOutput.SubstrateMaterialBsdfFeatures |= ESubstrateBsdfFeature::Anisotropy;
@@ -2195,7 +1806,7 @@ void FHLSLMaterialTranslator::TranslateMaterial()
 		MaterialCompilationOutput,
 		MaterialShadingModels,
 		BlendMode,
-		FrontMaterialExpr != nullptr,
+		SubstrateTranslatorData.FrontMaterialExpr != nullptr,
 		Platform,
 		ValidationErrors);
 
@@ -2277,7 +1888,7 @@ void FHLSLMaterialTranslator::TranslateMaterial()
 				(PropertyId == MP_WorldPositionOffset || PropertyId == CompiledMP_PrevWorldPositionOffset) : 
 				(PropertyId == MP_WorldPositionOffset || PropertyId == CompiledMP_PrevWorldPositionOffset || PropertyId == MP_Displacement);
 
-			if (bSubstrateEnabled && PropertyId >= MP_FrontMaterial && PropertyShaderFrequency == FrontMaterialShaderFrequency)
+			if (SubstrateTranslatorData.bSubstrateEnabled && PropertyId >= MP_FrontMaterial && PropertyShaderFrequency == FrontMaterialShaderFrequency)
 			{
 				int32 PropertyIdCodeChunk = Chunk[PropertyId];
 				bool bPropertyIsConstant = false;
@@ -2350,7 +1961,7 @@ void FHLSLMaterialTranslator::TranslateMaterial()
 
 		// The code chunk corresponding to FullySimplifiedSubstrateFrontMaterialCodeChunk have already been written as part of MP_FrontMaterial.
 		// Here we get the FullySimplifiedSubstrateFrontMaterialTranslatedCodeChunks representing the variable storing the final fully simplified SubstrateData.
-		if (bSubstrateEnabled)
+		if (SubstrateTranslatorData.bSubstrateEnabled)
 		{
 			uint32 PropertyId = MP_FrontMaterial;
 
@@ -2375,10 +1986,10 @@ void FHLSLMaterialTranslator::TranslateMaterial()
 			GetFixedParameterCode(
 				StartChunk,
 				SharedPropertyCodeChunks[PropertyShaderFrequency].Num(),
-				FullySimplifiedSubstrateFrontMaterialCodeChunk,								//Chunk[PropertyId],
+				SubstrateTranslatorData.FullySimplifiedSubstrateFrontMaterialCodeChunk,								//Chunk[PropertyId],
 				SharedPropertyCodeChunks[PropertyShaderFrequency],
-				FullySimplifiedSubstrateFrontMaterialTranslatedCodeChunkDefinitions,
-				FullySimplifiedSubstrateFrontMaterialTranslatedCodeChunks,
+				SubstrateTranslatorData.FullySimplifiedSubstrateFrontMaterialTranslatedCodeChunkDefinitions,
+				SubstrateTranslatorData.FullySimplifiedSubstrateFrontMaterialTranslatedCodeChunks,
 				Variation,
 				bReduceAfterReturnValue);
 		}
@@ -2506,303 +2117,21 @@ void FHLSLMaterialTranslator::TranslateMaterial()
 		TextureParameter.ShaderFrequencyMask = TextureExpression->ShaderFrequencyMask;
 	}
 
-	if (bSubstrateEnabled)
+	// Write substrate tree specific code.
+	if (SubstrateTranslatorData.bSubstrateEnabled)
 	{
-		for (uint32 SubstrateCompilationContextIndex = 0; SubstrateCompilationContextIndex < ESubstrateCompilationContext::SCC_MAX; ++SubstrateCompilationContextIndex)
-		{
-			CHECK_DDC_QUERY_FINISHED_ELSE_RETURN();
-			FSubstrateCompilationContext& SubstrateCtx = SubstrateCompilationContext[SubstrateCompilationContextIndex];
-
-			FString TreeFunctionPostFix = "ERROR";
-			switch (SubstrateCompilationContextIndex)
-			{
-			case ESubstrateCompilationContext::SCC_Default:
-				TreeFunctionPostFix = TEXT("");
-				break;
-			case ESubstrateCompilationContext::SCC_FullySimplified:
-				TreeFunctionPostFix = TEXT("_FullySimplified");
-				break;
-			default:
-				check(false);
-			}
-
-			const bool bSubstrateFrontMaterialProvided = Chunk[MP_FrontMaterial] != INDEX_NONE;
-			bool bSubstrateFrontMaterialIsValid = bSubstrateFrontMaterialProvided;
-			if (bSubstrateFrontMaterialIsValid)
-			{
-				// The material can be null when some entries are automatically generated, for instance in the material layer blending system
-				bSubstrateFrontMaterialIsValid &= SubstrateCtx.SubstrateMaterialEffectiveClosureCount > 0;
-
-				UMaterialInterface* MaterialInterface = Material->GetMaterialInterface();
-				UMaterial* BaseMaterial = MaterialInterface ? MaterialInterface->GetMaterial() : nullptr;
-				if (!bSubstrateFrontMaterialIsValid && BaseMaterial && BaseMaterial->bUseMaterialAttributes)
-				{
-					Errorf(TEXT(" %s [%s]: Substrate - Material has no BSDF: this can happen when the SubstrateConvertMaterialAttributes node is not used and when material attributes are enabled."), *Material->GetDebugName(), *Material->GetAssetPath().ToString());
-					TranslationResult = EHLSLMaterialTranslatorResult::Failure;
-					return;
-				}
-			}
-
-			if (bMaterialUsesRootNodeToSubstrateHiddenConversion)
-			{
-				// Hardcoded substrate tree node visit when using hidden conversion.
-				bMaterialIsSubstrate = true;
-
-				if (SubstrateCompilationContextIndex == ESubstrateCompilationContext::SCC_Default)
-				{
-					MaterialCompilationOutput.SubstrateMaterialCompilationOutput.SubstrateMaterialDescription = "";
-					ResourcesString += "// Substrate: HiddenMaterialAssetConversion\n";
-
-					// Adde default Substrate functions
-					ResourcesString += "#if TEMPLATE_USES_SUBSTRATE\n";
-					ResourcesString += "void FSubstratePixelHeader::PreUpdateAllBSDFWithBottomUpOperatorVisit(float3 V) {}\n";
-					ResourcesString += "void FSubstratePixelHeader::UpdateAllBSDFsOperatorCoverageTransmittance(FSubstrateIntegrationSettings Settings, float3 V)\n"
-										"{\n"
-										"	#if SUBSTRATE_COMPILER_SUPPORTS_STRUCT_FORWARD_DECLARATION\n"
-										"						SubstrateTree.UpdateSingleBSDFOperatorCoverageTransmittance(this, 0, Settings, V);\n"
-										"	#else\n"
-										"						UpdateSingleBSDFOperatorCoverageTransmittance(SubstrateTree, this, 0, Settings, V);\n"
-										"	#endif\n"
-										"}\n";
-					ResourcesString += "void FSubstratePixelHeader::UpdateAllOperatorsCoverageTransmittance() {}\n";
-					ResourcesString += "void FSubstratePixelHeader::UpdateAllBSDFWithBottomUpOperatorVisit() {}\n";
-					ResourcesString += "#endif // TEMPLATE_USES_SUBSTRATE\n";
-				}
-				else
-				{
-					ResourcesString += "#if TEMPLATE_USES_SUBSTRATE\n";
-					ResourcesString += "void FSubstratePixelHeader::PreUpdateAllBSDFWithBottomUpOperatorVisit_FullySimplified(float3 V) {}\n";
-					ResourcesString += "void FSubstratePixelHeader::UpdateAllBSDFsOperatorCoverageTransmittance_FullySimplified(FSubstrateIntegrationSettings Settings, float3 V)\n"
-										"{\n"
-										"	#if SUBSTRATE_COMPILER_SUPPORTS_STRUCT_FORWARD_DECLARATION\n"
-										"						SubstrateTree.UpdateSingleBSDFOperatorCoverageTransmittance(this, 0, Settings, V);\n"
-										"	#else\n"
-										"						UpdateSingleBSDFOperatorCoverageTransmittance(SubstrateTree, this, 0, Settings, V);\n"
-										"	#endif\n"
-										"}\n";
-					ResourcesString += "void FSubstratePixelHeader::UpdateAllOperatorsCoverageTransmittance_FullySimplified() {}\n";
-					ResourcesString += "void FSubstratePixelHeader::UpdateAllBSDFWithBottomUpOperatorVisit_FullySimplified() {}\n";
-					ResourcesString += "#endif // TEMPLATE_USES_SUBSTRATE\n";
-					}
-			}
-			else if (bSubstrateFrontMaterialIsValid)
-			{
-				bMaterialIsSubstrate = true;
-
-				if (SubstrateCtx.SubstrateMaterialRootOperator)
-				{
-					// Now implement the functions needed to process the material topology
-
-					// Pre-Update the slab BSDF with operators (like thin film coating, which can alter F0/F90)
-					{
-						// Update the coverage/transmittance of each node in the graph
-						check(SubstrateCtx.SubstrateMaterialRootOperator);
-						int32 RootMaximumDistanceToLeaves = SubstrateCtx.SubstrateMaterialRootOperator->MaxDistanceFromLeaves;
-
-						ResourcesString += FString::Printf(TEXT("void  FSubstratePixelHeader::PreUpdateAllBSDFWithBottomUpOperatorVisit%s(float3 V)\n"), *TreeFunctionPostFix);
-						ResourcesString += "{\n";
-						for (uint32 ClosureIndex = 0; ClosureIndex < SubstrateCtx.SubstrateMaterialEffectiveClosureCount; ++ClosureIndex)
-						{
-							ResourcesString += "\t{\n";
-							for (auto& It : SubstrateCtx.SubstrateMaterialExpressionRegisteredOperators)
-							{
-								if (!It.IsDiscarded() && It.BSDFIndex == ClosureIndex)
-								{
-									// Walk up the graph to the root node and apply weight factors
-									std::function<void(const FSubstrateOperator&, int32)> WalkOperatorsUp = [&](const FSubstrateOperator& CurrentOperator, int32 PreviousOperatorIndex) -> void
-										{
-											switch (CurrentOperator.OperatorType)
-											{
-											case SUBSTRATE_OPERATOR_WEIGHT:
-											{
-												break; // NOP
-											}
-											case SUBSTRATE_OPERATOR_HORIZONTAL:
-											{
-												break; // NOP
-											}
-											case SUBSTRATE_OPERATOR_VERTICAL:
-											{
-												// example ResourcesString += FString::Printf(TEXT("\t PreUpdateAllBSDFWithBottomUpOperatorVisit_Vertical(this, SubstrateTree, SubstrateTree.BSDFs[%d], V, %d /*Op index*/, %d /*PreviousIsInputA*/);\n"), BSDFIndex, CurrentOperator.Index, CurrentOperator.LeftIndex == PreviousOperatorIndex ? 1 : 0);
-												break; // NOP
-											}
-											case SUBSTRATE_OPERATOR_ADD:
-											{
-												break; // NOP
-											}
-											case SUBSTRATE_OPERATOR_SELECT:
-											{
-												break; // NOP
-											}
-											default:
-											case SUBSTRATE_OPERATOR_BSDF:
-											{
-												check(false);
-											}
-											}
-
-											if (CurrentOperator.ParentIndex != INDEX_NONE)
-											{
-												WalkOperatorsUp(SubstrateCtx.SubstrateMaterialExpressionRegisteredOperators[CurrentOperator.ParentIndex], CurrentOperator.Index);
-											}
-										};
-
-									const int32 BSDFOperatorIndex = It.Index;
-									const FSubstrateOperator& BSDFOperator = SubstrateCtx.SubstrateMaterialExpressionRegisteredOperators[BSDFOperatorIndex];
-
-									// Start visiting node up from the BSDF leaf only if it has a parent.
-									if (BSDFOperator.ParentIndex != INDEX_NONE)
-									{
-										WalkOperatorsUp(SubstrateCtx.SubstrateMaterialExpressionRegisteredOperators[BSDFOperator.ParentIndex], BSDFOperator.Index);
-									}
-								}
-							}
-							ResourcesString += "\t}\n";
-						}
-						ResourcesString += "}\n";
-					}
-
-					// Update the coverage/transmittance of each leaves (==BSDFs) of the Substrate tree.
-					{
-						ResourcesString += FString::Printf(TEXT("void FSubstratePixelHeader::UpdateAllBSDFsOperatorCoverageTransmittance%s(FSubstrateIntegrationSettings Settings, float3 V)\n"), *TreeFunctionPostFix);
-						ResourcesString += "{\n";
-						for (uint32 ClosureIndex = 0; ClosureIndex < SubstrateCtx.SubstrateMaterialEffectiveClosureCount; ++ClosureIndex)
-						{
-							ResourcesString += FString::Printf(TEXT("\t #if SUBSTRATE_COMPILER_SUPPORTS_STRUCT_FORWARD_DECLARATION\n"));
-							ResourcesString += FString::Printf(TEXT("\t SubstrateTree.UpdateSingleBSDFOperatorCoverageTransmittance(this, %d, Settings, V);\n"), ClosureIndex);
-							ResourcesString += FString::Printf(TEXT("\t #else\n"));
-							ResourcesString += FString::Printf(TEXT("\t UpdateSingleBSDFOperatorCoverageTransmittance(SubstrateTree, this, %d, Settings, V);\n"), ClosureIndex);
-							ResourcesString += FString::Printf(TEXT("\t #endif\n"));
-						}
-						ResourcesString += "}\n";
-					}
-
-					// Propagate up the coverage/transmittance of each node in the Substrate tree.
-					// For that we visit all the operator according to their distance from the Substrate tree leaves, from small to large.
-					{
-						check(SubstrateCtx.SubstrateMaterialRootOperator);
-						int32 RootMaximumDistanceToLeaves = SubstrateCtx.SubstrateMaterialRootOperator->MaxDistanceFromLeaves;
-
-						ResourcesString += FString::Printf(TEXT("void FSubstratePixelHeader::UpdateAllOperatorsCoverageTransmittance%s()\n"), *TreeFunctionPostFix);
-						ResourcesString += "{\n";
-						for (int32 DistanceToLeaves = 1; DistanceToLeaves <= RootMaximumDistanceToLeaves; ++DistanceToLeaves)
-						{
-							ResourcesString += FString::Printf(TEXT("\t// MaxDistanceFromLeaves = %d \n"), DistanceToLeaves);
-							for (auto& It : SubstrateCtx.SubstrateMaterialExpressionRegisteredOperators)
-							{
-								if (!It.IsDiscarded() && It.MaxDistanceFromLeaves == DistanceToLeaves)
-								{
-									ResourcesString += FString::Printf(TEXT("\t SubstrateTree.UpdateSingleOperatorCoverageTransmittance(%d /*operator index*/);\n"), It.Index);
-								}
-							}
-						}
-						ResourcesString += "}\n";
-					}
-
-					// Update the luminance weight of each BSDF according to the operators it has to traverse bottom-up to the Substrate tree root node.
-					{
-						// Update the coverage/transmittance of each node in the graph
-						check(SubstrateCtx.SubstrateMaterialRootOperator);
-						int32 RootMaximumDistanceToLeaves = SubstrateCtx.SubstrateMaterialRootOperator->MaxDistanceFromLeaves;
-
-						ResourcesString += FString::Printf(TEXT("void FSubstratePixelHeader::UpdateAllBSDFWithBottomUpOperatorVisit%s()\n"), *TreeFunctionPostFix);
-						ResourcesString += "{\n";
-						for (uint32 ClosureIndex = 0; ClosureIndex < SubstrateCtx.SubstrateMaterialEffectiveClosureCount; ++ClosureIndex)
-						{
-							for (auto& It : SubstrateCtx.SubstrateMaterialExpressionRegisteredOperators)
-							{
-								if (!It.IsDiscarded() && It.BSDFIndex == ClosureIndex)
-								{
-									// Walk up the graph to the root node and apply weight factors
-									std::function<void(const FSubstrateOperator&, int32)> WalkOperatorsUp = [&](const FSubstrateOperator& CurrentOperator, int32 PreviousOperatorIndex) -> void
-										{
-											switch (CurrentOperator.OperatorType)
-											{
-											case SUBSTRATE_OPERATOR_WEIGHT:
-											{
-												ResourcesString += FString::Printf(TEXT("\t SubstrateTree.UpdateAllBSDFWithBottomUpOperatorVisit_Weight(%d /*BSDFIndex*/, %d /*Op index*/, %d /*PreviousIsInputA*/);\n"), ClosureIndex, CurrentOperator.Index, 1);
-												break;
-											}
-											case SUBSTRATE_OPERATOR_HORIZONTAL:
-											{
-												ResourcesString += FString::Printf(TEXT("\t SubstrateTree.UpdateAllBSDFWithBottomUpOperatorVisit_Horizontal(%d /*BSDFIndex*/, %d /*Op index*/, %d /*PreviousIsInputA*/);\n"), ClosureIndex, CurrentOperator.Index, CurrentOperator.LeftIndex == PreviousOperatorIndex ? 1 : 0);
-												break;
-											}
-											case SUBSTRATE_OPERATOR_VERTICAL:
-											{
-												ResourcesString += FString::Printf(TEXT("\t SubstrateTree.UpdateAllBSDFWithBottomUpOperatorVisit_Vertical(%d /*BSDFIndex*/, %d /*Op index*/, %d /*PreviousIsInputA*/);\n"), ClosureIndex, CurrentOperator.Index, CurrentOperator.LeftIndex == PreviousOperatorIndex ? 1 : 0);
-												break;
-											}
-											case SUBSTRATE_OPERATOR_ADD:
-											case SUBSTRATE_OPERATOR_SELECT:
-											{
-												break; // NOP
-											}
-											default:
-											case SUBSTRATE_OPERATOR_BSDF:
-											{
-												check(false);
-											}
-											}
-
-											if (CurrentOperator.ParentIndex != INDEX_NONE)
-											{
-												WalkOperatorsUp(SubstrateCtx.SubstrateMaterialExpressionRegisteredOperators[CurrentOperator.ParentIndex], CurrentOperator.Index);
-											}
-										};
-
-									const int32 BSDFOperatorIndex = It.Index;
-									const FSubstrateOperator& BSDFOperator = SubstrateCtx.SubstrateMaterialExpressionRegisteredOperators[BSDFOperatorIndex];
-
-									// Start visiting node up from the BSDF leaf only if it has a parent.
-									if (BSDFOperator.ParentIndex != INDEX_NONE)
-									{
-										WalkOperatorsUp(SubstrateCtx.SubstrateMaterialExpressionRegisteredOperators[BSDFOperator.ParentIndex], BSDFOperator.Index);
-									}
-								}
-							}
-						}
-						ResourcesString += "}\n";
-					}
-				}
+		CHECK_DDC_QUERY_FINISHED_ELSE_RETURN();
 
-				// Check if normal/tangent basis are valid
-				{
-					SubstrateCtx.FinalUsedSharedLocalBasesCount = 0;
-					uint8 RequestedSharedLocalBasesCount = 0;
-					SubstrateCtx.SubstrateEvaluateSharedLocalBases(this, RequestedSharedLocalBasesCount, nullptr);
-					if (RequestedSharedLocalBasesCount > SUBSTRATE_MAX_SHAREDLOCALBASES_REGISTERS)
-					{
-						Errorf(TEXT(" %s [%s]: Substrate - Material has more unique normal/tangent basis than the allowed limit %d/%d."), *Material->GetDebugName(), *Material->GetAssetPath().ToString(), RequestedSharedLocalBasesCount, SUBSTRATE_MAX_SHAREDLOCALBASES_REGISTERS);
-					}
-				}
-			}
-			else
-			{
-				if (SubstrateCompilationContextIndex == ESubstrateCompilationContext::SCC_Default)
-				{
-					MaterialCompilationOutput.SubstrateMaterialCompilationOutput.SubstrateMaterialDescription = "";
-					ResourcesString += "// No Substrate material provided\n";
-
-					// Adde default Substrate functions
-					ResourcesString += "#if TEMPLATE_USES_SUBSTRATE\n";
-					ResourcesString += "void FSubstratePixelHeader::PreUpdateAllBSDFWithBottomUpOperatorVisit(float3 V) {}\n";
-					ResourcesString += "void FSubstratePixelHeader::UpdateAllBSDFsOperatorCoverageTransmittance(FSubstrateIntegrationSettings Settings, float3 V) {}\n";
-					ResourcesString += "void FSubstratePixelHeader::UpdateAllOperatorsCoverageTransmittance() {}\n";
-					ResourcesString += "void FSubstratePixelHeader::UpdateAllBSDFWithBottomUpOperatorVisit() {}\n";
-					ResourcesString += "#endif\n";
-				}
-				else
-				{
-					ResourcesString += "#if TEMPLATE_USES_SUBSTRATE\n";
-					ResourcesString += "void FSubstratePixelHeader::PreUpdateAllBSDFWithBottomUpOperatorVisit_FullySimplified(float3 V) {}\n";
-					ResourcesString += "void FSubstratePixelHeader::UpdateAllBSDFsOperatorCoverageTransmittance_FullySimplified(FSubstrateIntegrationSettings Settings, float3 V) {}\n";
-					ResourcesString += "void FSubstratePixelHeader::UpdateAllOperatorsCoverageTransmittance_FullySimplified() {}\n";
-					ResourcesString += "void FSubstratePixelHeader::UpdateAllBSDFWithBottomUpOperatorVisit_FullySimplified() {}\n";
-					ResourcesString += "#endif\n";
-				}
-			}
+		const bool bSubstrateFrontMaterialProvided = Chunk[MP_FrontMaterial] != INDEX_NONE;
+		const bool bValid = SubstrateTranslatorData.GenerateSubstrateTreeHLSLCode(bSubstrateFrontMaterialProvided);
+		if(!bValid)
+		{
+			TranslationResult = EHLSLMaterialTranslatorResult::Failure;
+			return;
 		}
+
+		// We conveniently append that code to the resources string.
+		ResourcesString += SubstrateTranslatorData.SubstrateTreeHLSL; // 
 	}
 
 	MaterialCompilationOutput.UniformExpressionSet.SetParameterCollections(ParameterCollections);
@@ -2976,7 +2305,6 @@ void FHLSLMaterialTranslator::GetMaterialEnvironment(EShaderPlatform InPlatform,
 	OutEnvironment.SetDefine(TEXT("USES_EMISSIVE_COLOR"), EnvironmentDefines->bUsesEmissiveColor);
 	// Distortion uses tangent space transform 
 	OutEnvironment.SetDefine(TEXT("USES_DISTORTION"), EnvironmentDefines->bUsesDistortion);
-	OutEnvironment.SetDefine(TEXT("DISTORTION_ACCOUNT_FOR_COVERAGE"), EnvironmentDefines->bDistortionAccountForCoverage);
 
 	OutEnvironment.SetDefine(TEXT("MATERIAL_ENABLE_TRANSLUCENCY_FOGGING"), EnvironmentDefines->bMaterialEnableTranslucencyFogging);
 	OutEnvironment.SetDefine(TEXT("MATERIAL_ENABLE_TRANSLUCENCY_CLOUD_FOGGING"), EnvironmentDefines->bMaterialEnableTranslucencyCloudFogging);
@@ -3091,17 +2419,17 @@ void FHLSLMaterialTranslator::GetMaterialEnvironment(EShaderPlatform InPlatform,
 			OutEnvironment.SetDefine(TEXT("MATERIAL_SHADINGMODEL_DEFAULT_LIT"), TEXT("1"));
 		}
 
-		if (EnvironmentDefines->HasShadingModel(MSM_Subsurface) || EnumHasAnyFlags(EnvironmentDefines->SubstrateMaterialBsdfFeatures, ESubstrateBsdfFeature::SSS))
+		if (EnvironmentDefines->HasShadingModel(MSM_Subsurface) || EnumHasAnyFlags(SubstrateTranslatorData.SubstrateMaterialBsdfFeatures, ESubstrateBsdfFeature::SSS))
 		{
 			OutEnvironment.SetDefine(TEXT("MATERIAL_SHADINGMODEL_SUBSURFACE"), TEXT("1"));
 		}
 
-		if (EnvironmentDefines->HasShadingModel(MSM_PreintegratedSkin) || EnumHasAnyFlags(EnvironmentDefines->SubstrateMaterialBsdfFeatures, ESubstrateBsdfFeature::SSS))
+		if (EnvironmentDefines->HasShadingModel(MSM_PreintegratedSkin) || EnumHasAnyFlags(SubstrateTranslatorData.SubstrateMaterialBsdfFeatures, ESubstrateBsdfFeature::SSS))
 		{
 			OutEnvironment.SetDefine(TEXT("MATERIAL_SHADINGMODEL_PREINTEGRATED_SKIN"), TEXT("1"));
 		}
 
-		if (EnvironmentDefines->HasShadingModel(MSM_SubsurfaceProfile) || EnumHasAnyFlags(EnvironmentDefines->SubstrateMaterialBsdfFeatures, ESubstrateBsdfFeature::SSS))
+		if (EnvironmentDefines->HasShadingModel(MSM_SubsurfaceProfile) || EnumHasAnyFlags(SubstrateTranslatorData.SubstrateMaterialBsdfFeatures, ESubstrateBsdfFeature::SSS))
 		{
 			OutEnvironment.SetDefine(TEXT("MATERIAL_SHADINGMODEL_SUBSURFACE_PROFILE"), TEXT("1"));
 		}
@@ -3111,32 +2439,32 @@ void FHLSLMaterialTranslator::GetMaterialEnvironment(EShaderPlatform InPlatform,
 			OutEnvironment.SetDefine(TEXT("MATERIAL_SUBSURFACE_PROFILE_USE_CURVATURE"), TEXT("1"));
 		}
 
-		if (EnvironmentDefines->HasShadingModel(MSM_ClearCoat) || EnumHasAnyFlags(EnvironmentDefines->SubstrateMaterialBsdfFeatures, ESubstrateBsdfFeature::SecondRoughnessOrSimpleClearCoat))
+		if (EnvironmentDefines->HasShadingModel(MSM_ClearCoat) || EnumHasAnyFlags(SubstrateTranslatorData.SubstrateMaterialBsdfFeatures, ESubstrateBsdfFeature::SecondRoughnessOrSimpleClearCoat))
 		{
 			OutEnvironment.SetDefine(TEXT("MATERIAL_SHADINGMODEL_CLEAR_COAT"), TEXT("1"));
 		}
 
-		if (EnvironmentDefines->HasShadingModel(MSM_TwoSidedFoliage) || EnumHasAnyFlags(EnvironmentDefines->SubstrateMaterialBsdfFeatures, ESubstrateBsdfFeature::SSS))
+		if (EnvironmentDefines->HasShadingModel(MSM_TwoSidedFoliage) || EnumHasAnyFlags(SubstrateTranslatorData.SubstrateMaterialBsdfFeatures, ESubstrateBsdfFeature::SSS))
 		{
 			OutEnvironment.SetDefine(TEXT("MATERIAL_SHADINGMODEL_TWOSIDED_FOLIAGE"), TEXT("1"));
 		}
 
-		if (EnvironmentDefines->HasShadingModel(MSM_Hair) || EnumHasAnyFlags(EnvironmentDefines->SubstrateMaterialBsdfFeatures, ESubstrateBsdfFeature::Hair))
+		if (EnvironmentDefines->HasShadingModel(MSM_Hair) || EnumHasAnyFlags(SubstrateTranslatorData.SubstrateMaterialBsdfFeatures, ESubstrateBsdfFeature::Hair))
 		{
 			OutEnvironment.SetDefine(TEXT("MATERIAL_SHADINGMODEL_HAIR"), TEXT("1"));
 		}
 
-		if (EnvironmentDefines->HasShadingModel(MSM_Cloth) || EnumHasAnyFlags(EnvironmentDefines->SubstrateMaterialBsdfFeatures, ESubstrateBsdfFeature::Fuzz))
+		if (EnvironmentDefines->HasShadingModel(MSM_Cloth) || EnumHasAnyFlags(SubstrateTranslatorData.SubstrateMaterialBsdfFeatures, ESubstrateBsdfFeature::Fuzz))
 		{
 			OutEnvironment.SetDefine(TEXT("MATERIAL_SHADINGMODEL_CLOTH"), TEXT("1"));
 		}
 
-		if (EnvironmentDefines->HasShadingModel(MSM_Eye) || EnumHasAnyFlags(EnvironmentDefines->SubstrateMaterialBsdfFeatures, ESubstrateBsdfFeature::Eye))
+		if (EnvironmentDefines->HasShadingModel(MSM_Eye) || EnumHasAnyFlags(SubstrateTranslatorData.SubstrateMaterialBsdfFeatures, ESubstrateBsdfFeature::Eye))
 		{
 			OutEnvironment.SetDefine(TEXT("MATERIAL_SHADINGMODEL_EYE"), TEXT("1"));
 		}
 
-		if (EnumHasAnyFlags(EnvironmentDefines->SubstrateMaterialBsdfFeatures, ESubstrateBsdfFeature::Toon))
+		if (EnumHasAnyFlags(SubstrateTranslatorData.SubstrateMaterialBsdfFeatures, ESubstrateBsdfFeature::Toon))
 		{
 			OutEnvironment.SetDefine(TEXT("MATERIAL_SHADINGMODEL_TOON"), TEXT("1"));
 		}
@@ -3209,30 +2537,10 @@ void FHLSLMaterialTranslator::GetMaterialEnvironment(EShaderPlatform InPlatform,
 		OutEnvironment.SetDefine(TEXT("MATERIAL_VOLUMETRIC_CLOUD_EMPTY_SPACE_SKIPPING_OUTPUT"), TEXT("1"));
 	}
 
-	OutEnvironment.SetDefine(TEXT("MATERIAL_IS_SUBSTRATE"), EnvironmentDefines->bMaterialIsSubstrate);
-	OutEnvironment.SetDefine(TEXT("MATERIAL_IS_SUBSTRATEHIDDENCONVERSION"), EnvironmentDefines->bMaterialUsesRootNodeToSubstrateHiddenConversion);
 	OutEnvironment.SetDefine(TEXT("DUAL_SOURCE_COLOR_BLENDING_ENABLED"), EnvironmentDefines->bDualSourceColorBlendingEnabled);
-	OutEnvironment.SetDefine(TEXT("SUBSTRATE_PREMULTIPLIED_ALPHA_OPACITY_OVERRIDEN"), EnvironmentDefines->bSubstratePremultipliedAlphaOpacityOverridden);
 
-	if (EnvironmentDefines->bMaterialIsSubstrate)
+	if (SubstrateTranslatorData.bMaterialIsSubstrate)
 	{
-		OutEnvironment.SetDefine(TEXT("SUBSTRATE_USES_CONVERSION_FROM_LEGACY"), EnvironmentDefines->bSubstrateUsesConversionFromLegacy);
-		OutEnvironment.SetDefine(TEXT("SUBSTRATE_OPTIMIZED_UNLIT"), EnvironmentDefines->bSubstrateOptimizedUnlit);
-
-		OutEnvironment.SetDefine(TEXT("SUBSTRATE_MATERIAL_OUTPUT_OPAQUE_ROUGH_REFRACTIONS"), EnvironmentDefines->bSubstrateMaterialOutputOpaqueRoughRefractions);
-		OutEnvironment.SetDefine(TEXT("SUBSTRATE_MATERIAL_EXPORT_TYPE"), EnvironmentDefines->SubstrateMaterialExportType);
-		OutEnvironment.SetDefine(TEXT("SUBSTRATE_MATERIAL_EXPORT_CONTEXT"), EnvironmentDefines->SubstrateMaterialExportContext);
-		OutEnvironment.SetDefine(TEXT("SUBSTRATE_MATERIAL_EXPORT_LEGACY_BLEND_MODE"), EnvironmentDefines->SubstrateMaterialExportLegacyBlendMode);
-		OutEnvironment.SetDefine(TEXT("SUBSTRATE_MATERIAL_ROUGHNESS_TRACKING"), EnvironmentDefines->bSubstrateRoughnessTracking);
-
-		OutEnvironment.SetDefine(TEXT("SUBSTRATE_LEGACY_IRIS_NORMAL"), EnvironmentDefines->bSubstrateLegacyIrisNormal);
-		OutEnvironment.SetDefine(TEXT("SUBSTRATE_LEGACY_IRIS_TANGENT"), EnvironmentDefines->bSubstrateLegacyIrisTangent);
-
-		OutEnvironment.SetDefine(TEXT("SUBSTRATE_SINGLEPATH"), EnvironmentDefines->bSubstrateSinglePath);
-		OutEnvironment.SetDefine(TEXT("SUBSTRATE_FASTPATH"), EnvironmentDefines->bSubstrateFastPath);
-		OutEnvironment.SetDefine(TEXT("SUBSTRATE_CLAMPED_CLOSURE_COUNT"), EnvironmentDefines->SubstrateClampedClosureCount);
-		OutEnvironment.SetDefine(TEXT("SUBSTRATE_COMPLEXSPECIALPATH"), EnvironmentDefines->bSubstrateComplexSpecialPath);
-
 		if (!Substrate::IsSubstrateBlendableGBufferEnabled(InPlatform))
 		{
 			static const TCHAR* SDCESymbols[] = {
@@ -3255,6 +2563,7 @@ void FHLSLMaterialTranslator::GetMaterialEnvironment(EShaderPlatform InPlatform,
 		}
 	}
 
+	// Append all substrate defines
 	for (int i = 0; i < EnvironmentDefines->SubstrateDefines.Num(); ++i)
 	{
 		OutEnvironment.SetDefine(*EnvironmentDefines->SubstrateDefines[i].Get<0>(), EnvironmentDefines->SubstrateDefines[i].Get<1>());
@@ -3390,26 +2699,12 @@ void FHLSLMaterialTranslator::GetSharedInputsMaterialCode(FString& PixelMembersD
 		PixelMembersInitializationEpilog += DerivativeVariations[DerivativeVariation].TranslatedCodeChunkDefinitions[LastProperty] + TEXT("\n");
 	}
 
+	// Generate some more substrate specific HLSL code.
 	if (Substrate::IsSubstrateEnabled())
 	{
-		PixelMembersDeclaration += FString::Printf(TEXT("\t#if SUBSTRATE_USE_FULLYSIMPLIFIED_MATERIAL == 1\n"));
-		PixelMembersDeclaration += FString::Printf(TEXT("\t%s FullySimplifiedFrontMaterial;\n"), HLSLTypeString(FMaterialAttributeDefinitionMap::GetValueType(MP_FrontMaterial)));
-		PixelMembersDeclaration += FString::Printf(TEXT("\t#endif\n"));
-		if (!FullySimplifiedSubstrateFrontMaterialTranslatedCodeChunks.IsEmpty())
-		{
-			PixelInputInitializerValues += FString::Printf(TEXT("\t#if SUBSTRATE_USE_FULLYSIMPLIFIED_MATERIAL == 1\n"));
-			PixelInputInitializerValues += FString::Printf(TEXT("\tPixelMaterialInputs.FullySimplifiedFrontMaterial = %s;\n"), *FullySimplifiedSubstrateFrontMaterialTranslatedCodeChunks);
-			PixelInputInitializerValues += FString::Printf(TEXT("\t#endif\n"));
-		}
-
-		for (uint32 SubstrateCompilationContextIndex = 0; SubstrateCompilationContextIndex < ESubstrateCompilationContext::SCC_MAX; ++SubstrateCompilationContextIndex)
-		{
-			FSubstrateCompilationContext& SubstrateCtx = SubstrateCompilationContext[SubstrateCompilationContextIndex];
-			if (SubstrateCtx.CodeChunkToSubstrateSharedLocalBasis.Num() > 0)
-			{
-				PixelInputInitializerValues += SubstrateCtx.SubstratePixelNormalInitializerValues;
-			}
-		}
+		SubstrateTranslatorData.GeneratePixelMemberDeclarationHLSLCode(PixelMembersDeclaration, HLSLTypeString(FMaterialAttributeDefinitionMap::GetValueType(MP_FrontMaterial)));
+		SubstrateTranslatorData.GeneratePixelInputInitializerHLSLCode(PixelInputInitializerValues); 
+		SubstrateTranslatorData.GeneratePixelNormalInitializerHLSLCode(PixelInputInitializerValues);
 	}
 
 	PixelMembersInitializationEpilog += PixelInputInitializerValues;
@@ -8899,7 +8194,7 @@ int32 FHLSLMaterialTranslator::SceneTextureLookup(int32 ViewportUV, uint32 InSce
 	// Reading SceneTexture, when Substrate is enabled, implies unpacking material buffer data. The unpacking function exists in different 'flavor' 
 	// for optimization purpose (simple/single/complex). To avoid compiling out single or complex unpacking paths (due to defines set by analyzing 
 	// the current shader, vs. scene texture pixels), we force Simple/Single versions to be disabled
-	FSubstrateCompilationContext& SubstrateCtx = SubstrateCompilationContext[CurrentSubstrateCompilationContext];
+	FSubstrateCompilationContext& SubstrateCtx = SubstrateTranslatorData.SubstrateCompilationContext[SubstrateTranslatorData.CurrentSubstrateCompilationContext];;
 	SubstrateCtx.SubstrateMaterialComplexity.bIsSimple = false;
 	SubstrateCtx.SubstrateMaterialComplexity.bIsSingle = false;
 	
@@ -13413,1380 +12708,24 @@ int32 FHLSLMaterialTranslator::NeuralOutput(int32 ViewportUV, uint32 NeuralIndex
 
 FSubstrateOperator& FHLSLMaterialTranslator::SubstrateCompilationRegisterOperator(int32 OperatorType, FGuid SubstrateExpressionGuid, FGuid ChildMaterialExpressionGuid, UMaterialExpression* Parent, FGuid SubstrateParentExpressionGuid, bool bUseParameterBlending)
 {
-	FSubstrateCompilationContext& SubstrateCtx = SubstrateCompilationContext[CurrentSubstrateCompilationContext];
-
-	if (OperatorType == SUBSTRATE_OPERATOR_BSDF_LEGACY)
-	{
-		// We register the fact that a legacy material conversion is used and register a simple BSDF
-		bSubstrateUsesConversionFromLegacy = true;
-		OperatorType = SUBSTRATE_OPERATOR_BSDF;
-	}
-
-	static FSubstrateOperator DefaultOperatorOnError = FSubstrateOperator();
-
-	if (SubstrateCtx.SubstrateMaterialExpressionToOperatorIndex.Find(SubstrateExpressionGuid))
-	{
-		// It is not possible to register/use a Substrate BSDF multiple times with this same exact graph path. (that would break the Substrate tree code generation)
-		Errorf(TEXT("Material %s: It is not possible to uses a Substrate BSDF (or any ouput of type SubstrateData) multiple times within a Substrate material topology with the same graph path GUID (asset: %s).\r\n"), *Material->GetDebugName(), *Material->GetAssetPath().ToString());
-		return DefaultOperatorOnError;
-	}
-
-	const uint32 NewOperatorIndex = SubstrateCtx.SubstrateMaterialExpressionToOperatorIndex.Num();
-	if (NewOperatorIndex >= SUBSTRATE_MAX_COMPILER_REGISTERED_OPERATOR_COUNT)
-	{
-		Errorf(TEXT("Material %s have too many Substrate Operators: the compiler is failing (asset: %s).\r\n"), *Material->GetDebugName(), *Material->GetAssetPath().ToString());
-		return DefaultOperatorOnError;
-	}
-
-	int32* ParentOperatorIndex = SubstrateCtx.SubstrateMaterialExpressionToOperatorIndex.Find(SubstrateParentExpressionGuid);
-	if (Parent!=nullptr && ParentOperatorIndex == nullptr)
-	{
-		Errorf(TEXT("Material %s tries to register unknown operator parents (asset: %s).\r\n"), *Material->GetDebugName(), *Material->GetAssetPath().ToString());
-		return DefaultOperatorOnError;
-	}
-
-	SubstrateCtx.SubstrateMaterialExpressionToOperatorIndex.Add(SubstrateExpressionGuid, NewOperatorIndex);
-
-	FSubstrateOperator& NewOperator = SubstrateCtx.SubstrateMaterialExpressionRegisteredOperators.AddDefaulted_GetRef();
-
-	NewOperator.OperatorType = OperatorType;
-	NewOperator.bNodeRequestParameterBlending = bUseParameterBlending;
-	NewOperator.Index = NewOperatorIndex;
-	NewOperator.ParentIndex = ParentOperatorIndex != nullptr ? *ParentOperatorIndex : INDEX_NONE;
-	NewOperator.LeftIndex   = INDEX_NONE;
-	NewOperator.RightIndex   = INDEX_NONE;
-
-	NewOperator.BSDFIndex = INDEX_NONE;	// Allocated later to be able to account for inline
-	NewOperator.SubUsage = SUBSTRATE_OPERATOR_SUBUSAGE_NONE;
-
-	NewOperator.MaxDistanceFromLeaves = 0;
-	NewOperator.bIsBottom = false;
-	NewOperator.bIsTop = false;
-	NewOperator.MaterialExpressionGuid = ChildMaterialExpressionGuid;
-
-	return NewOperator;
-}
-
-FSubstrateOperator& FHLSLMaterialTranslator::SubstrateCompilationGetOperator(FGuid SubstrateExpressionGuid)
-{
-	FSubstrateCompilationContext& SubstrateCtx = SubstrateCompilationContext[CurrentSubstrateCompilationContext];
-	auto* OperatorIndex = SubstrateCtx.SubstrateMaterialExpressionToOperatorIndex.Find(SubstrateExpressionGuid);
-	if (!(OperatorIndex && *OperatorIndex >= 0 && *OperatorIndex < SUBSTRATE_MAX_COMPILER_REGISTERED_OPERATOR_COUNT))
-	{
-		static FSubstrateOperator DefaultOperatorOnError = FSubstrateOperator();
-		return DefaultOperatorOnError;
-	};
-	return SubstrateCtx.SubstrateMaterialExpressionRegisteredOperators[*OperatorIndex];
-}
-
-FSubstrateOperator* FHLSLMaterialTranslator::SubstrateCompilationGetOperatorFromIndex(int32 OperatorIndex)
-{
-	FSubstrateCompilationContext& SubstrateCtx = SubstrateCompilationContext[CurrentSubstrateCompilationContext];
-	int32 OperatorCount = SubstrateCtx.SubstrateMaterialExpressionToOperatorIndex.Num();
-	if (OperatorIndex < 0 || OperatorIndex >= OperatorCount)
-	{
-		Errorf(TEXT("SubstrateCompilationGetOperatorFromIndex - OperatorIndex out of range %s (asset: %s).\r\n"), *Material->GetDebugName(), *Material->GetAssetPath().ToString());
-		return nullptr;
-	};
-	return &SubstrateCtx.SubstrateMaterialExpressionRegisteredOperators[OperatorIndex];
-}
-
-static FString GetParametersSharedLocalBasesName(ESubstrateCompilationContext CompilationContextIndex)
-{
-	switch (CompilationContextIndex)
-	{
-	case ESubstrateCompilationContext::SCC_Default:
-		return TEXT("SharedLocalBases");
-	case ESubstrateCompilationContext::SCC_FullySimplified:
-		return TEXT("SharedLocalBasesFullySimplified");
-	}
-	check(false);
-	return TEXT("ERROR");
-}
-
-static FString GetParametersSubstrateTreeName(ESubstrateCompilationContext CompilationContextIndex)
-{
-	switch (CompilationContextIndex)
-	{
-	case ESubstrateCompilationContext::SCC_Default:
-		return TEXT("SubstrateTree");
-	case ESubstrateCompilationContext::SCC_FullySimplified:
-		return TEXT("SubstrateTreeFullySimplified");
-	}
-	check(false);
-	return TEXT("ERROR");
-}
-
-/**
- * This function is temporary. It's only been duplicated for validation purposes and it will be removed in a few days when the latest changes settle without problems.
- */
-void FHLSLMaterialTranslator::FSubstrateCompilationContext::SubstrateEvaluateSharedLocalBases(
-	FHLSLMaterialTranslator* Compiler,
-	uint8& OutRequestedSharedLocalBasesCount,
-	FEnvironmentDefines* OutEnvironment)
-{
-	/*
-	* The final output code/workflow for shared tangent basis should look like
-	*
-	* #define SHAREDLOCALBASIS_INDEX_0 0		// default, unused
-	* #define SHAREDLOCALBASIS_INDEX_1 0		// default, unused
-	* #define SHAREDLOCALBASIS_INDEX_2 0
-	*
-	* FSubstrateData BSDF0 = GetSubstrateSlabBSDF(... SHAREDLOCALBASIS_0, NormalCode0 ...)
-	* FSubstrateData BSDF1 = GetSubstrateSlabBSDF(... SHAREDLOCALBASIS_1, NormalCode1 ...)
-	*
-	* float3 NormalCode2 = lerp(NormalCode0, NormalCode1, mix)
-	* FSubstrateData BSDF2 = SubstrateHorizontalMixingParameterBlending(BSDF0, BSDF1, mix, NormalCode2, SHAREDLOCALBASIS_INDEX_2, SharedLocalBases.Types) // will internally create NormalCode2
-	*
-	* tParameters.SharedLocalBases.Normals[SHAREDLOCALBASIS_INDEX_2] = NormalCode2;
-	* #if MATERIAL_TANGENTSPACENORMAL
-	* Parameters.SharedLocalBases.Normals[SHAREDLOCALBASIS_INDEX_2] *= Parameters.TwoSidedSign;
-	* #endif
-	*/
-
-	FSubstrateRegisteredSharedLocalBasis UsedSharedLocalBasesInfo[SUBSTRATE_MAX_SHAREDLOCALBASES_REGISTERS];
-	FinalUsedSharedLocalBasesCount = 0;
-	OutRequestedSharedLocalBasesCount = 0;
-	const FString ParameterSharedLocalBasesName = GetParametersSharedLocalBasesName(CompilationContextIndex);
-
-	if (CompilationContextIndex == SCC_FullySimplified)
-	{
-		SubstratePixelNormalInitializerValues += FString::Printf(TEXT("\t#if SUBSTRATE_USE_FULLYSIMPLIFIED_MATERIAL == 1\n"));
-	}
-
-	for (int32 OpIt = 0; OpIt < SubstrateMaterialExpressionRegisteredOperators.Num(); ++OpIt)
-	{
-		const FSubstrateOperator& BSDFOperator = SubstrateMaterialExpressionRegisteredOperators[OpIt];
-		if (BSDFOperator.BSDFIndex == INDEX_NONE || BSDFOperator.IsDiscarded())
-		{
-			continue;	// not a BSDF or if discarded (i.e. not the root of a parameter blending subtree), then there is no local basis to register
-		}
-
-		if (BSDFOperator.BSDFRegisteredSharedLocalBasis.NormalCodeChunk == INDEX_NONE && BSDFOperator.BSDFRegisteredSharedLocalBasis.TangentCodeChunk == INDEX_NONE)
-		{
-			continue;	// We skip null normal on certain BSDF, for instance unlit.
-		}
-		const FSubstrateSharedLocalBasesInfo& SubstrateSharedLocalBasesInfo = SubstrateCompilationInfoGetMatchingSharedLocalBasisInfo(BSDFOperator.BSDFRegisteredSharedLocalBasis);
-
-		// First, we check that the normal/tangent has not already written out (avoid 2 BSDFs sharing the same normal to note generate the same code twice)
-		bool bAlreadyProcessed = false;
-		for (uint8 i = 0; i < FinalUsedSharedLocalBasesCount; ++i)
-		{
-			if (UsedSharedLocalBasesInfo[i].NormalCodeChunkHash == SubstrateSharedLocalBasesInfo.SharedData.NormalCodeChunkHash &&
-				(UsedSharedLocalBasesInfo[i].TangentCodeChunkHash == SubstrateSharedLocalBasesInfo.SharedData.TangentCodeChunkHash || BSDFOperator.BSDFRegisteredSharedLocalBasis.TangentCodeChunk == INDEX_NONE))
-			{
-				bAlreadyProcessed = true;
-				break;
-			}
-		}
-		if (bAlreadyProcessed)
-		{
-			continue;
-		}
-
-		++OutRequestedSharedLocalBasesCount;
-		if (FinalUsedSharedLocalBasesCount >= SUBSTRATE_MAX_SHAREDLOCALBASES_REGISTERS)
-		{
-			continue;
-		}
-
-		const uint8 FinalSharedLocalBasisIndex = FinalUsedSharedLocalBasesCount++;
-		UsedSharedLocalBasesInfo[FinalSharedLocalBasisIndex] = SubstrateSharedLocalBasesInfo.SharedData;
-
-		// Write out normals
-		SubstratePixelNormalInitializerValues += FString::Printf(TEXT("\tParameters.%s.Normals[%u] = %s;\n"), *ParameterSharedLocalBasesName, FinalSharedLocalBasisIndex, *SubstrateSharedLocalBasesInfo.NormalCode);
-		SubstratePixelNormalInitializerValues += FString::Printf(TEXT("\t#if MATERIAL_TANGENTSPACENORMAL\n"));
-		SubstratePixelNormalInitializerValues += FString::Printf(TEXT("\tParameters.%s.Normals[%u] *= Parameters.TwoSidedSign;\n"), *ParameterSharedLocalBasesName, FinalSharedLocalBasisIndex);
-		SubstratePixelNormalInitializerValues += FString::Printf(TEXT("\t#endif\n"));
-
-		// Write out tangents
-		if (SubstrateSharedLocalBasesInfo.SharedData.TangentCodeChunk != INDEX_NONE)
-		{
-			SubstratePixelNormalInitializerValues += FString::Printf(TEXT("\tParameters.%s.Tangents[%u] = %s;\n"), *ParameterSharedLocalBasesName, FinalSharedLocalBasisIndex, *SubstrateSharedLocalBasesInfo.TangentCode);
-		}
-		else
-		{
-			SubstratePixelNormalInitializerValues += FString::Printf(TEXT("\tParameters.%s.Tangents[%u] = Parameters.TangentToWorld[0];\n"), *ParameterSharedLocalBasesName, FinalSharedLocalBasisIndex);
-		}
-		SubstratePixelNormalInitializerValues += FString::Printf(TEXT("\t#if MATERIAL_TANGENTSPACENORMAL\n"));
-		SubstratePixelNormalInitializerValues += FString::Printf(TEXT("\tParameters.%s.Tangents[%u] *= Parameters.TwoSidedSign;\n"), *ParameterSharedLocalBasesName, FinalSharedLocalBasisIndex);
-		SubstratePixelNormalInitializerValues += FString::Printf(TEXT("\t#endif\n"));
-	}
-
-	SubstratePixelNormalInitializerValues += FString::Printf(TEXT("\tParameters.%s.Count = %u;\n"), *ParameterSharedLocalBasesName, FinalUsedSharedLocalBasesCount);
-
-	if (CompilationContextIndex == SCC_FullySimplified)
-	{
-		SubstratePixelNormalInitializerValues += FString::Printf(TEXT("\t#endif // SUBSTRATE_USE_FULLYSIMPLIFIED_MATERIAL\n"));
-	}
-
-	if (OutEnvironment)
-	{
-		// Now write out all the macros, them mapping from the BSDF to the effective position/index in the shared local basis array they should write to.
-		OutEnvironment->SubstrateDefines.Reserve(CodeChunkToSubstrateSharedLocalBasis.Num());
-		for (TMultiMap<uint64, FSubstrateSharedLocalBasesInfo>::TConstIterator It(CodeChunkToSubstrateSharedLocalBasis); It; ++It)
-		{
-			// The default linear output index will be 0 by default, and different if in fact the shared local basis points to one that is effectively in used in the array of shared local bases.
-			uint8 LinearIndex = 0;
-			for (uint8 i = 0; i < FinalUsedSharedLocalBasesCount; ++i)
-			{
-				if (UsedSharedLocalBasesInfo[i].NormalCodeChunkHash == It->Value.SharedData.NormalCodeChunkHash &&
-					(UsedSharedLocalBasesInfo[i].TangentCodeChunkHash == It->Value.SharedData.TangentCodeChunkHash || It->Value.SharedData.TangentCodeChunk == INDEX_NONE))
-				{
-					LinearIndex = i;
-					break;
-				}
-			}
-
-			OutEnvironment->SubstrateDefines.Emplace(*Compiler->GetSubstrateSharedLocalBasisIndexMacroInner(It->Value.SharedData, CompilationContextIndex), LinearIndex);
-		}
-	}
-}
-
-bool FHLSLMaterialTranslator::FSubstrateCompilationContext::SubstrateGenerateDerivedMaterialOperatorData(FHLSLMaterialTranslator* Compiler)
-{
-	FMaterial* CompilerMaterial = Compiler->Material;
-
-	if (SubstrateMaterialExpressionRegisteredOperators.IsEmpty())
-	{
-		Compiler->Errorf(TEXT("Could not find any Substrate operators or BSDFs in Material %s (asset: %s).\r\n"), *CompilerMaterial->GetDebugName(), *CompilerMaterial->GetAssetPath().ToString());
-		return false;
-	}
-
-	//
-	// Evaluate the one and only root node.
-	// And make sure each and every path of the Substrate tree have valid children and path.
-	//
-	int32 RootIndex = INDEX_NONE;
-	for (auto& It : SubstrateMaterialExpressionRegisteredOperators)
-	{
-		if (It.ParentIndex == INDEX_NONE)
-		{
-			check(RootIndex == INDEX_NONE);	// There can only be one
-			RootIndex = It.Index;
-		}
-
-		if (It.OperatorType == SUBSTRATE_OPERATOR_BSDF)
-		{
-			// Gather information about data written by BSDF
-			bSubstrateWritesEmissive |= It.bBSDFWritesEmissive > 0;
-			bSubstrateWritesAmbientOcclusion |= It.bBSDFWritesAmbientOcclusion > 0;
-		}
-
-		if (It.IsDiscarded())
-		{
-			continue; // ignore discarded operations in sub tree using parameter blending
-		}
-
-		bool bMustHaveLeftChild = false;
-		bool bMustHaveRightChild = false;
-		switch (It.OperatorType)
-		{
-		// Operators without any child
-		case SUBSTRATE_OPERATOR_BSDF:
-		{
-		}
-		break;
-
-		// Operators with two children
-		case SUBSTRATE_OPERATOR_HORIZONTAL:
-		case SUBSTRATE_OPERATOR_VERTICAL:
-		case SUBSTRATE_OPERATOR_ADD:
-		case SUBSTRATE_OPERATOR_SELECT:
-		{
-			bMustHaveLeftChild = true;
-			bMustHaveRightChild = true;
-		}
-		break;
-
-		// Operators with a single child
-		case SUBSTRATE_OPERATOR_WEIGHT:
-		{
-			bMustHaveLeftChild = true;
-		}
-		break;
-		}
-
-		if (bMustHaveLeftChild && It.LeftIndex == INDEX_NONE)
-		{
-			Compiler->Errorf(TEXT("A Substrate Operator %s node is missing its first input from material %s (asset: %s).\r\n"), GetSubstrateOperatorStr(It.OperatorType), *CompilerMaterial->GetDebugName(), *CompilerMaterial->GetAssetPath().ToString());
-			return false;
-		}
-		if (bMustHaveRightChild && It.RightIndex == INDEX_NONE)
-		{
-			Compiler->Errorf(TEXT("A Substrate Operator %s node is missing its second input from material %s (asset: %s).\r\n"), GetSubstrateOperatorStr(It.OperatorType), *CompilerMaterial->GetDebugName(), *CompilerMaterial->GetAssetPath().ToString());
-			return false;
-		}
-	}
-	SubstrateMaterialRootOperator = &SubstrateMaterialExpressionRegisteredOperators[RootIndex];
-	if (!SubstrateMaterialRootOperator)
-	{
-		Compiler->Errorf(TEXT("Cannot find the root of the Substrate Tree for Material %s (asset: %s).\r\n"), *CompilerMaterial->GetDebugName(), *CompilerMaterial->GetAssetPath().ToString());
-		return false;
-	}
-
-	
-	//
-	// Validate the tree operators with respect to how BSDF and operator can flow.
-	// Also evaluate a bunch of data required.
-	//
-	enum EMaterialTypeValidation : uint32
-	{
-		MTV_BSDF_None				= 0x000,
-		MTV_BSDF_Unlit				= 0x001,
-		MTV_BSDF_Slab				= 0x002,
-		MTV_BSDF_Hair				= 0x004,
-		MTV_BSDF_Eye				= 0x008,
-		MTV_BSDF_SLW				= 0x010,
-		MTV_BSDF_Toon				= 0x020,
-		MTV_BSDF_VFogCloud			= 0x040,
-		MTV_UI						= 0x080,
-		MTV_Decal					= 0x100,
-		MTV_PostProcess				= 0x200,
-		MTV_LightFunction			= 0x400,
-
-		MTV_ValidMultipleBSDFs		= (MTV_BSDF_Slab | MTV_BSDF_Hair | MTV_BSDF_Toon),	// SUBSTRATE_TODO Unlit cannot be in there yet due to the specific define for it (SUBSTRATE_OPTIMIZED_UNLIT)
-		MTV_WeightOp_AllowedList	= (MTV_BSDF_Unlit | MTV_BSDF_Slab | MTV_BSDF_Toon),
-		MTV_SelectOp_AllowedList	= (MTV_ValidMultipleBSDFs),							// This only works because parameter blending is enforced so there will always be a single BSDF of a single type only. And hair or toon BSDFs will never pass through other operator allowed list.
-		MTV_HorizoOp_AllowedList	= MTV_BSDF_Slab,
-		MTV_VerticOp_AllowedList	= MTV_BSDF_Slab,
-		MTV_DecalOp_AllowedList		= MTV_BSDF_Slab,
-		MTV_AddOp_AllowedList		= MTV_BSDF_Slab
-	};
-	struct FSubstrateTreeValidation
-	{
-		uint32 MaterialType = MTV_BSDF_None;
-		bool bValid = true;
-	};
-	FSubstrateTreeValidation SubstrateTreeValidation;
-	{
-
-		auto MergeSubstrateTreeValidationParameterBlended = [](const FSubstrateTreeValidation& A, const FSubstrateTreeValidation& B)
-		{
-			FSubstrateTreeValidation Result;
-			Result.MaterialType = A.MaterialType | B.MaterialType;
-			return Result;
-		};
-		auto MergeSubstrateTreeValidation = [](const FSubstrateTreeValidation& A, const FSubstrateTreeValidation& B)
-		{
-			FSubstrateTreeValidation Result;
-			Result.MaterialType = A.MaterialType | B.MaterialType;
-			return Result;
-		};
-
-		auto ValidateAllowedList = [&](FSubstrateOperator& CurrentOperator, FSubstrateTreeValidation SubstrateTreeValidation)
-		{
-			if (CurrentOperator.OperatorType == SUBSTRATE_OPERATOR_BSDF)
-			{
-				return true;
-			}
-
-			FString OperatorName;
-			uint32 AllowedListMask = 0;
-			if (CurrentOperator.OperatorType == SUBSTRATE_OPERATOR_WEIGHT && CurrentOperator.SubUsage == SUBSTRATE_OPERATOR_SUBUSAGE_DECAL)
-			{
-				OperatorName = TEXT("ConvertToDecal");
-				AllowedListMask = MTV_DecalOp_AllowedList;
-			}
-			else if (CurrentOperator.OperatorType == SUBSTRATE_OPERATOR_WEIGHT)
-			{
-				OperatorName = TEXT("Weight");
-				AllowedListMask = MTV_WeightOp_AllowedList;
-			}
-			else if (CurrentOperator.OperatorType == SUBSTRATE_OPERATOR_HORIZONTAL)
-			{
-				OperatorName = TEXT("Horizontal");
-				AllowedListMask = MTV_HorizoOp_AllowedList;
-			}
-			else if (CurrentOperator.OperatorType == SUBSTRATE_OPERATOR_VERTICAL)
-			{
-				OperatorName = TEXT("Vertical");
-				AllowedListMask = MTV_VerticOp_AllowedList;
-			}
-			else if (CurrentOperator.OperatorType == SUBSTRATE_OPERATOR_ADD)
-			{
-				OperatorName = TEXT("Add");
-				AllowedListMask = MTV_AddOp_AllowedList;
-			}
-			else if (CurrentOperator.OperatorType == SUBSTRATE_OPERATOR_SELECT)
-			{
-				OperatorName = TEXT("Select");
-				AllowedListMask = MTV_SelectOp_AllowedList;
-			}
-			else
-			{
-				Compiler->Errorf(TEXT("Substrate error: unhandled operator type from ValidateAllowedList. (asset: %s).\r\n"), *CompilerMaterial->GetDebugName());
-				return false;
-			}
-
-			uint32 FoundNotAllowedTypeMask = SubstrateTreeValidation.MaterialType & (~AllowedListMask);
-
-			if (FoundNotAllowedTypeMask != 0)
-			{
-				FString NotAllowedTypeNames;
-				NotAllowedTypeNames += "{";
-				if (!!(FoundNotAllowedTypeMask & MTV_BSDF_Unlit))
-				{
-					NotAllowedTypeNames += TEXT("Unlit ");
-				}
-				if (!!(FoundNotAllowedTypeMask & MTV_BSDF_Slab))
-				{
-					NotAllowedTypeNames += TEXT("Slab ");
-				}
-				if (!!(FoundNotAllowedTypeMask & MTV_BSDF_Hair))
-				{
-					NotAllowedTypeNames += TEXT("Hair ");
-				}
-				if (!!(FoundNotAllowedTypeMask & MTV_BSDF_Eye))
-				{
-					NotAllowedTypeNames += TEXT("Eye ");
-				}
-				if (!!(FoundNotAllowedTypeMask & MTV_BSDF_SLW))
-				{
-					NotAllowedTypeNames += TEXT("SLW ");
-				}
-				if (!!(FoundNotAllowedTypeMask & MTV_BSDF_Toon))
-				{
-					NotAllowedTypeNames += TEXT("Toon ");
-				}
-				if (!!(FoundNotAllowedTypeMask & MTV_BSDF_VFogCloud))
-				{
-					NotAllowedTypeNames += TEXT("VFogCloud ");
-				}
-				if (!!(FoundNotAllowedTypeMask & MTV_UI))
-				{
-					NotAllowedTypeNames += TEXT("UI ");
-				}
-				if (!!(FoundNotAllowedTypeMask & MTV_Decal))
-				{
-					NotAllowedTypeNames += TEXT("Decal ");
-				}
-				if (!!(FoundNotAllowedTypeMask & MTV_PostProcess))
-				{
-					NotAllowedTypeNames += TEXT("PostProcess ");
-				}
-				if (!!(FoundNotAllowedTypeMask & MTV_LightFunction))
-				{
-					NotAllowedTypeNames += TEXT("LightFunction ");
-				}
-				NotAllowedTypeNames += "}";
-
-				Compiler->Errorf(TEXT("Substrate error: Operator %s cannot be used with the following data flowing through %s. (asset: %s).\r\n"), *OperatorName, *NotAllowedTypeNames, *CompilerMaterial->GetDebugName());
-				return false;
-			}
-			return true;
-		};
-
-		std::function<FSubstrateTreeValidation(FSubstrateOperator&)> WalkOperators = [&](FSubstrateOperator& CurrentOperator) -> FSubstrateTreeValidation
-		{
-			FSubstrateTreeValidation Out;
-			switch (CurrentOperator.OperatorType)
-			{
-			case SUBSTRATE_OPERATOR_VERTICAL:
-			case SUBSTRATE_OPERATOR_HORIZONTAL:
-			case SUBSTRATE_OPERATOR_ADD:
-			{
-				FSubstrateTreeValidation A = WalkOperators(SubstrateMaterialExpressionRegisteredOperators[CurrentOperator.LeftIndex]);
-				FSubstrateTreeValidation B = WalkOperators(SubstrateMaterialExpressionRegisteredOperators[CurrentOperator.RightIndex]);
-				Out = MergeSubstrateTreeValidation(A, B);
-				break;
-			}
-			case SUBSTRATE_OPERATOR_SELECT:
-			{
-				if (CurrentOperator.bNodeRequestParameterBlending)
-				{
-					FSubstrateTreeValidation A = WalkOperators(SubstrateMaterialExpressionRegisteredOperators[CurrentOperator.LeftIndex]);
-					FSubstrateTreeValidation B = WalkOperators(SubstrateMaterialExpressionRegisteredOperators[CurrentOperator.RightIndex]);
-					Out = MergeSubstrateTreeValidation(A, B);
-				}
-				else
-				{
-					Compiler->Errorf(TEXT("Substrate error: Select node can only be parameter blended as of today. (asset: %s).\r\n"), *CompilerMaterial->GetDebugName());
-				}
-				break;
-			}
-			case SUBSTRATE_OPERATOR_WEIGHT:
-			{
-				Out = WalkOperators(SubstrateMaterialExpressionRegisteredOperators[CurrentOperator.LeftIndex]);
-				break;
-			}
-			case SUBSTRATE_OPERATOR_BSDF:
-			{
-				if (CurrentOperator.BSDFType == SUBSTRATE_BSDF_TYPE_UNLIT)
-				{
-					if (CurrentOperator.SubUsage == SUBSTRATE_OPERATOR_SUBUSAGE_NONE)
-					{
-						Out.MaterialType = MTV_BSDF_Unlit;
-					}
-					else if (CurrentOperator.SubUsage == SUBSTRATE_OPERATOR_SUBUSAGE_UI)
-					{
-						Out.MaterialType = MTV_UI;
-					}
-					else if (CurrentOperator.SubUsage == SUBSTRATE_OPERATOR_SUBUSAGE_LIGHTFUNCTION)
-					{
-						Out.MaterialType = MTV_LightFunction;
-					}
-					else if (CurrentOperator.SubUsage == SUBSTRATE_OPERATOR_SUBUSAGE_POSTPROCESS)
-					{
-						Out.MaterialType = MTV_PostProcess;
-					}
-					else
-					{
-						Compiler->Errorf(TEXT("Substrate error: unhandled Unlit BSDF sub type from WalkParameterBlendedSubTree. (asset: %s).\r\n"), *CompilerMaterial->GetDebugName());
-					}
-				}
-				else if (CurrentOperator.BSDFType == SUBSTRATE_BSDF_TYPE_SLAB)
-				{
-					Out.MaterialType = MTV_BSDF_Slab;
-				}
-				else if (CurrentOperator.BSDFType == SUBSTRATE_BSDF_TYPE_VOLUMETRICFOGCLOUD)
-				{
-					Out.MaterialType = MTV_BSDF_VFogCloud;
-				}
-				else if (CurrentOperator.BSDFType == SUBSTRATE_BSDF_TYPE_HAIR)
-				{
-					Out.MaterialType = MTV_BSDF_Hair;
-				}
-				else if (CurrentOperator.BSDFType == SUBSTRATE_BSDF_TYPE_SINGLELAYERWATER)
-				{
-					Out.MaterialType = MTV_BSDF_SLW;
-				}
-				else if (CurrentOperator.BSDFType == SUBSTRATE_BSDF_TYPE_TOON)
-				{
-					Out.MaterialType = MTV_BSDF_Toon;
-				}
-				else if (CurrentOperator.BSDFType == SUBSTRATE_BSDF_TYPE_EYE)
-				{
-					Out.MaterialType = MTV_BSDF_Eye;
-				}
-				else
-				{
-					Compiler->Errorf(TEXT("Substrate error: unhandled BSDF type from WalkParameterBlendedSubTree. (asset: %s).\r\n"), *CompilerMaterial->GetDebugName());
-				}
-				break;
-			}
-			}
-			Out.bValid = ValidateAllowedList(CurrentOperator, Out);
-			return Out;
-		};
-
-		SubstrateTreeValidation = WalkOperators(*SubstrateMaterialRootOperator);
-		if (!SubstrateTreeValidation.bValid)
-		{
-			Compiler->Errorf(TEXT("Substrate error: invalid Substrate node graph. (asset: %s).\r\n"), *CompilerMaterial->GetDebugName());
-			return false;
-		}
-	}
-
-	// Evaluate if simplification is needed.
-	SubstrateSimplificationStatus.bRunFullSimplification = Compiler->SubstrateCompilationConfig.bFullSimplify || CompilationContextIndex == ESubstrateCompilationContext::SCC_FullySimplified;
-	if (!SubstrateSimplificationStatus.bRunFullSimplification)
-	{
-		//
-		// Generate LayerDepth value for all operators/bsdfs for progressive material simplification and sort them.
-		//
-		int VOpTopBranchCountTaken = 0;
-		int VOpBottomBranchCountTaken = 0;
-		std::function<void(FSubstrateOperator&)> WalkOperatorsForSimplification = [&](FSubstrateOperator& CurrentOperator) -> void
-		{
-			CurrentOperator.LayerDepth = VOpBottomBranchCountTaken;
-			switch (CurrentOperator.OperatorType)
-			{
-			case SUBSTRATE_OPERATOR_VERTICAL:
-			{
-				VOpTopBranchCountTaken++;
-				WalkOperatorsForSimplification(SubstrateMaterialExpressionRegisteredOperators[CurrentOperator.LeftIndex]);
-				VOpTopBranchCountTaken--;
-				VOpBottomBranchCountTaken++;
-				WalkOperatorsForSimplification(SubstrateMaterialExpressionRegisteredOperators[CurrentOperator.RightIndex]);
-				VOpBottomBranchCountTaken--;
-				break;
-			}
-			case SUBSTRATE_OPERATOR_HORIZONTAL:
-			case SUBSTRATE_OPERATOR_ADD:
-			case SUBSTRATE_OPERATOR_SELECT:
-			{
-				WalkOperatorsForSimplification(SubstrateMaterialExpressionRegisteredOperators[CurrentOperator.LeftIndex]);
-				WalkOperatorsForSimplification(SubstrateMaterialExpressionRegisteredOperators[CurrentOperator.RightIndex]);
-				break;
-			}
-			case SUBSTRATE_OPERATOR_WEIGHT:
-			{
-				WalkOperatorsForSimplification(SubstrateMaterialExpressionRegisteredOperators[CurrentOperator.LeftIndex]);
-				break;
-			}
-			case SUBSTRATE_OPERATOR_BSDF:
-			{
-				break;
-			}
-			}
-
-			// Add operators to simplify with a priority.
-			switch (CurrentOperator.OperatorType)
-			{
-			case SUBSTRATE_OPERATOR_VERTICAL:
-			case SUBSTRATE_OPERATOR_HORIZONTAL:
-			case SUBSTRATE_OPERATOR_ADD:
-			{
-				FSubstrateSimplificationStatus::FOperatorToSimplify OperatorToSimplify;
-				OperatorToSimplify.Data.Index = CurrentOperator.Index;
-				OperatorToSimplify.Data.Depth = CurrentOperator.LayerDepth;
-				SubstrateSimplificationStatus.OperatorSimplificationOrder.Push(OperatorToSimplify);
-				break;
-			}
-			}
-		};
-		WalkOperatorsForSimplification(*SubstrateMaterialRootOperator);
-
-		SubstrateSimplificationStatus.OperatorSimplificationOrder.Sort();	// sort according to depth
-	}
-
-	EShaderPlatform ShaderPlatform = Compiler->GetShaderPlatform();
-	const uint32 SubstrateBytePerPixel = Compiler->SubstrateCompilationConfig.BytesPerPixelOverride > 0 ? Compiler->SubstrateCompilationConfig.BytesPerPixelOverride : Substrate::GetBytePerPixel(ShaderPlatform);
-	const uint32 SubstrateClosurePerPixel = Compiler->SubstrateCompilationConfig.ClosuresPerPixelOverride > 0 ? Compiler->SubstrateCompilationConfig.ClosuresPerPixelOverride : Substrate::GetClosurePerPixel(ShaderPlatform);
-
-	bool bFirstLoop = true;
-	bool bRequestMaterialDetailsSet = false;
-	do 
-	{
-		if (!bFirstLoop && !SubstrateSimplificationStatus.bFullSimplificationStepHasBeenRun && !SubstrateSimplificationStatus.bMaterialFitsInMemoryBudget && SubstrateSimplificationStatus.OperatorSimplificationOrder.Num() > 0)
-		{
-			// Mark the deepest operator for parameter blending
-			FSubstrateSimplificationStatus::FOperatorToSimplify& OperatorToSimplify = SubstrateSimplificationStatus.OperatorSimplificationOrder.Top();
-			SubstrateMaterialExpressionRegisteredOperators[OperatorToSimplify.Data.Index].bNodeRequestParameterBlending = true;
-			SubstrateSimplificationStatus.OperatorSimplificationOrder.Pop(EAllowShrinking::No);
-
-			// Mark that this is similar to have run full simplification
-			SubstrateSimplificationStatus.bFullSimplificationStepHasBeenRun |= SubstrateSimplificationStatus.OperatorSimplificationOrder.Num() == 0;
-		}
-		else if (!bFirstLoop && SubstrateSimplificationStatus.bFullSimplificationStepHasBeenRun && !SubstrateSimplificationStatus.bMaterialFitsInMemoryBudget && SubstrateSimplificationStatus.OperatorSimplificationOrder.Num() == 0 && !SubstrateSimplificationStatus.bSlabSimplificationStepHasBeenRun)
-		{
-			for (auto& It : SubstrateMaterialExpressionRegisteredOperators)
-			{
-				// Disable all optional features for now to fit.
-				// SUBSTRATE_TODO we will need to refine that to account for platforms supporting SSS for instance.
-				It.BSDFFeatures = ESubstrateBsdfFeature::None;
-			}
-			SubstrateSimplificationStatus.bSlabSimplificationStepHasBeenRun = true;
-		}
-		else if (bFirstLoop)
-		{
-			// Fall through. This is needed for material with a single BSDF and no operator to simplify.
-			bFirstLoop = false;
-		}
-		else
-		{
-			Compiler->Errorf(TEXT("Unkown Substrate material simplification status error for %s (asset: %s).\r\n"), *CompilerMaterial->GetDebugName(), *CompilerMaterial->GetDebugName());
-			return false;
-		}
-
-		// Reset some data
-		SubstrateMaterialEffectiveClosureCount = 0;
-
-		//
-		// Parse the tree and mark nodes that are the root of a subtree using parameter blending, while other nodes in that tree are forced to use parameter blending.
-		// Allocate BSDFIndex at the same time.
-		// Check BSDF that should be unique.
-		//
-		bool bHasUnlit			= !!(SubstrateTreeValidation.MaterialType & MTV_BSDF_Unlit);
-		bool bHasVFogCloud		= !!(SubstrateTreeValidation.MaterialType & MTV_BSDF_VFogCloud);
-		bool bHasHair			= !!(SubstrateTreeValidation.MaterialType & MTV_BSDF_Hair);
-		bool bHasEye			= !!(SubstrateTreeValidation.MaterialType & MTV_BSDF_Eye);
-		bool bHasSLW			= !!(SubstrateTreeValidation.MaterialType & MTV_BSDF_SLW);
-		bool bHasToon			= !!(SubstrateTreeValidation.MaterialType & MTV_BSDF_Toon);
-		bool bHasSlab			= !!(SubstrateTreeValidation.MaterialType & MTV_BSDF_Slab);
-		bool bHasUI				= !!(SubstrateTreeValidation.MaterialType & MTV_UI);
-		bool bHasDecal			= !!(SubstrateTreeValidation.MaterialType & MTV_Decal);
-		bool bHasPostProcess	= !!(SubstrateTreeValidation.MaterialType & MTV_PostProcess);
-		bool bHasLightFunction	= !!(SubstrateTreeValidation.MaterialType & MTV_LightFunction);
-		{
-			int VOpTopBranchCountTaken = 0;
-			int VOpBottomBranchCountTaken = 0;
-			bool bSubstrateUsesVerticalLayering = false;
-
-			std::function<void(FSubstrateOperator&, bool)> WalkOperators = [&](FSubstrateOperator& CurrentOperator, bool bInsideParameterBlendingSubTree) -> void
-			{
-				const bool bCurrentOpRequestParameterBlending	= CurrentOperator.bNodeRequestParameterBlending || SubstrateSimplificationStatus.bRunFullSimplification;
-				const bool bRootOfParameterBlendingSubTree		= bCurrentOpRequestParameterBlending && !bInsideParameterBlendingSubTree;
-				const bool bUseParameterBlending				= bCurrentOpRequestParameterBlending || bInsideParameterBlendingSubTree;
-
-				if (CurrentOperator.BSDFType == SUBSTRATE_BSDF_TYPE_SLAB																// Any operators
-					|| CurrentOperator.BSDFType == SUBSTRATE_BSDF_TYPE_UNLIT															// Weight operator
-					|| CurrentOperator.BSDFType == SUBSTRATE_BSDF_TYPE_HAIR || CurrentOperator.BSDFType == SUBSTRATE_BSDF_TYPE_TOON)	// Select operator
-				{
-					// Update the parameter blending data for BSDFs supporting operators.
-					// We also need to do this for UNLIT since it supports Coverage operator.
-					CurrentOperator.bUseParameterBlending = bUseParameterBlending;
-					CurrentOperator.bRootOfParameterBlendingSubTree = bRootOfParameterBlendingSubTree;
-				}
-
-				switch (CurrentOperator.OperatorType)
-				{
-				case SUBSTRATE_OPERATOR_VERTICAL:
-				{
-					WalkOperators(SubstrateMaterialExpressionRegisteredOperators[CurrentOperator.LeftIndex], bUseParameterBlending);
-					WalkOperators(SubstrateMaterialExpressionRegisteredOperators[CurrentOperator.RightIndex], bUseParameterBlending);
-					break;
-				}
-				case SUBSTRATE_OPERATOR_HORIZONTAL:
-				case SUBSTRATE_OPERATOR_ADD:
-				case SUBSTRATE_OPERATOR_SELECT:
-				{
-					WalkOperators(SubstrateMaterialExpressionRegisteredOperators[CurrentOperator.LeftIndex], bUseParameterBlending);
-					WalkOperators(SubstrateMaterialExpressionRegisteredOperators[CurrentOperator.RightIndex], bUseParameterBlending);
-					break;
-				}
-				case SUBSTRATE_OPERATOR_WEIGHT:
-				{
-					WalkOperators(SubstrateMaterialExpressionRegisteredOperators[CurrentOperator.LeftIndex], bUseParameterBlending);
-					break;
-				}
-				case SUBSTRATE_OPERATOR_BSDF:
-				{
-					if (!bInsideParameterBlendingSubTree)
-					{
-						CurrentOperator.BSDFIndex = SubstrateMaterialEffectiveClosureCount++;
-					}
-
-					bHasUnlit |= CurrentOperator.BSDFType == SUBSTRATE_BSDF_TYPE_UNLIT;
-					bHasVFogCloud |= CurrentOperator.BSDFType == SUBSTRATE_BSDF_TYPE_VOLUMETRICFOGCLOUD;
-					bHasHair |= CurrentOperator.BSDFType == SUBSTRATE_BSDF_TYPE_HAIR;
-					bHasEye |= CurrentOperator.BSDFType == SUBSTRATE_BSDF_TYPE_EYE;
-					bHasSLW |= CurrentOperator.BSDFType == SUBSTRATE_BSDF_TYPE_SINGLELAYERWATER;
-					bHasToon |= CurrentOperator.BSDFType == SUBSTRATE_BSDF_TYPE_TOON;
-					bHasSlab |= CurrentOperator.BSDFType == SUBSTRATE_BSDF_TYPE_SLAB;
-					break;
-				}
-				}
-
-				// We mark the top of a parameter blending tree as a BSDF now to allocate a slot for it that can then be used next for non-parameter blending operations.
-				// Intermediate parameter blending BSDF and operation will be done inline and stored in FSubstrateData.
-				if (CurrentOperator.OperatorType != SUBSTRATE_OPERATOR_BSDF && bRootOfParameterBlendingSubTree)
-				{
-					CurrentOperator.OperatorType = SUBSTRATE_OPERATOR_BSDF;
-					CurrentOperator.BSDFIndex = SubstrateMaterialEffectiveClosureCount++;
-					// We do not reset LeftIndex and RightIndex because those are needed to recover local tangent basis information needed with parameter blending.
-				}
-
-				// When at least one vertical operator exists that is not parameter blending, we can enabled writing to opaque rough refraction buffer.
-				bSubstrateUsesVerticalLayering = bSubstrateUsesVerticalLayering || (!CurrentOperator.bUseParameterBlending && CurrentOperator.OperatorType == SUBSTRATE_OPERATOR_VERTICAL);
-			};
-
-			WalkOperators(*SubstrateMaterialRootOperator, false);
-
-			if (CompilationContextIndex == ESubstrateCompilationContext::SCC_Default)
-			{
-				// Only write those data for the default material
-				const bool bIsOpaqueOrMasked = IsOpaqueOrMaskedBlendMode(*CompilerMaterial);
-				Compiler->bSubstrateOutputsOpaqueRoughRefractions = bSubstrateUsesVerticalLayering && bIsOpaqueOrMasked;
-			}
-			bSubstrateMaterialIsUnlitNode = bHasUnlit;
-		}
-
-		//
-		// Make sure all the types have valid children operator indices according to the parameter blending setup for this simplification status.
-		//
-		for (const auto& It : SubstrateMaterialExpressionRegisteredOperators)
-		{
-			if (It.IsDiscarded())
-			{
-				continue; // ignore discarded operations in sub tree using parameter blending
-			}
-
-			check(It.Index != INDEX_NONE);
-
-			switch (It.OperatorType)
-			{
-			// Operators without any child
-			case SUBSTRATE_OPERATOR_BSDF:
-			{
-				if (!It.bUseParameterBlending)
-				{
-					// When using parameter blending, we need to keep indices to be able to recover local basis information for normal blending.
-					check(It.LeftIndex == INDEX_NONE && It.RightIndex == INDEX_NONE && It.BSDFIndex != INDEX_NONE);
-				}
-				break;
-			}
-
-			// Operators with two children
-			case SUBSTRATE_OPERATOR_HORIZONTAL:
-			case SUBSTRATE_OPERATOR_VERTICAL:
-			case SUBSTRATE_OPERATOR_ADD:
-			case SUBSTRATE_OPERATOR_SELECT:
-			{
-				check(It.RightIndex != INDEX_NONE);
-			}
-			// Fallthrough
-
-			// Operators with a single child
-			case SUBSTRATE_OPERATOR_WEIGHT:
-			{
-				check(It.LeftIndex != INDEX_NONE);
-			}
-			// Fallthrough
-			}
-		}
-
-		//
-		// Compute the maximum depth from the BSDF node for each operator
-		//
-		{
-			std::function<void(FSubstrateOperator&)>  WalkOperatorsToRoot = [&](FSubstrateOperator& CurrentOperator) -> void
-			{
-				switch (CurrentOperator.OperatorType)
-				{
-				case SUBSTRATE_OPERATOR_WEIGHT:
-				{
-					CurrentOperator.MaxDistanceFromLeaves = SubstrateMaterialExpressionRegisteredOperators[CurrentOperator.LeftIndex].MaxDistanceFromLeaves + 1;
-					break;
-				}
-				case SUBSTRATE_OPERATOR_VERTICAL:
-				case SUBSTRATE_OPERATOR_HORIZONTAL:
-				case SUBSTRATE_OPERATOR_ADD:
-				case SUBSTRATE_OPERATOR_SELECT:
-				{
-					CurrentOperator.MaxDistanceFromLeaves = FMath::Max(
-						SubstrateMaterialExpressionRegisteredOperators[CurrentOperator.LeftIndex].MaxDistanceFromLeaves,
-						SubstrateMaterialExpressionRegisteredOperators[CurrentOperator.RightIndex].MaxDistanceFromLeaves) + 1;
-					break;
-				}
-				case SUBSTRATE_OPERATOR_BSDF:
-				{
-					CurrentOperator.MaxDistanceFromLeaves = 0;
-					break;
-				}
-				}
-
-				if (CurrentOperator.ParentIndex != INDEX_NONE)
-				{
-					WalkOperatorsToRoot(SubstrateMaterialExpressionRegisteredOperators[CurrentOperator.ParentIndex]);
-				}
-			};
-
-			for (auto& It : SubstrateMaterialExpressionRegisteredOperators)
-			{
-				if (It.IsDiscarded())
-				{
-					continue; // ignore discarded operations in sub tree using parameter blending
-				}
-
-				if (It.OperatorType == SUBSTRATE_OPERATOR_BSDF)
-				{
-					// Recursively parse all nodes from BSDF to the root node and update the necessary properties.
-					WalkOperatorsToRoot(It);
-				}
-			}
-		}
-
-		//
-		// Compute IsTop or IsBottom layer using a depth first tree visit while counting vertical right and left branches taken.
-		// When a BSDF is encountered, and it is the root of a parameter blend subtree, we continue to parse, setting the same information as the root.
-		//
-		{
-			int VOpTopBranchCountTaken = 0;
-			int VOpBottomBranchCountTaken = 0;
-
-			std::function<void(FSubstrateOperator& , const FSubstrateOperator&)> PopulateParameterBlendedSubtree = [&](FSubstrateOperator& CurrentOperator, const FSubstrateOperator& RootOperator) -> void
-			{
-				switch (CurrentOperator.OperatorType)
-				{
-				case SUBSTRATE_OPERATOR_VERTICAL:
-				case SUBSTRATE_OPERATOR_HORIZONTAL:
-				case SUBSTRATE_OPERATOR_ADD:
-				case SUBSTRATE_OPERATOR_SELECT:
-				{
-					PopulateParameterBlendedSubtree(SubstrateMaterialExpressionRegisteredOperators[CurrentOperator.LeftIndex], RootOperator);
-					PopulateParameterBlendedSubtree(SubstrateMaterialExpressionRegisteredOperators[CurrentOperator.RightIndex], RootOperator);
-					break;
-				}
-				case SUBSTRATE_OPERATOR_WEIGHT:
-				{
-					PopulateParameterBlendedSubtree(SubstrateMaterialExpressionRegisteredOperators[CurrentOperator.LeftIndex], RootOperator);
-					break;
-				}
-				case SUBSTRATE_OPERATOR_BSDF:
-				{
-					CurrentOperator.LayerDepth = RootOperator.LayerDepth;
-					CurrentOperator.bIsTop = RootOperator.bIsTop;
-					CurrentOperator.bIsBottom = RootOperator.bIsBottom;
-					break;
-				}
-				}
-			};
-
-			std::function<void(FSubstrateOperator&)> WalkOperators = [&](FSubstrateOperator& CurrentOperator) -> void
-			{
-				switch (CurrentOperator.OperatorType)
-				{
-				case SUBSTRATE_OPERATOR_VERTICAL:
-				{
-					VOpTopBranchCountTaken++;
-					WalkOperators(SubstrateMaterialExpressionRegisteredOperators[CurrentOperator.LeftIndex]);
-					VOpTopBranchCountTaken--;
-					VOpBottomBranchCountTaken++;
-					WalkOperators(SubstrateMaterialExpressionRegisteredOperators[CurrentOperator.RightIndex]);
-					VOpBottomBranchCountTaken--;
-					break;
-				}
-				case SUBSTRATE_OPERATOR_HORIZONTAL:
-				case SUBSTRATE_OPERATOR_ADD:
-				case SUBSTRATE_OPERATOR_SELECT:
-				{
-					WalkOperators(SubstrateMaterialExpressionRegisteredOperators[CurrentOperator.LeftIndex]);
-					WalkOperators(SubstrateMaterialExpressionRegisteredOperators[CurrentOperator.RightIndex]);
-					break;
-				}
-				case SUBSTRATE_OPERATOR_WEIGHT:
-				{
-					WalkOperators(SubstrateMaterialExpressionRegisteredOperators[CurrentOperator.LeftIndex]);
-					break;
-				}
-				case SUBSTRATE_OPERATOR_BSDF:
-				{
-					const int32 VopCount = VOpTopBranchCountTaken + VOpBottomBranchCountTaken;
-					CurrentOperator.LayerDepth = VOpBottomBranchCountTaken;
-					CurrentOperator.bIsTop = VopCount == 0 || VOpTopBranchCountTaken == VopCount;
-					CurrentOperator.bIsBottom = VopCount == 0 || VOpBottomBranchCountTaken == VopCount;
-					if (CurrentOperator.bRootOfParameterBlendingSubTree)
-					{
-						// Make sure to also set up those values onto the BSDF within parameter blend sub trees.
-						if (CurrentOperator.LeftIndex != INDEX_NONE)
-						{
-							PopulateParameterBlendedSubtree(SubstrateMaterialExpressionRegisteredOperators[CurrentOperator.LeftIndex], CurrentOperator);
-						}
-						if (CurrentOperator.RightIndex != INDEX_NONE)
-						{
-							PopulateParameterBlendedSubtree(SubstrateMaterialExpressionRegisteredOperators[CurrentOperator.RightIndex], CurrentOperator);
-						}
-					}
-					break;
-				}
-				}
-			};
-
-			WalkOperators(*SubstrateMaterialRootOperator);
-		}
-
-		//
-		// Compute the per pixel byte count required by the materials.
-		//
-		{
-			// Compute the shared local basis count only
-			// But we cannot use SubstrateEvaluateSharedLocalBases here, because the material has not been compiled yet so all the bases would just default to the same.
-			// SUBSTRATE_TODO: can we do that material generation in two passes? 
-			//		1- A first one to evaluate the normal/tangent code
-			//		2- Operators are processed and simplification computed based on memory budget
-			//		3- Material is finally compiled for with operator updated to fit in memory budget.
-			uint8 UsedSharedLocalBasesCount = SubstrateMaterialEffectiveClosureCount;
-
-			const uint32 UintByteSize = sizeof(uint32);
-			SubstrateMaterialRequestedSizeByte = 0;
-
-			// 1. Evaluate simple/single BSDF
-			SubstrateMaterialComplexity.bIsSimple = SubstrateMaterialEffectiveClosureCount == 1;
-			SubstrateMaterialComplexity.bIsSingle = SubstrateMaterialEffectiveClosureCount == 1;
-			SubstrateMaterialComplexity.bIsComplexSpecial = false;
-			ESubstrateBsdfFeature SubstrateMaterialBsdfFeatures = ESubstrateBsdfFeature::None;
-			bool bIsFastWaterPath = false;
-			bool bCustomEncoding = false;
-			bool bMayHaveCoverageLessThan1 = false;
-			for (const auto& It : SubstrateMaterialExpressionRegisteredOperators)
-			{
-				if (It.IsDiscarded() && It.OperatorType == SUBSTRATE_OPERATOR_WEIGHT)
-				{
-					// If a BSDF modified by a weight operator, its weight will be < 1.0f, and it won't be a "single" material anymore.
-					// This is also the case if the operator is "discarded" due to parameter blending, because the resulting BSDF might not have a coverage of 1 in the end.
-					// For instance with a BSDF => Weight => Horizonal node with parameter blending is encountered and weight is less than 1.
-					bMayHaveCoverageLessThan1 |= true;
-				}
-
-				if (It.IsDiscarded())
-				{
-					continue; // ignore discarded operations in sub tree using parameter blending
-				}
-
-				switch (It.OperatorType)
-				{
-				case SUBSTRATE_OPERATOR_BSDF:
-				{
-					// From the compiler side, we can only assume the top layer has gray scale luminance weight.
-					const bool bMayHaveColoredWeight = !It.bIsTop;
-
-					// Aggregate all BSDFs features used by the material
-					SubstrateMaterialBsdfFeatures |= ESubstrateBsdfFeature(It.BSDFFeatures);
-
-					switch (It.BSDFType)
-					{
-					case SUBSTRATE_BSDF_TYPE_SLAB:
-					{
-						bool bIsComplexSpecial = false;
-						if (It.Has(ESubstrateBsdfFeature::ComplexSpecialMask))
-						{
-							// We need to check each features if really enabled for a platofrms
-							bIsComplexSpecial = It.Has(ESubstrateBsdfFeature::Glint) && Substrate::IsGlintEnabled(ShaderPlatform);
-						}
-
-						SubstrateMaterialComplexity.bIsSimple = SubstrateMaterialComplexity.bIsSimple && !bMayHaveColoredWeight && !It.Has(ESubstrateBsdfFeature::ComplexMask) && !bIsComplexSpecial && !It.Has(ESubstrateBsdfFeature::SingleMask);
-						SubstrateMaterialComplexity.bIsSingle = SubstrateMaterialComplexity.bIsSingle && !bMayHaveColoredWeight && !It.Has(ESubstrateBsdfFeature::ComplexMask) && !bIsComplexSpecial;
-						SubstrateMaterialComplexity.bIsComplexSpecial |= bIsComplexSpecial;
-						break;
-					}
-					case SUBSTRATE_BSDF_TYPE_HAIR:
-					{
-						SubstrateMaterialComplexity.bIsSimple = false;
-						SubstrateMaterialComplexity.bIsSingle = false;
-						bCustomEncoding			= true;
-						break;
-					}
-					case SUBSTRATE_BSDF_TYPE_EYE:
-					{
-						SubstrateMaterialComplexity.bIsSimple = false;
-						SubstrateMaterialComplexity.bIsSingle = false;
-						bCustomEncoding			= true;
-						break;
-					}
-					case SUBSTRATE_BSDF_TYPE_SINGLELAYERWATER:
-					{
-						SubstrateMaterialComplexity.bIsSimple = false;
-						SubstrateMaterialComplexity.bIsSingle = false;
-						bIsFastWaterPath		= true;
-						break;
-					}
-					case SUBSTRATE_BSDF_TYPE_TOON:
-					{
-						// Use the complex path for now.
-						SubstrateMaterialComplexity.bIsSimple = false;
-						SubstrateMaterialComplexity.bIsSingle = false;
-						break;
-					}
-					}
-					break;
-				}
-				case SUBSTRATE_OPERATOR_WEIGHT:
-				{
-					// If a BSDF modified by a weight operator, its weight will be < 1.0f, and it won't be a single material anymore
-					SubstrateMaterialComplexity.bIsSimple = false;
-					SubstrateMaterialComplexity.bIsSingle = false;
-					break;
-				}
-				}
-			}
-			SubstrateMaterialComplexity.bIsSingle = SubstrateMaterialComplexity.bIsSingle && !SubstrateMaterialComplexity.bIsSimple;
-
-			if (bMayHaveCoverageLessThan1)
-			{
-				// Material with coverage<1 can only be complex. It will be made simple or single at export time if possible.
-				SubstrateMaterialComplexity.bIsSimple = false;
-				SubstrateMaterialComplexity.bIsSingle = false;
-			}
-
-			// 2. Header
-
-			if (!SubstrateMaterialComplexity.bIsSimple && !SubstrateMaterialComplexity.bIsSingle && !bCustomEncoding && !bIsFastWaterPath) // header written later, 
-			{
-				// Packed Header
-				SubstrateMaterialRequestedSizeByte += UintByteSize;
-
-				// Shared local bases between BSDFs
-				SubstrateMaterialRequestedSizeByte += UsedSharedLocalBasesCount * SUBSTRATE_PACKED_SHAREDLOCALBASIS_STRIDE_BYTES;
-			}
-			// Note:
-			//  - We do not need to account for the Top Normal texture when evaluating the material byte count for the optimization algorithm.
-			//  - This is because we only need to optimize for the Substrate uint material buffer.
-
-			// 2. Process the list of BSDFs for worst case memory usage and count operators.
-			static_assert(SUBSTRATE_MAX_CLOSURE_COUNT_FOR_CLOSUREOFFSET	== (32u / SUBSTRATE_CLOSURE_OFFSET_BIT_COUNT));
-			static_assert(SUBSTRATE_MAX_CLOSURE_COUNT				    <= (1u << SUBSTRATE_CLOSURE_OFFSET_BIT_COUNT));
-			const uint32 ClosureMaxByteCountForOffset = uint32(1u << SUBSTRATE_CLOSURE_OFFSET_BIT_COUNT) * sizeof(uint32);
-			uint32 OperatorCount = 0;
-			SubstrateMaterialClosureCount = 0;
-			for (auto& It : SubstrateMaterialExpressionRegisteredOperators)
-			{
-				if (It.IsDiscarded())
-				{
-					continue; // ignore discarded operations in sub tree using parameter blending
-				}
-
-				// Operators are weight, vertical layering, etc. 
-				// Be aware that BSDFs also count as Operators when they are promoted from parameter blending!
-				OperatorCount++;
-
-				const uint32 PreSubstrateMaterialRequestedSizeByte = SubstrateMaterialRequestedSizeByte;
-				switch (It.OperatorType)
-				{
-				case SUBSTRATE_OPERATOR_BSDF:
-				{
-					// we have encountered a new BSDF which directly link to a single closure evaluation
-					SubstrateMaterialClosureCount++;
-
-					// From the compiler side, we can only assume the top layer has gray scale luminance weight.
-					const bool bMayHaveColoredWeight = !It.bIsTop;
-
-					if (SubstrateMaterialComplexity.bIsSimple)
-					{
-						// Header
-						SubstrateMaterialRequestedSizeByte += UintByteSize;
-						// Disney material
-						SubstrateMaterialRequestedSizeByte += UintByteSize;
-						break; // Stop here
-					}
-					else if (SubstrateMaterialComplexity.bIsSingle)
-					{
-						// Header
-						SubstrateMaterialRequestedSizeByte += UintByteSize;
-					}
-					else if (bCustomEncoding)
-					{
-						// Header
-						SubstrateMaterialRequestedSizeByte += UintByteSize;
-					}
-					else if (bIsFastWaterPath)
-					{
-						// Header + Data
-						SubstrateMaterialRequestedSizeByte += UintByteSize;
-						// Data
-						SubstrateMaterialRequestedSizeByte += UintByteSize;
-						break; // Stop here
-					}
-					else if (bMayHaveColoredWeight)
-					{
-						// BSDF state
-						SubstrateMaterialRequestedSizeByte += UintByteSize;
-						// Color weight
-						SubstrateMaterialRequestedSizeByte += UintByteSize;
-						// Light transmittance weight
-						SubstrateMaterialRequestedSizeByte += UintByteSize;
-					}
-					else
-					{
-						// BSDF state with gray scale weight
-						SubstrateMaterialRequestedSizeByte += UintByteSize;
-					}
-
-					switch (It.BSDFType)
-					{
-					case SUBSTRATE_BSDF_TYPE_SLAB:
-					{
-						// Compute values closer to the reality for HasSSS and IsSimpleVolume, now that we know that we know the topology of the material.
-						const bool bIsSimpleVolume = !It.bIsBottom && It.Has(ESubstrateBsdfFeature::MFPPluggedIn);
-						const bool bHasSSS = It.bIsBottom && It.Has(ESubstrateBsdfFeature::SSS);
-
-						SubstrateMaterialRequestedSizeByte += UintByteSize;
-						SubstrateMaterialRequestedSizeByte += UintByteSize;
-
-						// When using blendable GBuffer, do not account for features's byte requests. as downcasting is done SubstrateExport(). 
-						// Otherwise this would cause material to not compile as we do not demote feature during simplification.
-						if (Substrate::IsSubstrateBlendableGBufferEnabled(ShaderPlatform))
-						{
-							break;
-						}
-
-						if (It.Has(ESubstrateBsdfFeature::EdgeColor | ESubstrateBsdfFeature::SecondRoughnessOrSimpleClearCoat))
-						{
-							SubstrateMaterialRequestedSizeByte += UintByteSize;
-						}
-						if (bHasSSS || bIsSimpleVolume)
-						{
-							SubstrateMaterialRequestedSizeByte += UintByteSize;
-						}
-						if (It.Has(ESubstrateBsdfFeature::Fuzz))
-						{
-							SubstrateMaterialRequestedSizeByte += UintByteSize;
-						}
-						if (It.Has(ESubstrateBsdfFeature::Glint) && Substrate::IsGlintEnabled(ShaderPlatform))
-						{
-							SubstrateMaterialRequestedSizeByte += UintByteSize;
-							SubstrateMaterialRequestedSizeByte += UintByteSize;
-							SubstrateMaterialRequestedSizeByte += UintByteSize;
-							SubstrateMaterialRequestedSizeByte += UintByteSize;
-						}
-						if (It.Has(ESubstrateBsdfFeature::SpecularProfile))
-						{
-							SubstrateMaterialRequestedSizeByte += UintByteSize;
-						}
-						if (It.Has(ESubstrateBsdfFeature::Eye))
-						{
-							SubstrateMaterialRequestedSizeByte += UintByteSize;
-							SubstrateMaterialRequestedSizeByte += UintByteSize;
-						}
-						if (It.Has(ESubstrateBsdfFeature::Hair))
-						{
-							SubstrateMaterialRequestedSizeByte += UintByteSize;
-							SubstrateMaterialRequestedSizeByte += UintByteSize;
-						}
-						break;
-					}
-					case SUBSTRATE_BSDF_TYPE_HAIR:
-					{
-						// Custom encoding
-						SubstrateMaterialRequestedSizeByte += UintByteSize;
-						SubstrateMaterialRequestedSizeByte += UintByteSize;
-						break;
-					}
-					case SUBSTRATE_BSDF_TYPE_EYE:
-					{
-						// Custom encoding
-						SubstrateMaterialRequestedSizeByte += UintByteSize;
-						SubstrateMaterialRequestedSizeByte += UintByteSize;
-						break;
-					}
-					case SUBSTRATE_BSDF_TYPE_SINGLELAYERWATER:
-					{
-						Compiler->Errorf(TEXT("Substrate error: single layer water should go through through a dedicated fast path in %s (asset: %s).\r\n"), *CompilerMaterial->GetDebugName(), *CompilerMaterial->GetAssetPath().ToString());
-						return false;
-					}
-					case SUBSTRATE_BSDF_TYPE_TOON:
-					{
-						SubstrateMaterialRequestedSizeByte += UintByteSize;
-						SubstrateMaterialRequestedSizeByte += UintByteSize;
-						// We cannot add more uint here to stay compatible with legacy gbuffer when Substrate::IsSubstrateBlendableGBufferEnabled(ShaderPlatform).
-						// Possible work around: disable some toon shading feature that cannot be stored in the GBuffer, or increase Custom GBufferD size to 64 bits.
-						break;
-					}
-					case SUBSTRATE_BSDF_TYPE_UNLIT:
-					{
-						// Never stored, it goes directly into the scene as emitted luminance.
-						break;
-					}
-					default:
-					{
-						Compiler->Errorf(TEXT("Unkownd BSDF type encountered in %s (asset: %s).\r\n"), *CompilerMaterial->GetDebugName(), *CompilerMaterial->GetAssetPath().ToString());
-						return false;
-					}
-					}
-					break;
-				} // case SUBSTRATE_OPERATOR_BSDF
-				} // switch (It.OperatorType)
-			}
-
-			if (OperatorCount > SUBSTRATE_MAX_OPERATOR_COUNT)
-			{
-				// Why do we have an operator limit: due to the size of the array of Operator in FSubstrateTre and the way the Substrate tree is exported for advanced debug purpose. Parameter blending can help working around that. 
-				Compiler->Errorf(TEXT("Material %s have too many Substrate Operators (asset: %s): %d / %d. Please note that BSDFs also count as an operator. Use parameter blending to workaround that limitation.\r\n"), *CompilerMaterial->GetDebugName(), *CompilerMaterial->GetAssetPath().ToString(), OperatorCount, SUBSTRATE_MAX_OPERATOR_COUNT);
-				return false;
-			}
-
-			SubstrateSimplificationStatus.bMaterialFitsInMemoryBudget = (SubstrateMaterialRequestedSizeByte <= SubstrateBytePerPixel) && (SubstrateMaterialClosureCount <= SubstrateClosurePerPixel);
-			if (!SubstrateSimplificationStatus.bMaterialFitsInMemoryBudget && SubstrateSimplificationStatus.bFullSimplificationStepHasBeenRun && SubstrateSimplificationStatus.bSlabSimplificationStepHasBeenRun)
-			{
-				// If we have already run the full simplification but the material still does not fit in memory, we must fail the material compilation.
-				Compiler->Errorf(TEXT("Material %s could not be simplified to fit in Substrate per pixel (asset: %s).\r\n"), *CompilerMaterial->GetDebugName(), *CompilerMaterial->GetAssetPath().ToString());
-				return false;
-			}
-			if (!bRequestMaterialDetailsSet)
-			{
-				// Record the original requested byte size before simplification, only for the first pass.
-				SubstrateSimplificationStatus.OriginalRequestedByteSize = SubstrateMaterialRequestedSizeByte;
-				SubstrateSimplificationStatus.OriginalRequestedClosureCount = SubstrateMaterialClosureCount;
-				bRequestMaterialDetailsSet = true;
-			}
-			SubstrateSimplificationStatus.bFullSimplificationStepHasBeenRun |= SubstrateSimplificationStatus.bRunFullSimplification;
-
-			const uint32 RequestedSizeInUint = FMath::DivideAndRoundUp(SubstrateMaterialRequestedSizeByte, 4u);
-			check(RequestedSizeInUint < 256u);
-
-			if (CompilationContextIndex == ESubstrateCompilationContext::SCC_Default)
-			{
-				// Only write those data for the default material
-				Compiler->MaterialCompilationOutput.SubstrateMaterialCompilationOutput.SubstrateMaterialType = SubstrateMaterialComplexity.SubstrateMaterialType(); 
-				Compiler->MaterialCompilationOutput.SubstrateMaterialCompilationOutput.SubstrateClosureCount = SubstrateMaterialEffectiveClosureCount;
-				Compiler->MaterialCompilationOutput.SubstrateMaterialCompilationOutput.SubstrateUintPerPixel = uint8(FMath::Clamp(RequestedSizeInUint, 0u, 0xFF));
-				Compiler->MaterialCompilationOutput.SubstrateMaterialCompilationOutput.SubstrateMaterialBsdfFeatures = SubstrateMaterialBsdfFeatures;
-
-#if WITH_EDITOR
-				Compiler->MaterialCompilationOutput.SubstrateMaterialCompilationOutput.SharedLocalBasesCount = 0; // FinalUsedSharedLocalBasesCount is not valid yet
-				Compiler->MaterialCompilationOutput.SubstrateMaterialCompilationOutput.RequestedBytePerPixel = SubstrateSimplificationStatus.OriginalRequestedByteSize;
-				Compiler->MaterialCompilationOutput.SubstrateMaterialCompilationOutput.PlatformBytePerPixel = SubstrateBytePerPixel;
-
-				Compiler->MaterialCompilationOutput.SubstrateMaterialCompilationOutput.RequestedClosurePerPixel = SubstrateSimplificationStatus.OriginalRequestedClosureCount;
-				Compiler->MaterialCompilationOutput.SubstrateMaterialCompilationOutput.PlatformClosurePixel = SubstrateClosurePerPixel;
-
-				Compiler->MaterialCompilationOutput.SubstrateMaterialCompilationOutput.bIsThin = CompilerMaterial->IsThinSurface() ? 1 : 0;
-
-				uint32 BSDFTypeCount = FMath::CountBits(SubstrateTreeValidation.MaterialType & MTV_ValidMultipleBSDFs);
-
-				// The order of ifs here is important.
-				if (CompilerMaterial->IsLightFunction())
-				{
-					Compiler->MaterialCompilationOutput.SubstrateMaterialCompilationOutput.MaterialType = SUBSTRATE_MATERIAL_TYPE_LIGHTFUNCTION;
-				}
-				else if (CompilerMaterial->IsPostProcessMaterial())
-				{
-					Compiler->MaterialCompilationOutput.SubstrateMaterialCompilationOutput.MaterialType = SUBSTRATE_MATERIAL_TYPE_POSTPROCESS;
-				}
-				else if (CompilerMaterial->IsUIMaterial())
-				{
-					Compiler->MaterialCompilationOutput.SubstrateMaterialCompilationOutput.MaterialType = SUBSTRATE_MATERIAL_TYPE_UI;
-				}
-				else if (CompilerMaterial->IsDeferredDecal())
-				{
-					Compiler->MaterialCompilationOutput.SubstrateMaterialCompilationOutput.MaterialType = SUBSTRATE_MATERIAL_TYPE_DECAL;
-				}
-				else if (SubstrateMaterialEffectiveClosureCount > 1 && (SubstrateTreeValidation.MaterialType == (SubstrateTreeValidation.MaterialType & MTV_BSDF_Slab)))
-				{
-					// Only slabs, not hair nor toon.
-					Compiler->MaterialCompilationOutput.SubstrateMaterialCompilationOutput.MaterialType = SUBSTRATE_MATERIAL_TYPE_MULTIPLESLABS;
-				}
-				else if (SubstrateMaterialEffectiveClosureCount > 1) // Only slabs can have multiple closures because the select nodde will always parameter blend slabs into a single one
-				{
-					Compiler->MaterialCompilationOutput.SubstrateMaterialCompilationOutput.MaterialType = SUBSTRATE_MATERIAL_TYPE_MULTIPLESLABS;
-				}
-				else if (BSDFTypeCount > 1)
-				{
-					Compiler->MaterialCompilationOutput.SubstrateMaterialCompilationOutput.MaterialType = SUBSTRATE_MATERIAL_TYPE_MULTIPLEBSDFS;
-				}
-				else if (bHasUnlit)
-				{
-					Compiler->MaterialCompilationOutput.SubstrateMaterialCompilationOutput.MaterialType = SUBSTRATE_MATERIAL_TYPE_UNLIT;
-				}
-				else if (bHasVFogCloud)
-				{
-					Compiler->MaterialCompilationOutput.SubstrateMaterialCompilationOutput.MaterialType = SUBSTRATE_MATERIAL_TYPE_VOLUMETRICFOGCLOUD;
-				}
-				else if (bHasHair)
-				{
-					Compiler->MaterialCompilationOutput.SubstrateMaterialCompilationOutput.MaterialType = SUBSTRATE_MATERIAL_TYPE_HAIR;
-				}
-				else if (bHasEye)
-				{
-					Compiler->MaterialCompilationOutput.SubstrateMaterialCompilationOutput.MaterialType = SUBSTRATE_MATERIAL_TYPE_EYE;
-				}
-				else if (bHasSLW)
-				{
-					Compiler->MaterialCompilationOutput.SubstrateMaterialCompilationOutput.MaterialType = SUBSTRATE_MATERIAL_TYPE_SINGLELAYERWATER;
-				}
-				else if (bHasToon)
-				{
-					Compiler->MaterialCompilationOutput.SubstrateMaterialCompilationOutput.MaterialType = SUBSTRATE_MATERIAL_TYPE_TOON;
-				}
-				else
-				{
-					Compiler->MaterialCompilationOutput.SubstrateMaterialCompilationOutput.MaterialType = SUBSTRATE_MATERIAL_TYPE_SINGLESLAB;
-				}
-
-				Compiler->MaterialCompilationOutput.SubstrateMaterialCompilationOutput.bMaterialOutOfBudgetHasBeenSimplified |= !SubstrateSimplificationStatus.bMaterialFitsInMemoryBudget;
+	return SubstrateTranslatorData.SubstrateCompilationRegisterOperator(OperatorType, SubstrateExpressionGuid, ChildMaterialExpressionGuid, Parent, SubstrateParentExpressionGuid, bUseParameterBlending);
+}
 
-				if (OperatorCount <= SUBSTRATE_COMPILATION_OUTPUT_MAX_OPERATOR)
-				{
-					int32 OperatorIndex = 0;
-					for (auto& It : SubstrateMaterialExpressionRegisteredOperators)
-					{
-						Compiler->MaterialCompilationOutput.SubstrateMaterialCompilationOutput.Operators[OperatorIndex] = SubstrateMaterialExpressionRegisteredOperators[OperatorIndex];
-						OperatorIndex++;
-					}
-				}
-				Compiler->MaterialCompilationOutput.SubstrateMaterialCompilationOutput.RootOperatorIndex = RootIndex;
-#endif // EDITOR_ONLY
-			}
-		}
-	} while (!SubstrateSimplificationStatus.bMaterialFitsInMemoryBudget);
+FSubstrateOperator& FHLSLMaterialTranslator::SubstrateCompilationGetOperator(FGuid SubstrateExpressionGuid)
+{
+	return SubstrateTranslatorData.SubstrateCompilationGetOperator(SubstrateExpressionGuid);
+}
 
-	return true; // Success
+FSubstrateOperator* FHLSLMaterialTranslator::SubstrateCompilationGetOperatorFromIndex(int32 OperatorIndex)
+{
+	return SubstrateTranslatorData.SubstrateCompilationGetOperatorFromIndex(OperatorIndex);
 }
 
 FSubstrateRegisteredSharedLocalBasis FHLSLMaterialTranslator::SubstrateCompilationInfoRegisterSharedLocalBasis(int32 NormalCodeChunk)
 {
-	FSubstrateCompilationContext& SubstrateCtx = SubstrateCompilationContext[CurrentSubstrateCompilationContext];
+	// Specific implementation for FHLSLMaterialTranslator on the common FSubstrateTranslatorData
+
+	FSubstrateCompilationContext& SubstrateCtx = SubstrateTranslatorData.SubstrateCompilationContext[SubstrateTranslatorData.CurrentSubstrateCompilationContext];
 	check(NormalCodeChunk != INDEX_NONE);
 	check(SubstrateCtx.NextFreeSubstrateShaderNormalIndex < 255);	// Out of shared local basis slots
 
@@ -14812,7 +12751,9 @@ FSubstrateRegisteredSharedLocalBasis FHLSLMaterialTranslator::SubstrateCompilati
 
 FSubstrateRegisteredSharedLocalBasis FHLSLMaterialTranslator::SubstrateCompilationInfoRegisterSharedLocalBasis(int32 NormalCodeChunk, int32 TangentCodeChunk)
 {
-	FSubstrateCompilationContext& SubstrateCtx = SubstrateCompilationContext[CurrentSubstrateCompilationContext];
+	// Specific implementation for FHLSLMaterialTranslator on the common FSubstrateTranslatorData
+
+	FSubstrateCompilationContext& SubstrateCtx = SubstrateTranslatorData.SubstrateCompilationContext[SubstrateTranslatorData.CurrentSubstrateCompilationContext];
 	check(NormalCodeChunk != INDEX_NONE);
 	check(SubstrateCtx.NextFreeSubstrateShaderNormalIndex < 255);	// Out of shared local basis slots
 	
@@ -14851,29 +12792,7 @@ FSubstrateRegisteredSharedLocalBasis FHLSLMaterialTranslator::SubstrateCompilati
 
 FString FHLSLMaterialTranslator::GetSubstrateSharedLocalBasisIndexMacro(const FSubstrateRegisteredSharedLocalBasis& SharedLocalBasis)
 {
-	return GetSubstrateSharedLocalBasisIndexMacroInner(SharedLocalBasis, CurrentSubstrateCompilationContext);
-}
-
-FHLSLMaterialTranslator::FSubstrateSharedLocalBasesInfo FHLSLMaterialTranslator::FSubstrateCompilationContext::SubstrateCompilationInfoGetMatchingSharedLocalBasisInfo(const FSubstrateRegisteredSharedLocalBasis& SearchedSharedLocalBasis)
-{
-	check(NextFreeSubstrateShaderNormalIndex < 255);	// Out of shared local basis slots
-
-	// Find a basis which matches both the Normal & the Tangent code chunks
-	TArray<const FSubstrateSharedLocalBasesInfo*> NormalInfos;
-	CodeChunkToSubstrateSharedLocalBasis.MultiFindPointer(SearchedSharedLocalBasis.NormalCodeChunkHash, NormalInfos);
-
-	// We first try to find a perfect match for normal and tangent from all the registered element.
-	for (const FSubstrateSharedLocalBasesInfo* NormalInfo : NormalInfos)
-	{
-		if (SearchedSharedLocalBasis.TangentCodeChunk == INDEX_NONE ||											// We selected the first available normal if there is no tangent specified on the material.
-			SearchedSharedLocalBasis.TangentCodeChunkHash == NormalInfo->SharedData.TangentCodeChunkHash)		// Otherwise we select the normal+tangent that exactly matches the request.
-		{
-			return *NormalInfo;
-		}
-	}
-
-	check(0);	// When the compiler is querying, this is to get a result to generate code from a fully processed graph. No result means a bug happened during graph processing.
-	return FSubstrateSharedLocalBasesInfo();
+	return Substrate::GetSubstrateSharedLocalBasisIndexMacroInner(SharedLocalBasis, SubstrateTranslatorData.CurrentSubstrateCompilationContext);
 }
 
 int32 FHLSLMaterialTranslator::SubstrateAddParameterBlendingBSDFCoverageToNormalMixCodeChunk(int32 ACodeChunk, int32 BCodeChunk)
@@ -14989,7 +12908,7 @@ int32 FHLSLMaterialTranslator::SubstrateSlabBSDF(
 				/*LocalBasisIndex*/                  %s,\n\
 				Parameters.%s.Types)\n\
 				 /* Normal = %s ; Tangent = %s ; Thickness = %s */, %u, %u, %u, %u)"),
-			*GetParametersSubstrateTreeName(CurrentSubstrateCompilationContext),
+			*Substrate::GetParametersSubstrateTreeName(SubstrateTranslatorData.CurrentSubstrateCompilationContext),
 			*SubstrateGetCastParameterCode(Normal,					MCT_Float3),
 			*SubstrateGetCastParameterCode(DiffuseAlbedo,			MCT_Float3),
 			*SubstrateGetCastParameterCode(F0,						MCT_Float3),
@@ -15017,7 +12936,7 @@ int32 FHLSLMaterialTranslator::SubstrateSlabBSDF(
 			bIsThinSurface ? TEXT("true") : TEXT("false"),
 			bIsAtTheBottomOfTopology ? TEXT("true") : TEXT("false"),
 			*SharedLocalBasisIndexMacro,
-			*GetParametersSharedLocalBasesName(CurrentSubstrateCompilationContext),
+			*Substrate::GetParametersSharedLocalBasesName(SubstrateTranslatorData.CurrentSubstrateCompilationContext),
 			*NormalCode,
 			*TangentCode,
 			*ThicknessCode,
@@ -15086,7 +13005,7 @@ int32 FHLSLMaterialTranslator::SubstrateSlabBSDF(
 		bIsThinSurface ? TEXT("true") : TEXT("false"),
 		bIsAtTheBottomOfTopology ? TEXT("true") : TEXT("false"),
 		*SharedLocalBasisIndexMacro,
-		*GetParametersSharedLocalBasesName(CurrentSubstrateCompilationContext),
+		*Substrate::GetParametersSharedLocalBasesName(SubstrateTranslatorData.CurrentSubstrateCompilationContext),
 		*NormalCode,
 		*TangentCode,
 		*ThicknessCode
@@ -15171,7 +13090,7 @@ int32 FHLSLMaterialTranslator::SubstrateConversionFromLegacy(
 				/*ClearCoat_SharedLocalBasisIndex*/   %s,\n\
 				Parameters.%s.Types)\n\
 				/* Normal = %s ; Tangent = %s ; ClearCoat_Normal = %s ; ClearCoat_Tangent = %s */, %u, %u, %u, %u)"),
-			*GetParametersSubstrateTreeName(CurrentSubstrateCompilationContext),
+			*Substrate::GetParametersSubstrateTreeName(SubstrateTranslatorData.CurrentSubstrateCompilationContext),
 			bHasDynamicShadingModels ? TEXT("Dynamic") : TEXT("Static"),
 			*SubstrateGetCastParameterCode(BaseColor,							MCT_Float3),
 			*SubstrateGetCastParameterCode(Specular,							MCT_Float),
@@ -15198,7 +13117,7 @@ int32 FHLSLMaterialTranslator::SubstrateConversionFromLegacy(
 			*SubstrateGetCastParameterCode(CustomTangent_Tangent,				MCT_Float3),
 			*SharedLocalBasisIndexMacro,
 			*ClearCoat_SharedLocalBasisIndexMacro,
-			*GetParametersSharedLocalBasesName(CurrentSubstrateCompilationContext),
+			*Substrate::GetParametersSharedLocalBasesName(SubstrateTranslatorData.CurrentSubstrateCompilationContext),
 			// Regular normal basis
 			*NormalCode,
 			*TangentCode,
@@ -15266,7 +13185,7 @@ int32 FHLSLMaterialTranslator::SubstrateConversionFromLegacy(
 		*SubstrateGetCastParameterCode(CustomTangent_Tangent,				MCT_Float3),
 		*SharedLocalBasisIndexMacro,
 		*ClearCoat_SharedLocalBasisIndexMacro,
-		*GetParametersSharedLocalBasesName(CurrentSubstrateCompilationContext),
+		*Substrate::GetParametersSharedLocalBasesName(SubstrateTranslatorData.CurrentSubstrateCompilationContext),
 		// Regular normal basis
 		*NormalCode,
 		*TangentCode,
@@ -15321,7 +13240,7 @@ int32 FHLSLMaterialTranslator::SubstrateUnlitBSDF(int32 EmissiveColor, int32 Tra
 	{
 		return AddCodeChunk(
 			MCT_Substrate, TEXT("Parameters.%s.PromoteParameterBlendedBSDFToOperator(GetSubstrateUnlitBSDF(%s, %s, %s), %u, %u, %u, %u)"),
-			*GetParametersSubstrateTreeName(CurrentSubstrateCompilationContext),
+			*Substrate::GetParametersSubstrateTreeName(SubstrateTranslatorData.CurrentSubstrateCompilationContext),
 			*SubstrateGetCastParameterCode(EmissiveColor,		MCT_Float3),
 			*SubstrateGetCastParameterCode(TransmittanceColor,	MCT_Float3),
 			*SubstrateGetCastParameterCode(Normal,				MCT_Float3),
@@ -15345,7 +13264,7 @@ int32 FHLSLMaterialTranslator::SubstrateUIBSDF(int32 EmissiveColor, int32 Opacit
 	{
 		return AddCodeChunk(
 			MCT_Substrate, TEXT("Parameters.%s.PromoteParameterBlendedBSDFToOperator(SubstrateCreateUIMaterial(%s, %s), %u, %u, %u, %u)"),
-			*GetParametersSubstrateTreeName(CurrentSubstrateCompilationContext),
+			*Substrate::GetParametersSubstrateTreeName(SubstrateTranslatorData.CurrentSubstrateCompilationContext),
 			*SubstrateGetCastParameterCode(EmissiveColor, MCT_Float3),
 			*SubstrateGetCastParameterCode(Opacity, MCT_Float1),
 			PromoteToOperator->Index,
@@ -15367,7 +13286,7 @@ int32 FHLSLMaterialTranslator::SubstrateHairBSDF(int32 BaseColor, int32 Scatter,
 	{
 		return AddCodeChunk(
 			MCT_Substrate, TEXT("Parameters.%s.PromoteParameterBlendedBSDFToOperator(GetSubstrateHairBSDF(%s, %s, %s, %s, %s, %s, false/*bComplexTransmittance*/, %s), %u, %u, %u, %u) /* Tangent:%s */"),
-			*GetParametersSubstrateTreeName(CurrentSubstrateCompilationContext),
+			*Substrate::GetParametersSubstrateTreeName(SubstrateTranslatorData.CurrentSubstrateCompilationContext),
 			*SubstrateGetCastParameterCode(BaseColor, MCT_Float3),
 			*SubstrateGetCastParameterCode(Scatter, MCT_Float),
 			*SubstrateGetCastParameterCode(Specular, MCT_Float),
@@ -15399,7 +13318,7 @@ int32 FHLSLMaterialTranslator::SubstrateEyeBSDF(int32 DiffuseAlbedo, int32 Rough
 {
 	return AddCodeChunk(
 		MCT_Substrate, TEXT("Parameters.%s.PromoteParameterBlendedBSDFToOperator(GetSubstrateEyeBSDF(%s, %s, %s, %s, %s, %s, ExtractSubsurfaceProfileInt(%s), %s, %s), %u, %u, %u, %u) /* Cornea:%s Iris:%s */"),
-		*GetParametersSubstrateTreeName(CurrentSubstrateCompilationContext),
+		*Substrate::GetParametersSubstrateTreeName(SubstrateTranslatorData.CurrentSubstrateCompilationContext),
 		*SubstrateGetCastParameterCode(DiffuseAlbedo, MCT_Float3),
 		*SubstrateGetCastParameterCode(Roughness, MCT_Float),
 		*SubstrateGetCastParameterCode(IrisMask, MCT_Float),
@@ -15425,7 +13344,7 @@ int32 FHLSLMaterialTranslator::SubstrateSingleLayerWaterBSDF(
 {
 	return AddCodeChunk(
 		MCT_Substrate, TEXT("Parameters.%s.PromoteParameterBlendedBSDFToOperator(GetSubstrateSingleLayerWaterBSDF(%s, %s, %s, %s, %s, %s, %s, %s, %s, %s, %s), %u, %u, %u, %u) /* Normal:%s */"),
-		*GetParametersSubstrateTreeName(CurrentSubstrateCompilationContext),
+		*Substrate::GetParametersSubstrateTreeName(SubstrateTranslatorData.CurrentSubstrateCompilationContext),
 		*SubstrateGetCastParameterCode(BaseColor,				MCT_Float3),
 		*SubstrateGetCastParameterCode(Metallic,				MCT_Float),
 		*SubstrateGetCastParameterCode(Specular,				MCT_Float),
@@ -15460,7 +13379,7 @@ int32 FHLSLMaterialTranslator::SubstrateToonBSDF(
 		}
 		return AddCodeChunk(
 			MCT_Substrate, TEXT("Parameters.%s.PromoteParameterBlendedBSDFToOperator(GetSubstrateToonBSDF(%s, %s, %s, %s, %s, %s, %s, Parameters.%s.Types), %u, %u, %u, %u) /* Normal:%s */"),
-			*GetParametersSubstrateTreeName(CurrentSubstrateCompilationContext),
+			*Substrate::GetParametersSubstrateTreeName(SubstrateTranslatorData.CurrentSubstrateCompilationContext),
 			*SubstrateGetCastParameterCode(ToonProfileId, MCT_Float),
 			*SubstrateGetCastParameterCode(DiffuseColor, MCT_Float3),
 			*SubstrateGetCastParameterCode(SpecularColor, MCT_Float3),
@@ -15468,7 +13387,7 @@ int32 FHLSLMaterialTranslator::SubstrateToonBSDF(
 			*SubstrateGetCastParameterCode(Anisotropy, MCT_Float),
 			*SubstrateGetCastParameterCode(EmissiveColor, MCT_Float3),
 			*SharedLocalBasisIndexMacro,
-			*GetParametersSharedLocalBasesName(CurrentSubstrateCompilationContext),
+			*Substrate::GetParametersSharedLocalBasesName(SubstrateTranslatorData.CurrentSubstrateCompilationContext),
 			PromoteToOperator->Index,
 			PromoteToOperator->BSDFIndex,
 			PromoteToOperator->LayerDepth,
@@ -15486,7 +13405,7 @@ int32 FHLSLMaterialTranslator::SubstrateToonBSDF(
 		*SubstrateGetCastParameterCode(Anisotropy, MCT_Float),
 		*SubstrateGetCastParameterCode(EmissiveColor, MCT_Float3),
 		*SharedLocalBasisIndexMacro,
-		*GetParametersSharedLocalBasesName(CurrentSubstrateCompilationContext)
+		*Substrate::GetParametersSharedLocalBasesName(SubstrateTranslatorData.CurrentSubstrateCompilationContext)
 	);
 }
 
@@ -15498,7 +13417,7 @@ int32 FHLSLMaterialTranslator::SubstrateHorizontalMixing(int32 Background, int32
 	}
 	return AddCodeChunk(
 		MCT_Substrate, TEXT("Parameters.%s.SubstrateHorizontalMixing(%s, %s, %s, %u, %u)"),
-		*GetParametersSubstrateTreeName(CurrentSubstrateCompilationContext),
+		*Substrate::GetParametersSubstrateTreeName(SubstrateTranslatorData.CurrentSubstrateCompilationContext),
 		*GetParameterCode(Background),
 		*GetParameterCode(Foreground),
 		*GetParameterCode(Mix),
@@ -15530,13 +13449,13 @@ int32 FHLSLMaterialTranslator::SubstrateHorizontalMixingParameterBlending(
 		check(PromoteToOperator->BSDFIndex != INDEX_NONE);
 		return AddCodeChunk(
 			MCT_Substrate, TEXT("Parameters.%s.PromoteParameterBlendedBSDFToOperator(SubstrateHorizontalMixingParameterBlending(%s, %s, %s, %s, %s, Parameters.%s.Types, %s, %s), %u, %u, %u, %u)"),
-			*GetParametersSubstrateTreeName(CurrentSubstrateCompilationContext),
+			*Substrate::GetParametersSubstrateTreeName(SubstrateTranslatorData.CurrentSubstrateCompilationContext),
 			*GetParameterCode(Background),
 			*GetParameterCode(Foreground),
 			*GetParameterCode(HorizontalMixCodeChunk),
 			*GetParameterCode(NormalMixCodeChunk),
 			*SharedLocalBasisIndexMacro,
-			*GetParametersSharedLocalBasesName(CurrentSubstrateCompilationContext),
+			*Substrate::GetParametersSharedLocalBasesName(SubstrateTranslatorData.CurrentSubstrateCompilationContext),
 			*GetParameterCode(BackgroundNoV),
 			*GetParameterCode(ForegroundNoV),
 			PromoteToOperator->Index,
@@ -15553,7 +13472,7 @@ int32 FHLSLMaterialTranslator::SubstrateHorizontalMixingParameterBlending(
 		*GetParameterCode(HorizontalMixCodeChunk),
 		*GetParameterCode(NormalMixCodeChunk),
 		*SharedLocalBasisIndexMacro,
-		*GetParametersSharedLocalBasesName(CurrentSubstrateCompilationContext),
+		*Substrate::GetParametersSharedLocalBasesName(SubstrateTranslatorData.CurrentSubstrateCompilationContext),
 		*GetParameterCode(BackgroundNoV),
 		*GetParameterCode(ForegroundNoV)
 	);
@@ -15568,7 +13487,7 @@ int32 FHLSLMaterialTranslator::SubstrateVerticalLayering(int32 Top, int32 Base,
 	const FString ThicknessCode = Thickness != INDEX_NONE ? GetParameterCode(Thickness) : TEXT("NONE");;
 	return AddCodeChunk(
 		MCT_Substrate, TEXT("Parameters.%s.SubstrateVerticalLayering(%s, %s, %u, %u) /* Thickness = %s */"),
-		*GetParametersSubstrateTreeName(CurrentSubstrateCompilationContext),
+		*Substrate::GetParametersSubstrateTreeName(SubstrateTranslatorData.CurrentSubstrateCompilationContext),
 		*GetParameterCode(Top),
 		*GetParameterCode(Base),
 		OperatorIndex,
@@ -15601,7 +13520,7 @@ int32 FHLSLMaterialTranslator::SubstrateVerticalLayeringParameterBlending(int32
 		check(PromoteToOperator->BSDFIndex != INDEX_NONE);
 		return AddCodeChunk(
 			MCT_Substrate, TEXT("Parameters.%s.PromoteParameterBlendedBSDFToOperator(SubstrateVerticalLayeringParameterBlending(%s, %s, %s, %s, %s), %u, %u, %u, %u) /* Thickness = %s */"),
-			*GetParametersSubstrateTreeName(CurrentSubstrateCompilationContext),
+			*Substrate::GetParametersSubstrateTreeName(SubstrateTranslatorData.CurrentSubstrateCompilationContext),
 			*GetParameterCode(Top),
 			*GetParameterCode(Base),
 			*SharedLocalBasisIndexMacro,
@@ -15634,7 +13553,7 @@ int32 FHLSLMaterialTranslator::SubstrateAdd(int32 A, int32 B, int OperatorIndex,
 	}
 	return AddCodeChunk(
 		MCT_Substrate, TEXT("Parameters.%s.SubstrateAdd(%s, %s, %u, %u)"),
-		*GetParametersSubstrateTreeName(CurrentSubstrateCompilationContext),
+		*Substrate::GetParametersSubstrateTreeName(SubstrateTranslatorData.CurrentSubstrateCompilationContext),
 		*GetParameterCode(A),
 		*GetParameterCode(B),
 		OperatorIndex,
@@ -15664,7 +13583,7 @@ int32 FHLSLMaterialTranslator::SubstrateAddParameterBlending(int32 A, int32 B, i
 		check(PromoteToOperator->BSDFIndex != INDEX_NONE);
 		return AddCodeChunk(
 			MCT_Substrate, TEXT("Parameters.%s.PromoteParameterBlendedBSDFToOperator(SubstrateAddParameterBlending(%s, %s, %s, %s, %s, %s), %u, %u, %u, %u)"),
-			*GetParametersSubstrateTreeName(CurrentSubstrateCompilationContext),
+			*Substrate::GetParametersSubstrateTreeName(SubstrateTranslatorData.CurrentSubstrateCompilationContext),
 			*GetParameterCode(A),
 			*GetParameterCode(B),
 			*GetParameterCode(AMixWeight),
@@ -15697,7 +13616,7 @@ int32 FHLSLMaterialTranslator::SubstrateWeight(int32 A, int32 Weight, int Operat
 	}
 	return AddCodeChunk(
 		MCT_Substrate, TEXT("Parameters.%s.SubstrateWeight(%s, %s, %u, %u)"),
-		*GetParametersSubstrateTreeName(CurrentSubstrateCompilationContext),
+		*Substrate::GetParametersSubstrateTreeName(SubstrateTranslatorData.CurrentSubstrateCompilationContext),
 		*GetParameterCode(A),
 		*GetParameterCode(Weight),
 		OperatorIndex,
@@ -15718,7 +13637,7 @@ int32 FHLSLMaterialTranslator::SubstrateWeightParameterBlending(int32 A, int32 W
 		check(PromoteToOperator->BSDFIndex != INDEX_NONE);
 		return AddCodeChunk(
 			MCT_Substrate, TEXT("Parameters.%s.PromoteParameterBlendedBSDFToOperator(SubstrateWeightParameterBlending(%s, %s), %u, %u, %u, %u)"),
-			*GetParametersSubstrateTreeName(CurrentSubstrateCompilationContext),
+			*Substrate::GetParametersSubstrateTreeName(SubstrateTranslatorData.CurrentSubstrateCompilationContext),
 			*GetParameterCode(A),
 			*GetParameterCode(Weight),
 			PromoteToOperator->Index,
@@ -15748,12 +13667,12 @@ int32 FHLSLMaterialTranslator::SubstrateSelectParameterBlending(int32 A, int32 B
 		check(PromoteToOperator->BSDFIndex != INDEX_NONE);
 		return AddCodeChunk(
 			MCT_Substrate, TEXT("Parameters.%s.PromoteParameterBlendedBSDFToOperator(SubstrateSelectParameterBlending(%s, %s, %s, %s, Parameters.%s.Types), %u, %u, %u, %u)"),
-			*GetParametersSubstrateTreeName(CurrentSubstrateCompilationContext),
+			*Substrate::GetParametersSubstrateTreeName(SubstrateTranslatorData.CurrentSubstrateCompilationContext),
 			*GetParameterCode(A),
 			*GetParameterCode(B),
 			*GetParameterCode(SelectValue),
 			*SharedLocalBasisIndexMacro,
-			*GetParametersSharedLocalBasesName(CurrentSubstrateCompilationContext),
+			*Substrate::GetParametersSharedLocalBasesName(SubstrateTranslatorData.CurrentSubstrateCompilationContext),
 			PromoteToOperator->Index,
 			PromoteToOperator->BSDFIndex,
 			PromoteToOperator->LayerDepth,
@@ -15767,7 +13686,7 @@ int32 FHLSLMaterialTranslator::SubstrateSelectParameterBlending(int32 A, int32 B
 		*GetParameterCode(B),
 		*GetParameterCode(SelectValue),
 		*SharedLocalBasisIndexMacro,
-		*GetParametersSharedLocalBasesName(CurrentSubstrateCompilationContext)
+		*Substrate::GetParametersSharedLocalBasesName(SubstrateTranslatorData.CurrentSubstrateCompilationContext)
 	);
 }
 
@@ -15899,104 +13818,47 @@ bool FHLSLMaterialTranslator::SubstrateSkipsOpacityEvaluation()
 
 FGuid FHLSLMaterialTranslator::SubstrateTreeStackPush(UMaterialExpression* Expression, uint32 InputIndex)
 {
-	FSubstrateCompilationContext& SubstrateCtx = SubstrateCompilationContext[CurrentSubstrateCompilationContext];
-
-	// Create an md5 hash for the parent, its input pin index and current node to represent the path.
-	uint32 IntputHashBuffer[9];
-	FGuid PreviousNodeGuid = SubstrateTreeStackGetPathUniqueId();
-	FGuid NodeGuid = Expression->MaterialExpressionGuid;
-	IntputHashBuffer[0] = PreviousNodeGuid.A;
-	IntputHashBuffer[1] = PreviousNodeGuid.B;
-	IntputHashBuffer[2] = PreviousNodeGuid.C;
-	IntputHashBuffer[3] = PreviousNodeGuid.D;
-	IntputHashBuffer[4] = InputIndex;
-	IntputHashBuffer[5] = NodeGuid.A;
-	IntputHashBuffer[6] = NodeGuid.B;
-	IntputHashBuffer[7] = NodeGuid.C;
-	IntputHashBuffer[8] = NodeGuid.D;
-
-	uint32 OutputHashBuffer[]{ 0, 0, 0, 0 };
-	FMD5 IdentifierStringHash;
-	IdentifierStringHash.Update((uint8*)IntputHashBuffer, sizeof(IntputHashBuffer));
-	IdentifierStringHash.Final((uint8*)&OutputHashBuffer);
-
-	SubstrateCtx.SubstrateNodeIdentifierStack.Push(FGuid(OutputHashBuffer[0], OutputHashBuffer[1], OutputHashBuffer[2], OutputHashBuffer[3]));
-
-#if DEBUG_SUBSTRATE_TREE_STACK
-	UE_LOG(LogMaterial, Display, TEXT(" SubstrateTreeStack: Push (input %i of %s)"), InputIndex , *Expression->GetName());
-	TStringBuilder<2048> GuidStack;
-	for (auto& Entry : SubstrateCtx.SubstrateNodeIdentifierStack)
-	{
-		GuidStack.Append(*Entry.ToString());
-		GuidStack.Append(TEXT("  "));
-	}
-	UE_LOG(LogMaterial, Display, TEXT(" SubstrateTreeStack: %s."), *GuidStack);
-#endif
-
-	SubstrateCtx.bSubstrateTreeOutOfStackDepthOccurred = SubstrateCtx.bSubstrateTreeOutOfStackDepthOccurred || (SubstrateCtx.SubstrateNodeIdentifierStack.Num() > SUBSTRATE_TREE_MAX_DEPTH);
-
-
-	return SubstrateCtx.SubstrateNodeIdentifierStack.Top();
+	return SubstrateTranslatorData.SubstrateTreeStackPush(Expression, InputIndex);
 }
 
 FGuid FHLSLMaterialTranslator::SubstrateTreeStackGetPathUniqueId()
 {
-	FSubstrateCompilationContext& SubstrateCtx = SubstrateCompilationContext[CurrentSubstrateCompilationContext];
-	return SubstrateCtx.SubstrateNodeIdentifierStack.Top();
+	return SubstrateTranslatorData.SubstrateTreeStackGetPathUniqueId();
 }
 
 FGuid FHLSLMaterialTranslator::SubstrateTreeStackGetParentPathUniqueId()
 {
-	FSubstrateCompilationContext& SubstrateCtx = SubstrateCompilationContext[CurrentSubstrateCompilationContext];
-	if (SubstrateCtx.SubstrateNodeIdentifierStack.Num() < 2)
-	{
-		// return some default when Substrate tree stack unique guid cannot be found
-		FGuid NullParent;
-		return NullParent;
-	}
-	return SubstrateCtx.SubstrateNodeIdentifierStack.Last(1);
+	return SubstrateTranslatorData.SubstrateTreeStackGetParentPathUniqueId();
 }
 
 void FHLSLMaterialTranslator::SubstrateTreeStackPop()
 {
-	FSubstrateCompilationContext& SubstrateCtx = SubstrateCompilationContext[CurrentSubstrateCompilationContext];
-	check(SubstrateCtx.SubstrateNodeIdentifierStack.Num() >= 2);// 2 because there must always be the root remaining.
-	SubstrateCtx.SubstrateNodeIdentifierStack.Pop();
-
-#if DEBUG_SUBSTRATE_TREE_STACK
-	TStringBuilder<2048> GuidStack;
-	for (auto& Entry : SubstrateCtx.SubstrateNodeIdentifierStack)
-	{
-		GuidStack.Append(*Entry.ToString());
-		GuidStack.Append(TEXT("  "));
-	}
-	UE_LOG(LogMaterial, Display, TEXT(" SubstrateTreeStack: Pop %s."), *GuidStack);
-#endif
+	SubstrateTranslatorData.SubstrateTreeStackPop();
 }
 
 bool FHLSLMaterialTranslator::GetSubstrateTreeOutOfStackDepthOccurred()
 {
-	FSubstrateCompilationContext& SubstrateCtx = SubstrateCompilationContext[CurrentSubstrateCompilationContext];
-	return SubstrateCtx.bSubstrateTreeOutOfStackDepthOccurred;
+	return SubstrateTranslatorData.GetSubstrateTreeOutOfStackDepthOccurred();
 }
 
 int32 FHLSLMaterialTranslator::SubstrateThicknessStackGetThicknessIndex()
 {
-	FSubstrateCompilationContext& SubstrateCtx = SubstrateCompilationContext[CurrentSubstrateCompilationContext];
-	return SubstrateCtx.SubstrateThicknessStack.Top();
+	return SubstrateTranslatorData.SubstrateThicknessStackGetThicknessIndex();
 }
 
 int32 FHLSLMaterialTranslator::SubstrateThicknessStackGetThicknessCode(int32 Index)
 {
-	FSubstrateCompilationContext& SubstrateCtx = SubstrateCompilationContext[CurrentSubstrateCompilationContext];
+	FSubstrateTranslatorDataInterface::FSubstrateThicknessExpression* ThicknessExpression = SubstrateTranslatorData.SubstrateThicknessStackGetExpression(Index);
+
 	int32 OutCode = INDEX_NONE;
-	if (Index == INDEX_NONE || Index >= SubstrateCtx.SubstrateThicknessIndexToExpressionInput.Num())
+	if (ThicknessExpression == nullptr)
 	{
 		UE_LOG(LogMaterial, Error, TEXT(" SubstrateThichkness: %i could not be found)"), Index);
 	}
-	else if (FExpressionInput* Input = SubstrateCtx.SubstrateThicknessIndexToExpressionInput[Index].ExpressionInput)
+	else if (ThicknessExpression->ExpressionInput)
 	{
-		OutCode = Input->GetTracedInput().Expression ? Input->Compile(this) : Constant(SUBSTRATE_LAYER_DEFAULT_THICKNESS_CM);
+		FExpressionInput* ExpressionInput = ThicknessExpression->ExpressionInput;
+		OutCode = ExpressionInput->GetTracedInput().Expression ? ExpressionInput->Compile(this) : Constant(SUBSTRATE_LAYER_DEFAULT_THICKNESS_CM);
 		EMaterialValueType Type = GetType(OutCode);
 		if (IsLWCType(Type))
 		{
@@ -16004,8 +13866,9 @@ int32 FHLSLMaterialTranslator::SubstrateThicknessStackGetThicknessCode(int32 Ind
 			OutCode = ValidCast(OutCode, Type);
 		}
 	}
-	else if (FScalarMaterialInput* MaterialInput = SubstrateCtx.SubstrateThicknessIndexToExpressionInput[Index].MaterialInput)
+	else if (ThicknessExpression->MaterialInput)
 	{
+		FScalarMaterialInput* MaterialInput = ThicknessExpression->MaterialInput;
 		OutCode = MaterialInput->UseConstant ? Constant(MaterialInput->Constant) : Constant(SUBSTRATE_LAYER_DEFAULT_THICKNESS_CM);
 		EMaterialValueType Type = GetType(OutCode);
 		if (IsLWCType(Type))
@@ -16023,35 +13886,22 @@ int32 FHLSLMaterialTranslator::SubstrateThicknessStackGetThicknessCode(int32 Ind
 
 int32 FHLSLMaterialTranslator::SubstrateThicknessStackPush(UMaterialExpression* Expression, FScalarMaterialInput* Input)
 {
-	// If the input is actually not a constant, reroute it to the Expression overload to that the expression can be traced.
-	if (Input && !Input->UseConstant)
-	{
-		return SubstrateThicknessStackPush(Expression, (FExpressionInput*)Input);
-	}
-
-	FSubstrateCompilationContext& SubstrateCtx = SubstrateCompilationContext[CurrentSubstrateCompilationContext];
-	int32 Index = SubstrateCtx.SubstrateThicknessIndexToExpressionInput.Num();
-	FSubstrateCompilationContext::FSubstrateThicknessExpression& Entry = SubstrateCtx.SubstrateThicknessIndexToExpressionInput.AddDefaulted_GetRef();
-	Entry.MaterialInput = Input;
-	SubstrateCtx.SubstrateThicknessStack.Push(Index);
-	return Index;
+	return SubstrateTranslatorData.SubstrateThicknessStackPush(Expression, Input);
 }
 
 int32 FHLSLMaterialTranslator::SubstrateThicknessStackPush(UMaterialExpression* Expression, FExpressionInput* Input)
 {
-	FSubstrateCompilationContext& SubstrateCtx = SubstrateCompilationContext[CurrentSubstrateCompilationContext];
-	int32 Index = SubstrateCtx.SubstrateThicknessIndexToExpressionInput.Num();
-	FSubstrateCompilationContext::FSubstrateThicknessExpression& Entry = SubstrateCtx.SubstrateThicknessIndexToExpressionInput.AddDefaulted_GetRef();
-	Entry.ExpressionInput = Input;
-	SubstrateCtx.SubstrateThicknessStack.Push(Index);
-	return Index;	
+	return SubstrateTranslatorData.SubstrateThicknessStackPush(Expression, Input);
 }
 
 void FHLSLMaterialTranslator::SubstrateThicknessStackPop()
 {
-	FSubstrateCompilationContext& SubstrateCtx = SubstrateCompilationContext[CurrentSubstrateCompilationContext];
-	check(SubstrateCtx.SubstrateThicknessStack.Num() >= 1);
-	SubstrateCtx.SubstrateThicknessStack.Pop();
+	SubstrateTranslatorData.SubstrateThicknessStackPop();
+}
+
+struct FSubstrateTranslatorDataInterface& FHLSLMaterialTranslator::GetSubstrateTranslatorData()
+{
+	return SubstrateTranslatorData;
 }
 
 int32 FHLSLMaterialTranslator::MapARPassthroughCameraUV(int32 UV)
@@ -16837,7 +14687,6 @@ void FHLSLMaterialTranslator::PrepareEnvironmentDefines()
 	EnvironmentDefines->bUsesDistortion = Material->IsDistorted();
 	EnvironmentDefines->bUsesExplicitDerivatives = bUsesExplicitDerivatives;
 	EnvironmentDefines->bUsesFirstPersonInterpolation = bUsesFirstPersonInterpolation;
-	EnvironmentDefines->bDistortionAccountForCoverage = bSubstrateEnabled && !bMaterialUsesRootNodeToSubstrateHiddenConversion && Material->GetRefractionCoverageMode() == RCM_CoverageAccountedFor ? 1 : 0;
 	EnvironmentDefines->bMaterialEnableTranslucencyFogging = Material->ShouldApplyFogging();
 	EnvironmentDefines->bMaterialEnableTranslucencyCloudFogging = Material->ShouldApplyCloudFogging();
 	EnvironmentDefines->bMaterialIsSky = Material->IsSky();
@@ -16971,7 +14820,7 @@ void FHLSLMaterialTranslator::PrepareEnvironmentDefines()
 			EnvironmentDefines->bMaterialVolumetricAdvancedClampMultiscatteringContribution = VolumetricAdvancedNode->bClampMultiScatteringContribution;
 			EnvironmentDefines->MaterialVolumetricAdvancedMultiscatteringOctaveCount = VolumetricAdvancedNode->GetMultiScatteringApproximationOctaveCount();
 			EnvironmentDefines->bMaterialVolumetricAdvancedConservativeDensity = VolumetricAdvancedNode->ConservativeDensity.IsConnected();
-			EnvironmentDefines->bMaterialVolumetricAdvancedOverrideAmbientOcclusion = Material->HasAmbientOcclusionConnected() || bSubstrateWritesAmbientOcclusion;
+			EnvironmentDefines->bMaterialVolumetricAdvancedOverrideAmbientOcclusion = Material->HasAmbientOcclusionConnected() || SubstrateTranslatorData.bSubstrateWritesAmbientOcclusion;
 			EnvironmentDefines->bMaterialVolumetricAdvancedAdvancedGroundContribution = VolumetricAdvancedNode->bGroundContribution;
 		}
 
@@ -16987,178 +14836,20 @@ void FHLSLMaterialTranslator::PrepareEnvironmentDefines()
 		}
 	}
 
-	EnvironmentDefines->bMaterialIsSubstrate = bMaterialIsSubstrate;
-	EnvironmentDefines->bMaterialUsesRootNodeToSubstrateHiddenConversion = bMaterialUsesRootNodeToSubstrateHiddenConversion;
-
 	// Substrate requests dual source blending only for BLEND_TranslucentColoredTransmittance
-	bMaterialRequestsDualSourceBlending |= EnvironmentDefines->bMaterialIsSubstrate && Material->GetBlendMode() == EBlendMode::BLEND_TranslucentColoredTransmittance;
+	bMaterialRequestsDualSourceBlending |= SubstrateTranslatorData.bMaterialIsSubstrate && Material->GetBlendMode() == EBlendMode::BLEND_TranslucentColoredTransmittance;
 
 	// if duals source blending (colored transmittance) is not supported on a platform, it will fall back to standard alpha blending (grey scale transmittance)
 	EnvironmentDefines->bDualSourceColorBlendingEnabled = Material->IsDualBlendingEnabled();
-	EnvironmentDefines->bSubstratePremultipliedAlphaOpacityOverridden = EnvironmentDefines->bMaterialIsSubstrate && bOpacityPropertyIsUsed;
-
-	if (EnvironmentDefines->bMaterialIsSubstrate)
-	{
-		EnvironmentDefines->bSubstrateUsesConversionFromLegacy = bSubstrateUsesConversionFromLegacy;
-		EnvironmentDefines->bSubstrateMaterialOutputOpaqueRoughRefractions = bSubstrateOutputsOpaqueRoughRefractions;
-
-		EnvironmentDefines->SubstrateMaterialExportType = (int32)GetSubstrateMaterialExportType();
-		EnvironmentDefines->SubstrateMaterialExportContext = (int32)GetSubstrateMaterialExportContext();
-		EnvironmentDefines->SubstrateMaterialExportLegacyBlendMode = (int32)GetSubstrateMaterialExportLegacyBlendMode();
-
-		// Unlit cannot be combined with other BSDF so we can simply pick the default strata context
-		EnvironmentDefines->bSubstrateOptimizedUnlit = SubstrateCompilationContext[ESubstrateCompilationContext::SCC_Default].bSubstrateMaterialIsUnlitNode;
-
-		EnvironmentDefines->bSubstrateRoughnessTracking = Material->HasSubstrateRoughnessTracking();
-
-		{
-			// For now, the fully simplified mode is used for Lumen or anything else supported inlined evaluation. The export is only valid for the default case.
-			// STRATA_TODO: generate an export for the different context (need to generate two export functions: the default one and the FullSimplification one)
-			FSubstrateCompilationContext& SubstrateCtx = SubstrateCompilationContext[ESubstrateCompilationContext::SCC_Default];
-			EnvironmentDefines->bSubstrateSinglePath = SubstrateCtx.SubstrateMaterialComplexity.IsSingle();
-			EnvironmentDefines->bSubstrateFastPath = SubstrateCtx.SubstrateMaterialComplexity.IsSimple();
-			EnvironmentDefines->SubstrateClampedClosureCount = SubstrateCtx.SubstrateMaterialEffectiveClosureCount;
-			EnvironmentDefines->bSubstrateComplexSpecialPath = SubstrateCtx.SubstrateMaterialComplexity.IsComplexSpecial();
-		}
-
-		bool EyeIrisNormalPluggedIn = false;
-		bool EyeIrisTangentPluggedIn = false;
-
-		FString SubstrateMaterialDescription;
-		for (uint32 SubstrateCompilationContextIndex = 0; SubstrateCompilationContextIndex < ESubstrateCompilationContext::SCC_MAX; ++SubstrateCompilationContextIndex)
-		{
-			ESubstrateBsdfFeature MaterialBsdfFeatures = ESubstrateBsdfFeature::None;
-
-			FSubstrateCompilationContext& SubstrateCtx = SubstrateCompilationContext[SubstrateCompilationContextIndex];
-			const FSubstrateSimplificationStatus& SubstrateSimplificationStatus = SubstrateCtx.SubstrateSimplificationStatus;
-
-			check(SubstrateCtx.SubstrateMaterialRootOperator);
-			uint32 RootMaximumDistanceToLeaves = SubstrateCtx.SubstrateMaterialRootOperator->MaxDistanceFromLeaves;
-
-			// Compute the shared local basis count and generate the hlsl shader code for it.
-			SubstrateCtx.SubstratePixelNormalInitializerValues = FString::Printf(TEXT("\n\n\n\t// Substrate normal and tangent\n"));
-			SubstrateCtx.FinalUsedSharedLocalBasesCount = 0;
-			uint8 RequestedSharedLocalBasesCount = 0;
-			SubstrateCtx.SubstrateEvaluateSharedLocalBases(this, RequestedSharedLocalBasesCount, EnvironmentDefines.Get());
-
-#if WITH_EDITOR
-			// Now write some feedback to the user, but only produce debug string if in editor
-			{
-				// Output some debug info as comment in code and in the material stat window
-				const uint32 SubstrateBytePerPixel_Platform = Substrate::GetBytePerPixel(InPlatform);
-				const uint32 SubstrateClosurePerPixel_Platform = Substrate::GetClosurePerPixel(InPlatform);
-				FString SubstrateMaterialContextDescription;
-
-				auto GetSubstrateCompilationContextName = [](uint32 Index)
-				{
-					switch (Index)
-					{
-					case ESubstrateCompilationContext::SCC_Default:
-						return TEXT("Default");
-					case ESubstrateCompilationContext::SCC_FullySimplified:
-						return TEXT("FullySimplified");
-					default:
-						check(false);
-					}
-					return TEXT("ERROR");
-				};
-				FString SubstrateCompilationContextName = GetSubstrateCompilationContextName(SubstrateCompilationContextIndex);
-				SubstrateMaterialContextDescription += FString::Printf(TEXT("\n\n----- SUBSTRATE - %s -----\n"), *SubstrateCompilationContextName);
-				SubstrateMaterialContextDescription += FString::Printf(TEXT("SubstrateCompilationInfo -\n"));
-				SubstrateMaterialContextDescription += FString::Printf(TEXT(" - Byte Per Pixel Budget                           %u\n"), SubstrateBytePerPixel_Platform);
-				SubstrateMaterialContextDescription += FString::Printf(TEXT(" - Closure Per Pixel Budget                        %u\n"), SubstrateClosurePerPixel_Platform);
-				SubstrateMaterialContextDescription += FString::Printf(TEXT(" - Requested Byte Size before simplification       %u (%d UINT32)\n"), SubstrateSimplificationStatus.OriginalRequestedByteSize, SubstrateSimplificationStatus.OriginalRequestedByteSize / 4);
-				SubstrateMaterialContextDescription += FString::Printf(TEXT(" - Requested Byte Size after simplification        %u (%d UINT32)\n"), SubstrateCtx.SubstrateMaterialRequestedSizeByte, SubstrateCtx.SubstrateMaterialRequestedSizeByte / 4);
-				SubstrateMaterialContextDescription += FString::Printf(TEXT(" - Requested Closure Count before simplification   %u\n"), SubstrateSimplificationStatus.OriginalRequestedClosureCount);
-				SubstrateMaterialContextDescription += FString::Printf(TEXT(" - Requested Closure Count after simplification    %u\n"), SubstrateCtx.SubstrateMaterialClosureCount);
-				SubstrateMaterialContextDescription += FString::Printf(TEXT(" - Material complexity                             %s\n"), *FSubstrateMaterialComplexity::ToString(SubstrateCtx.SubstrateMaterialComplexity.SubstrateMaterialType(), true /* Upper case */));
-				SubstrateMaterialContextDescription += FString::Printf(TEXT(" - BSDF Count                                      %i\n"), SubstrateCtx.SubstrateMaterialEffectiveClosureCount); // REMOVE?
-				if (RequestedSharedLocalBasesCount > SUBSTRATE_MAX_SHAREDLOCALBASES_REGISTERS)
-				{
-					SubstrateMaterialDescription += FString::Printf(TEXT(" - SharedLocalBasesCount                      %i (Requested:%i)\n"), SubstrateCtx.FinalUsedSharedLocalBasesCount, RequestedSharedLocalBasesCount);
-				}
-				else
-				{
-					SubstrateMaterialDescription += FString::Printf(TEXT(" - SharedLocalBasesCount                      %i\n"), SubstrateCtx.FinalUsedSharedLocalBasesCount);
-				}
-
-
-				for (int32 OpIt = 0; OpIt < SubstrateCtx.SubstrateMaterialExpressionRegisteredOperators.Num(); ++OpIt)
-				{
-					FSubstrateOperator& BSDFOperator = SubstrateCtx.SubstrateMaterialExpressionRegisteredOperators[OpIt];
-
-					EyeIrisNormalPluggedIn  = EyeIrisNormalPluggedIn  || BSDFOperator.Has(ESubstrateBsdfFeature::EyeIrisNormalPluggedIn);
-					EyeIrisTangentPluggedIn = EyeIrisTangentPluggedIn || BSDFOperator.Has(ESubstrateBsdfFeature::EyeIrisTangentPluggedIn);
-
-					if (BSDFOperator.BSDFIndex == INDEX_NONE || BSDFOperator.IsDiscarded())
-					{
-						continue;	// not a BSDF or if discarded (i.e. not the root of a parameter blending subtree), then there is no local basis to register
-					}
-					if (BSDFOperator.BSDFRegisteredSharedLocalBasis.NormalCodeChunk == INDEX_NONE && BSDFOperator.BSDFRegisteredSharedLocalBasis.TangentCodeChunk == INDEX_NONE)
-					{
-						continue;	// We skip null normal on certain BSDF, for instance unlit.
-					}
-					const FSubstrateSharedLocalBasesInfo& SubstrateSharedLocalBasesInfo = SubstrateCtx.SubstrateCompilationInfoGetMatchingSharedLocalBasisInfo(BSDFOperator.BSDFRegisteredSharedLocalBasis);
-
-					SubstrateMaterialContextDescription += FString::Printf(TEXT("     - %s - SharedLocalBasisIndexMacro = %s \n"), *GetSubstrateBSDFName(BSDFOperator.BSDFType), *GetSubstrateSharedLocalBasisIndexMacroInner(BSDFOperator.BSDFRegisteredSharedLocalBasis, SubstrateCompilationContextIndex));
-				}
-
-				SubstrateMaterialContextDescription += FString::Printf(TEXT("----------- SUBSTRATE TREE - %s -----------\n"), *SubstrateCompilationContextName);
-				SubstrateMaterialContextDescription += FString::Printf(TEXT("Graph maximum distance to leaves %u\n"), RootMaximumDistanceToLeaves);
-				// Debug print operators according to depth from root.
-				{
-
-					for (int32 DistanceToLeaves = RootMaximumDistanceToLeaves; DistanceToLeaves >= 0; --DistanceToLeaves)
-					{
-						SubstrateMaterialContextDescription += FString::Printf(TEXT("----- DistanceFromLeaves = %d -----\n"), DistanceToLeaves);
-						for (auto& It : SubstrateCtx.SubstrateMaterialExpressionRegisteredOperators)
-						{
-							if (!It.IsDiscarded() && It.MaxDistanceFromLeaves == DistanceToLeaves)
-							{
-								SubstrateMaterialContextDescription += FString::Printf(TEXT("\tIdx=%d Op=%s ParentIdx=%d LeftIndex=%d RightIndex=%d BSDFIdx=%d LayerDepth=%d IsTop=%d IsBot=%d BSDFType=%s SSS=%d MFP=%d F90=%d Rough2=%d Fuzz=%d Aniso=%d Glint=%d SpecularProfile=%d Eye=%d Hair=%d Toon=%d\n GUID=%s\n"),
-									It.Index, GetSubstrateOperatorStr(It.OperatorType), It.ParentIndex, It.LeftIndex, It.RightIndex, It.BSDFIndex, It.LayerDepth, It.bIsTop, It.bIsBottom,
-									*GetSubstrateBSDFName(It.BSDFType), 
-									It.Has(ESubstrateBsdfFeature::SSS), 
-									It.Has(ESubstrateBsdfFeature::MFPPluggedIn), 
-									It.Has(ESubstrateBsdfFeature::EdgeColor), 
-									It.Has(ESubstrateBsdfFeature::SecondRoughnessOrSimpleClearCoat), 
-									It.Has(ESubstrateBsdfFeature::Fuzz), 
-									It.Has(ESubstrateBsdfFeature::Anisotropy), 
-									It.Has(ESubstrateBsdfFeature::Glint), 
-									It.Has(ESubstrateBsdfFeature::SpecularProfile), 
-									It.Has(ESubstrateBsdfFeature::Eye), 
-									It.Has(ESubstrateBsdfFeature::Hair),
-									It.Has(ESubstrateBsdfFeature::Toon),
-									*It.MaterialExpressionGuid.ToString());
-
-								// Aggregate all the features used within the material
-								MaterialBsdfFeatures |= ESubstrateBsdfFeature(It.BSDFFeatures);
-							}
-						}
-					}
-				}
-
-				ResourcesString += TEXT("/*");
-				ResourcesString += SubstrateMaterialContextDescription;
-				ResourcesString += TEXT("*/");
 
-				SubstrateMaterialDescription += SubstrateMaterialContextDescription;
-
-				MaterialCompilationOutput.SubstrateMaterialCompilationOutput.SubstrateMaterialBsdfFeatures = MaterialBsdfFeatures;
-
-				if (SubstrateCompilationContextIndex == ESubstrateCompilationContext::SCC_Default)
-				{
-					MaterialCompilationOutput.SubstrateMaterialCompilationOutput.SharedLocalBasesCount = SubstrateCtx.FinalUsedSharedLocalBasesCount;
-				}
-			}
-#endif // WITH_EDITOR
-		}
-		MaterialCompilationOutput.SubstrateMaterialCompilationOutput.SubstrateMaterialDescription = SubstrateMaterialDescription;
-
-		EnvironmentDefines->bSubstrateLegacyIrisNormal = EyeIrisNormalPluggedIn;
-		EnvironmentDefines->bSubstrateLegacyIrisTangent = EyeIrisTangentPluggedIn;
-		EnvironmentDefines->SubstrateMaterialBsdfFeatures = MaterialCompilationOutput.SubstrateMaterialCompilationOutput.SubstrateMaterialBsdfFeatures;
-	}
+	// Finalise compilation and defines. Always need to be called to setup all the defines.
+	FString DescriptionStringCommentForDebug;
+	SubstrateTranslatorData.GenerateMaterialCompilationOutputAndDefines(
+		bOpacityPropertyIsUsed,
+		(int32)GetSubstrateMaterialExportType(), (int32)GetSubstrateMaterialExportContext(), (int32)GetSubstrateMaterialExportLegacyBlendMode(),
+		DescriptionStringCommentForDebug,
+		EnvironmentDefines.Get() ? &EnvironmentDefines.Get()->SubstrateDefines : nullptr);
+	ResourcesString += DescriptionStringCommentForDebug;
 
 	EnvironmentDefines->bTextureSampleDebug = IsDebugTextureSampleEnabled();
 	EnvironmentDefines->ParameterCollections = ParameterCollections;
