Substrate

diff --git a/Engine/Source/Runtime/Engine/Private/Materials/HLSLMaterialTranslator.h b/Engine/Source/Runtime/Engine/Private/Materials/HLSLMaterialTranslator.h
--- a/Engine/Source/Runtime/Engine/Private/Materials/HLSLMaterialTranslator.h
+++ b/Engine/Source/Runtime/Engine/Private/Materials/HLSLMaterialTranslator.h
@@ -52,8 +52,8 @@
 #include "Containers/LazyPrintf.h"
 #include "Containers/HashTable.h"
 #include "Engine/Texture2D.h"
-#include "SubstrateMaterial.h"
 #include "HLSLMaterialDerivativeAutogen.h"
+#include "SubstrateTranslatorCommon.h"
 #endif
 
 class Error;
@@ -213,13 +213,6 @@ enum class EMaterialCastFlags : uint32
 };
 ENUM_CLASS_FLAGS(EMaterialCastFlags);
 
-enum ESubstrateCompilationContext : uint8
-{
-	SCC_Default = 0u,
-	SCC_FullySimplified = 1u,
-	SCC_MAX = 2u
-};
-
 /** 
  * Enumerates Translate() results.
  * If RetryWithoutDDC is returned, the caller should call Translate() again
@@ -470,20 +463,6 @@ protected:
 	/** True if this material write anisotropy material property */
 	uint32 bUsesAnisotropy : 1;
 
-	/** True whether Substrate is enabled */
-	uint32 bSubstrateEnabled : 1;
-
-	/** True if the material is detected as a Substrate material at compile time.
-	 * This is decoupled from runtime FMaterialResource::IsSubstrateMaterial but practically fine since this is only temporary until Substrate is the main shading system. Only really used at runtime for translucency dual source blending.
-	 */
-	uint32 bMaterialIsSubstrate : 1;
-
-	/** True if Substrate hidden conversion from the root node is used by this material. (enabled for the project, FrontMaterial is not plugged in). */
-	uint32 bMaterialUsesRootNodeToSubstrateHiddenConversion : 1;
-
-	/** True if Substrate hidden conversion from the root node is enabled for the project.*/
-	uint32 bProjectSubstrateHiddenMaterialAssetConversionEnabled : 1;
-
 	/** True if the opacity input is plugged in */
 	uint32 bOpacityPropertyIsUsed : 1;
 	
@@ -522,137 +501,13 @@ protected:
 
 	struct FEnvironmentDefines;
 
-	// Describe the simplification status. Once the material has been compiled, it can be used to understand if and how it has been simplified.
-	struct FSubstrateSimplificationStatus
-	{
-		bool bMaterialFitsInMemoryBudget = false;	// Track whether or not the material fits.
-
-		uint32 OriginalRequestedByteSize = 0;
-		uint32 OriginalRequestedClosureCount = 0;
-		bool bRunFullSimplification = false;	// Simple implementation for now: if the material does not fit, we simply everything.
-		bool bFullSimplificationStepHasBeenRun = false;
-
-		bool bSlabSimplificationStepHasBeenRun = false;
-
-		struct FOperatorToSimplify
-		{
-			union
-			{
-				uint32 PackedData;
-				struct
-				{
-					uint32 Index : 16; // Index of the operator
-					uint32 Depth : 16; // Depth of the operator
-				} Data;
-			};
-
-			FORCEINLINE bool operator!=(FOperatorToSimplify B) const
-			{
-				return PackedData != B.PackedData;
-			}
-
-			FORCEINLINE bool operator<(FOperatorToSimplify B) const
-			{
-				return PackedData < B.PackedData;
-			}
-		};
-		TArray<FOperatorToSimplify> OperatorSimplificationOrder;
-	};
-
-	/** Represent a shared local basis description with its associated code. */
-	struct FSubstrateSharedLocalBasesInfo
-	{
-		FSubstrateRegisteredSharedLocalBasis SharedData;
-		FString NormalCode;
-		FString TangentCode;
-	};
-
-	struct FSubstrateCompilationContext
-	{
-		FSubstrateCompilationContext();
-
-		FSubstrateCompilationContext(ESubstrateCompilationContext InCompilationContext);
-
-		ESubstrateCompilationContext CompilationContextIndex;
-
-		/** The code initializing the array of shared local bases. */
-		FString SubstratePixelNormalInitializerValues;
-		/** The next free index that can be used to represent a unique macros pointing to the position in the array of shared local bases written to memory once the shader is executed. */
-		uint8 NextFreeSubstrateShaderNormalIndex;
-		/** The effective final shared local bases count used by the final shader. */
-		uint8 FinalUsedSharedLocalBasesCount;
-		/** Tracks shared local bases used by Substrate materials, mapping a normal code chunk hash to a SharedMaterialInfo.
-		 * A normal code chunk hash can point to multiple shared info in case it is paired with different tangents. */
-		TMultiMap<uint64, FSubstrateSharedLocalBasesInfo> CodeChunkToSubstrateSharedLocalBasis;
-
-		TMap<FGuid, int32> SubstrateMaterialExpressionToOperatorIndex;
-		TArray<FSubstrateOperator> SubstrateMaterialExpressionRegisteredOperators;
-		FSubstrateOperator* SubstrateMaterialRootOperator;
-		uint32 SubstrateMaterialEffectiveClosureCount; // Also acts as requested 
-		uint32 SubstrateMaterialRequestedSizeByte;
-		uint32 SubstrateMaterialClosureCount;
-		FSubstrateMaterialComplexity SubstrateMaterialComplexity;
-		bool bSubstrateMaterialIsUnlitNode;
-
-		bool bSubstrateWritesEmissive;
-		bool bSubstrateWritesAmbientOcclusion;
-
-		/** Stack of unique id for each node of the Substrate tree
-		* This is transient and updated on the fly in the exact same way when parsing node for
-		*  1- SubstrateGenerateMaterialTopologyTree: generating a picture of the Substrate material tree for code generation and simplifications.
-		*  2- CompilePropertyAndSetMaterialProperty(MP_ShadingModel): compiling the code with some features enabled/disabled and accounting for tree simplification decided beforehand.
-		* It is not valid to use that information outside of those functions.
-		*/
-		TArray<FGuid> SubstrateNodeIdentifierStack;
-
-		/**
-		 * This can be used to know if the Substrate tree we are trying to build is too deep and we should stop the compilation.
-		 * True means that we have likely encountered node re-entry leading to cyclic graph we cannot handle and compile internally: we must fail the compilation.
-		 */
-		bool bSubstrateTreeOutOfStackDepthOccurred;
-
-		/** Stack of thickness input used for propagating thickness information from root node and vertical operation
-		* This is transient and updated when calling SubstrateGenerateMaterialTopologyTree. The information is then stored into the FStratOperator
-		*/
-		struct FSubstrateThicknessExpression
-		{
-			FExpressionInput* ExpressionInput = nullptr;
-			FScalarMaterialInput* MaterialInput = nullptr;
-		};
-		TArray<int32> SubstrateThicknessStack;
-		TArray<FSubstrateThicknessExpression> SubstrateThicknessIndexToExpressionInput;
-
-		FSubstrateSimplificationStatus SubstrateSimplificationStatus;
-
-		bool SubstrateGenerateDerivedMaterialOperatorData(FHLSLMaterialTranslator* Compiler);
-
-		void SubstrateEvaluateSharedLocalBases(FHLSLMaterialTranslator* Compiler, uint8& RequestedSharedLocalBasesCount, FEnvironmentDefines* OutEnvironment);
-
-		FSubstrateSharedLocalBasesInfo SubstrateCompilationInfoGetMatchingSharedLocalBasisInfo(const FSubstrateRegisteredSharedLocalBasis& SearchedSharedLocalBasis);
-
-	private:
-		void Initialise();
-	};
-	FSubstrateCompilationContext SubstrateCompilationContext[ESubstrateCompilationContext::SCC_MAX];
-
-	ESubstrateCompilationContext CurrentSubstrateCompilationContext = ESubstrateCompilationContext::SCC_Default;
-	int32 FullySimplifiedSubstrateFrontMaterialCodeChunk = INDEX_NONE;
-	FString FullySimplifiedSubstrateFrontMaterialTranslatedCodeChunkDefinitions;
-	FString FullySimplifiedSubstrateFrontMaterialTranslatedCodeChunks;
-
-	bool bSubstrateWritesEmissive;
-	bool bSubstrateWritesAmbientOcclusion;
-
-	bool bSubstrateUsesConversionFromLegacy;
-	bool bSubstrateOutputsOpaqueRoughRefractions;
+	/** Substrate data for a Translator*/
+	Substrate::FSubstrateTranslatorData SubstrateTranslatorData;
 
 	/** Tracks the total number of vt samples in the shader. */
 	uint32 NumVtSamples;
 
 	const ITargetPlatform* TargetPlatform;
-
-	/** Substrate material compilation and simplification configuration. */
-	FSubstrateCompilationConfig SubstrateCompilationConfig;
 	
 	/** The DDC key for the material translation results */
 	FIoHash DDCKeyHash;
@@ -1462,6 +1317,8 @@ protected:
 	virtual int32 SubstrateThicknessStackPush(UMaterialExpression* Expression, FExpressionInput* Input) override;
 	virtual void SubstrateThicknessStackPop() override;
 
+	virtual struct FSubstrateTranslatorDataInterface& GetSubstrateTranslatorData() override;
+
 	virtual FSubstrateOperator& SubstrateCompilationRegisterOperator(int32 OperatorType, FGuid SubstrateExpressionGuid, FGuid ChildMaterialExpressionGuid, UMaterialExpression* Parent, FGuid SubstrateParentExpressionGuid, bool bUseParameterBlending = false) override;
 	virtual FSubstrateOperator& SubstrateCompilationGetOperator(FGuid SubstrateExpressionGuid) override;
 	virtual FSubstrateOperator* SubstrateCompilationGetOperatorFromIndex(int32 OperatorIndex) override;
diff --git a/Engine/Source/Runtime/Engine/Private/Materials/MaterialExpressionSubstrate.cpp b/Engine/Source/Runtime/Engine/Private/Materials/MaterialExpressionSubstrate.cpp
--- a/Engine/Source/Runtime/Engine/Private/Materials/MaterialExpressionSubstrate.cpp
+++ b/Engine/Source/Runtime/Engine/Private/Materials/MaterialExpressionSubstrate.cpp
@@ -17,6 +17,7 @@
 #if WITH_EDITOR
 #include "MaterialGraph/MaterialGraphSchema.h"
 #include "SubstrateMaterial.h"
+#include "SubstrateTranslatorCommon.h"
 #endif
 
 #include UE_INLINE_GENERATED_CPP_BY_NAME(MaterialExpressionSubstrate)
@@ -618,24 +619,41 @@ FSubstrateOperator* UMaterialExpressionSubstrateShadingModels::SubstrateGenerate
 	const FExpressionInput& ClearCoatNormal,
 	const FExpressionInput& CustomTangent,
 	const FExpressionInput& ShadingModel)
+{
+	return SubstrateGenerateMaterialTopologyTreeCommon(
+		Compiler->GetSubstrateTranslatorData(), ThisExpressionGuid, Parent, OutputIndex,
+		EmissiveColor,
+		Anisotropy,
+		ClearCoatNormal,
+		CustomTangent,
+		ShadingModel);
+}
+
+FSubstrateOperator* UMaterialExpressionSubstrateShadingModels::SubstrateGenerateMaterialTopologyTreeCommon(
+	struct FSubstrateTranslatorDataInterface& SubstrateTranslatorData, struct FGuid ThisExpressionGuid, class UMaterialExpression* Parent, int32 OutputIndex,
+	const FExpressionInput& EmissiveColor,
+	const FExpressionInput& Anisotropy,
+	const FExpressionInput& ClearCoatNormal,
+	const FExpressionInput& CustomTangent,
+	const FExpressionInput& ShadingModel)
 {
 	// Note Thickness has no meaning/usage in the context of SubstrateLegacyConversionNode
-	int32 ThicknessIndex = Compiler->SubstrateThicknessStackGetThicknessIndex();
+	int32 ThicknessIndex = SubstrateTranslatorData.SubstrateThicknessStackGetThicknessIndex();
 
 	auto AddDefaultWorstCase = [&](ESubstrateBsdfFeature In)
-		{
-			FSubstrateOperator& SlabOperator = Compiler->SubstrateCompilationRegisterOperator(SUBSTRATE_OPERATOR_BSDF_LEGACY, Compiler->SubstrateTreeStackGetPathUniqueId(), ThisExpressionGuid, Parent, Compiler->SubstrateTreeStackGetParentPathUniqueId());
-			SlabOperator.BSDFType = SUBSTRATE_BSDF_TYPE_SLAB;
-			SlabOperator.BSDFFeatures = In | (Anisotropy.IsConnected() ? ESubstrateBsdfFeature::Anisotropy : ESubstrateBsdfFeature::None);
-			SlabOperator.SubSurfaceType = uint8(EMaterialSubSurfaceType::MSS_None);
-			SlabOperator.bBSDFWritesEmissive = EmissiveColor.IsConnected();
-			SlabOperator.ThicknessIndex = ThicknessIndex;
+	{
+		FSubstrateOperator& SlabOperator = SubstrateTranslatorData.SubstrateCompilationRegisterOperator(SUBSTRATE_OPERATOR_BSDF_LEGACY, SubstrateTranslatorData.SubstrateTreeStackGetPathUniqueId(), ThisExpressionGuid, Parent, SubstrateTranslatorData.SubstrateTreeStackGetParentPathUniqueId());
+		SlabOperator.BSDFType = SUBSTRATE_BSDF_TYPE_SLAB;
+		SlabOperator.BSDFFeatures = In | (Anisotropy.IsConnected() ? ESubstrateBsdfFeature::Anisotropy : ESubstrateBsdfFeature::None);
+		SlabOperator.SubSurfaceType = uint8(EMaterialSubSurfaceType::MSS_None);
+		SlabOperator.bBSDFWritesEmissive = EmissiveColor.IsConnected();
+		SlabOperator.ThicknessIndex = ThicknessIndex;
 
-			return &SlabOperator;
-		};
+		return &SlabOperator;
+	};
 
 	// Get the shading models resulting from the UMaterial::RebuildShadingModelField().
-	FMaterialShadingModelField ShadingModels = Compiler->GetMaterialShadingModels();
+	FMaterialShadingModelField ShadingModels = SubstrateTranslatorData.GetMaterialShadingModels();
 
 	bool bEyeIrisNormalPluggedIn = ClearCoatNormal.IsConnected();
 	bool bEyeIrisTangentPluggedIn = CustomTangent.IsConnected();
@@ -652,7 +670,7 @@ FSubstrateOperator* UMaterialExpressionSubstrateShadingModels::SubstrateGenerate
 		// Special case for unlit only material to get fast path
 		if (ShadingModels.HasOnlyShadingModel(MSM_Unlit))
 		{
-			FSubstrateOperator& Operator = Compiler->SubstrateCompilationRegisterOperator(SUBSTRATE_OPERATOR_BSDF_LEGACY, Compiler->SubstrateTreeStackGetPathUniqueId(), ThisExpressionGuid, Parent, Compiler->SubstrateTreeStackGetParentPathUniqueId());
+			FSubstrateOperator& Operator = SubstrateTranslatorData.SubstrateCompilationRegisterOperator(SUBSTRATE_OPERATOR_BSDF_LEGACY, SubstrateTranslatorData.SubstrateTreeStackGetPathUniqueId(), ThisExpressionGuid, Parent, SubstrateTranslatorData.SubstrateTreeStackGetParentPathUniqueId());
 			Operator.BSDFType = SUBSTRATE_BSDF_TYPE_UNLIT;
 			Operator.ThicknessIndex = ThicknessIndex;
 			return &Operator;
@@ -676,7 +694,7 @@ FSubstrateOperator* UMaterialExpressionSubstrateShadingModels::SubstrateGenerate
 
 		if (ShadingModels.HasShadingModel(MSM_Unlit))
 		{
-			FSubstrateOperator& Operator = Compiler->SubstrateCompilationRegisterOperator(SUBSTRATE_OPERATOR_BSDF_LEGACY, Compiler->SubstrateTreeStackGetPathUniqueId(), ThisExpressionGuid, Parent, Compiler->SubstrateTreeStackGetParentPathUniqueId());
+			FSubstrateOperator& Operator = SubstrateTranslatorData.SubstrateCompilationRegisterOperator(SUBSTRATE_OPERATOR_BSDF_LEGACY, SubstrateTranslatorData.SubstrateTreeStackGetPathUniqueId(), ThisExpressionGuid, Parent, SubstrateTranslatorData.SubstrateTreeStackGetParentPathUniqueId());
 			Operator.BSDFType = SUBSTRATE_BSDF_TYPE_UNLIT;
 			Operator.ThicknessIndex = ThicknessIndex;
 			Operator.bBSDFWritesEmissive = true;
@@ -712,7 +730,7 @@ FSubstrateOperator* UMaterialExpressionSubstrateShadingModels::SubstrateGenerate
 		}
 		else if (ShadingModels.HasShadingModel(MSM_ClearCoat))
 		{
-			FSubstrateOperator& Operator = Compiler->SubstrateCompilationRegisterOperator(SUBSTRATE_OPERATOR_BSDF_LEGACY, Compiler->SubstrateTreeStackGetPathUniqueId(), ThisExpressionGuid, Parent, Compiler->SubstrateTreeStackGetParentPathUniqueId());
+			FSubstrateOperator& Operator = SubstrateTranslatorData.SubstrateCompilationRegisterOperator(SUBSTRATE_OPERATOR_BSDF_LEGACY, SubstrateTranslatorData.SubstrateTreeStackGetPathUniqueId(), ThisExpressionGuid, Parent, SubstrateTranslatorData.SubstrateTreeStackGetParentPathUniqueId());
 			Operator.BSDFType = SUBSTRATE_BSDF_TYPE_SLAB;
 			Operator.ThicknessIndex = ThicknessIndex;
 			Operator.BSDFFeatures = ESubstrateBsdfFeature::SecondRoughnessOrSimpleClearCoat | (Anisotropy.IsConnected() ? ESubstrateBsdfFeature::Anisotropy : ESubstrateBsdfFeature::None);
@@ -721,7 +739,7 @@ FSubstrateOperator* UMaterialExpressionSubstrateShadingModels::SubstrateGenerate
 		}
 		else if (ShadingModels.HasShadingModel(MSM_Hair))
 		{
-			FSubstrateOperator& Operator = Compiler->SubstrateCompilationRegisterOperator(SUBSTRATE_OPERATOR_BSDF_LEGACY, Compiler->SubstrateTreeStackGetPathUniqueId(), ThisExpressionGuid, Parent, Compiler->SubstrateTreeStackGetParentPathUniqueId());
+			FSubstrateOperator& Operator = SubstrateTranslatorData.SubstrateCompilationRegisterOperator(SUBSTRATE_OPERATOR_BSDF_LEGACY, SubstrateTranslatorData.SubstrateTreeStackGetPathUniqueId(), ThisExpressionGuid, Parent, SubstrateTranslatorData.SubstrateTreeStackGetParentPathUniqueId());
 			Operator.BSDFType = SUBSTRATE_BSDF_TYPE_HAIR;
 			Operator.BSDFFeatures = ESubstrateBsdfFeature::Hair;
 			Operator.ThicknessIndex = ThicknessIndex;
@@ -730,7 +748,7 @@ FSubstrateOperator* UMaterialExpressionSubstrateShadingModels::SubstrateGenerate
 		}
 		else if (ShadingModels.HasShadingModel(MSM_Eye))
 		{
-			FSubstrateOperator& Operator = Compiler->SubstrateCompilationRegisterOperator(SUBSTRATE_OPERATOR_BSDF_LEGACY, Compiler->SubstrateTreeStackGetPathUniqueId(), ThisExpressionGuid, Parent, Compiler->SubstrateTreeStackGetParentPathUniqueId());
+			FSubstrateOperator& Operator = SubstrateTranslatorData.SubstrateCompilationRegisterOperator(SUBSTRATE_OPERATOR_BSDF_LEGACY, SubstrateTranslatorData.SubstrateTreeStackGetPathUniqueId(), ThisExpressionGuid, Parent, SubstrateTranslatorData.SubstrateTreeStackGetParentPathUniqueId());
 			Operator.BSDFType = SUBSTRATE_BSDF_TYPE_EYE;
 			Operator.BSDFFeatures = ESubstrateBsdfFeature::Eye;
 			ApplyEyeIrisUsedFeatures(&Operator);
@@ -741,7 +759,7 @@ FSubstrateOperator* UMaterialExpressionSubstrateShadingModels::SubstrateGenerate
 		}
 		else if (ShadingModels.HasShadingModel(MSM_SingleLayerWater))
 		{
-			FSubstrateOperator& Operator = Compiler->SubstrateCompilationRegisterOperator(SUBSTRATE_OPERATOR_BSDF_LEGACY, Compiler->SubstrateTreeStackGetPathUniqueId(), ThisExpressionGuid, Parent, Compiler->SubstrateTreeStackGetParentPathUniqueId());
+			FSubstrateOperator& Operator = SubstrateTranslatorData.SubstrateCompilationRegisterOperator(SUBSTRATE_OPERATOR_BSDF_LEGACY, SubstrateTranslatorData.SubstrateTreeStackGetPathUniqueId(), ThisExpressionGuid, Parent, SubstrateTranslatorData.SubstrateTreeStackGetParentPathUniqueId());
 			Operator.BSDFType = SUBSTRATE_BSDF_TYPE_SINGLELAYERWATER;
 			Operator.ThicknessIndex = ThicknessIndex;
 			Operator.bBSDFWritesEmissive = EmissiveColor.IsConnected();
@@ -754,10 +772,10 @@ FSubstrateOperator* UMaterialExpressionSubstrateShadingModels::SubstrateGenerate
 	}
 }
 
-FSubstrateOperator* UMaterialExpressionSubstrateShadingModels::SubstrateGenerateMaterialTopologyTree(class FMaterialCompiler* Compiler, class UMaterialExpression* Parent, int32 OutputIndex)
+FSubstrateOperator* UMaterialExpressionSubstrateShadingModels::SubstrateGenerateMaterialTopologyTree(struct FSubstrateTranslatorDataInterface& SubstrateTranslatorData, class UMaterialExpression* Parent, int32 OutputIndex)
 {
 	return SubstrateGenerateMaterialTopologyTreeCommon(
-		Compiler, this->MaterialExpressionGuid, Parent, OutputIndex,
+		SubstrateTranslatorData, this->MaterialExpressionGuid, Parent, OutputIndex,
 		EmissiveColor,
 		Anisotropy,
 		ClearCoatNormal,
@@ -1209,9 +1227,10 @@ void UMaterialExpressionSubstrateSlabBSDF::GatherSubstrateMaterialInfo(FSubstrat
 	SubstrateMaterialInfo.AddGuid(MaterialExpressionGuid);
 }
 
-FSubstrateOperator* UMaterialExpressionSubstrateSlabBSDF::SubstrateGenerateMaterialTopologyTree(class FMaterialCompiler* Compiler, class UMaterialExpression* Parent, int32 OutputIndex)
+FSubstrateOperator* UMaterialExpressionSubstrateSlabBSDF::SubstrateGenerateMaterialTopologyTree(
+	struct FSubstrateTranslatorDataInterface& SubstrateTranslatorData, class UMaterialExpression* Parent, int32 OutputIndex)
 {
-	FSubstrateOperator& SubstrateOperator = Compiler->SubstrateCompilationRegisterOperator(SUBSTRATE_OPERATOR_BSDF, Compiler->SubstrateTreeStackGetPathUniqueId(), this->MaterialExpressionGuid, Parent, Compiler->SubstrateTreeStackGetParentPathUniqueId());
+	FSubstrateOperator& SubstrateOperator = SubstrateTranslatorData.SubstrateCompilationRegisterOperator(SUBSTRATE_OPERATOR_BSDF, SubstrateTranslatorData.SubstrateTreeStackGetPathUniqueId(), this->MaterialExpressionGuid, Parent, SubstrateTranslatorData.SubstrateTreeStackGetParentPathUniqueId());
 	SubstrateOperator.BSDFType = SUBSTRATE_BSDF_TYPE_SLAB;
 
 	if (HasEdgeColor())			{ SubstrateOperator.BSDFFeatures |= ESubstrateBsdfFeature::EdgeColor; }
@@ -1225,7 +1244,7 @@ FSubstrateOperator* UMaterialExpressionSubstrateSlabBSDF::SubstrateGenerateMater
 	if (HasSSSProfile())		{ SubstrateOperator.BSDFFeatures |= ESubstrateBsdfFeature::SecondRoughnessOrSimpleClearCoat; } // If a Slab has a subsurface profile, it will have haziness in order to to support the dual-specular lobe from the profile.
 
 	SubstrateOperator.SubSurfaceType = SubSurfaceType;
-	SubstrateOperator.ThicknessIndex = Compiler->SubstrateThicknessStackGetThicknessIndex();
+	SubstrateOperator.ThicknessIndex = SubstrateTranslatorData.SubstrateThicknessStackGetThicknessIndex();
 	SubstrateOperator.bBSDFWritesEmissive = EmissiveColor.IsConnected();
 	return &SubstrateOperator;
 }
@@ -1481,12 +1500,13 @@ void UMaterialExpressionSubstrateSimpleClearCoatBSDF::GatherSubstrateMaterialInf
 	SubstrateMaterialInfo.AddGuid(MaterialExpressionGuid);
 }
 
-FSubstrateOperator* UMaterialExpressionSubstrateSimpleClearCoatBSDF::SubstrateGenerateMaterialTopologyTree(class FMaterialCompiler* Compiler, class UMaterialExpression* Parent, int32 OutputIndex)
+FSubstrateOperator* UMaterialExpressionSubstrateSimpleClearCoatBSDF::SubstrateGenerateMaterialTopologyTree(
+	struct FSubstrateTranslatorDataInterface& SubstrateTranslatorData, class UMaterialExpression* Parent, int32 OutputIndex)
 {
-	FSubstrateOperator& SubstrateOperator = Compiler->SubstrateCompilationRegisterOperator(SUBSTRATE_OPERATOR_BSDF, Compiler->SubstrateTreeStackGetPathUniqueId(), this->MaterialExpressionGuid, Parent, Compiler->SubstrateTreeStackGetParentPathUniqueId());
+	FSubstrateOperator& SubstrateOperator = SubstrateTranslatorData.SubstrateCompilationRegisterOperator(SUBSTRATE_OPERATOR_BSDF, SubstrateTranslatorData.SubstrateTreeStackGetPathUniqueId(), this->MaterialExpressionGuid, Parent, SubstrateTranslatorData.SubstrateTreeStackGetParentPathUniqueId());
 	SubstrateOperator.BSDFType = SUBSTRATE_BSDF_TYPE_SLAB;
 	SubstrateOperator.BSDFFeatures = ESubstrateBsdfFeature::SecondRoughnessOrSimpleClearCoat; // This node explicitly requires simple clear coat
-	SubstrateOperator.ThicknessIndex = Compiler->SubstrateThicknessStackGetThicknessIndex();
+	SubstrateOperator.ThicknessIndex = SubstrateTranslatorData.SubstrateThicknessStackGetThicknessIndex();
 	SubstrateOperator.bBSDFWritesEmissive = EmissiveColor.IsConnected();
 	return &SubstrateOperator;
 }
@@ -1581,21 +1601,33 @@ FSubstrateOperator* UMaterialExpressionSubstrateVolumetricFogCloudBSDF::Substrat
 	const FExpressionInput& EmissiveColor,
 	const FExpressionInput& AmbientOcclusion)
 {
-	FSubstrateOperator& SubstrateOperator = Compiler->SubstrateCompilationRegisterOperator(SUBSTRATE_OPERATOR_BSDF, Compiler->SubstrateTreeStackGetPathUniqueId(), ThisExpressionGuid, Parent, Compiler->SubstrateTreeStackGetParentPathUniqueId());
+	return SubstrateGenerateMaterialTopologyTreeCommon(
+		Compiler->GetSubstrateTranslatorData(), ThisExpressionGuid, Parent, OutputIndex,
+		EmissiveColor,
+		AmbientOcclusion);
+}
+
+FSubstrateOperator* UMaterialExpressionSubstrateVolumetricFogCloudBSDF::SubstrateGenerateMaterialTopologyTreeCommon(
+	struct FSubstrateTranslatorDataInterface& SubstrateTranslatorData, struct FGuid ThisExpressionGuid, class UMaterialExpression* Parent, int32 OutputIndex,
+	const FExpressionInput& EmissiveColor,
+	const FExpressionInput& AmbientOcclusion)
+{
+	FSubstrateOperator& SubstrateOperator = SubstrateTranslatorData.SubstrateCompilationRegisterOperator(SUBSTRATE_OPERATOR_BSDF, SubstrateTranslatorData.SubstrateTreeStackGetPathUniqueId(), ThisExpressionGuid, Parent, SubstrateTranslatorData.SubstrateTreeStackGetParentPathUniqueId());
 	SubstrateOperator.BSDFType = SUBSTRATE_BSDF_TYPE_VOLUMETRICFOGCLOUD;
-	SubstrateOperator.ThicknessIndex = Compiler->SubstrateThicknessStackGetThicknessIndex();
+	SubstrateOperator.ThicknessIndex = SubstrateTranslatorData.SubstrateThicknessStackGetThicknessIndex();
 	SubstrateOperator.bBSDFWritesEmissive = EmissiveColor.IsConnected();
 	SubstrateOperator.bBSDFWritesAmbientOcclusion = AmbientOcclusion.IsConnected();
 	return &SubstrateOperator;
 }
 
-FSubstrateOperator* UMaterialExpressionSubstrateVolumetricFogCloudBSDF::SubstrateGenerateMaterialTopologyTree(class FMaterialCompiler* Compiler, class UMaterialExpression* Parent, int32 OutputIndex)
+FSubstrateOperator* UMaterialExpressionSubstrateVolumetricFogCloudBSDF::SubstrateGenerateMaterialTopologyTree(struct FSubstrateTranslatorDataInterface& SubstrateTranslatorData, class UMaterialExpression* Parent, int32 OutputIndex)
 {
 	return SubstrateGenerateMaterialTopologyTreeCommon(
-		Compiler, MaterialExpressionGuid, Parent, OutputIndex,
+		SubstrateTranslatorData, this->MaterialExpressionGuid, Parent, OutputIndex,
 		EmissiveColor,
 		AmbientOcclusion);
 }
+
 #endif // WITH_EDITOR
 
 
@@ -1669,17 +1701,24 @@ void UMaterialExpressionSubstrateLightFunction::GatherSubstrateMaterialInfo(FSub
 FSubstrateOperator* UMaterialExpressionSubstrateLightFunction::SubstrateGenerateMaterialTopologyTreeCommon(
 	class FMaterialCompiler* Compiler, struct FGuid ThisExpressionGuid, class UMaterialExpression* Parent, int32 OutputIndex)
 {
-	FSubstrateOperator& SubstrateOperator = Compiler->SubstrateCompilationRegisterOperator(SUBSTRATE_OPERATOR_BSDF, Compiler->SubstrateTreeStackGetPathUniqueId(), ThisExpressionGuid, Parent, Compiler->SubstrateTreeStackGetParentPathUniqueId());
+	return SubstrateGenerateMaterialTopologyTreeCommon(Compiler->GetSubstrateTranslatorData(), ThisExpressionGuid, Parent, OutputIndex);
+}
+
+FSubstrateOperator* UMaterialExpressionSubstrateLightFunction::SubstrateGenerateMaterialTopologyTreeCommon(
+	struct FSubstrateTranslatorDataInterface& SubstrateTranslatorData, struct FGuid ThisExpressionGuid, class UMaterialExpression* Parent, int32 OutputIndex)
+{
+	FSubstrateOperator& SubstrateOperator = SubstrateTranslatorData.SubstrateCompilationRegisterOperator(SUBSTRATE_OPERATOR_BSDF, SubstrateTranslatorData.SubstrateTreeStackGetPathUniqueId(), ThisExpressionGuid, Parent, SubstrateTranslatorData.SubstrateTreeStackGetParentPathUniqueId());
 	SubstrateOperator.BSDFType = SUBSTRATE_BSDF_TYPE_UNLIT;
 	SubstrateOperator.SubUsage = SUBSTRATE_OPERATOR_SUBUSAGE_LIGHTFUNCTION;
-	SubstrateOperator.ThicknessIndex = Compiler->SubstrateThicknessStackGetThicknessIndex();
+	SubstrateOperator.ThicknessIndex = SubstrateTranslatorData.SubstrateThicknessStackGetThicknessIndex();
 	return &SubstrateOperator;
 }
 
-FSubstrateOperator* UMaterialExpressionSubstrateLightFunction::SubstrateGenerateMaterialTopologyTree(class FMaterialCompiler* Compiler, class UMaterialExpression* Parent, int32 OutputIndex)
+FSubstrateOperator* UMaterialExpressionSubstrateLightFunction::SubstrateGenerateMaterialTopologyTree(struct FSubstrateTranslatorDataInterface& SubstrateTranslatorData, class UMaterialExpression* Parent, int32 OutputIndex)
 {
-	return SubstrateGenerateMaterialTopologyTreeCommon(Compiler, MaterialExpressionGuid, Parent, OutputIndex);
+	return SubstrateGenerateMaterialTopologyTreeCommon(SubstrateTranslatorData, this->MaterialExpressionGuid, Parent, OutputIndex);
 }
+
 #endif // WITH_EDITOR
 
 
@@ -1758,16 +1797,22 @@ void UMaterialExpressionSubstratePostProcess::GatherSubstrateMaterialInfo(FSubst
 FSubstrateOperator* UMaterialExpressionSubstratePostProcess::SubstrateGenerateMaterialTopologyTreeCommon(
 	class FMaterialCompiler* Compiler, struct FGuid ThisExpressionGuid, class UMaterialExpression* Parent, int32 OutputIndex)
 {
-	FSubstrateOperator& SubstrateOperator = Compiler->SubstrateCompilationRegisterOperator(SUBSTRATE_OPERATOR_BSDF, Compiler->SubstrateTreeStackGetPathUniqueId(), ThisExpressionGuid, Parent, Compiler->SubstrateTreeStackGetParentPathUniqueId());
+	return SubstrateGenerateMaterialTopologyTreeCommon(Compiler->GetSubstrateTranslatorData(), ThisExpressionGuid, Parent, OutputIndex);
+}
+
+FSubstrateOperator* UMaterialExpressionSubstratePostProcess::SubstrateGenerateMaterialTopologyTreeCommon(
+	struct FSubstrateTranslatorDataInterface& SubstrateTranslatorData, struct FGuid ThisExpressionGuid, class UMaterialExpression* Parent, int32 OutputIndex)
+{
+	FSubstrateOperator& SubstrateOperator = SubstrateTranslatorData.SubstrateCompilationRegisterOperator(SUBSTRATE_OPERATOR_BSDF, SubstrateTranslatorData.SubstrateTreeStackGetPathUniqueId(), ThisExpressionGuid, Parent, SubstrateTranslatorData.SubstrateTreeStackGetParentPathUniqueId());
 	SubstrateOperator.BSDFType = SUBSTRATE_BSDF_TYPE_UNLIT;
 	SubstrateOperator.SubUsage = SUBSTRATE_OPERATOR_SUBUSAGE_POSTPROCESS;
-	SubstrateOperator.ThicknessIndex = Compiler->SubstrateThicknessStackGetThicknessIndex();
+	SubstrateOperator.ThicknessIndex = SubstrateTranslatorData.SubstrateThicknessStackGetThicknessIndex();
 	return &SubstrateOperator;
 }
 
-FSubstrateOperator* UMaterialExpressionSubstratePostProcess::SubstrateGenerateMaterialTopologyTree(class FMaterialCompiler* Compiler, class UMaterialExpression* Parent, int32 OutputIndex)
+FSubstrateOperator* UMaterialExpressionSubstratePostProcess::SubstrateGenerateMaterialTopologyTree(struct FSubstrateTranslatorDataInterface& SubstrateTranslatorData, class UMaterialExpression* Parent, int32 OutputIndex)
 {
-	return SubstrateGenerateMaterialTopologyTreeCommon(Compiler, MaterialExpressionGuid, Parent, OutputIndex);
+	return SubstrateGenerateMaterialTopologyTreeCommon(SubstrateTranslatorData, this->MaterialExpressionGuid, Parent, OutputIndex);
 }
 #endif // WITH_EDITOR
 
@@ -1845,16 +1890,22 @@ void UMaterialExpressionSubstrateUI::GatherSubstrateMaterialInfo(FSubstrateMater
 FSubstrateOperator* UMaterialExpressionSubstrateUI::SubstrateGenerateMaterialTopologyTreeCommon(
 	class FMaterialCompiler* Compiler, struct FGuid ThisExpressionGuid, class UMaterialExpression* Parent, int32 OutputIndex)
 {
-	FSubstrateOperator& SubstrateOperator = Compiler->SubstrateCompilationRegisterOperator(SUBSTRATE_OPERATOR_BSDF, Compiler->SubstrateTreeStackGetPathUniqueId(), ThisExpressionGuid, Parent, Compiler->SubstrateTreeStackGetParentPathUniqueId());
+	return SubstrateGenerateMaterialTopologyTreeCommon(Compiler->GetSubstrateTranslatorData(), ThisExpressionGuid, Parent, OutputIndex);
+}
+
+FSubstrateOperator* UMaterialExpressionSubstrateUI::SubstrateGenerateMaterialTopologyTreeCommon(
+	struct FSubstrateTranslatorDataInterface& SubstrateTranslatorData, struct FGuid ThisExpressionGuid, class UMaterialExpression* Parent, int32 OutputIndex)
+{
+	FSubstrateOperator& SubstrateOperator = SubstrateTranslatorData.SubstrateCompilationRegisterOperator(SUBSTRATE_OPERATOR_BSDF, SubstrateTranslatorData.SubstrateTreeStackGetPathUniqueId(), ThisExpressionGuid, Parent, SubstrateTranslatorData.SubstrateTreeStackGetParentPathUniqueId());
 	SubstrateOperator.BSDFType = SUBSTRATE_BSDF_TYPE_UNLIT;
 	SubstrateOperator.SubUsage = SUBSTRATE_OPERATOR_SUBUSAGE_UI;
-	SubstrateOperator.ThicknessIndex = Compiler->SubstrateThicknessStackGetThicknessIndex();
+	SubstrateOperator.ThicknessIndex = SubstrateTranslatorData.SubstrateThicknessStackGetThicknessIndex();
 	return &SubstrateOperator;
 }
 
-FSubstrateOperator* UMaterialExpressionSubstrateUI::SubstrateGenerateMaterialTopologyTree(class FMaterialCompiler* Compiler, class UMaterialExpression* Parent, int32 OutputIndex)
+FSubstrateOperator* UMaterialExpressionSubstrateUI::SubstrateGenerateMaterialTopologyTree(struct FSubstrateTranslatorDataInterface& SubstrateTranslatorData, class UMaterialExpression* Parent, int32 OutputIndex)
 {
-	return SubstrateGenerateMaterialTopologyTreeCommon(Compiler, MaterialExpressionGuid, Parent, OutputIndex);
+	return SubstrateGenerateMaterialTopologyTreeCommon(SubstrateTranslatorData, this->MaterialExpressionGuid, Parent, OutputIndex);
 }
 #endif // WITH_EDITOR
 
@@ -1956,12 +2007,13 @@ void UMaterialExpressionSubstrateConvertToDecal::GatherSubstrateMaterialInfo(FSu
 	SubstrateMaterialInfo.PopSubstrateTreeStack();
 }
 
-FSubstrateOperator* UMaterialExpressionSubstrateConvertToDecal::SubstrateGenerateMaterialTopologyTree(class FMaterialCompiler* Compiler, class UMaterialExpression* Parent, int32 OutputIndex)
+FSubstrateOperator* UMaterialExpressionSubstrateConvertToDecal::SubstrateGenerateMaterialTopologyTree(
+	struct FSubstrateTranslatorDataInterface& SubstrateTranslatorData, class UMaterialExpression* Parent, int32 OutputIndex)
 {
 	const bool bUseParameterBlending = true;
-	FSubstrateOperator& SubstrateOperator = Compiler->SubstrateCompilationRegisterOperator(SUBSTRATE_OPERATOR_WEIGHT, Compiler->SubstrateTreeStackGetPathUniqueId(), this->MaterialExpressionGuid, Parent, Compiler->SubstrateTreeStackGetParentPathUniqueId(), bUseParameterBlending);
+	FSubstrateOperator& SubstrateOperator = SubstrateTranslatorData.SubstrateCompilationRegisterOperator(SUBSTRATE_OPERATOR_WEIGHT, SubstrateTranslatorData.SubstrateTreeStackGetPathUniqueId(), this->MaterialExpressionGuid, Parent, SubstrateTranslatorData.SubstrateTreeStackGetParentPathUniqueId(), bUseParameterBlending);
 	SubstrateOperator.SubUsage = SUBSTRATE_OPERATOR_SUBUSAGE_DECAL;
-	if (Compiler->GetSubstrateTreeOutOfStackDepthOccurred())
+	if (SubstrateTranslatorData.GetSubstrateTreeOutOfStackDepthOccurred())
 	{
 		return &SubstrateOperator; // Out ot stack space, return now to fail the compilation
 	}
@@ -1971,9 +2023,9 @@ FSubstrateOperator* UMaterialExpressionSubstrateConvertToDecal::SubstrateGenerat
 	FSubstrateOperator* OpA = nullptr;
 	if (ChildDecalMaterialExpression)
 	{
-		Compiler->SubstrateTreeStackPush(this, 0);
-		OpA = ChildDecalMaterialExpression->SubstrateGenerateMaterialTopologyTree(Compiler, this, TracedInput.OutputIndex);
-		Compiler->SubstrateTreeStackPop();
+		SubstrateTranslatorData.SubstrateTreeStackPush(this, 0);
+		OpA = ChildDecalMaterialExpression->SubstrateGenerateMaterialTopologyTree(SubstrateTranslatorData, this, TracedInput.OutputIndex);
+		SubstrateTranslatorData.SubstrateTreeStackPop();
 		AssignOperatorIndexIfNotNull(SubstrateOperator.LeftIndex, OpA);
 	}
 	CombineFlagForParameterBlending(SubstrateOperator, OpA);
@@ -2048,11 +2100,12 @@ void UMaterialExpressionSubstrateUnlitBSDF::GatherSubstrateMaterialInfo(FSubstra
 	SubstrateMaterialInfo.AddGuid(MaterialExpressionGuid);
 }
 
-FSubstrateOperator* UMaterialExpressionSubstrateUnlitBSDF::SubstrateGenerateMaterialTopologyTree(class FMaterialCompiler* Compiler, class UMaterialExpression* Parent, int32 OutputIndex)
+FSubstrateOperator* UMaterialExpressionSubstrateUnlitBSDF::SubstrateGenerateMaterialTopologyTree(
+	struct FSubstrateTranslatorDataInterface& SubstrateTranslatorData, class UMaterialExpression* Parent, int32 OutputIndex)
 {
-	FSubstrateOperator& SubstrateOperator = Compiler->SubstrateCompilationRegisterOperator(SUBSTRATE_OPERATOR_BSDF, Compiler->SubstrateTreeStackGetPathUniqueId(), this->MaterialExpressionGuid, Parent, Compiler->SubstrateTreeStackGetParentPathUniqueId());
+	FSubstrateOperator& SubstrateOperator = SubstrateTranslatorData.SubstrateCompilationRegisterOperator(SUBSTRATE_OPERATOR_BSDF, SubstrateTranslatorData.SubstrateTreeStackGetPathUniqueId(), this->MaterialExpressionGuid, Parent, SubstrateTranslatorData.SubstrateTreeStackGetParentPathUniqueId());
 	SubstrateOperator.BSDFType = SUBSTRATE_BSDF_TYPE_UNLIT;
-	SubstrateOperator.ThicknessIndex = Compiler->SubstrateThicknessStackGetThicknessIndex();
+	SubstrateOperator.ThicknessIndex = SubstrateTranslatorData.SubstrateThicknessStackGetThicknessIndex();
 	SubstrateOperator.bBSDFWritesEmissive = EmissiveColor.IsConnected();
 	return &SubstrateOperator;
 }
@@ -2150,12 +2203,13 @@ void UMaterialExpressionSubstrateHairBSDF::GatherSubstrateMaterialInfo(FSubstrat
 	SubstrateMaterialInfo.AddGuid(MaterialExpressionGuid);
 }
 
-FSubstrateOperator* UMaterialExpressionSubstrateHairBSDF::SubstrateGenerateMaterialTopologyTree(class FMaterialCompiler* Compiler, class UMaterialExpression* Parent, int32 OutputIndex)
+FSubstrateOperator* UMaterialExpressionSubstrateHairBSDF::SubstrateGenerateMaterialTopologyTree(
+	struct FSubstrateTranslatorDataInterface& SubstrateTranslatorData, class UMaterialExpression* Parent, int32 OutputIndex)
 {
-	FSubstrateOperator& SubstrateOperator = Compiler->SubstrateCompilationRegisterOperator(SUBSTRATE_OPERATOR_BSDF, Compiler->SubstrateTreeStackGetPathUniqueId(), this->MaterialExpressionGuid, Parent, Compiler->SubstrateTreeStackGetParentPathUniqueId());
+	FSubstrateOperator& SubstrateOperator = SubstrateTranslatorData.SubstrateCompilationRegisterOperator(SUBSTRATE_OPERATOR_BSDF, SubstrateTranslatorData.SubstrateTreeStackGetPathUniqueId(), this->MaterialExpressionGuid, Parent, SubstrateTranslatorData.SubstrateTreeStackGetParentPathUniqueId());
 	SubstrateOperator.BSDFType = SUBSTRATE_BSDF_TYPE_HAIR;
 	SubstrateOperator.BSDFFeatures = ESubstrateBsdfFeature::Hair;
-	SubstrateOperator.ThicknessIndex = Compiler->SubstrateThicknessStackGetThicknessIndex();
+	SubstrateOperator.ThicknessIndex = SubstrateTranslatorData.SubstrateThicknessStackGetThicknessIndex();
 	SubstrateOperator.bBSDFWritesEmissive = EmissiveColor.IsConnected();
 	return &SubstrateOperator;
 }
@@ -2271,9 +2325,10 @@ void UMaterialExpressionSubstrateEyeBSDF::GatherSubstrateMaterialInfo(FSubstrate
 	SubstrateMaterialInfo.AddGuid(MaterialExpressionGuid);
 }
 
-FSubstrateOperator* UMaterialExpressionSubstrateEyeBSDF::SubstrateGenerateMaterialTopologyTree(class FMaterialCompiler* Compiler, class UMaterialExpression* Parent, int32 OutputIndex)
+FSubstrateOperator* UMaterialExpressionSubstrateEyeBSDF::SubstrateGenerateMaterialTopologyTree(
+	struct FSubstrateTranslatorDataInterface& SubstrateTranslatorData, class UMaterialExpression* Parent, int32 OutputIndex)
 {
-	FSubstrateOperator& SubstrateOperator = Compiler->SubstrateCompilationRegisterOperator(SUBSTRATE_OPERATOR_BSDF, Compiler->SubstrateTreeStackGetPathUniqueId(), this->MaterialExpressionGuid, Parent, Compiler->SubstrateTreeStackGetParentPathUniqueId());
+	FSubstrateOperator& SubstrateOperator = SubstrateTranslatorData.SubstrateCompilationRegisterOperator(SUBSTRATE_OPERATOR_BSDF, SubstrateTranslatorData.SubstrateTreeStackGetPathUniqueId(), this->MaterialExpressionGuid, Parent, SubstrateTranslatorData.SubstrateTreeStackGetParentPathUniqueId());
 	SubstrateOperator.BSDFType = SUBSTRATE_BSDF_TYPE_EYE;
 	SubstrateOperator.BSDFFeatures = ESubstrateBsdfFeature::Eye;
 	if (IrisNormal.IsConnected())
@@ -2285,7 +2340,7 @@ FSubstrateOperator* UMaterialExpressionSubstrateEyeBSDF::SubstrateGenerateMateri
 		SubstrateOperator.BSDFFeatures |= ESubstrateBsdfFeature::EyeIrisTangentPluggedIn;
 	}
 
-	SubstrateOperator.ThicknessIndex = Compiler->SubstrateThicknessStackGetThicknessIndex();
+	SubstrateOperator.ThicknessIndex = SubstrateTranslatorData.SubstrateThicknessStackGetThicknessIndex();
 	SubstrateOperator.bBSDFWritesEmissive = EmissiveColor.IsConnected();
 	return &SubstrateOperator;
 }
@@ -2403,11 +2458,12 @@ void UMaterialExpressionSubstrateSingleLayerWaterBSDF::GatherSubstrateMaterialIn
 	SubstrateMaterialInfo.AddGuid(MaterialExpressionGuid);
 }
 
-FSubstrateOperator* UMaterialExpressionSubstrateSingleLayerWaterBSDF::SubstrateGenerateMaterialTopologyTree(class FMaterialCompiler* Compiler, class UMaterialExpression* Parent, int32 OutputIndex)
+FSubstrateOperator* UMaterialExpressionSubstrateSingleLayerWaterBSDF::SubstrateGenerateMaterialTopologyTree(
+	struct FSubstrateTranslatorDataInterface& SubstrateTranslatorData, class UMaterialExpression* Parent, int32 OutputIndex)
 {
-	FSubstrateOperator& SubstrateOperator = Compiler->SubstrateCompilationRegisterOperator(SUBSTRATE_OPERATOR_BSDF, Compiler->SubstrateTreeStackGetPathUniqueId(), this->MaterialExpressionGuid, Parent, Compiler->SubstrateTreeStackGetParentPathUniqueId());
+	FSubstrateOperator& SubstrateOperator = SubstrateTranslatorData.SubstrateCompilationRegisterOperator(SUBSTRATE_OPERATOR_BSDF, SubstrateTranslatorData.SubstrateTreeStackGetPathUniqueId(), this->MaterialExpressionGuid, Parent, SubstrateTranslatorData.SubstrateTreeStackGetParentPathUniqueId());
 	SubstrateOperator.BSDFType = SUBSTRATE_BSDF_TYPE_SINGLELAYERWATER;
-	SubstrateOperator.ThicknessIndex = Compiler->SubstrateThicknessStackGetThicknessIndex();
+	SubstrateOperator.ThicknessIndex = SubstrateTranslatorData.SubstrateThicknessStackGetThicknessIndex();
 	SubstrateOperator.bBSDFWritesEmissive = EmissiveColor.IsConnected();
 	return &SubstrateOperator;
 }
@@ -2516,11 +2572,12 @@ void UMaterialExpressionSubstrateToonBSDF::GatherSubstrateMaterialInfo(FSubstrat
 	}
 }
 
-FSubstrateOperator* UMaterialExpressionSubstrateToonBSDF::SubstrateGenerateMaterialTopologyTree(class FMaterialCompiler* Compiler, class UMaterialExpression* Parent, int32 OutputIndex)
+FSubstrateOperator* UMaterialExpressionSubstrateToonBSDF::SubstrateGenerateMaterialTopologyTree(
+	struct FSubstrateTranslatorDataInterface& SubstrateTranslatorData, class UMaterialExpression* Parent, int32 OutputIndex)
 {
-	FSubstrateOperator& SubstrateOperator = Compiler->SubstrateCompilationRegisterOperator(SUBSTRATE_OPERATOR_BSDF, Compiler->SubstrateTreeStackGetPathUniqueId(), this->MaterialExpressionGuid, Parent, Compiler->SubstrateTreeStackGetParentPathUniqueId());
+	FSubstrateOperator& SubstrateOperator = SubstrateTranslatorData.SubstrateCompilationRegisterOperator(SUBSTRATE_OPERATOR_BSDF, SubstrateTranslatorData.SubstrateTreeStackGetPathUniqueId(), this->MaterialExpressionGuid, Parent, SubstrateTranslatorData.SubstrateTreeStackGetParentPathUniqueId());
 	SubstrateOperator.BSDFType = SUBSTRATE_BSDF_TYPE_TOON;
-	SubstrateOperator.ThicknessIndex = Compiler->SubstrateThicknessStackGetThicknessIndex();
+	SubstrateOperator.ThicknessIndex = SubstrateTranslatorData.SubstrateThicknessStackGetThicknessIndex();
 	SubstrateOperator.bBSDFWritesEmissive = EmissiveColor.IsConnected();
 	SubstrateOperator.BSDFFeatures |= ESubstrateBsdfFeature::Toon;
 	if (Anisotropy.IsConnected())
@@ -2680,10 +2737,11 @@ void UMaterialExpressionSubstrateHorizontalMixing::GatherSubstrateMaterialInfo(F
 	SubstrateMaterialInfo.PopSubstrateTreeStack();
 }
 
-FSubstrateOperator* UMaterialExpressionSubstrateHorizontalMixing::SubstrateGenerateMaterialTopologyTree(class FMaterialCompiler* Compiler, class UMaterialExpression* Parent, int32 OutputIndex)
+FSubstrateOperator* UMaterialExpressionSubstrateHorizontalMixing::SubstrateGenerateMaterialTopologyTree(
+	struct FSubstrateTranslatorDataInterface& SubstrateTranslatorData, class UMaterialExpression* Parent, int32 OutputIndex)
 {
-	FSubstrateOperator& SubstrateOperator = Compiler->SubstrateCompilationRegisterOperator(SUBSTRATE_OPERATOR_HORIZONTAL, Compiler->SubstrateTreeStackGetPathUniqueId(), this->MaterialExpressionGuid, Parent, Compiler->SubstrateTreeStackGetParentPathUniqueId(), bUseParameterBlending);
-	if (Compiler->GetSubstrateTreeOutOfStackDepthOccurred())
+	FSubstrateOperator& SubstrateOperator = SubstrateTranslatorData.SubstrateCompilationRegisterOperator(SUBSTRATE_OPERATOR_HORIZONTAL, SubstrateTranslatorData.SubstrateTreeStackGetPathUniqueId(), this->MaterialExpressionGuid, Parent, SubstrateTranslatorData.SubstrateTreeStackGetParentPathUniqueId(), bUseParameterBlending);
+	if (SubstrateTranslatorData.GetSubstrateTreeOutOfStackDepthOccurred())
 	{
 		return &SubstrateOperator; // Out ot stack space, return now to fail the compilation
 	}
@@ -2696,16 +2754,16 @@ FSubstrateOperator* UMaterialExpressionSubstrateHorizontalMixing::SubstrateGener
 	FSubstrateOperator* OpB = nullptr;
 	if (ChildAExpression)
 	{
-		Compiler->SubstrateTreeStackPush(this, 0);
-		OpA = ChildAExpression->SubstrateGenerateMaterialTopologyTree(Compiler, this, TracedInputA.OutputIndex);
-		Compiler->SubstrateTreeStackPop();
+		SubstrateTranslatorData.SubstrateTreeStackPush(this, 0);
+		OpA = ChildAExpression->SubstrateGenerateMaterialTopologyTree(SubstrateTranslatorData, this, TracedInputA.OutputIndex);
+		SubstrateTranslatorData.SubstrateTreeStackPop();
 		AssignOperatorIndexIfNotNull(SubstrateOperator.LeftIndex, OpA);
 	}
 	if (ChildBExpression)
 	{
-		Compiler->SubstrateTreeStackPush(this, 1);
-		OpB = ChildBExpression->SubstrateGenerateMaterialTopologyTree(Compiler, this, TracedInputB.OutputIndex);
-		Compiler->SubstrateTreeStackPop();
+		SubstrateTranslatorData.SubstrateTreeStackPush(this, 1);
+		OpB = ChildBExpression->SubstrateGenerateMaterialTopologyTree(SubstrateTranslatorData, this, TracedInputB.OutputIndex);
+		SubstrateTranslatorData.SubstrateTreeStackPop();
 		AssignOperatorIndexIfNotNull(SubstrateOperator.RightIndex, OpB);
 	}
 	CombineFlagForParameterBlending(SubstrateOperator, OpA, OpB);
@@ -2880,10 +2938,11 @@ void UMaterialExpressionSubstrateVerticalLayering::GatherSubstrateMaterialInfo(F
 	SubstrateMaterialInfo.PopSubstrateTreeStack();
 }
 
-FSubstrateOperator* UMaterialExpressionSubstrateVerticalLayering::SubstrateGenerateMaterialTopologyTree(class FMaterialCompiler* Compiler, class UMaterialExpression* Parent, int32 OutputIndex)
+FSubstrateOperator* UMaterialExpressionSubstrateVerticalLayering::SubstrateGenerateMaterialTopologyTree(
+	struct FSubstrateTranslatorDataInterface& SubstrateTranslatorData, class UMaterialExpression* Parent, int32 OutputIndex)
 {
-	FSubstrateOperator& SubstrateOperator = Compiler->SubstrateCompilationRegisterOperator(SUBSTRATE_OPERATOR_VERTICAL, Compiler->SubstrateTreeStackGetPathUniqueId(), this->MaterialExpressionGuid, Parent, Compiler->SubstrateTreeStackGetParentPathUniqueId(), bUseParameterBlending);
-	if (Compiler->GetSubstrateTreeOutOfStackDepthOccurred())
+	FSubstrateOperator& SubstrateOperator = SubstrateTranslatorData.SubstrateCompilationRegisterOperator(SUBSTRATE_OPERATOR_VERTICAL, SubstrateTranslatorData.SubstrateTreeStackGetPathUniqueId(), this->MaterialExpressionGuid, Parent, SubstrateTranslatorData.SubstrateTreeStackGetParentPathUniqueId(), bUseParameterBlending);
+	if (SubstrateTranslatorData.GetSubstrateTreeOutOfStackDepthOccurred())
 	{
 		return &SubstrateOperator; // Out ot stack space, return now to fail the compilation
 	}
@@ -2898,19 +2957,19 @@ FSubstrateOperator* UMaterialExpressionSubstrateVerticalLayering::SubstrateGener
 	// Top - Use the vertical operator thickness
 	if (ChildAExpression)
 	{
-		Compiler->SubstrateThicknessStackPush(this, &Thickness);
-		Compiler->SubstrateTreeStackPush(this, 0);
-		OpA = ChildAExpression->SubstrateGenerateMaterialTopologyTree(Compiler, this, TracedInputTop.OutputIndex);
-		Compiler->SubstrateTreeStackPop();
-		Compiler->SubstrateThicknessStackPop();
+		SubstrateTranslatorData.SubstrateThicknessStackPush(this, &Thickness);
+		SubstrateTranslatorData.SubstrateTreeStackPush(this, 0);
+		OpA = ChildAExpression->SubstrateGenerateMaterialTopologyTree(SubstrateTranslatorData, this, TracedInputTop.OutputIndex);
+		SubstrateTranslatorData.SubstrateTreeStackPop();
+		SubstrateTranslatorData.SubstrateThicknessStackPop();
 		AssignOperatorIndexIfNotNull(SubstrateOperator.LeftIndex, OpA);
 	}
 	// Bottom - Use the propagated thickness from parent
 	if (ChildBExpression)
 	{
-		Compiler->SubstrateTreeStackPush(this, 1);
-		OpB = ChildBExpression->SubstrateGenerateMaterialTopologyTree(Compiler, this, TracedInputBase.OutputIndex);
-		Compiler->SubstrateTreeStackPop();
+		SubstrateTranslatorData.SubstrateTreeStackPush(this, 1);
+		OpB = ChildBExpression->SubstrateGenerateMaterialTopologyTree(SubstrateTranslatorData, this, TracedInputBase.OutputIndex);
+		SubstrateTranslatorData.SubstrateTreeStackPop();
 		AssignOperatorIndexIfNotNull(SubstrateOperator.RightIndex, OpB);
 	}
 	CombineFlagForParameterBlending(SubstrateOperator, OpA, OpB);
@@ -3059,10 +3118,11 @@ void UMaterialExpressionSubstrateAdd::GatherSubstrateMaterialInfo(FSubstrateMate
 	SubstrateMaterialInfo.PopSubstrateTreeStack();
 }
 
-FSubstrateOperator* UMaterialExpressionSubstrateAdd::SubstrateGenerateMaterialTopologyTree(class FMaterialCompiler* Compiler, class UMaterialExpression* Parent, int32 OutputIndex)
+FSubstrateOperator* UMaterialExpressionSubstrateAdd::SubstrateGenerateMaterialTopologyTree(
+	struct FSubstrateTranslatorDataInterface& SubstrateTranslatorData, class UMaterialExpression* Parent, int32 OutputIndex)
 {
-	FSubstrateOperator& SubstrateOperator = Compiler->SubstrateCompilationRegisterOperator(SUBSTRATE_OPERATOR_ADD, Compiler->SubstrateTreeStackGetPathUniqueId(), this->MaterialExpressionGuid, Parent, Compiler->SubstrateTreeStackGetParentPathUniqueId(), bUseParameterBlending);
-	if (Compiler->GetSubstrateTreeOutOfStackDepthOccurred())
+	FSubstrateOperator& SubstrateOperator = SubstrateTranslatorData.SubstrateCompilationRegisterOperator(SUBSTRATE_OPERATOR_ADD, SubstrateTranslatorData.SubstrateTreeStackGetPathUniqueId(), this->MaterialExpressionGuid, Parent, SubstrateTranslatorData.SubstrateTreeStackGetParentPathUniqueId(), bUseParameterBlending);
+	if (SubstrateTranslatorData.GetSubstrateTreeOutOfStackDepthOccurred())
 	{
 		return &SubstrateOperator; // Out ot stack space, return now to fail the compilation
 	}
@@ -3075,16 +3135,16 @@ FSubstrateOperator* UMaterialExpressionSubstrateAdd::SubstrateGenerateMaterialTo
 	FSubstrateOperator* OpB = nullptr;
 	if (ChildAExpression)
 	{
-		Compiler->SubstrateTreeStackPush(this, 0);
-		OpA = ChildAExpression->SubstrateGenerateMaterialTopologyTree(Compiler, this, TracedInputA.OutputIndex);
-		Compiler->SubstrateTreeStackPop();
+		SubstrateTranslatorData.SubstrateTreeStackPush(this, 0);
+		OpA = ChildAExpression->SubstrateGenerateMaterialTopologyTree(SubstrateTranslatorData, this, TracedInputA.OutputIndex);
+		SubstrateTranslatorData.SubstrateTreeStackPop();
 		AssignOperatorIndexIfNotNull(SubstrateOperator.LeftIndex, OpA);
 	}
 	if (ChildBExpression)
 	{
-		Compiler->SubstrateTreeStackPush(this, 1);
-		OpB = ChildBExpression->SubstrateGenerateMaterialTopologyTree(Compiler, this, TracedInputB.OutputIndex);
-		Compiler->SubstrateTreeStackPop();
+		SubstrateTranslatorData.SubstrateTreeStackPush(this, 1);
+		OpB = ChildBExpression->SubstrateGenerateMaterialTopologyTree(SubstrateTranslatorData, this, TracedInputB.OutputIndex);
+		SubstrateTranslatorData.SubstrateTreeStackPop();
 		AssignOperatorIndexIfNotNull(SubstrateOperator.RightIndex, OpB);
 	}
 	CombineFlagForParameterBlending(SubstrateOperator, OpA, OpB);
@@ -3197,10 +3257,11 @@ void UMaterialExpressionSubstrateWeight::GatherSubstrateMaterialInfo(FSubstrateM
 	SubstrateMaterialInfo.PopSubstrateTreeStack();
 }
 
-FSubstrateOperator* UMaterialExpressionSubstrateWeight::SubstrateGenerateMaterialTopologyTree(class FMaterialCompiler* Compiler, class UMaterialExpression* Parent, int32 OutputIndex)
+FSubstrateOperator* UMaterialExpressionSubstrateWeight::SubstrateGenerateMaterialTopologyTree(
+	struct FSubstrateTranslatorDataInterface& SubstrateTranslatorData, class UMaterialExpression* Parent, int32 OutputIndex)
 {
-	FSubstrateOperator& SubstrateOperator = Compiler->SubstrateCompilationRegisterOperator(SUBSTRATE_OPERATOR_WEIGHT, Compiler->SubstrateTreeStackGetPathUniqueId(), this->MaterialExpressionGuid, Parent, Compiler->SubstrateTreeStackGetParentPathUniqueId());
-	if (Compiler->GetSubstrateTreeOutOfStackDepthOccurred())
+	FSubstrateOperator& SubstrateOperator = SubstrateTranslatorData.SubstrateCompilationRegisterOperator(SUBSTRATE_OPERATOR_WEIGHT, SubstrateTranslatorData.SubstrateTreeStackGetPathUniqueId(), this->MaterialExpressionGuid, Parent, SubstrateTranslatorData.SubstrateTreeStackGetParentPathUniqueId());
+	if (SubstrateTranslatorData.GetSubstrateTreeOutOfStackDepthOccurred())
 	{
 		return &SubstrateOperator; // Out ot stack space, return now to fail the compilation
 	}
@@ -3210,9 +3271,9 @@ FSubstrateOperator* UMaterialExpressionSubstrateWeight::SubstrateGenerateMateria
 	FSubstrateOperator* OpA = nullptr;
 	if (ChildAExpression)
 	{
-		Compiler->SubstrateTreeStackPush(this, 0);
-		OpA = ChildAExpression->SubstrateGenerateMaterialTopologyTree(Compiler, this, TracedInputA.OutputIndex);
-		Compiler->SubstrateTreeStackPop();
+		SubstrateTranslatorData.SubstrateTreeStackPush(this, 0);
+		OpA = ChildAExpression->SubstrateGenerateMaterialTopologyTree(SubstrateTranslatorData, this, TracedInputA.OutputIndex);
+		SubstrateTranslatorData.SubstrateTreeStackPop();
 		AssignOperatorIndexIfNotNull(SubstrateOperator.LeftIndex, OpA);
 	}
 	CombineFlagForParameterBlending(SubstrateOperator, OpA);
@@ -3364,10 +3425,11 @@ void UMaterialExpressionSubstrateSelect::GatherSubstrateMaterialInfo(FSubstrateM
 	SubstrateMaterialInfo.PopSubstrateTreeStack();
 }
 
-FSubstrateOperator* UMaterialExpressionSubstrateSelect::SubstrateGenerateMaterialTopologyTree(class FMaterialCompiler* Compiler, class UMaterialExpression* Parent, int32 OutputIndex)
+FSubstrateOperator* UMaterialExpressionSubstrateSelect::SubstrateGenerateMaterialTopologyTree(
+	struct FSubstrateTranslatorDataInterface& SubstrateTranslatorData, class UMaterialExpression* Parent, int32 OutputIndex)
 {
-	FSubstrateOperator& SubstrateOperator = Compiler->SubstrateCompilationRegisterOperator(SUBSTRATE_OPERATOR_SELECT, Compiler->SubstrateTreeStackGetPathUniqueId(), this->MaterialExpressionGuid, Parent, Compiler->SubstrateTreeStackGetParentPathUniqueId(), bUseParameterBlending);
-	if (Compiler->GetSubstrateTreeOutOfStackDepthOccurred())
+	FSubstrateOperator& SubstrateOperator = SubstrateTranslatorData.SubstrateCompilationRegisterOperator(SUBSTRATE_OPERATOR_SELECT, SubstrateTranslatorData.SubstrateTreeStackGetPathUniqueId(), this->MaterialExpressionGuid, Parent, SubstrateTranslatorData.SubstrateTreeStackGetParentPathUniqueId(), bUseParameterBlending);
+	if (SubstrateTranslatorData.GetSubstrateTreeOutOfStackDepthOccurred())
 	{
 		return &SubstrateOperator; // Out ot stack space, return now to fail the compilation
 	}
@@ -3380,16 +3442,16 @@ FSubstrateOperator* UMaterialExpressionSubstrateSelect::SubstrateGenerateMateria
 	FSubstrateOperator* OpB = nullptr;
 	if (ChildAExpression)
 	{
-		Compiler->SubstrateTreeStackPush(this, 0);
-		OpA = ChildAExpression->SubstrateGenerateMaterialTopologyTree(Compiler, this, TracedInputA.OutputIndex);
-		Compiler->SubstrateTreeStackPop();
+		SubstrateTranslatorData.SubstrateTreeStackPush(this, 0);
+		OpA = ChildAExpression->SubstrateGenerateMaterialTopologyTree(SubstrateTranslatorData, this, TracedInputA.OutputIndex);
+		SubstrateTranslatorData.SubstrateTreeStackPop();
 		AssignOperatorIndexIfNotNull(SubstrateOperator.LeftIndex, OpA);
 	}
 	if (ChildBExpression)
 	{
-		Compiler->SubstrateTreeStackPush(this, 1);
-		OpB = ChildBExpression->SubstrateGenerateMaterialTopologyTree(Compiler, this, TracedInputB.OutputIndex);
-		Compiler->SubstrateTreeStackPop();
+		SubstrateTranslatorData.SubstrateTreeStackPush(this, 1);
+		OpB = ChildBExpression->SubstrateGenerateMaterialTopologyTree(SubstrateTranslatorData, this, TracedInputB.OutputIndex);
+		SubstrateTranslatorData.SubstrateTreeStackPop();
 		AssignOperatorIndexIfNotNull(SubstrateOperator.RightIndex, OpB);
 	}
 
@@ -4244,24 +4306,32 @@ void UMaterialExpressionSubstrateConvertMaterialAttributes::GatherSubstrateMater
 FSubstrateOperator* UMaterialExpressionSubstrateConvertMaterialAttributes::SubstrateGenerateMaterialTopologyTreeCommon(
 	class FMaterialCompiler* Compiler, struct FGuid ThisExpressionGuid, class UMaterialExpression* Parent, int32 OutputIndex,
 	const uint64 CachedConnectedMaterialAttributesInputs, const bool bShadingModelFromMaterialExpression, const bool bIsEmissiveConnected)
+{
+	return SubstrateGenerateMaterialTopologyTreeCommon(Compiler->GetSubstrateTranslatorData(), ThisExpressionGuid, Parent, OutputIndex,
+		CachedConnectedMaterialAttributesInputs, bShadingModelFromMaterialExpression, bIsEmissiveConnected);
+}
+
+FSubstrateOperator* UMaterialExpressionSubstrateConvertMaterialAttributes::SubstrateGenerateMaterialTopologyTreeCommon(
+	struct FSubstrateTranslatorDataInterface& SubstrateTranslatorData, struct FGuid ThisExpressionGuid, class UMaterialExpression* Parent, int32 OutputIndex,
+	const uint64 CachedConnectedMaterialAttributesInputs, const bool bShadingModelFromMaterialExpression, const bool bIsEmissiveConnected)
 {
 	// Note Thickness has no meaning/usage in the context of SubstrateLegacyConversionNode
-	int32 ThicknessIndex = Compiler->SubstrateThicknessStackGetThicknessIndex();
+	int32 ThicknessIndex = SubstrateTranslatorData.SubstrateThicknessStackGetThicknessIndex();
 
 	const bool bHasAnisotropy = FMaterialAttributeDefinitionMap::IsAttributeInBitmask(CachedConnectedMaterialAttributesInputs, MP_Anisotropy);
 
 	auto AddDefaultWorstCase = [&](ESubstrateBsdfFeature In)
-		{
-			FSubstrateOperator& SlabOperator = Compiler->SubstrateCompilationRegisterOperator(SUBSTRATE_OPERATOR_BSDF_LEGACY, Compiler->SubstrateTreeStackGetPathUniqueId(), ThisExpressionGuid, Parent, Compiler->SubstrateTreeStackGetParentPathUniqueId());
-			SlabOperator.BSDFType = SUBSTRATE_BSDF_TYPE_SLAB;
-			SlabOperator.BSDFFeatures = In | (bHasAnisotropy ? ESubstrateBsdfFeature::Anisotropy : ESubstrateBsdfFeature::None);
-			SlabOperator.ThicknessIndex = ThicknessIndex;
+	{
+		FSubstrateOperator& SlabOperator = SubstrateTranslatorData.SubstrateCompilationRegisterOperator(SUBSTRATE_OPERATOR_BSDF_LEGACY, SubstrateTranslatorData.SubstrateTreeStackGetPathUniqueId(), ThisExpressionGuid, Parent, SubstrateTranslatorData.SubstrateTreeStackGetParentPathUniqueId());
+		SlabOperator.BSDFType = SUBSTRATE_BSDF_TYPE_SLAB;
+		SlabOperator.BSDFFeatures = In | (bHasAnisotropy ? ESubstrateBsdfFeature::Anisotropy : ESubstrateBsdfFeature::None);
+		SlabOperator.ThicknessIndex = ThicknessIndex;
 
-			return &SlabOperator;
-		};
+		return &SlabOperator;
+	};
 
 	// Get the shading models resulting from the UMaterial::RebuildShadingModelField().
-	FMaterialShadingModelField ShadingModels = Compiler->GetMaterialShadingModels();
+	FMaterialShadingModelField ShadingModels = SubstrateTranslatorData.GetMaterialShadingModels();
 
 	// Logic about shading models and complexity should match UMaterialExpressionSubstrateConvertMaterialAttributes::Compile.
 	const bool bHasShadingModelFromExpression = FMaterialAttributeDefinitionMap::IsAttributeInBitmask(CachedConnectedMaterialAttributesInputs, MP_ShadingModel) || bShadingModelFromMaterialExpression; // We keep HasShadingModelFromExpression in case all shading models cannot be safely recovered from material functions.
@@ -4277,7 +4347,7 @@ FSubstrateOperator* UMaterialExpressionSubstrateConvertMaterialAttributes::Subst
 
 		if (ShadingModels.HasShadingModel(MSM_Unlit))
 		{
-			FSubstrateOperator& Operator = Compiler->SubstrateCompilationRegisterOperator(SUBSTRATE_OPERATOR_BSDF_LEGACY, Compiler->SubstrateTreeStackGetPathUniqueId(), ThisExpressionGuid, Parent, Compiler->SubstrateTreeStackGetParentPathUniqueId());
+			FSubstrateOperator& Operator = SubstrateTranslatorData.SubstrateCompilationRegisterOperator(SUBSTRATE_OPERATOR_BSDF_LEGACY, SubstrateTranslatorData.SubstrateTreeStackGetPathUniqueId(), ThisExpressionGuid, Parent, SubstrateTranslatorData.SubstrateTreeStackGetParentPathUniqueId());
 			Operator.BSDFType = SUBSTRATE_BSDF_TYPE_UNLIT;
 			Operator.ThicknessIndex = ThicknessIndex;
 			return &Operator;
@@ -4312,7 +4382,7 @@ FSubstrateOperator* UMaterialExpressionSubstrateConvertMaterialAttributes::Subst
 		}
 		else if (ShadingModels.HasShadingModel(MSM_ClearCoat))
 		{
-			FSubstrateOperator& Operator = Compiler->SubstrateCompilationRegisterOperator(SUBSTRATE_OPERATOR_BSDF_LEGACY, Compiler->SubstrateTreeStackGetPathUniqueId(), ThisExpressionGuid, Parent, Compiler->SubstrateTreeStackGetParentPathUniqueId());
+			FSubstrateOperator& Operator = SubstrateTranslatorData.SubstrateCompilationRegisterOperator(SUBSTRATE_OPERATOR_BSDF_LEGACY, SubstrateTranslatorData.SubstrateTreeStackGetPathUniqueId(), ThisExpressionGuid, Parent, SubstrateTranslatorData.SubstrateTreeStackGetParentPathUniqueId());
 			Operator.BSDFType = SUBSTRATE_BSDF_TYPE_SLAB;
 			Operator.ThicknessIndex = ThicknessIndex;
 			Operator.BSDFFeatures = ESubstrateBsdfFeature::SecondRoughnessOrSimpleClearCoat | (bHasAnisotropy ? ESubstrateBsdfFeature::Anisotropy : ESubstrateBsdfFeature::None);
@@ -4320,7 +4390,7 @@ FSubstrateOperator* UMaterialExpressionSubstrateConvertMaterialAttributes::Subst
 		}
 		else if (ShadingModels.HasShadingModel(MSM_Hair))
 		{
-			FSubstrateOperator& Operator = Compiler->SubstrateCompilationRegisterOperator(SUBSTRATE_OPERATOR_BSDF_LEGACY, Compiler->SubstrateTreeStackGetPathUniqueId(), ThisExpressionGuid, Parent, Compiler->SubstrateTreeStackGetParentPathUniqueId());
+			FSubstrateOperator& Operator = SubstrateTranslatorData.SubstrateCompilationRegisterOperator(SUBSTRATE_OPERATOR_BSDF_LEGACY, SubstrateTranslatorData.SubstrateTreeStackGetPathUniqueId(), ThisExpressionGuid, Parent, SubstrateTranslatorData.SubstrateTreeStackGetParentPathUniqueId());
 			Operator.BSDFType = SUBSTRATE_BSDF_TYPE_HAIR;
 			Operator.BSDFFeatures = ESubstrateBsdfFeature::Hair;
 			Operator.ThicknessIndex = ThicknessIndex;
@@ -4328,7 +4398,7 @@ FSubstrateOperator* UMaterialExpressionSubstrateConvertMaterialAttributes::Subst
 		}
 		else if (ShadingModels.HasShadingModel(MSM_Eye))
 		{
-			FSubstrateOperator& Operator = Compiler->SubstrateCompilationRegisterOperator(SUBSTRATE_OPERATOR_BSDF_LEGACY, Compiler->SubstrateTreeStackGetPathUniqueId(), ThisExpressionGuid, Parent, Compiler->SubstrateTreeStackGetParentPathUniqueId());
+			FSubstrateOperator& Operator = SubstrateTranslatorData.SubstrateCompilationRegisterOperator(SUBSTRATE_OPERATOR_BSDF_LEGACY, SubstrateTranslatorData.SubstrateTreeStackGetPathUniqueId(), ThisExpressionGuid, Parent, SubstrateTranslatorData.SubstrateTreeStackGetParentPathUniqueId());
 			Operator.BSDFType = SUBSTRATE_BSDF_TYPE_EYE;
 			Operator.BSDFFeatures = ESubstrateBsdfFeature::Eye;
 			Operator.ThicknessIndex = ThicknessIndex;
@@ -4336,7 +4406,7 @@ FSubstrateOperator* UMaterialExpressionSubstrateConvertMaterialAttributes::Subst
 		}
 		else if (ShadingModels.HasShadingModel(MSM_SingleLayerWater))
 		{
-			FSubstrateOperator& Operator = Compiler->SubstrateCompilationRegisterOperator(SUBSTRATE_OPERATOR_BSDF_LEGACY, Compiler->SubstrateTreeStackGetPathUniqueId(), ThisExpressionGuid, Parent, Compiler->SubstrateTreeStackGetParentPathUniqueId());
+			FSubstrateOperator& Operator = SubstrateTranslatorData.SubstrateCompilationRegisterOperator(SUBSTRATE_OPERATOR_BSDF_LEGACY, SubstrateTranslatorData.SubstrateTreeStackGetPathUniqueId(), ThisExpressionGuid, Parent, SubstrateTranslatorData.SubstrateTreeStackGetParentPathUniqueId());
 			Operator.BSDFType = SUBSTRATE_BSDF_TYPE_SINGLELAYERWATER;
 			Operator.ThicknessIndex = ThicknessIndex;
 			Operator.bBSDFWritesEmissive = bIsEmissiveConnected;
@@ -4349,12 +4419,12 @@ FSubstrateOperator* UMaterialExpressionSubstrateConvertMaterialAttributes::Subst
 	}
 }
 
-FSubstrateOperator* UMaterialExpressionSubstrateConvertMaterialAttributes::SubstrateGenerateMaterialTopologyTree(class FMaterialCompiler* Compiler, class UMaterialExpression* Parent, int32 OutputIndex)
+FSubstrateOperator* UMaterialExpressionSubstrateConvertMaterialAttributes::SubstrateGenerateMaterialTopologyTree(struct FSubstrateTranslatorDataInterface& SubstrateTranslatorData, class UMaterialExpression* Parent, int32 OutputIndex)
 {
 	const uint64 Cached = GetConnectedMaterialAttributesInputs(Material);
 
 	return SubstrateGenerateMaterialTopologyTreeCommon(
-		Compiler, MaterialExpressionGuid, Parent, OutputIndex,
+		SubstrateTranslatorData, MaterialExpressionGuid, Parent, OutputIndex,
 		Cached, ShadingModelOverride == MSM_FromMaterialExpression, MaterialAttributes.IsConnected(MP_EmissiveColor));
 }
 
diff --git a/Engine/Source/Runtime/Engine/Private/Materials/MaterialExpressions.cpp b/Engine/Source/Runtime/Engine/Private/Materials/MaterialExpressions.cpp
--- a/Engine/Source/Runtime/Engine/Private/Materials/MaterialExpressions.cpp
+++ b/Engine/Source/Runtime/Engine/Private/Materials/MaterialExpressions.cpp
@@ -2439,7 +2439,14 @@ bool UMaterialExpression::IsUsingNewHLSLGenerator() const
 
 FSubstrateOperator* UMaterialExpression::SubstrateGenerateMaterialTopologyTree(class FMaterialCompiler* Compiler, class UMaterialExpression* Parent, int32 OutputIndex)
 {
-	Compiler->Errorf(TEXT("%s nodes type does not support generating/processing/flowing Substrate data.\nPlease reach out to the development team for feedback and if you want support to be added."), *GetClass()->GetName());
+	// Deprecated so call the general version.
+	return SubstrateGenerateMaterialTopologyTree(Compiler->GetSubstrateTranslatorData(), Parent, OutputIndex);
+}
+
+FSubstrateOperator* UMaterialExpression::SubstrateGenerateMaterialTopologyTree(struct FSubstrateTranslatorDataInterface& SubstrateTranslatorData, class UMaterialExpression* Parent, int32 OutputIndex)
+{
+	// Deprecated to call the general version.
+	SubstrateTranslatorData.Errorf(FString::Printf(TEXT("%s nodes type does not support generating/processing/flowing Substrate data.\nPlease reach out to the development team for feedback and if you want support to be added."), *GetClass()->GetName()));
 	return nullptr;
 }
 
@@ -6551,11 +6558,11 @@ void UMaterialExpressionGetMaterialAttributes::GatherSubstrateMaterialInfo(FSubs
 	}
 }
 
-FSubstrateOperator* UMaterialExpressionGetMaterialAttributes::SubstrateGenerateMaterialTopologyTree(class FMaterialCompiler* Compiler, class UMaterialExpression* Parent, int32 OutputIndex)
+FSubstrateOperator* UMaterialExpressionGetMaterialAttributes::SubstrateGenerateMaterialTopologyTree(struct FSubstrateTranslatorDataInterface& SubstrateTranslatorData, class UMaterialExpression* Parent, int32 OutputIndex)
 {
 	if (IsResultSubstrateMaterial(OutputIndex))
 	{
-		return MaterialAttributes.Expression->SubstrateGenerateMaterialTopologyTree(Compiler, Parent, MaterialAttributes.OutputIndex);
+		return MaterialAttributes.Expression->SubstrateGenerateMaterialTopologyTree(SubstrateTranslatorData, Parent, MaterialAttributes.OutputIndex);
 	}
 	return nullptr;
 }
@@ -6924,12 +6931,12 @@ void UMaterialExpressionSetMaterialAttributes::GatherSubstrateMaterialInfo(FSubs
 	}
 }
 
-FSubstrateOperator* UMaterialExpressionSetMaterialAttributes::SubstrateGenerateMaterialTopologyTree(FMaterialCompiler* Compiler, UMaterialExpression* Parent, int32 OutputIndex)
+FSubstrateOperator* UMaterialExpressionSetMaterialAttributes::SubstrateGenerateMaterialTopologyTree(struct FSubstrateTranslatorDataInterface& SubstrateTranslatorData, UMaterialExpression* Parent, int32 OutputIndex)
 {
 	int32 InputIndex = INDEX_NONE;
 	if (GetSubstrateMaterialInputIndex(OutputIndex, InputIndex))
 	{
-		return Inputs[InputIndex].Expression->SubstrateGenerateMaterialTopologyTree(Compiler, Parent, Inputs[InputIndex].OutputIndex);
+		return Inputs[InputIndex].Expression->SubstrateGenerateMaterialTopologyTree(SubstrateTranslatorData, Parent, Inputs[InputIndex].OutputIndex);
 	}
 	return nullptr;
 }
@@ -7201,11 +7208,11 @@ void UMaterialExpressionBlendMaterialAttributes::GatherSubstrateMaterialInfo(FSu
 	}
 }
 
-FSubstrateOperator* UMaterialExpressionBlendMaterialAttributes::SubstrateGenerateMaterialTopologyTree(class FMaterialCompiler* Compiler, class UMaterialExpression* Parent, int32 OutputIndex)
+FSubstrateOperator* UMaterialExpressionBlendMaterialAttributes::SubstrateGenerateMaterialTopologyTree(struct FSubstrateTranslatorDataInterface& SubstrateTranslatorData, class UMaterialExpression* Parent, int32 OutputIndex)
 {
 	// SUBSTRATE_TODO: this likely no longer work (or only works with the new layering system). We would need to stop parsing and always do parameter blending at this stage.
 	const bool bUseParameterBlending = false;
-	FSubstrateOperator& SubstrateOperator = Compiler->SubstrateCompilationRegisterOperator(SUBSTRATE_OPERATOR_HORIZONTAL, Compiler->SubstrateTreeStackGetPathUniqueId(), this->MaterialExpressionGuid, Parent, Compiler->SubstrateTreeStackGetParentPathUniqueId(), bUseParameterBlending);
+	FSubstrateOperator& SubstrateOperator = SubstrateTranslatorData.SubstrateCompilationRegisterOperator(SUBSTRATE_OPERATOR_HORIZONTAL, SubstrateTranslatorData.SubstrateTreeStackGetPathUniqueId(), this->MaterialExpressionGuid, Parent, SubstrateTranslatorData.SubstrateTreeStackGetParentPathUniqueId(), bUseParameterBlending);
 
 	UMaterialExpression* ChildAExpression = A.GetTracedInput().Expression;
 	UMaterialExpression* ChildBExpression = B.GetTracedInput().Expression;
@@ -7213,16 +7220,16 @@ FSubstrateOperator* UMaterialExpressionBlendMaterialAttributes::SubstrateGenerat
 	FSubstrateOperator* OpB = nullptr;
 	if (ChildAExpression)
 	{
-		Compiler->SubstrateTreeStackPush(this, 0);
-		OpA = ChildAExpression->SubstrateGenerateMaterialTopologyTree(Compiler, this, A.OutputIndex);
-		Compiler->SubstrateTreeStackPop();
+		SubstrateTranslatorData.SubstrateTreeStackPush(this, 0);
+		OpA = ChildAExpression->SubstrateGenerateMaterialTopologyTree(SubstrateTranslatorData, this, A.OutputIndex);
+		SubstrateTranslatorData.SubstrateTreeStackPop();
 		AssignOperatorIndexIfNotNull(SubstrateOperator.LeftIndex, OpA);
 	}
 	if (ChildBExpression)
 	{
-		Compiler->SubstrateTreeStackPush(this, 1);
-		OpB = ChildBExpression->SubstrateGenerateMaterialTopologyTree(Compiler, this, B.OutputIndex);
-		Compiler->SubstrateTreeStackPop();
+		SubstrateTranslatorData.SubstrateTreeStackPush(this, 1);
+		OpB = ChildBExpression->SubstrateGenerateMaterialTopologyTree(SubstrateTranslatorData, this, B.OutputIndex);
+		SubstrateTranslatorData.SubstrateTreeStackPop();
 		AssignOperatorIndexIfNotNull(SubstrateOperator.RightIndex, OpB);
 	}
 	CombineFlagForParameterBlending(SubstrateOperator, OpA, OpB);
@@ -8809,18 +8816,28 @@ void UMaterialExpressionStaticSwitchParameter::GatherSubstrateMaterialInfo(FSubs
 	}
 }
 
-FSubstrateOperator* UMaterialExpressionStaticSwitchParameter::SubstrateGenerateMaterialTopologyTree(class FMaterialCompiler* Compiler, class UMaterialExpression* Parent, int32 OutputIndex)
+FSubstrateOperator* UMaterialExpressionStaticSwitchParameter::SubstrateGenerateMaterialTopologyTree(struct FSubstrateTranslatorDataInterface& SubstrateTranslatorData, class UMaterialExpression* Parent, int32 OutputIndex)
 {
 	if (DynamicBranch)
 	{
-		Compiler->Errorf(TEXT("Static Switch nodes processing Substrate data do not support dynamic branching. The compiler must know the topology when translating HLSL (different branches could have different topologies)."));
+		SubstrateTranslatorData.Errorf(TEXT("Static Switch nodes processing Substrate data do not support dynamic branching. The compiler must know the topology when translating HLSL (different branches could have different topologies)."));
 		return nullptr;
 	}
 
-	FExpressionInput* EffectiveInput = GetEffectiveInput(Compiler);
+	FExpressionInput* EffectiveInput = nullptr;
+	if (SubstrateTranslatorData.GetCompilerLegacy())
+	{
+		EffectiveInput = GetEffectiveInput(SubstrateTranslatorData.GetCompilerLegacy());
+	}
+	else
+	{
+		check(false);	// SUBSTRATE_TODO: Not implmenented. We need to do a selection of A or B based on static select evaluation.
+		EffectiveInput = &A;
+	}
+
 	if (EffectiveInput && EffectiveInput->Expression)
 	{
-		return EffectiveInput->Expression->SubstrateGenerateMaterialTopologyTree(Compiler, Parent, EffectiveInput->OutputIndex);
+		return EffectiveInput->Expression->SubstrateGenerateMaterialTopologyTree(SubstrateTranslatorData, Parent, EffectiveInput->OutputIndex);
 	}
 	return nullptr;
 }
@@ -9042,12 +9059,22 @@ void UMaterialExpressionStaticSwitch::GatherSubstrateMaterialInfo(FSubstrateMate
 	}
 }
 
-FSubstrateOperator* UMaterialExpressionStaticSwitch::SubstrateGenerateMaterialTopologyTree(class FMaterialCompiler* Compiler, class UMaterialExpression* Parent, int32 OutputIndex)
+FSubstrateOperator* UMaterialExpressionStaticSwitch::SubstrateGenerateMaterialTopologyTree(struct FSubstrateTranslatorDataInterface& SubstrateTranslatorData, class UMaterialExpression* Parent, int32 OutputIndex)
 {
-	FExpressionInput* EffectiveInput = GetEffectiveInput(Compiler);
+	FExpressionInput* EffectiveInput = nullptr;
+	if (SubstrateTranslatorData.GetCompilerLegacy())
+	{
+		EffectiveInput = GetEffectiveInput(SubstrateTranslatorData.GetCompilerLegacy());
+	}
+	else
+	{
+		check(false);	// SUBSTRATE_TODO: Not implmenented. We need to do a selection of A or B based on static select evaluation.
+		EffectiveInput = &A;
+	}
+
 	if (EffectiveInput && EffectiveInput->Expression)
 	{
-		return EffectiveInput->Expression->SubstrateGenerateMaterialTopologyTree(Compiler, Parent, EffectiveInput->OutputIndex);
+		return EffectiveInput->Expression->SubstrateGenerateMaterialTopologyTree(SubstrateTranslatorData, Parent, EffectiveInput->OutputIndex);
 	}
 	return nullptr;
 }
@@ -9215,19 +9242,29 @@ void UMaterialExpressionQualitySwitch::GatherSubstrateMaterialInfo(FSubstrateMat
 	}
 }
 
-FSubstrateOperator* UMaterialExpressionQualitySwitch::SubstrateGenerateMaterialTopologyTree(class FMaterialCompiler* Compiler, class UMaterialExpression* Parent, int32 OutputIndex)
+FSubstrateOperator* UMaterialExpressionQualitySwitch::SubstrateGenerateMaterialTopologyTree(struct FSubstrateTranslatorDataInterface& SubstrateTranslatorData, class UMaterialExpression* Parent, int32 OutputIndex)
 {
 	FExpressionInput DefaultTraced = Default.GetTracedInput();
 	if (!DefaultTraced.Expression)
 	{
-		Compiler->Errorf(TEXT("Quality switch missing default input"));
+		SubstrateTranslatorData.Errorf(TEXT("Quality switch missing default input"));
 		return nullptr;
 	}
 
-	FExpressionInput* EffectiveInput = GetEffectiveInput(Compiler);
+	FExpressionInput* EffectiveInput = nullptr;
+	if (SubstrateTranslatorData.GetCompilerLegacy())
+	{
+		EffectiveInput = GetEffectiveInput(SubstrateTranslatorData.GetCompilerLegacy());
+	}
+	else
+	{
+		check(false);	// SUBSTRATE_TODO: Not implmenented. We need to do a selection of A or B based on static select evaluation.
+		EffectiveInput = &Default;
+	}
+
 	if (EffectiveInput && EffectiveInput->Expression)
 	{
-		return EffectiveInput->Expression->SubstrateGenerateMaterialTopologyTree(Compiler, Parent, EffectiveInput->OutputIndex);
+		return EffectiveInput->Expression->SubstrateGenerateMaterialTopologyTree(SubstrateTranslatorData, Parent, EffectiveInput->OutputIndex);
 	}
 	return nullptr;
 }
@@ -9973,19 +10010,29 @@ void UMaterialExpressionShadingPathSwitch::GatherSubstrateMaterialInfo(FSubstrat
 	}
 }
 
-FSubstrateOperator* UMaterialExpressionShadingPathSwitch::SubstrateGenerateMaterialTopologyTree(class FMaterialCompiler* Compiler, class UMaterialExpression* Parent, int32 OutputIndex)
+FSubstrateOperator* UMaterialExpressionShadingPathSwitch::SubstrateGenerateMaterialTopologyTree(struct FSubstrateTranslatorDataInterface& SubstrateTranslatorData, class UMaterialExpression* Parent, int32 OutputIndex)
 {
 	FExpressionInput DefaultTraced = Default.GetTracedInput();
 	if (!DefaultTraced.Expression)
 	{
-		Compiler->Errorf(TEXT("Shading path switch missing default input"));
+		SubstrateTranslatorData.Errorf(TEXT("Shading path switch missing default input"));
 		return nullptr;
 	}
 
-	FExpressionInput* EffectiveInput = GetEffectiveInput(Compiler);
+	FExpressionInput* EffectiveInput = nullptr;
+	if (SubstrateTranslatorData.GetCompilerLegacy())
+	{
+		EffectiveInput = GetEffectiveInput(SubstrateTranslatorData.GetCompilerLegacy());
+	}
+	else
+	{
+		check(false);	// SUBSTRATE_TODO: Not implmenented. We need to do a selection of A or B based on static select evaluation.
+		EffectiveInput = &Default;
+	}
+
 	if (EffectiveInput && EffectiveInput->Expression)
 	{
-		return EffectiveInput->Expression->SubstrateGenerateMaterialTopologyTree(Compiler, Parent, EffectiveInput->OutputIndex);
+		return EffectiveInput->Expression->SubstrateGenerateMaterialTopologyTree(SubstrateTranslatorData, Parent, EffectiveInput->OutputIndex);
 	}
 	return nullptr;
 }
@@ -16296,17 +16343,17 @@ void UMaterialExpressionMaterialFunctionCall::GatherSubstrateMaterialInfo(FSubst
 	}
 }
 
-FSubstrateOperator* UMaterialExpressionMaterialFunctionCall::SubstrateGenerateMaterialTopologyTree(class FMaterialCompiler* Compiler, class UMaterialExpression* Parent, int32 OutputIndex)
+FSubstrateOperator* UMaterialExpressionMaterialFunctionCall::SubstrateGenerateMaterialTopologyTree(struct FSubstrateTranslatorDataInterface& SubstrateTranslatorData, class UMaterialExpression* Parent, int32 OutputIndex)
 {
 	if (OutputIndex >= 0 && OutputIndex < FunctionOutputs.Num() && FunctionOutputs[OutputIndex].ExpressionOutput)
 	{
 		this->LinkFunctionIntoCaller(nullptr);
 		FMaterialFunctionCompileState LocalState(this);
-		Compiler->PushFunction(SharedCompileState ? SharedCompileState : &LocalState);
+		SubstrateTranslatorData.PushFunction(SharedCompileState ? SharedCompileState : &LocalState);
 
-		FSubstrateOperator* ResultingOperator = FunctionOutputs[OutputIndex].ExpressionOutput->SubstrateGenerateMaterialTopologyTree(Compiler, Parent, 0);
+		FSubstrateOperator* ResultingOperator = FunctionOutputs[OutputIndex].ExpressionOutput->SubstrateGenerateMaterialTopologyTree(SubstrateTranslatorData, Parent, 0);
 
-		Compiler->PopFunction();
+		SubstrateTranslatorData.PopFunction();
 		this->UnlinkFunctionFromCaller(nullptr);
 		return ResultingOperator;
 	}
@@ -16815,7 +16862,7 @@ void UMaterialExpressionFunctionInput::GatherSubstrateMaterialInfo(FSubstrateMat
 	SubstrateMaterialInfo.AddGuid(MaterialExpressionGuid);
 }
 
-FSubstrateOperator* UMaterialExpressionFunctionInput::SubstrateGenerateMaterialTopologyTree(class FMaterialCompiler* Compiler, class UMaterialExpression* Parent, int32 OutputIndex)
+FSubstrateOperator* UMaterialExpressionFunctionInput::SubstrateGenerateMaterialTopologyTree(struct FSubstrateTranslatorDataInterface& SubstrateTranslatorData, class UMaterialExpression* Parent, int32 OutputIndex)
 {
 	FSubstrateOperator* SlabOperator = nullptr;
 	if (IsResultSubstrateMaterial(OutputIndex))
@@ -16825,11 +16872,11 @@ FSubstrateOperator* UMaterialExpressionFunctionInput::SubstrateGenerateMaterialT
 		{
 			// Backup EffectivePreviewDuringCompile which will be modified by UnlinkFromCaller and LinkIntoCaller of any potential chained function calls to the same function
 			FExpressionInput LocalPreviewTracedInput = LocalPreviewDuringCompile->GetTracedInput();
-			SlabOperator = LocalPreviewTracedInput.Expression->SubstrateGenerateMaterialTopologyTree(Compiler, Parent, LocalPreviewTracedInput.OutputIndex);
+			SlabOperator = LocalPreviewTracedInput.Expression->SubstrateGenerateMaterialTopologyTree(SubstrateTranslatorData, Parent, LocalPreviewTracedInput.OutputIndex);
 		}
 		else if(Preview.Expression && bUsePreviewValueAsDefault)
 		{
-			SlabOperator = Preview.Expression->SubstrateGenerateMaterialTopologyTree(Compiler, Parent, Preview.OutputIndex);
+			SlabOperator = Preview.Expression->SubstrateGenerateMaterialTopologyTree(SubstrateTranslatorData, Parent, Preview.OutputIndex);
 		}
 
 		if (LocalPreviewDuringCompile)
@@ -16843,14 +16890,7 @@ FSubstrateOperator* UMaterialExpressionFunctionInput::SubstrateGenerateMaterialT
 	// If this is not the case however, we need to compile a default material to compile the shader, or to preview the material function.
 	if(!SlabOperator)
 	{
-		int32 NormalCodeChunk = Compiler->VertexNormal();
-		const int32 NullTangentCodeChunk = INDEX_NONE;
-		const FSubstrateRegisteredSharedLocalBasis NewRegisteredSharedLocalBasis = SubstrateCompilationInfoCreateSharedLocalBasis(Compiler, NormalCodeChunk, NullTangentCodeChunk);
-
-		SlabOperator = &Compiler->SubstrateCompilationRegisterOperator(SUBSTRATE_OPERATOR_BSDF, Compiler->SubstrateTreeStackGetPathUniqueId(), this->MaterialExpressionGuid, Parent, Compiler->SubstrateTreeStackGetParentPathUniqueId());
-		SlabOperator->BSDFRegisteredSharedLocalBasis = NewRegisteredSharedLocalBasis;
-		SlabOperator->BSDFType = SUBSTRATE_BSDF_TYPE_SLAB;
-		SlabOperator->ThicknessIndex = Compiler->SubstrateThicknessStackGetThicknessIndex();
+		SlabOperator = &SubstrateTranslatorData.SubstrateGetDefaultBSDFOperator(this->MaterialExpressionGuid, Parent);
 	}
 	return SlabOperator;
 }
@@ -17089,11 +17129,11 @@ void UMaterialExpressionFunctionOutput::GatherSubstrateMaterialInfo(FSubstrateMa
 	}
 }
 
-FSubstrateOperator* UMaterialExpressionFunctionOutput::SubstrateGenerateMaterialTopologyTree(class FMaterialCompiler* Compiler, class UMaterialExpression* Parent, int32 OutputIndex)
+FSubstrateOperator* UMaterialExpressionFunctionOutput::SubstrateGenerateMaterialTopologyTree(struct FSubstrateTranslatorDataInterface& SubstrateTranslatorData, class UMaterialExpression* Parent, int32 OutputIndex)
 {
 	if (A.GetTracedInput().Expression)
 	{
-		return A.Expression->SubstrateGenerateMaterialTopologyTree(Compiler, Parent, A.OutputIndex);
+		return A.Expression->SubstrateGenerateMaterialTopologyTree(SubstrateTranslatorData, Parent, A.OutputIndex);
 	}
 	return nullptr;
 }
@@ -18066,7 +18106,7 @@ void UMaterialExpressionRerouteBase::GatherSubstrateMaterialInfo(FSubstrateMater
 	}
 }
 
-FSubstrateOperator* UMaterialExpressionRerouteBase::SubstrateGenerateMaterialTopologyTree(class FMaterialCompiler* Compiler, class UMaterialExpression* Parent, int32 OutputIndex)
+FSubstrateOperator* UMaterialExpressionRerouteBase::SubstrateGenerateMaterialTopologyTree(struct FSubstrateTranslatorDataInterface& SubstrateTranslatorData, class UMaterialExpression* Parent, int32 OutputIndex)
 {
 	FExpressionInput Input;
 	if (GetRerouteInput(Input))
@@ -18079,7 +18119,7 @@ FSubstrateOperator* UMaterialExpressionRerouteBase::SubstrateGenerateMaterialTop
 			UMaterialExpression* RealExpression = TraceInputsToRealExpression(RealExpressionOutputIndex);
 			if (RealExpression != nullptr)
 			{
-				RealExpression->SubstrateGenerateMaterialTopologyTree(Compiler, Parent, RealExpressionOutputIndex);
+				RealExpression->SubstrateGenerateMaterialTopologyTree(SubstrateTranslatorData, Parent, RealExpressionOutputIndex);
 			}
 		}
 	}
