Substrate

diff --git a/Engine/Source/Runtime/Engine/Private/Materials/MaterialExpressionsIR.cpp b/Engine/Source/Runtime/Engine/Private/Materials/MaterialExpressionsIR.cpp
--- a/Engine/Source/Runtime/Engine/Private/Materials/MaterialExpressionsIR.cpp
+++ b/Engine/Source/Runtime/Engine/Private/Materials/MaterialExpressionsIR.cpp
@@ -255,16 +255,20 @@
 #include "Materials/MaterialSharedPrivate.h"
 #include "Materials/MaterialExpressionMotionVectorWorldOffsetOutput.h"
 #include "Materials/MaterialExpressionSubsurfaceMediumMaterialOutput.h"
+#include "Materials/MaterialExpressionSubstrate.h"
 #include "MaterialShared.h"
 #include "Misc/MemStackUtility.h"
 #include "RenderUtils.h"
 #include "Shader/Preshader2.h"
 #include "TextureResource.h"
 #include "ColorManagement/ColorSpace.h"
+#include "SubstrateTranslatorCommon.h"
+#include "Engine/SubsurfaceProfile.h"
 
 #include "Materials/MaterialIREmitter.h"
 
 using FValueRef = MIR::FValueRef;
+using namespace Substrate;
 
 static FName NAME_CameraVector("CameraVector");
 
@@ -635,6 +639,22 @@ static void EmitExternalCodeConditionalReplace(MIR::FEmitter& Em, FName External
 	Em.Output(0, Em.Branch(EmitInlineHLSL(Em, ExternalCodeIdentifier), ReplaceValue, DefaultValue));
 }
 
+static FValueRef EmitVertexNormal(MIR::FEmitter& Em)
+{
+	/*
+		static FName NAME_VertexNormal("VertexNormal");
+		NormalValue = EmitInlineHLSL(Em, NAME_VertexNormal);
+	*/
+	static FName NAME_VertexNormal("VertexNormal");
+	return EmitInlineHLSL(Em, NAME_VertexNormal);
+}
+
+static FValueRef EmitVertexTangent(MIR::FEmitter& Em)
+{
+	static FName NAME_VertexTangent("VertexTangent");
+	return EmitInlineHLSL(Em, NAME_VertexTangent);
+}
+
 #define REPLACE_INPUTS(DEFAULT, REPLACE) DEFAULT, REPLACE, TEXT(#DEFAULT), TEXT(#REPLACE)
 
 void UMaterialExpressionLightmassReplace::Build(MIR::FEmitter& Em)
@@ -4323,4 +4343,1027 @@ void UMaterialExpressionAbsorptionMediumMaterialOutput::Build(MIR::FEmitter& Em)
 	Em.SetCustomOutputs(TEXTVIEW("AbsorptionMediumMaterialOutput"), OutputValues, MIR::MaterialOutputStage_Pixel);
 }
 
+////////////////////
+//  Substrate IR  //
+////////////////////
+
+struct FSubstrateRegisteredSharedLocalBasisSearchResult
+{
+	bool bFound = false;
+	FSubstrateSharedLocalBasesInfo BasisInfo;
+};
+FSubstrateRegisteredSharedLocalBasisSearchResult SubstrateCompilationFindExistingSharedLocalBasis(Substrate::FSubstrateTranslatorData* SubstrateTranslatorData, uint32 NormalCodeChunkHash, uint32 TangentCodeChunkHash)
+{
+	FSubstrateRegisteredSharedLocalBasisSearchResult Result;
+
+	FSubstrateCompilationContext& SubstrateCtx = SubstrateTranslatorData->SubstrateCompilationContext[SubstrateTranslatorData->CurrentSubstrateCompilationContext];
+
+	// Find a basis which matches both the Normal & the Tangent code chunks
+	TArray<FSubstrateSharedLocalBasesInfo*> NormalInfos;
+	SubstrateCtx.CodeChunkToSubstrateSharedLocalBasis.MultiFindPointer(NormalCodeChunkHash, NormalInfos);
+	for (FSubstrateSharedLocalBasesInfo* NormalInfo : NormalInfos)
+	{
+		// * Either we find a perfect match (normal & tangent matches)
+		// * Or we find a normal which doesn't have a tangent associated with, and we set the tangent for code
+		if (TangentCodeChunkHash == INDEX_NONE || TangentCodeChunkHash == NormalInfo->SharedData.TangentCodeChunk)
+		{
+			Result.BasisInfo = *NormalInfo;
+			Result.bFound = true;
+		}
+	}
+
+	return Result;
+}
+
+FSubstrateRegisteredSharedLocalBasis MIRSubstrateCompilationInfoCreateSharedLocalBasis2(Substrate::FSubstrateTranslatorData* SubstrateTranslatorData, FValueRef* NormalArg, FValueRef* TangentArg)
+{
+	// Specific implementation for IR emitter on the common FSubstrateTranslatorData
+	// Note that:
+	//  - We do not have generated code or unique identifier yet at this stage.
+	//  - And we do not have any unique hash for the generated IR tree.
+	//  - So there is no local basis deduplication right now with the new translator based on IR.
+
+	MIR::FEmitter* Em = SubstrateTranslatorData->GetCompilerNew();
+	check(Em); // This function is meant to be used with the new translator only.
+
+	FString SubstrateSharedNormalParamCode;
+	FString SubstrateSharedTangentParamCode;
+	const uint32 FinalSharedLocalBasisIndex = SubstrateTranslatorData->GetFreeSharedLocalBasesIndex();
+
+	FString SubstrateSharedNormalParamName;
+	FString SubstrateSharedTangentParamName;
+
+	SubstrateSharedNormalParamName = FString(TEXT("SubstrateNormal"));
+	SubstrateSharedNormalParamName.AppendInt(FinalSharedLocalBasisIndex);
+	Em->SetCustomOutputs(SubstrateSharedNormalParamName, { NormalArg, 1 }, MIR::MaterialOutputStage_Pixel);
+	SubstrateSharedNormalParamCode = FString::Printf(TEXT("Parameters.%s0"), *SubstrateSharedNormalParamName);
+
+	// Always write tangent even if not used later for now.
+	SubstrateSharedTangentParamName = FString(TEXT("SubstrateTangent"));
+	SubstrateSharedTangentParamName.AppendInt(FinalSharedLocalBasisIndex);
+	FValueRef DefaultTangent = EmitVertexTangent(*Em);
+	Em->SetCustomOutputs(SubstrateSharedTangentParamName, { TangentArg ? TangentArg : &DefaultTangent, 1}, MIR::MaterialOutputStage_Pixel);
+	SubstrateSharedTangentParamCode = FString::Printf(TEXT("Parameters.%s0"), *SubstrateSharedTangentParamName);
+
+	FSubstrateCompilationContext& SubstrateCtx = SubstrateTranslatorData->SubstrateCompilationContext[SubstrateTranslatorData->CurrentSubstrateCompilationContext];
+	check(NormalArg != nullptr);
+	check(SubstrateCtx.NextFreeSubstrateShaderNormalIndex < 255);	// Out of shared local basis slots
+
+	static const FString NullSubstrateIRNormalTangentCode(TEXT("Parameters.TangentToWorld[2]"));	//
+
+	FSubstrateRegisteredSharedLocalBasis SubstrateRegisteredSharedLocalBasis;
+	SubstrateRegisteredSharedLocalBasis.NormalCodeChunk = FinalSharedLocalBasisIndex;
+	SubstrateRegisteredSharedLocalBasis.NormalCodeChunkHash = FinalSharedLocalBasisIndex;
+	SubstrateRegisteredSharedLocalBasis.TangentCodeChunk = FinalSharedLocalBasisIndex;
+	SubstrateRegisteredSharedLocalBasis.TangentCodeChunkHash = FinalSharedLocalBasisIndex;
+	SubstrateRegisteredSharedLocalBasis.GraphSharedLocalBasisIndex = SubstrateCtx.NextFreeSubstrateShaderNormalIndex++;
+
+	// Find a basis which matches both the Normal & the Tangent code chunks
+	TArray<FSubstrateSharedLocalBasesInfo*> NormalInfos;
+	SubstrateCtx.CodeChunkToSubstrateSharedLocalBasis.MultiFindPointer(SubstrateRegisteredSharedLocalBasis.NormalCodeChunkHash, NormalInfos);
+	for (FSubstrateSharedLocalBasesInfo* NormalInfo : NormalInfos)
+	{
+		// * Either we find a perfect match (normal & tangent matches)
+		// * Or we find a normal which doesn't have a tangent associated with, and we set the tangent for code
+		if (SubstrateRegisteredSharedLocalBasis.TangentCodeChunkHash == NormalInfo->SharedData.TangentCodeChunk)
+		{
+			return NormalInfo->SharedData;
+		}
+		else if (NormalInfo->SharedData.TangentCodeChunk == INDEX_NONE)
+		{
+			NormalInfo->SharedData.TangentCodeChunk = FinalSharedLocalBasisIndex;
+			NormalInfo->SharedData.TangentCodeChunkHash = SubstrateRegisteredSharedLocalBasis.TangentCodeChunkHash;
+			NormalInfo->TangentCode = SubstrateSharedTangentParamCode;
+			NormalInfo->TangentArg = TangentArg != nullptr ? *TangentArg : MIR::FValueRef();
+			return NormalInfo->SharedData;
+		}
+	}
+
+	// Allocate a new slot for a new shared local basis
+	SubstrateCtx.CodeChunkToSubstrateSharedLocalBasis.Add(SubstrateRegisteredSharedLocalBasis.NormalCodeChunkHash, { SubstrateRegisteredSharedLocalBasis, *SubstrateSharedNormalParamCode, *SubstrateSharedTangentParamCode, *NormalArg, TangentArg != nullptr ? *TangentArg : MIR::FValueRef()});
+	return SubstrateRegisteredSharedLocalBasis;
+}
+
+static bool IsTangentSpaceNormal(UMaterial* Material)
+{
+	check(Material); // SUBSTRATE_TODO We should never have a null material here, but it seems to happen in some cases. Figure out why and fix if this can result in bad tangent space normal.
+	if (!Material)
+	{
+		return true; // This is here to fall back to using tangent space when debugging and working with the new translator.
+	}
+	return Material->MaterialDomain == MD_DeferredDecal || Material->bTangentSpaceNormal;
+}
+
+static FValueRef TransformNormalFromRequestedBasisToWorld(MIR::FEmitter& Em, UMaterial* Material, FValueRef NormalArg)
+{
+	// When feeding tangent space or world space, we want to have the final normal/tangent normalized when stored in SharedLocalBases.
+	// So that we do not have to normalise it later in any forward processes and avoid overhead in instructions.
+	if (IsTangentSpaceNormal(Material))
+	{
+		// See TransformTangentNormalToWorld definitions in MaterialTemplate.ush
+		return EmitTransformVectorBase(Em, NormalArg, MCB_Tangent, MCB_World, false, nullptr, nullptr);
+	}
+	// Normalize
+	return NormalArg = Em.Multiply(NormalArg, Em.Rsqrt(Em.Dot(NormalArg, NormalArg)));
+}
+
+
+// The compilation of an expression can sometimes lead to a INDEX_NONE code chunk when editing material graphs 
+// or when the node is inside a material function, linked to an input pin of the material function and that input is not plugged in to anything.
+// But for normals or tangents, Substrate absolutely need a valid code chunk to de-duplicate when stored in memory. 
+// Also, we want all our nodes to have default, as that is needed when creating BSDF, when registering code chunk representing material topology.
+
+static FValueRef CompileWithDefaultFloat1(MIR::FEmitter& Em, FExpressionInput& Input, float X, const FScalarMaterialInput* RootNodeInput = nullptr)
+{
+	FValueRef DefaultResultArg = Em.ConstantFloat(X);
+	if (RootNodeInput && RootNodeInput->UseConstant)
+	{
+		DefaultResultArg = Em.ConstantFloat(RootNodeInput->Constant);
+	}
+	FValueRef ResultArg = Input.GetTracedInput().Expression ? Em.TryInput(&Input) : DefaultResultArg;
+	ResultArg = Em.CastToFloatKind(ResultArg);	// Move to non LWC
+	return !ResultArg.IsValid() ? DefaultResultArg : ResultArg;
+}
+static FValueRef CompileWithDefaultFloat2(MIR::FEmitter& Em, FExpressionInput& Input, float X, float Y, const FVector2MaterialInput* RootNodeInput = nullptr)
+{
+	FValueRef DefaultResultArg = Em.ConstantFloat2(FVector2f(X, Y));
+	if (RootNodeInput && RootNodeInput->UseConstant)
+	{
+		DefaultResultArg = Em.ConstantFloat2(FVector2f(RootNodeInput->Constant.X, RootNodeInput->Constant.Y));
+	}
+	FValueRef ResultArg = Input.GetTracedInput().Expression ? Em.TryInput(&Input) : DefaultResultArg;
+	ResultArg = Em.CastToFloatKind(ResultArg);	// Move to non LWC
+	return !ResultArg.IsValid() ? DefaultResultArg : ResultArg;
+}
+static FValueRef CompileWithDefaultFloat3(MIR::FEmitter& Em, FExpressionInput& Input, float X, float Y, float Z, const FColorMaterialInput* RootNodeInput = nullptr)
+{
+	FValueRef DefaultResultArg = Em.ConstantFloat3(FVector3f(X, Y, Z));
+	if (RootNodeInput && RootNodeInput->UseConstant)
+	{
+		DefaultResultArg = Em.ConstantFloat3(FVector3f(RootNodeInput->Constant.R, RootNodeInput->Constant.G, RootNodeInput->Constant.B));
+	}
+	FValueRef ResultArg = Input.GetTracedInput().Expression ? Em.TryInput(&Input) : DefaultResultArg;
+	ResultArg = Em.CastToFloatKind(ResultArg);	// Move to non LWC
+	return !ResultArg.IsValid() ? DefaultResultArg : ResultArg;
+}
+
+
+static FValueRef CompileWithDefaultNormalWS(MIR::FEmitter& Em, UMaterial* Material, FExpressionInput& NormalInput, bool bConvertToRequestedSpace = true)
+{
+	FSubstrateTranslatorData* SubstrateTranslatorData = Em.GetSubstrateTranslatorData();
+
+	if (NormalInput.GetTracedInput().Expression != nullptr)
+	{
+		FValueRef NormalArg = Em.TryInput(&NormalInput);
+		if (!NormalArg)
+		{
+			// Nothing is plug in from the linked input, so specify world space normal the BSDF node expects.
+			NormalArg = EmitVertexNormal(Em);
+		}
+
+		// Ensure the normal has always a valid float3 type
+		NormalArg = Em.Cast(NormalArg, MIR::FType::MakeFloatVector(3));
+		NormalArg = Em.CastToFloatKind(NormalArg);	// Move to non LWC
+		if (!NormalArg.IsValid())
+		{
+			NormalArg = EmitVertexNormal(Em);
+		}
+
+		// Transform into world space normal if needed. BSDF nodes always expects world space normal as input.
+		if (bConvertToRequestedSpace)
+		{
+			NormalArg = TransformNormalFromRequestedBasisToWorld(Em, Material, NormalArg);
+		}
+		return NormalArg;
+	}
+	// Nothing is plug in on the BSDF node, so specify world space normal the node expects.
+	return EmitVertexNormal(Em);
+}
+
+static FValueRef CompileWithDefaultTangentWS(MIR::FEmitter& Em, UMaterial* Material, FExpressionInput& TangentInput, bool bConvertToRequestedSpace = true)
+{
+	FSubstrateTranslatorData* SubstrateTranslatorData = Em.GetSubstrateTranslatorData();
+
+	if (TangentInput.GetTracedInput().Expression != nullptr)
+	{
+		FValueRef NormalArg = Em.TryInput(&TangentInput);
+		if (!NormalArg)
+		{
+			// Nothing is plug in from the linked input, so specify world space normal the BSDF node expects.
+			NormalArg = EmitVertexTangent(Em);
+		}
+
+		// Ensure the normal has always a valid float3 type
+		NormalArg = Em.Cast(NormalArg, MIR::FType::MakeFloatVector(3));
+		NormalArg = Em.CastToFloatKind(NormalArg);	// Move to non LWC
+		if (!NormalArg.IsValid())
+		{
+			NormalArg = EmitVertexTangent(Em);
+		}
+
+		// Transform into world space normal if needed. BSDF nodes always expects world space normal as input.
+		if (bConvertToRequestedSpace)
+		{
+			NormalArg = TransformNormalFromRequestedBasisToWorld(Em, Material, NormalArg);
+		}
+		return NormalArg;
+	}
+	// Nothing is plug in on the BSDF node, so specify world space normal the node expects.
+	return EmitVertexTangent(Em);
+}
+
+void UMaterialExpressionSubstrateSlabBSDF::Build(MIR::FEmitter& Em)
+{
+	// SUBSTRATE_TODO make this similar to UMaterialExpressionSubstrateSlabBSDF::Compile
+	// SUBSTRATE_TODO put SubstrateTranslatorCommon on Em.
+
+	FSubstrateTranslatorData* SubstrateTranslatorData = Em.GetSubstrateTranslatorData();
+
+	FGuid PathUniqueId = SubstrateTranslatorData->SubstrateTreeStackGetPathUniqueId();
+	FSubstrateOperator& SubstrateOperator = SubstrateTranslatorData->SubstrateCompilationGetOperator(PathUniqueId);
+
+#if 0
+	int32 SSSProfileCodeChunk = INDEX_NONE;
+	FName SubsurfaceProfileParameterName;
+	if (SubstrateOperator.Has(ESubstrateBsdfFeature::SSS) && SubsurfaceProfile)
+	{
+		SubsurfaceProfileParameterName = SubsurfaceProfile::CreateSubsurfaceProfileParameterName(SubsurfaceProfile);
+	}
+	else
+	{
+		SubsurfaceProfileParameterName = SubsurfaceProfile::GetSubsurfaceProfileParameterName();
+	}
+	Em.PreshaderParameter(MIR::FType::MakeScalar(MIR::EScalarKind::Float), EPreshader2Opcode::Abs, MIR::FValueRef::, )
+
+	//	FPreshaderParameter
+	//	PRESHADER2_OP(SparseVolumeTextureUniform)
+
+		// FValueRef FEmitter::PreshaderParameter(FType Type, EPreshader2Opcode Opcode, FValueRef SourceParameter, FPreshaderParameterPayload Payload)
+
+		/*
+		const UE::Shader::EValueType UniformType = URuntimeVirtualTexture::GetUniformParameterType(UniformIndex);
+		Uniforms[UniformIndex] = Em.PreshaderParameter(
+			MIR::FType::FromShaderType(UniformType), EPreshader2Opcode::RuntimeVirtualTextureUniform,
+			TextureObjects[0], MIR::FPreshaderParameterPayload{ .UniformIndex = UniformIndex });
+		*/
+	// SUBSTRATE_TODO implement
+	//SubstrateTranslatorData->SubstrateThicknessStackGetThicknessCode()
+	//SubstrateTranslatorData->GetSubstrateSharedLocalBasisIndexMacro(NewRegisteredSharedLocalBasis)
+
+/*
+	int32 SSSProfileCodeChunk = INDEX_NONE;
+	if (SubstrateOperator.Has(ESubstrateBsdfFeature::SSS) && SubsurfaceProfile)
+	{
+		SSSProfileCodeChunk = CreateSubsurfaceProfileParameter(Compiler, SubsurfaceProfile);
+	}
+	else
+	{
+		SSSProfileCodeChunk = CreateDefaultSubsurfaceProfileParameter(Compiler);
+	}
+*/
+#else
+	FValueRef SubsurfaceProfileArg = Em.ConstantFloat(0.0f); // SUBSTRATE_TODO
+#endif
+
+#if 0
+/*
+	int32 SpecularProfileCodeChunk = INDEX_NONE;
+	if (SubstrateOperator.Has(ESubstrateBsdfFeature::SpecularProfile))
+	{
+		const FName SpecularProfileParameterName = SpecularProfile::CreateSpecularProfileParameterName(SpecularProfile);
+		SpecularProfileCodeChunk = Compiler->ForceCast(Compiler->ScalarParameter(SpecularProfileParameterName, 1.0f), MCT_Float1);
+	}
+*/
+#else
+	FValueRef SpecularProfileArg = Em.ConstantFloat(0.0f); // SUBSTRATE_TODO
+#endif
+
+	const float DefaultSpecular = 0.5f;
+	const float DefaultF0 = DielectricSpecularToF0(DefaultSpecular);
+
+	// We also cannot ignore the tangent when using the default Tangent because GetTangentBasis
+	// used in SubstrateGetBSDFSharedBasis cannot be relied on for smooth tangent used for lighting on any mesh.
+	const bool bHasAnisotropy = SubstrateOperator.Has(ESubstrateBsdfFeature::Anisotropy);
+	FValueRef NormalArg = CompileWithDefaultNormalWS(Em, Material, Normal);
+	FValueRef TangentArg = CompileWithDefaultTangentWS(Em, Material, Tangent);
+	const bool bHasValidTangent = bHasAnisotropy && TangentArg.IsValid();
+
+	const FSubstrateRegisteredSharedLocalBasis NewRegisteredSharedLocalBasis = MIRSubstrateCompilationInfoCreateSharedLocalBasis2(SubstrateTranslatorData, &NormalArg, bHasValidTangent ? &TangentArg : nullptr);
+	SubstrateOperator.BSDFRegisteredSharedLocalBasis = NewRegisteredSharedLocalBasis;
+
+	FValueRef ThicknessArg;
+	if (SubstrateOperator.ThicknessIndex != INDEX_NONE)
+	{
+		FSubstrateTranslatorDataInterface::FSubstrateThicknessExpression* ThicknessExpression = SubstrateTranslatorData->SubstrateThicknessStackGetExpression(SubstrateOperator.ThicknessIndex);
+
+		if (ThicknessExpression == nullptr)
+		{
+			Em.Errorf(TEXT(" SubstrateThichkness: %i could not be found)"), SubstrateOperator.ThicknessIndex);
+			ThicknessArg = Em.ConstantFloat(SUBSTRATE_LAYER_DEFAULT_THICKNESS_CM);
+		}
+		else if (ThicknessExpression->ExpressionInput)
+		{
+			FExpressionInput* ExpressionInput = ThicknessExpression->ExpressionInput;
+			ThicknessArg = ExpressionInput->GetTracedInput().Expression ? Em.TryInput(ExpressionInput) : Em.ConstantFloat(SUBSTRATE_LAYER_DEFAULT_THICKNESS_CM);
+			ThicknessArg = Em.CastToFloatKind(ThicknessArg); // remove LWC
+		}
+		else if (ThicknessExpression->MaterialInput)
+		{
+			FScalarMaterialInput* MaterialInput = ThicknessExpression->MaterialInput;
+			ThicknessArg = MaterialInput->UseConstant ? Em.ConstantFloat(MaterialInput->Constant) : Em.ConstantFloat(SUBSTRATE_LAYER_DEFAULT_THICKNESS_CM);
+			ThicknessArg = Em.CastToFloatKind(ThicknessArg); // remove LWC
+		}
+		if (!ThicknessArg.IsValid())
+		{
+			ThicknessArg = Em.ConstantFloat(SUBSTRATE_LAYER_DEFAULT_THICKNESS_CM);
+		}
+	}
+	else
+	{
+		// Thickness is not tracked properly, this can happen when opening a material function in editor
+		ThicknessArg = Em.ConstantFloat(SUBSTRATE_LAYER_DEFAULT_THICKNESS_CM);
+	}
+	check(ThicknessArg.IsValid());
+
+	FValueRef EmissiveColorArg				= CompileWithDefaultFloat3(Em, EmissiveColor, 0.0f, 0.0f, 0.0f);
+	FValueRef DiffuseAlbedoArg				= CompileWithDefaultFloat3(Em, DiffuseAlbedo, 0.18f, 0.18f, 0.18f);
+	FValueRef F0Arg							= CompileWithDefaultFloat3(Em, F0, DefaultF0, DefaultF0, DefaultF0);
+	FValueRef RoughnessArg					= CompileWithDefaultFloat1(Em, Roughness, 0.5f);
+	FValueRef AnisotropyArg					= CompileWithDefaultFloat1(Em, Anisotropy, 0.0f);
+	FValueRef F90Arg						= CompileWithDefaultFloat3(Em, F90, 1.0f, 1.0f, 1.0f);
+	FValueRef SSSMFPArg						= CompileWithDefaultFloat3(Em, SSSMFP, 0.0f, 0.0f, 0.0f);
+	FValueRef SSSMFPScaleArg				= CompileWithDefaultFloat1(Em, SSSMFPScale, 1.0f);
+	FValueRef SSSPhaseAnisotropyArg			= CompileWithDefaultFloat1(Em, SSSPhaseAnisotropy, 0.0f);
+	FValueRef SecondRoughnessArg			= CompileWithDefaultFloat1(Em, SecondRoughness, 0.0f);
+	FValueRef SecondRoughnessWeightArg		= CompileWithDefaultFloat1(Em, SecondRoughnessWeight, 0.0f);
+	FValueRef FuzzAmountArg					= CompileWithDefaultFloat1(Em, FuzzAmount, 0.0f);
+	FValueRef FuzzColorArg					= CompileWithDefaultFloat3(Em, FuzzColor, 0.0f, 0.0f, 0.0f);
+	FValueRef FuzzRoughnessArg				= HasFuzzRoughness() ? CompileWithDefaultFloat1(Em, FuzzRoughness, 0.5f) : RoughnessArg;
+	FValueRef GlintValueArg					= CompileWithDefaultFloat1(Em, GlintValue, 1.0f);
+	FValueRef GlintUVArg					= CompileWithDefaultFloat2(Em, GlintUV, 0.0f, 0.0f);
+
+	// Disable some features if requested by the simplification process
+	if (!SubstrateOperator.Has(ESubstrateBsdfFeature::MFPPluggedIn) || SubSurfaceType == EMaterialSubSurfaceType::MSS_None)
+	{
+		SSSMFPArg = Em.ConstantFloat3(FVector3f(0.0f, 0.0f, 0.0f));
+	}
+	if (!SubstrateOperator.Has(ESubstrateBsdfFeature::EdgeColor))
+	{
+		F90Arg = Em.ConstantFloat3(FVector3f(1.0f, 1.0f, 1.0));
+	}
+	if (!SubstrateOperator.Has(ESubstrateBsdfFeature::Fuzz))
+	{
+		FuzzAmountArg = Em.ConstantFloat(0.0f);
+	}
+	if (!SubstrateOperator.Has(ESubstrateBsdfFeature::SecondRoughnessOrSimpleClearCoat))
+	{
+		SecondRoughnessWeightArg = Em.ConstantFloat(0.0f);
+	}
+	if (!SubstrateOperator.Has(ESubstrateBsdfFeature::Anisotropy))
+	{
+		AnisotropyArg = Em.ConstantFloat(0.0f);
+	}
+	if (!SubstrateOperator.Has(ESubstrateBsdfFeature::Glint))
+	{
+		GlintValueArg = Em.ConstantFloat(1.0f);
+	}
+	if (!SubstrateOperator.Has(ESubstrateBsdfFeature::SpecularProfile))
+	{
+		SpecularProfileArg = Em.ConstantFloat(float(INDEX_NONE));
+	}
+
+	const bool bIsAtTheBottomOfTopology = SubstrateOperator.bIsBottom > 0;
+	const bool bIsThinSurface = Material ? Material->IsThinSurface() : false;
+	float ClearCoatUseSecondNormal = 0.0f; // Constant(ClearCoatBottomNormal != Normal ? 1.0f : 0.0f);SUBSTRATE_TODO
+
+	FValueRef SharedLocalBasisIndexMacroArg = Em.InlineHLSL(MIR::FType::MakeIntScalar(), *GetSubstrateSharedLocalBasisIndexMacroInner(NewRegisteredSharedLocalBasis, SubstrateTranslatorData->CurrentSubstrateCompilationContext));
+
+	FSubstrateOperator* PromoteToOperator = !SubstrateOperator.bUseParameterBlending || (SubstrateOperator.bUseParameterBlending && SubstrateOperator.bRootOfParameterBlendingSubTree) ? &SubstrateOperator : nullptr;
+	if(PromoteToOperator)
+	{
+		MIR::FSubstrateSlabDesc SubstrateSlabDesc = {
+				NormalArg,
+				DiffuseAlbedoArg,
+				F0Arg,
+				F90Arg,
+				RoughnessArg,
+				AnisotropyArg,
+				SubsurfaceProfileArg,
+				SSSMFPArg,
+				SSSMFPScaleArg,
+				SSSPhaseAnisotropyArg,
+				Em.ConstantFloat(float(SubSurfaceType)),
+				EmissiveColorArg,
+				SecondRoughnessArg,
+				SecondRoughnessWeightArg,
+				Em.ConstantFloat(0.0f),												//		SecondRoughnessAsSimpleClearCoatSUBSTRATE_TODO
+				Em.ConstantFloat(ClearCoatUseSecondNormal),							//		*SubstrateGetCastParameterCode(ClearCoatUseSecondNormal, MCT_Float),SUBSTRATE_TODO
+				Em.ConstantFloat3(FVector3f::Zero()),								//		*SubstrateGetCastParameterCode(ClearCoatBottomNormal, MCT_Float3),SUBSTRATE_TODO
+				FuzzAmountArg,
+				FuzzColorArg,
+				FuzzRoughnessArg,
+				GlintValueArg,
+				GlintUVArg,
+				SpecularProfileArg,
+				ThicknessArg,
+				bIsThinSurface ? Em.ConstantTrue() : Em.ConstantFalse(),
+				bIsAtTheBottomOfTopology ? Em.ConstantTrue() : Em.ConstantFalse(),
+				SharedLocalBasisIndexMacroArg };
+
+		FValueRef Slab = Em.SubstrateSlab(SubstrateSlabDesc);
+
+		FValueRef PromototedSlab = Em.SubstratePromoteToOperator(
+			Slab,
+			Em.ConstantInt(PromoteToOperator->Index),
+			Em.ConstantInt(PromoteToOperator->BSDFIndex),
+			Em.ConstantInt(PromoteToOperator->LayerDepth),
+			Em.ConstantInt(PromoteToOperator->bIsBottom ? 1 : 0));
+
+		Em.Output(0, PromototedSlab);
+	}
+	else
+	{
+		MIR::FSubstrateSlabDesc SubstrateSlabDesc = {
+			NormalArg,
+			DiffuseAlbedoArg,
+			F0Arg,
+			F90Arg,
+			RoughnessArg,
+			AnisotropyArg,
+			SubsurfaceProfileArg,
+			SSSMFPArg,
+			SSSMFPScaleArg,
+			SSSPhaseAnisotropyArg,
+			Em.ConstantFloat(float(SubSurfaceType)),
+			EmissiveColorArg,
+			SecondRoughnessArg,
+			SecondRoughnessWeightArg,
+			Em.ConstantFloat(0.0f),												//		SecondRoughnessAsSimpleClearCoat SUBSTRATE_TODO
+			Em.ConstantFloat(ClearCoatUseSecondNormal),							//		*SubstrateGetCastParameterCode(ClearCoatUseSecondNormal, MCT_Float), SUBSTRATE_TODO
+			Em.ConstantFloat3(FVector3f::Zero()),								//		*SubstrateGetCastParameterCode(ClearCoatBottomNormal, MCT_Float3), SUBSTRATE_TODO
+			FuzzAmountArg,
+			FuzzColorArg,
+			FuzzRoughnessArg,
+			GlintValueArg,
+			GlintUVArg,
+			SpecularProfileArg,
+			ThicknessArg,
+			bIsThinSurface ? Em.ConstantTrue() : Em.ConstantFalse(),
+			bIsAtTheBottomOfTopology ? Em.ConstantTrue() : Em.ConstantFalse(),
+			SharedLocalBasisIndexMacroArg};
+		FValueRef Slab = Em.SubstrateSlab(SubstrateSlabDesc);
+
+		Em.Output(0, Slab);
+	}
+}
+
+MIR::FValueRef UMaterialExpressionSubstrateShadingModels::BuildCommon(MIR::FEmitter& Em, UMaterial* Material,
+	FExpressionInput& BaseColor, FExpressionInput& Specular, FExpressionInput& Metallic, FExpressionInput& Roughness, FExpressionInput& EmissiveColor,
+	FExpressionInput& Opacity, FExpressionInput& SubSurfaceColor, FExpressionInput& ClearCoat, FExpressionInput& ClearCoatRoughness,
+	FExpressionInput& ShadingModel, TEnumAsByte<enum EMaterialShadingModel> ShadingModelOverride,
+	FExpressionInput& TransmittanceColor, FExpressionInput& ThinTranslucentSurfaceCoverage,
+	FExpressionInput& WaterScatteringCoefficients, FExpressionInput& WaterAbsorptionCoefficients, FExpressionInput& WaterPhaseG, FExpressionInput& ColorScaleBehindWater,
+	const bool bHasAnisotropy, FExpressionInput& Anisotropy,
+	FExpressionInput& Normal, FExpressionInput& Tangent,
+	FExpressionInput& ClearCoatNormal, FExpressionInput& CustomTangent,
+	const bool bHasSSS, USubsurfaceProfile* SSSProfile,
+	const class UMaterialEditorOnlyData* EditorOnlyData)
+{
+	FSubstrateTranslatorData* SubstrateTranslatorData = Em.GetSubstrateTranslatorData();
+
+	FGuid PathUniqueId = SubstrateTranslatorData->SubstrateTreeStackGetPathUniqueId();
+	FSubstrateOperator& SubstrateOperator = SubstrateTranslatorData->SubstrateCompilationGetOperator(PathUniqueId);
+
+	if (bHasSSS && SSSProfile != nullptr)
+	{
+		check(false);	// SUBSTRATE TODO implemnent
+		//	SSSProfileCodeChunk = CreateSubsurfaceProfileParameter(Compiler, SSSProfile);
+	}
+	//else
+	//{
+	//	SSSProfileCodeChunk = CreateDefaultSubsurfaceProfileParameter(Compiler);SUBSTRATE_TODO
+	//}
+	FValueRef SubsurfaceProfileIdArg = Em.ConstantFloat(0.0f);
+
+
+	// We also cannot ignore the tangent when using the default Tangent because GetTangentBasis
+	// used in SubstrateGetBSDFSharedBasis cannot be relied on for smooth tangent used for lighting on any mesh.
+//	const bool bHasAnisotropy = SubstrateOperator.Has(ESubstrateBsdfFeature::Anisotropy);
+	FValueRef NormalArg = CompileWithDefaultNormalWS(Em, Material, Normal);
+	FValueRef TangentArg = CompileWithDefaultTangentWS(Em, Material, Tangent);
+	const bool bHasValidTangent = bHasAnisotropy && TangentArg.IsValid();
+
+	const FSubstrateRegisteredSharedLocalBasis NewRegisteredSharedLocalBasis = MIRSubstrateCompilationInfoCreateSharedLocalBasis2(SubstrateTranslatorData, &NormalArg, bHasValidTangent ? &TangentArg : nullptr);
+	SubstrateOperator.BSDFRegisteredSharedLocalBasis = NewRegisteredSharedLocalBasis;
+	const FString BasisIndexMacro = GetSubstrateSharedLocalBasisIndexMacroInner(NewRegisteredSharedLocalBasis, SubstrateTranslatorData->CurrentSubstrateCompilationContext);
+
+
+	FValueRef EmissiveColorArg = CompileWithDefaultFloat3(Em, EmissiveColor, 0.0f, 0.0f, 0.0f, EditorOnlyData ? &EditorOnlyData->EmissiveColor : nullptr);
+	FValueRef BaseColorArg = CompileWithDefaultFloat3(Em, BaseColor, 0.18f, 0.18f, 0.18f, EditorOnlyData ? &EditorOnlyData->BaseColor : nullptr);
+	FValueRef SpecularArg = CompileWithDefaultFloat1(Em, Specular, 0.5f, EditorOnlyData ? &EditorOnlyData->Specular : nullptr);
+	FValueRef MetallicArg = CompileWithDefaultFloat1(Em, Metallic, 0.0f, EditorOnlyData ? &EditorOnlyData->Metallic : nullptr);
+	FValueRef RoughnessArg = CompileWithDefaultFloat1(Em, Roughness, 0.5f, EditorOnlyData ? &EditorOnlyData->Roughness : nullptr);
+	FValueRef AnisotropyArg = CompileWithDefaultFloat1(Em, Anisotropy, 0.0f, EditorOnlyData ? &EditorOnlyData->Anisotropy : nullptr);
+	FValueRef SubSurfaceColorArg = CompileWithDefaultFloat3(Em, SubSurfaceColor, 1.0f, 1.0f, 1.0f, EditorOnlyData ? &EditorOnlyData->SubsurfaceColor : nullptr);
+	FValueRef ClearCoatArg = CompileWithDefaultFloat1(Em, ClearCoat, 1.0f, EditorOnlyData ? &EditorOnlyData->ClearCoat : nullptr);
+	FValueRef ClearCoatRoughnessArg = CompileWithDefaultFloat1(Em, ClearCoatRoughness, 0.1f, EditorOnlyData ? &EditorOnlyData->ClearCoatRoughness : nullptr);
+	FValueRef ThinTranslucentTransmittanceColorArg = CompileWithDefaultFloat3(Em, TransmittanceColor, 1.0f, 1.0f, 1.0f);
+	FValueRef WaterScatteringCoefficientsArg = CompileWithDefaultFloat3(Em, WaterScatteringCoefficients, 0.0f, 0.0f, 0.0f);
+	FValueRef WaterAbsorptionCoefficientsArg = CompileWithDefaultFloat3(Em, WaterAbsorptionCoefficients, 0.0f, 0.0f, 0.0f);
+	FValueRef WaterPhaseGArg = CompileWithDefaultFloat1(Em, WaterPhaseG, 0.0f);
+	FValueRef ColorScaleBehindWaterArg = CompileWithDefaultFloat3(Em, ColorScaleBehindWater, 1.0f, 1.0f, 1.0f);
+	FValueRef CustomTangentArg = CompileWithDefaultTangentWS(Em, Material, CustomTangent);
+	FValueRef ThinTranslucentSurfaceCoverageArg = CompileWithDefaultFloat1(Em, ThinTranslucentSurfaceCoverage, 1.0f);
+
+
+	const bool bHasCoatNormal = ClearCoatNormal.IsConnected();
+	// Clear coat normal basis
+	FValueRef ClearCoatNormalArg = NormalArg;
+	FValueRef ClearCoatTangentArg = TangentArg;
+	FString ClearCoat_BasisIndexMacro = BasisIndexMacro;
+	FSubstrateRegisteredSharedLocalBasis ClearCoat_NewRegisteredSharedLocalBasis = NewRegisteredSharedLocalBasis;
+	if (bHasCoatNormal)
+	{
+		ClearCoatNormalArg = CompileWithDefaultNormalWS(Em, Material, ClearCoatNormal);
+		ClearCoat_NewRegisteredSharedLocalBasis = MIRSubstrateCompilationInfoCreateSharedLocalBasis2(SubstrateTranslatorData, &ClearCoatNormalArg, bHasValidTangent ? &ClearCoatTangentArg : nullptr);
+		//SubstrateCompilationInfoCreateSharedLocalBasis(Compiler, ClearCoat_NormalCodeChunk, ClearCoat_TangentCodeChunk);
+		ClearCoat_BasisIndexMacro = GetSubstrateSharedLocalBasisIndexMacroInner(ClearCoat_NewRegisteredSharedLocalBasis, SubstrateTranslatorData->CurrentSubstrateCompilationContext);
+	}
+
+	FValueRef OpacityArg = FValueRef();
+	const bool bSubstrateSkipsOpacityEvaluation = false; // SUBSTRATE_TODO implemente the below code shared with legacy trnaslator
+	/*
+	return !IsTranslucentBlendMode(Material)
+		&& GetSubstrateMaterialExportType() == ESubstrateMaterialExport::SME_None	// We never skip when exporting information
+		&& Material->GetShadingModels().CountShadingModels() == 1
+		&& !Material->GetShadingModels().HasShadingModel(MSM_SingleLayerWater)
+		&& !Material->GetShadingModels().HasShadingModel(MSM_Subsurface)
+		&& !Material->GetShadingModels().HasShadingModel(MSM_SubsurfaceProfile)
+		&& !Material->GetShadingModels().HasShadingModel(MSM_TwoSidedFoliage)
+		&& !Material->GetShadingModels().HasShadingModel(MSM_PreintegratedSkin);
+	*/
+	const float DefaultOpacity = 1.0f;
+	if (!bSubstrateSkipsOpacityEvaluation)
+	{
+		OpacityArg = CompileWithDefaultFloat1(Em, Opacity, DefaultOpacity, EditorOnlyData ? &EditorOnlyData->Opacity : nullptr);
+	}
+	if (!OpacityArg.IsValid())
+	{
+		OpacityArg = Em.ConstantFloat(DefaultOpacity);
+	}
+
+	FValueRef ShadingModelArg = ShadingModel.IsConnected() ? CompileWithDefaultFloat1(Em, ShadingModel, float(MSM_DefaultLit)) : Em.ConstantFloat(float(ShadingModelOverride));
+	int32 ShadingModelCount = Material ? Material->GetShadingModels().CountShadingModels() : 1;
+	bool bHasDynamicShadingModels = ShadingModelCount > 1;
+
+	check(Material);
+	UMaterial* BaseMaterial = Material->GetBaseMaterial();
+	// Material is probably an instanced material, so we check that this is the case before potentially overriding the ShadingModel from the material instance.
+	UMaterialInterface* BaseMaterialInterface = static_cast<UMaterialInterface*>(BaseMaterial);
+	UMaterialInterface* MaterialInterface = Material;
+	if (BaseMaterial && (BaseMaterialInterface != MaterialInterface))
+	{
+		FMaterialShadingModelField BaseMaterialShadingModels = BaseMaterial->GetShadingModels();
+		FMaterialShadingModelField MaterialShadingModels = Material->GetShadingModels();
+
+		// If the potentially instanced material does not have the same shading model as the instance material, this means it would have overridden the shading model.
+		// From the UI, only a single shading model is selectable, so we simply apply the one coming form the material instance.
+		if (MaterialShadingModels.IsValid() && MaterialShadingModels.CountShadingModels() == 1 && MaterialShadingModels != BaseMaterialShadingModels)
+		{
+			bHasDynamicShadingModels = false;	// No need to go dynamic when there is only a single shading model selected.
+			ShadingModelArg = Em.ConstantFloat(float(MaterialShadingModels.GetFirstShadingModel()));
+		}
+	}
+
+	FValueRef BasisIndexMacroArg = Em.InlineHLSL(MIR::FType::MakeIntScalar(), *BasisIndexMacro);
+	FValueRef ClearCoat_BasisIndexMacroArg = Em.InlineHLSL(MIR::FType::MakeIntScalar(), *ClearCoat_BasisIndexMacro);
+
+	FSubstrateOperator* PromoteToOperator = !SubstrateOperator.bUseParameterBlending || (SubstrateOperator.bUseParameterBlending && SubstrateOperator.bRootOfParameterBlendingSubTree) ? &SubstrateOperator : nullptr;
+	if (PromoteToOperator)
+	{
+		if (PromoteToOperator->Index == INDEX_NONE || PromoteToOperator->BSDFIndex == INDEX_NONE)
+		{
+			Em.Errorf(TEXT("Invalid SubstrateSlabBSDF operator and BSDF indices during promotion in Material %s (asset: %s).\r\n"), *Material->GetName(), *Material->GetPathName());
+			return Em.SubstrateDefaultSlab();
+		}
+
+		MIR::FSubstrateShadingModelsDesc SubstrateShadingModelsDesc = {
+			BaseColorArg,
+			SpecularArg,
+			MetallicArg,
+			RoughnessArg,
+			AnisotropyArg,
+			SubSurfaceColorArg,
+			SubsurfaceProfileIdArg,
+			ClearCoatArg,
+			ClearCoatRoughnessArg,
+			EmissiveColorArg,
+			OpacityArg,
+			ThinTranslucentTransmittanceColorArg,
+			ThinTranslucentSurfaceCoverageArg,
+			WaterScatteringCoefficientsArg,
+			WaterAbsorptionCoefficientsArg,
+			WaterPhaseGArg,
+			ColorScaleBehindWaterArg,
+			ShadingModelArg,
+			NormalArg,
+			TangentArg,
+			ClearCoatNormalArg,
+			CustomTangentArg,
+			BasisIndexMacroArg,
+			ClearCoat_BasisIndexMacroArg,
+			bHasDynamicShadingModels };
+
+		FValueRef Slab = Em.SubstrateShadingModels(SubstrateShadingModelsDesc);
+
+		FValueRef PromototedSlab = Em.SubstratePromoteToOperator(
+			Slab,
+			Em.ConstantInt(PromoteToOperator->Index),
+			Em.ConstantInt(PromoteToOperator->BSDFIndex),
+			Em.ConstantInt(PromoteToOperator->LayerDepth),
+			Em.ConstantInt(PromoteToOperator->bIsBottom ? 1 : 0));
+
+		Em.Output(0, PromototedSlab);
+		return PromototedSlab;
+	}
+
+	MIR::FSubstrateShadingModelsDesc SubstrateShadingModelsDesc = {
+		BaseColorArg,
+		SpecularArg,
+		MetallicArg,
+		RoughnessArg,
+		AnisotropyArg,
+		SubSurfaceColorArg,
+		SubsurfaceProfileIdArg,
+		ClearCoatArg,
+		ClearCoatRoughnessArg,
+		EmissiveColorArg,
+		OpacityArg,
+		ThinTranslucentTransmittanceColorArg,
+		ThinTranslucentSurfaceCoverageArg,
+		WaterScatteringCoefficientsArg,
+		WaterAbsorptionCoefficientsArg,
+		WaterPhaseGArg,
+		ColorScaleBehindWaterArg,
+		ShadingModelArg,
+		NormalArg,
+		TangentArg,
+		ClearCoatNormalArg,
+		CustomTangentArg,
+		BasisIndexMacroArg,
+		ClearCoat_BasisIndexMacroArg,
+		bHasDynamicShadingModels };
+
+	FValueRef Slab = Em.SubstrateShadingModels(SubstrateShadingModelsDesc);
+
+	Em.Output(0, Slab);
+	return Slab;
+}
+
+void UMaterialExpressionSubstrateShadingModels::Build(MIR::FEmitter& Em)
+{
+	FSubstrateTranslatorData* SubstrateTranslatorData = Em.GetSubstrateTranslatorData();
+	FGuid PathUniqueId = SubstrateTranslatorData->SubstrateTreeStackGetPathUniqueId();
+	FSubstrateOperator& SubstrateOperator = SubstrateTranslatorData->SubstrateCompilationGetOperator(PathUniqueId);
+
+	const bool bHasAnisotropy = SubstrateOperator.Has(ESubstrateBsdfFeature::Anisotropy);
+	const bool bHasSSS = HasSSS();
+
+	BuildCommon(Em, Material.Get(),
+		BaseColor, Specular, Metallic, Roughness, EmissiveColor,
+		Opacity,  SubSurfaceColor, ClearCoat, ClearCoatRoughness,
+		ShadingModel, ShadingModelOverride,
+		TransmittanceColor, ThinTranslucentSurfaceCoverage,
+		WaterScatteringCoefficients, WaterAbsorptionCoefficients, WaterPhaseG, ColorScaleBehindWater,
+		bHasAnisotropy, Anisotropy,
+		Normal, Tangent,
+		ClearCoatNormal, CustomTangent,
+		bHasSSS, SubsurfaceProfile/*, EditorOnlyData = nullptr*/);
+}
+
+static FValueRef SubstrateBlendNormal(MIR::FEmitter& Em, FValueRef NormalCodeChunk0, FValueRef NormalCodeChunk1, FValueRef MixCodeChunk)
+{
+	FValueRef SafeMixCodeChunk = Em.Saturate(MixCodeChunk);
+	FValueRef LerpedNormal = Em.Lerp(NormalCodeChunk0, NormalCodeChunk1, SafeMixCodeChunk);
+	FValueRef BlendedNormalCodeChunk = Em.Divide(LerpedNormal, Em.Sqrt(Em.Dot(LerpedNormal, LerpedNormal)));
+	return BlendedNormalCodeChunk;
+}
+
+void UMaterialExpressionSubstrateHorizontalMixing::Build(MIR::FEmitter& Em)
+{
+	FSubstrateTranslatorData* SubstrateTranslatorData = Em.GetSubstrateTranslatorData();
+
+	if (!Foreground.GetTracedInput().Expression)
+	{
+		return Em.Errorf(TEXT("Missing Foreground input"));
+	}
+	if (!Background.GetTracedInput().Expression)
+	{
+		return Em.Errorf(TEXT("Missing Background input"));
+	}
+
+	// Note: we no longer do SubstrateTreeStackPush/Pop since the path guid is regenerated for every substrate node from BuildTopMaterialExpression.
+	FValueRef BackgroundArg = Em.TryInput(&Background);
+	FValueRef ForegroundArg = Em.TryInput(&Foreground);
+
+	FValueRef HorizontalMixArg = CompileWithDefaultFloat1(Em, Mix, 0.5f);
+
+	int32 OutputCodeChunk = INDEX_NONE;
+	FSubstrateOperator& SubstrateOperator = SubstrateTranslatorData->SubstrateCompilationGetOperator(SubstrateTranslatorData->SubstrateTreeStackGetPathUniqueId());
+	if (SubstrateOperator.bUseParameterBlending)
+	{
+		if (!ForegroundArg.IsValid())
+		{
+			return Em.Errorf(TEXT("Foreground input graphs could not be evaluated for parameter blending."));
+		}
+		if (!BackgroundArg.IsValid())
+		{
+			return Em.Errorf(TEXT("Background input graphs could not be evaluated for parameter blending."));
+		}
+
+		FValueRef NormalMixArg = Em.InlineHLSL(MIR::FType::MakeFloatScalar(), TEXTVIEW("HorizontalMixingParameterBlendingBSDFCoverageToNormalMix($0, $1, $2)"),
+			{
+				BackgroundArg,
+				ForegroundArg,
+				HorizontalMixArg
+			}
+		);
+
+
+		FSubstrateOperator* BackgroundBSDFOperator = SubstrateTranslatorData->SubstrateCompilationGetOperatorFromIndex(SubstrateOperator.LeftIndex);
+		FSubstrateOperator* ForegroundBSDFOperator = SubstrateTranslatorData->SubstrateCompilationGetOperatorFromIndex(SubstrateOperator.RightIndex);
+		if (!BackgroundBSDFOperator || !ForegroundBSDFOperator)
+		{
+			return Em.Errorf(TEXT("Missing input on horizontal blending node."));
+		}
+
+		FSubstrateRegisteredSharedLocalBasisSearchResult ForegroundBasis = SubstrateCompilationFindExistingSharedLocalBasis(SubstrateTranslatorData, ForegroundBSDFOperator->BSDFRegisteredSharedLocalBasis.NormalCodeChunkHash, ForegroundBSDFOperator->BSDFRegisteredSharedLocalBasis.TangentCodeChunkHash);
+		FSubstrateRegisteredSharedLocalBasisSearchResult BackgroundBasis = SubstrateCompilationFindExistingSharedLocalBasis(SubstrateTranslatorData, BackgroundBSDFOperator->BSDFRegisteredSharedLocalBasis.NormalCodeChunkHash, BackgroundBSDFOperator->BSDFRegisteredSharedLocalBasis.TangentCodeChunkHash);
+
+		if (!ForegroundBasis.bFound || !BackgroundBasis.bFound)
+		{
+			Em.Errorf(TEXT("Foreground & Background shared local basis could not be found."));
+			return;
+		}
+
+		// Compute the new Normal and Tangent resulting from the blending using code chunk
+		FValueRef NewNormalArg = SubstrateBlendNormal(Em, BackgroundBasis.BasisInfo.NormalArg, ForegroundBasis.BasisInfo.NormalArg, NormalMixArg);
+
+		// The tangent is optional so we treat it differently if INDEX_NONE is specified
+		FValueRef NewTangentArg;
+		if (ForegroundBasis.BasisInfo.TangentArg != nullptr && BackgroundBasis.BasisInfo.TangentArg != nullptr)
+		{
+			NewTangentArg = SubstrateBlendNormal(Em, BackgroundBasis.BasisInfo.TangentArg, ForegroundBasis.BasisInfo.TangentArg, NormalMixArg);
+		}
+		else if (ForegroundBasis.BasisInfo.TangentArg != nullptr)
+		{
+			NewTangentArg = ForegroundBasis.BasisInfo.TangentArg;
+		}
+		else if (BackgroundBasis.BasisInfo.TangentArg != nullptr)
+		{
+			NewTangentArg = BackgroundBasis.BasisInfo.TangentArg;
+		}
+		const FSubstrateRegisteredSharedLocalBasis NewRegisteredSharedLocalBasis = MIRSubstrateCompilationInfoCreateSharedLocalBasis2(SubstrateTranslatorData, &NewNormalArg, NewTangentArg.IsValid() ? &NewTangentArg : nullptr);
+
+
+		FString SharedLocalBasisName = GetSubstrateSharedLocalBasisIndexMacroInner(NewRegisteredSharedLocalBasis, SubstrateTranslatorData->CurrentSubstrateCompilationContext);
+		FValueRef SharedLocalBasisIndexMacroArg = Em.InlineHLSL(MIR::FType::MakeIntScalar(), *SharedLocalBasisName);
+
+		if (!ForegroundArg.IsValid() || !BackgroundArg.IsValid() || !HorizontalMixArg.IsValid() || !NormalMixArg.IsValid())
+		{
+			Em.Errorf(TEXT("Invalid Horizontal operator Foreground, Background, Mix or NormalMix input."));
+		}
+
+		FValueRef CameraVectorArg = EmitInlineHLSL(Em, NAME_CameraVector);
+		FValueRef BackgroundNoV = Em.Saturate(Em.Dot(BackgroundBasis.BasisInfo.NormalArg, CameraVectorArg));
+		FValueRef ForegroundNoV = Em.Saturate(Em.Dot(ForegroundBasis.BasisInfo.NormalArg, CameraVectorArg));
+
+		const FSubstrateOperator* PromoteToOperator = SubstrateOperator.bRootOfParameterBlendingSubTree ? &SubstrateOperator : nullptr;
+		if (PromoteToOperator)
+		{
+			check(PromoteToOperator->Index != INDEX_NONE);
+			check(PromoteToOperator->BSDFIndex != INDEX_NONE);
+
+			FValueRef MixedSlab = Em.SubstrateHorizontalMixingParameterBlending(
+				BackgroundArg,
+				ForegroundArg,
+				HorizontalMixArg,
+				NormalMixArg,
+				SharedLocalBasisIndexMacroArg,
+				BackgroundNoV,
+				ForegroundNoV);
+
+			FValueRef PromotedSlab = Em.SubstratePromoteToOperator(
+				MixedSlab,
+				Em.ConstantInt(PromoteToOperator->Index),
+				Em.ConstantInt(PromoteToOperator->BSDFIndex),
+				Em.ConstantInt(PromoteToOperator->LayerDepth),
+				Em.ConstantInt(PromoteToOperator->bIsBottom ? 1 : 0));
+
+			Em.Output(0, PromotedSlab);
+		}
+		else
+		{
+			FValueRef MixedSlab = Em.InlineHLSL(MIR::FType::MakeSubstrateData(), TEXTVIEW("SubstrateHorizontalMixingParameterBlending($0, $1, $2, $3, $4, Parameters.SharedLocalBases.Types, $5, $6)"),
+				{
+					BackgroundArg,
+					ForegroundArg,
+					HorizontalMixArg,
+					NormalMixArg,
+					SharedLocalBasisIndexMacroArg,
+					BackgroundNoV,
+					ForegroundNoV
+				}
+			);
+			Em.Output(0, MixedSlab);
+		}
+
+		// Propagate the parameter blended normal
+		SubstrateOperator.BSDFRegisteredSharedLocalBasis = NewRegisteredSharedLocalBasis;
+	}
+	else
+	{
+		if (!ForegroundArg.IsValid() || !BackgroundArg.IsValid() || !HorizontalMixArg.IsValid())
+		{
+			Em.Errorf(TEXT("Invalid Horizontal operator Foreground, Background or Mix input."));
+		}
+
+		FValueRef MixedSlabs = Em.SubstrateHorizontalMixing(
+				BackgroundArg,
+				ForegroundArg,
+				HorizontalMixArg,
+				Em.ConstantInt(SubstrateOperator.Index),
+				Em.ConstantInt(SubstrateOperator.MaxDistanceFromLeaves));
+
+		Em.Output(0, MixedSlabs);
+	}
+}
+
+
+void UMaterialExpressionSubstrateVerticalLayering::Build(MIR::FEmitter& Em)
+{
+	FSubstrateTranslatorData* SubstrateTranslatorData = Em.GetSubstrateTranslatorData();
+
+
+	if (!Top.GetTracedInput().Expression)
+	{
+		return Em.Errorf(TEXT("Missing Foreground input"));
+	}
+	if (!Base.GetTracedInput().Expression)
+	{
+		return Em.Errorf(TEXT("Missing Background input"));
+	}
+
+	// Note: we no longer do SubstrateTreeStackPush/Pop since the path guid is regenerated for every substrate node from BuildTopMaterialExpression.
+	FValueRef TopArg = Em.TryInput(&Top);
+	FValueRef BaseArg = Em.TryInput(&Base);
+
+	FValueRef ThicknessArg = CompileWithDefaultFloat1(Em, Thickness, SUBSTRATE_LAYER_DEFAULT_THICKNESS_CM);
+
+	int32 OutputCodeChunk = INDEX_NONE;
+	FSubstrateOperator& SubstrateOperator = SubstrateTranslatorData->SubstrateCompilationGetOperator(SubstrateTranslatorData->SubstrateTreeStackGetPathUniqueId());
+
+	if (SubstrateOperator.bUseParameterBlending)
+	{
+		if (!TopArg.IsValid())
+		{
+			return Em.Errorf(TEXT("Foreground input graphs could not be evaluated for parameter blending."));
+		}
+		if (!BaseArg.IsValid())
+		{
+			return Em.Errorf(TEXT("Background input graphs could not be evaluated for parameter blending."));
+		}
+
+		FValueRef NormalMixArg = Em.InlineHLSL(MIR::FType::MakeFloatScalar(), TEXTVIEW("VerticalLayeringParameterBlendingBSDFCoverageToNormalMix($0)"), { TopArg } );
+
+
+		FSubstrateOperator* TopBSDFOperator = SubstrateTranslatorData->SubstrateCompilationGetOperatorFromIndex(SubstrateOperator.LeftIndex);
+		FSubstrateOperator* BaseBSDFOperator = SubstrateTranslatorData->SubstrateCompilationGetOperatorFromIndex(SubstrateOperator.RightIndex);
+		if (!TopBSDFOperator || !BaseBSDFOperator)
+		{
+			return Em.Errorf(TEXT("Missing input on horizontal blending node."));
+		}
+
+		FSubstrateRegisteredSharedLocalBasisSearchResult BaseBasis = SubstrateCompilationFindExistingSharedLocalBasis(SubstrateTranslatorData, BaseBSDFOperator->BSDFRegisteredSharedLocalBasis.NormalCodeChunkHash, BaseBSDFOperator->BSDFRegisteredSharedLocalBasis.TangentCodeChunkHash);
+		FSubstrateRegisteredSharedLocalBasisSearchResult TopBasis  = SubstrateCompilationFindExistingSharedLocalBasis(SubstrateTranslatorData, TopBSDFOperator->BSDFRegisteredSharedLocalBasis.NormalCodeChunkHash,  TopBSDFOperator->BSDFRegisteredSharedLocalBasis.TangentCodeChunkHash);
+
+		if (!BaseBasis.bFound || !TopBasis.bFound)
+		{
+			Em.Errorf(TEXT("Top & Base shared local basis could not be found."));
+			return;
+		}
+
+		// Compute the new Normal and Tangent resulting from the blending using code chunk
+		FValueRef NewNormalArg = SubstrateBlendNormal(Em, BaseBasis.BasisInfo.NormalArg, TopBasis.BasisInfo.NormalArg, NormalMixArg);
+
+		// The tangent is optional so we treat it differently if INDEX_NONE is specified
+		FValueRef NewTangentArg;
+		if (TopBasis.BasisInfo.TangentArg != nullptr && BaseBasis.BasisInfo.TangentArg != nullptr)
+		{
+			NewTangentArg = SubstrateBlendNormal(Em, BaseBasis.BasisInfo.TangentArg, TopBasis.BasisInfo.TangentArg, NormalMixArg);
+		}
+		else if (TopBasis.BasisInfo.TangentArg != nullptr)
+		{
+			NewTangentArg = TopBasis.BasisInfo.TangentArg;
+		}
+		else if (BaseBasis.BasisInfo.TangentArg != nullptr)
+		{
+			NewTangentArg = BaseBasis.BasisInfo.TangentArg;
+		}
+		const FSubstrateRegisteredSharedLocalBasis NewRegisteredSharedLocalBasis = MIRSubstrateCompilationInfoCreateSharedLocalBasis2(SubstrateTranslatorData, &NewNormalArg, NewTangentArg.IsValid() ? &NewTangentArg : nullptr);
+
+		FString SharedLocalBasisName = GetSubstrateSharedLocalBasisIndexMacroInner(NewRegisteredSharedLocalBasis, SubstrateTranslatorData->CurrentSubstrateCompilationContext);
+		FValueRef SharedLocalBasisIndexMacroArg = Em.InlineHLSL(MIR::FType::MakeIntScalar(), *SharedLocalBasisName);
+
+		if (!TopArg.IsValid() || !BaseArg.IsValid() || !ThicknessArg.IsValid() || !NormalMixArg.IsValid())
+		{
+			Em.Errorf(TEXT("Invalid Vertical operator Base, Top, Thickness or NormalMix input."));
+		}
+
+		FValueRef CameraVectorArg = EmitInlineHLSL(Em, NAME_CameraVector);
+		FValueRef TopNoV = Em.Saturate(Em.Dot(TopBasis.BasisInfo.NormalArg, CameraVectorArg));
+		FValueRef BaseNoV = Em.Saturate(Em.Dot(BaseBasis.BasisInfo.NormalArg, CameraVectorArg));
+
+		const FSubstrateOperator* PromoteToOperator = SubstrateOperator.bRootOfParameterBlendingSubTree ? &SubstrateOperator : nullptr;
+		if (PromoteToOperator)
+		{
+			check(PromoteToOperator->Index != INDEX_NONE);
+			check(PromoteToOperator->BSDFIndex != INDEX_NONE);
+
+			FValueRef LayeredSlab = Em.SubstrateVerticalLayeringParameterBlending(
+				TopArg,
+				BaseArg,
+				SharedLocalBasisIndexMacroArg,
+				TopNoV,
+				BaseNoV);
+
+			FValueRef PromotedSlab = Em.SubstratePromoteToOperator(
+				LayeredSlab,
+				Em.ConstantInt(PromoteToOperator->Index),
+				Em.ConstantInt(PromoteToOperator->BSDFIndex),
+				Em.ConstantInt(PromoteToOperator->LayerDepth),
+				Em.ConstantInt(PromoteToOperator->bIsBottom ? 1 : 0));
+
+			Em.Output(0, PromotedSlab);
+		}
+		else
+		{
+			FValueRef LayeredSlab = Em.SubstrateVerticalLayeringParameterBlending(
+				TopArg,
+				BaseArg,
+				SharedLocalBasisIndexMacroArg,
+				TopNoV,
+				BaseNoV);
+
+			Em.Output(0, LayeredSlab);
+		}
+
+		// Propagate the parameter blended normal
+		SubstrateOperator.BSDFRegisteredSharedLocalBasis = NewRegisteredSharedLocalBasis;
+	}
+	else
+	{
+		if (!TopArg.IsValid() || !BaseArg.IsValid())
+		{
+			Em.Errorf(TEXT("Invalid Vertical operator Top, Base or Thickness input."));
+		}
+
+		FValueRef LayeredSlabs = Em.SubstrateVerticalLayering(
+				TopArg,
+				BaseArg,
+				Em.ConstantInt(SubstrateOperator.Index),
+				Em.ConstantInt(SubstrateOperator.MaxDistanceFromLeaves));
+
+		Em.Output(0, LayeredSlabs);
+	}
+}
+
+
+void UMaterialExpressionSubstrateTransmittanceToMFP::Build(MIR::FEmitter& Em)
+{
+	FSubstrateTranslatorData* SubstrateTranslatorData = Em.GetSubstrateTranslatorData();
+
+	FValueRef TransmittanceColorArg	= TransmittanceColor.GetTracedInput().Expression	? Em.TryInput(&TransmittanceColor)	: Em.ConstantFloat(0.5f);
+	FValueRef ThicknessArg			= Thickness.GetTracedInput().Expression				? Em.TryInput(&Thickness)			: Em.ConstantFloat(SUBSTRATE_LAYER_DEFAULT_THICKNESS_CM);
+
+	if (!TransmittanceColorArg.IsValid())
+	{
+		Em.Errorf(TEXT("TransmittanceColor input graph could not be evaluated for TransmittanceToMFP."));
+		return;
+	}
+	if (!ThicknessArg.IsValid())
+	{
+		Em.Errorf(TEXT("Thickness input graph could not be evaluated for TransmittanceToMFP."));
+		return;
+	}
+
+	MIR::FValueRef MFPArg = Em.InlineHLSL(
+		MIR::FType::MakeFloatVector(3), 
+		TEXTVIEW("(TransmittanceToMeanFreePath($0, $1 * CENTIMETER_TO_METER) * METER_TO_CENTIMETER)"),
+		{ TransmittanceColorArg, ThicknessArg} );
+
+	Em.Output(0, MFPArg);
+	Em.Output(1, ThicknessArg);
+}
+
 #endif // WITH_EDITOR
diff --git a/Engine/Source/Runtime/Engine/Private/Materials/MaterialIR.cpp b/Engine/Source/Runtime/Engine/Private/Materials/MaterialIR.cpp
--- a/Engine/Source/Runtime/Engine/Private/Materials/MaterialIR.cpp
+++ b/Engine/Source/Runtime/Engine/Private/Materials/MaterialIR.cpp
@@ -47,7 +47,12 @@ const TCHAR* LexToString(EValueKind Kind)
 		case VK_TextureRead: return TEXT("TextureRead");
 		case VK_VTPageTableRead: return TEXT("VTPageTableRead");
 		case VK_InlineHLSL: return TEXT("InlineHLSL");
-		case VK_PromoteSubstrateParameter: return TEXT("PromoteSubstrateParameter");
+		case VK_SubstrateDefaultSlab: return TEXT("SubstrateDefaultSlab");
+		case VK_SubstrateSlab: return TEXT("SubstrateSlab");
+		case VK_SubstrateShadingModels: return TEXT("SubstrateShadingModel");
+		case VK_SubstrateHorizontalMixing: return TEXT("SubstrateHorizontalMixing");
+		case VK_SubstrateVerticalLayering: return TEXT("SubstrateVerticalLayering");
+		case VK_SubstratePromoteToOperator: return TEXT("SubstratePromoteToOperator");
 		case VK_StageSwitch: return TEXT("StageSwitch");
 		case VK_HardwarePartialDerivative: return TEXT("HardwarePartialDerivative");
 		case VK_Nop: return TEXT("Nop");
@@ -114,7 +119,12 @@ uint32 FValue::GetSizeInBytes() const
 		case VK_TextureRead: return sizeof(FTextureRead);
 		case VK_VTPageTableRead: return sizeof(FVTPageTableRead);
 		case VK_InlineHLSL: return sizeof(FInlineHLSL);
-		case VK_PromoteSubstrateParameter: return sizeof(FPromoteSubstrateParameter);
+		case VK_SubstrateDefaultSlab: return sizeof(FSubstrateDefaultSlab);
+		case VK_SubstrateSlab: return sizeof(FSubstrateSlab);
+		case VK_SubstrateShadingModels: return sizeof(FSubstrateShadingModels);
+		case VK_SubstrateHorizontalMixing: return sizeof(FSubstrateHorizontalMixing);
+		case VK_SubstrateVerticalLayering: return sizeof(FSubstrateVerticalLayering);
+		case VK_SubstratePromoteToOperator: return sizeof(FSubstratePromoteToOperator);
 		case VK_StageSwitch: return sizeof(FStageSwitch);
 		case VK_HardwarePartialDerivative: return sizeof(FHardwarePartialDerivative);
 		case VK_Nop: return sizeof(FNop);
@@ -193,10 +203,36 @@ TConstArrayView<FValue*> FValue::GetUses() const
 			return { This->Arguments, This->NumArguments };
 		}
 
-		case VK_PromoteSubstrateParameter:
+		case VK_SubstrateDefaultSlab:
 		{
-			auto This = static_cast<const FPromoteSubstrateParameter*>(this);
-			return { This->WorldSpaceTangentsAndNormals, FPromoteSubstrateParameter::NumStaticUses };
+			auto This = static_cast<const FSubstrateDefaultSlab*>(this);
+			return { &This->Dummy, FSubstrateDefaultSlab::NumStaticUses };
+		}
+
+		case VK_SubstrateSlab:
+		{
+			auto This = static_cast<const FSubstrateSlab*>(this);
+			return { &This->Normal, FSubstrateSlab::NumStaticUses };
+		}
+		case VK_SubstrateShadingModels:
+		{
+			auto This = static_cast<const FSubstrateShadingModels*>(this);
+			return { &This->BaseColor, FSubstrateShadingModels::NumStaticUses };
+		}
+		case VK_SubstrateHorizontalMixing:
+		{
+			auto This = static_cast<const FSubstrateHorizontalMixing*>(this);
+			return { &This->Background, FSubstrateHorizontalMixing::NumStaticUses };
+		}
+		case VK_SubstrateVerticalLayering:
+		{
+			auto This = static_cast<const FSubstrateVerticalLayering*>(this);
+			return { &This->Top, FSubstrateVerticalLayering::NumStaticUses };
+		}
+		case VK_SubstratePromoteToOperator:
+		{
+			auto This = static_cast<const FSubstratePromoteToOperator*>(this);
+			return { &This->SubstrateDataInput, FSubstratePromoteToOperator::NumStaticUses };
 		}
 
 		case VK_StageSwitch:
diff --git a/Engine/Source/Runtime/Engine/Private/Materials/MaterialIREmitter.cpp b/Engine/Source/Runtime/Engine/Private/Materials/MaterialIREmitter.cpp
--- a/Engine/Source/Runtime/Engine/Private/Materials/MaterialIREmitter.cpp
+++ b/Engine/Source/Runtime/Engine/Private/Materials/MaterialIREmitter.cpp
@@ -17,6 +17,7 @@
 #include "Engine/Texture2D.h"
 #include "TextureResource.h"
 #include "RenderUtils.h"
+#include "Materials/SubstrateTranslatorCommon.h"
 
 #if WITH_EDITOR
 
@@ -1589,7 +1590,7 @@ static FValueRef EmitAttributeDefaultValue(FEmitter& Emitter, const UMaterial* M
 		case EMaterialAggregateAttributeType::Float4: return Emitter.ConstantFloat4(UE::Math::TVector4<TFloat>{ Attribute.DefaultValue.X, Attribute.DefaultValue.Y, Attribute.DefaultValue.Z, Attribute.DefaultValue.W });
 		case EMaterialAggregateAttributeType::ShadingModel: return Emitter.ConstantInt((TInteger)Material->GetShadingModels().GetFirstShadingModel());
 		case EMaterialAggregateAttributeType::MaterialAttributes: return Emitter.Aggregate(UMaterialAggregate::GetMaterialAttributes());
-		case EMaterialAggregateAttributeType::SubstrateData: return Emitter.PromoteSubstrateParameter();	// is that really a default value ???
+		case EMaterialAggregateAttributeType::SubstrateData: return Emitter.SubstrateDefaultSlab();
 		case EMaterialAggregateAttributeType::Aggregate: return Emitter.Aggregate(Attribute.Aggregate);
 		default: UE_MIR_UNREACHABLE();
 	}
@@ -4088,9 +4089,185 @@ FValueRef FEmitter::InlineHLSL(const FMaterialExternalCodeDeclaration* InExterna
 	return EmitInlineHLSL(*this, ReturnType, InExternalCodeDeclaration, {}, InputValues, ValueFlags, UsedGraphProperties);
 }
 
-FValueRef FEmitter::PromoteSubstrateParameter()
+FValueRef FEmitter::SubstrateDefaultSlab()
 {
-	FPromoteSubstrateParameter Prototype = MakePrototype<FPromoteSubstrateParameter>(FType::MakeSubstrateData());
+	FSubstrateDefaultSlab Prototype = MakePrototype<FSubstrateDefaultSlab>(FType::MakeSubstrateData());
+	Prototype.Dummy = nullptr;
+	return EmitPrototype(*this, Prototype);
+}
+
+FValueRef FEmitter::SubstratePromoteToOperator(FValueRef SubstrateData, FValueRef OperatorIndex, FValueRef BSDFIndex, FValueRef LayerDepth, FValueRef bIsBottom)
+{
+	FSubstratePromoteToOperator Prototype = MakePrototype<FSubstratePromoteToOperator>(FType::MakeSubstrateData());
+	Prototype.SubstrateDataInput = SubstrateData;
+	Prototype.OperatorIndex = OperatorIndex;
+	Prototype.BSDFIndex = BSDFIndex;
+	Prototype.LayerDepth = LayerDepth;
+	Prototype.bIsBottom = bIsBottom;
+	Prototype.bFullySimplifiedInstruction = SubstrateTranslatorData->CurrentSubstrateCompilationContext == Substrate::ESubstrateCompilationContext::SCC_FullySimplified;
+	return EmitPrototype(*this, Prototype);
+}
+
+FValueRef FEmitter::SubstrateShadingModels(const FSubstrateShadingModelsDesc Desc)
+{
+	FSubstrateShadingModels Prototype = MakePrototype<FSubstrateShadingModels>(FType::MakeSubstrateData());
+
+	Prototype.BaseColor = Desc.BaseColor;
+	Prototype.Specular = Desc.Specular;
+	Prototype.Metallic = Desc.Metallic;
+	Prototype.Roughness = Desc.Roughness;
+	Prototype.Anisotropy = Desc.Anisotropy;
+	Prototype.SubSurfaceColor = Desc.SubSurfaceColor;
+	Prototype.SubSurfaceProfileId = Desc.SubSurfaceProfileId;
+	Prototype.ClearCoat = Desc.ClearCoat;
+	Prototype.ClearCoatRoughness = Desc.ClearCoatRoughness;
+	Prototype.EmissiveColor = Desc.EmissiveColor;
+	Prototype.Opacity = Desc.Opacity;
+	Prototype.ThinTranslucentTransmittanceColor = Desc.ThinTranslucentTransmittanceColor;
+	Prototype.ThinTranslucentSurfaceCoverage = Desc.ThinTranslucentSurfaceCoverage;
+	Prototype.WaterScatteringCoefficients = Desc.WaterScatteringCoefficients;
+	Prototype.WaterAbsorptionCoefficients = Desc.WaterAbsorptionCoefficients;
+	Prototype.WaterPhaseG = Desc.WaterPhaseG;
+	Prototype.ColorScaleBehindWater = Desc.ColorScaleBehindWater;
+	Prototype.ShadingModel = Desc.ShadingModel;
+	Prototype.Normal = Desc.Normal;
+	Prototype.Tangent = Desc.Tangent;
+	Prototype.ClearCoatNormal = Desc.ClearCoatNormal;
+	Prototype.CustomTangent = Desc.CustomTangent;
+	Prototype.BasisIndexMacro = Desc.BasisIndexMacro;
+	Prototype.ClearCoat_BasisIndexMacro = Desc.ClearCoat_BasisIndexMacro;
+
+	Prototype.bHasDynamicShadingModels = Desc.bHasDynamicShadingModels;
+	Prototype.bFullySimplifiedInstruction = SubstrateTranslatorData->CurrentSubstrateCompilationContext == Substrate::ESubstrateCompilationContext::SCC_FullySimplified;
+
+	return EmitPrototype(*this, Prototype);
+}
+
+FValueRef FEmitter::SubstrateSlab(const FSubstrateSlabDesc& Desc)
+{
+	FSubstrateSlab Prototype = MakePrototype<FSubstrateSlab>(FType::MakeSubstrateData());
+	
+	Prototype.Normal = Desc.Normal;
+	Prototype.DiffuseAlbedo = Desc.DiffuseAlbedo;
+	Prototype.F0 = Desc.F0;
+	Prototype.F90 = Desc.F90;
+	Prototype.Roughness = Desc.Roughness;
+	Prototype.Anisotropy = Desc.Anisotropy;
+	Prototype.SSSProfileId = Desc.SSSProfileId;
+	Prototype.SSSMFP = Desc.SSSMFP;
+	Prototype.SSSMFPScale = Desc.SSSMFPScale;
+	Prototype.SSSPhaseAniso = Desc.SSSPhaseAniso;
+	Prototype.SSSType = Desc.SSSType;
+	Prototype.EmissiveColor = Desc.EmissiveColor;
+	Prototype.SecondRoughness = Desc.SecondRoughness;
+	Prototype.SecondRoughnessWeight = Desc.SecondRoughnessWeight;
+	Prototype.SecondRoughnessAsSimpleClearCoat = Desc.SecondRoughnessAsSimpleClearCoat;
+	Prototype.ClearCoatUseSecondNormal = Desc.ClearCoatUseSecondNormal;
+	Prototype.ClearCoatBottomNormal = Desc.ClearCoatBottomNormal;
+	Prototype.FuzzAmount = Desc.FuzzAmount;
+	Prototype.FuzzColor = Desc.FuzzColor;
+	Prototype.FuzzRoughness = Desc.FuzzRoughness;
+	Prototype.GlintValue = Desc.GlintValue;
+	Prototype.GlintUV = Desc.GlintUV;
+	Prototype.SpecularProfileId = Desc.SpecularProfileId;
+	Prototype.Thickness = Desc.Thickness;
+	Prototype.IsThin = Desc.IsThin;
+	Prototype.IsAtBottom = Desc.IsAtBottom;
+	Prototype.LocalBasisIndex = Desc.LocalBasisIndex;
+
+	Prototype.bFullySimplifiedInstruction = SubstrateTranslatorData->CurrentSubstrateCompilationContext == Substrate::ESubstrateCompilationContext::SCC_FullySimplified;
+	return EmitPrototype(*this, Prototype);
+}
+
+FValueRef FEmitter::SubstrateHorizontalMixing(FValueRef Background, FValueRef Foreground, FValueRef Mix, FValueRef OperatorIndex, FValueRef MaxDistanceFromLeaves)
+{
+	FSubstrateHorizontalMixing Prototype = MakePrototype<FSubstrateHorizontalMixing>(FType::MakeSubstrateData());
+	Prototype.Background = Background;
+	Prototype.Foreground = Foreground;
+	Prototype.Mix = Mix;
+
+	// ParameterBlending = 0
+	Prototype.OperatorIndex = OperatorIndex;
+	Prototype.MaxDistanceFromLeaves = MaxDistanceFromLeaves;
+
+	// ParameterBlending = 1
+	Prototype.NormalMix = nullptr;
+	Prototype.SharedLocalBasisIndexMacro = nullptr;
+	Prototype.BackgroundNoV = nullptr;
+	Prototype.ForegroundNoV = nullptr;
+
+	// Settings
+	Prototype.bParameterBlendingEnabled = false;
+	Prototype.bFullySimplifiedInstruction = SubstrateTranslatorData->CurrentSubstrateCompilationContext == Substrate::ESubstrateCompilationContext::SCC_FullySimplified;
+
+	return EmitPrototype(*this, Prototype);
+}
+
+FValueRef FEmitter::SubstrateHorizontalMixingParameterBlending(FValueRef Background, FValueRef Foreground, FValueRef Mix, FValueRef NormalMix, FValueRef SharedLocalBasisIndexMacro, FValueRef BackgroundNoV, FValueRef ForegroundNoV)
+{
+	FSubstrateHorizontalMixing Prototype = MakePrototype<FSubstrateHorizontalMixing>(FType::MakeSubstrateData());
+	Prototype.Background = Background;
+	Prototype.Foreground = Foreground;
+	Prototype.Mix = Mix;
+
+	// ParameterBlending = 0
+	Prototype.OperatorIndex = nullptr;
+	Prototype.MaxDistanceFromLeaves = nullptr;
+
+	// ParameterBlending = 1
+	Prototype.NormalMix = NormalMix;
+	Prototype.SharedLocalBasisIndexMacro = SharedLocalBasisIndexMacro;
+	Prototype.BackgroundNoV = BackgroundNoV;
+	Prototype.ForegroundNoV = ForegroundNoV;
+
+	// Settings
+	Prototype.bParameterBlendingEnabled = true;
+	Prototype.bFullySimplifiedInstruction = SubstrateTranslatorData->CurrentSubstrateCompilationContext == Substrate::ESubstrateCompilationContext::SCC_FullySimplified;
+
+	return EmitPrototype(*this, Prototype);
+}
+
+FValueRef FEmitter::SubstrateVerticalLayering(FValueRef Top, FValueRef Base, FValueRef OperatorIndex, FValueRef MaxDistanceFromLeaves)
+{
+	FSubstrateVerticalLayering Prototype = MakePrototype<FSubstrateVerticalLayering>(FType::MakeSubstrateData());
+	Prototype.Top = Top;
+	Prototype.Base = Base;
+
+	// ParameterBlending = 0
+	Prototype.OperatorIndex = OperatorIndex;
+	Prototype.MaxDistanceFromLeaves = MaxDistanceFromLeaves;
+
+	// ParameterBlending = 1
+	Prototype.SharedLocalBasisIndexMacro = nullptr;
+	Prototype.TopNoV = nullptr;
+	Prototype.BaseNoV = nullptr;
+
+	// Settings
+	Prototype.bParameterBlendingEnabled = false;
+	Prototype.bFullySimplifiedInstruction = SubstrateTranslatorData->CurrentSubstrateCompilationContext == Substrate::ESubstrateCompilationContext::SCC_FullySimplified;
+
+	return EmitPrototype(*this, Prototype);
+}
+
+FValueRef FEmitter::SubstrateVerticalLayeringParameterBlending(FValueRef Top, FValueRef Base, FValueRef SharedLocalBasisIndexMacro, FValueRef TopNoV, FValueRef BaseNoV)
+{
+	FSubstrateVerticalLayering Prototype = MakePrototype<FSubstrateVerticalLayering>(FType::MakeSubstrateData());
+	Prototype.Top = Top;
+	Prototype.Base = Base;
+
+	// ParameterBlending = 0
+	Prototype.OperatorIndex = nullptr;
+	Prototype.MaxDistanceFromLeaves = nullptr;
+
+	// ParameterBlending = 1
+	Prototype.SharedLocalBasisIndexMacro = SharedLocalBasisIndexMacro;
+	Prototype.TopNoV = TopNoV;
+	Prototype.BaseNoV = BaseNoV;
+
+	// Settings
+	Prototype.bParameterBlendingEnabled = true;
+	Prototype.bFullySimplifiedInstruction = SubstrateTranslatorData->CurrentSubstrateCompilationContext == Substrate::ESubstrateCompilationContext::SCC_FullySimplified;
+
 	return EmitPrototype(*this, Prototype);
 }
 
@@ -4217,6 +4394,11 @@ EMaterialQualityLevel::Type FEmitter::GetQualityLevel() const
 	return Module->GetQualityLevel();
 }
 
+Substrate::FSubstrateTranslatorData* FEmitter::GetSubstrateTranslatorData() const
+{
+	return SubstrateTranslatorData;
+}
+
 bool FEmitter::FValueKeyFuncs::Matches(KeyInitType A, KeyInitType B)
 {
 	return A->Equals(B);
diff --git a/Engine/Source/Runtime/Engine/Private/Materials/MaterialIRModuleBuilder.cpp b/Engine/Source/Runtime/Engine/Private/Materials/MaterialIRModuleBuilder.cpp
--- a/Engine/Source/Runtime/Engine/Private/Materials/MaterialIRModuleBuilder.cpp
+++ b/Engine/Source/Runtime/Engine/Private/Materials/MaterialIRModuleBuilder.cpp
@@ -20,6 +20,7 @@
 #include "Materials/MaterialExpressionNamedReroute.h"
 #include "Materials/MaterialExpressionSingleLayerWaterMaterialOutput.h"
 #include "Materials/MaterialExpressionThinTranslucentMaterialOutput.h"
+#include "Materials/MaterialExpressionSubstrate.h"
 #include "Materials/MaterialSharedPrivate.h"
 #include "Materials/Material.h"
 #include "MaterialDomain.h"
@@ -32,6 +33,7 @@
 #include "Misc/FileHelper.h"
 #include "PostProcess/PostProcessMaterialInputs.h"
 #include "RenderUtils.h"
+#include "SubstrateTranslatorCommon.h"
 
 static TAutoConsoleVariable<int32> CVarMaterialIRDebugDumpLevel(
 	TEXT("r.Material.Translator.DebugDump"),
@@ -83,6 +85,8 @@ struct FMaterialIRModuleBuilderImpl
 	MIR::FValue* DefaultMaterialAggregate;
 	FColorMaterialInput PreviewInput;
 
+	Substrate::FSubstrateTranslatorData SubstrateTranslatorData;
+
 	void Step_Initialize(FMaterialIRModuleBuilder* InBuilder, FMaterialIRModule* InModule)
 	{
 		this->Builder = InBuilder;
@@ -112,6 +116,7 @@ struct FMaterialIRModuleBuilderImpl
 		Emitter.Material = Builder->Material;
 		Emitter.Module = Module;
 		Emitter.StaticParameterSet = &Builder->StaticParameters;
+		Emitter.SubstrateTranslatorData = &SubstrateTranslatorData;
 		Emitter.Initialize();
 
 		// Create an IR value to hold the material attributes aggregate default
@@ -122,6 +127,32 @@ struct FMaterialIRModuleBuilderImpl
 
 		// Set the preview input expression
 		PreviewInput.Expression = Builder->PreviewExpression;
+
+		// Setup Substrate translator data to process the Substrate tree and simplification.
+		FSubstrateCompilationConfig* SubstrateCompilationConfig = nullptr;	// SUBSTRATE_TODO: forward SubstrateCompilationConfig from compilation step to make substrate tab simplification preview work.
+		if (SubstrateCompilationConfig)
+		{
+			SubstrateTranslatorData.SubstrateCompilationConfig = *SubstrateCompilationConfig;
+		}
+		SubstrateTranslatorData.Initialize(Builder->Material->GetMaterialResource(Module->ShaderPlatform), Module->CompilationOutput, nullptr, &Emitter);
+
+		FMaterialShadingModelField MaterialShadingModels;
+		// If the material gets its shading model from material expressions and we have compiled one or more shading model expressions already, 
+		// then use that shading model field instead. It's the most optimal set of shading models
+		if (Builder->Material->IsShadingModelFromMaterialExpression())// && ShadingModelsFromCompilation.IsValid())
+		{
+			check(false); // SUBSTRATE_TODO add support for that case.
+		}
+		else
+		{
+			MaterialShadingModels = Builder->Material->GetShadingModels();
+			UMaterialInterface::FilterOutPlatformShadingModels(Module->ShaderPlatform, MaterialShadingModels);
+		}
+
+		bool bDoClearFunctionStack = false;
+
+		UMaterialExpression* MaterialGraphNodePreviewExpression = nullptr; // SUBSTRATE_TODO use Builder->Material->GetMaterialGraphNodePreviewExpression();
+		SubstrateTranslatorData.ProcessSubstrateTreeAndTopology(MaterialShadingModels, MaterialGraphNodePreviewExpression, bDoClearFunctionStack);
 	}
 
 	void Step_PushRootExpressionDependencies()
@@ -187,6 +218,91 @@ struct FMaterialIRModuleBuilderImpl
 		return Module->IsValid();
 	}
 
+	bool Step_BuildSubstrateFrontMaterialOutput()
+	{
+		if (SubstrateTranslatorData.FrontMaterialExpr)
+		{
+			// Save the the FrontMaterial to a custom output.
+			MIR::FValueRef SubstrateFrontMaterial = MIR::FValueRef(AnalysisContextStack.Last().GetOutputValue(SubstrateTranslatorData.FrontMaterialExpr->GetOutput(0))); // SUBSTRATE_TODO always assuming output 0 for now but this can be wrong for some material function
+			Emitter.SetCustomOutputs(TEXTVIEW("SubstrateFrontMaterial"), { &SubstrateFrontMaterial, 1 }, MIR::EMaterialOutputStage::MaterialOutputStage_Pixel);
+		}
+
+		// Remove all substrate nodes and function calls from all the contexts.
+		bool bNeedSubstrateExpressionProcessing = false;
+		for (FAnalysisContext& Context : AnalysisContextStack)
+		{
+			// Re add expression we need to process them again
+			for (UMaterialExpression* BuiltExpression : Context.BuiltExpressions)
+			{
+				if (BuiltExpression->IsA(UMaterialExpressionSubstrateBSDF::StaticClass()))
+				{
+					Context.ExpressionStack.Push(BuiltExpression);
+					bNeedSubstrateExpressionProcessing = SubstrateTranslatorData.FrontMaterialExpr != nullptr;
+				}
+			}
+			// Remove them from built
+			for (UMaterialExpression* BuiltExpression : Context.ExpressionStack)
+			{
+				if (BuiltExpression->IsA(UMaterialExpressionSubstrateBSDF::StaticClass()))
+				{
+					Context.BuiltExpressions.Remove(BuiltExpression);
+				}
+			}
+		}
+
+		// Process the graph once again to generate IR for the FullySimplifiedFrontMaterial output.
+		if (bNeedSubstrateExpressionProcessing)
+		{
+			SubstrateTranslatorData.CurrentSubstrateCompilationContext = Substrate::ESubstrateCompilationContext::SCC_FullySimplified;
+
+			FAnalysisContext& FirstContext = AnalysisContextStack.Last();
+			FirstContext.ExpressionStack.Empty();
+			FirstContext.ExpressionStack.Push(SubstrateTranslatorData.FrontMaterialExpr);
+
+			while (true)
+			{
+				FMemMark MemMark{ FMemStack::Get() };
+				FAnalysisContext& Context = AnalysisContextStack.Last();
+
+				if (!Context.ExpressionStack.IsEmpty())
+				{
+					// Some expression is on the expression stack of this context. Analyze it. This will
+					// have the effect of either building the expression or pushing its other expression
+					// dependencies onto the stack.
+					BuildTopMaterialExpression();
+				}
+				else if (Context.Call)
+				{
+					// There are no more expressions to analyze on the stack, this analysis context is complete.
+					// Context.Call isn't null so this context is for a function call, which has now been fully analyzed.
+					// Pop the callee context from the stack and resume analyzing the parent context (the caller).
+					PopFunctionCall();
+				}
+				else
+				{
+					// No other expressions on the stack to evaluate, nor this is a function
+					// call context but the root context. Nothing left to do so simply quit.
+					break;
+				}
+			}
+
+			if (SubstrateTranslatorData.CurrentSubstrateCompilationContext == Substrate::ESubstrateCompilationContext::SCC_FullySimplified)
+			{
+				// Save the output to the FrontMaterial, now fully simplified, to the custom output representing fully simplified.
+				MIR::FValueRef SubstrateFrontMaterial = MIR::FValueRef(AnalysisContextStack.Last().GetOutputValue(SubstrateTranslatorData.FrontMaterialExpr->GetOutput(0)));
+				Emitter.SetCustomOutputs(TEXTVIEW("SubstrateFrontMaterialFullySimplified"), { &SubstrateFrontMaterial, 1 }, MIR::EMaterialOutputStage::MaterialOutputStage_Pixel);
+
+				SubstrateTranslatorData.CurrentSubstrateCompilationContext = Substrate::ESubstrateCompilationContext::SCC_Default;
+			}
+		}
+		else
+		{
+			// SUBSTRATE_TODO copy to fully simplified the first input?
+		}
+
+		return Module->IsValid();
+	}
+
 	void BuildTopMaterialExpression()
 	{
 		FAnalysisContext& CurrContext = AnalysisContextStack.Last();
@@ -248,6 +364,8 @@ struct FMaterialIRModuleBuilderImpl
 		}
 		else
 		{
+			UpdateSubstrateNodeIdentifierStack(CurrContext);
+
 			// Invoke the expression build function. This will perform semantic analysis, error reporting and
 			// emit IR values for its outputs (which will flow into connected expressions inputs).
 			Emitter.Expression->Build(Emitter);
@@ -256,6 +374,151 @@ struct FMaterialIRModuleBuilderImpl
 			AddExpressionConnectionInsights(Emitter.Expression);
 		}
 	}
+
+	void UpdateSubstrateNodeIdentifierStack(FAnalysisContext& CurrContext)
+	{
+		// Build the Substrate expressions unique GUID as if we had a SubstrateTreeStack.
+		// It needs to be done each time because expression are appended breath first while we need position Top-1 to be the parent
+		if (auto NewSubstrateExpr = Cast<UMaterialExpressionSubstrateBSDF>(Emitter.Expression))
+		{
+			Substrate::FSubstrateCompilationContext& SubstrateCtx = SubstrateTranslatorData.SubstrateCompilationContext[Substrate::ESubstrateCompilationContext::SCC_Default];
+			SubstrateCtx.SubstrateNodeIdentifierStack.Empty();
+
+			FAnalysisContext& SearchContext = CurrContext;
+			int32 SearchContextParentIndex = AnalysisContextStack.Num() - 2;
+
+			UMaterialExpression* TrackedNode = NewSubstrateExpr;
+			struct FWalkedTree
+			{
+				UMaterialExpression* Expression = nullptr;
+				uint32 InputIndex = 0;
+			};
+			TArray<FWalkedTree> WalkedNodeStack;
+			// If WalkedNodeStack is empty then we are only dealing with a single BSDF, otherwise we chain all operator from the BSDF to the root
+			for (int32 Index = SearchContext.ExpressionStack.Num() - 1; Index >= 0; Index--)
+			{
+				if (auto SubstrateHorizontalExpr = Cast<UMaterialExpressionSubstrateHorizontalMixing>(SearchContext.ExpressionStack[Index]))
+				{
+					if (SubstrateHorizontalExpr->Background.GetTracedInput().Expression == TrackedNode)
+					{
+						WalkedNodeStack.Push({ SubstrateHorizontalExpr , 0 });
+						TrackedNode = SubstrateHorizontalExpr;
+					}
+					else if (SubstrateHorizontalExpr->Foreground.GetTracedInput().Expression == TrackedNode)
+					{
+						WalkedNodeStack.Push({ SubstrateHorizontalExpr , 1 });
+						TrackedNode = SubstrateHorizontalExpr;
+					}
+					// Else, cannot error since it could be another mixing from another branch we never walk from there
+				}
+				else if (auto SubstrateDecalExpr = Cast<UMaterialExpressionSubstrateConvertToDecal>(SearchContext.ExpressionStack[Index]))
+				{
+					if (SubstrateDecalExpr->DecalMaterial.GetTracedInput().Expression == TrackedNode)
+					{
+						WalkedNodeStack.Push({ SubstrateDecalExpr , 0 });
+						TrackedNode = SubstrateDecalExpr;
+					}
+				}
+				else if (auto SubstrateVerticalExpr = Cast<UMaterialExpressionSubstrateVerticalLayering>(SearchContext.ExpressionStack[Index]))
+				{
+					if (SubstrateVerticalExpr->Top.GetTracedInput().Expression == TrackedNode)
+					{
+						WalkedNodeStack.Push({ SubstrateVerticalExpr , 0 });
+						TrackedNode = SubstrateVerticalExpr;
+					}
+					else if (SubstrateVerticalExpr->Base.GetTracedInput().Expression == TrackedNode)
+					{
+						WalkedNodeStack.Push({ SubstrateVerticalExpr , 1 });
+						TrackedNode = SubstrateVerticalExpr;
+					}
+				}
+				else if (auto SubstrateAddExpr = Cast<UMaterialExpressionSubstrateAdd>(SearchContext.ExpressionStack[Index]))
+				{
+					if (SubstrateAddExpr->A.GetTracedInput().Expression == TrackedNode)
+					{
+						WalkedNodeStack.Push({ SubstrateAddExpr , 0 });
+						TrackedNode = SubstrateAddExpr;
+					}
+					else if (SubstrateAddExpr->B.GetTracedInput().Expression == TrackedNode)
+					{
+						WalkedNodeStack.Push({ SubstrateAddExpr , 1 });
+						TrackedNode = SubstrateAddExpr;
+					}
+				}
+				else if (auto SubstrateWeightExpr = Cast<UMaterialExpressionSubstrateWeight>(SearchContext.ExpressionStack[Index]))
+				{
+					if (SubstrateWeightExpr->A.GetTracedInput().Expression == TrackedNode)
+					{
+						WalkedNodeStack.Push({ SubstrateWeightExpr , 0 });
+						TrackedNode = SubstrateWeightExpr;
+					}
+				}
+				else if (auto SubstrateSelectExpr = Cast<UMaterialExpressionSubstrateSelect>(SearchContext.ExpressionStack[Index]))
+				{
+					if (SubstrateSelectExpr->A.GetTracedInput().Expression == TrackedNode)
+					{
+						WalkedNodeStack.Push({ SubstrateSelectExpr , 0 });
+						TrackedNode = SubstrateSelectExpr;
+					}
+					else if (SubstrateSelectExpr->B.GetTracedInput().Expression == TrackedNode)
+					{
+						WalkedNodeStack.Push({ SubstrateSelectExpr , 1 });
+						TrackedNode = SubstrateSelectExpr;
+					}
+				}
+				else if (auto FunctionOutput = Cast<UMaterialExpressionFunctionOutput>(SearchContext.ExpressionStack[Index]))
+				{
+					// We need to select another context
+					if (SearchContext.Call == nullptr)
+					{
+						break; // we are done
+					}
+					else
+					{
+						// Update the tracked node to the output
+						TrackedNode = SearchContext.Call;
+
+						SearchContext = AnalysisContextStack[SearchContextParentIndex];
+						SearchContextParentIndex--;
+						Index = SearchContext.ExpressionStack.Num() - 1;
+					}
+				}
+			}
+
+			// Now update the stack on the context
+			SubstrateCtx.SubstrateNodeIdentifierStack.Push(FGuid(0x7AEE, 0xBAD, 0xDEAD, 0xBEEF));
+			for (int32 Index = WalkedNodeStack.Num() - 1; Index >= 0; Index--)
+			{
+				FWalkedTree& SubstrateExpr = WalkedNodeStack[Index];
+
+
+				// Create an md5 hash for the parent, its input pin index and current node to represent the path.
+				uint32 IntputHashBuffer[9];
+				FGuid PreviousNodeGuid = SubstrateCtx.SubstrateNodeIdentifierStack.Top();
+				FGuid NodeGuid = SubstrateExpr.Expression->MaterialExpressionGuid;
+				IntputHashBuffer[0] = PreviousNodeGuid.A;
+				IntputHashBuffer[1] = PreviousNodeGuid.B;
+				IntputHashBuffer[2] = PreviousNodeGuid.C;
+				IntputHashBuffer[3] = PreviousNodeGuid.D;
+				IntputHashBuffer[4] = SubstrateExpr.InputIndex;
+				IntputHashBuffer[5] = NodeGuid.A;
+				IntputHashBuffer[6] = NodeGuid.B;
+				IntputHashBuffer[7] = NodeGuid.C;
+				IntputHashBuffer[8] = NodeGuid.D;
+
+				uint32 OutputHashBuffer[]{ 0, 0, 0, 0 };
+				FMD5 IdentifierStringHash;
+				IdentifierStringHash.Update((uint8*)IntputHashBuffer, sizeof(IntputHashBuffer));
+				IdentifierStringHash.Final((uint8*)&OutputHashBuffer);
+
+				SubstrateCtx.SubstrateNodeIdentifierStack.Push(FGuid(OutputHashBuffer[0], OutputHashBuffer[1], OutputHashBuffer[2], OutputHashBuffer[3]));
+			}
+
+			// Simply copy for now to the fully simplified as it should be the same. SUBSTRATE_TODO: simplify, move the stack onto the SubstrateTranslatorData instead of context.
+			check(Substrate::ESubstrateCompilationContext::SCC_MAX == 2);
+			SubstrateTranslatorData.SubstrateCompilationContext[Substrate::ESubstrateCompilationContext::SCC_FullySimplified].SubstrateNodeIdentifierStack = SubstrateTranslatorData.SubstrateCompilationContext[Substrate::ESubstrateCompilationContext::SCC_Default].SubstrateNodeIdentifierStack;
+		}
+	}
 	
 	// Pushes an expression dependency on this context expression stack.
 	void PushDependency(UMaterialExpression* Expression)
@@ -436,6 +699,15 @@ struct FMaterialIRModuleBuilderImpl
 						PropertyValue = Emitter.Vector3(PropertyValue, Zero, Zero);
 					}
 				}
+				else if (Property == MP_FrontMaterial )
+				{
+					// Just check that we get a struct for now. 
+					// SUBSTRATE_TODO We should probably add a proper SubstrateData struct enum?
+					if (InputDesc.Type != UE::Shader::EValueType::Struct)
+					{
+						Emitter.Error(TEXT("FrontMaterial must be fed with a Struct"));
+					}
+				}
 				else
 				{
 					// If a value is flowing in through the connection, cast it to this material attribute type and assign it.
@@ -443,6 +715,10 @@ struct FMaterialIRModuleBuilderImpl
 					PropertyValue = Emitter.Cast(PropertyValue, OutputArgType);
 				}
 			}
+			else if (Property == MP_FrontMaterial && PropertyValue == nullptr && SubstrateTranslatorData.bMaterialUsesRootNodeToSubstrateHiddenConversion)
+			{
+				PropertyValue = BuildSubstrateFrontMaterialIRFromRootNodeHiddenConversion();
+			}
 			else if (InputDesc.bUseConstant)
 			{
 				// If input is marked to use constant, assign this output to the specified constant value.
@@ -486,6 +762,146 @@ struct FMaterialIRModuleBuilderImpl
 		return Module->IsValid();
 	}
 
+	MIR::FValue* BuildSubstrateFrontMaterialIRFromRootNodeHiddenConversion()
+	{
+		MIR::FValue* FrontMaterialPropertyValue = Emitter.SubstrateDefaultSlab();
+
+		FMaterial* Material = SubstrateTranslatorData.GetMaterial();
+		if (Material)
+		{
+			UMaterialInterface* MaterialInterface = Material ? Material->GetMaterialInterface() : nullptr;
+			UMaterial* BaseMaterial = MaterialInterface ? MaterialInterface->GetMaterial() : nullptr;
+			UMaterialEditorOnlyData* EditorOnlyData = BaseMaterial ? BaseMaterial->GetEditorOnlyData() : nullptr;
+
+			if (BaseMaterial && BaseMaterial->bUseMaterialAttributes)
+			{
+				check(false); // SUBSTRATE_TODO
+				return FrontMaterialPropertyValue;
+				//	// Need to root subsurface profile for the uniform scalar numeric parameter name to be correct.
+				//	USubsurfaceProfile* SSSProfile = MaterialInterface ? MaterialInterface->GetSubsurfaceProfileRoot_Internal() : nullptr;
+				//	const bool bHasSSS = SSSProfile != nullptr;
+				//
+				//	return UMaterialExpressionSubstrateConvertMaterialAttributes::CompileCommon(this, 0,
+				//		FMaterialAttributeDefinitionMap::GetConnectedMaterialAttributesBitmask(BaseMaterial->GetExpressions()), EditorOnlyData->MaterialAttributes, MaterialInterface->GetShadingModels().GetFirstShadingModel(),
+				//		*SubstrateTranslatorData.GetWaterScatteringCoefficients(), *SubstrateTranslatorData.GetWaterAbsorptionCoefficients(), *SubstrateTranslatorData.GetWaterPhaseG(), *SubstrateTranslatorData.GetColorScaleBehindWater(),
+				//		bHasSSS, SSSProfile,
+				//		*SubstrateTranslatorData.GetClearCoatNormal(), *SubstrateTranslatorData.GetCustomTangent());
+			}
+			else if ((Material->GetMaterialDomain() == MD_Surface || Material->GetMaterialDomain() == MD_RuntimeVirtualTexture || Material->GetMaterialDomain() == MD_DeferredDecal) && EditorOnlyData)
+			{
+				USubsurfaceProfile* SSSProfile = MaterialInterface ? MaterialInterface->GetSubsurfaceProfileRoot_Internal() : nullptr;
+				const bool bHasSSS = SSSProfile != nullptr;
+				const bool bHasAnisotropy = EditorOnlyData->Anisotropy.IsConnected();
+
+				FrontMaterialPropertyValue = UMaterialExpressionSubstrateShadingModels::BuildCommon(Emitter, Builder->Material,
+					EditorOnlyData->BaseColor, EditorOnlyData->Specular, EditorOnlyData->Metallic, EditorOnlyData->Roughness, EditorOnlyData->EmissiveColor,
+					EditorOnlyData->Opacity, EditorOnlyData->SubsurfaceColor, EditorOnlyData->ClearCoat, EditorOnlyData->ClearCoatRoughness,
+					EditorOnlyData->ShadingModelFromMaterialExpression, MaterialInterface->GetShadingModels().GetFirstShadingModel(),
+					*SubstrateTranslatorData.GetThinTranslucentTransmittanceColor(), *SubstrateTranslatorData.GetThinTranslucentSurfaceCoverage(),
+					*SubstrateTranslatorData.GetWaterScatteringCoefficients(), *SubstrateTranslatorData.GetWaterAbsorptionCoefficients(), *SubstrateTranslatorData.GetWaterPhaseG(), *SubstrateTranslatorData.GetColorScaleBehindWater(),
+					bHasAnisotropy, EditorOnlyData->Anisotropy,
+					EditorOnlyData->Normal, EditorOnlyData->Tangent,
+					*SubstrateTranslatorData.GetClearCoatNormal(), *SubstrateTranslatorData.GetCustomTangent(),
+					bHasSSS, SSSProfile,
+					EditorOnlyData);
+			}
+			else if (Material->GetMaterialDomain() == MD_Volume && EditorOnlyData)
+			{
+				check(false); // SUBSTRATE_TODO
+				return FrontMaterialPropertyValue;
+				//	return UMaterialExpressionSubstrateVolumetricFogCloudBSDF::CompileCommon(
+				//		this,
+				//		EditorOnlyData->BaseColor,
+				//		EditorOnlyData->SubsurfaceColor,
+				//		EditorOnlyData->EmissiveColor,
+				//		EditorOnlyData->AmbientOcclusion,
+				//		MaterialInterface->GetShadingModels().HasOnlyShadingModel(MSM_Unlit),
+				//		EditorOnlyData);
+			}
+			else if (Material->GetMaterialDomain() == MD_LightFunction && EditorOnlyData)
+			{
+				check(false); // SUBSTRATE_TODO
+				return FrontMaterialPropertyValue;
+				//	return UMaterialExpressionSubstrateLightFunction::CompileCommon(this, EditorOnlyData->EmissiveColor, EditorOnlyData);
+			}
+			else if (Material->GetMaterialDomain() == MD_PostProcess && EditorOnlyData)
+			{
+				check(false); // SUBSTRATE_TODO
+				return FrontMaterialPropertyValue;
+				//	return UMaterialExpressionSubstratePostProcess::CompileCommon(this, EditorOnlyData->EmissiveColor, EditorOnlyData->Opacity, EditorOnlyData);
+			}
+			else if (Material->GetMaterialDomain() == MD_UI && EditorOnlyData)
+			{
+				check(false); // SUBSTRATE_TODO
+				return FrontMaterialPropertyValue;
+				//	return UMaterialExpressionSubstrateUI::CompileCommon(this, EditorOnlyData->EmissiveColor, EditorOnlyData->Opacity, EditorOnlyData);
+			}
+
+
+			// When we are running hidden conversion (bMaterialUsesRootNodeToSubstrateHiddenConversion),
+			// We know that we will only ever have a single BSDF and shared local bases. So simply use the same non-fully-simplified as the fully-simplified version of the material too.
+			MIR::FValueRef PropertyValueRef = MIR::FValueRef(FrontMaterialPropertyValue);
+			Emitter.SetCustomOutputs(TEXTVIEW("SubstrateFrontMaterialFullySimplified"), { &PropertyValueRef, 1 }, MIR::EMaterialOutputStage::MaterialOutputStage_Pixel);
+			Emitter.SetCustomOutputs(TEXTVIEW("SubstrateFrontMaterial"), { &PropertyValueRef, 1 }, MIR::EMaterialOutputStage::MaterialOutputStage_Pixel);
+			// We also copy the copy the shared local bases structure from within GeneratePixelInputInitializerHLSLCode.
+		}
+		else
+		{
+			check(false); // SUBSTRATE_TODO should not happen
+			return FrontMaterialPropertyValue;
+		}
+
+		return FrontMaterialPropertyValue;
+	}
+
+	void GenerateSubstrateMaterialCompilationOutputAndDefines()
+	{
+		// Write substrate tree specific code.
+		if (SubstrateTranslatorData.bSubstrateEnabled)
+		{
+			const bool bSubstrateFrontMaterialProvided = Module->IsMaterialPropertyUsed(MP_FrontMaterial);
+			const bool bValid = SubstrateTranslatorData.GenerateSubstrateTreeHLSLCode(bSubstrateFrontMaterialProvided);
+			Module->SubstrateData.TreeHLSL = SubstrateTranslatorData.SubstrateTreeHLSL;
+
+			// TOOD add that back
+			//			if (!bValid)
+			//			{
+			//				TranslationResult = EHLSLMaterialTranslatorResult::Failure;
+			//				return;
+			//			}
+		}
+
+		// SUBSTRATE_TODO export context
+		ESubstrateMaterialExport SubstrateMaterialExportType = ESubstrateMaterialExport::SME_None;
+		ESubstrateMaterialExportContext SubstrateMaterialExportContext = ESubstrateMaterialExportContext::SMEC_Opaque;
+		uint8 SubstrateMaterialExportLegacyBlendMode = BLEND_Opaque;
+
+		const bool bOpacityPropertyIsUsed = Module->IsMaterialPropertyUsed(MP_Opacity);
+		FString DescriptionStringCommentForDebug;
+		Substrate::FSubstrateDefines OutSubstrateDefines;
+		SubstrateTranslatorData.GenerateMaterialCompilationOutputAndDefines(
+			bOpacityPropertyIsUsed,
+			(int32)SubstrateMaterialExportType, (int32)SubstrateMaterialExportContext, (int32)SubstrateMaterialExportLegacyBlendMode,
+			DescriptionStringCommentForDebug,
+			&OutSubstrateDefines);
+
+
+		if (SubstrateTranslatorData.bSubstrateEnabled)
+		{
+			FString HLSLSubstrateTypeString = TEXT("FSubstrateData");	// TODO we would need to abstract that away only when lowering the IR
+			SubstrateTranslatorData.GeneratePixelMemberDeclarationHLSLCode(Module->SubstrateData.PixelMembersDeclarationHLSL, HLSLSubstrateTypeString);
+			SubstrateTranslatorData.GeneratePixelInputInitializerHLSLCode(Module->SubstrateData.PixelInputInitializerValuesHLSL);
+			SubstrateTranslatorData.GeneratePixelNormalInitializerHLSLCode(Module->SubstrateData.PixelInputInitializerValuesHLSL);
+		}
+
+		// Copie Substrate defines
+		// SUBSTRATE_TODO: have SubstrateTranslatorData and both translator use the same define format to just add to the existing array one
+		for (auto& Define : OutSubstrateDefines)
+		{
+			Module->AddIntegerEnvironmentDefine(*Define.Key, Define.Value);
+		}
+	}
+
 	bool GetExpressionInputDescription(EMaterialProperty Property, FMaterialInputDescription& Input)
 	{
 		if (PreviewInput.IsConnected() && Property == MP_EmissiveColor)
@@ -653,6 +1069,9 @@ struct FMaterialIRModuleBuilderImpl
 			Module->AddError(nullptr, FString::Printf(TEXT("Maximum Scene Texture post process inputs exceeded (%d > %d), between SceneTexture nodes with PostProcessInputs or UserSceneTexture nodes."), NumPostProcessInputs, kPostProcessMaterialInputCountMax));
 		}
 
+		// Generate specific code and define.
+		GenerateSubstrateMaterialCompilationOutputAndDefines();
+
 		// Consolidate and validate environment defines
 		ConsolidateEnvironmentDefines();
 		ValidateEnvironmentDefines();
@@ -961,6 +1380,14 @@ bool FMaterialIRModuleBuilder::Build(FMaterialIRModule* TargetModule)
 		return false;
 	}
 
+	// New that we have generated the IR for the front material output, we store it on a custom output used explicitely later in materials.
+	// We also generate the IR for the fully simplified version of the front material output (reprocess all the substrate nodes with the FullySimplified context).
+	// Then a material shader can decide to use the normal of fully simplified materials based on SUBSTRATE_USE_FULLYSIMPLIFIED_MATERIAL from this single translated shader file.
+	if (!Impl.Step_BuildSubstrateFrontMaterialOutput())
+	{
+		return false;
+	}
+
 	// Materials dont have a final "output expression", so this step grabs the values flowing
 	// into the material property output pins (if any) and generates SetMaterialProperty instructions
     // handling details such as default values, preview material expression, etc.
