Substrate

diff --git a/Engine/Source/Runtime/Engine/Private/Materials/MaterialIRToHLSLTranslator.cpp b/Engine/Source/Runtime/Engine/Private/Materials/MaterialIRToHLSLTranslator.cpp
--- a/Engine/Source/Runtime/Engine/Private/Materials/MaterialIRToHLSLTranslator.cpp
+++ b/Engine/Source/Runtime/Engine/Private/Materials/MaterialIRToHLSLTranslator.cpp
@@ -303,6 +303,11 @@ static bool IsFoldable(const MIR::FInstruction* Instr, int32 EntryPointIndex)
 		case MIR::VK_SetMaterialOutput:
 		case MIR::VK_VTPageTableRead:
 		case MIR::VK_Call:
+		case MIR::VK_SubstrateHorizontalMixing:
+		case MIR::VK_SubstrateVerticalLayering:
+		case MIR::VK_SubstratePromoteToOperator:
+		case MIR::VK_SubstrateSlab:
+		case MIR::VK_SubstrateShadingModels:
 			return false;
 		default:
 			return true;
@@ -395,8 +400,21 @@ struct FPrivate : FMaterialIRToHLSLTranslation
 			Printer << TEXTVIEW("PixelMaterialInputs.Subsurface = 0") << EndOfStatement;
 		}
 
-		FString ParamName = (Stage == MIR::Stage_Pixel) ? FString{ TEXTVIEW("calc_pixel_material_inputs_other_inputs") } : FString{ TEXTVIEW("calc_pixel_material_inputs_analytic_derivatives_other_inputs") };
-		OutputShaderTemplateParams.Add(MoveTemp(ParamName), MoveTemp(Printer.Buffer));
+		FString ParamName = FString((Stage == MIR::Stage_Pixel) ? TEXT("calc_pixel_material_inputs_other_inputs") : TEXT("calc_pixel_material_inputs_analytic_derivatives_other_inputs"));
+		OutputShaderTemplateParams.Add(ParamName, MoveTemp(Printer.Buffer));
+
+		// Substrate extra code. For instance the shared local bases generation.
+		const FMaterialIRModule::FSubstrateData& SubstrateData = Module->GetSubstrateData();
+
+		FString* PixelMaterialOtherInputs = OutputShaderTemplateParams.Find(ParamName);
+		if (PixelMaterialOtherInputs)
+		{
+			PixelMaterialOtherInputs->Append(SubstrateData.PixelInputInitializerValuesHLSL);
+		}
+		else
+		{
+			OutputShaderTemplateParams.Add(ParamName, SubstrateData.PixelInputInitializerValuesHLSL);
+		}
 	}
 	
 	ENoOp LowerBlock(const MIR::FBlock& Block)
@@ -513,7 +531,12 @@ struct FPrivate : FMaterialIRToHLSLTranslation
 			case MIR::VK_TextureRead: LowerTextureRead(static_cast<const MIR::FTextureRead*>(Instr)); break;
 			case MIR::VK_VTPageTableRead: LowerVTPageTableRead(static_cast<const MIR::FVTPageTableRead*>(Instr)); break;
 			case MIR::VK_InlineHLSL: LowerInlineHLSL(static_cast<const MIR::FInlineHLSL*>(Instr)); break;
-			case MIR::VK_PromoteSubstrateParameter: LowerPromoteSubstrateParameter(static_cast<const MIR::FPromoteSubstrateParameter*>(Instr)); break;
+			case MIR::VK_SubstrateDefaultSlab: LowerSubstrateDefaultSlab(static_cast<const MIR::FSubstrateDefaultSlab*>(Instr)); break;
+			case MIR::VK_SubstrateSlab: LowerSubstrateSlab(static_cast<const MIR::FSubstrateSlab*>(Instr)); break;
+			case MIR::VK_SubstrateShadingModels: LowerSubstrateShadingModel(static_cast<const MIR::FSubstrateShadingModels*>(Instr)); break;
+			case MIR::VK_SubstrateHorizontalMixing: LowerSubstrateHorizontalMixing(static_cast<const MIR::FSubstrateHorizontalMixing*>(Instr)); break;
+			case MIR::VK_SubstrateVerticalLayering: LowerSubstrateVerticalLayering(static_cast<const MIR::FSubstrateVerticalLayering*>(Instr)); break;
+			case MIR::VK_SubstratePromoteToOperator: LowerSubstratePromoteToOperator(static_cast<const MIR::FSubstratePromoteToOperator*>(Instr)); break;
 			case MIR::VK_StageSwitch: LowerStageSwitch(static_cast<const MIR::FStageSwitch*>(Instr)); break;
 			case MIR::VK_HardwarePartialDerivative: LowerHardwarePartialDerivative(static_cast<const MIR::FHardwarePartialDerivative*>(Instr)); break;
 			case MIR::VK_Nop: LowerNop(static_cast<const MIR::FNop*>(Instr)); break;
@@ -871,115 +894,155 @@ struct FPrivate : FMaterialIRToHLSLTranslation
 		}
 	}
 
-	struct FSubstrateLegacyArgument
+	void LowerSubstrateDefaultSlab(const MIR::FSubstrateDefaultSlab* SubstrateParameter)
 	{
-		MIR::FType ParameterType;
-		FStringView ParameterName;
-		EMaterialProperty MaterialProperty;
-		FStringView DefaultInlineHLSL;
-	};
+		Printer << TEXTVIEW("GetInitialisedSubstrateData();");
 
-	static TConstArrayView<FSubstrateLegacyArgument> GetSubstrateLegacyConversionArguments()
-	{
-		static const FSubstrateLegacyArgument LegacyArguments[] =
-		{
-			{ MIR::FType::MakeFloatVector(3), TEXTVIEW("BaseColor"), MP_BaseColor, {} },
-			{ MIR::FType::MakeFloatScalar(),  TEXTVIEW("Specular"), MP_SpecularColor, TEXTVIEW("0.5f") },
-			{ MIR::FType::MakeFloatScalar(),  TEXTVIEW("Metallic"), MP_Metallic, {} },
-			{ MIR::FType::MakeFloatScalar(),  TEXTVIEW("Roughness"), MP_Roughness, {} },
-			{ MIR::FType::MakeFloatScalar(),  TEXTVIEW("Anisotropy"), MP_Anisotropy, {} },
-			{ MIR::FType::MakeFloatVector(3), TEXTVIEW("SubSurfaceColor"), MP_SubsurfaceColor, TEXTVIEW("(float3)1") },
-			{ MIR::FType::MakeFloatScalar(),  TEXTVIEW("SubSurfaceProfileId"), MP_MAX, TEXTVIEW("0") },
-			{ MIR::FType::MakeFloatScalar(),  TEXTVIEW("ClearCoat"), MP_MAX, TEXTVIEW("1.0f") },
-			{ MIR::FType::MakeFloatScalar(),  TEXTVIEW("ClearCoatRoughness"), MP_MAX, TEXTVIEW("0.1f") },
-			{ MIR::FType::MakeFloatVector(3), TEXTVIEW("Emissive"), MP_EmissiveColor, {} },
-			{ MIR::FType::MakeFloatScalar(),  TEXTVIEW("Opacity"), MP_Opacity, {} },
-			{ MIR::FType::MakeFloatVector(3), TEXTVIEW("ThinTranslucentTransmittanceColor"), MP_MAX, TEXTVIEW("(float3)0.5f") },
-			{ MIR::FType::MakeFloatScalar(),  TEXTVIEW("ThinTranslucentSurfaceCoverage"), MP_MAX, TEXTVIEW("1.0f") },
-			{ MIR::FType::MakeFloatVector(3), TEXTVIEW("WaterScatteringCoefficients"), MP_MAX, {} },
-			{ MIR::FType::MakeFloatVector(3), TEXTVIEW("WaterAbsorptionCoefficients"), MP_MAX, {} },
-			{ MIR::FType::MakeFloatScalar(),  TEXTVIEW("WaterPhaseG"), MP_MAX, {} },
-			{ MIR::FType::MakeFloatVector(3), TEXTVIEW("ColorScaleBehindWater"), MP_MAX, TEXTVIEW("(float3)1.0f") },
-			{ MIR::FType::MakeIntScalar(),    TEXTVIEW("ShadingModel"), MP_MAX, TEXTVIEW("1") },
-			{ MIR::FType::MakeFloatVector(3), TEXTVIEW("RawNormal"), MP_Normal, TEXTVIEW("Parameters.TangentToWorld[2]") },
-			{ MIR::FType::MakeFloatVector(3), TEXTVIEW("RawTangent"), MP_Tangent, TEXTVIEW("Parameters.TangentToWorld[0]") },
-			{ MIR::FType::MakeFloatVector(3), TEXTVIEW("RawClearCoatNormal"), MP_Normal, TEXTVIEW("Parameters.TangentToWorld[2]") },
-			{ MIR::FType::MakeFloatVector(3), TEXTVIEW("RawCustomTangent"), MP_Tangent, TEXTVIEW("Parameters.TangentToWorld[0]") },
-			{ MIR::FType::MakeIntScalar(),    TEXTVIEW("SharedLocalBasisIndex"), MP_MAX, {} },
-			{ MIR::FType::MakeIntScalar(),    TEXTVIEW("ClearCoatBottomNormal_SharedLocalBasisIndex"), MP_MAX, {} },
-			{ MIR::FType::MakeIntScalar(),    TEXTVIEW("SharedLocalBasisTypes"), MP_MAX, TEXTVIEW("Parameters.SharedLocalBases.Types") },
-		};
-		return LegacyArguments;
 	}
 
-	static TConstArrayView<FSubstrateLegacyArgument> GetSubstrateUnlitArguments()
+	static FStringView GetSharedLocalBasesTypes(bool bFullySimplifiedInstruction)
+	{
+		return bFullySimplifiedInstruction ? TEXTVIEW("Parameters.SharedLocalBasesFullySimplified.Types") : TEXTVIEW("Parameters.SharedLocalBases.Types");
+	}
+	static FStringView GetSubstrateTree(bool bFullySimplifiedInstruction)
+	{
+		return bFullySimplifiedInstruction ? TEXTVIEW("Parameters.SubstrateTreeFullySimplified") : TEXTVIEW("Parameters.SubstrateTree");
+	}
+
+	void LowerSubstrateSlab(const MIR::FSubstrateSlab* SubstrateParameter)
 	{
-		static const FSubstrateLegacyArgument LegacyArguments[] =
+		Printer << TEXTVIEW("GetSubstrateSlabBSDF");
+		Printer << BeginArgs;
+		Printer << TEXTVIEW("Parameters.SubstratePixelFootprint") << TEXTVIEW(",");
+		Printer << LowerValue(SubstrateParameter->Normal) << TEXTVIEW(",");
+		Printer << LowerValue(SubstrateParameter->DiffuseAlbedo) << TEXTVIEW(",");
+		Printer << LowerValue(SubstrateParameter->F0) << TEXTVIEW(",");
+		Printer << LowerValue(SubstrateParameter->F90) << TEXTVIEW(",");
+		Printer << LowerValue(SubstrateParameter->Roughness) << TEXTVIEW(",");
+		Printer << LowerValue(SubstrateParameter->Anisotropy) << TEXTVIEW(",");
+		Printer << LowerValue(SubstrateParameter->SSSProfileId) << TEXTVIEW(",");
+		Printer << TEXTVIEW("false") << TEXTVIEW(","); // bSupportDefaultSSSProfile
+		Printer << LowerValue(SubstrateParameter->SSSMFP) << TEXTVIEW(",");
+		Printer << LowerValue(SubstrateParameter->SSSMFPScale) << TEXTVIEW(",");
+		Printer << LowerValue(SubstrateParameter->SSSPhaseAniso) << TEXTVIEW(",");
+		Printer << LowerValue(SubstrateParameter->SSSType) << TEXTVIEW(",");
+		Printer << LowerValue(SubstrateParameter->EmissiveColor) << TEXTVIEW(",");
+		Printer << LowerValue(SubstrateParameter->SecondRoughness) << TEXTVIEW(",");
+		Printer << LowerValue(SubstrateParameter->SecondRoughnessWeight) << TEXTVIEW(",");
+		Printer << LowerValue(SubstrateParameter->SecondRoughnessAsSimpleClearCoat) << TEXTVIEW(",");
+		Printer << LowerValue(SubstrateParameter->ClearCoatUseSecondNormal) << TEXTVIEW(",");
+		Printer << LowerValue(SubstrateParameter->ClearCoatBottomNormal) << TEXTVIEW(",");
+		Printer << LowerValue(SubstrateParameter->FuzzAmount) << TEXTVIEW(",");
+		Printer << LowerValue(SubstrateParameter->FuzzColor) << TEXTVIEW(",");
+		Printer << LowerValue(SubstrateParameter->FuzzRoughness) << TEXTVIEW(",");
+		Printer << LowerValue(SubstrateParameter->GlintValue) << TEXTVIEW(",");
+		Printer << LowerValue(SubstrateParameter->GlintUV) << TEXTVIEW(",");
+		Printer << LowerValue(SubstrateParameter->SpecularProfileId) << TEXTVIEW(",");
+		Printer << LowerValue(SubstrateParameter->Thickness) << TEXTVIEW(",");
+		Printer << LowerValue(SubstrateParameter->IsThin) << TEXTVIEW(",");
+		Printer << LowerValue(SubstrateParameter->IsAtBottom) << TEXTVIEW(",");
+		Printer << LowerValue(SubstrateParameter->LocalBasisIndex) << TEXTVIEW(",");
+		Printer << GetSharedLocalBasesTypes(SubstrateParameter->bFullySimplifiedInstruction);
+		Printer << EndArgs;
+	}
+
+	void LowerSubstrateShadingModel(const MIR::FSubstrateShadingModels* SubstrateParameter)
+	{
+		Printer << TEXTVIEW("SubstrateConvertLegacyMaterial") << (SubstrateParameter->bHasDynamicShadingModels ? TEXTVIEW("Dynamic") : TEXTVIEW("Static"));
+		Printer << BeginArgs;
+		Printer << TEXTVIEW("Parameters.SubstratePixelFootprint") << TEXTVIEW(",");
+		Printer << LowerValue(SubstrateParameter->BaseColor) << TEXTVIEW(",");
+		Printer << LowerValue(SubstrateParameter->Specular) << TEXTVIEW(",");
+		Printer << LowerValue(SubstrateParameter->Metallic) << TEXTVIEW(",");
+		Printer << LowerValue(SubstrateParameter->Roughness) << TEXTVIEW(",");
+		Printer << LowerValue(SubstrateParameter->Anisotropy) << TEXTVIEW(",");
+		Printer << LowerValue(SubstrateParameter->SubSurfaceColor) << TEXTVIEW(",");
+		Printer << LowerValue(SubstrateParameter->SubSurfaceProfileId) << TEXTVIEW(",");
+		Printer << LowerValue(SubstrateParameter->ClearCoat) << TEXTVIEW(",");
+		Printer << LowerValue(SubstrateParameter->ClearCoatRoughness) << TEXTVIEW(",");
+		Printer << LowerValue(SubstrateParameter->EmissiveColor) << TEXTVIEW(",");
+		Printer << LowerValue(SubstrateParameter->Opacity) << TEXTVIEW(",");
+		Printer << LowerValue(SubstrateParameter->ThinTranslucentTransmittanceColor) << TEXTVIEW(",");
+		Printer << LowerValue(SubstrateParameter->ThinTranslucentSurfaceCoverage) << TEXTVIEW(",");
+		Printer << LowerValue(SubstrateParameter->WaterScatteringCoefficients) << TEXTVIEW(",");
+		Printer << LowerValue(SubstrateParameter->WaterAbsorptionCoefficients) << TEXTVIEW(",");
+		Printer << LowerValue(SubstrateParameter->WaterPhaseG) << TEXTVIEW(",");
+		Printer << LowerValue(SubstrateParameter->ColorScaleBehindWater) << TEXTVIEW(",");
+		Printer << LowerValue(SubstrateParameter->ShadingModel) << TEXTVIEW(",");
+		Printer << LowerValue(SubstrateParameter->Normal) << TEXTVIEW(",");
+		Printer << LowerValue(SubstrateParameter->Tangent) << TEXTVIEW(",");
+		Printer << LowerValue(SubstrateParameter->ClearCoatNormal) << TEXTVIEW(",");
+		Printer << LowerValue(SubstrateParameter->CustomTangent) << TEXTVIEW(",");
+		Printer << LowerValue(SubstrateParameter->BasisIndexMacro) << TEXTVIEW(",");
+		Printer << LowerValue(SubstrateParameter->ClearCoat_BasisIndexMacro) << TEXTVIEW(",");
+		Printer << GetSharedLocalBasesTypes(SubstrateParameter->bFullySimplifiedInstruction);
+		Printer << EndArgs;
+	}
+
+	void LowerSubstrateHorizontalMixing(const MIR::FSubstrateHorizontalMixing* SubstrateParameter)
+	{
+		if (SubstrateParameter->bParameterBlendingEnabled)
 		{
-			{ MIR::FType::MakeFloatVector(3), TEXTVIEW("Emissive"), MP_EmissiveColor, {} },
-			{ MIR::FType::MakeFloatScalar(), TEXTVIEW("TransmittanceColor"), MP_MAX, TEXTVIEW("(float3)0.0f") },
-			{ MIR::FType::MakeFloatVector(3), TEXTVIEW("RawNormal"), MP_Normal, TEXTVIEW("Parameters.TangentToWorld[2]") },
-		};
-		return LegacyArguments;
-	}
-
-	void LowerPromoteSubstrateParameter(const MIR::FPromoteSubstrateParameter* SubstrateParameter)
-	{
-		Printer << TEXTVIEW("Parameters.SubstrateTree.PromoteParameterBlendedBSDFToOperator") << BeginArgs;
-		Printer.Tabs++;
-		Printer << NewLine << (SubstrateParameter->bIsUnlit ? TEXTVIEW("GetSubstrateUnlitBSDF") : TEXTVIEW("SubstrateConvertLegacyMaterialStatic")) << BeginArgs;
-		Printer.Tabs++;
-		Printer << ListSeparator << NewLine << TEXTVIEW("/*PixelFootprint:*/ Parameters.SubstratePixelFootprint");
-
-		TConstArrayView<FSubstrateLegacyArgument> Arguments = SubstrateParameter->bIsUnlit ? GetSubstrateUnlitArguments() : GetSubstrateLegacyConversionArguments();
-
-		for (const FSubstrateLegacyArgument& Argument : Arguments)
-		{
-			Printer << ListSeparator << NewLine << TEXTVIEW("/*") << Argument.ParameterName << TEXTVIEW("*/ ");
-
-			switch (Argument.MaterialProperty)
-			{
-				case MP_Normal: Printer << TEXTVIEW("float3(0,0,1)") << NewLine; break;
-				case MP_Tangent:  Printer << TEXTVIEW("float3(1,0,0)") << NewLine; break;
-
-				default:
-					if (Argument.MaterialProperty != MP_MAX && SubstrateParameter->PropertyArgs[Argument.MaterialProperty])
-					{
-						LowerValue(SubstrateParameter->PropertyArgs[Argument.MaterialProperty]);
-					}
-					else if (!Argument.DefaultInlineHLSL.IsEmpty())
-					{
-						Printer << Argument.DefaultInlineHLSL;
-					}
-					else if (TOptional<MIR::FPrimitive> PrimitiveType = Argument.ParameterType.AsPrimitive())
-					{
-						if (PrimitiveType->IsScalar())
-						{
-							Printer << TEXTVIEW("0.0f");
-						}
-						else
-						{
-							Printer << Argument.ParameterType << BeginArgs;
-							for (int32 ComponentIndex = 0; ComponentIndex < PrimitiveType->NumComponents(); ++ComponentIndex)
-							{
-								Printer << ListSeparator << TEXTVIEW("0.0f");
-							}
-							Printer << EndArgs;
-						}
-					}
-					else
-					{
-						UE_MIR_UNREACHABLE();
-					}
-					break;
-			}
+			Printer << TEXTVIEW("SubstrateHorizontalMixingParameterBlending");
+			Printer << BeginArgs;
+			Printer << LowerValue(SubstrateParameter->Background) << TEXTVIEW(",");
+			Printer << LowerValue(SubstrateParameter->Foreground) << TEXTVIEW(",");
+			Printer << LowerValue(SubstrateParameter->Mix) << TEXTVIEW(",");
+			Printer << LowerValue(SubstrateParameter->NormalMix) << TEXTVIEW(",");
+			Printer << LowerValue(SubstrateParameter->SharedLocalBasisIndexMacro) << TEXTVIEW(",");
+			Printer << GetSharedLocalBasesTypes(SubstrateParameter->bFullySimplifiedInstruction) << TEXTVIEW(",");
+			Printer << LowerValue(SubstrateParameter->BackgroundNoV) << TEXTVIEW(",");
+			Printer << LowerValue(SubstrateParameter->ForegroundNoV);
+			Printer << EndArgs;
+		}
+		else
+		{
+			Printer << GetSubstrateTree(SubstrateParameter->bFullySimplifiedInstruction) << TEXTVIEW(".SubstrateHorizontalMixing");
+			Printer << BeginArgs;
+			Printer << LowerValue(SubstrateParameter->Background) << TEXTVIEW(",");
+			Printer << LowerValue(SubstrateParameter->Foreground) << TEXTVIEW(",");
+			Printer << LowerValue(SubstrateParameter->Mix) << TEXTVIEW(",");
+			Printer << LowerValue(SubstrateParameter->OperatorIndex) << TEXTVIEW(",");
+			Printer << LowerValue(SubstrateParameter->MaxDistanceFromLeaves);
+			Printer << EndArgs;
+		}
+	}
+
+	void LowerSubstrateVerticalLayering(const MIR::FSubstrateVerticalLayering* SubstrateParameter)
+	{
+		if (SubstrateParameter->bParameterBlendingEnabled)
+		{
+			Printer << TEXTVIEW("SubstrateVerticalLayeringParameterBlending");
+			Printer << BeginArgs;
+			Printer << LowerValue(SubstrateParameter->Top) << TEXTVIEW(",");
+			Printer << LowerValue(SubstrateParameter->Base) << TEXTVIEW(",");
+			Printer << LowerValue(SubstrateParameter->SharedLocalBasisIndexMacro) << TEXTVIEW(",");
+			Printer << LowerValue(SubstrateParameter->TopNoV) << TEXTVIEW(",");
+			Printer << LowerValue(SubstrateParameter->BaseNoV);
+			Printer << EndArgs;
 		}
+		else
+		{
+			Printer << GetSubstrateTree(SubstrateParameter->bFullySimplifiedInstruction) << TEXTVIEW(".SubstrateVerticalLayering");
+			Printer << BeginArgs;
+			Printer << LowerValue(SubstrateParameter->Top) << TEXTVIEW(",");
+			Printer << LowerValue(SubstrateParameter->Base) << TEXTVIEW(",");
+			Printer << LowerValue(SubstrateParameter->OperatorIndex) << TEXTVIEW(",");
+			Printer << LowerValue(SubstrateParameter->MaxDistanceFromLeaves);
+			Printer << EndArgs;
+		}
+	}
 
-		Printer.Tabs--;
-		Printer << NewLine << EndArgs << ListSeparator;
-		Printer << NewLine << TEXTVIEW("0, 0, 0, 1");
-		Printer.Tabs--;
-		Printer << NewLine << EndArgs;
+	void LowerSubstratePromoteToOperator(const MIR::FSubstratePromoteToOperator* SubstrateParameter)
+	{
+		Printer << GetSubstrateTree(SubstrateParameter->bFullySimplifiedInstruction) << TEXTVIEW(".PromoteParameterBlendedBSDFToOperator");
+		Printer << BeginArgs;
+		Printer << LowerValue(SubstrateParameter->SubstrateDataInput) << TEXTVIEW(",");
+		Printer << LowerValue(SubstrateParameter->OperatorIndex) << TEXTVIEW(",");
+		Printer << LowerValue(SubstrateParameter->BSDFIndex) << TEXTVIEW(",");
+		Printer << LowerValue(SubstrateParameter->LayerDepth) << TEXTVIEW(",");
+		Printer << LowerValue(SubstrateParameter->bIsBottom);
+		Printer << EndArgs;
 	}
 
 	void EmitPreshaderBufferReadDoubleVector(const MIR::FPrimitive& PrimitiveType, uint32 GlobalComponentOffset)
@@ -2140,15 +2203,18 @@ struct FPrivate : FMaterialIRToHLSLTranslation
 				Printer << EndOfStatement << CloseBlock << NewLine;
 				
 				// Print the LWC get function.
-				MIR::FType ReturnType = Instr->Arg->Type.GetPrimitive().ToScalarKind(MIR::EScalarKind::Double);
-				Printer << ReturnType << TEXTVIEW(" Get") << Instr->Name << TEXTVIEW("_LWC") << Signature
-					<< OpenBlock << TEXTVIEW("return ");
-				
-				Instr->Arg->Type.IsDouble()
-					? Printer << TEXTVIEW("Parameters.") << Instr->Name
-					: Printer << TEXTVIEW("WSPromote(Parameters.") << Instr->Name << TEXT(')');
+				if (!Instr->Arg->Type.IsSubstrateData())
+				{
+					MIR::FType ReturnType = Instr->Arg->Type.GetPrimitive().ToScalarKind(MIR::EScalarKind::Double);
+					Printer << ReturnType << TEXTVIEW(" Get") << Instr->Name << TEXTVIEW("_LWC") << Signature
+						<< OpenBlock << TEXTVIEW("return ");
 				
-				Printer << EndOfStatement << CloseBlock << NewLine;
+					Instr->Arg->Type.IsDouble()
+						? Printer << TEXTVIEW("Parameters.") << Instr->Name
+						: Printer << TEXTVIEW("WSPromote(Parameters.") << Instr->Name << TEXT(')');
+					
+					Printer << EndOfStatement << CloseBlock << NewLine;
+				}
 			}
 		}
 	}
@@ -2221,6 +2287,8 @@ struct FPrivate : FMaterialIRToHLSLTranslation
 	{
 		const FMaterialIRModule::FStatistics ModuleStatistics = Module->GetStatistics();
 
+		const FMaterialIRModule::FSubstrateData& SubstrateData = Module->GetSubstrateData();
+
 		// material_declarations
 		{
 			Printer = {};
@@ -2232,11 +2300,6 @@ struct FPrivate : FMaterialIRToHLSLTranslation
 			}
 			Printer << CloseBlock << TEXT(";\n");
 
-			// temporary placeholder support for Substrate
-			Printer << TEXTVIEW("#if SUBSTRATE_USE_FULLYSIMPLIFIED_MATERIAL == 1") << NewLine;
-			Printer << GetHLSLTypeString(FMaterialAttributeDefinitionMap::GetValueType(MP_FrontMaterial)) << TEXTVIEW(" FullySimplifiedFrontMaterial") << EndOfStatement;
-			Printer << TEXTVIEW("#endif") << NewLine;
-
 			OutputShaderTemplateParams.Add(TEXT("material_declarations"), MoveTemp(Printer.Buffer));
 		}
 
@@ -2254,6 +2317,10 @@ struct FPrivate : FMaterialIRToHLSLTranslation
 					Printer << MIR::FType::FromMaterialValueType(PropertyAggregate->ToMaterialValueType()) << TEXT(' ') << *PropertyAggregate->Name.ToString() << EndOfStatement;
 				}
 			}
+
+			// Add Substrate pixel members declaration code
+			Printer.Buffer.Append(SubstrateData.PixelMembersDeclarationHLSL);
+
 			OutputShaderTemplateParams.Add(TEXT("pixel_material_inputs"), MoveTemp(Printer.Buffer));
 		}
 		
@@ -2296,23 +2363,9 @@ struct FPrivate : FMaterialIRToHLSLTranslation
 	
 		if (Substrate::IsSubstrateEnabled())
 		{
-			// Adde default Substrate functions
-			OutputShaderTemplateParams.Add(TEXT("uniform_material_expressions"), TEXT(
-				"// Substrate: HiddenMaterialAssetConversion\n"
-				"#if TEMPLATE_USES_SUBSTRATE\n"
-				"void FSubstratePixelHeader::PreUpdateAllBSDFWithBottomUpOperatorVisit(float3 V) {}\n"
-				"void FSubstratePixelHeader::UpdateAllBSDFsOperatorCoverageTransmittance(FSubstrateIntegrationSettings Settings, float3 V)\n"
-				"{\n"
-				"#if SUBSTRATE_COMPILER_SUPPORTS_STRUCT_FORWARD_DECLARATION\n"
-				"\tSubstrateTree.UpdateSingleBSDFOperatorCoverageTransmittance(this, 0, Settings, V);\n"
-				"#else\n"
-				"\tUpdateSingleBSDFOperatorCoverageTransmittance(SubstrateTree, this, 0, Settings, V);\n"
-				"#endif\n"
-				"}\n"
-				"void FSubstratePixelHeader::UpdateAllOperatorsCoverageTransmittance() {}\n"
-				"void FSubstratePixelHeader::UpdateAllBSDFWithBottomUpOperatorVisit() {}\n"
-				"#endif // TEMPLATE_USES_SUBSTRATE\n"
-				"\n"));
+			// No one else write to uniform_material_expressions so no need to search before append.
+			OutputShaderTemplateParams.Add(TEXT("uniform_material_expressions"), SubstrateData.TreeHLSL);
+			// SUBSTRATE_TODO also add in there DescriptionStringCommentForDebug
 		}  
 	}
 	
@@ -2393,6 +2446,7 @@ struct FPrivate : FMaterialIRToHLSLTranslation
 
 			for (int32 i = 0; i < MSM_NUM; ++i)
 			{
+				// SUBSTRATE_TODO handle that using EnumHasAnyFlags(SubstrateTranslatorData.SubstrateMaterialBsdfFeatures, ESubstrateBsdfFeature::SSS)
 				EMaterialShadingModel Model = (EMaterialShadingModel)i;
 				if (Model == MSM_Strata || !ShadingModels.HasShadingModel(Model))
 				{
@@ -2483,18 +2537,9 @@ struct FPrivate : FMaterialIRToHLSLTranslation
 			}
 		}
 
-		const bool bIsSubstrate = Substrate::IsSubstrateEnabled();
-
-		OutEnvironment.SetDefine(TEXT("MATERIAL_IS_SUBSTRATE"), bIsSubstrate);
 		OutEnvironment.SetDefine(TEXT("DUAL_SOURCE_COLOR_BLENDING_ENABLED"), false);
 		OutEnvironment.SetDefine(TEXT("TEXTURE_SAMPLE_DEBUG"), false);
 
-		if (bIsSubstrate)
-		{
-			OutEnvironment.SetDefine(TEXT("SUBSTRATE_USE_FULLYSIMPLIFIED_MATERIAL"), false);
-			OutEnvironment.SetDefine(TEXT("SUBSTRATE_CLAMPED_CLOSURE_COUNT"), 1);
-		}
-
 		for (int32 VTStackIndex = 0; VTStackIndex < CompilationOutput.UniformExpressionSet.GetVTStacks().Num(); ++VTStackIndex)
 		{
 			const FMaterialVirtualTextureStack& VTStack = CompilationOutput.UniformExpressionSet.GetVTStack(VTStackIndex);
diff --git a/Engine/Source/Runtime/Engine/Public/MaterialCompiler.h b/Engine/Source/Runtime/Engine/Public/MaterialCompiler.h
--- a/Engine/Source/Runtime/Engine/Public/MaterialCompiler.h
+++ b/Engine/Source/Runtime/Engine/Public/MaterialCompiler.h
@@ -18,6 +18,43 @@
 #include "Materials/MaterialExpressionViewProperty.h"
 #include "Materials/MaterialExpressionRuntimeVirtualTextureSample.h"
 
+// This interface is needed so that the SubstrateTranslatorData can be passed around in material expression
+struct FSubstrateTranslatorDataInterface
+{
+	virtual void Errorf(const FString& Format) = 0;
+
+	virtual FMaterialShadingModelField GetMaterialShadingModels() = 0;
+
+
+	virtual FGuid SubstrateTreeStackGetPathUniqueId() = 0;
+	virtual FGuid SubstrateTreeStackGetParentPathUniqueId() = 0;
+	virtual bool GetSubstrateTreeOutOfStackDepthOccurred() = 0;
+
+	virtual int32 SubstrateThicknessStackGetThicknessIndex() = 0;
+	virtual int32 SubstrateThicknessStackPush(class UMaterialExpression* Expression, FScalarMaterialInput* Input) = 0;
+	virtual int32 SubstrateThicknessStackPush(class UMaterialExpression* Expression, FExpressionInput* Input) = 0;
+	virtual void SubstrateThicknessStackPop() = 0;
+	struct FSubstrateThicknessExpression
+	{
+		FExpressionInput* ExpressionInput = nullptr;
+		FScalarMaterialInput* MaterialInput = nullptr;
+	};
+	virtual FSubstrateThicknessExpression* SubstrateThicknessStackGetExpression(int32 Index) = 0;
+
+	virtual FGuid SubstrateTreeStackPush(UMaterialExpression* Expression, uint32 InputIndex) = 0;
+	virtual void SubstrateTreeStackPop() = 0;
+
+	virtual FSubstrateOperator& SubstrateCompilationRegisterOperator(int32 OperatorType, FGuid SubstrateExpressionGuid, FGuid ChildMaterialExpressionGuid, class UMaterialExpression* Parent, FGuid SubstrateParentExpressionGuid, bool bUseParameterBlending = false) = 0;
+	virtual FSubstrateOperator& SubstrateGetDefaultBSDFOperator(FGuid MaterialExpressionGuid, class UMaterialExpression* Parent) = 0;
+	virtual FSubstrateOperator& SubstrateCompilationGetOperator(FGuid SubstrateExpressionGuid) = 0;
+	virtual FSubstrateOperator* SubstrateCompilationGetOperatorFromIndex(int32 OperatorIndex) = 0;
+
+	virtual class FMaterialCompiler* GetCompilerLegacy() = 0;
+	virtual class MIR::FEmitter* GetCompilerNew() = 0;
+
+	virtual void PushFunction(class FMaterialFunctionCompileState* FunctionState) = 0;
+	virtual void PopFunction() = 0;
+};
 
 namespace UE
 {
@@ -766,6 +803,8 @@ public:
 	virtual int32 SubstrateThicknessStackPush(UMaterialExpression* Expression, FExpressionInput* Input) = 0;
 	virtual void SubstrateThicknessStackPop() = 0;
 
+	virtual struct FSubstrateTranslatorDataInterface& GetSubstrateTranslatorData() = 0;
+
 	/**
 	 * Register an operator of the tree representation the Substrate material and its topology.
 	 */
@@ -800,15 +839,6 @@ public:
 	bool IsValidating() const { return bIsValidating; }
 
 protected:
-	FString GetSubstrateSharedLocalBasisIndexMacroInner(const FSubstrateRegisteredSharedLocalBasis& SharedLocalBasis, uint32 Mode)
-	{
-#if WITH_EDITOR
-		return FString::Printf(TEXT("SHAREDLOCALBASIS_INDEX_%u_%u"), SharedLocalBasis.GraphSharedLocalBasisIndex, Mode);
-#else
-		return TEXT("");
-#endif
-	}
-
 
 private:
 	ESubstrateMaterialExport SubstrateMaterialExport = SME_None;
@@ -1649,6 +1679,10 @@ public:
 	{
 		Compiler->SubstrateThicknessStackPop();
 	}
+	virtual struct FSubstrateTranslatorDataInterface& GetSubstrateTranslatorData() override
+	{
+		return Compiler->GetSubstrateTranslatorData();
+	}
 
 	virtual FSubstrateOperator& SubstrateCompilationRegisterOperator(int32 OperatorType, FGuid SubstrateExpressionGuid, FGuid ChildMaterialExpressionGuid, UMaterialExpression* Parent, FGuid SubstrateParentExpressionGuid, bool bUseParameterBlending = false) override
 	{
diff --git a/Engine/Source/Runtime/Engine/Public/Materials/MaterialExpression.h b/Engine/Source/Runtime/Engine/Public/Materials/MaterialExpression.h
--- a/Engine/Source/Runtime/Engine/Public/Materials/MaterialExpression.h
+++ b/Engine/Source/Runtime/Engine/Public/Materials/MaterialExpression.h
@@ -9,6 +9,7 @@
 #include "Misc/Guid.h"
 #include "MaterialExpressionIO.h"
 #include "MaterialValueType.h"
+//#include "MaterialCompiler.h"
 
 #include "MaterialExpression.generated.h"
 
@@ -38,6 +39,8 @@ namespace MIR
 	struct FValue;
 }
 
+struct FSubstrateTranslatorDataInterface;
+
 //@warning: FExpressionInput is mirrored in MaterialShared.h and manually "subclassed" in Material.h (FMaterialInput)
 #if !CPP      //noexport struct
 USTRUCT(noexport)
@@ -467,12 +470,15 @@ class UMaterialExpression : public UObject
 	 */
 	virtual void GatherSubstrateMaterialInfo(FSubstrateMaterialInfo& SubstrateMaterialInfo, int32 OutputIndex) { }
 
+	UE_DEPRECATED(5.8, "Please use SubstrateGenerateMaterialTopologyTree taking the FSubstrateTranslatorDataInterface parameter.")
+	ENGINE_API virtual FSubstrateOperator* SubstrateGenerateMaterialTopologyTree(class FMaterialCompiler* Compiler, class UMaterialExpression* Parent, int32 OutputIndex);
+
 	/**
-	 * A starta material is a tree with FrontMateiral being its root and BSDF being leaves, with operators in the middle.
-	 * This recursively parse nodes outputing Substrate material in order to gather the maximum distance to any leaves. 
-	 * This is used to drive the bottom up order processing of those nodes.
+	 * A Substrate material is a tree with FrontMaterial being its root and BSDF being leaves, with operators in the middle.
+	 * This recursively parse nodes outputing Substrate material in order to gather the maximum distance to any leaves.
+	 * This is used to drive the bottom up order processing of those nodes. 
 	 */
-	ENGINE_API virtual FSubstrateOperator* SubstrateGenerateMaterialTopologyTree(class FMaterialCompiler* Compiler, class UMaterialExpression* Parent, int32 OutputIndex);
+	ENGINE_API virtual FSubstrateOperator* SubstrateGenerateMaterialTopologyTree(struct FSubstrateTranslatorDataInterface& SubstrateTranslatorData, class UMaterialExpression* Parent, int32 OutputIndex);
 
 	/**
 	 * If true, discards the output index when caching this expression which allows more cases to re-use the output instead of adding a separate instruction
diff --git a/Engine/Source/Runtime/Engine/Public/Materials/MaterialExpressionBlendMaterialAttributes.h b/Engine/Source/Runtime/Engine/Public/Materials/MaterialExpressionBlendMaterialAttributes.h
--- a/Engine/Source/Runtime/Engine/Public/Materials/MaterialExpressionBlendMaterialAttributes.h
+++ b/Engine/Source/Runtime/Engine/Public/Materials/MaterialExpressionBlendMaterialAttributes.h
@@ -63,7 +63,7 @@ class UMaterialExpressionBlendMaterialAttributes : public UMaterialExpression
 	virtual bool IsResultMaterialAttributes(int32 OutputIndex) override {return true;}
 	virtual bool IsResultSubstrateMaterial(int32 OutputIndex) override;
 	virtual void GatherSubstrateMaterialInfo(FSubstrateMaterialInfo& SubstrateMaterialInfo, int32 OutputIndex) override;
-	virtual FSubstrateOperator* SubstrateGenerateMaterialTopologyTree(class FMaterialCompiler* Compiler, class UMaterialExpression* Parent, int32 OutputIndex) override;
+	virtual FSubstrateOperator* SubstrateGenerateMaterialTopologyTree(struct FSubstrateTranslatorDataInterface& SubstrateTranslatorData, class UMaterialExpression* Parent, int32 OutputIndex) override;
 	virtual EMaterialValueType GetInputValueType(int32 InputIndex) override {return InputIndex == 2 ? MCT_Float1 : MCT_MaterialAttributes;}
 #endif
 	//~ End UMaterialExpression Interface
diff --git a/Engine/Source/Runtime/Engine/Public/Materials/MaterialExpressionFunctionInput.h b/Engine/Source/Runtime/Engine/Public/Materials/MaterialExpressionFunctionInput.h
--- a/Engine/Source/Runtime/Engine/Public/Materials/MaterialExpressionFunctionInput.h
+++ b/Engine/Source/Runtime/Engine/Public/Materials/MaterialExpressionFunctionInput.h
@@ -121,7 +121,7 @@ class UMaterialExpressionFunctionInput : public UMaterialExpression
 
 	virtual bool IsResultSubstrateMaterial(int32 OutputIndex) override;
 	virtual void GatherSubstrateMaterialInfo(FSubstrateMaterialInfo& SubstrateMaterialInfo, int32 OutputIndex) override;
-	virtual FSubstrateOperator* SubstrateGenerateMaterialTopologyTree(class FMaterialCompiler* Compiler, class UMaterialExpression* Parent, int32 OutputIndex) override;
+	virtual FSubstrateOperator* SubstrateGenerateMaterialTopologyTree(struct FSubstrateTranslatorDataInterface& SubstrateTranslatorData, class UMaterialExpression* Parent, int32 OutputIndex) override;
 	virtual bool IsResultMaterialAttributes(int32 OutputIndex) override;
 	virtual EMaterialValueType GetInputValueType(int32 InputIndex) override;
 	virtual EMaterialValueType GetOutputValueType(int32 OutputIndex) override;
diff --git a/Engine/Source/Runtime/Engine/Public/Materials/MaterialExpressionMaterialFunctionCall.h b/Engine/Source/Runtime/Engine/Public/Materials/MaterialExpressionMaterialFunctionCall.h
--- a/Engine/Source/Runtime/Engine/Public/Materials/MaterialExpressionMaterialFunctionCall.h
+++ b/Engine/Source/Runtime/Engine/Public/Materials/MaterialExpressionMaterialFunctionCall.h
@@ -134,7 +134,7 @@ class UMaterialExpressionMaterialFunctionCall : public UMaterialExpression
 	ENGINE_API virtual bool IsResultMaterialAttributes(int32 OutputIndex) override;
 	ENGINE_API virtual bool IsResultSubstrateMaterial(int32 OutputIndex) override;
 	ENGINE_API virtual void GatherSubstrateMaterialInfo(FSubstrateMaterialInfo& SubstrateMaterialInfo, int32 OutputIndex) override;
-	ENGINE_API virtual FSubstrateOperator* SubstrateGenerateMaterialTopologyTree(class FMaterialCompiler* Compiler, class UMaterialExpression* Parent, int32 OutputIndex) override;
+	ENGINE_API virtual FSubstrateOperator* SubstrateGenerateMaterialTopologyTree(struct FSubstrateTranslatorDataInterface& SubstrateTranslatorData, class UMaterialExpression* Parent, int32 OutputIndex) override;
 	ENGINE_API virtual EMaterialValueType GetInputValueType(int32 InputIndex) override;
 	ENGINE_API virtual EMaterialValueType GetOutputValueType(int32 OutputIndex) override;
 	//~ End UMaterialExpression Interface
diff --git a/Engine/Source/Runtime/Engine/Public/Materials/MaterialExpressionRerouteBase.h b/Engine/Source/Runtime/Engine/Public/Materials/MaterialExpressionRerouteBase.h
--- a/Engine/Source/Runtime/Engine/Public/Materials/MaterialExpressionRerouteBase.h
+++ b/Engine/Source/Runtime/Engine/Public/Materials/MaterialExpressionRerouteBase.h
@@ -31,7 +31,7 @@ public:
 	ENGINE_API virtual bool IsResultMaterialAttributes(int32 OutputIndex) override;
 	ENGINE_API virtual bool IsResultSubstrateMaterial(int32 OutputIndex) override;
 	ENGINE_API void GatherSubstrateMaterialInfo(FSubstrateMaterialInfo& SubstrateMaterialInfo, int32 OutputIndex) override;
-	ENGINE_API FSubstrateOperator* SubstrateGenerateMaterialTopologyTree(class FMaterialCompiler* Compiler, class UMaterialExpression* Parent, int32 OutputIndex) override;
+	ENGINE_API FSubstrateOperator* SubstrateGenerateMaterialTopologyTree(struct FSubstrateTranslatorDataInterface& SubstrateTranslatorData, class UMaterialExpression* Parent, int32 OutputIndex) override;
 #endif
 	//~ End UMaterialExpression Interface
 
diff --git a/Engine/Source/Runtime/Engine/Public/Materials/MaterialExpressionSetMaterialAttributes.h b/Engine/Source/Runtime/Engine/Public/Materials/MaterialExpressionSetMaterialAttributes.h
--- a/Engine/Source/Runtime/Engine/Public/Materials/MaterialExpressionSetMaterialAttributes.h
+++ b/Engine/Source/Runtime/Engine/Public/Materials/MaterialExpressionSetMaterialAttributes.h
@@ -50,7 +50,7 @@ class UMaterialExpressionSetMaterialAttributes : public UMaterialExpression
 	bool GetSubstrateMaterialInputIndex(int32 OutputIndex, int32& InputIndex);
 	virtual bool IsResultSubstrateMaterial(int32 OutputIndex) override;
 	virtual void GatherSubstrateMaterialInfo(FSubstrateMaterialInfo& SubstrateMaterialInfo, int32 OutputIndex) override;
-	virtual FSubstrateOperator* SubstrateGenerateMaterialTopologyTree(FMaterialCompiler* Compiler, UMaterialExpression* Parent, int32 OutputIndex) override;
+	virtual FSubstrateOperator* SubstrateGenerateMaterialTopologyTree(struct FSubstrateTranslatorDataInterface& SubstrateTranslatorData, UMaterialExpression* Parent, int32 OutputIndex) override;
 	virtual void GetExpressionToolTip(TArray<FString>& OutToolTip) override;
 
 	uint64 GetConnectedInputs() const;
diff --git a/Engine/Source/Runtime/Engine/Public/Materials/MaterialExpressionSubstrate.h b/Engine/Source/Runtime/Engine/Public/Materials/MaterialExpressionSubstrate.h
--- a/Engine/Source/Runtime/Engine/Public/Materials/MaterialExpressionSubstrate.h
+++ b/Engine/Source/Runtime/Engine/Public/Materials/MaterialExpressionSubstrate.h
@@ -6,12 +6,18 @@
 #include "Engine/EngineTypes.h"
 #include "Materials/MaterialExpression.h"
 #include "SubstrateDefinitions.h"
+#include "MaterialCompiler.h"
 #include "MaterialExpressionSubstrate.generated.h"
 
 class FMaterialCompiler;
 struct FSubstrateMaterialComplexity;
 struct FSubstrateOperator;
 
+namespace MIR
+{
+	struct FValueRef;
+};
+
 UENUM()
 enum EMaterialSubSurfaceType : int
 {
@@ -218,14 +224,16 @@ class UMaterialExpressionSubstrateShadingModels : public UMaterialExpressionSubs
 	virtual EMaterialValueType GetInputValueType(int32 InputIndex) override;
 	virtual bool IsResultSubstrateMaterial(int32 OutputIndex) override;
 	virtual void GatherSubstrateMaterialInfo(FSubstrateMaterialInfo& SubstrateMaterialInfo, int32 OutputIndex) override;
-	virtual FSubstrateOperator* SubstrateGenerateMaterialTopologyTree(class FMaterialCompiler* Compiler, class UMaterialExpression* Parent, int32 OutputIndex) override;
+	virtual FSubstrateOperator* SubstrateGenerateMaterialTopologyTree(struct FSubstrateTranslatorDataInterface& SubstrateTranslatorData, class UMaterialExpression* Parent, int32 OutputIndex) override;
 	virtual FName GetInputName(int32 InputIndex) const override;
 	virtual void GetConnectorToolTip(int32 InputIndex, int32 OutputIndex, TArray<FString>& OutToolTip) override;
 	virtual void PostEditChangeProperty(FPropertyChangedEvent& PropertyChangedEvent) override;
+	virtual void Build(MIR::FEmitter& Emitter) override;
 
 	bool HasSSS() const;
 	bool HasAnisotropy() const;
 
+	UE_DEPRECATED(5.8, "Please, SubstrateGenerateMaterialTopologyTreeCommon taking the FSubstrateTranslatorData parameter.")
 	static FSubstrateOperator* SubstrateGenerateMaterialTopologyTreeCommon(
 		class FMaterialCompiler* Compiler, struct FGuid ThisExpressionGuid, class UMaterialExpression* Parent, int32 OutputIndex,
 		const FExpressionInput& EmissiveColor,
@@ -234,6 +242,14 @@ class UMaterialExpressionSubstrateShadingModels : public UMaterialExpressionSubs
 		const FExpressionInput& CustomTangent,
 		const FExpressionInput& ShadingModel);
 
+	static FSubstrateOperator* SubstrateGenerateMaterialTopologyTreeCommon(
+		struct FSubstrateTranslatorDataInterface& SubstrateTranslatorData, struct FGuid ThisExpressionGuid, class UMaterialExpression* Parent, int32 OutputIndex,
+		const FExpressionInput& EmissiveColor,
+		const FExpressionInput& Anisotropy,
+		const FExpressionInput& ClearCoatNormal,
+		const FExpressionInput& CustomTangent,
+		const FExpressionInput& ShadingModel);
+
 	static int32 CompileCommon(class FMaterialCompiler* Compiler,
 		FExpressionInput& BaseColor, FExpressionInput& Specular, FExpressionInput& Metallic, FExpressionInput& Roughness, FExpressionInput& EmissiveColor,
 		FExpressionInput& Opacity, FExpressionInput& SubSurfaceColor, FExpressionInput& ClearCoat, FExpressionInput& ClearCoatRoughness,
@@ -245,6 +261,18 @@ class UMaterialExpressionSubstrateShadingModels : public UMaterialExpressionSubs
 		FExpressionInput& ClearCoatNormal, FExpressionInput& CustomTangent,
 		const bool bHasSSS, USubsurfaceProfile* SSSProfile,
 		const class UMaterialEditorOnlyData* EditorOnlyData = nullptr /*Used to provide default values from the root node pins*/);
+
+	static MIR::FValueRef BuildCommon(MIR::FEmitter& Em, UMaterial* Material,
+		FExpressionInput& BaseColor, FExpressionInput& Specular, FExpressionInput& Metallic, FExpressionInput& Roughness, FExpressionInput& EmissiveColor,
+		FExpressionInput& Opacity, FExpressionInput& SubSurfaceColor, FExpressionInput& ClearCoat, FExpressionInput& ClearCoatRoughness,
+		FExpressionInput& ShadingModel, TEnumAsByte<enum EMaterialShadingModel> ShadingModelOverride,
+		FExpressionInput& TransmittanceColor, FExpressionInput& ThinTranslucentSurfaceCoverage,
+		FExpressionInput& WaterScatteringCoefficients, FExpressionInput& WaterAbsorptionCoefficients, FExpressionInput& WaterPhaseG, FExpressionInput& ColorScaleBehindWater,
+		const bool bHasAnisotropy, FExpressionInput& Anisotropy,
+		FExpressionInput& Normal, FExpressionInput& Tangent,
+		FExpressionInput& ClearCoatNormal, FExpressionInput& CustomTangent,
+		const bool bHasSSS, USubsurfaceProfile* SSSProfile,
+		const class UMaterialEditorOnlyData* EditorOnlyData = nullptr /*Used to provide default values from the root node pins*/);
 #endif
 	//~ End UMaterialExpression Interface
 };
@@ -391,11 +419,12 @@ class UMaterialExpressionSubstrateSlabBSDF : public UMaterialExpressionSubstrate
 	virtual EMaterialValueType GetInputValueType(int32 InputIndex) override;
 	virtual bool IsResultSubstrateMaterial(int32 OutputIndex) override;
 	virtual void GatherSubstrateMaterialInfo(FSubstrateMaterialInfo& SubstrateMaterialInfo, int32 OutputIndex) override;
-	virtual FSubstrateOperator* SubstrateGenerateMaterialTopologyTree(class FMaterialCompiler* Compiler, class UMaterialExpression* Parent, int32 OutputIndex) override;
+	virtual FSubstrateOperator* SubstrateGenerateMaterialTopologyTree(struct FSubstrateTranslatorDataInterface& SubstrateTranslatorData, class UMaterialExpression* Parent, int32 OutputIndex) override;
 	virtual FName GetInputName(int32 InputIndex) const override;
 	virtual void GetConnectorToolTip(int32 InputIndex, int32 OutputIndex, TArray<FString>& OutToolTip) override;
 	virtual void GetExpressionToolTip(TArray<FString>& OutToolTip) override;
 	virtual void PostEditChangeProperty(FPropertyChangedEvent& PropertyChangedEvent) override;
+	virtual void Build(MIR::FEmitter& Emitter) override;
 
 	bool HasEdgeColor() const;
 	bool HasFuzz() const;
@@ -474,7 +503,7 @@ class UMaterialExpressionSubstrateSimpleClearCoatBSDF : public UMaterialExpressi
 	virtual EMaterialValueType GetInputValueType(int32 InputIndex) override;
 	virtual bool IsResultSubstrateMaterial(int32 OutputIndex) override;
 	virtual void GatherSubstrateMaterialInfo(FSubstrateMaterialInfo& SubstrateMaterialInfo, int32 OutputIndex) override;
-	virtual FSubstrateOperator* SubstrateGenerateMaterialTopologyTree(class FMaterialCompiler* Compiler, class UMaterialExpression* Parent, int32 OutputIndex) override;
+	virtual FSubstrateOperator* SubstrateGenerateMaterialTopologyTree(struct FSubstrateTranslatorDataInterface& SubstrateTranslatorData, class UMaterialExpression* Parent, int32 OutputIndex) override;
 	virtual FName GetInputName(int32 InputIndex) const override;
 	virtual void GetExpressionToolTip(TArray<FString>& OutToolTip) override;
 #endif
@@ -522,13 +551,19 @@ class UMaterialExpressionSubstrateVolumetricFogCloudBSDF : public UMaterialExpre
 	virtual EMaterialValueType GetInputValueType(int32 InputIndex) override;
 	virtual bool IsResultSubstrateMaterial(int32 OutputIndex) override;
 	virtual void GatherSubstrateMaterialInfo(FSubstrateMaterialInfo& SubstrateMaterialInfo, int32 OutputIndex) override;
-	virtual FSubstrateOperator* SubstrateGenerateMaterialTopologyTree(class FMaterialCompiler* Compiler, class UMaterialExpression* Parent, int32 OutputIndex) override;
+	virtual FSubstrateOperator* SubstrateGenerateMaterialTopologyTree(struct FSubstrateTranslatorDataInterface& SubstrateTranslatorData, class UMaterialExpression* Parent, int32 OutputIndex) override;
 
+	UE_DEPRECATED(5.8, "Please, SubstrateGenerateMaterialTopologyTreeCommon taking the FSubstrateTranslatorData parameter.")
 	static FSubstrateOperator* SubstrateGenerateMaterialTopologyTreeCommon(
 		class FMaterialCompiler* Compiler, struct FGuid ThisExpressionGuid, class UMaterialExpression* Parent, int32 OutputIndex,
 		const FExpressionInput& EmissiveColor,
 		const FExpressionInput& AmbientOcclusion);
 
+	static FSubstrateOperator* SubstrateGenerateMaterialTopologyTreeCommon(
+		struct FSubstrateTranslatorDataInterface& SubstrateTranslatorData, struct FGuid ThisExpressionGuid, class UMaterialExpression* Parent, int32 OutputIndex,
+		const FExpressionInput& EmissiveColor,
+		const FExpressionInput& AmbientOcclusion);
+
 	static int32 CompileCommon(class FMaterialCompiler* Compiler,
 		FExpressionInput& Albedo, FExpressionInput& Extinction, FExpressionInput& EmissiveColor, FExpressionInput& AmbientOcclusion,
 		bool bEmissiveOnly,
@@ -568,7 +603,7 @@ class UMaterialExpressionSubstrateUnlitBSDF : public UMaterialExpressionSubstrat
 	virtual EMaterialValueType GetInputValueType(int32 InputIndex) override;
 	virtual bool IsResultSubstrateMaterial(int32 OutputIndex) override;
 	virtual void GatherSubstrateMaterialInfo(FSubstrateMaterialInfo& SubstrateMaterialInfo, int32 OutputIndex) override;
-	virtual FSubstrateOperator* SubstrateGenerateMaterialTopologyTree(class FMaterialCompiler* Compiler, class UMaterialExpression* Parent, int32 OutputIndex) override;
+	virtual FSubstrateOperator* SubstrateGenerateMaterialTopologyTree(struct FSubstrateTranslatorDataInterface& SubstrateTranslatorData, class UMaterialExpression* Parent, int32 OutputIndex) override;
 #endif
 	//~ End UMaterialExpression Interface
 };
@@ -628,7 +663,7 @@ class UMaterialExpressionSubstrateHairBSDF : public UMaterialExpressionSubstrate
 	virtual EMaterialValueType GetInputValueType(int32 InputIndex) override;
 	virtual bool IsResultSubstrateMaterial(int32 OutputIndex) override;
 	virtual void GatherSubstrateMaterialInfo(FSubstrateMaterialInfo& SubstrateMaterialInfo, int32 OutputIndex) override;
-	virtual FSubstrateOperator* SubstrateGenerateMaterialTopologyTree(class FMaterialCompiler* Compiler, class UMaterialExpression* Parent, int32 OutputIndex) override;
+	virtual FSubstrateOperator* SubstrateGenerateMaterialTopologyTree(struct FSubstrateTranslatorDataInterface& SubstrateTranslatorData, class UMaterialExpression* Parent, int32 OutputIndex) override;
 #endif
 	//~ End UMaterialExpression Interface
 };
@@ -698,7 +733,7 @@ class UMaterialExpressionSubstrateEyeBSDF : public UMaterialExpressionSubstrateB
 	virtual EMaterialValueType GetInputValueType(int32 InputIndex) override;
 	virtual bool IsResultSubstrateMaterial(int32 OutputIndex) override;
 	virtual void GatherSubstrateMaterialInfo(FSubstrateMaterialInfo& SubstrateMaterialInfo, int32 OutputIndex) override;
-	virtual FSubstrateOperator* SubstrateGenerateMaterialTopologyTree(class FMaterialCompiler* Compiler, class UMaterialExpression* Parent, int32 OutputIndex) override;
+	virtual FSubstrateOperator* SubstrateGenerateMaterialTopologyTree(struct FSubstrateTranslatorDataInterface& SubstrateTranslatorData, class UMaterialExpression* Parent, int32 OutputIndex) override;
 #endif
 	//~ End UMaterialExpression Interface
 };
@@ -782,7 +817,7 @@ class UMaterialExpressionSubstrateSingleLayerWaterBSDF : public UMaterialExpress
 	virtual EMaterialValueType GetInputValueType(int32 InputIndex) override;
 	virtual bool IsResultSubstrateMaterial(int32 OutputIndex) override;
 	virtual void GatherSubstrateMaterialInfo(FSubstrateMaterialInfo& SubstrateMaterialInfo, int32 OutputIndex) override;
-	virtual FSubstrateOperator* SubstrateGenerateMaterialTopologyTree(class FMaterialCompiler* Compiler, class UMaterialExpression* Parent, int32 OutputIndex) override;
+	virtual FSubstrateOperator* SubstrateGenerateMaterialTopologyTree(struct FSubstrateTranslatorDataInterface& SubstrateTranslatorData, class UMaterialExpression* Parent, int32 OutputIndex) override;
 #endif
 	//~ End UMaterialExpression Interface
 };
@@ -806,11 +841,15 @@ class UMaterialExpressionSubstrateLightFunction : public UMaterialExpressionSubs
 	virtual EMaterialValueType GetInputValueType(int32 InputIndex) override;
 	virtual bool IsResultSubstrateMaterial(int32 OutputIndex) override;
 	virtual void GatherSubstrateMaterialInfo(FSubstrateMaterialInfo& SubstrateMaterialInfo, int32 OutputIndex) override;
-	virtual FSubstrateOperator* SubstrateGenerateMaterialTopologyTree(class FMaterialCompiler* Compiler, class UMaterialExpression* Parent, int32 OutputIndex) override;
+	virtual FSubstrateOperator* SubstrateGenerateMaterialTopologyTree(struct FSubstrateTranslatorDataInterface& SubstrateTranslatorData, class UMaterialExpression* Parent, int32 OutputIndex) override;
 
+	UE_DEPRECATED(5.8, "Please, SubstrateGenerateMaterialTopologyTreeCommon taking the FSubstrateTranslatorData parameter.")
 	static FSubstrateOperator* SubstrateGenerateMaterialTopologyTreeCommon(
 		class FMaterialCompiler* Compiler, struct FGuid ThisExpressionGuid, class UMaterialExpression* Parent, int32 OutputIndex);
 
+	static FSubstrateOperator* SubstrateGenerateMaterialTopologyTreeCommon(
+		struct FSubstrateTranslatorDataInterface& SubstrateTranslatorData, struct FGuid ThisExpressionGuid, class UMaterialExpression* Parent, int32 OutputIndex);
+
 	static int32 CompileCommon(class FMaterialCompiler* Compiler,
 		FExpressionInput& Color,
 		const UMaterialEditorOnlyData* EditorOnlyData = nullptr);
@@ -872,7 +911,7 @@ class UMaterialExpressionSubstrateToonBSDF : public UMaterialExpressionSubstrate
 	virtual EMaterialValueType GetInputValueType(int32 InputIndex) override;
 	virtual bool IsResultSubstrateMaterial(int32 OutputIndex) override;
 	virtual void GatherSubstrateMaterialInfo(FSubstrateMaterialInfo& SubstrateMaterialInfo, int32 OutputIndex) override;
-	virtual FSubstrateOperator* SubstrateGenerateMaterialTopologyTree(class FMaterialCompiler* Compiler, class UMaterialExpression* Parent, int32 OutputIndex) override;
+	virtual FSubstrateOperator* SubstrateGenerateMaterialTopologyTree(struct FSubstrateTranslatorDataInterface& SubstrateTranslatorData, class UMaterialExpression* Parent, int32 OutputIndex) override;
 #endif
 	//~ End UMaterialExpression Interface
 };
@@ -902,11 +941,15 @@ class UMaterialExpressionSubstratePostProcess : public UMaterialExpressionSubstr
 	virtual EMaterialValueType GetInputValueType(int32 InputIndex) override;
 	virtual bool IsResultSubstrateMaterial(int32 OutputIndex) override;
 	virtual void GatherSubstrateMaterialInfo(FSubstrateMaterialInfo& SubstrateMaterialInfo, int32 OutputIndex) override;
-	virtual FSubstrateOperator* SubstrateGenerateMaterialTopologyTree(class FMaterialCompiler* Compiler, class UMaterialExpression* Parent, int32 OutputIndex) override;
+	virtual FSubstrateOperator* SubstrateGenerateMaterialTopologyTree(struct FSubstrateTranslatorDataInterface& SubstrateTranslatorData, class UMaterialExpression* Parent, int32 OutputIndex) override;
 
+	UE_DEPRECATED(5.8, "Please, SubstrateGenerateMaterialTopologyTreeCommon taking the FSubstrateTranslatorData parameter.")
 	static FSubstrateOperator* SubstrateGenerateMaterialTopologyTreeCommon(
 		class FMaterialCompiler* Compiler, struct FGuid ThisExpressionGuid, class UMaterialExpression* Parent, int32 OutputIndex);
 
+	static FSubstrateOperator* SubstrateGenerateMaterialTopologyTreeCommon(
+		struct FSubstrateTranslatorDataInterface& SubstrateTranslatorData, struct FGuid ThisExpressionGuid, class UMaterialExpression* Parent, int32 OutputIndex);
+
 	static int32 CompileCommon(class FMaterialCompiler* Compiler,
 		FExpressionInput& Color, FExpressionInput& Opacity,
 		const UMaterialEditorOnlyData* EditorOnlyData = nullptr);
@@ -939,11 +982,15 @@ class UMaterialExpressionSubstrateUI : public UMaterialExpressionSubstrateBSDF
 	virtual EMaterialValueType GetInputValueType(int32 InputIndex) override;
 	virtual bool IsResultSubstrateMaterial(int32 OutputIndex) override;
 	virtual void GatherSubstrateMaterialInfo(FSubstrateMaterialInfo& SubstrateMaterialInfo, int32 OutputIndex) override;
-	virtual FSubstrateOperator* SubstrateGenerateMaterialTopologyTree(class FMaterialCompiler* Compiler, class UMaterialExpression* Parent, int32 OutputIndex) override;
+	virtual FSubstrateOperator* SubstrateGenerateMaterialTopologyTree(struct FSubstrateTranslatorDataInterface& SubstrateTranslatorData, class UMaterialExpression* Parent, int32 OutputIndex) override;
 
+	UE_DEPRECATED(5.8, "Please, SubstrateGenerateMaterialTopologyTreeCommon taking the FSubstrateTranslatorData parameter.")
 	static FSubstrateOperator* SubstrateGenerateMaterialTopologyTreeCommon(
 		class FMaterialCompiler* Compiler, struct FGuid ThisExpressionGuid, class UMaterialExpression* Parent, int32 OutputIndex);
 
+	static FSubstrateOperator* SubstrateGenerateMaterialTopologyTreeCommon(
+		struct FSubstrateTranslatorDataInterface& SubstrateTranslatorData, struct FGuid ThisExpressionGuid, class UMaterialExpression* Parent, int32 OutputIndex);
+
 	static int32 CompileCommon(class FMaterialCompiler* Compiler,
 		FExpressionInput& Color, FExpressionInput& Opacity,
 		const UMaterialEditorOnlyData* EditorOnlyData = nullptr);
@@ -976,7 +1023,7 @@ class UMaterialExpressionSubstrateConvertToDecal : public UMaterialExpressionSub
 	virtual EMaterialValueType GetInputValueType(int32 InputIndex) override;
 	virtual bool IsResultSubstrateMaterial(int32 OutputIndex) override;
 	virtual void GatherSubstrateMaterialInfo(FSubstrateMaterialInfo& SubstrateMaterialInfo, int32 OutputIndex) override;
-	virtual FSubstrateOperator* SubstrateGenerateMaterialTopologyTree(class FMaterialCompiler* Compiler, class UMaterialExpression* Parent, int32 OutputIndex) override;
+	virtual FSubstrateOperator* SubstrateGenerateMaterialTopologyTree(struct FSubstrateTranslatorDataInterface& SubstrateTranslatorData, class UMaterialExpression* Parent, int32 OutputIndex) override;
 #endif
 	//~ End UMaterialExpression Interface
 };
@@ -1033,17 +1080,22 @@ class UMaterialExpressionSubstrateConvertMaterialAttributes : public UMaterialEx
 	virtual bool IsResultSubstrateMaterial(int32 OutputIndex) override;
 	virtual bool IsResultMaterialAttributes(int32 OutputIndex) override;
 	virtual void GatherSubstrateMaterialInfo(FSubstrateMaterialInfo& SubstrateMaterialInfo, int32 OutputIndex) override;
-	virtual FSubstrateOperator* SubstrateGenerateMaterialTopologyTree(class FMaterialCompiler* Compiler, class UMaterialExpression* Parent, int32 OutputIndex) override;
+	virtual FSubstrateOperator* SubstrateGenerateMaterialTopologyTree(struct FSubstrateTranslatorDataInterface& SubstrateTranslatorData, class UMaterialExpression* Parent, int32 OutputIndex) override;
 	virtual void PostEditChangeProperty(FPropertyChangedEvent& PropertyChangedEvent) override;
 	virtual void GetConnectorToolTip(int32 InputIndex, int32 OutputIndex, TArray<FString>& OutToolTip) override;
 	virtual bool IsInputConnectionRequired(int32 InputIndex) const override {return true;}
 	
 	bool HasSSS() const;
 
+	UE_DEPRECATED(5.8, "Please, SubstrateGenerateMaterialTopologyTreeCommon taking the FSubstrateTranslatorData parameter.")
 	static FSubstrateOperator* SubstrateGenerateMaterialTopologyTreeCommon(
 		class FMaterialCompiler* Compiler, struct FGuid ThisExpressionGuid, class UMaterialExpression* Parent, int32 OutputIndex,
 		const uint64 CachedConnectedMaterialAttributesInputs, const bool bShadingModelFromMaterialExpression, const bool bIsEmissiveConnected);
 
+	static FSubstrateOperator* SubstrateGenerateMaterialTopologyTreeCommon(
+		struct FSubstrateTranslatorDataInterface& SubstrateTranslatorData, struct FGuid ThisExpressionGuid, class UMaterialExpression* Parent, int32 OutputIndex,
+		const uint64 CachedConnectedMaterialAttributesInputs, const bool bShadingModelFromMaterialExpression, const bool bIsEmissiveConnected);
+
 	static int32 CompileCommon(class FMaterialCompiler* Compiler, int32 OutputIndex,
 		const uint64 CachedConnectedMaterialAttributesInputs, FMaterialAttributesInput& MaterialAttributes, enum EMaterialShadingModel ShadingModelOverride,
 		FExpressionInput& WaterScatteringCoefficients, FExpressionInput& WaterAbsorptionCoefficients, FExpressionInput& WaterPhaseG, FExpressionInput& ColorScaleBehindWater,
@@ -1096,7 +1148,8 @@ class UMaterialExpressionSubstrateHorizontalMixing : public UMaterialExpressionS
 	virtual EMaterialValueType GetInputValueType(int32 InputIndex) override;
 	virtual bool IsResultSubstrateMaterial(int32 OutputIndex) override;
 	virtual void GatherSubstrateMaterialInfo(FSubstrateMaterialInfo& SubstrateMaterialInfo, int32 OutputIndex) override;
-	virtual FSubstrateOperator* SubstrateGenerateMaterialTopologyTree(class FMaterialCompiler* Compiler, class UMaterialExpression* Parent, int32 OutputIndex) override;
+	virtual FSubstrateOperator* SubstrateGenerateMaterialTopologyTree(struct FSubstrateTranslatorDataInterface& SubstrateTranslatorData, class UMaterialExpression* Parent, int32 OutputIndex) override;
+	virtual void Build(MIR::FEmitter& Emitter) override;
 #endif
 	//~ End UMaterialExpression Interface
 };
@@ -1139,8 +1192,9 @@ class UMaterialExpressionSubstrateVerticalLayering : public UMaterialExpressionS
 	virtual EMaterialValueType GetInputValueType(int32 InputIndex) override;
 	virtual bool IsResultSubstrateMaterial(int32 OutputIndex) override;
 	virtual void GatherSubstrateMaterialInfo(FSubstrateMaterialInfo& SubstrateMaterialInfo, int32 OutputIndex) override;
-	virtual FSubstrateOperator* SubstrateGenerateMaterialTopologyTree(class FMaterialCompiler* Compiler, UMaterialExpression* Parent, int32 OutputIndex) override;
+	virtual FSubstrateOperator* SubstrateGenerateMaterialTopologyTree(struct FSubstrateTranslatorDataInterface& SubstrateTranslatorData, class UMaterialExpression* Parent, int32 OutputIndex) override;
 	virtual FName GetInputName(int32 InputIndex) const override;
+	virtual void Build(MIR::FEmitter& Emitter) override;
 #endif
 	//~ End UMaterialExpression Interface
 };
@@ -1176,7 +1230,7 @@ class UMaterialExpressionSubstrateAdd : public UMaterialExpressionSubstrateBSDF
 	virtual EMaterialValueType GetInputValueType(int32 InputIndex) override;
 	virtual bool IsResultSubstrateMaterial(int32 OutputIndex) override;
 	virtual void GatherSubstrateMaterialInfo(FSubstrateMaterialInfo& SubstrateMaterialInfo, int32 OutputIndex) override;
-	virtual FSubstrateOperator* SubstrateGenerateMaterialTopologyTree(class FMaterialCompiler* Compiler, class UMaterialExpression* Parent, int32 OutputIndex) override;
+	virtual FSubstrateOperator* SubstrateGenerateMaterialTopologyTree(struct FSubstrateTranslatorDataInterface& SubstrateTranslatorData, class UMaterialExpression* Parent, int32 OutputIndex) override;
 #endif
 	//~ End UMaterialExpression Interface
 };
@@ -1206,7 +1260,7 @@ class UMaterialExpressionSubstrateWeight : public UMaterialExpressionSubstrateBS
 	virtual EMaterialValueType GetInputValueType(int32 InputIndex) override;
 	virtual bool IsResultSubstrateMaterial(int32 OutputIndex) override;
 	virtual void GatherSubstrateMaterialInfo(FSubstrateMaterialInfo& SubstrateMaterialInfo, int32 OutputIndex) override;
-	virtual FSubstrateOperator* SubstrateGenerateMaterialTopologyTree(class FMaterialCompiler* Compiler, class UMaterialExpression* Parent, int32 OutputIndex) override;
+	virtual FSubstrateOperator* SubstrateGenerateMaterialTopologyTree(struct FSubstrateTranslatorDataInterface& SubstrateTranslatorData, class UMaterialExpression* Parent, int32 OutputIndex) override;
 #endif
 	//~ End UMaterialExpression Interface
 };
@@ -1254,7 +1308,7 @@ class UMaterialExpressionSubstrateSelect : public UMaterialExpressionSubstrateBS
 	virtual EMaterialValueType GetInputValueType(int32 InputIndex) override;
 	virtual bool IsResultSubstrateMaterial(int32 OutputIndex) override;
 	virtual void GatherSubstrateMaterialInfo(FSubstrateMaterialInfo& SubstrateMaterialInfo, int32 OutputIndex) override;
-	virtual FSubstrateOperator* SubstrateGenerateMaterialTopologyTree(class FMaterialCompiler* Compiler, class UMaterialExpression* Parent, int32 OutputIndex) override;
+	virtual FSubstrateOperator* SubstrateGenerateMaterialTopologyTree(struct FSubstrateTranslatorDataInterface& SubstrateTranslatorData, class UMaterialExpression* Parent, int32 OutputIndex) override;
 #endif
 	//~ End UMaterialExpression Interface
 };
@@ -1296,6 +1350,7 @@ class UMaterialExpressionSubstrateTransmittanceToMFP : public UMaterialExpressio
 	virtual EMaterialValueType GetInputValueType(int32 InputIndex) override;
 	virtual void GetConnectorToolTip(int32 InputIndex, int32 OutputIndex, TArray<FString>& OutToolTip) override;
 	virtual void GetExpressionToolTip(TArray<FString>& OutToolTip) override;
+	virtual void Build(MIR::FEmitter& Emitter) override;
 #endif
 	//~ End UMaterialExpression Interface
 };
diff --git a/Engine/Source/Runtime/Engine/Public/Materials/MaterialIR.h b/Engine/Source/Runtime/Engine/Public/Materials/MaterialIR.h
--- a/Engine/Source/Runtime/Engine/Public/Materials/MaterialIR.h
+++ b/Engine/Source/Runtime/Engine/Public/Materials/MaterialIR.h
@@ -123,7 +123,12 @@ enum EValueKind
 	VK_TextureRead,
 	VK_VTPageTableRead,
 	VK_InlineHLSL,
-	VK_PromoteSubstrateParameter,
+	VK_SubstrateDefaultSlab,
+	VK_SubstrateSlab,
+	VK_SubstrateShadingModels,
+	VK_SubstrateHorizontalMixing,
+	VK_SubstrateVerticalLayering,
+	VK_SubstratePromoteToOperator,
 	VK_StageSwitch,
 	VK_HardwarePartialDerivative,
 	VK_Nop,
@@ -746,16 +751,125 @@ struct FInlineHLSL : TInstruction<VK_InlineHLSL, 0>
 };
 
 // Material IR representation of the FSubstrateData shader parameter to transition legacy materials into Substrate.
-struct FPromoteSubstrateParameter : TInstruction<VK_PromoteSubstrateParameter, 4>
+struct FSubstrateDefaultSlab : TInstruction<VK_SubstrateDefaultSlab, 0>
 {
-	// Extra IR values for tangent and normal vectors that are transformed into world-space specifically for Substrate conversion.
-	FValue* WorldSpaceTangentsAndNormals[4];
+	FValue* Dummy;
+};
+
+struct FSubstrateSlab : TInstruction<VK_SubstrateSlab, 27> 
+{
+	FValue* Normal;
+	FValue* DiffuseAlbedo;
+	FValue* F0;
+	FValue* F90;
+	FValue* Roughness;
+	FValue* Anisotropy;
+	FValue* SSSProfileId;
+	FValue* SSSMFP;
+	FValue* SSSMFPScale;
+	FValue* SSSPhaseAniso;
+	FValue* SSSType;
+	FValue* EmissiveColor;
+	FValue* SecondRoughness;
+	FValue* SecondRoughnessWeight;
+	FValue* SecondRoughnessAsSimpleClearCoat;
+	FValue* ClearCoatUseSecondNormal;
+	FValue* ClearCoatBottomNormal;
+	FValue* FuzzAmount;
+	FValue* FuzzColor;
+	FValue* FuzzRoughness;
+	FValue* GlintValue;
+	FValue* GlintUV;
+	FValue* SpecularProfileId;
+	FValue* Thickness;
+	FValue* IsThin;
+	FValue* IsAtBottom;
+	FValue* LocalBasisIndex;
+
+	bool bFullySimplifiedInstruction;
+};
+
+struct FSubstrateShadingModels : TInstruction<VK_SubstrateShadingModels, 24>
+{
+	FValue* BaseColor;
+	FValue* Specular;
+	FValue* Metallic;
+	FValue* Roughness;
+	FValue* Anisotropy;
+	FValue* SubSurfaceColor;
+	FValue* SubSurfaceProfileId;
+	FValue* ClearCoat;
+	FValue* ClearCoatRoughness;
+	FValue* EmissiveColor;
+	FValue* Opacity;
+	FValue* ThinTranslucentTransmittanceColor;
+	FValue* ThinTranslucentSurfaceCoverage;
+	FValue* WaterScatteringCoefficients;
+	FValue* WaterAbsorptionCoefficients;
+	FValue* WaterPhaseG;
+	FValue* ColorScaleBehindWater;
+	FValue* ShadingModel;
+	FValue* Normal;
+	FValue* Tangent;
+	FValue* ClearCoatNormal;
+	FValue* CustomTangent;
+	FValue* BasisIndexMacro;
+	FValue* ClearCoat_BasisIndexMacro;
+
+	bool bHasDynamicShadingModels;
+	bool bFullySimplifiedInstruction;
+};
+
+struct FSubstrateHorizontalMixing : TInstruction<VK_SubstrateHorizontalMixing, 9>
+{
+	FValue* Background;
+	FValue* Foreground;
+	FValue* Mix;
+
+	// ParameterBlending = 0
+	FValue* OperatorIndex;
+	FValue* MaxDistanceFromLeaves;
+
+	// ParameterBlending = 1
+	FValue* NormalMix;
+	FValue* SharedLocalBasisIndexMacro;
+	FValue* BackgroundNoV;
+	FValue* ForegroundNoV;
+
+	// Settings
+	bool bParameterBlendingEnabled;
+	bool bFullySimplifiedInstruction;
+};
+
+struct FSubstrateVerticalLayering : TInstruction<VK_SubstrateVerticalLayering, 7>
+{
+	FValue* Top;
+	FValue* Base;
 
-	// IR values for all property arguments that forwarded to the Substrate conversion function.
-	FValue* PropertyArgs[MP_MAX];
+	// ParameterBlending = 0
+	FValue* OperatorIndex;
+	FValue* MaxDistanceFromLeaves;
 
-	// Specifies what Substrate BSDF to select, default or unlit.
-	bool bIsUnlit;
+	// ParameterBlending = 1
+	FValue* SharedLocalBasisIndexMacro;
+	FValue* TopNoV;
+	FValue* BaseNoV;
+
+	// Settings
+	bool bParameterBlendingEnabled;
+	bool bFullySimplifiedInstruction;
+};
+
+struct FSubstratePromoteToOperator : TInstruction<VK_SubstratePromoteToOperator, 5>
+{
+	FValue* SubstrateDataInput;
+	FValue* OperatorIndex;
+	FValue* BSDFIndex;
+	FValue* LayerDepth;
+	FValue* bIsBottom;
+
+	// Settings
+	bool bFullySimplifiedInstruction;
 };
 
 // Operator enumeration.
diff --git a/Engine/Source/Runtime/Engine/Public/Materials/MaterialIRCommon.h b/Engine/Source/Runtime/Engine/Public/Materials/MaterialIRCommon.h
--- a/Engine/Source/Runtime/Engine/Public/Materials/MaterialIRCommon.h
+++ b/Engine/Source/Runtime/Engine/Public/Materials/MaterialIRCommon.h
@@ -30,6 +30,11 @@ struct FMaterialIRModuleBuilderImpl;
 struct FShaderCompilerEnvironment;
 struct FStaticParameterSet;
 
+namespace Substrate
+{
+	struct FSubstrateTranslatorData;
+};
+
 namespace UE::Shader
 {
 	struct FValue;
diff --git a/Engine/Source/Runtime/Engine/Public/Materials/MaterialIREmitter.h b/Engine/Source/Runtime/Engine/Public/Materials/MaterialIREmitter.h
--- a/Engine/Source/Runtime/Engine/Public/Materials/MaterialIREmitter.h
+++ b/Engine/Source/Runtime/Engine/Public/Materials/MaterialIREmitter.h
@@ -187,6 +187,66 @@ struct FTextureSampleBaseAttributes
 	bool bIsAdaptive : 1 { false };
 };
 
+struct FSubstrateShadingModelsDesc
+{
+	FValueRef BaseColor;
+	FValueRef Specular;
+	FValueRef Metallic;
+	FValueRef Roughness;
+	FValueRef Anisotropy;
+	FValueRef SubSurfaceColor;
+	FValueRef SubSurfaceProfileId;
+	FValueRef ClearCoat;
+	FValueRef ClearCoatRoughness;
+	FValueRef EmissiveColor;
+	FValueRef Opacity;
+	FValueRef ThinTranslucentTransmittanceColor;
+	FValueRef ThinTranslucentSurfaceCoverage;
+	FValueRef WaterScatteringCoefficients;
+	FValueRef WaterAbsorptionCoefficients;
+	FValueRef WaterPhaseG;
+	FValueRef ColorScaleBehindWater;
+	FValueRef ShadingModel;
+	FValueRef Normal;
+	FValueRef Tangent;
+	FValueRef ClearCoatNormal;
+	FValueRef CustomTangent;
+	FValueRef BasisIndexMacro;
+	FValueRef ClearCoat_BasisIndexMacro;
+	bool bHasDynamicShadingModels;
+};
+
+struct FSubstrateSlabDesc
+{
+	FValueRef Normal;
+	FValueRef DiffuseAlbedo;
+	FValueRef F0;
+	FValueRef F90;
+	FValueRef Roughness;
+	FValueRef Anisotropy;
+	FValueRef SSSProfileId;
+	FValueRef SSSMFP;
+	FValueRef SSSMFPScale;
+	FValueRef SSSPhaseAniso;
+	FValueRef SSSType;
+	FValueRef EmissiveColor;
+	FValueRef SecondRoughness;
+	FValueRef SecondRoughnessWeight;
+	FValueRef SecondRoughnessAsSimpleClearCoat;
+	FValueRef ClearCoatUseSecondNormal;
+	FValueRef ClearCoatBottomNormal;
+	FValueRef FuzzAmount;
+	FValueRef FuzzColor;
+	FValueRef FuzzRoughness;
+	FValueRef GlintValue;
+	FValueRef GlintUV;
+	FValueRef SpecularProfileId;
+	FValueRef Thickness;
+	FValueRef IsThin;
+	FValueRef IsAtBottom;
+	FValueRef LocalBasisIndex;
+};
+
 // The FEmitter is a helper object responsible for emitting MIR values and instructions.
 // It is primarily used within UMaterialExpression::Build() implementations to lower the
 // high-level semantics of a material expression into corresponding lower-level MIR nodes.
@@ -645,8 +705,29 @@ public:
 	// @param UsedGraphProperties Properties used during graph processing (default is None).
 	FValueRef InlineHLSL(const FMaterialExternalCodeDeclaration* InExternalCodeDeclaration, TConstArrayView<FValueRef> InputValues, EValueFlags ValueFlags = EValueFlags::None, EGraphProperties UsedGraphProperties = EGraphProperties::None);
 
-	// Emits an instruction that promotes a Substrate parameter, encoded as MCT_Substrate.
-	FValueRef PromoteSubstrateParameter();
+	// Emits an instruction that promotes a SubstrateData representing a default slab, encoded as MCT_Substrate.
+	FValueRef SubstrateDefaultSlab();
+
+	// Promote a SubstrateData BSDF to an operator so that is becomes a node of the SubstrateTree and can be combined with other operators.
+	FValueRef SubstratePromoteToOperator(FValueRef SubstrateData, FValueRef OperatorIndex, FValueRef BSDFIndex, FValueRef LayerDepth, FValueRef bIsBottom);
+
+	// Create a SubstrateData representing a legacy BSDF.
+	FValueRef SubstrateShadingModels(const FSubstrateShadingModelsDesc Desc);
+
+	// Create a SubstrateData representing a Slab BSDF.
+	FValueRef SubstrateSlab(const FSubstrateSlabDesc& Desc);
+
+	// Get the SubstrateData resulting from the horizontal mixing of the Background and Foreground SubstrateData.
+	FValueRef SubstrateHorizontalMixing(FValueRef Background, FValueRef Foreground, FValueRef Mix, FValueRef OperatorIndex, FValueRef MaxDistanceFromLeaves);
+
+	// Parameter blended version of SubstrateHorizontalMixing.
+	FValueRef SubstrateHorizontalMixingParameterBlending(FValueRef Background, FValueRef Foreground, FValueRef Mix, FValueRef NormalMix, FValueRef SharedLocalBasisIndexMacro, FValueRef BackgroundNoV, FValueRef ForegroundNoV);
+
+	// Get the SubstrateData resulting from the vertical layering of the Top over the Base SubstrateData.
+	FValueRef SubstrateVerticalLayering(FValueRef Top, FValueRef Base, FValueRef OperatorIndex, FValueRef MaxDistanceFromLeaves);
+
+	// Parameter blended version of SubstrateVerticalLayering.
+	FValueRef SubstrateVerticalLayeringParameterBlending(FValueRef Top, FValueRef Base, FValueRef SharedLocalBasisIndexMacro, FValueRef TopNoV, FValueRef BaseNoV);
 
 	// Defines a user-specified HLSL function to embed in the generated HLSL when generated IR Module is translated to HLSL.
 	// This will create a local function with the specified description. You can issue calls to the returned functions using the Call() API.
@@ -735,6 +816,7 @@ public:
 	const ITargetPlatform* GetTargetPlatform() const;
 	ERHIFeatureLevel::Type GetFeatureLevel() const;
 	EMaterialQualityLevel::Type GetQualityLevel() const;
+	Substrate::FSubstrateTranslatorData* GetSubstrateTranslatorData() const;
 
 	// Internal
 	struct FPrivate;
@@ -768,6 +850,9 @@ private:
 	// The current expression being built (set by the builder).
 	UMaterialExpression* Expression{};
 
+	// The substrate translator common data representing the material topology.
+	struct Substrate::FSubstrateTranslatorData* SubstrateTranslatorData;
+
 	// Whether the current expression has reported any error.
 	bool bCurrentExpressionHasErrors = false;
 
diff --git a/Engine/Source/Runtime/Engine/Public/Materials/MaterialIRModule.h b/Engine/Source/Runtime/Engine/Public/Materials/MaterialIRModule.h
--- a/Engine/Source/Runtime/Engine/Public/Materials/MaterialIRModule.h
+++ b/Engine/Source/Runtime/Engine/Public/Materials/MaterialIRModule.h
@@ -80,6 +80,14 @@ public:
 		int32 CustomOutputInstructionsBegin;
 	};
 
+	// Specific Substrate data that needs to be handled in a specific way today when emitting back shader code from IR.
+	struct FSubstrateData
+	{
+		FString PixelMembersDeclarationHLSL;
+		FString PixelInputInitializerValuesHLSL;
+		FString TreeHLSL;
+	};
+
 public:
 	FMaterialIRModule();
 	~FMaterialIRModule();
@@ -199,6 +207,11 @@ public:
 		return { (T*)Allocate(sizeof(T) * Count, alignof(T)), Count };
 	}
 
+	const FSubstrateData& GetSubstrateData() const
+	{
+		return SubstrateData;
+	}
+
 private:
 	// The chunk allocator used to contain child allocations (e.g., values, interned strings, etc)
 	FMemStackBase Allocator;
@@ -263,6 +276,9 @@ private:
 	// Shading models used by this material.
 	FMaterialShadingModelField ShadingModelsFromCompilation;
 
+	// Substrate compiled data.
+	FSubstrateData SubstrateData;
+
 	// List of compilation errors.
 	TArray<FError> Errors;
 
