[SkeletalMesh] improve tool tip + added warning against direct access of target mesh deformers

#Synced-CL 49562189

diff --git a/Engine/Source/Runtime/Engine/Classes/Engine/SkeletalMesh.h b/Engine/Source/Runtime/Engine/Classes/Engine/SkeletalMesh.h
--- a/Engine/Source/Runtime/Engine/Classes/Engine/SkeletalMesh.h
+++ b/Engine/Source/Runtime/Engine/Classes/Engine/SkeletalMesh.h
@@ -3089,7 +3089,8 @@ protected:
 	UPROPERTY(EditAnywhere, BlueprintReadOnly, Category = "Deformer", BlueprintGetter=GetDefaultMeshDeformer)
 	TObjectPtr<UMeshDeformer> DefaultMeshDeformer;
 	
-	/** Skeletal Mesh needs this collection of deformers to make sure it cooks any extra data required by these deformers */
+	/** Additional deformers that will be applied to this mesh at runtime. Skeletal Mesh needs this collection of deformers to make sure it cooks any extra data required by these deformers */
+	UE_DEPRECATED(5.8, "This must be protected for async build, always use the accessors even internally.")
 	UPROPERTY(EditAnywhere, BlueprintReadOnly, Category = "Deformer", BlueprintGetter=GetTargetMeshDeformers)
 	TObjectPtr<UMeshDeformerCollection> TargetMeshDeformers;
 
@@ -3136,7 +3137,9 @@ public:
 	void SetTargetMeshDeformers(UMeshDeformerCollection* InMeshDeformerCollection)
 	{
 		WaitUntilAsyncPropertyReleased(ESkeletalMeshAsyncProperties::TargetMeshDeformers, ESkinnedAssetAsyncPropertyLockType::WriteOnly);
+		PRAGMA_DISABLE_DEPRECATION_WARNINGS
 		TargetMeshDeformers = InMeshDeformerCollection;
+		PRAGMA_ENABLE_DEPRECATION_WARNINGS
 	}
 	
 	/** Set the collection of mesh deformers that may be applied to this mesh. Skeletal Mesh use these deformers to determined if any extra data needs to be cooked*/
@@ -3144,7 +3147,9 @@ public:
 	virtual UMeshDeformerCollection* GetTargetMeshDeformers() const override
 	{
 		WaitUntilAsyncPropertyReleased(ESkeletalMeshAsyncProperties::TargetMeshDeformers, ESkinnedAssetAsyncPropertyLockType::ReadOnly);
+		PRAGMA_DISABLE_DEPRECATION_WARNINGS
 		return TargetMeshDeformers;
+		PRAGMA_ENABLE_DEPRECATION_WARNINGS
 	}
 
 	/** True if this mesh LOD builds the half edge buffer. USkinnedAsset interface.	*/
