Fix code static analysis post 49733857.
bMaterialIsSubstrate was also the correct second test.

#tests PC editor dx12
#rnx
#Synced-CL 49731272

diff --git a/Engine/Source/Runtime/Engine/Private/Materials/SubstrateTranslatorCommon.cpp b/Engine/Source/Runtime/Engine/Private/Materials/SubstrateTranslatorCommon.cpp
--- a/Engine/Source/Runtime/Engine/Private/Materials/SubstrateTranslatorCommon.cpp
+++ b/Engine/Source/Runtime/Engine/Private/Materials/SubstrateTranslatorCommon.cpp
@@ -2313,25 +2313,28 @@ void FSubstrateTranslatorData::GenerateMaterialCompilationOutputAndDefines(
 	}
 
 	// Setup the remaining defines
-	OutSubstrateDefines->Add(TPair<FString, int32>(TEXT("MATERIAL_IS_SUBSTRATE"),									bMaterialIsSubstrate ? 1 : 0));
-	OutSubstrateDefines->Add(TPair<FString, int32>(TEXT("MATERIAL_IS_SUBSTRATEHIDDENCONVERSION"),					bMaterialUsesRootNodeToSubstrateHiddenConversion ? 1 : 0));
-	OutSubstrateDefines->Add(TPair<FString, int32>(TEXT("SUBSTRATE_PREMULTIPLIED_ALPHA_OPACITY_OVERRIDEN"),			bSubstratePremultipliedAlphaOpacityOverridden ? 1 : 0));
-	OutSubstrateDefines->Add(TPair<FString, int32>(TEXT("DISTORTION_ACCOUNT_FOR_COVERAGE"),							bDistortionAccountForCoverage ? 1 : 0));
 	if (OutSubstrateDefines)
 	{
-		OutSubstrateDefines->Add(TPair<FString, int32>(TEXT("SUBSTRATE_USES_CONVERSION_FROM_LEGACY"),					bSubstrateUsesConversionFromLegacy ? 1 : 0));
-		OutSubstrateDefines->Add(TPair<FString, int32>(TEXT("SUBSTRATE_OPTIMIZED_UNLIT"),								bSubstrateOptimizedUnlit ? 1 : 0));
-		OutSubstrateDefines->Add(TPair<FString, int32>(TEXT("SUBSTRATE_MATERIAL_OUTPUT_OPAQUE_ROUGH_REFRACTIONS"),		bSubstrateMaterialOutputOpaqueRoughRefractions ? 1 : 0));
-		OutSubstrateDefines->Add(TPair<FString, int32>(TEXT("SUBSTRATE_MATERIAL_EXPORT_TYPE"),							SubstrateMaterialExportType));
-		OutSubstrateDefines->Add(TPair<FString, int32>(TEXT("SUBSTRATE_MATERIAL_EXPORT_CONTEXT"),						SubstrateMaterialExportContext));
-		OutSubstrateDefines->Add(TPair<FString, int32>(TEXT("SUBSTRATE_MATERIAL_EXPORT_LEGACY_BLEND_MODE"),				SubstrateMaterialExportLegacyBlendMode));
-		OutSubstrateDefines->Add(TPair<FString, int32>(TEXT("SUBSTRATE_MATERIAL_ROUGHNESS_TRACKING"),					bSubstrateRoughnessTracking ? 1 : 0));
-		OutSubstrateDefines->Add(TPair<FString, int32>(TEXT("SUBSTRATE_LEGACY_IRIS_NORMAL"),							bSubstrateLegacyIrisNormal ? 1 : 0));
-		OutSubstrateDefines->Add(TPair<FString, int32>(TEXT("SUBSTRATE_LEGACY_IRIS_TANGENT"),							bSubstrateLegacyIrisTangent ? 1 : 0));
-		OutSubstrateDefines->Add(TPair<FString, int32>(TEXT("SUBSTRATE_SINGLEPATH"),									bSubstrateSinglePath ? 1 : 0));
-		OutSubstrateDefines->Add(TPair<FString, int32>(TEXT("SUBSTRATE_FASTPATH"),										bSubstrateFastPath ? 1 : 0));
-		OutSubstrateDefines->Add(TPair<FString, int32>(TEXT("SUBSTRATE_CLAMPED_CLOSURE_COUNT"),							(int32)SubstrateClampedClosureCount));
-		OutSubstrateDefines->Add(TPair<FString, int32>(TEXT("SUBSTRATE_COMPLEXSPECIALPATH"),							bSubstrateComplexSpecialPath ? 1 : 0));
+		OutSubstrateDefines->Add(TPair<FString, int32>(TEXT("MATERIAL_IS_SUBSTRATE"),									bMaterialIsSubstrate ? 1 : 0));
+		OutSubstrateDefines->Add(TPair<FString, int32>(TEXT("MATERIAL_IS_SUBSTRATEHIDDENCONVERSION"),					bMaterialUsesRootNodeToSubstrateHiddenConversion ? 1 : 0));
+		OutSubstrateDefines->Add(TPair<FString, int32>(TEXT("SUBSTRATE_PREMULTIPLIED_ALPHA_OPACITY_OVERRIDEN"),			bSubstratePremultipliedAlphaOpacityOverridden ? 1 : 0));
+		OutSubstrateDefines->Add(TPair<FString, int32>(TEXT("DISTORTION_ACCOUNT_FOR_COVERAGE"),							bDistortionAccountForCoverage ? 1 : 0));
+		if(bMaterialIsSubstrate)
+		{
+			OutSubstrateDefines->Add(TPair<FString, int32>(TEXT("SUBSTRATE_USES_CONVERSION_FROM_LEGACY"),				bSubstrateUsesConversionFromLegacy ? 1 : 0));
+			OutSubstrateDefines->Add(TPair<FString, int32>(TEXT("SUBSTRATE_OPTIMIZED_UNLIT"),							bSubstrateOptimizedUnlit ? 1 : 0));
+			OutSubstrateDefines->Add(TPair<FString, int32>(TEXT("SUBSTRATE_MATERIAL_OUTPUT_OPAQUE_ROUGH_REFRACTIONS"),	bSubstrateMaterialOutputOpaqueRoughRefractions ? 1 : 0));
+			OutSubstrateDefines->Add(TPair<FString, int32>(TEXT("SUBSTRATE_MATERIAL_EXPORT_TYPE"),						SubstrateMaterialExportType));
+			OutSubstrateDefines->Add(TPair<FString, int32>(TEXT("SUBSTRATE_MATERIAL_EXPORT_CONTEXT"),					SubstrateMaterialExportContext));
+			OutSubstrateDefines->Add(TPair<FString, int32>(TEXT("SUBSTRATE_MATERIAL_EXPORT_LEGACY_BLEND_MODE"),			SubstrateMaterialExportLegacyBlendMode));
+			OutSubstrateDefines->Add(TPair<FString, int32>(TEXT("SUBSTRATE_MATERIAL_ROUGHNESS_TRACKING"),				bSubstrateRoughnessTracking ? 1 : 0));
+			OutSubstrateDefines->Add(TPair<FString, int32>(TEXT("SUBSTRATE_LEGACY_IRIS_NORMAL"),						bSubstrateLegacyIrisNormal ? 1 : 0));
+			OutSubstrateDefines->Add(TPair<FString, int32>(TEXT("SUBSTRATE_LEGACY_IRIS_TANGENT"),						bSubstrateLegacyIrisTangent ? 1 : 0));
+			OutSubstrateDefines->Add(TPair<FString, int32>(TEXT("SUBSTRATE_SINGLEPATH"),								bSubstrateSinglePath ? 1 : 0));
+			OutSubstrateDefines->Add(TPair<FString, int32>(TEXT("SUBSTRATE_FASTPATH"),									bSubstrateFastPath ? 1 : 0));
+			OutSubstrateDefines->Add(TPair<FString, int32>(TEXT("SUBSTRATE_CLAMPED_CLOSURE_COUNT"),						(int32)SubstrateClampedClosureCount));
+			OutSubstrateDefines->Add(TPair<FString, int32>(TEXT("SUBSTRATE_COMPLEXSPECIALPATH"),						bSubstrateComplexSpecialPath ? 1 : 0));
+			}
 	}
 }
 
