Prevent FChildActorComponentInstanceData from being recreated when the ChildActorComponent or it's parent Actor are being destroyed

#tests Windows Test client with Insights capture
#Synced-CL 49641272

diff --git a/Engine/Source/Runtime/Engine/Private/Components/ChildActorComponent.cpp b/Engine/Source/Runtime/Engine/Private/Components/ChildActorComponent.cpp
--- a/Engine/Source/Runtime/Engine/Private/Components/ChildActorComponent.cpp
+++ b/Engine/Source/Runtime/Engine/Private/Components/ChildActorComponent.cpp
@@ -925,15 +925,13 @@ void UChildActorComponent::SetChildActorNameIsExact(bool bInExact)
 void UChildActorComponent::DestroyChildActor()
 {
 	UWorld* LocalWorld = GetWorld();
-	if (LocalWorld != nullptr)
+	AActor* MyOwner = GetOwner();
+
+	// If we're in a replay scrub and this is a startup actor and the class exists, assume the reference will be restored later
+	if (LocalWorld && LocalWorld->IsPlayingReplay() && LocalWorld->GetDemoNetDriver()->IsRestoringStartupActors() && MyOwner && MyOwner->IsNetStartupActor() && ChildActorClass)
 	{
-		AActor* MyOwner = GetOwner();
-		// If we're in a replay scrub and this is a startup actor and the class exists, assume the reference will be restored later
-		if (LocalWorld->IsPlayingReplay() && LocalWorld->GetDemoNetDriver()->IsRestoringStartupActors() && MyOwner && MyOwner->IsNetStartupActor() && ChildActorClass)
-		{
-			UE_LOG(LogChildActorComponent, Display, TEXT("Not creating destroying actor due to replay! Comp:%s ChildActor:%s"), *GetPathName(), *GetPathNameSafe(ChildActor));
-			return;
-		}
+		UE_LOG(LogChildActorComponent, Display, TEXT("Not creating destroying actor due to replay! Comp:%s ChildActor:%s"), *GetPathName(), *GetPathNameSafe(ChildActor));
+		return;
 	}
 
 	// If we own an Actor, kill it now unless we don't have authority on it, for that we rely on the server
@@ -956,7 +954,7 @@ void UChildActorComponent::DestroyChildActor()
 				check(!CachedInstanceData);
 #endif
 				// If we're already tearing down we won't be needing this
-				if (!HasAnyFlags(RF_BeginDestroyed) && !IsUnreachable())
+				if (!HasAnyFlags(RF_BeginDestroyed) && !IsUnreachable() && !IsBeingDestroyed() && !(MyOwner && MyOwner->IsActorBeingDestroyed()))
 				{
 					CachedInstanceData = new FChildActorComponentInstanceData(this);
 				}
