[NetPhysics] Fix issue with NetworkPhysicsSettingsDataAsset values getting stomped. Limited fix only to Redundant inputs and state settings.

diff --git a/Engine/Source/Runtime/Engine/Public/Physics/NetworkPhysicsSettingsComponent.h b/Engine/Source/Runtime/Engine/Public/Physics/NetworkPhysicsSettingsComponent.h
--- a/Engine/Source/Runtime/Engine/Public/Physics/NetworkPhysicsSettingsComponent.h
+++ b/Engine/Source/Runtime/Engine/Public/Physics/NetworkPhysicsSettingsComponent.h
@@ -584,18 +584,18 @@ struct FNetworkPhysicsSettingsNetworkPhysicsComponent
 
 	/* Overrides CVar : np2.Resim.RedundantInputs -- How many extra inputs to send with each unreliable network message, to account for packetloss.From owning client to server and server to owning client.
 	 * NOTE: This is disabled while np2.Resim.DynamicInputReplicationScaling.Enabled is enabled. */
-	UPROPERTY(config, EditDefaultsOnly, Category = "Overrides", Meta = (EditCondition = "bOverrideRedundantInputs"))
-	uint16 RedundantInputs = (uint16)PhysicsReplicationCVars::ResimulationCVars::RedundantInputs;
+	UPROPERTY(EditDefaultsOnly, Category = "Overrides", Meta = (EditCondition = "bOverrideRedundantInputs"))
+	uint16 RedundantInputs = 2;
 	const uint16 GetRedundantInputs() const { return bOverrideRedundantInputs ? RedundantInputs : (uint16)PhysicsReplicationCVars::ResimulationCVars::RedundantInputs; }
 
 	// Overrides CVar: np2.Resim.RedundantRemoteInputs -- How many extra inputs to send with each unreliable network message, to account for packetloss. From Server to remote clients.
-	UPROPERTY(config, EditDefaultsOnly, Category = "Overrides", Meta = (EditCondition = "bOverrideRedundantRemoteInputs"))
-	uint16 RedundantRemoteInputs = (uint16)PhysicsReplicationCVars::ResimulationCVars::RedundantRemoteInputs;
+	UPROPERTY(EditDefaultsOnly, Category = "Overrides", Meta = (EditCondition = "bOverrideRedundantRemoteInputs"))
+	uint16 RedundantRemoteInputs = 1;
 	const uint16 GetRedundantRemoteInputs() const { return bOverrideRedundantRemoteInputs ? RedundantRemoteInputs : (uint16)PhysicsReplicationCVars::ResimulationCVars::RedundantRemoteInputs; }
 
 	// Overrides CVar: np2.Resim.RedundantStates -- How many extra states to send with each unreliable network message, to account for packetloss. From Server to remote clients.
-	UPROPERTY(config, EditDefaultsOnly, Category = "Overrides", Meta = (EditCondition = "bOverrideRedundantStates"))
-	uint16 RedundantStates = (uint16)PhysicsReplicationCVars::ResimulationCVars::RedundantStates;
+	UPROPERTY(EditDefaultsOnly, Category = "Overrides", Meta = (EditCondition = "bOverrideRedundantStates"))
+	uint16 RedundantStates = 0;
 	const uint16 GetRedundantStates() const { return bOverrideRedundantStates ? RedundantStates : (uint16)PhysicsReplicationCVars::ResimulationCVars::RedundantStates; }
 
 	// Overrides CVar: np2.Resim.DynamicInputBufferScaling.Enabled -- Enable dynmic scaling of input buffer on the server.
