[SceneViewport] Adjust calculation of unbounded mouse position to only run on mouse & touch move events. Track motion of bounded mouse position when possible to avoid duplicate event (and thus duplicate delta) processing. Set unbounded position on mousedown prior to processing viewport mousedown events. Set unbounded positon when the mouse is set explicitly.

Fixes an issue with ITF capture stealing where a mouse drag based on bound position and a mouse drag based on unbounded position would quickly drift from one another, making relative comparison impossible.

#rnx
#Synced-CL 49683958

diff --git a/Engine/Source/Runtime/Engine/Private/Slate/SceneViewport.cpp b/Engine/Source/Runtime/Engine/Private/Slate/SceneViewport.cpp
--- a/Engine/Source/Runtime/Engine/Private/Slate/SceneViewport.cpp
+++ b/Engine/Source/Runtime/Engine/Private/Slate/SceneViewport.cpp
@@ -231,7 +231,7 @@ void FSceneViewport::SetMouse( int32 X, int32 Y )
 	const FVector2D NormalizedLocalMousePosition = FVector2D(X, Y) / GetSizeXY();
 	FVector2D AbsolutePos = CachedGeometry.LocalToAbsolute(NormalizedLocalMousePosition * CachedGeometry.GetLocalSize());
 	FSlateApplication::Get().SetCursorPos( AbsolutePos.RoundToVector() );
-	CachedCursorPos = FIntPoint(X, Y);
+	CachedUnboundedCursorPos = CachedCursorPos = FIntPoint(X, Y);
 }
 
 void FSceneViewport::ProcessInput( float DeltaTime )
@@ -239,7 +239,7 @@ void FSceneViewport::ProcessInput( float DeltaTime )
 	// Required 
 }
 
-void FSceneViewport::UpdateCachedCursorPos( const FGeometry& InGeometry, const FPointerEvent& InMouseEvent )
+void FSceneViewport::UpdateCachedCursorPos( const FGeometry& InGeometry, const FPointerEvent& InMouseEvent, bool bIsMoveEvent )
 {
 	const FPlatformUserId UserId = IPlatformInputDeviceMapper::Get().GetUserForInputDevice(InMouseEvent.GetInputDeviceId());
 	if (UserId == FSlateApplication::SlateAppPrimaryPlatformUser)
@@ -247,12 +247,27 @@ void FSceneViewport::UpdateCachedCursorPos( const FGeometry& InGeometry, const F
 		FVector2D LocalPixelMousePos = InGeometry.AbsoluteToLocal(InMouseEvent.GetScreenSpacePosition());
 		LocalPixelMousePos.X = FMath::Clamp(LocalPixelMousePos.X * CachedGeometry.Scale, (double)TNumericLimits<int32>::Min(), (double)TNumericLimits<int32>::Max());
 		LocalPixelMousePos.Y = FMath::Clamp(LocalPixelMousePos.Y * CachedGeometry.Scale, (double)TNumericLimits<int32>::Min(), (double)TNumericLimits<int32>::Max());
-
-		CachedCursorPos = LocalPixelMousePos.IntPoint();
+		
+		const FIntPoint NewPos = LocalPixelMousePos.IntPoint();
+		const FIntPoint Delta = NewPos - CachedCursorPos;
+		
+		CachedCursorPos = NewPos;
+		
+		// The cached cursor value is scaled and needs to be reversed to get an accurate calculation
+		if (InGeometry.IsUnderLocation(InGeometry.LocalToAbsolute(FVector2f(CachedCursorPos) / InGeometry.Scale)))
+		{
+			// When the mouse is over the viewport, the unbound cursor should move in lock step with the normal cursor
+			CachedUnboundedCursorPos += Delta;
+		}
+		else if (bIsMoveEvent)
+		{
+			// When the mouse is outside the viewport, rely on straight deltas.
+			// This distinction was made due to seeing duplicate OnMouseMove deltas being applied.
+			// If that problem can be ameliorated, the above branch should no longer be necessary.
+			const FIntPoint UnboundDelta = (InMouseEvent.GetCursorDelta() * InGeometry.Scale).IntPoint();
+			CachedUnboundedCursorPos += UnboundDelta;
+		}
 	}
-
-	CachedUnboundedCursorPos.X += InMouseEvent.GetCursorDelta().X;
-	CachedUnboundedCursorPos.Y += InMouseEvent.GetCursorDelta().Y;
 }
 
 void FSceneViewport::UpdateCachedGeometry( const FGeometry& InGeometry )
@@ -570,6 +585,16 @@ FReply FSceneViewport::OnMouseButtonDown( const FGeometry& InGeometry, const FPo
 
 	UpdateCachedGeometry(InGeometry);
 	UpdateCachedCursorPos(InGeometry, InMouseEvent);
+	
+	const uint32 DownButtonsCount = InMouseEvent.IsMouseButtonDown(EKeys::LeftMouseButton)
+								  + InMouseEvent.IsMouseButtonDown(EKeys::MiddleMouseButton)
+								  + InMouseEvent.IsMouseButtonDown(EKeys::RightMouseButton);
+
+	// This is the first pressed button: initialize unbounded cursor position to match current cursor position
+	if (DownButtonsCount == 1)
+	{
+		CachedUnboundedCursorPos = CachedCursorPos;
+	}
 
 	// Switch to the viewport clients world before processing input
 	FScopedConditionalWorldSwitcher WorldSwitcher(ViewportClient);
@@ -618,17 +643,6 @@ FReply FSceneViewport::OnMouseButtonDown( const FGeometry& InGeometry, const FPo
 
 	// Re-set prevent throttling here as it can get reset when inside of InputKey()
 	CurrentReplyState.PreventThrottling();
-
-	const uint32 DownButtonsCount = InMouseEvent.IsMouseButtonDown(EKeys::LeftMouseButton)
-								  + InMouseEvent.IsMouseButtonDown(EKeys::MiddleMouseButton)
-								  + InMouseEvent.IsMouseButtonDown(EKeys::RightMouseButton);
-
-	// This is the first pressed button: initialize unbounded cursor position to match current cursor position
-	if (DownButtonsCount == 1)
-	{
-		CachedUnboundedCursorPos = CachedCursorPos;
-	}
-
 	return CurrentReplyState;
 }
 
@@ -792,7 +806,9 @@ FReply FSceneViewport::OnMouseMove( const FGeometry& InGeometry, const FPointerE
 	CurrentReplyState = FReply::Handled();
 
 	UpdateCachedGeometry(InGeometry);
-	UpdateCachedCursorPos(InGeometry, InMouseEvent);
+	
+	constexpr bool bIsMoveEvent = true;
+	UpdateCachedCursorPos(InGeometry, InMouseEvent, bIsMoveEvent);
 
 	const bool bViewportHasCapture = ViewportWidget.IsValid() && ViewportWidget.Pin()->HasMouseCapture();
 	if ( ViewportClient && GetSizeXY() != FIntPoint::ZeroValue )
@@ -923,7 +939,9 @@ FReply FSceneViewport::OnTouchMoved( const FGeometry& MyGeometry, const FPointer
 	CurrentReplyState = FReply::Handled();
 
 	UpdateCachedGeometry(MyGeometry);
-	UpdateCachedCursorPos(MyGeometry, TouchEvent);
+	
+	constexpr bool bIsMoveEvent = true;
+	UpdateCachedCursorPos(MyGeometry, TouchEvent, bIsMoveEvent);
 
 	if( ViewportClient )
 	{
diff --git a/Engine/Source/Runtime/Engine/Public/Slate/SceneViewport.h b/Engine/Source/Runtime/Engine/Public/Slate/SceneViewport.h
--- a/Engine/Source/Runtime/Engine/Public/Slate/SceneViewport.h
+++ b/Engine/Source/Runtime/Engine/Public/Slate/SceneViewport.h
@@ -321,7 +321,7 @@ private:
 	 * @param InGeometry	The geometry of the viewport to convert to local space
 	 * @param InMouseEvent	The mouse event containing the position of the mouse in absolute space
 	 */
-	ENGINE_API void UpdateCachedCursorPos( const FGeometry& InGeometry, const FPointerEvent& InMouseEvent );
+	ENGINE_API void UpdateCachedCursorPos( const FGeometry& InGeometry, const FPointerEvent& InMouseEvent, bool bIsMoveEvent = false);
 
 	/**
 	 * Updates the cached viewport geometry
