Fix stress test unity ambiguous symbol error
#rnx
#Synced-CL 49749982

diff --git a/Engine/Source/Runtime/Engine/Private/Materials/SubstrateMaterial.cpp b/Engine/Source/Runtime/Engine/Private/Materials/SubstrateMaterial.cpp
--- a/Engine/Source/Runtime/Engine/Private/Materials/SubstrateMaterial.cpp
+++ b/Engine/Source/Runtime/Engine/Private/Materials/SubstrateMaterial.cpp
@@ -43,10 +43,13 @@ FSubstrateRegisteredSharedLocalBasis SubstrateCompilationInfoCreateSharedLocalBa
 	return Compiler->SubstrateCompilationInfoRegisterSharedLocalBasis(NormalCodeChunk, TangentCodeChunk);
 }
 
-inline bool IsSubstrateEnabled()
+namespace SubstrateMaterialPrivate
 {
-	static bool bSubstrateEnabled = Substrate::IsSubstrateEnabled();
-	return bSubstrateEnabled;
+	inline bool IsSubstrateEnabled()
+	{
+		static bool bSubstrateEnabled = Substrate::IsSubstrateEnabled();
+		return bSubstrateEnabled;
+	}
 }
 
 #define IsGenericBlendMode_Conversion0(Name, LegacyCondition) \
@@ -56,7 +59,7 @@ inline bool IsSubstrateEnabled()
 	bool Is##Name##BlendMode(const FMaterialShaderParameters& In)							{ return Is##Name##BlendMode(In.BlendMode); }
 
 #define IsGenericBlendMode_Conversion1(Name, LegacyCondition, SubstrateCondition) \
-	bool Is##Name##BlendMode(EBlendMode BlendMode)											{ return IsSubstrateEnabled() ? (SubstrateCondition) : (LegacyCondition); } \
+	bool Is##Name##BlendMode(EBlendMode BlendMode)											{ return SubstrateMaterialPrivate::IsSubstrateEnabled() ? (SubstrateCondition) : (LegacyCondition); } \
 	bool Is##Name##BlendMode(const FMaterial& In)											{ return Is##Name##BlendMode(In.GetBlendMode()); } \
 	bool Is##Name##BlendMode(const UMaterialInterface& In)									{ return Is##Name##BlendMode(In.GetBlendMode()); } \
 	bool Is##Name##BlendMode(const FMaterialShaderParameters& In)							{ return Is##Name##BlendMode(In.BlendMode); }
