Fixed some (ISM) instance data bugs.
* In instance data manager/proxy the custom data count was updated before movement due to deleted instances, causing OOB array access.
* Custom data was also not reset properly when the count is changed.
* Added bounds checks in ISMC before marking instances as removed.

#Synced-CL 49722754

diff --git a/Engine/Source/Runtime/Engine/Private/InstanceData/InstanceDataManager.cpp b/Engine/Source/Runtime/Engine/Private/InstanceData/InstanceDataManager.cpp
--- a/Engine/Source/Runtime/Engine/Private/InstanceData/InstanceDataManager.cpp
+++ b/Engine/Source/Runtime/Engine/Private/InstanceData/InstanceDataManager.cpp
@@ -117,8 +117,6 @@ public:
 		FInstanceSceneDataBuffers::FAccessTag AccessTag(PointerHash(this));
 		FInstanceSceneDataBuffers::FWriteView ProxyData = InstanceSceneDataBuffers.BeginWriteAccess(AccessTag);
 
-		ProxyData.Flags = ChangeSet.Flags;
-
 		int32 PostUpdateNumInstances = ChangeSet.NumSourceInstances;
 
 		// Handle data movement, needs old & new ID maps
@@ -140,16 +138,17 @@ public:
 					const int32 FromIndex =  InstanceIdIndexMap.IdToIndex(InstanceId);
 
 					ProxyData.InstanceToPrimitiveRelative[ToIndex] = ProxyData.InstanceToPrimitiveRelative[FromIndex];
-					CondMove(ChangeSet.Flags.bHasPerInstanceCustomData, ProxyData.InstanceCustomData, FromIndex, ToIndex, ChangeSet.NumCustomDataFloats);
-					CondMove(ChangeSet.Flags.bHasPerInstanceRandom, ProxyData.InstanceRandomIDs, FromIndex, ToIndex);
-					CondMove(ChangeSet.Flags.bHasPerInstanceLMSMUVBias, ProxyData.InstanceLightShadowUVBias, FromIndex, ToIndex);
-					CondMove(ChangeSet.Flags.bHasPerInstanceSkinningData, ProxyData.InstanceSkinningData, FromIndex, ToIndex);
+					CondMove(ProxyData.Flags.bHasPerInstanceCustomData, ProxyData.InstanceCustomData, FromIndex, ToIndex, ProxyData.NumCustomDataFloats);
+					CondMove(ProxyData.Flags.bHasPerInstanceRandom, ProxyData.InstanceRandomIDs, FromIndex, ToIndex);
+					CondMove(ProxyData.Flags.bHasPerInstanceLMSMUVBias, ProxyData.InstanceLightShadowUVBias, FromIndex, ToIndex);
+					CondMove(ProxyData.Flags.bHasPerInstanceSkinningData, ProxyData.InstanceSkinningData, FromIndex, ToIndex);
 	#if WITH_EDITOR
-					CondMove(ChangeSet.Flags.bHasPerInstanceEditorData, ProxyData.InstanceEditorData, FromIndex, ToIndex);
+					CondMove(ProxyData.Flags.bHasPerInstanceEditorData, ProxyData.InstanceEditorData, FromIndex, ToIndex);
 	#endif
 				}
 			}
 		}
+		ProxyData.Flags = ChangeSet.Flags;
 
 		UpdateIdMapping(ChangeSet, FIdentityIndexRemap(), InstanceIdIndexMap);
 		check(ChangeSet.NumSourceInstances == InstanceIdIndexMap.GetMaxInstanceIndex());
diff --git a/Engine/Source/Runtime/Engine/Private/InstancedStaticMesh.cpp b/Engine/Source/Runtime/Engine/Private/InstancedStaticMesh.cpp
--- a/Engine/Source/Runtime/Engine/Private/InstancedStaticMesh.cpp
+++ b/Engine/Source/Runtime/Engine/Private/InstancedStaticMesh.cpp
@@ -4040,55 +4040,60 @@ bool UInstancedStaticMeshComponent::RemoveInstanceInternal(int32 InstanceIndex,
 	// For performance we would prefer to use RemoveAtSwap() but some old code may be relying on the old
 	// RemoveAt() behavior, since there was no explicit contract about how instance indices can move around.
 	const bool bUseRemoveAtSwap = bForceRemoveAtSwap || bSupportRemoveAtSwap || CVarISMForceRemoveAtSwap.GetValueOnGameThread() != 0;
-	if (bUseRemoveAtSwap)
-	{
-		PrimitiveInstanceDataManager.RemoveAtSwap(InstanceIndex);
-	}
-	else
-	{
-		PrimitiveInstanceDataManager.RemoveAt(InstanceIndex);
-	}
 
-	// remove instance
-	if (!InstanceAlreadyRemoved && PerInstanceSMData.IsValidIndex(InstanceIndex))
+	if (PerInstanceSMData.IsValidIndex(InstanceIndex))
 	{
-		// Navigation update required explicitly requested and if the component is relevant to navigation and registered.
-		const bool bNavigationUpdateRequired = bUpdateNavigation && bNavigationRelevant && IsRegistered();
-
-		// Note that it is done before removing the instance since partial update needs
-		// the instance's transform to dirty the area and full update will dirty the whole covered by the navigation bounds.
-		if (bNavigationUpdateRequired && SupportsPartialNavigationUpdate())
-		{
-			PartialNavigationUpdate(InstanceIndex);
-		}
-
 		if (bUseRemoveAtSwap)
 		{
-			PerInstanceSMData.RemoveAtSwap(InstanceIndex, EAllowShrinking::No);
-			PerInstanceSMCustomData.RemoveAtSwap(InstanceIndex * NumCustomDataFloats, NumCustomDataFloats, EAllowShrinking::No);
+			PrimitiveInstanceDataManager.RemoveAtSwap(InstanceIndex);
 		}
 		else
 		{
-			PerInstanceSMData.RemoveAt(InstanceIndex);
-			PerInstanceSMCustomData.RemoveAt(InstanceIndex * NumCustomDataFloats, NumCustomDataFloats);
+			PrimitiveInstanceDataManager.RemoveAt(InstanceIndex);
 		}
-
-		if (bNavigationUpdateRequired)
+		
+		// remove instance
+		if (!InstanceAlreadyRemoved)
 		{
-			// If it's the last instance, unregister the component since component with no instances are not registered.
-			// (because of GetInstanceCount() > 0 in UInstancedStaticMeshComponent::IsNavigationRelevant())
-			if (GetInstanceCount() == 0)
+			// Navigation update required explicitly requested and if the component is relevant to navigation and registered.
+			const bool bNavigationUpdateRequired = bUpdateNavigation && bNavigationRelevant && IsRegistered();
+
+			// Note that it is done before removing the instance since partial update needs
+			// the instance's transform to dirty the area and full update will dirty the whole covered by the navigation bounds.
+			if (bNavigationUpdateRequired && SupportsPartialNavigationUpdate())
 			{
-				bNavigationRelevant = false;
-				FNavigationSystem::UnregisterComponent(*this);
+				PartialNavigationUpdate(InstanceIndex);
 			}
-			else if (!SupportsPartialNavigationUpdate())
+
+			if (bUseRemoveAtSwap)
 			{
-				FullNavigationUpdate();
+				PerInstanceSMData.RemoveAtSwap(InstanceIndex, EAllowShrinking::No);
+				PerInstanceSMCustomData.RemoveAtSwap(InstanceIndex * NumCustomDataFloats, NumCustomDataFloats, EAllowShrinking::No);
+			}
+			else
+			{
+				PerInstanceSMData.RemoveAt(InstanceIndex);
+				PerInstanceSMCustomData.RemoveAt(InstanceIndex * NumCustomDataFloats, NumCustomDataFloats);
+			}
+
+			if (bNavigationUpdateRequired)
+			{
+				// If it's the last instance, unregister the component since component with no instances are not registered.
+				// (because of GetInstanceCount() > 0 in UInstancedStaticMeshComponent::IsNavigationRelevant())
+				if (GetInstanceCount() == 0)
+				{
+					bNavigationRelevant = false;
+					FNavigationSystem::UnregisterComponent(*this);
+				}
+				else if (!SupportsPartialNavigationUpdate())
+				{
+					FullNavigationUpdate();
+				}
 			}
 		}
 	}
-	if (bHasPreviousTransforms)
+
+	if (bHasPreviousTransforms && PerInstancePrevTransform.IsValidIndex(InstanceIndex))
 	{
 		if (bUseRemoveAtSwap)
 		{
diff --git a/Engine/Source/Runtime/Engine/Private/InstancedStaticMesh/ISMInstanceDataManager.cpp b/Engine/Source/Runtime/Engine/Private/InstancedStaticMesh/ISMInstanceDataManager.cpp
--- a/Engine/Source/Runtime/Engine/Private/InstancedStaticMesh/ISMInstanceDataManager.cpp
+++ b/Engine/Source/Runtime/Engine/Private/InstancedStaticMesh/ISMInstanceDataManager.cpp
@@ -436,6 +436,9 @@ void FPrimitiveInstanceDataManager::InitChangeSet(const FChangeDesc &ChangeDesc,
 	bAnyEditorDataChanged = false;
 #endif	
 
+	// trigger full copy of custom data if num changed (because this zero-fills the custom data array).
+	ChangeSet.ForceFullFlags.bHasPerInstanceCustomData = ChangeSet.ForceFullFlags.bHasPerInstanceCustomData || bNumCustomDataChanged;
+
 	bAnyInstanceChange = false;
 	bTransformChangedAllInstances = false;
 	bNumCustomDataChanged = false;
@@ -587,7 +590,7 @@ bool FPrimitiveInstanceDataManager::FlushChanges(FInstanceUpdateComponentDesc &&
 		ChangeDesc.bPrimitiveTransformChanged = true;
 	}
 
-	// Init previous LocalToWorld to current on inital update
+	// Init previous LocalToWorld to current on initial update
 	FMatrix PrevPrimitiveLocalToWorld = GetState() == ETrackingState::Initial ? ComponentData.PrimitiveLocalToWorld : PrimitiveLocalToWorld;
 	// Update the tracked state
 	PrimitiveLocalToWorld = ComponentData.PrimitiveLocalToWorld;
diff --git a/Engine/Source/Runtime/Engine/Private/InstancedStaticMesh/ISMInstanceDataSceneProxy.cpp b/Engine/Source/Runtime/Engine/Private/InstancedStaticMesh/ISMInstanceDataSceneProxy.cpp
--- a/Engine/Source/Runtime/Engine/Private/InstancedStaticMesh/ISMInstanceDataSceneProxy.cpp
+++ b/Engine/Source/Runtime/Engine/Private/InstancedStaticMesh/ISMInstanceDataSceneProxy.cpp
@@ -8,7 +8,7 @@
 #include "Rendering/RenderingSpatialHash.h"
 #include "Rendering/MotionVectorSimulation.h"
 
-//UE_DISABLE_OPTIMIZATION
+// UE_DISABLE_OPTIMIZATION
 
 DEFINE_LOG_CATEGORY(LogInstanceProxy);
 
@@ -339,7 +339,6 @@ void FISMCInstanceDataSceneProxy::Update(FISMInstanceUpdateChangeSet&& ChangeSet
 	FInstanceSceneDataBuffers::FAccessTag AccessTag(PointerHash(this));
 	FInstanceSceneDataBuffers::FWriteView ProxyData = InstanceSceneDataBuffers.BeginWriteAccess(AccessTag);
 
-	ProxyData.Flags = ChangeSet.Flags;
 
 	int32 PostUpdateNumInstances = ChangeSet.PostUpdateNumInstances;
 
@@ -362,15 +361,16 @@ void FISMCInstanceDataSceneProxy::Update(FISMInstanceUpdateChangeSet&& ChangeSet
 				const int32 FromIndex =  InstanceIdIndexMap.IdToIndex(InstanceId);
 
 				ProxyData.InstanceToPrimitiveRelative[ToIndex] = ProxyData.InstanceToPrimitiveRelative[FromIndex];
-				CondMove(ChangeSet.Flags.bHasPerInstanceCustomData, ProxyData.InstanceCustomData, FromIndex, ToIndex, ChangeSet.NumCustomDataFloats);
-				CondMove(ChangeSet.Flags.bHasPerInstanceRandom, ProxyData.InstanceRandomIDs, FromIndex, ToIndex);
-				CondMove(ChangeSet.Flags.bHasPerInstanceLMSMUVBias, ProxyData.InstanceLightShadowUVBias, FromIndex, ToIndex);
+				CondMove(ProxyData.Flags.bHasPerInstanceCustomData, ProxyData.InstanceCustomData, FromIndex, ToIndex, ProxyData.NumCustomDataFloats);
+				CondMove(ProxyData.Flags.bHasPerInstanceRandom, ProxyData.InstanceRandomIDs, FromIndex, ToIndex);
+				CondMove(ProxyData.Flags.bHasPerInstanceLMSMUVBias, ProxyData.InstanceLightShadowUVBias, FromIndex, ToIndex);
 #if WITH_EDITOR
-				CondMove(ChangeSet.Flags.bHasPerInstanceEditorData, ProxyData.InstanceEditorData, FromIndex, ToIndex);
+				CondMove(ProxyData.Flags.bHasPerInstanceEditorData, ProxyData.InstanceEditorData, FromIndex, ToIndex);
 #endif
 			}
 		}
 	}
+	ProxyData.Flags = ChangeSet.Flags;
 
 	UpdateIdMapping(ChangeSet, FIdentityIndexRemap());
 	check(ChangeSet.PostUpdateNumInstances == InstanceIdIndexMap.GetMaxInstanceIndex());
diff --git a/Engine/Source/Runtime/Engine/Public/InstancedStaticMesh/ISMInstanceUpdateChangeSet.h b/Engine/Source/Runtime/Engine/Public/InstancedStaticMesh/ISMInstanceUpdateChangeSet.h
--- a/Engine/Source/Runtime/Engine/Public/InstancedStaticMesh/ISMInstanceUpdateChangeSet.h
+++ b/Engine/Source/Runtime/Engine/Public/InstancedStaticMesh/ISMInstanceUpdateChangeSet.h
@@ -18,7 +18,7 @@ public:
 	FInstanceAttributeTracker::FDeltaRange<FInstanceAttributeTracker::EFlag::CustomDataChanged> GetCustomDataDelta() const 
 	{ 
 		// Force empty range if no custom data
-		return GetDelta<FInstanceAttributeTracker::EFlag::CustomDataChanged>(NumCustomDataFloats == 0 || !Flags.bHasPerInstanceCustomData);
+		return GetDelta<FInstanceAttributeTracker::EFlag::CustomDataChanged>(NumCustomDataFloats == 0 || !Flags.bHasPerInstanceCustomData, ForceFullFlags.bHasPerInstanceCustomData);
 	}
 
 	FIdentityDeltaRange GetInstanceLightShadowUVBiasDelta() const
