Add helper functions to check if graphics or compute PSO have been precached already and caches the state for faster retrieval - if it's not precached yet then it will kick the precache request once and used cached request ID to check on next call if PSO is compiled

#Synced-CL 49719249

diff --git a/Engine/Source/Runtime/Engine/Private/PSOPrecache.cpp b/Engine/Source/Runtime/Engine/Private/PSOPrecache.cpp
--- a/Engine/Source/Runtime/Engine/Private/PSOPrecache.cpp
+++ b/Engine/Source/Runtime/Engine/Private/PSOPrecache.cpp
@@ -259,4 +259,140 @@ bool FPSOPrecacheComponentData::CheckPSOPrecachingAndBoostPriority(EPSOPrecacheP
 #else
 	return false;
 #endif
+}
+
+#if UE_WITH_PSO_PRECACHING
+
+class FPSOPrecacheStateCache
+{
+public:
+
+	bool PrecachePSO(const FGraphicsPipelineStateInitializer& InInitializer)
+	{	
+		uint64 StatePrecachePSOHash = RHIComputeStatePrecachePSOHash(InInitializer);
+
+		auto StartPSOPrecache = [InInitializer, StatePrecachePSOHash]()
+		{
+			FGraphicsPipelineStateInitializer Initializer = InInitializer;
+			Initializer.StatePrecachePSOHash = StatePrecachePSOHash;
+			return PipelineStateCache::PrecacheGraphicsPipelineState(Initializer, EPSOPrecacheInMemoryMode::Uncached);
+		};
+		return PrecachePSOInternal(StatePrecachePSOHash, StartPSOPrecache);			
+	}
+
+	bool PrecachePSO(const FComputePipelineStateInitializer& InInitializer)
+	{	
+		uint64 ComputePSOHash = GetTypeHash(InInitializer.ComputeShader->GetHash());
+		
+		auto StartPSOPrecache = [InInitializer]()
+		{
+			return PipelineStateCache::PrecacheComputePipelineState(InInitializer, EPSOPrecacheInMemoryMode::Uncached);
+		};
+		return PrecachePSOInternal(ComputePSOHash, StartPSOPrecache);			
+	}
+
+private:
+	
+	template<typename FStartPrecacheFunction>
+	bool PrecachePSOInternal(uint64 PSOInitializerHash, FStartPrecacheFunction StartPrecacheFunction)
+	{	
+		FPrecacheStateAndRequestData ActiveRequestData;
+
+		// Check if PSO is already cached with read only lock
+		{
+			FRWScopeLock ReadLock(RWLock, SLT_ReadOnly);
+			FPrecacheStateAndRequestData* StateAndRequestData = CachedPSOPrecacheResults.Find(PSOInitializerHash);
+			if (StateAndRequestData)
+			{
+				ActiveRequestData = *StateAndRequestData;
+			}
+		}
+
+		// Early out when done
+		if (ActiveRequestData.bIsPrecached)
+		{
+			return true;
+		}		
+
+		// Not requested yet?
+		if (!ActiveRequestData.RequestData.RequestID.IsValid())
+		{
+			// Start the PSO precache request and boost to highest priority (multiple thread could in theory do this which is 'fine')
+			FPrecacheStateAndRequestData NewStateAndRequestData;
+			NewStateAndRequestData.RequestData = StartPrecacheFunction();
+			NewStateAndRequestData.bIsPrecached = !PipelineStateCache::IsPrecaching(NewStateAndRequestData.RequestData.RequestID);
+			if (!NewStateAndRequestData.bIsPrecached)
+			{
+				PipelineStateCache::BoostPrecachePriority(EPSOPrecachePriority::Highest, NewStateAndRequestData.RequestData.RequestID);
+			}
+			ActiveRequestData.bIsPrecached = NewStateAndRequestData.bIsPrecached;
+			
+			// Add with write lock
+			{
+				FRWScopeLock WriteLock(RWLock, SLT_Write);
+			
+				// check if it's already in the cache now that we have a write lock
+				FPrecacheStateAndRequestData* StateAndRequestData = CachedPSOPrecacheResults.Find(PSOInitializerHash);
+				if (StateAndRequestData)
+				{
+					check(StateAndRequestData->RequestData == NewStateAndRequestData.RequestData);
+				}
+				else
+				{
+					CachedPSOPrecacheResults.Add(PSOInitializerHash, NewStateAndRequestData);
+				}
+			}			
+		}
+		else 
+		{
+			// Check the current state again
+			if (!PipelineStateCache::IsPrecaching(ActiveRequestData.RequestData.RequestID))
+			{
+				// Update the cache state with write lock 
+				// Multiple thread can write this state after each other but that's fine because multiple threads can try
+				// and precache the same PSO initializer and thus should get the same hash with StateAndRequestData
+				FRWScopeLock WriteLock(RWLock, SLT_Write);
+				FPrecacheStateAndRequestData* StateAndRequestData = CachedPSOPrecacheResults.Find(PSOInitializerHash);
+				check(StateAndRequestData && StateAndRequestData->RequestData == ActiveRequestData.RequestData);
+				StateAndRequestData->bIsPrecached = true;
+				StateAndRequestData->RequestData.AsyncCompileEvent = nullptr;
+
+				ActiveRequestData.bIsPrecached = true;
+			}
+		}
+
+		return ActiveRequestData.bIsPrecached;				
+	}
+
+	struct FPrecacheStateAndRequestData
+	{
+		FPSOPrecacheRequestResult RequestData;
+		bool bIsPrecached = false;
+	};
+
+	FRWLock RWLock;
+	TMap <uint64, FPrecacheStateAndRequestData> CachedPSOPrecacheResults;
+};
+
+FPSOPrecacheStateCache GPSOPrecacheStateCache;
+
+#endif // UE_WITH_PSO_PRECACHING
+
+
+bool PrecachePSO(const FGraphicsPipelineStateInitializer& GraphicsPSO)
+{
+#if UE_WITH_PSO_PRECACHING
+	return GPSOPrecacheStateCache.PrecachePSO(GraphicsPSO);
+#else
+	return true;
+#endif
+}
+
+bool PrecachePSO(const FComputePipelineStateInitializer& ComputePSO)
+{
+#if UE_WITH_PSO_PRECACHING
+	return GPSOPrecacheStateCache.PrecachePSO(ComputePSO);
+#else
+	return true;
+#endif
 }
\ No newline at end of file
diff --git a/Engine/Source/Runtime/Engine/Public/PSOPrecache.h b/Engine/Source/Runtime/Engine/Public/PSOPrecache.h
--- a/Engine/Source/Runtime/Engine/Public/PSOPrecache.h
+++ b/Engine/Source/Runtime/Engine/Public/PSOPrecache.h
@@ -412,6 +412,16 @@ extern ENGINE_API bool IsPSOShaderPreloadingEnabled();
  */
 extern ENGINE_API void BoostPrecachedPSORequestsOnDraw(const class FPrimitiveSceneInfo* SceneInfo);
 
+/**
+ * Precache graphics PSO and return true if PSO is fully precached
+ */
+extern ENGINE_API bool PrecachePSO(const FGraphicsPipelineStateInitializer& GraphicsPSO);
+
+/**
+* Precache compute PSO and return true if PSO is fully precached
+ */
+extern ENGINE_API bool PrecachePSO(const FComputePipelineStateInitializer& ComputePSO);
+
 /**
  * Start the actual async PSO precache request from the given list of initializers
  */
