Fixed an issue where Viewmode LODColoration was not working on builds with cooked assets unless ODSC was used

#Synced-CL 49758588

diff --git a/Engine/Source/Runtime/Engine/Private/DebugViewModeHelpers.cpp b/Engine/Source/Runtime/Engine/Private/DebugViewModeHelpers.cpp
--- a/Engine/Source/Runtime/Engine/Private/DebugViewModeHelpers.cpp
+++ b/Engine/Source/Runtime/Engine/Private/DebugViewModeHelpers.cpp
@@ -20,6 +20,7 @@
 #include "RenderingThread.h"
 #include "ShaderCompiler.h"
 #include "RenderUtils.h"
+#include "ODSC/ODSCManager.h"
 
 #define LOCTEXT_NAMESPACE "LogDebugViewMode"
 
@@ -134,16 +135,19 @@ bool AllowDebugViewShaderMode(EDebugViewShaderMode ShaderMode, EShaderPlatform P
 		return false;
 	}
 
+	if (FODSCManager::IsODSCEnabled() && FODSCManager::IsODSCActive())
+	{
 #if WITH_EDITORONLY_DATA
-	bool bRequireODSCShaders = IsDebugViewShaderModeODSCOnly();
+		bool bRequireODSCShaders = IsDebugViewShaderModeODSCOnly();
 #else
-	bool bRequireODSCShaders = true;
+		bool bRequireODSCShaders = true;
 #endif
 
-	// Don't allow debug modes if we require ODSC but can't compile ODSC shaders.
-	if (bRequireODSCShaders && !ShouldCompileODSCOnlyShaders())
-	{
-		return false;
+		// Don't allow debug modes if we require ODSC but can't compile ODSC shaders.
+		if (bRequireODSCShaders && !ShouldCompileODSCOnlyShaders())
+		{
+			return false;
+		}
 	}
 
 	return true;
diff --git a/Engine/Source/Runtime/Renderer/Private/DebugViewModeRendering.cpp b/Engine/Source/Runtime/Renderer/Private/DebugViewModeRendering.cpp
--- a/Engine/Source/Runtime/Renderer/Private/DebugViewModeRendering.cpp
+++ b/Engine/Source/Runtime/Renderer/Private/DebugViewModeRendering.cpp
@@ -167,6 +167,11 @@ static bool ShouldCompileDebugViewModeShaders(const FMeshMaterialShaderPermutati
 		return false;
 	}
 
+	if (AllowDebugViewmodes())
+	{
+		return true;
+	}
+
 	EShaderPermutationFlags Flags = IsDebugViewShaderModeODSCOnly() ? EShaderPermutationFlags::IsODSCOnly : EShaderPermutationFlags::HasEditorOnlyData | EShaderPermutationFlags::IsODSCOnly;
 	return EnumHasAnyFlags(Parameters.Flags, Flags);
 }
