- Make sure padding on texture coordinate buffer with uneven amount of TC sets is taking into account because ResourceCount in CreateRHIBuffer (it doesn't have the element size so it also can't really use it)
- Restore shared setup for TC elements in input layout and force all entries back to half4/float4 to reduce amount of (unneeded) extra PSOs compile due to permutations on different amount of TCs used

#Synced-CL 49719249

diff --git a/Engine/Source/Runtime/Engine/Private/LocalVertexFactory.cpp b/Engine/Source/Runtime/Engine/Private/LocalVertexFactory.cpp
--- a/Engine/Source/Runtime/Engine/Private/LocalVertexFactory.cpp
+++ b/Engine/Source/Runtime/Engine/Private/LocalVertexFactory.cpp
@@ -558,22 +558,48 @@ void AddStaticMeshTextureCoordinateElements(
 	FVertexDeclarationElementList& Elements,
 	FVertexFactory::FVertexStreamList& InOutStreams)
 {
-	for (int32 CoordinateIndex = 0; CoordinateIndex < TextureCoordinates.Num(); ++CoordinateIndex)
+	if (TextureCoordinates.IsEmpty())
 	{
-		Elements.Add(FVertexFactory::AccessStreamComponent(
-			TextureCoordinates[CoordinateIndex],
-			BaseTexCoordAttribute + CoordinateIndex,
-			InOutStreams
-		));
+		return;
 	}
-
-	for (int32 CoordinateIndex = TextureCoordinates.Num(); CoordinateIndex < MAX_STATIC_TEXCOORDS / 2; ++CoordinateIndex)
+	
+#if DO_CHECK 
+	// Validate all texture coordinates share the same stream data
+	for (int32 TextureCoordIndex = 1; TextureCoordIndex < TextureCoordinates.Num(); ++TextureCoordIndex)
+	{
+		checkf(TextureCoordinates[0].VertexBuffer == TextureCoordinates[TextureCoordIndex].VertexBuffer, TEXT("Texture coordinate %d has a different vertex buffer than index 0"), TextureCoordIndex);
+		checkf(TextureCoordinates[0].Stride == TextureCoordinates[TextureCoordIndex].Stride, TEXT("Texture coordinate %d has a different stride than index 0"), TextureCoordIndex);
+		checkf(TextureCoordinates[0].StreamOffset == TextureCoordinates[TextureCoordIndex].StreamOffset, TEXT("Texture coordinate %d has a different stream offset than index 0"), TextureCoordIndex);
+		checkf(TextureCoordinates[0].VertexStreamUsage == TextureCoordinates[TextureCoordIndex].VertexStreamUsage, TEXT("Texture coordinate %d has a different vertex stream usage than index 0"), TextureCoordIndex);
+	}
+#endif // DO_CHECK
+	
+	EVertexElementType FirstElementType = TextureCoordinates[0].Type;
+	
+	EVertexElementType UVDoubleWideVertexElementType = VET_None;
+	uint32 DoubleUVSizeInBytes = 0;
+	if (FirstElementType == VET_Float2 || FirstElementType == VET_Float4)
+	{
+		DoubleUVSizeInBytes = sizeof(TStaticMeshVertexUVsTypeSelector<EStaticMeshVertexUVType::HighPrecision>::UVsTypeT) * 2;
+		UVDoubleWideVertexElementType = VET_Float4;
+	}
+	else
+	{
+		check(FirstElementType == VET_Half2 || FirstElementType == VET_Half4)
+		DoubleUVSizeInBytes = sizeof(TStaticMeshVertexUVsTypeSelector<EStaticMeshVertexUVType::Default>::UVsTypeT) * 2;
+		UVDoubleWideVertexElementType = VET_Half4;
+	}
+	
+	uint8 TexCoordStreamIndex = FVertexFactory::AddVertexStream(TextureCoordinates[0], InOutStreams);
+	
+	for (int32 CoordinateIndex = 0; CoordinateIndex < MAX_STATIC_TEXCOORDS / 2; ++CoordinateIndex)
 	{
-		Elements.Add(FVertexFactory::AccessStreamComponent(
-			TextureCoordinates[TextureCoordinates.Num() - 1],
+		Elements.Add(FVertexElement(
+			TexCoordStreamIndex,
+			DoubleUVSizeInBytes * CoordinateIndex,
+			UVDoubleWideVertexElementType,
 			BaseTexCoordAttribute + CoordinateIndex,
-			InOutStreams
-		));
+			TextureCoordinates[0].Stride));
 	}
 }
 
diff --git a/Engine/Source/Runtime/Engine/Private/Rendering/StaticMeshVertexBuffer.cpp b/Engine/Source/Runtime/Engine/Private/Rendering/StaticMeshVertexBuffer.cpp
--- a/Engine/Source/Runtime/Engine/Private/Rendering/StaticMeshVertexBuffer.cpp
+++ b/Engine/Source/Runtime/Engine/Private/Rendering/StaticMeshVertexBuffer.cpp
@@ -11,6 +11,45 @@
 #include "RHIResourceReplace.h"
 #include "StaticMeshVertexData.h"
 
+	
+static uint32 GetNumTexCoordInVertexBuffer(uint32 NumTexCoords, uint32 NumVertices)
+{
+	// Increase size with single vertex in case we have odd number of texcoords because the vertex declaration
+	// will declare the stream as half4/float4 while it only contains a half2/float2 and the last element can in theory
+	// be out of bounds then
+	// Vertex streams are all created with half4/float4 to reduce the number of PSO permutations coming from different texcoords in the vertex declaration
+	uint32 Padding = NumTexCoords % 2;
+	return (NumVertices + Padding) * NumTexCoords;
+}
+
+/** The implementation of the static mesh texture coordinate vertex data storage type which takes extra padding into account. */
+template<typename VertexDataType>
+class TStaticTexCoordVertexData :
+	public TStaticMeshVertexData<VertexDataType>
+{
+private:
+
+	uint32 NumTexCoords;
+
+public:
+
+	TStaticTexCoordVertexData(uint32 InNumTexCoords, bool InNeedsCPUAccess=false)
+		: TStaticMeshVertexData<VertexDataType>(InNeedsCPUAccess), NumTexCoords(InNumTexCoords)
+	{
+	}
+
+	void ResizeBuffer(uint32 NumVertices, EResizeBufferFlags BufferFlags = EResizeBufferFlags::None) override
+	{
+		uint32 NumTexCoordData = GetNumTexCoordInVertexBuffer(NumTexCoords, NumVertices);
+		TStaticMeshVertexData<VertexDataType>::ResizeBuffer(NumTexCoordData);
+	}
+
+	void Empty(uint32 NumVertices) override
+	{
+		TStaticMeshVertexData<VertexDataType>::Empty(GetNumTexCoordInVertexBuffer(NumTexCoords, NumVertices));
+	}
+};
+
 FStaticMeshVertexBuffer::FStaticMeshVertexBuffer() :
 	TangentsData(nullptr),
 	TexcoordData(nullptr),
@@ -54,7 +93,7 @@ void FStaticMeshVertexBuffer::Init(uint32 InNumVertices, uint32 InNumTexCoords,
 	// Allocate the vertex data buffer.
 	TangentsData->ResizeBuffer(NumVertices);
 	TangentsDataPtr = NumVertices ? TangentsData->GetDataPointer() : nullptr;
-	TexcoordData->ResizeBuffer(NumVertices * GetNumTexCoords());
+	TexcoordData->ResizeBuffer(NumVertices);
 	TexcoordDataPtr = NumVertices ? TexcoordData->GetDataPointer() : nullptr;
 }
 
@@ -118,38 +157,13 @@ void FStaticMeshVertexBuffer::Init(const FStaticMeshVertexBuffer& InVertexBuffer
 			check(GetNumTexCoords() == InVertexBuffer.GetNumTexCoords());
 			const uint8* InData = InVertexBuffer.TexcoordDataPtr;
 
-			TexcoordData->ResizeBuffer(NumVertices * GetNumTexCoords());
+			TexcoordData->ResizeBuffer(NumVertices);
 			TexcoordDataPtr = TexcoordData->GetDataPointer();
 			FMemory::Memcpy(TexcoordDataPtr, InData, TexcoordData->GetStride() * NumVertices * GetNumTexCoords());
 		}
 	}
 }
 
-void FStaticMeshVertexBuffer::ConvertHalfTexcoordsToFloat(const uint8* InData)
-{
-	check(TexcoordData);
-	SetUseFullPrecisionUVs(true);
-
-	FStaticMeshVertexDataInterface* OriginalTexcoordData = TexcoordData;
-
-	typedef TStaticMeshVertexUVsDatum<typename TStaticMeshVertexUVsTypeSelector<EStaticMeshVertexUVType::HighPrecision>::UVsTypeT> UVType;
-	TexcoordData = new TStaticMeshVertexData<UVType>(OriginalTexcoordData->GetAllowCPUAccess());
-	TexcoordData->ResizeBuffer(NumVertices * GetNumTexCoords());
-	TexcoordDataPtr = TexcoordData->GetDataPointer();
-	TexcoordStride = sizeof(UVType);
-
-	FVector2f* DestTexcoordDataPtr = (FVector2f*)TexcoordDataPtr;
-	FVector2DHalf* SourceTexcoordDataPtr = (FVector2DHalf*)(InData ? InData : OriginalTexcoordData->GetDataPointer());
-	for (uint32 i = 0; i < NumVertices * GetNumTexCoords(); i++)
-	{
-		*DestTexcoordDataPtr++ = *SourceTexcoordDataPtr++;
-	}
-
-	delete OriginalTexcoordData;
-	OriginalTexcoordData = nullptr;
-}
-
-
 void FStaticMeshVertexBuffer::AppendVertices( const FStaticMeshBuildVertex* Vertices, const uint32 NumVerticesToAppend, bool bUseBackwardsCompatibleF16TruncUVs)
 {
 	if ((TangentsData == nullptr || TexcoordData == nullptr) && NumVerticesToAppend > 0)
@@ -169,7 +183,7 @@ void FStaticMeshVertexBuffer::AppendVertices( const FStaticMeshBuildVertex* Vert
 		NumVertices += NumVerticesToAppend;
 
 		TangentsData->ResizeBuffer(NumVertices);
-		TexcoordData->ResizeBuffer(NumVertices * GetNumTexCoords());
+		TexcoordData->ResizeBuffer(NumVertices);
 
 		if( NumVertices > 0 )
 		{
@@ -284,12 +298,7 @@ FBufferRHIRef FStaticMeshVertexBuffer::CreateTangentsRHIBuffer(FRHICommandListBa
 
 FBufferRHIRef FStaticMeshVertexBuffer::CreateTexCoordRHIBuffer(FRHICommandListBase& RHICmdList)
 {
-	// Increase size with single vertex in case we have odd number of texcoords because the vertex declaration
-	// will declare the stream as half4/float4 while it only contains a half2/float2 and the last element can in theory
-	// be out of bounds then
-	// Vertex streams are all created with half4/float4 to reduce the number of PSO permutations coming from different texcoords in the vertex declaration
-	uint32 Padding = (GetNumTexCoords() % 2 == 0) ? 0 : 1;
-	uint32 ResourceCount = (GetNumVertices() + Padding) * GetNumTexCoords();
+	uint32 ResourceCount = GetNumTexCoordInVertexBuffer(GetNumTexCoords(), GetNumVertices());
 	return CreateRHIBuffer(RHICmdList, TexcoordData, ResourceCount, BUF_Static | BUF_ShaderResource, TEXT("TexCoordRHIBuffer"));
 }
 
@@ -389,13 +398,13 @@ void FStaticMeshVertexBuffer::AllocateData(bool bNeedsCPUAccess /*= true*/)
 	{
 		typedef TStaticMeshVertexUVsDatum<typename TStaticMeshVertexUVsTypeSelector<EStaticMeshVertexUVType::HighPrecision>::UVsTypeT> UVType;
 		TexcoordStride = sizeof(UVType);
-		TexcoordData = new TStaticMeshVertexData<UVType>(bNeedsCPUAccess);
+		TexcoordData = new TStaticTexCoordVertexData<UVType>(NumTexCoords, bNeedsCPUAccess);
 	}
 	else
 	{
 		typedef TStaticMeshVertexUVsDatum<typename TStaticMeshVertexUVsTypeSelector<EStaticMeshVertexUVType::Default>::UVsTypeT> UVType;
 		TexcoordStride = sizeof(UVType);
-		TexcoordData = new TStaticMeshVertexData<UVType>(bNeedsCPUAccess);
+		TexcoordData = new TStaticTexCoordVertexData<UVType>(NumTexCoords, bNeedsCPUAccess);
 	}
 }
 
diff --git a/Engine/Source/Runtime/Engine/Public/Rendering/StaticMeshVertexBuffer.h b/Engine/Source/Runtime/Engine/Public/Rendering/StaticMeshVertexBuffer.h
--- a/Engine/Source/Runtime/Engine/Public/Rendering/StaticMeshVertexBuffer.h
+++ b/Engine/Source/Runtime/Engine/Public/Rendering/StaticMeshVertexBuffer.h
@@ -534,11 +534,6 @@ private:
 
 	/** Allocates the vertex data storage type. */
 	void AllocateData(bool bNeedsCPUAccess = true);
-
-	/** Convert half float data to full float if the HW requires it.
-	* @param InData - optional half float source data to convert into full float texture coordinate buffer. if null, convert existing half float texture coordinates to a new float buffer.
-	*/
-	void ConvertHalfTexcoordsToFloat(const uint8* InData);
-
+	   
 	void InitTangentAndTexCoordStrides();
 };
