[WorldPartition][MiniMap] Added Included Data Layer assets option to the WorldPartitionMiniMap

#Synced-CL 49728154

diff --git a/Engine/Source/Editor/UnrealEd/Private/WorldPartition/WorldPartitionMiniMapBuilder.cpp b/Engine/Source/Editor/UnrealEd/Private/WorldPartition/WorldPartitionMiniMapBuilder.cpp
--- a/Engine/Source/Editor/UnrealEd/Private/WorldPartition/WorldPartitionMiniMapBuilder.cpp
+++ b/Engine/Source/Editor/UnrealEd/Private/WorldPartition/WorldPartitionMiniMapBuilder.cpp
@@ -86,6 +86,7 @@ bool UWorldPartitionMiniMapBuilder::PreRun(UWorld* World, FPackageSourceControlH
 		WorldToMinimap *= FScaleMatrix(Scaling);
 	}
 
+	IncludedDataLayers = WorldMiniMap->IncludedDataLayers;
 	ExcludedDataLayers = WorldMiniMap->ExcludedDataLayers;
 
 	return true;
diff --git a/Engine/Source/Runtime/Engine/Public/WorldPartition/WorldPartitionMiniMap.h b/Engine/Source/Runtime/Engine/Public/WorldPartition/WorldPartitionMiniMap.h
--- a/Engine/Source/Runtime/Engine/Public/WorldPartition/WorldPartitionMiniMap.h
+++ b/Engine/Source/Runtime/Engine/Public/WorldPartition/WorldPartitionMiniMap.h
@@ -50,6 +50,10 @@ public:
 	UPROPERTY(VisibleAnywhere, Category = WorldPartitionMiniMap)
 	TObjectPtr<UTexture2D> MiniMapTexture;
 
+	/* Datalayers included in MiniMap rendering */
+	UPROPERTY(EditAnywhere, Category = WorldPartitionMiniMap)
+	TSet<TObjectPtr<UDataLayerAsset>> IncludedDataLayers;
+
 	/* Datalayers excluded from MiniMap rendering */
 	UPROPERTY(EditAnywhere, Category = WorldPartitionMiniMap)
 	TSet<TObjectPtr<UDataLayerAsset>> ExcludedDataLayers;
