[First Person] Implemented TransformFirstPersonToWorld() utility function.

#Synced-CL 49760438

diff --git a/Engine/Source/Runtime/Engine/Classes/Camera/CameraTypes.h b/Engine/Source/Runtime/Engine/Classes/Camera/CameraTypes.h
--- a/Engine/Source/Runtime/Engine/Classes/Camera/CameraTypes.h
+++ b/Engine/Source/Runtime/Engine/Classes/Camera/CameraTypes.h
@@ -238,6 +238,12 @@ public:
 	 */
 	ENGINE_API FVector TransformWorldToFirstPerson(const FVector& WorldPosition, bool bIgnoreFirstPersonScale) const;
 
+	/**
+	 * Transforms a first person space location back into world space.
+	 * bScaleWasIgnoredOnInput should be set to true if the first person position was calculated with TransformWorldToFirstPerson() and bIgnoreFirstPersonScale set to true.
+	 */
+	ENGINE_API FVector TransformFirstPersonToWorld(const FVector& FirstPersonPosition, bool bScaleWasIgnoredOnInput) const;
+
 	/**
 	 * Correction factor to apply to the first person transform used on primitives tagged as "IsFirstPerson" to achieve a first person specific field of view.
 	 * It is computed as tan(SceneFOVRadians * 0.5) / tan(FirstPersonFOVRadians * 0.5).
diff --git a/Engine/Source/Runtime/Engine/Classes/Kismet/GameplayStatics.h b/Engine/Source/Runtime/Engine/Classes/Kismet/GameplayStatics.h
--- a/Engine/Source/Runtime/Engine/Classes/Kismet/GameplayStatics.h
+++ b/Engine/Source/Runtime/Engine/Classes/Kismet/GameplayStatics.h
@@ -1472,6 +1472,15 @@ public:
 	UFUNCTION(BlueprintPure, Category = "Camera")
 	static ENGINE_API FVector TransformWorldToFirstPerson(const FMinimalViewInfo& ViewInfo, const FVector& WorldPosition, bool bIgnoreFirstPersonScale);
 
+	/**
+	 * Transforms a first person space location back into world space. This function is the inverse of TransformWorldToFirstPerson.
+	 * @param ViewInfo					FMinimalViewInfo struct holding the first person camera parameters.
+	 * @param FirstPersonPosition		First person space position to transform.
+	 * @param bScaleWasIgnoredOnInput	Should be set to true if FirstPersonPosition was computed with TransformWorldToFirstPerson and the bIgnoreFirstPersonScale parameter set to true.
+	 */
+	UFUNCTION(BlueprintPure, Category = "Camera")
+	static ENGINE_API FVector TransformFirstPersonToWorld(const FMinimalViewInfo& ViewInfo, const FVector& FirstPersonPosition, bool bScaleWasIgnoredOnInput);
+
 	/**
 	 * Returns the View Matrix, Projection Matrix and the View x Projection Matrix for a given view
 	 * @param DesiredView			FMinimalViewInfo struct for a camera.
diff --git a/Engine/Source/Runtime/Engine/Private/Camera/CameraStackTypes.cpp b/Engine/Source/Runtime/Engine/Private/Camera/CameraStackTypes.cpp
--- a/Engine/Source/Runtime/Engine/Private/Camera/CameraStackTypes.cpp
+++ b/Engine/Source/Runtime/Engine/Private/Camera/CameraStackTypes.cpp
@@ -480,8 +480,8 @@ FVector FMinimalViewInfo::TransformWorldToFirstPerson(const FVector& WorldPositi
 		const FVector CameraRelativePosition = WorldPosition - Location;
 		const FVector ProjectedPosition = FVector::DotProduct(Forward, CameraRelativePosition) * Forward;
 		const FVector Rejection = CameraRelativePosition - ProjectedPosition;
-		const float FOVCorrectionFactor = CalculateFirstPersonFOVCorrectionFactor() - 1.0f;
-		const FVector FOVCorrectedPosition = CameraRelativePosition + Rejection * FOVCorrectionFactor;
+		const float FOVCorrectionFactor = CalculateFirstPersonFOVCorrectionFactor();
+		const FVector FOVCorrectedPosition = ProjectedPosition + Rejection * FOVCorrectionFactor;
 		const FVector ScaledPosition = FOVCorrectedPosition * FirstPersonScale;
 		const FVector Result = (bIgnoreFirstPersonScale ? FOVCorrectedPosition : ScaledPosition) + Location;
 		return Result;
@@ -489,6 +489,23 @@ FVector FMinimalViewInfo::TransformWorldToFirstPerson(const FVector& WorldPositi
 	return WorldPosition;
 }
 
+FVector FMinimalViewInfo::TransformFirstPersonToWorld(const FVector& FirstPersonPosition, bool bScaleWasIgnoredOnInput) const
+{
+	if (ProjectionMode == ECameraProjectionMode::Perspective)
+	{
+		const FVector Forward = Rotation.Vector();
+		const FVector CameraRelativePosition = FirstPersonPosition - Location;
+		const FVector UnscaledPosition = bScaleWasIgnoredOnInput || FMath::Abs(FirstPersonScale) < UE_SMALL_NUMBER ? CameraRelativePosition : CameraRelativePosition / FirstPersonScale;
+		const FVector ProjectedPosition = FVector::DotProduct(Forward, UnscaledPosition) * Forward;
+		const FVector Rejection = UnscaledPosition - ProjectedPosition;
+		const float FOVCorrectionFactor = CalculateFirstPersonFOVCorrectionFactor();
+		const FVector FOVUncorrectedPosition = FMath::Abs(FOVCorrectionFactor) > UE_SMALL_NUMBER ? (ProjectedPosition + (Rejection / FOVCorrectionFactor)) : UnscaledPosition;
+		const FVector Result = FOVUncorrectedPosition + Location;
+		return Result;
+	}
+	return FirstPersonPosition;
+}
+
 float FMinimalViewInfo::CalculateFirstPersonFOVCorrectionFactor() const
 {
 	const float HalfTanSceneFOV = FMath::Tan(FMath::DegreesToRadians(FOV) * 0.5f);
diff --git a/Engine/Source/Runtime/Engine/Private/GameplayStatics.cpp b/Engine/Source/Runtime/Engine/Private/GameplayStatics.cpp
--- a/Engine/Source/Runtime/Engine/Private/GameplayStatics.cpp
+++ b/Engine/Source/Runtime/Engine/Private/GameplayStatics.cpp
@@ -3347,6 +3347,11 @@ FVector UGameplayStatics::TransformWorldToFirstPerson(const FMinimalViewInfo& Vi
 	return ViewInfo.TransformWorldToFirstPerson(WorldPosition, bIgnoreFirstPersonScale);
 }
 
+FVector UGameplayStatics::TransformFirstPersonToWorld(const FMinimalViewInfo& ViewInfo, const FVector& FirstPersonPosition, bool bScaleWasIgnoredOnInput)
+{
+	return ViewInfo.TransformFirstPersonToWorld(FirstPersonPosition, bScaleWasIgnoredOnInput);
+}
+
 void UGameplayStatics::CalculateViewProjectionMatricesFromViewTarget(AActor* InViewTarget, FMatrix& OutViewMatrix, FMatrix& OutProjectionMatrix, FMatrix& OutViewProjectionMatrix)
 {
 	if (InViewTarget)
