- Fixed an issue where WorldDataLayers data layer instances had their 'Loaded in Editor' state reset by DataLayerEditorPerProjectUserSettings when loading a Level Instance, causing builders' IncludedDataLayers / ExcludedDataLayers settings to be ignored.
- Fixed an issue where External Data Layers 'Loaded in Editor' state was not properly restored when reading from DataLayerEditorPerProjectUserSettings.
- Deprecated GetWorldDataLayersLoadedInEditor / GetWorldDataLayersNotLoadedInEditor and replaced them with GetWorldDataLayerAssetsLoadedInEditor / GetWorldDataLayerAssetsNotLoadedInEditor.

#Synced-CL 49728154

diff --git a/Engine/Source/Runtime/Engine/Private/WorldPartition/DataLayer/DataLayerManager.cpp b/Engine/Source/Runtime/Engine/Private/WorldPartition/DataLayer/DataLayerManager.cpp
--- a/Engine/Source/Runtime/Engine/Private/WorldPartition/DataLayer/DataLayerManager.cpp
+++ b/Engine/Source/Runtime/Engine/Private/WorldPartition/DataLayer/DataLayerManager.cpp
@@ -733,10 +733,10 @@ void UDataLayerManager::GetUserLoadedInEditorStates(TArray<const UDataLayerInsta
 	OutDataLayersLoadedInEditor.Empty();
 	OutDataLayersNotLoadedInEditor.Empty();
 
-	const TArray<const UDataLayerInstance*>& SettingsDataLayersNotLoadedInEditor = GetDefault<UWorldPartitionEditorPerProjectUserSettings>()->GetWorldDataLayersNotLoadedInEditor(GetWorld());
-	const TArray<const UDataLayerInstance*>& SettingsDataLayersLoadedInEditor = GetDefault<UWorldPartitionEditorPerProjectUserSettings>()->GetWorldDataLayersLoadedInEditor(GetWorld());
+	const TArray<const UDataLayerAsset*>& SettingsDataLayerAssetsNotLoadedInEditor = GetDefault<UWorldPartitionEditorPerProjectUserSettings>()->GetWorldDataLayerAssetsNotLoadedInEditor(GetWorld());
+	const TArray<const UDataLayerAsset*>& SettingsDataLayerAssetsLoadedInEditor = GetDefault<UWorldPartitionEditorPerProjectUserSettings>()->GetWorldDataLayerAssetsLoadedInEditor(GetWorld());
 
-	ForEachDataLayerInstance([&OutDataLayersLoadedInEditor, &OutDataLayersNotLoadedInEditor, &SettingsDataLayersNotLoadedInEditor, &SettingsDataLayersLoadedInEditor](UDataLayerInstance* DataLayerInstance)
+	ForEachDataLayerInstance([&OutDataLayersLoadedInEditor, &OutDataLayersNotLoadedInEditor, &SettingsDataLayerAssetsNotLoadedInEditor, &SettingsDataLayerAssetsLoadedInEditor](UDataLayerInstance* DataLayerInstance)
 	{
 		if (DataLayerInstance->IsLoadedInEditorChangedByUserOperation())
 		{
@@ -753,11 +753,12 @@ void UDataLayerManager::GetUserLoadedInEditorStates(TArray<const UDataLayerInsta
 		}
 		else
 		{
-			if (SettingsDataLayersNotLoadedInEditor.Contains(DataLayerInstance))
+			const UDataLayerAsset* DataLayerAsset = DataLayerInstance->GetAsset();
+			if (SettingsDataLayerAssetsNotLoadedInEditor.Contains(DataLayerAsset))
 			{
 				OutDataLayersNotLoadedInEditor.Add(DataLayerInstance);
 			}
-			else if (SettingsDataLayersLoadedInEditor.Contains(DataLayerInstance))
+			else if (SettingsDataLayerAssetsLoadedInEditor.Contains(DataLayerAsset))
 			{
 				OutDataLayersLoadedInEditor.Add(DataLayerInstance);
 			}
diff --git a/Engine/Source/Runtime/Engine/Private/WorldPartition/DataLayer/WorldDataLayers.cpp b/Engine/Source/Runtime/Engine/Private/WorldPartition/DataLayer/WorldDataLayers.cpp
--- a/Engine/Source/Runtime/Engine/Private/WorldPartition/DataLayer/WorldDataLayers.cpp
+++ b/Engine/Source/Runtime/Engine/Private/WorldPartition/DataLayer/WorldDataLayers.cpp
@@ -1193,31 +1193,27 @@ void AWorldDataLayers::OnDataLayerManagerInitialized()
 		RemoveEditorDataLayers();
 	}
 
-	// Setup defaults before overriding with user settings
-	ForEachDataLayerInstance([](UDataLayerInstance* DataLayerInstance)
-	{
-		DataLayerInstance->SetIsLoadedInEditor(DataLayerInstance->IsInitiallyLoadedInEditor(), /*bFromUserChange*/false);
-		return true;
-	});
+	TSet<const UDataLayerAsset*> SettingsDataLayerAssetsLoadedInEditor;
+	TSet<const UDataLayerAsset*> SettingsDataLayerAssetsNotLoadedInEditor;
+	SettingsDataLayerAssetsLoadedInEditor.Append(GetMutableDefault<UWorldPartitionEditorPerProjectUserSettings>()->GetWorldDataLayerAssetsLoadedInEditor(GetWorld()));
+	SettingsDataLayerAssetsNotLoadedInEditor.Append(GetMutableDefault<UWorldPartitionEditorPerProjectUserSettings>()->GetWorldDataLayerAssetsNotLoadedInEditor(GetWorld()));
 
-	// Initialize DataLayer's IsLoadedInEditor based on DataLayerEditorPerProjectUserSettings
-	TArray<UDataLayerInstance*> SettingsDataLayersNotLoadedInEditor = GetMutableDefault<UWorldPartitionEditorPerProjectUserSettings>()->GetWorldDataLayersNotLoadedInEditor(GetWorld());
-	for (UDataLayerInstance* DataLayerInstance : SettingsDataLayersNotLoadedInEditor)
+	// Setup IsLoadedInEditor based on default value and DataLayerEditorPerProjectUserSettings
+	ForEachDataLayerInstance([&SettingsDataLayerAssetsLoadedInEditor, &SettingsDataLayerAssetsNotLoadedInEditor](UDataLayerInstance* DataLayerInstance)
 	{
-		if (DataLayerInstance)
+		bool bIsLoadedInEditor = DataLayerInstance->IsInitiallyLoadedInEditor();
+		const UDataLayerAsset* DataLayerAsset = DataLayerInstance->GetAsset();
+		if (SettingsDataLayerAssetsLoadedInEditor.Contains(DataLayerAsset))
 		{
-			DataLayerInstance->SetIsLoadedInEditor(false, /*bFromUserChange*/false);
+			bIsLoadedInEditor = true;
 		}
-	}
-
-	TArray<UDataLayerInstance*> SettingsDataLayersLoadedInEditor = GetMutableDefault<UWorldPartitionEditorPerProjectUserSettings>()->GetWorldDataLayersLoadedInEditor(GetWorld());
-	for (UDataLayerInstance* DataLayerInstance : SettingsDataLayersLoadedInEditor)
-	{
-		if (DataLayerInstance)
+		else if (SettingsDataLayerAssetsNotLoadedInEditor.Contains(DataLayerAsset))
 		{
-			DataLayerInstance->SetIsLoadedInEditor(true, /*bFromUserChange*/false);
+			bIsLoadedInEditor = false;
 		}
-	}
+		DataLayerInstance->SetIsLoadedInEditor(bIsLoadedInEditor, /*bFromUserChange*/false);
+		return true;
+	});
 #endif
 
 	InitializeDataLayerRuntimeStates();
diff --git a/Engine/Source/Runtime/Engine/Private/WorldPartition/WorldPartitionEditorPerProjectUserSettings.cpp b/Engine/Source/Runtime/Engine/Private/WorldPartition/WorldPartitionEditorPerProjectUserSettings.cpp
--- a/Engine/Source/Runtime/Engine/Private/WorldPartition/WorldPartitionEditorPerProjectUserSettings.cpp
+++ b/Engine/Source/Runtime/Engine/Private/WorldPartition/WorldPartitionEditorPerProjectUserSettings.cpp
@@ -69,70 +69,56 @@ TArray<FName> UWorldPartitionEditorPerProjectUserSettings::GetEditorLoadedLocati
 	return TArray<FName>();
 }
 
-TArray<const UDataLayerInstance*> UWorldPartitionEditorPerProjectUserSettings::GetWorldDataLayersNotLoadedInEditor(UWorld* InWorld) const
+TArray<const UDataLayerAsset*> UWorldPartitionEditorPerProjectUserSettings::GetWorldDataLayerAssetsNotLoadedInEditor(UWorld* InWorld) const
 {
-	TArray<const UDataLayerInstance*> NotLoadedDataLayers;
+	TArray<const UDataLayerAsset*> NotLoadedDataLayerAssets;
 
-	if (UDataLayerManager* DataLayerManager = UDataLayerManager::GetDataLayerManager(InWorld))
+	if (const FWorldPartitionPerWorldSettings* PerWorldSettings = GetWorldPartitionPerWorldSettings(InWorld))
 	{
-		if (const FWorldPartitionPerWorldSettings* PerWorldSettings = GetWorldPartitionPerWorldSettings(InWorld))
+		for (const TSoftObjectPtr<const UDataLayerAsset>& SoftDataLayerAsset : PerWorldSettings->NotLoadedDataLayers)
 		{
-			for (const TSoftObjectPtr<const UDataLayerAsset>& DataLayerAsset : PerWorldSettings->NotLoadedDataLayers)
+			if (const UDataLayerAsset* DataLayerAsset = SoftDataLayerAsset.Get())
 			{
-				if (DataLayerAsset != nullptr)
-				{
-					const UDataLayerInstance* DataLayerInstance = DataLayerManager->GetDataLayerInstance(DataLayerAsset.Get());
-					if (DataLayerInstance != nullptr)
-					{
-						NotLoadedDataLayers.Emplace(DataLayerInstance);
-					}
-				}
+				NotLoadedDataLayerAssets.Emplace(DataLayerAsset);
 			}
 		}
 	}
 
-	return NotLoadedDataLayers;
+	return NotLoadedDataLayerAssets;
 }
 
-TArray<UDataLayerInstance*> UWorldPartitionEditorPerProjectUserSettings::GetWorldDataLayersNotLoadedInEditor(UWorld* InWorld)
+TArray<UDataLayerAsset*> UWorldPartitionEditorPerProjectUserSettings::GetWorldDataLayerAssetsNotLoadedInEditor(UWorld* InWorld)
 {
-	TArray<const UDataLayerInstance*> NotLoadedDataLayers = AsConst(*this).GetWorldDataLayersNotLoadedInEditor(InWorld);
-	TArray<UDataLayerInstance*> NotLoadedDataLayersMutable;
-	Algo::Transform(NotLoadedDataLayers, NotLoadedDataLayersMutable, [](const UDataLayerInstance* DataLayerInstance) { return const_cast<UDataLayerInstance*>(DataLayerInstance); });
-	return NotLoadedDataLayersMutable;
+	TArray<const UDataLayerAsset*> NotLoadedDataLayerAssets = AsConst(*this).GetWorldDataLayerAssetsNotLoadedInEditor(InWorld);
+	TArray<UDataLayerAsset*> NotLoadedDataLayerAssetsMutable;
+	Algo::Transform(NotLoadedDataLayerAssets, NotLoadedDataLayerAssetsMutable, [](const UDataLayerAsset* DataLayerAsset) { return const_cast<UDataLayerAsset*>(DataLayerAsset); });
+	return NotLoadedDataLayerAssetsMutable;
 }
 
-TArray<const UDataLayerInstance*> UWorldPartitionEditorPerProjectUserSettings::GetWorldDataLayersLoadedInEditor(UWorld* InWorld) const
+TArray<const UDataLayerAsset*> UWorldPartitionEditorPerProjectUserSettings::GetWorldDataLayerAssetsLoadedInEditor(UWorld* InWorld) const
 {
-	TArray<const UDataLayerInstance*> LoadedDataLayers;
+	TArray<const UDataLayerAsset*> LoadedDataLayerAssets;
 
-	if (UDataLayerManager* DataLayerManager = UDataLayerManager::GetDataLayerManager(InWorld))
+	if (const FWorldPartitionPerWorldSettings* PerWorldSettings = GetWorldPartitionPerWorldSettings(InWorld))
 	{
-		if (const FWorldPartitionPerWorldSettings* PerWorldSettings = GetWorldPartitionPerWorldSettings(InWorld))
+		for (const TSoftObjectPtr<const UDataLayerAsset>& SoftDataLayerAsset : PerWorldSettings->LoadedDataLayers)
 		{
-			for (const TSoftObjectPtr<const UDataLayerAsset>& DataLayerAsset : PerWorldSettings->LoadedDataLayers)
+			if (const UDataLayerAsset* DataLayerAsset = SoftDataLayerAsset.Get())
 			{
-				if (DataLayerAsset != nullptr)
-				{
-					const UDataLayerInstance* DataLayerInstance = DataLayerManager->GetDataLayerInstance(DataLayerAsset.Get());
-					if (DataLayerInstance != nullptr)
-					{
-						LoadedDataLayers.Emplace(DataLayerInstance);
-					}
-				}
+				LoadedDataLayerAssets.Emplace(DataLayerAsset);
 			}
 		}
 	}
 
-	return LoadedDataLayers;
+	return LoadedDataLayerAssets;
 }
 
-TArray<UDataLayerInstance*> UWorldPartitionEditorPerProjectUserSettings::GetWorldDataLayersLoadedInEditor(UWorld* InWorld)
+TArray<UDataLayerAsset*> UWorldPartitionEditorPerProjectUserSettings::GetWorldDataLayerAssetsLoadedInEditor(UWorld* InWorld)
 {
-	TArray<const UDataLayerInstance*> LoadedDataLayers = AsConst(*this).GetWorldDataLayersLoadedInEditor(InWorld);
-	TArray<UDataLayerInstance*> LoadedDataLayersMutable;
-	Algo::Transform(LoadedDataLayers, LoadedDataLayersMutable, [](const UDataLayerInstance* DataLayerInstance) { return const_cast<UDataLayerInstance*>(DataLayerInstance); });
-	return LoadedDataLayersMutable;
+	TArray<const UDataLayerAsset*> LoadedDataLayerAssets = AsConst(*this).GetWorldDataLayerAssetsLoadedInEditor(InWorld);
+	TArray<UDataLayerAsset*> LoadedDataLayerAssetsMutable;
+	Algo::Transform(LoadedDataLayerAssets, LoadedDataLayerAssetsMutable, [](const UDataLayerAsset* DataLayerAsset) { return const_cast<UDataLayerAsset*>(DataLayerAsset); });
+	return LoadedDataLayerAssetsMutable;
 }
 
 const FWorldPartitionPerWorldSettings* UWorldPartitionEditorPerProjectUserSettings::GetWorldPartitionPerWorldSettings(UWorld* InWorld) const
diff --git a/Engine/Source/Runtime/Engine/Public/WorldPartition/WorldPartitionEditorPerProjectUserSettings.h b/Engine/Source/Runtime/Engine/Public/WorldPartition/WorldPartitionEditorPerProjectUserSettings.h
--- a/Engine/Source/Runtime/Engine/Public/WorldPartition/WorldPartitionEditorPerProjectUserSettings.h
+++ b/Engine/Source/Runtime/Engine/Public/WorldPartition/WorldPartitionEditorPerProjectUserSettings.h
@@ -125,11 +125,19 @@ public:
 		SaveConfig();
 	}
 
-	ENGINE_API TArray<const UDataLayerInstance*> GetWorldDataLayersNotLoadedInEditor(UWorld* InWorld) const;
-	ENGINE_API TArray<const UDataLayerInstance*> GetWorldDataLayersLoadedInEditor(UWorld* InWorld) const;
-
-	ENGINE_API TArray<UDataLayerInstance*> GetWorldDataLayersNotLoadedInEditor(UWorld* InWorld);
-	ENGINE_API TArray<UDataLayerInstance*> GetWorldDataLayersLoadedInEditor(UWorld* InWorld);
+	UE_DEPRECATED(5.8, "Use GetWorldDataLayerAssetsNotLoadedInEditor instead.")
+	TArray<const UDataLayerInstance*> GetWorldDataLayersNotLoadedInEditor(UWorld* InWorld) const { return TArray<const UDataLayerInstance*>(); }
+	UE_DEPRECATED(5.8, "Use GetWorldDataLayerAssetsLoadedInEditor instead.")
+	TArray<const UDataLayerInstance*> GetWorldDataLayersLoadedInEditor(UWorld* InWorld) const { return TArray<const UDataLayerInstance*>(); }
+	UE_DEPRECATED(5.8, "Use GetWorldDataLayerAssetsNotLoadedInEditor instead.")
+	TArray<UDataLayerInstance*> GetWorldDataLayersNotLoadedInEditor(UWorld* InWorld) { return TArray<UDataLayerInstance*>(); }
+	UE_DEPRECATED(5.8, "Use GetWorldDataLayerAssetsLoadedInEditor instead.")
+	TArray<UDataLayerInstance*> GetWorldDataLayersLoadedInEditor(UWorld* InWorld) { return TArray<UDataLayerInstance*>(); }
+
+	ENGINE_API TArray<const UDataLayerAsset*> GetWorldDataLayerAssetsLoadedInEditor(UWorld* InWorld) const;
+	ENGINE_API TArray<UDataLayerAsset*> GetWorldDataLayerAssetsLoadedInEditor(UWorld* InWorld);
+	ENGINE_API TArray<const UDataLayerAsset*> GetWorldDataLayerAssetsNotLoadedInEditor(UWorld* InWorld) const;
+	ENGINE_API TArray<UDataLayerAsset*> GetWorldDataLayerAssetsNotLoadedInEditor(UWorld* InWorld);
 
 	ENGINE_API void SetWorldDataLayersNonDefaultEditorLoadStates(UWorld* InWorld, const TArray<const UDataLayerInstance*>& InDataLayersLoadedInEditor, const TArray<const UDataLayerInstance*>& InDataLayersNotLoadedInEditor);
 		
