Fixed incorrect instance count caused by updates to HISMs affected by density scaling.

#Synced-CL 49761124

diff --git a/Engine/Source/Runtime/Engine/Private/InstancedStaticMesh/ISMInstanceDataSceneProxy.cpp b/Engine/Source/Runtime/Engine/Private/InstancedStaticMesh/ISMInstanceDataSceneProxy.cpp
--- a/Engine/Source/Runtime/Engine/Private/InstancedStaticMesh/ISMInstanceDataSceneProxy.cpp
+++ b/Engine/Source/Runtime/Engine/Private/InstancedStaticMesh/ISMInstanceDataSceneProxy.cpp
@@ -429,13 +429,13 @@ void FISMCInstanceDataSceneProxyLegacyReordered::Update(FISMInstanceUpdateChange
 	// Use the index reorder table to scatter the data to the correct locations.
 	ApplyDataChanges(ChangeSet, IndexRemap, ChangeSet.PostUpdateNumInstances, ProxyData);
 	
-	if (bLegacyReordered && ChangeSet.PostUpdateNumInstances != LegacyInstanceReorderTable.Num())
+	if (bLegacyReordered && !LegacyInstanceReorderTable.IsEmpty())
 	{
 		// Make sure any instance no longer represented in the reorder table is hidden.
 		ProxyData.VisibleInstances.Init(false, ChangeSet.PostUpdateNumInstances);
 		for (int32 InstanceIndex : LegacyInstanceReorderTable)
 		{
-			if (IndexRemap.RemapDestIndex(InstanceIndex))
+			if (ProxyData.VisibleInstances.IsValidIndex(InstanceIndex))
 			{
 				ProxyData.VisibleInstances[InstanceIndex] = true;
 			}
@@ -832,7 +832,7 @@ void FISMCInstanceDataSceneProxyNoGPUScene::Build(FISMInstanceUpdateChangeSet&&
 	InstanceSceneDataBuffers.EndWriteAccess(AccessTag);
 
 	// Is there is a reorder table and it does not have the same number as the instances, some must be hidden
-	if (bLegacyReordered && ChangeSet.PostUpdateNumInstances != LegacyInstanceReorderTable.Num())
+	if (bLegacyReordered && !LegacyInstanceReorderTable.IsEmpty())
 	{
 		TBitArray<> HiddenInstances;
 		HiddenInstances.SetNum(ChangeSet.PostUpdateNumInstances, true);
