PR #14271: DataTable export struct to string and file as objects.

#Synced-CL 49717872

diff --git a/Engine/Source/Runtime/Engine/Private/DataTableFunctionLibrary.cpp b/Engine/Source/Runtime/Engine/Private/DataTableFunctionLibrary.cpp
--- a/Engine/Source/Runtime/Engine/Private/DataTableFunctionLibrary.cpp
+++ b/Engine/Source/Runtime/Engine/Private/DataTableFunctionLibrary.cpp
@@ -348,7 +348,7 @@ bool UDataTableFunctionLibrary::ExportDataTableToJSONString(const UDataTable* Da
 		return false;
 	}
 
-	OutJSONString = DataTable->GetTableAsJSON();
+	OutJSONString = DataTable->GetTableAsJSON(EDataTableExportFlags::UseJsonObjectsForStructs);
 	return true;
 }
 
@@ -360,7 +360,7 @@ bool UDataTableFunctionLibrary::ExportDataTableToJSONFile(const UDataTable* Data
 		return false;
 	}
 
-	if (!FFileHelper::SaveStringToFile(DataTable->GetTableAsJSON(), *JSONFilePath, FFileHelper::EEncodingOptions::ForceUTF8WithoutBOM))
+	if (!FFileHelper::SaveStringToFile(DataTable->GetTableAsJSON(EDataTableExportFlags::UseJsonObjectsForStructs), *JSONFilePath, FFileHelper::EEncodingOptions::ForceUTF8WithoutBOM))
 	{
 		UE_LOG(LogDataTable, Error, TEXT("ExportDataTableToJSONFile - Failed to write file '%s'."), *JSONFilePath);
 		return false;
