Added the Blueprint function GetDataLayerManagerForObject to UWorldPartitionBlueprintLibrary. This allows retrieving the DataLayerManager for any provided object.
This is useful for retrieving the DataLayerManager of a loaded partitioned Level Instance when running in a Blueprint where the current world context is the main world (i.e., different from the target partitioned level).

#Synced-CL 49728154

diff --git a/Engine/Source/Runtime/Engine/Private/WorldPartitionBlueprintLibrary.cpp b/Engine/Source/Runtime/Engine/Private/WorldPartitionBlueprintLibrary.cpp
--- a/Engine/Source/Runtime/Engine/Private/WorldPartitionBlueprintLibrary.cpp
+++ b/Engine/Source/Runtime/Engine/Private/WorldPartitionBlueprintLibrary.cpp
@@ -289,3 +289,8 @@ UDataLayerManager* UWorldPartitionBlueprintLibrary::GetDataLayerManager(class UO
 {
 	return UDataLayerManager::GetDataLayerManager(WorldContextObject);
 }
+
+UDataLayerManager* UWorldPartitionBlueprintLibrary::GetDataLayerManagerForObject(class UObject* Object)
+{
+	return UDataLayerManager::GetDataLayerManager(Object);
+}
\ No newline at end of file
diff --git a/Engine/Source/Runtime/Engine/Public/WorldPartition/WorldPartitionBlueprintLibrary.h b/Engine/Source/Runtime/Engine/Public/WorldPartition/WorldPartitionBlueprintLibrary.h
--- a/Engine/Source/Runtime/Engine/Public/WorldPartition/WorldPartitionBlueprintLibrary.h
+++ b/Engine/Source/Runtime/Engine/Public/WorldPartition/WorldPartitionBlueprintLibrary.h
@@ -159,8 +159,14 @@ public:
 	static bool GetActorDescsForActors(const TArray<AActor*>& InActors, TArray<FActorDesc>& OutActorDescs);
 
 	/**
-	 * Returns the Data Layer Manager for this object.
+	 * Returns the Data Layer Manager for the current world context.
 	 */
 	UFUNCTION(BlueprintCallable, Category = "World Partition", meta = (WorldContext = "WorldContextObject"))
 	static UDataLayerManager* GetDataLayerManager(class UObject* WorldContextObject);
+
+	/**
+	 * Returns the Data Layer Manager for the provided object.
+	 */
+	UFUNCTION(BlueprintCallable, Category = "World Partition")
+	static UDataLayerManager* GetDataLayerManagerForObject(class UObject* Object);
 };
