Skeletal Mesh: Fix a crash when regenerating a skeletal mesh using geometry scripting

#rnx
#Synced-CL 49749960

diff --git a/Engine/Plugins/Runtime/MeshModelingToolset/Source/ModelingComponentsEditorOnly/Private/AssetUtils/CreateSkeletalMeshUtil.cpp b/Engine/Plugins/Runtime/MeshModelingToolset/Source/ModelingComponentsEditorOnly/Private/AssetUtils/CreateSkeletalMeshUtil.cpp
--- a/Engine/Plugins/Runtime/MeshModelingToolset/Source/ModelingComponentsEditorOnly/Private/AssetUtils/CreateSkeletalMeshUtil.cpp
+++ b/Engine/Plugins/Runtime/MeshModelingToolset/Source/ModelingComponentsEditorOnly/Private/AssetUtils/CreateSkeletalMeshUtil.cpp
@@ -65,6 +65,7 @@ UE::AssetUtils::ECreateSkeletalMeshResult UE::AssetUtils::CreateSkeletalMeshAsse
 		FSkeletalMeshModel* ImportedResource = NewSkeletalMesh->GetImportedModel();
 		ImportedResource->EmptyOriginalReductionSourceMeshData();
 		ImportedResource->LODModels.Empty();
+		ImportedResource->InlineReductionCacheDatas.Empty();
 
 		NewSkeletalMesh->SetNumSourceModels(0);
 		NewSkeletalMesh->GetMaterials().Empty();
diff --git a/Engine/Source/Runtime/Engine/Private/SkeletalMesh.cpp b/Engine/Source/Runtime/Engine/Private/SkeletalMesh.cpp
--- a/Engine/Source/Runtime/Engine/Private/SkeletalMesh.cpp
+++ b/Engine/Source/Runtime/Engine/Private/SkeletalMesh.cpp
@@ -2389,15 +2389,11 @@ void USkeletalMesh::BeginBuildInternal(FSkinnedAssetBuildContext& Context)
 	{
 		RecacheNaniteAssemblyReferences(false);
 	}
-
-	//Make sure InlineReduction structure are created
-	const int32 MaxLodIndex = GetLODNum() - 1;
-	if (GetImportedModel()->InlineReductionCacheDatas.IsValidIndex(MaxLodIndex))
-	{
-		//We should not do that in main thread, this is why there is an ensure
-		GetImportedModel()->InlineReductionCacheDatas.AddDefaulted((MaxLodIndex + 1) - GetImportedModel()->InlineReductionCacheDatas.Num());
-	}
-
+	
+	// Make sure the InlineReductionCacheDatas structure has enough room for all available LODs. Since the reduction can happen on multiple threads
+	// this has to happen before we trigger the mesh build.
+	GetImportedModel()->InlineReductionCacheDatas.SetNum(GetNumSourceModels());
+	
 	// Flush the resource release commands to the rendering thread to ensure that the build doesn't occur while a resource is still
 	// allocated, and potentially accessing the USkeletalMesh.
 	ReleaseResourcesFence.Wait();
