Make sure to follow redirectors when retrieving the world type from UWorld::WorldTypePreLoadMap in UWorld::PostLoad.

#Synced-CL 49726820

diff --git a/Engine/Source/Runtime/Engine/Private/World.cpp b/Engine/Source/Runtime/Engine/Private/World.cpp
--- a/Engine/Source/Runtime/Engine/Private/World.cpp
+++ b/Engine/Source/Runtime/Engine/Private/World.cpp
@@ -1675,15 +1675,30 @@ bool UWorld::IsReadyForFinishDestroy()
 
 void UWorld::PostLoad()
 {
-	EWorldType::Type * PreLoadWorldType = UWorld::WorldTypePreLoadMap.Find(GetOuter()->GetFName());
-	if (PreLoadWorldType)
-	{
-		WorldType = *PreLoadWorldType;
-	}
-	else
+	// By default, assume the world was loaded from a package to persist in memory
+	WorldType = EWorldType::Inactive;
+
+	if (!UWorld::WorldTypePreLoadMap.IsEmpty())
 	{
-		// Since we did not specify a world type, assume it was loaded from a package to persist in memory
-		WorldType = EWorldType::Inactive;
+		EWorldType::Type * PreLoadWorldType = UWorld::WorldTypePreLoadMap.Find(GetOuter()->GetFName());
+		if (PreLoadWorldType)
+		{
+			WorldType = *PreLoadWorldType;
+		}
+		else
+		{			
+			for (const TPair<FName, EWorldType::Type>& WorldTypePair : UWorld::WorldTypePreLoadMap)
+			{
+				if (UPackage* WorldPackage = FindPackage(nullptr, *WorldTypePair.Key.ToString()))
+				{
+					if (FollowWorldRedirectorInPackage(WorldPackage) == this)
+					{
+						check(WorldType == WorldTypePair.Value);
+						break;
+					}
+				}
+			}
+		}
 	}
 
 	Super::PostLoad();
