Generate the TraceType/ObjectType ScriptName from the CollisionChannel DisplayName, if the CollisionChannel ScriptName is empty

The old broken names are stil available as deprecated aliases

#Synced-CL 49717872

diff --git a/Engine/Source/Runtime/Engine/Private/Collision/CollisionProfile.cpp b/Engine/Source/Runtime/Engine/Private/Collision/CollisionProfile.cpp
--- a/Engine/Source/Runtime/Engine/Private/Collision/CollisionProfile.cpp
+++ b/Engine/Source/Runtime/Engine/Private/Collision/CollisionProfile.cpp
@@ -314,6 +314,23 @@ void UCollisionProfile::LoadProfileConfig(bool bForceInit)
 	// first go through enum entry, and add suffix to displaynames
 	static const FString Prefix = "ECC_";
 
+	auto DisplayNameToScriptName = 
+		[](const FString& DisplayName, FStringView DeprecatedScriptName = FStringView())
+		{
+			FString ScriptName = DisplayName;
+			ScriptName.RemoveSpacesInline();
+			if (!ScriptName.StartsWith(Prefix))
+			{
+				ScriptName.InsertAt(0, Prefix);
+			}
+			if (!DeprecatedScriptName.IsEmpty())
+			{
+				ScriptName += TEXT(';');
+				ScriptName += DeprecatedScriptName;
+			}
+			return ScriptName;
+		};
+
 	// need to have mapping table between ECollisionChannel and EObjectTypeQuery/ETraceTypeQuery
 	UEnum* ObjectTypeEnum = StaticEnum<EObjectTypeQuery>();
 	UEnum* TraceTypeEnum = StaticEnum<ETraceTypeQuery>();
@@ -412,12 +429,7 @@ void UCollisionProfile::LoadProfileConfig(bool bForceInit)
 
 #if WITH_EDITOR
 				// Create the name to use when exposing this channel for scripting
-				FString ScriptValue = DisplayValue;
-				ScriptValue.RemoveSpacesInline();
-				if (!ScriptValue.StartsWith(Prefix))
-				{
-					ScriptValue.InsertAt(0, Prefix);
-				}
+				FString ScriptValue = DisplayNameToScriptName(DisplayValue);
 
 				// set displayvalue for this enum entry
 				Enum->SetMetaData(*DisplayNameKey, *DisplayValue, EnumIndex);
@@ -501,7 +513,7 @@ void UCollisionProfile::LoadProfileConfig(bool bForceInit)
 					{
 						TraceTypeEnum->RemoveMetaData(*HiddenMeta, TraceTypeIndex);
 						TraceTypeEnum->SetMetaData(*DisplayNameKey, *DisplayName, TraceTypeIndex);
-						TraceTypeEnum->SetMetaData(*ScriptNameKey, *ScriptName, TraceTypeIndex);
+						TraceTypeEnum->SetMetaData(*ScriptNameKey, ScriptName.IsEmpty() ? *DisplayNameToScriptName(DisplayName, TraceTypeEnum->GetNameStringByIndex(TraceTypeIndex)) : *ScriptName, TraceTypeIndex);
 						++TraceTypeMappingCount;
 					}
 				}
@@ -512,7 +524,7 @@ void UCollisionProfile::LoadProfileConfig(bool bForceInit)
 					{
 						ObjectTypeEnum->RemoveMetaData(*HiddenMeta, ObjectTypeIndex);
 						ObjectTypeEnum->SetMetaData(*DisplayNameKey, *DisplayName, ObjectTypeIndex);
-						ObjectTypeEnum->SetMetaData(*ScriptNameKey, *ScriptName, ObjectTypeIndex);
+						ObjectTypeEnum->SetMetaData(*ScriptNameKey, ScriptName.IsEmpty() ? *DisplayNameToScriptName(DisplayName, ObjectTypeEnum->GetNameStringByIndex(ObjectTypeIndex)) : *ScriptName, ObjectTypeIndex);
 						++ObjectTypeMappingCount;
 					}
 				}
