Added MinInGameMaterialQualityMobile option to runtime virtual texture component.
This works like the existing MinInGameMaterialQuality to enable/disable the volume according to material quality, but only applies on mobile platforms.
If it is not set (the default) then mobile platforms fall back to using MinInGameMaterialQuality.
Also did a small rearrange of details panel ordering for runtime virtual texture component.
#Synced-CL 49765706

diff --git a/Engine/Source/Runtime/Engine/Classes/Components/RuntimeVirtualTextureComponent.h b/Engine/Source/Runtime/Engine/Classes/Components/RuntimeVirtualTextureComponent.h
--- a/Engine/Source/Runtime/Engine/Classes/Components/RuntimeVirtualTextureComponent.h
+++ b/Engine/Source/Runtime/Engine/Classes/Components/RuntimeVirtualTextureComponent.h
@@ -64,9 +64,13 @@ protected:
 	UPROPERTY(EditAnywhere, BlueprintReadOnly, TextExportTransient, Category = RuntimeVirtualTexture)
 	TObjectPtr<URuntimeVirtualTexture> VirtualTexture = nullptr;
 
-	/** Per platform overrides for enabling the virtual texture. Only affects In-Game and PIE. */
-	UPROPERTY(EditAnywhere, AdvancedDisplay, Category = RuntimeVirtualTexture)
-	FPerPlatformBool EnableInGamePerPlatform;
+	/** Set to true to enable scalability settings for the virtual texture. */
+	UPROPERTY(EditAnywhere, AdvancedDisplay, Category = RuntimeVirtualTexture, meta = (InlineEditConditionToggle))
+	bool bEnableScalability = false;
+
+	/** Group index of the scalability settings to use for the virtual texture. */
+	UPROPERTY(EditAnywhere, AdvancedDisplay, Category = RuntimeVirtualTexture, meta = (UIMin = "0", UIMax = "2", EditCondition = bEnableScalability))
+	uint32 ScalabilityGroup = 0;
 
 	/** Enable the virtual texture only when Nanite is enabled. Can be used for a Displacement virtual texture with Nanite tessellation. */
 	UPROPERTY(EditAnywhere, AdvancedDisplay, Category = RuntimeVirtualTexture)
@@ -80,13 +84,17 @@ protected:
 	UPROPERTY(EditAnywhere, AdvancedDisplay, Category = RuntimeVirtualTexture, Meta = (EditCondition = "bUseMinMaterialQuality"))
 	ERuntimeVirtualTextureMaterialQuality MinInGameMaterialQuality = ERuntimeVirtualTextureMaterialQuality::Low;
 
-	/** Set to true to enable scalability settings for the virtual texture. */
-	UPROPERTY(EditAnywhere, AdvancedDisplay, Category = RuntimeVirtualTexture, meta = (InlineEditConditionToggle))
-	bool bEnableScalability = false;
+	/** Use a minimum material quality to determine if we enable the virtual texture on mobile. */
+	UPROPERTY(EditAnywhere, AdvancedDisplay, Category = RuntimeVirtualTexture, Meta = (InlineEditConditionToggle))
+	bool bUseMinMaterialQualityMobile = false;
 
-	/** Group index of the scalability settings to use for the virtual texture. */
-	UPROPERTY(EditAnywhere, AdvancedDisplay, Category = RuntimeVirtualTexture, meta = (UIMin = "0", UIMax = "2", EditCondition = bEnableScalability))
-	uint32 ScalabilityGroup = 0;
+	/** The minimum material quality for which we enable the virtual texture on mobile. Only affects In-Game and PIE. */
+	UPROPERTY(EditAnywhere, AdvancedDisplay, Category = RuntimeVirtualTexture, Meta = (EditCondition = "bUseMinMaterialQualityMobile"))
+	ERuntimeVirtualTextureMaterialQuality MinInGameMaterialQualityMobile = ERuntimeVirtualTextureMaterialQuality::Low;
+
+	/** Per platform overrides for enabling the virtual texture. Only affects In-Game and PIE. */
+	UPROPERTY(EditAnywhere, AdvancedDisplay, Category = RuntimeVirtualTexture)
+	FPerPlatformBool EnableInGamePerPlatform;
 
 	/** Hide primitives in the main pass. Hidden primitives will be those that draw to this virtual texture with 'Draw in Main Pass' set to 'From Virtual Texture'. */
 	UPROPERTY(EditAnywhere, AdvancedDisplay, Category = RuntimeVirtualTexture)
diff --git a/Engine/Source/Runtime/Engine/Private/Components/RuntimeVirtualTextureComponent.cpp b/Engine/Source/Runtime/Engine/Private/Components/RuntimeVirtualTextureComponent.cpp
--- a/Engine/Source/Runtime/Engine/Private/Components/RuntimeVirtualTextureComponent.cpp
+++ b/Engine/Source/Runtime/Engine/Private/Components/RuntimeVirtualTextureComponent.cpp
@@ -227,9 +227,17 @@ bool URuntimeVirtualTextureComponent::IsEnabledInScene() const
 				return false;
 			}
 
-			if (bUseMinMaterialQuality)
+			const ERuntimeVirtualTextureMaterialQuality CurrentQuality = ConvertMaterialQualityEnum(GetCachedScalabilityCVars().MaterialQualityLevel);
+
+			if (IsMobilePlatform(ShaderPlatform) && bUseMinMaterialQualityMobile)
+			{
+				if (CurrentQuality < MinInGameMaterialQualityMobile)
+				{
+					return false;
+				}
+			}
+			else if (bUseMinMaterialQuality)
 			{
-				ERuntimeVirtualTextureMaterialQuality CurrentQuality = ConvertMaterialQualityEnum(GetCachedScalabilityCVars().MaterialQualityLevel);
  				if (CurrentQuality < MinInGameMaterialQuality)
  				{
  					return false;
