Fixed inverse refpose being applied to an identity transform when bone isn't present.

#Synced-CL 49747543

diff --git a/Engine/Source/Runtime/Engine/Private/SkeletalRender.cpp b/Engine/Source/Runtime/Engine/Private/SkeletalRender.cpp
--- a/Engine/Source/Runtime/Engine/Private/SkeletalRender.cpp
+++ b/Engine/Source/Runtime/Engine/Private/SkeletalRender.cpp
@@ -523,7 +523,10 @@ void UpdateRefToLocalMatricesInner(
 
 	for (int32 ThisBoneIndex = 0; ThisBoneIndex < NumReferenceToLocal; ++ThisBoneIndex)
 	{
-		VectorMatrixMultiply(&ReferenceToLocal[ThisBoneIndex], &(*RefBasesInvMatrix)[ThisBoneIndex], &ReferenceToLocal[ThisBoneIndex]);
+		if (!BitsToReset[ThisBoneIndex])
+		{
+			VectorMatrixMultiply(&ReferenceToLocal[ThisBoneIndex], &(*RefBasesInvMatrix)[ThisBoneIndex], &ReferenceToLocal[ThisBoneIndex]);
+		}
 	}
 }
 /**
