Engine - Exposing AnimMontages track properties to BluePrint

#Synced-CL 49766586

diff --git a/Engine/Source/Runtime/Engine/Classes/Animation/AnimCompositeBase.h b/Engine/Source/Runtime/Engine/Classes/Animation/AnimCompositeBase.h
--- a/Engine/Source/Runtime/Engine/Classes/Animation/AnimCompositeBase.h
+++ b/Engine/Source/Runtime/Engine/Classes/Animation/AnimCompositeBase.h
@@ -63,7 +63,7 @@ struct FRootMotionExtractionStep
 };
 
 /** this is anim segment that defines what animation and how **/
-USTRUCT()
+USTRUCT(BlueprintType)
 struct FAnimSegment
 {
 	GENERATED_USTRUCT_BODY()
@@ -86,7 +86,7 @@ struct FAnimSegment
 
 protected:
 	/** Anim Reference to play - only allow AnimSequence or AnimComposite **/
-	UPROPERTY(EditAnywhere, Category=AnimSegment, meta=(DisplayName = "Animation Reference"))
+	UPROPERTY(BlueprintReadOnly, EditAnywhere, Category=AnimSegment, meta=(DisplayName = "Animation Reference"))
 	TObjectPtr<UAnimSequenceBase> AnimReference;
 public:
 
@@ -101,22 +101,22 @@ public:
 #endif // WITH_EDITOR
 	
 	/** Start Pos within this AnimCompositeBase */
-	UPROPERTY(VisibleAnywhere, Category=AnimSegment, meta=(DisplayName = "Starting Position"))
+	UPROPERTY(BlueprintReadOnly, VisibleAnywhere, Category=AnimSegment, meta=(DisplayName = "Starting Position"))
 	float StartPos;
 
 	/** Time to start playing AnimSequence at. */
-	UPROPERTY(EditAnywhere, Category=AnimSegment, meta=(DisplayName = "Start Time"))
+	UPROPERTY(BlueprintReadOnly, EditAnywhere, Category=AnimSegment, meta=(DisplayName = "Start Time"))
 	float AnimStartTime;
 
 	/** Time to end playing the AnimSequence at. */
-	UPROPERTY(EditAnywhere, Category=AnimSegment, meta=(DisplayName = "End Time"))
+	UPROPERTY(BlueprintReadOnly, EditAnywhere, Category=AnimSegment, meta=(DisplayName = "End Time"))
 	float AnimEndTime;
 
 	/** Playback speed of this animation. If you'd like to reverse, set -1*/
-	UPROPERTY(EditAnywhere, Category=AnimSegment, meta=(DisplayName = "Play Rate"))
+	UPROPERTY(BlueprintReadOnly, EditAnywhere, Category=AnimSegment, meta=(DisplayName = "Play Rate"))
 	float AnimPlayRate;
 
-	UPROPERTY(EditAnywhere, Category=AnimSegment, meta=(DisplayName = "Loop Count"))
+	UPROPERTY(BlueprintReadOnly, EditAnywhere, Category=AnimSegment, meta=(DisplayName = "Loop Count"))
 	int32 LoopingCount;
 
 	FAnimSegment()
@@ -225,12 +225,12 @@ private:
  * For now this is only one TArray, but in the future 
  * we should define more transition/blending behaviors
  **/
-USTRUCT()
+USTRUCT(BlueprintType)
 struct FAnimTrack
 {
 	GENERATED_USTRUCT_BODY()
 
-	UPROPERTY(EditAnywhere, Category=AnimTrack, EditFixedSize)
+	UPROPERTY(BlueprintReadOnly, EditAnywhere, Category=AnimTrack, EditFixedSize)
 	TArray<FAnimSegment>	AnimSegments;
 
 	FAnimTrack() {}
diff --git a/Engine/Source/Runtime/Engine/Classes/Animation/AnimMontage.h b/Engine/Source/Runtime/Engine/Classes/Animation/AnimMontage.h
--- a/Engine/Source/Runtime/Engine/Classes/Animation/AnimMontage.h
+++ b/Engine/Source/Runtime/Engine/Classes/Animation/AnimMontage.h
@@ -79,15 +79,15 @@ public:
  * Each slot data referenced by Animation Slot 
  * contains slot name, and animation data 
  */
-USTRUCT()
+USTRUCT(BlueprintType)
 struct FSlotAnimationTrack
 {
 	GENERATED_USTRUCT_BODY()
 
-	UPROPERTY(EditAnywhere, Category=Slot)
+	UPROPERTY(BlueprintReadOnly, EditAnywhere, Category=Slot)
 	FName SlotName;
 
-	UPROPERTY(EditAnywhere, Category=Slot)
+	UPROPERTY(BlueprintReadOnly, EditAnywhere, Category=Slot)
 	FAnimTrack AnimTrack;
 
 	ENGINE_API FSlotAnimationTrack();
@@ -695,7 +695,7 @@ class UAnimMontage : public UAnimCompositeBase
 	TArray<FCompositeSection> CompositeSections;
 	
 	// slot data, each slot contains anim track
-	UPROPERTY()
+	UPROPERTY(BlueprintReadOnly, Category = "Montage")
 	TArray<struct FSlotAnimationTrack> SlotAnimTracks;
 
 #if WITH_EDITORONLY_DATA
