Skeletal Mesh: Only emit skin weight profile override message once for each skelmesh

#rnx
#Synced-CL 49749960

diff --git a/Engine/Source/Runtime/Engine/Private/Animation/SkinWeightProfile.cpp b/Engine/Source/Runtime/Engine/Private/Animation/SkinWeightProfile.cpp
--- a/Engine/Source/Runtime/Engine/Private/Animation/SkinWeightProfile.cpp
+++ b/Engine/Source/Runtime/Engine/Private/Animation/SkinWeightProfile.cpp
@@ -296,7 +296,11 @@ void FSkinWeightProfilesData::SetDynamicDefaultSkinWeightProfile(USkeletalMesh*
 {
 	if (bStaticOverridden)
 	{
-		UE_LOG(LogSkeletalMesh, Error, TEXT("[%s] Skeletal Mesh has overridden the default Skin Weights buffer during serialization, cannot set any other skin weight profile."), *Mesh->GetName());
+		if (!bEmittedOverrideWarning)
+		{
+			UE_LOG(LogSkeletalMesh, Error, TEXT("[%s] Skeletal Mesh has overridden the default Skin Weights buffer during serialization, cannot set any other skin weight profile."), *Mesh->GetName());
+			bEmittedOverrideWarning = true;
+		}
 		return;
 	}
 
@@ -454,9 +458,10 @@ FSkinWeightVertexBuffer* FSkinWeightProfilesData::GetOverrideBuffer(const FSkinW
 	// In case we have overridden the default skin weight buffer we do not need to create an override buffer, if it was statically overridden we cannot load any other profile
 	if (bDefaultOverridden && (ProfileStack == DefaultProfileStack || bStaticOverridden))
 	{	
-		if (bStaticOverridden && ProfileStack != DefaultProfileStack)
+		if (!bEmittedOverrideWarning && bStaticOverridden && ProfileStack != DefaultProfileStack)
 		{
-			UE_LOG(LogSkeletalMesh, Verbose, TEXT("Skeletal Mesh has overridden the default Skin Weights buffer during serialization, cannot set any other skin weight profile."));
+			UE_LOG(LogSkeletalMesh, Error, TEXT("Skeletal Mesh has overridden the default Skin Weights buffer during serialization, cannot set any other skin weight profile."));
+			bEmittedOverrideWarning = true;
 		}	
 
 		return nullptr;
diff --git a/Engine/Source/Runtime/Engine/Public/Animation/SkinWeightProfile.h b/Engine/Source/Runtime/Engine/Public/Animation/SkinWeightProfile.h
--- a/Engine/Source/Runtime/Engine/Public/Animation/SkinWeightProfile.h
+++ b/Engine/Source/Runtime/Engine/Public/Animation/SkinWeightProfile.h
@@ -307,8 +307,11 @@ private:
 
 	bool bDefaultOverridden = false;
 	bool bStaticOverridden = false;
+	
 	FSkinWeightProfileStack DefaultProfileStack;
 
 	FSkinWeightReadbackData ReadbackData;
+	
+	mutable bool bEmittedOverrideWarning = false;
 };
 
