Make `FHitProxy` safe transactionally.

#Synced-CL 49758756

diff --git a/Engine/Source/Runtime/Engine/Private/HitProxies.cpp b/Engine/Source/Runtime/Engine/Private/HitProxies.cpp
--- a/Engine/Source/Runtime/Engine/Private/HitProxies.cpp
+++ b/Engine/Source/Runtime/Engine/Private/HitProxies.cpp
@@ -25,7 +25,7 @@ const FHitProxyId FHitProxyId::InvisibleHitProxyId( INDEX_NONE - 1 );
 class FHitProxyArray
 {
 	TSparseArray<HHitProxy*> HitProxies;
-	FCriticalSection Lock;
+	FTransactionallySafeCriticalSection Lock;
 public:
 
 	static FHitProxyArray& Get()
@@ -36,19 +36,19 @@ public:
 
 	void Remove(int32 Index)
 	{
-		FScopeLock ScopeLock(&Lock);
+		UE::TScopeLock ScopeLock(Lock);
 		HitProxies.RemoveAt(Index);
 	}
 
 	int32 Add(HHitProxy* Proxy)
 	{
-		FScopeLock ScopeLock(&Lock);
+		UE::TScopeLock ScopeLock(Lock);
 		return HitProxies.Add(Proxy);
 	}
 
 	HHitProxy* GetHitProxyById(int32 Index)
 	{
-		FScopeLock ScopeLock(&Lock);
+		UE::TScopeLock ScopeLock(Lock);
 		if(Index >= 0 && Index < HitProxies.GetMaxIndex() && HitProxies.IsAllocated(Index))
 		{
 			return HitProxies[Index];
