Use a transactionally safe lock in streaming manager texture.

#Synced-CL 49758756

diff --git a/Engine/Source/Runtime/Engine/Private/Streaming/StreamingManagerTexture.cpp b/Engine/Source/Runtime/Engine/Private/Streaming/StreamingManagerTexture.cpp
--- a/Engine/Source/Runtime/Engine/Private/Streaming/StreamingManagerTexture.cpp
+++ b/Engine/Source/Runtime/Engine/Private/Streaming/StreamingManagerTexture.cpp
@@ -264,7 +264,7 @@ void FRenderAssetStreamingManager::OnPreGarbageCollect()
 {
 	TRACE_CPUPROFILER_EVENT_SCOPE(FRenderAssetStreamingManager::OnPreGarbageCollect);
 
-	FScopeLock ScopeLock(&CriticalSection);
+	UE::TScopeLock ScopeLock(CriticalSection);
 
 	if (StreamingRenderAssetsSyncEvent.IsValid())
 	{
@@ -328,7 +328,7 @@ void FRenderAssetStreamingManager::MarkMountedStateDirty(FIoFilenameHash Filenam
  */
 void FRenderAssetStreamingManager::CancelForcedResources()
 {
-	FScopeLock ScopeLock(&CriticalSection);
+	UE::TScopeLock ScopeLock(CriticalSection);
 
 	TArray<FStreamingRenderAsset>& StreamingRenderAssets = GetStreamingRenderAssetsAsyncSafe();
 
@@ -378,7 +378,7 @@ void FRenderAssetStreamingManager::SetDisregardWorldResourcesForFrames( int32 Nu
  **/
 bool FRenderAssetStreamingManager::StreamOutRenderAssetData( int64 RequiredMemorySize )
 {
-	FScopeLock ScopeLock(&CriticalSection);
+	UE::TScopeLock ScopeLock(CriticalSection);
 	LLM_SCOPE(ELLMTag::StreamingManager);
 	LLM_TAGSET_SCOPE_CLEAR(ELLMTagSet::Assets);
 	LLM_TAGSET_SCOPE_CLEAR(ELLMTagSet::AssetClasses);
@@ -807,7 +807,7 @@ void FRenderAssetStreamingManager::PrepareAsyncTask(bool bProcessEverything)
  */
 void FRenderAssetStreamingManager::BoostTextures( AActor* Actor, float BoostFactor )
 {
-	FScopeLock ScopeLock(&CriticalSection);
+	UE::TScopeLock ScopeLock(CriticalSection);
 
 	if ( Actor )
 	{
@@ -872,7 +872,7 @@ void FRenderAssetStreamingManager::AddLevel( ULevel* Level )
 		return;
 	}
 
-	FScopeLock ScopeLock(&CriticalSection);
+	UE::TScopeLock ScopeLock(CriticalSection);
 	LLM_SCOPE(ELLMTag::StreamingManager);
 	LLM_TAGSET_SCOPE_CLEAR(ELLMTagSet::Assets);
 	LLM_TAGSET_SCOPE_CLEAR(ELLMTagSet::AssetClasses);
@@ -939,7 +939,7 @@ void FRenderAssetStreamingManager::RemoveLevel( ULevel* Level )
 	{
 		return;
 	}
-	FScopeLock ScopeLock(&CriticalSection);
+	UE::TScopeLock ScopeLock(CriticalSection);
 
 	check(Level);
 
@@ -981,7 +981,7 @@ void FRenderAssetStreamingManager::NotifyLevelOffset(ULevel* Level, const FVecto
 		return;
 	}
 	
-	FScopeLock ScopeLock(&CriticalSection);
+	UE::TScopeLock ScopeLock(CriticalSection);
 
 	for (FLevelRenderAssetManager* LevelManager : LevelRenderAssetManagers)
 	{
@@ -995,7 +995,7 @@ void FRenderAssetStreamingManager::NotifyLevelOffset(ULevel* Level, const FVecto
 
 void FRenderAssetStreamingManager::AddStreamingRenderAsset(UStreamableRenderAsset* InAsset)
 {
-	FScopeLock ScopeLock(&CriticalSection);
+	UE::TScopeLock ScopeLock(CriticalSection);
 	LLM_SCOPE(ELLMTag::StreamingManager);
 	LLM_TAGSET_SCOPE_CLEAR(ELLMTagSet::Assets);
 	LLM_TAGSET_SCOPE_CLEAR(ELLMTagSet::AssetClasses);
@@ -1024,7 +1024,7 @@ void FRenderAssetStreamingManager::RemoveStreamingRenderAsset( UStreamableRender
 	// FSimpleStreamableAssetManager::UnregisterAsset can be called outside critical Section
 	FSimpleStreamableAssetManager::UnregisterAsset(RenderAsset);
 	
-	FScopeLock ScopeLock(&CriticalSection);
+	UE::TScopeLock ScopeLock(CriticalSection);
 	LLM_SCOPE(ELLMTag::StreamingManager);
 	LLM_TAGSET_SCOPE_CLEAR(ELLMTagSet::Assets);
 	LLM_TAGSET_SCOPE_CLEAR(ELLMTagSet::AssetClasses);
@@ -1089,7 +1089,7 @@ void FRenderAssetStreamingManager::NotifyActorDestroyed( AActor* Actor )
 		return;
 	}
 	
-	FScopeLock ScopeLock(&CriticalSection);
+	UE::TScopeLock ScopeLock(CriticalSection);
 
 	STAT(GatheredStats.CallbacksCycles = -(int32)FPlatformTime::Cycles();)
 	FRemovedRenderAssetArray RemovedRenderAssets;
@@ -1136,7 +1136,7 @@ void FRenderAssetStreamingManager::RemoveStaticReferences(const IPrimitiveCompon
 		return;
 	}
 	
-	FScopeLock ScopeLock(&CriticalSection);
+	UE::TScopeLock ScopeLock(CriticalSection);
 
 	check(Primitive);
 	FRenderAssetOwnerStreamingState& StreamingState = Primitive->GetStreamingState();
@@ -1169,7 +1169,7 @@ void FRenderAssetStreamingManager::NotifyPrimitiveDetached( const IPrimitiveComp
 		return;
 	}
 	
-	FScopeLock ScopeLock(&CriticalSection);
+	UE::TScopeLock ScopeLock(CriticalSection);
 
 	if (!Primitive || !Primitive->GetStreamingState().IsAttachedToStreamingManager())
 	{
@@ -1267,7 +1267,7 @@ void FRenderAssetStreamingManager::NotifyPrimitiveUpdated( const IPrimitiveCompo
 	// When this happens, the states will be initialized when the components render states will be set.
 	if (IsInGameThread() && bUseDynamicStreaming && Primitive && !StreamingState.IsIgnoreStreamingManagerUpdate())
 	{
-		FScopeLock ScopeLock(&CriticalSection);
+		UE::TScopeLock ScopeLock(CriticalSection);
 
 		// Check if there is a pending renderstate update, useful since streaming data can be updated in UPrimitiveComponent::CreateRenderState_Concurrent().
 		// We handle this here to prevent the primitive from being updated twice in the same frame.
@@ -1321,7 +1321,7 @@ void FRenderAssetStreamingManager::NotifyPrimitiveUpdated_Concurrent( const IPri
 	// The level context is not used currently.
 	if (bUseDynamicStreaming && Primitive)
 	{
-		FScopeLock ScopeLock(&CriticalSection);
+		UE::TScopeLock ScopeLock(CriticalSection);
 		FStreamingTextureLevelContext LevelContext(GetCurrentMaterialQualityLevelChecked());
 		if (DynamicComponentManager.Add(Primitive, LevelContext) == EAddComponentResult::Success)
 		{
@@ -1333,7 +1333,11 @@ void FRenderAssetStreamingManager::NotifyPrimitiveUpdated_Concurrent( const IPri
 	}
 
 	STAT(CallbackCycle += (int32)FPlatformTime::Cycles();)
-	STAT(FPlatformAtomics::InterlockedAdd(&GatheredStats.CallbacksCycles, CallbackCycle));
+
+	UE_AUTORTFM_OPEN
+	{
+		STAT(FPlatformAtomics::InterlockedAdd(&GatheredStats.CallbacksCycles, CallbackCycle));
+	};
 }
 
 void FRenderAssetStreamingManager::SyncStates(bool bCompleteFullUpdateCycle)
@@ -1359,7 +1363,7 @@ void FRenderAssetStreamingManager::SyncStates(bool bCompleteFullUpdateCycle)
  */
 FStreamingRenderAsset* FRenderAssetStreamingManager::GetStreamingRenderAsset( const UStreamableRenderAsset* RenderAsset )
 {
-	FScopeLock ScopeLock(&CriticalSection);
+	UE::TScopeLock ScopeLock(CriticalSection);
 
 	if (RenderAsset && AsyncUnsafeStreamingRenderAssets.IsValidIndex(RenderAsset->StreamingIndex))
 	{
@@ -1382,7 +1386,7 @@ FStreamingRenderAsset* FRenderAssetStreamingManager::GetStreamingRenderAsset( co
  */
 void FRenderAssetStreamingManager::UpdateIndividualRenderAsset( UStreamableRenderAsset* RenderAsset )
 {
-	FScopeLock ScopeLock(&CriticalSection);
+	UE::TScopeLock ScopeLock(CriticalSection);
 
 	if (!IStreamingManager::Get().IsStreamingEnabled() || !RenderAsset) return;
 
@@ -1908,7 +1912,7 @@ public:
 
 void FRenderAssetStreamingManager::UpdateResourceStreaming( float DeltaTime, bool bProcessEverything/*=false*/ )
 {
-	FScopeLock ScopeLock(&CriticalSection);
+	UE::TScopeLock ScopeLock(CriticalSection);
 	TArray<FStreamingRenderAsset>& StreamingRenderAssets = GetStreamingRenderAssetsAsyncSafe();
 
 	SCOPE_CYCLE_COUNTER(STAT_RenderAssetStreaming_GameThreadUpdateTime);
@@ -2102,7 +2106,7 @@ void FRenderAssetStreamingManager::UpdateResourceStreaming( float DeltaTime, boo
  */
 int32 FRenderAssetStreamingManager::BlockTillAllRequestsFinished( float TimeLimit /*= 0.0f*/, bool bLogResults /*= false*/ )
 {
-	FScopeLock ScopeLock(&CriticalSection);
+	UE::TScopeLock ScopeLock(CriticalSection);
 	TArray<FStreamingRenderAsset>& StreamingRenderAssets = GetStreamingRenderAssetsAsyncSafe();
 	QUICK_SCOPE_CYCLE_COUNTER(STAT_FRenderAssetStreamingManager_BlockTillAllRequestsFinished);
 
@@ -2150,7 +2154,7 @@ int32 FRenderAssetStreamingManager::BlockTillAllRequestsFinished( float TimeLimi
 
 void FRenderAssetStreamingManager::GetObjectReferenceBounds(const UObject* RefObject, TArray<FBox>& AssetBoxes)
 {
-	FScopeLock ScopeLock(&CriticalSection);
+	UE::TScopeLock ScopeLock(CriticalSection);
 	QUICK_SCOPE_CYCLE_COUNTER(STAT_FRenderAssetStreamingManager_GetObjectReferenceBounds);
 
 	const UStreamableRenderAsset* RenderAsset = Cast<const UStreamableRenderAsset>(RefObject);
@@ -2223,7 +2227,7 @@ void FRenderAssetStreamingManager::GetAssetComponents(const UStreamableRenderAss
 	
 	if (RenderAsset && RenderAsset->StreamingIndex >= 0)
 	{
-		FScopeLock Lock(&CriticalSection);
+		UE::TScopeLock Lock(CriticalSection);
 		QUICK_SCOPE_CYCLE_COUNTER(STAT_FRenderAssetStreamingManager_GetAssetComponents);
 
 		const int32 Idx = RenderAsset->StreamingIndex;
@@ -2255,7 +2259,7 @@ void FRenderAssetStreamingManager::PropagateLightingScenarioChange()
 		return;
 	}
 	
-	FScopeLock ScopeLock(&CriticalSection);
+	UE::TScopeLock ScopeLock(CriticalSection);
 	QUICK_SCOPE_CYCLE_COUNTER(STAT_FRenderAssetStreamingManager_PropagateLightingScenarioChange);
 
 	// Note that dynamic components don't need to be handled because their renderstates are updated, which triggers and update.
@@ -2282,7 +2286,7 @@ void FRenderAssetStreamingManager::PropagateLightingScenarioChange()
 
 void FRenderAssetStreamingManager::AddRenderedTextureStats(TMap<FString, FRenderedTextureStats>& OutRenderedTextureStats)
 {
-	FScopeLock ScopeLock(&CriticalSection);
+	UE::TScopeLock ScopeLock(CriticalSection);
 	
 	for (const FStreamingRenderAsset& StreamingRenderAsset : AsyncUnsafeStreamingRenderAssets)
 	{
@@ -2320,7 +2324,7 @@ void FRenderAssetStreamingManager::OnRenderAssetUpdateAbandoned(UStreamableRende
 
 bool FRenderAssetStreamingManager::HandleDumpTextureStreamingStatsCommand( const TCHAR* Cmd, FOutputDevice& Ar )
 {
-	FScopeLock ScopeLock(&CriticalSection);
+	UE::TScopeLock ScopeLock(CriticalSection);
 
 	float CurrentOccupancyPct = 0.f;
 	float TargetOccupancyPct = 0.f;
@@ -2347,7 +2351,7 @@ bool FRenderAssetStreamingManager::HandleDumpTextureStreamingStatsCommand( const
 
 bool FRenderAssetStreamingManager::HandleListStreamingRenderAssetsCommand( const TCHAR* Cmd, FOutputDevice& Ar )
 {
-	FScopeLock ScopeLock(&CriticalSection);
+	UE::TScopeLock ScopeLock(CriticalSection);
 	TArray<FStreamingRenderAsset>& StreamingRenderAssets = GetStreamingRenderAssetsAsyncSafe();
 
 	SyncStates(true);
@@ -2478,7 +2482,7 @@ bool FRenderAssetStreamingManager::HandleListStreamingRenderAssetsCommand( const
 
 bool FRenderAssetStreamingManager::HandleResetMaxEverRequiredRenderAssetMemoryCommand(const TCHAR* Cmd, FOutputDevice& Ar)
 {
-	FScopeLock ScopeLock(&CriticalSection);
+	UE::TScopeLock ScopeLock(CriticalSection);
 
 	Ar.Logf(TEXT("OldMax: %u MaxEverRequired Reset."), MaxEverRequired);
 	ResetMaxEverRequired();	
@@ -2487,7 +2491,7 @@ bool FRenderAssetStreamingManager::HandleResetMaxEverRequiredRenderAssetMemoryCo
 
 bool FRenderAssetStreamingManager::HandleLightmapStreamingFactorCommand( const TCHAR* Cmd, FOutputDevice& Ar )
 {
-	FScopeLock ScopeLock(&CriticalSection);
+	UE::TScopeLock ScopeLock(CriticalSection);
 
 	FString FactorString(FParse::Token(Cmd, 0));
 	float NewFactor = ( FactorString.Len() > 0 ) ? FCString::Atof(*FactorString) : GLightmapStreamingFactor;
@@ -2501,7 +2505,7 @@ bool FRenderAssetStreamingManager::HandleLightmapStreamingFactorCommand( const T
 
 bool FRenderAssetStreamingManager::HandleCancelRenderAssetStreamingCommand( const TCHAR* Cmd, FOutputDevice& Ar )
 {
-	FScopeLock ScopeLock(&CriticalSection);
+	UE::TScopeLock ScopeLock(CriticalSection);
 
 	UTexture::CancelPendingTextureStreaming();
 	UStaticMesh::CancelAllPendingStreamingActions();
@@ -2511,7 +2515,7 @@ bool FRenderAssetStreamingManager::HandleCancelRenderAssetStreamingCommand( cons
 
 bool FRenderAssetStreamingManager::HandleShadowmapStreamingFactorCommand( const TCHAR* Cmd, FOutputDevice& Ar )
 {
-	FScopeLock ScopeLock(&CriticalSection);
+	UE::TScopeLock ScopeLock(CriticalSection);
 
 	FString FactorString(FParse::Token(Cmd, 0));
 	float NewFactor = ( FactorString.Len() > 0 ) ? FCString::Atof(*FactorString) : GShadowmapStreamingFactor;
@@ -2525,7 +2529,7 @@ bool FRenderAssetStreamingManager::HandleShadowmapStreamingFactorCommand( const
 
 bool FRenderAssetStreamingManager::HandleNumStreamedMipsCommand( const TCHAR* Cmd, FOutputDevice& Ar )
 {
-	FScopeLock ScopeLock(&CriticalSection);
+	UE::TScopeLock ScopeLock(CriticalSection);
 
 	FString NumTextureString(FParse::Token(Cmd, 0));
 	FString NumMipsString(FParse::Token(Cmd, 0));
@@ -2560,7 +2564,7 @@ bool FRenderAssetStreamingManager::HandleNumStreamedMipsCommand( const TCHAR* Cm
 
 bool FRenderAssetStreamingManager::HandleTrackRenderAssetCommand( const TCHAR* Cmd, FOutputDevice& Ar )
 {
-	FScopeLock ScopeLock(&CriticalSection);
+	UE::TScopeLock ScopeLock(CriticalSection);
 
 	FString AssetName(FParse::Token(Cmd, 0));
 	if ( TrackRenderAsset(AssetName) )
@@ -2572,7 +2576,7 @@ bool FRenderAssetStreamingManager::HandleTrackRenderAssetCommand( const TCHAR* C
 
 bool FRenderAssetStreamingManager::HandleListTrackedRenderAssetsCommand( const TCHAR* Cmd, FOutputDevice& Ar )
 {
-	FScopeLock ScopeLock(&CriticalSection);
+	UE::TScopeLock ScopeLock(CriticalSection);
 
 	FString NumAssetString(FParse::Token(Cmd, 0));
 	int32 NumAssets = (NumAssetString.Len() > 0) ? FCString::Atoi(*NumAssetString) : -1;
@@ -2587,7 +2591,7 @@ FORCEINLINE float SqrtKeepMax(float V)
 
 bool FRenderAssetStreamingManager::HandleDebugTrackedRenderAssetsCommand( const TCHAR* Cmd, FOutputDevice& Ar )
 {
-	FScopeLock ScopeLock(&CriticalSection);
+	UE::TScopeLock ScopeLock(CriticalSection);
 
 	// The ENABLE_RENDER_ASSET_TRACKING macro is defined in ContentStreaming.cpp and not available here. This code does not compile any more.
 #ifdef ENABLE_TEXTURE_TRACKING_BROKEN
@@ -2693,7 +2697,7 @@ bool FRenderAssetStreamingManager::HandleDebugTrackedRenderAssetsCommand( const
 
 bool FRenderAssetStreamingManager::HandleUntrackRenderAssetCommand( const TCHAR* Cmd, FOutputDevice& Ar )
 {
-	FScopeLock ScopeLock(&CriticalSection);
+	UE::TScopeLock ScopeLock(CriticalSection);
 
 	FString AssetName(FParse::Token(Cmd, 0));
 	if (UntrackRenderAsset(AssetName))
@@ -2705,7 +2709,7 @@ bool FRenderAssetStreamingManager::HandleUntrackRenderAssetCommand( const TCHAR*
 
 bool FRenderAssetStreamingManager::HandleStreamOutCommand( const TCHAR* Cmd, FOutputDevice& Ar )
 {
-	FScopeLock ScopeLock(&CriticalSection);
+	UE::TScopeLock ScopeLock(CriticalSection);
 
 	FString Parameter(FParse::Token(Cmd, 0));
 	int64 FreeMB = (Parameter.Len() > 0) ? FCString::Atoi(*Parameter) : 0;
@@ -2723,7 +2727,7 @@ bool FRenderAssetStreamingManager::HandleStreamOutCommand( const TCHAR* Cmd, FOu
 
 bool FRenderAssetStreamingManager::HandlePauseRenderAssetStreamingCommand( const TCHAR* Cmd, FOutputDevice& Ar )
 {
-	FScopeLock ScopeLock(&CriticalSection);
+	UE::TScopeLock ScopeLock(CriticalSection);
 
 	bPauseRenderAssetStreaming = !bPauseRenderAssetStreaming;
 	Ar.Logf( TEXT("Render asset streaming is now \"%s\"."), bPauseRenderAssetStreaming ? TEXT("PAUSED") : TEXT("UNPAUSED") );
@@ -2732,7 +2736,7 @@ bool FRenderAssetStreamingManager::HandlePauseRenderAssetStreamingCommand( const
 
 bool FRenderAssetStreamingManager::HandleStreamingManagerMemoryCommand( const TCHAR* Cmd, FOutputDevice& Ar, UWorld* InWorld )
 {
-	FScopeLock ScopeLock(&CriticalSection);
+	UE::TScopeLock ScopeLock(CriticalSection);
 	TArray<FStreamingRenderAsset>& StreamingRenderAssets = GetStreamingRenderAssetsAsyncSafe();
 
 	SyncStates(true);
@@ -2762,7 +2766,7 @@ bool FRenderAssetStreamingManager::HandleStreamingManagerMemoryCommand( const TC
 
 bool FRenderAssetStreamingManager::HandleLODGroupsCommand( const TCHAR* Cmd, FOutputDevice& Ar )
 {
-	FScopeLock ScopeLock(&CriticalSection);
+	UE::TScopeLock ScopeLock(CriticalSection);
 	SyncStates(true);
 
 	struct FTextureGroupStats
@@ -2869,7 +2873,7 @@ bool FRenderAssetStreamingManager::HandleLODGroupsCommand( const TCHAR* Cmd, FOu
 
 bool FRenderAssetStreamingManager::HandleInvestigateRenderAssetCommand(const TCHAR* Cmd, FOutputDevice& Ar, UWorld* InWorld)
 {
-	FScopeLock ScopeLock(&CriticalSection);
+	UE::TScopeLock ScopeLock(CriticalSection);
 	TArray<FStreamingRenderAsset>& StreamingRenderAssets = GetStreamingRenderAssetsAsyncSafe();
 	SyncStates(true);
 
diff --git a/Engine/Source/Runtime/Engine/Private/Streaming/StreamingManagerTexture.h b/Engine/Source/Runtime/Engine/Private/Streaming/StreamingManagerTexture.h
--- a/Engine/Source/Runtime/Engine/Private/Streaming/StreamingManagerTexture.h
+++ b/Engine/Source/Runtime/Engine/Private/Streaming/StreamingManagerTexture.h
@@ -455,7 +455,7 @@ private:
 	ENGINE_API virtual void MarkMountedStateDirty(FIoFilenameHash FilenameHash) override;
 
 	// A critical section use around code that could be called in parallel with NotifyPrimitiveUpdated() or NotifyPrimitiveUpdated_Concurrent().
-	FCriticalSection CriticalSection;
+	FTransactionallySafeCriticalSection CriticalSection;
 
 	// An event used to prevent FUpdateStreamingRenderAssetsTask from overlapping with related work
 	FGraphEventRef StreamingRenderAssetsSyncEvent;
