New Material Translator work:
- Added implementation for many material expressions: DecalMipmapLevel, DepthOfFieldFunction, Sobol, TemporalSobol, SpeedTree, physics field sampling, and others
- Implemented inline HLSL emission helpers and unified external code handling, though external code emissino will be refactored soon into a better solution
- Added new external inputs including DistanceCullFade and default path tracing UVs
- Expanded emitter validation utilities and analysis error handling
- General refactoring and renaming.

#rnx

#Synced-CL 49801225

diff --git a/Engine/Config/BaseMaterialExpressions.ini b/Engine/Config/BaseMaterialExpressions.ini
--- a/Engine/Config/BaseMaterialExpressions.ini
+++ b/Engine/Config/BaseMaterialExpressions.ini
@@ -251,6 +251,7 @@ bShaderInputData=true
 
 [/Script/Engine.MaterialExpressionDecalMipmapLevel]
 +MenuCategories=NSLOCTEXT("MaterialExpression", "Utils", "Utils")
++ExternalCodeIdentifiers="DecalMipmapLevel"
 
 [/Script/Engine.MaterialExpressionDeltaTime]
 +MenuCategories=NSLOCTEXT("MaterialExpression", "Constants", "Constants")
@@ -263,6 +264,7 @@ bShaderInputData=true
 
 [/Script/Engine.MaterialExpressionDepthOfFieldFunction]
 +MenuCategories=NSLOCTEXT("MaterialExpression", "Utility", "Utility")
++ExternalCodeIdentifiers="DepthOfFieldFunction"
 
 [/Script/Engine.MaterialExpressionDeriveNormalZ]
 +MenuCategories=NSLOCTEXT("MaterialExpression", "VectorOps", "VectorOps")
@@ -1074,6 +1076,8 @@ bShaderInputData=true
 +ExternalCodeDeclarations=(ReturnType=Float3,			Name="ReflectionVector",				Definition="Parameters.ReflectionVector", ShaderFrequency=Pixel|Compute, bIsInlined=true, EnvironmentDefines=((Name="USES_TRANSFORM_VECTOR", ShaderFrequency=Pixel|Compute)))
 +ExternalCodeDeclarations=(ReturnType=Float,			Name="ScalarBlueNoise",					Definition="ViewScalarBlueNoise(GetPixelPosition(Parameters), View.StateFrameIndex)", ShaderFrequency=Pixel|Compute)
 +ExternalCodeDeclarations=(ReturnType=Float,			Name="SphericalParticleOpacity",		Definition="GetSphericalParticleOpacity(Parameters,$0)", ShaderFrequency=Pixel|Compute, EnvironmentDefines=((Name="NEEDS_PARTICLE_POSITION"), (Name="SPHERICAL_PARTICLE_OPACITY"), (Name="NEEDS_WORLD_POSITION_EXCLUDING_SHADER_OFFSETS"), (Name="MIR.SceneDepth")))
++ExternalCodeDeclarations=(ReturnType=Float,			Name="DecalMipmapLevel",			    Definition="ComputeDecalMipmapLevel(Parameters,$0)", ShaderFrequency=Pixel|Compute)
++ExternalCodeDeclarations=(ReturnType=Float,			Name="DepthOfFieldFunction",		    Definition="MaterialExpressionDepthOfFieldFunction($0,$1)", ShaderFrequency=Pixel|Compute)
 +ExternalCodeDeclarations=(ReturnType=Float,			Name="TwoSidedSign",					Definition="Parameters.TwoSidedSign", ShaderFrequency=Pixel|Compute, Derivative=Zero, bIsInlined=true)
 +ExternalCodeDeclarations=(ReturnType=Float4,			Name="VertexColor",						Definition="Parameters.VertexColor", DefinitionDDX="Parameters.VertexColor_DDX", DefinitionDDY="Parameters.VertexColor_DDY", Derivative=Valid, bIsInlined=true, EnvironmentDefines=((Name="INTERPOLATE_VERTEX_COLOR", ShaderFrequency=Pixel|Compute)))
 +ExternalCodeDeclarations=(ReturnType=Float3,			Name="VertexNormal",					Definition="Parameters.TangentToWorld[2]", DefinitionDDX="Parameters.WorldGeoNormal_DDX", DefinitionDDY="Parameters.WorldGeoNormal_DDY", Derivative=Valid, bIsInlined=true, EnvironmentDefines=((Name="INTERPOLATE_VERTEX_COLOR", ShaderFrequency=Pixel|Compute)))
diff --git a/Engine/Shaders/Private/MaterialTemplate.ush b/Engine/Shaders/Private/MaterialTemplate.ush
--- a/Engine/Shaders/Private/MaterialTemplate.ush
+++ b/Engine/Shaders/Private/MaterialTemplate.ush
@@ -2838,6 +2838,9 @@ FWSVector2 RotateScaleOffsetTexCoords(FWSVector2 InTexCoords, float4 InRotationS
 /** Vertex offset for SpeedTree wind and LOD */
 float3 GetSpeedTreeVertexOffsetInner(FMaterialVertexParameters Parameters, int GeometryType, int WindType, int LODType, float BillboardThreshold, bool bExtraBend, float3 ExtraBend, FSpeedTreeData STData) 
 {
+#if !VF_SUPPORTS_SPEEDTREE_WIND
+	return 0;
+#else
 	#if (NUM_MATERIAL_TEXCOORDS_VERTEX < 6) || IS_MESHPARTICLE_FACTORY
 		return float4(0,0,0);
 	#endif
@@ -2997,20 +3000,27 @@ float3 GetSpeedTreeVertexOffsetInner(FMaterialVertexParameters Parameters, int G
 
 	// convert into a world space offset
 	return (FinalPosition - OriginalPosition) * TreeScale;
+#endif // !VF_SUPPORTS_SPEEDTREE_WIND
 }
 
 /** Vertex offset for SpeedTree wind and LOD */
 float3 GetSpeedTreeVertexOffset(FMaterialVertexParameters Parameters, int GeometryType, int WindType, int LODType, float BillboardThreshold, bool bUsePreviousFrame, bool bExtraBend, float3 ExtraBend) 
 {
-#if VF_SUPPORTS_SPEEDTREE_WIND
 	if (bUsePreviousFrame)
 	{
 		return GetSpeedTreeVertexOffsetInner(Parameters, GeometryType, WindType, LODType, BillboardThreshold, bExtraBend, ExtraBend, GetPreviousSpeedTreeData());
 	}
 	return GetSpeedTreeVertexOffsetInner(Parameters, GeometryType, WindType, LODType, BillboardThreshold, bExtraBend, ExtraBend, GetCurrentSpeedTreeData());
-#else
-	return 0;
-#endif
+}
+
+float3 GetSpeedTreeVertexOffset(FMaterialVertexParameters Parameters, int GeometryType, int WindType, int LODType, float BillboardThreshold, bool bExtraBend, float3 ExtraBend) 
+{
+	return GetSpeedTreeVertexOffsetInner(Parameters, GeometryType, WindType, LODType, BillboardThreshold, bExtraBend, ExtraBend, GetCurrentSpeedTreeData());
+}
+
+float3 GetPrevSpeedTreeVertexOffset(FMaterialVertexParameters Parameters, int GeometryType, int WindType, int LODType, float BillboardThreshold, bool bExtraBend, float3 ExtraBend) 
+{
+	return GetSpeedTreeVertexOffsetInner(Parameters, GeometryType, WindType, LODType, BillboardThreshold, bExtraBend, ExtraBend, GetPreviousSpeedTreeData());
 }
 
 #endif
diff --git a/Engine/Source/Runtime/Engine/Private/Materials/MaterialExpressions.cpp b/Engine/Source/Runtime/Engine/Private/Materials/MaterialExpressions.cpp
--- a/Engine/Source/Runtime/Engine/Private/Materials/MaterialExpressions.cpp
+++ b/Engine/Source/Runtime/Engine/Private/Materials/MaterialExpressions.cpp
@@ -9566,25 +9566,24 @@ void UMaterialExpressionDataDrivenShaderPlatformInfoSwitch::PostEditChangeProper
 //
 
 #if WITH_EDITOR
-static bool IsBindlessEnabledForCompiler(FMaterialCompiler* Compiler)
+
+bool IsBindlessEnabledForCompiler(EShaderPlatform ShaderPlatform, EShaderFrequency ShaderFrequency)
 {
-	const EShaderPlatform ShaderPlatform = Compiler->GetShaderPlatform();
 	const ERHIBindlessConfiguration BindlessConfiguration = UE::ShaderCompiler::GetBindlessConfiguration(ShaderPlatform);
-
 	switch (BindlessConfiguration)
 	{
-	case ERHIBindlessConfiguration::All:
-		return true;
-	case ERHIBindlessConfiguration::Minimal:
-		return true;
-
-	case ERHIBindlessConfiguration::RayTracing:
-		return IsRayTracingShaderFrequency(Compiler->GetCurrentShaderFrequency());
+		case ERHIBindlessConfiguration::All:        return true;
+		case ERHIBindlessConfiguration::Minimal:    return true;
+		case ERHIBindlessConfiguration::RayTracing: return IsRayTracingShaderFrequency(ShaderFrequency);
 	}
-
 	return false;
 }
 
+static bool IsBindlessEnabledForCompiler(FMaterialCompiler* Compiler)
+{
+	return IsBindlessEnabledForCompiler(Compiler->GetShaderPlatform(), Compiler->GetCurrentShaderFrequency());
+}
+
 int32 UMaterialExpressionBindlessSwitch::Compile(FMaterialCompiler* Compiler, int32 OutputIndex)
 {
 	if (IsBindlessEnabledForCompiler(Compiler))
diff --git a/Engine/Source/Runtime/Engine/Private/Materials/MaterialExpressionsIR.cpp b/Engine/Source/Runtime/Engine/Private/Materials/MaterialExpressionsIR.cpp
--- a/Engine/Source/Runtime/Engine/Private/Materials/MaterialExpressionsIR.cpp
+++ b/Engine/Source/Runtime/Engine/Private/Materials/MaterialExpressionsIR.cpp
@@ -255,6 +255,8 @@
 #include "Materials/MaterialSharedPrivate.h"
 #include "Materials/MaterialExpressionMotionVectorWorldOffsetOutput.h"
 #include "Materials/MaterialExpressionSubsurfaceMediumMaterialOutput.h"
+#include "Materials/MaterialExpressionPostVolumeUserFlagTest.h"
+#include "Materials/MaterialExpressionBindlessSwitch.h"
 #include "Materials/MaterialExpressionSubstrate.h"
 #include "MaterialShared.h"
 #include "Misc/MemStackUtility.h"
@@ -264,12 +266,56 @@
 #include "ColorManagement/ColorSpace.h"
 #include "SubstrateTranslatorCommon.h"
 #include "Engine/SubsurfaceProfile.h"
-
 #include "Materials/MaterialIREmitter.h"
+#include "Materials/MaterialIRModule.h"
 
 using FValueRef = MIR::FValueRef;
 using namespace Substrate;
 
+/* Helper functions to emit inline HLSL instructions */
+
+static FValueRef MaterialExpressionExternalCodeBase_EmitInlineHLSL(MIR::FEmitter& Em, const UMaterialExpressionExternalCodeBase& InExpression, int32 InExternalCodeIdentifierIndex, TConstArrayView<FValueRef> InArguments = {}, MIR::EValueFlags ValueFlags = MIR::EValueFlags::None)
+{
+	checkf(InExternalCodeIdentifierIndex >= 0 && InExternalCodeIdentifierIndex < InExpression.ExternalCodeIdentifiers.Num(),
+		   TEXT("External code identifier index (%d) out of bounds; Upper bound is %d"), InExternalCodeIdentifierIndex, InExpression.ExternalCodeIdentifiers.Num());
+
+	const FName ExternalCodeIdentifier = InExpression.ExternalCodeIdentifiers[InExternalCodeIdentifierIndex];
+	const FMaterialExternalCodeDeclaration* ExternalCodeDeclaration = MaterialExternalCodeRegistry::Get().FindExternalCode(ExternalCodeIdentifier);
+	if (!ExternalCodeDeclaration)
+	{
+		Em.Errorf(TEXT("Missing external code declaration for '%s' [Index=%d]"), *ExternalCodeIdentifier.ToString(), InExternalCodeIdentifierIndex);
+		return {};
+	}
+
+	return Em.InlineHLSL(ExternalCodeDeclaration, InArguments, ValueFlags);
+}
+
+static void BuildInlineHLSLOutput(MIR::FEmitter& Em, UMaterialExpressionExternalCodeBase& InExpression, TConstArrayView<FValueRef> InArguments)
+{
+	// If the external code expression has only one identifier, evaluate it and flow its result to all expression outputs.
+	if (InExpression.ExternalCodeIdentifiers.Num() == 1)
+	{
+		Em.OutputsWithComponentMask(MaterialExpressionExternalCodeBase_EmitInlineHLSL(Em, InExpression, 0, InArguments));
+	}
+	else
+	{
+		// Make sure the number of identifiers perfectly matches the number of expression outputs.
+		if (InExpression.ExternalCodeIdentifiers.Num() != InExpression.GetOutputs().Num())
+		{
+			Em.Errorf(TEXT("Internal Error: The number of external code identifiers (%d) does not match the number of expression outputs (%d) in expression '%s'.")
+					  , InExpression.ExternalCodeIdentifiers.Num(), InExpression.GetOutputs().Num(), *InExpression.GetName());
+			return;
+		}
+
+		// Evaluate each external code identifier separately and flow it to its own matching expression output only.
+		for (int32 OutputIndex = 0; OutputIndex < InExpression.ExternalCodeIdentifiers.Num(); ++OutputIndex)
+		{
+			Em.Output(OutputIndex, MaterialExpressionExternalCodeBase_EmitInlineHLSL(Em, InExpression, OutputIndex, InArguments));
+		}
+	}
+}
+
+
 static FName NAME_CameraVector("CameraVector");
 
 /* Constants */
@@ -492,8 +538,10 @@ void UMaterialExpressionStaticSwitchParameter::Build(MIR::FEmitter& Em)
 	Em.Output(0, bCondition ? AValue : BValue);
 }
 
-static void EmitEffectiveInputOrError(MIR::FEmitter& Em, FExpressionInput* EffectiveInput, FString& Error)
+void UMaterialExpressionDataDrivenShaderPlatformInfoSwitch::Build(MIR::FEmitter& Em)
 {
+	FString Error;
+	FExpressionInput* EffectiveInput = GetEffectiveInput(Em.GetShaderPlatform(), Error);
 	if (!EffectiveInput)
 	{
 		Em.Error(Error);
@@ -503,12 +551,6 @@ static void EmitEffectiveInputOrError(MIR::FEmitter& Em, FExpressionInput* Effec
 	Em.Output(0, Em.Input(EffectiveInput));
 }
 
-void UMaterialExpressionDataDrivenShaderPlatformInfoSwitch::Build(MIR::FEmitter& Em)
-{
-	FString Error;
-	EmitEffectiveInputOrError(Em, GetEffectiveInput(Em.GetShaderPlatform(), Error), Error);
-}
-
 void UMaterialExpressionDistanceFieldsRenderingSwitch::Build(MIR::FEmitter& Em)
 {
 	EShaderPlatform ShaderPlatform = Em.GetShaderPlatform();
@@ -612,18 +654,6 @@ void UMaterialExpressionVirtualTextureFeatureSwitch::Build(MIR::FEmitter& Em)
 	Em.Output(0, UseVirtualTexturing(Em.GetShaderPlatform()) ? YesValue : NoValue);
 }
 
-static FValueRef EmitInlineHLSL(MIR::FEmitter& Em, FName ExternalCodeIdentifier, TConstArrayView<FValueRef> InArguments = {})
-{
-	const FMaterialExternalCodeDeclaration* ExternalCodeDeclaration = MaterialExternalCodeRegistry::Get().FindExternalCode(ExternalCodeIdentifier);
-	if (!ExternalCodeDeclaration)
-	{
-		Em.Errorf(TEXT("Missing external code declaration for '%s'"), *ExternalCodeIdentifier.ToString());
-		return {};
-	}
-
-	return Em.InlineHLSL(ExternalCodeDeclaration, InArguments);
-}
-
 static void EmitExternalCodeConditionalReplace(MIR::FEmitter& Em, FName ExternalCodeIdentifier, FExpressionInput& Default, FExpressionInput& Replace, const TCHAR* DefaultDesc, const TCHAR* ReplaceDesc)
 {
 	FValueRef DefaultValue = Em.Input(&Default);
@@ -636,23 +666,19 @@ static void EmitExternalCodeConditionalReplace(MIR::FEmitter& Em, FName External
 	DefaultValue = Em.Cast(DefaultValue, CommonType);
 	ReplaceValue = Em.Cast(ReplaceValue, CommonType);
 
-	Em.Output(0, Em.Branch(EmitInlineHLSL(Em, ExternalCodeIdentifier), ReplaceValue, DefaultValue));
+	Em.Output(0, Em.Branch(Em.InlineHLSL(ExternalCodeIdentifier), ReplaceValue, DefaultValue));
 }
 
 static FValueRef EmitVertexNormal(MIR::FEmitter& Em)
 {
-	/*
-		static FName NAME_VertexNormal("VertexNormal");
-		NormalValue = EmitInlineHLSL(Em, NAME_VertexNormal);
-	*/
-	static FName NAME_VertexNormal("VertexNormal");
-	return EmitInlineHLSL(Em, NAME_VertexNormal);
+	static const FName NAME_VertexNormal("VertexNormal");
+	return Em.InlineHLSL(NAME_VertexNormal);
 }
 
 static FValueRef EmitVertexTangent(MIR::FEmitter& Em)
 {
-	static FName NAME_VertexTangent("VertexTangent");
-	return EmitInlineHLSL(Em, NAME_VertexTangent);
+	static const FName NAME_VertexTangent("VertexTangent");
+	return Em.InlineHLSL(NAME_VertexTangent);
 }
 
 #define REPLACE_INPUTS(DEFAULT, REPLACE) DEFAULT, REPLACE, TEXT(#DEFAULT), TEXT(#REPLACE)
@@ -934,39 +960,6 @@ void UMaterialExpressionSquareRoot::Build(MIR::FEmitter& Em)
 	Em.Output(0, Em.Sqrt(Em.Input(&Input)));
 }
 
-static FValueRef EmitInlineHLSL(MIR::FEmitter& Em, const UMaterialExpressionExternalCodeBase& InExternalCodeExpression, int32 InExternalCodeIdentifierIndex, TConstArrayView<FValueRef> InArguments = {}, MIR::EValueFlags ValueFlags = MIR::EValueFlags::None)
-{
-	checkf(InExternalCodeIdentifierIndex >= 0 && InExternalCodeIdentifierIndex < InExternalCodeExpression.ExternalCodeIdentifiers.Num(),
-		TEXT("External code identifier index (%d) out of bounds; Upper bound is %d"), InExternalCodeIdentifierIndex, InExternalCodeExpression.ExternalCodeIdentifiers.Num());
-
-	const FName ExternalCodeIdentifier = InExternalCodeExpression.ExternalCodeIdentifiers[InExternalCodeIdentifierIndex];
-	const FMaterialExternalCodeDeclaration* ExternalCodeDeclaration = MaterialExternalCodeRegistry::Get().FindExternalCode(ExternalCodeIdentifier);
-	if (!ExternalCodeDeclaration)
-	{
-		Em.Errorf(TEXT("Missing external code declaration for '%s' [Index=%d]"), *ExternalCodeIdentifier.ToString(), InExternalCodeIdentifierIndex);
-		return {};
-	}
-
-	return Em.InlineHLSL(ExternalCodeDeclaration, InArguments, ValueFlags);
-}
-
-static void BuildInlineHLSLOutput(MIR::FEmitter& Em, const UMaterialExpressionExternalCodeBase& InExternalCodeExpression, TConstArrayView<FValueRef> InArguments)
-{
-	// If there are multiple output pins but only one external code identifier, use it for all outputs and let the emitter handle the swizzling.
-	// This is used for output pins that map to component swizzling, e.g. R, G, B, RGB, RGBA
-	if (InExternalCodeExpression.Outputs.Num() > 1 && InExternalCodeExpression.ExternalCodeIdentifiers.Num() == 1)
-	{
-		Em.Outputs(InExternalCodeExpression.Outputs, EmitInlineHLSL(Em, InExternalCodeExpression, 0, InArguments));
-	}
-	else
-	{
-		for (int32 OutputIndex = 0; OutputIndex < InExternalCodeExpression.ExternalCodeIdentifiers.Num(); ++OutputIndex)
-		{
-			Em.Output(OutputIndex, EmitInlineHLSL(Em, InExternalCodeExpression, OutputIndex, InArguments));
-		}
-	}
-}
-
 void UMaterialExpressionExternalCodeBase::Build(MIR::FEmitter& Em)
 {
 	BuildInlineHLSLOutput(Em, *this, {});
@@ -1068,7 +1061,7 @@ void UMaterialExpressionEyeAdaptationInverse::Build(MIR::FEmitter& Em)
 	check(ExternalCodeIdentifiers.Num() == 1);
 	FValueRef LightValue = Em.CastToFloat(Em.InputDefaultFloat(&LightValueInput, 1.0f), 3);
 	FValueRef AlphaValue = Em.CastToFloat(Em.InputDefaultFloat(&AlphaInput, 1.0f), 1);
-	FValueRef MultiplierValue = EmitInlineHLSL(Em, *this, 0, { AlphaValue });
+	FValueRef MultiplierValue = MaterialExpressionExternalCodeBase_EmitInlineHLSL(Em, *this, 0, { AlphaValue });
 	Em.Output(0, Em.Multiply(LightValue, MultiplierValue));
 }
 
@@ -1461,15 +1454,17 @@ static void EmitParticleSubUV(MIR::FEmitter& Em, UMaterialExpressionTextureSampl
 	static FName NAME_ParticleSubUVCoords1("ParticleSubUVCoords1");
 	static FName NAME_ParticleSubUVLerp("ParticleSubUVLerp");
 
-	FValueRef TexCoords0 = DummyDependency ? EmitInlineHLSL(Em, NAME_ParticleSubUVCoords0, { DummyDependency }) : EmitInlineHLSL(Em, NAME_ParticleSubUVCoords0);
+	FValueRef TexCoords0 = DummyDependency
+		? Em.InlineHLSL(NAME_ParticleSubUVCoords0, { DummyDependency })
+		: Em.InlineHLSL(NAME_ParticleSubUVCoords0);
 	FValueRef Sample0 = BuildTextureSample(Em, Expr, TextureValue, TextureType, TexCoords0, bAutomaticViewMipBias);
 
 	if (bBlend)
 	{
-		FValueRef TexCoords1 = EmitInlineHLSL(Em, NAME_ParticleSubUVCoords1);
+		FValueRef TexCoords1 = Em.InlineHLSL(NAME_ParticleSubUVCoords1);
 		FValueRef Sample1 = BuildTextureSample(Em, Expr, TextureValue, TextureType, TexCoords1, bAutomaticViewMipBias);
 
-		FValueRef SubImageLerp = EmitInlineHLSL(Em, NAME_ParticleSubUVLerp);
+		FValueRef SubImageLerp = Em.InlineHLSL(NAME_ParticleSubUVLerp);
 
 		Sample0 = Em.Lerp(Sample0, Sample1, SubImageLerp);
 	}
@@ -1663,14 +1658,14 @@ void UMaterialExpressionFontSample::Build(MIR::FEmitter& Em)
 			return;
 		}
 
-		Em.Outputs(Outputs,
-			Em.TextureSample(
-				Em.TextureObject(Texture, ExpectedSamplerType),
-				Em.ExternalInput(MIR::TexCoordIndexToExternalInput(0)),
-				false,
-				{ SSM_FromTextureAsset, ExpectedSamplerType }
-			)
+		FValueRef OutputValue = Em.TextureSample(
+			Em.TextureObject(Texture, ExpectedSamplerType),
+			Em.ExternalInput(MIR::TexCoordIndexToExternalInput(0)),
+			false,
+			{ SSM_FromTextureAsset, ExpectedSamplerType }
 		);
+
+		Em.OutputsWithComponentMask(OutputValue);
 	}
 }
 #endif
@@ -1704,14 +1699,14 @@ void UMaterialExpressionFontSampleParameter::Build(MIR::FEmitter& Em)
 			return;
 		}
 
-		Em.Outputs(Outputs,
-			Em.TextureSample(
-				TextureParameter,
-				Em.ExternalInput(MIR::TexCoordIndexToExternalInput(0)),
-				false,
-				{ SSM_FromTextureAsset, ExpectedSamplerType }
-			)
+		FValueRef OutputValue = Em.TextureSample(
+			TextureParameter,
+			Em.ExternalInput(MIR::TexCoordIndexToExternalInput(0)),
+			false,
+			{ SSM_FromTextureAsset, ExpectedSamplerType }
 		);
+
+		Em.OutputsWithComponentMask(OutputValue);
 	}
 }
 #endif
@@ -1742,7 +1737,7 @@ static MIR::FValueRef BuildConstantVector(MIR::FEmitter& Em, TArrayView<float> I
 
 static FValueRef EmitWorldPosition(MIR::FEmitter& Em, EWorldPositionIncludedOffsets WorldPositionShaderOffset);
 
-static FValueRef EmitTransformVectorBase(
+static FValueRef EmitTransformVector(
 	MIR::FEmitter& Em, FValueRef InputValue, EMaterialCommonBasis TransformSourceBasis, EMaterialCommonBasis TransformDestBasis,
 	bool bIsPositionTransform, MIR::FValueRef PeriodicWorldTileSizeValue, MIR::FValueRef FirstPersonInterpolationAlphaValue);
 
@@ -1855,7 +1850,7 @@ void UMaterialExpressionRuntimeVirtualTextureSample::Build(MIR::FEmitter& Em)
 			{
 				//TODO: optimize by calculating translated world to VT directly.
 				//This requires some more work as the transform is currently fed in through a preshader variable, which is cached.
-				MIR::FValueRef AbsWorldPosIndex = EmitTransformVectorBase(Em, WorldPositionValue, EMaterialCommonBasis::MCB_TranslatedWorld, EMaterialCommonBasis::MCB_World, true, {}, {});
+				MIR::FValueRef AbsWorldPosIndex = EmitTransformVector(Em, WorldPositionValue, EMaterialCommonBasis::MCB_TranslatedWorld, EMaterialCommonBasis::MCB_World, true, {}, {});
 
 				MIR::FValueRef P0 = Uniforms[ERuntimeVirtualTextureShaderUniform_WorldToUVTransform0];
 				MIR::FValueRef P1 = Uniforms[ERuntimeVirtualTextureShaderUniform_WorldToUVTransform1];
@@ -2064,10 +2059,9 @@ static FValueRef EmitPeriodicWorldPositionOrigin(MIR::FEmitter& Em, FValueRef Ti
 	return Em.InlineHLSL(TileScaleIndexValue->Type, PeriodicWorldOriginFunctionName, { TileScaleIndexValue });
 }
 
-// Emits inline HLSL from an external code declaration that does not take any parameters.
 static FValueRef EmitFixedExternalCode(MIR::FEmitter& Em, FName InExternalCodeIdentifier)
 {
-	return Em.InlineHLSL(MaterialExternalCodeRegistry::Get().FindExternalCode(InExternalCodeIdentifier), {}, MIR::EValueFlags::SubstituteTags);
+	return Em.InlineHLSL(InExternalCodeIdentifier, {}, MIR::EValueFlags::SubstituteTags);
 }
 
 static FValueRef EmitMatrixCastTo3x3(MIR::FEmitter& Em, FValueRef MatrixValue)
@@ -2112,8 +2106,7 @@ static FValueRef EmitMultiplyLWCMatrix(MIR::FEmitter& Em, FValueRef VectorValue,
 		: Em.InlineHLSL(VectorValue->Type, TEXTVIEW("WSMultiplyVector($0, $1)"), { VectorValue, MatrixValue });
 }
 
-static FValueRef EmitTransformVectorBase(
-	MIR::FEmitter& Em, FValueRef InputValue, EMaterialCommonBasis TransformSourceBasis, EMaterialCommonBasis TransformDestBasis,
+static FValueRef EmitTransformVector(MIR::FEmitter& Em, FValueRef InputValue, EMaterialCommonBasis TransformSourceBasis, EMaterialCommonBasis TransformDestBasis,
 	bool bIsPositionTransform, MIR::FValueRef PeriodicWorldTileSizeValue, MIR::FValueRef FirstPersonInterpolationAlphaValue)
 {
 	// Construct float3(0,0,x) out of the input if it is a scalar
@@ -2143,7 +2136,7 @@ static FValueRef EmitTransformVectorBase(
 			check(!bIsPositionTransform);
 			if (TransformDestBasis == MCB_World)
 			{
-				return Em.MatrixMultiply(InputValue, EmitFixedExternalCode(Em, TEXT("TangentToWorld")));
+				return Em.MatrixMultiply(InputValue, Em.InlineHLSL(TEXT("TangentToWorld")));
 			}
 			// else use MCB_World as intermediary basis
 			break;
@@ -2341,20 +2334,20 @@ static FValueRef EmitTransformVectorBase(
 	check(IntermediaryBasis != TransformDestBasis);
 
 	// Use intermediary basis
-	FValueRef IntermediaryBasisA = EmitTransformVectorBase(Em, InputValue, TransformSourceBasis, IntermediaryBasis, bIsPositionTransform, PeriodicWorldTileSizeValue, FirstPersonInterpolationAlphaValue);
-	FValueRef IntermediaryBasisB = EmitTransformVectorBase(Em, IntermediaryBasisA, IntermediaryBasis, TransformDestBasis, bIsPositionTransform, PeriodicWorldTileSizeValue, FirstPersonInterpolationAlphaValue);
+	FValueRef IntermediaryBasisA = EmitTransformVector(Em, InputValue, TransformSourceBasis, IntermediaryBasis, bIsPositionTransform, PeriodicWorldTileSizeValue, FirstPersonInterpolationAlphaValue);
+	FValueRef IntermediaryBasisB = EmitTransformVector(Em, IntermediaryBasisA, IntermediaryBasis, TransformDestBasis, bIsPositionTransform, PeriodicWorldTileSizeValue, FirstPersonInterpolationAlphaValue);
 
 	return IntermediaryBasisB;
 }
 
-static void BuildTransformVectorBase(
+static void BuildTransformVector(
 	MIR::FEmitter& Em, const FExpressionInput* Input, EMaterialCommonBasis TransformSourceBasis, EMaterialCommonBasis TransformDestBasis,
 	bool bIsPositionTransform, MIR::FValueRef PeriodicWorldTileSizeValue, MIR::FValueRef FirstPersonInterpolationAlphaValue)
 {
 	FValueRef InputValue = Em.CheckIsPrimitive(Em.Input(Input));
 	UE_MIR_CHECKPOINT(Em);
 
-	FValueRef OutputValue = EmitTransformVectorBase(Em, InputValue, TransformSourceBasis, TransformDestBasis, bIsPositionTransform, PeriodicWorldTileSizeValue, FirstPersonInterpolationAlphaValue);
+	FValueRef OutputValue = EmitTransformVector(Em, InputValue, TransformSourceBasis, TransformDestBasis, bIsPositionTransform, PeriodicWorldTileSizeValue, FirstPersonInterpolationAlphaValue);
 
 	if (TransformSourceBasis == MCB_World && bIsPositionTransform)
 	{
@@ -2377,7 +2370,7 @@ void UMaterialExpressionTransform::Build(MIR::FEmitter& Em)
 	const EMaterialCommonBasis TransformDestBasis = UE::MaterialTranslatorUtils::GetMaterialCommonBasis(TransformType);
 
 	constexpr bool bIsPositionTransform = false;
-	BuildTransformVectorBase(Em, &Input, TransformSourceBasis, TransformDestBasis, bIsPositionTransform, {}, {});
+	BuildTransformVector(Em, &Input, TransformSourceBasis, TransformDestBasis, bIsPositionTransform, {}, {});
 }
 
 void UMaterialExpressionTransformPosition::Build(MIR::FEmitter& Em)
@@ -2396,7 +2389,7 @@ void UMaterialExpressionTransformPosition::Build(MIR::FEmitter& Em)
 	const EMaterialCommonBasis TransformDestBasis = UE::MaterialTranslatorUtils::GetMaterialCommonBasis(TransformType);
 
 	constexpr bool bIsPositionTransform = true;
-	BuildTransformVectorBase(Em, &Input, TransformSourceBasis, TransformDestBasis, bIsPositionTransform, PeriodicWorldTileSizeValue, FirstPersonInterpolationAlphaValue);
+	BuildTransformVector(Em, &Input, TransformSourceBasis, TransformDestBasis, bIsPositionTransform, PeriodicWorldTileSizeValue, FirstPersonInterpolationAlphaValue);
 }
 
 void UMaterialExpressionReroute::Build(MIR::FEmitter& Em)
@@ -2911,7 +2904,7 @@ static FValueRef EmitLocalPosition(MIR::FEmitter& Em, ELocalPositionOrigin Local
 	if (LocalOrigin == ELocalPositionOrigin::InstancePreSkinning)
 	{
 		static FName NAME_PreSkinnedPosition("PreSkinnedPosition");
-		return EmitInlineHLSL(Em, NAME_PreSkinnedPosition);
+		return Em.InlineHLSL(NAME_PreSkinnedPosition);
 	}
 
 	// Otherwise, there's a 2x2 configuration of Origin and Offset type.  Given that ELocalPositionOrigin::InstancePreSkinning is unused, valid origin values are 0 or 2,
@@ -3470,7 +3463,7 @@ void UMaterialExpressionAtmosphericFogColor::Build(MIR::FEmitter& Em)
 {
 	// This node is deprecated in favor of UMaterialExpressionSkyAtmosphereAerialPerspective, and falls through to the newer expression
 	FValueRef PositionValue = EmitWorldPositionOrDefault(Em, Em.TryInput(&WorldPosition), WorldPositionOriginType);
-	Em.Output(0, EmitInlineHLSL(Em, *this, 0, { PositionValue }));
+	Em.Output(0, MaterialExpressionExternalCodeBase_EmitInlineHLSL(Em, *this, 0, { PositionValue }));
 }
 
 void UMaterialExpressionBlackBody::Build(MIR::FEmitter& Em)
@@ -3486,7 +3479,7 @@ void UMaterialExpressionDepthFade::Build(MIR::FEmitter& Em)
 	FValueRef FadeDistanceValue = Em.Max(Em.InputDefaultFloat(&FadeDistance, FadeDistanceDefault), Em.ConstantFloat(UE_DELTA));
 
 	static FName NAME_PixelDepth("PixelDepth");
-	FValueRef PixelDepth = EmitInlineHLSL(Em, NAME_PixelDepth);
+	FValueRef PixelDepth = Em.InlineHLSL(NAME_PixelDepth);
 	// On mobile scene depth is limited to 65500 
 	// to avoid false fading on objects that are close or exceed this limit we clamp pixel depth to (65500 - FadeDistance)
 	if (Em.GetFeatureLevel() <= ERHIFeatureLevel::ES3_1)
@@ -3520,7 +3513,7 @@ void UMaterialExpressionDistanceFieldApproxAO::Build(MIR::FEmitter& Em)
 	if (!NormalValue)
 	{
 		static FName NAME_VertexNormal("VertexNormal");
-		NormalValue = EmitInlineHLSL(Em, NAME_VertexNormal);
+		NormalValue = Em.InlineHLSL(NAME_VertexNormal);
 	}
 
 	FValueRef BaseDistanceValue = Em.InputDefaultFloat(&BaseDistance, BaseDistanceDefault);
@@ -3593,7 +3586,7 @@ void UMaterialExpressionFresnel::Build(MIR::FEmitter& Em)
 		NormalArg = Em.InlineHLSL(MIR::FType::MakeFloatVector(3), TEXTVIEW("Parameters.WorldNormal"), {}, MIR::EValueFlags::None, MIR::EGraphProperties::ReadsPixelNormal);
 	}
 
-	FValueRef DotArg = Em.Dot(NormalArg, EmitInlineHLSL(Em, NAME_CameraVector));
+	FValueRef DotArg = Em.Dot(NormalArg, Em.InlineHLSL(NAME_CameraVector));
 	FValueRef MaxArg = Em.Max(Em.ConstantFloat(0.f), DotArg);
 	FValueRef MinusArg = Em.Subtract(Em.ConstantFloat(1.f), MaxArg);
 	FValueRef ExponentArg = Em.InputDefaultFloat(&ExponentIn, Exponent);
@@ -3622,13 +3615,13 @@ void UMaterialExpressionReflectionVectorWS::Build(MIR::FEmitter& Em)
 		}
 
 		// Normal * dot(Normal, CameraVector) * 2.0 - CameraVector;
-		FValueRef CameraVector = EmitInlineHLSL(Em, NAME_CameraVector);
+		FValueRef CameraVector = Em.InlineHLSL(NAME_CameraVector);
 		Em.Output(0, Em.Subtract(Em.Multiply(NormalValue, Em.Multiply(Em.Dot(NormalValue, CameraVector), Em.ConstantFloat(2.0f))), CameraVector));
 	}
 	else
 	{
 		static FName NAME_ReflectionVector("ReflectionVector");
-		Em.Output(0, EmitInlineHLSL(Em, NAME_ReflectionVector));
+		Em.Output(0, Em.InlineHLSL(NAME_ReflectionVector));
 	}
 }
 
@@ -3695,12 +3688,12 @@ void UMaterialExpressionRotator::Build(MIR::FEmitter& Em)
 void UMaterialExpressionSkyAtmosphereAerialPerspective::Build(MIR::FEmitter& Em)
 {
 	FValueRef PositionValue = EmitWorldPositionOrDefault(Em, Em.TryInput(&WorldPosition), WorldPositionOriginType);
-	Em.Output(0, EmitInlineHLSL(Em, *this, 0, { PositionValue }));
+	Em.Output(0, MaterialExpressionExternalCodeBase_EmitInlineHLSL(Em, *this, 0, { PositionValue }));
 }
 
 void UMaterialExpressionSkyAtmosphereLightDirection::Build(MIR::FEmitter& Em)
 {
-	Em.Output(0, EmitInlineHLSL(Em, *this, 0, { Em.ConstantInt(LightIndex) }));
+	Em.Output(0, MaterialExpressionExternalCodeBase_EmitInlineHLSL(Em, *this, 0, { Em.ConstantInt(LightIndex) }));
 }
 
 void UMaterialExpressionSkyAtmosphereLightDiskLuminance::Build(MIR::FEmitter& Em)
@@ -3715,18 +3708,18 @@ void UMaterialExpressionSkyAtmosphereLightDiskLuminance::Build(MIR::FEmitter& Em
 	{
 		CosHalfDiskRadius = Em.ConstantFloat(-1.0f);
 	}
-	Em.Output(0, EmitInlineHLSL(Em, *this, 0, { Em.ConstantInt(LightIndex), CosHalfDiskRadius }));
+	Em.Output(0, MaterialExpressionExternalCodeBase_EmitInlineHLSL(Em, *this, 0, { Em.ConstantInt(LightIndex), CosHalfDiskRadius }));
 }
 
 void UMaterialExpressionSkyAtmosphereLightIlluminance::Build(MIR::FEmitter& Em)
 {
 	FValueRef PositionValue = EmitWorldPositionOrDefault(Em, Em.TryInput(&WorldPosition), WorldPositionOriginType);
-	Em.Output(0, EmitInlineHLSL(Em, *this, 0, { PositionValue, Em.ConstantInt(LightIndex) }));
+	Em.Output(0, MaterialExpressionExternalCodeBase_EmitInlineHLSL(Em, *this, 0, { PositionValue, Em.ConstantInt(LightIndex) }));
 }
 
 void UMaterialExpressionSkyAtmosphereLightIlluminanceOnGround::Build(MIR::FEmitter& Em)
 {
-	Em.Output(0, EmitInlineHLSL(Em, *this, 0, { Em.ConstantInt(LightIndex) }));
+	Em.Output(0, MaterialExpressionExternalCodeBase_EmitInlineHLSL(Em, *this, 0, { Em.ConstantInt(LightIndex) }));
 }
 
 void UMaterialExpressionSkyAtmosphereViewLuminance::Build(MIR::FEmitter& Em)
@@ -3734,9 +3727,9 @@ void UMaterialExpressionSkyAtmosphereViewLuminance::Build(MIR::FEmitter& Em)
 	FValueRef WorldDirectionValue = Em.TryInput(&WorldDirection);
 	if (!WorldDirectionValue)
 	{
-		WorldDirectionValue = Em.Multiply(Em.ConstantFloat(-1.0f), EmitInlineHLSL(Em, NAME_CameraVector));
+		WorldDirectionValue = Em.Multiply(Em.ConstantFloat(-1.0f), Em.InlineHLSL(NAME_CameraVector));
 	}
-	Em.Output(0, EmitInlineHLSL(Em, *this, 0, { WorldDirectionValue }));
+	Em.Output(0, MaterialExpressionExternalCodeBase_EmitInlineHLSL(Em, *this, 0, { WorldDirectionValue }));
 }
 
 void UMaterialExpressionSkyLightEnvMapSample::Build(MIR::FEmitter& Em)
@@ -3749,7 +3742,7 @@ void UMaterialExpressionSkyLightEnvMapSample::Build(MIR::FEmitter& Em)
 	FValueRef DirectionValue = Em.InputDefaultFloat3(&Direction, FVector3f(0.0f, 0.0f, 1.0f));
 	FValueRef RoughnessValue = Em.InputDefaultFloat(&Roughness, 0.0f);
 
-	Em.Output(0, EmitInlineHLSL(Em, *this, 0, { DirectionValue, RoughnessValue }));
+	Em.Output(0, MaterialExpressionExternalCodeBase_EmitInlineHLSL(Em, *this, 0, { DirectionValue, RoughnessValue }));
 }
 
 void UMaterialExpressionSphereMask::Build(MIR::FEmitter& Em)
@@ -3931,10 +3924,10 @@ void UMaterialExpressionBounds::Build(MIR::FEmitter& Em)
 	static_assert(MEILB_PreSkinnedLocal == 2);
 
 	int32 OutputOffset = Type * 4;
-	Em.Output(0, EmitInlineHLSL(Em, *this, OutputOffset + 0));
-	Em.Output(1, EmitInlineHLSL(Em, *this, OutputOffset + 1));
-	Em.Output(2, EmitInlineHLSL(Em, *this, OutputOffset + 2));
-	Em.Output(3, EmitInlineHLSL(Em, *this, OutputOffset + 3));
+	Em.Output(0, MaterialExpressionExternalCodeBase_EmitInlineHLSL(Em, *this, OutputOffset + 0));
+	Em.Output(1, MaterialExpressionExternalCodeBase_EmitInlineHLSL(Em, *this, OutputOffset + 1));
+	Em.Output(2, MaterialExpressionExternalCodeBase_EmitInlineHLSL(Em, *this, OutputOffset + 2));
+	Em.Output(3, MaterialExpressionExternalCodeBase_EmitInlineHLSL(Em, *this, OutputOffset + 3));
 }
 
 void UMaterialExpressionBumpOffset::Build(MIR::FEmitter& Em)
@@ -3950,7 +3943,7 @@ void UMaterialExpressionBumpOffset::Build(MIR::FEmitter& Em)
 	Em.Output(0,
 		Em.Add(
 			Em.Multiply(
-				Em.Swizzle(EmitTransformVectorBase(Em, EmitInlineHLSL(Em, NAME_CameraVector), MCB_World, MCB_Tangent, false, nullptr, nullptr), MIR::FSwizzleMask(MIR::EVectorComponent::X, MIR::EVectorComponent::Y)),
+				Em.Swizzle(EmitTransformVector(Em, Em.InlineHLSL(NAME_CameraVector), MCB_World, MCB_Tangent, false, nullptr, nullptr), MIR::FSwizzleMask(MIR::EVectorComponent::X, MIR::EVectorComponent::Y)),
 				Em.Add(
 					Em.Multiply(
 						HeightRatioArg,
@@ -3985,7 +3978,7 @@ void UMaterialExpressionNoise::Build(MIR::FEmitter& Em)
 	if (WorldPositionOriginType == EPositionOrigin::CameraRelative)
 	{
 		// LWC_TODO: add support for translated world positions in the corresponding HLSL function
-		PositionInput = EmitTransformVectorBase(Em, PositionInput, MCB_TranslatedWorld, MCB_World, true, nullptr, nullptr);
+		PositionInput = EmitTransformVector(Em, PositionInput, MCB_TranslatedWorld, MCB_World, true, {}, {});
 	}
 
 	FValueRef FilterWidthInput = Em.InputDefaultFloat(&FilterWidth, 0.0f);
@@ -4024,7 +4017,7 @@ void UMaterialExpressionVectorNoise::Build(MIR::FEmitter& Em)
 	if (WorldPositionOriginType == EPositionOrigin::CameraRelative)
 	{
 		// LWC_TODO: add support for translated world positions in the corresponding HLSL function
-		PositionInput = EmitTransformVectorBase(Em, PositionInput, MCB_TranslatedWorld, MCB_World, true, nullptr, nullptr);
+		PositionInput = EmitTransformVector(Em, PositionInput, MCB_TranslatedWorld, MCB_World, true, {}, {});
 	}
 
 	// LWC_TODO - maybe possible/useful to add LWC-aware noise functions
@@ -4092,7 +4085,7 @@ void UMaterialExpressionParticlePositionWS::Build(MIR::FEmitter& Em)
 	int32 ExternalCodeIndex = OriginType == EPositionOrigin::Absolute ? 0 : 1;
 	MIR::FType ResultType = OriginType == EPositionOrigin::Absolute ? MIR::FType::MakeDoubleVector(3) : MIR::FType::MakeFloatVector(3);
 
-	Em.Output(0, EmitInlineHLSL(Em, *this, ExternalCodeIndex, {}, MIR::EValueFlags::SubstituteTags));
+	Em.Output(0, MaterialExpressionExternalCodeBase_EmitInlineHLSL(Em, *this, ExternalCodeIndex, {}, MIR::EValueFlags::SubstituteTags));
 }
 
 void UMaterialExpressionPerInstanceCustomData::Build(MIR::FEmitter& Em)
@@ -4100,7 +4093,7 @@ void UMaterialExpressionPerInstanceCustomData::Build(MIR::FEmitter& Em)
 	FValueRef DataIndexArgument = Em.ConstantInt(DataIndex);
 	FValueRef DefaultArgument = Em.InputDefaultFloat(&DefaultValue, ConstDefaultValue);
 
-	Em.Output(0, EmitInlineHLSL(Em, *this, 0, { DataIndexArgument, DefaultArgument }));
+	Em.Output(0, MaterialExpressionExternalCodeBase_EmitInlineHLSL(Em, *this, 0, { DataIndexArgument, DefaultArgument }));
 }
 
 void UMaterialExpressionPerInstanceCustomData3Vector::Build(MIR::FEmitter& Em)
@@ -4108,7 +4101,7 @@ void UMaterialExpressionPerInstanceCustomData3Vector::Build(MIR::FEmitter& Em)
 	FValueRef DataIndexArgument = Em.ConstantInt(DataIndex);
 	FValueRef DefaultArgument = Em.InputDefaultFloat3(&DefaultValue, FVector3f(ConstDefaultValue.R, ConstDefaultValue.G, ConstDefaultValue.B));
 
-	Em.Output(0, EmitInlineHLSL(Em, *this, 0, { DataIndexArgument, DefaultArgument }));
+	Em.Output(0, MaterialExpressionExternalCodeBase_EmitInlineHLSL(Em, *this, 0, { DataIndexArgument, DefaultArgument }));
 }
 
 void UMaterialExpressionPreviousFrameSwitch::Build(MIR::FEmitter& Em)
@@ -4319,6 +4312,13 @@ void UMaterialExpressionThinTranslucentMaterialOutput::Build(MIR::FEmitter& Em)
 	Em.SetCustomOutputs(TEXTVIEW("ThinTranslucentMaterialOutput"), OutputValues, MIR::MaterialOutputStage_Pixel);
 }
 
+void UMaterialExpressionVolumetricAdvancedMaterialInput::Build(MIR::FEmitter& Em)
+{
+	FValueRef VolumeSampleConservativeDensity = Em.InlineHLSL(TEXT("GetVolumeSampleConservativeDensity"));
+	Em.Output(0, Em.CastToFloat(VolumeSampleConservativeDensity, 3));
+	Em.Output(1, VolumeSampleConservativeDensity);
+}
+
 void UMaterialExpressionVolumetricAdvancedMaterialOutput::Build(MIR::FEmitter& Em)
 {
 	FValueRef OutputValues[] =
@@ -4343,6 +4343,487 @@ void UMaterialExpressionAbsorptionMediumMaterialOutput::Build(MIR::FEmitter& Em)
 	Em.SetCustomOutputs(TEXTVIEW("AbsorptionMediumMaterialOutput"), OutputValues, MIR::MaterialOutputStage_Pixel);
 }
 
+void UMaterialExpressionDecalMipmapLevel::Build(MIR::FEmitter& Em)
+{
+	Em.CheckMaterialDomainIsAnyOf({ MD_DeferredDecal });
+
+	FValueRef TextureSizeValue = Em.CastToFloat(Em.InputDefaultFloat2(&TextureSize, { ConstWidth, ConstHeight }), 2);
+	BuildInlineHLSLOutput(Em, *this, { TextureSizeValue });
+}
+
+void UMaterialExpressionDepthOfFieldFunction::Build(MIR::FEmitter& Em)
+{
+	FValueRef DepthValue = Em.CastToFloat(Em.TryInput(&Depth), 1);
+	if (!DepthValue)
+	{
+		DepthValue = Em.InlineHLSL(FName("PixelDepth"));
+	}
+
+	UE_MIR_CHECKPOINT(Em);
+	BuildInlineHLSLOutput(Em, *this, { DepthValue, Em.ConstantInt((MIR::TInteger)FunctionValue) });
+}
+
+void UMaterialExpressionDistanceCullFade::Build(MIR::FEmitter& Em)
+{
+	Em.Output(0, Em.ExternalInput(MIR::EExternalInput::DistanceCullFade));
+}
+
+void UMaterialExpressionGIReplace::Build(MIR::FEmitter& Em)
+{
+	FValueRef DefaultValue = Em.Input(&Default);
+	UE_MIR_CHECKPOINT(Em);
+
+	// Note: StaticIndirect is unused.
+	FValueRef LocalDynamicIndirectValue = Em.InputDefault(&DynamicIndirect, DefaultValue);
+	FValueRef GIReplaceStateValue = Em.InlineHLSL(MIR::FType::MakeBoolScalar(), TEXT("GetGIReplaceState()"));
+
+	Em.Output(0, Em.Branch(GIReplaceStateValue, LocalDynamicIndirectValue, DefaultValue));
+}
+
+void UMaterialExpressionHsvToRgb::Build(MIR::FEmitter& Em)
+{
+    FValueRef InputValue = Em.CastToFloatKind(Em.Input(&Input));
+    UE_MIR_CHECKPOINT(Em);
+
+    MIR::FPrimitive Primitive = *InputValue->Type.AsPrimitive();
+    if ((!Primitive.IsRowVector() && !Primitive.IsColumnVector()) || Primitive.NumComponents() < 3)
+    {
+        Em.Error(InputValue, TEXT("Expected a 3D or 4D vector."));
+        return;
+    }
+
+    FValueRef H = Em.Subscript(InputValue, 0);
+    FValueRef S = Em.Subscript(InputValue, 1);
+    FValueRef V = Em.Subscript(InputValue, 2);
+
+    FValueRef One   = Em.ConstantFloat(1.0f);
+    FValueRef Two   = Em.ConstantFloat(2.0f);
+    FValueRef Three = Em.ConstantFloat(3.0f);
+    FValueRef Four  = Em.ConstantFloat(4.0f);
+    FValueRef Six   = Em.ConstantFloat(6.0f);
+
+    FValueRef HTimesSix = Em.Multiply(H, Six);
+
+    // R = abs(H * 6 - 3) - 1;
+    FValueRef R = Em.Subtract(Em.Abs(Em.Subtract(HTimesSix, Three)), One);
+
+    // G = 2 - abs(H * 6 - 2);
+    FValueRef G = Em.Subtract(Two, Em.Abs(Em.Subtract(HTimesSix, Two)));
+
+    // B = 2 - abs(H * 6 - 4);
+    FValueRef B = Em.Subtract(Two, Em.Abs(Em.Subtract(HTimesSix, Four)));
+
+    // RGB = saturate(float3(R, G, B));
+    FValueRef RGB = Em.Saturate(Em.Vector3(R, G, B));
+
+    // RGB = ((RGB - 1) * S + 1) * V;
+    RGB = Em.Multiply(
+        Em.Add(Em.Multiply(Em.Subtract(RGB, One), S), One),
+        V
+    );
+
+    if (Primitive.NumComponents() == 4)
+    {
+        RGB = Em.Vector4(RGB, Em.Subscript(InputValue, 3));
+    }
+
+    if (Primitive.IsColumnVector())
+    {
+        RGB = Em.Transpose(RGB);
+    }
+
+    Em.Output(0, RGB);
+}
+
+void UMaterialExpressionMaterialProxyReplace::Build(MIR::FEmitter& Em)
+{
+	// New translator does not have the concept of "proxy". Simply pass through the realtime input.
+	Em.Output(0, Em.Input(&Realtime));
+}
+
+void UMaterialExpressionPathTracingBufferTexture::Build(MIR::FEmitter& Em)
+{
+	Em.CheckMaterialDomainIsAnyOf({ MD_PostProcess });
+	UE_MIR_CHECKPOINT(Em);
+
+	FValueRef CoordinatesValue = Em.TryInput(&Coordinates);
+	if (!CoordinatesValue)
+	{
+		CoordinatesValue = Em.InlineHLSL(MIR::FType::MakeFloatVector(2), TEXT("GetDefaultPathTracingBufferTextureUV(Parameters, 0)"));
+	}
+	CoordinatesValue = Em.CastToFloat(CoordinatesValue, 2);
+
+	FValueRef PathTracingBufferTextureIndexValue = Em.ConstantInt(PathTracingBufferTextureId);
+
+	FValueRef OutputValue = Em.InlineHLSL(
+		MIR::FType::MakeFloatVector(4),
+		TEXTVIEW("MaterialExpressionPathTracingBufferTextureLookup(Parameters, $0, $1)"),
+		{ CoordinatesValue, PathTracingBufferTextureIndexValue },
+		MIR::EValueFlags::None, MIR::EGraphProperties::None,
+		[](FMaterialIRModule* Module, uintptr_t PathTracingBufferTextureId)
+		{
+			Module->GetCompilationOutput().SetIsPathTracingBufferTextureUsed((int32)PathTracingBufferTextureId);
+		},
+		PathTracingBufferTextureId);
+
+	Em.Output(0, OutputValue);
+}
+
+void UMaterialExpressionPostVolumeUserFlagTest::Build(MIR::FEmitter& Em)
+{
+	FValueRef BitIndexValue = Em.TryInput(&BitIndex); 
+	if (!BitIndexValue)
+	{
+		if (ConstBitIndex < 0)
+		{
+			Em.Error(TEXTVIEW("Invalid constant bit index provided."));
+			return;
+		}
+		BitIndexValue = Em.ConstantInt(ConstBitIndex);
+	}
+
+	FValueRef OutputValue = Em.InlineHLSL(MIR::FType::MakeFloatScalar(), TEXTVIEW("PostVolumeUserFlagTest($0)"), { BitIndexValue });
+	Em.Output(0, OutputValue);
+}
+
+void UMaterialExpressionPreSkinnedLocalBounds::Build(MIR::FEmitter& Em)
+{
+	Em.CheckMaterialDomainIsAnyOf( { MD_Surface, MD_Volume });
+	UE_MIR_CHECKPOINT(Em);
+
+	Em.Output(0, Em.InlineHLSL(TEXT("PreSkinnedLocalBoundsHalf"))); // Half extents
+	Em.Output(1, Em.InlineHLSL(TEXT("PreSkinnedLocalBoundsFull"))); // Full extents
+	Em.Output(2, Em.InlineHLSL(TEXT("PreSkinnedLocalBoundsMin")));  // Min point
+	Em.Output(3, Em.InlineHLSL(TEXT("PreSkinnedLocalBoundsMax")));  // Max point
+}
+
+static void EmitRayOrPathTracingQualitySwitch(MIR::FEmitter& Em, FStringView ConditionHLSL, FValueRef NormalValue, FValueRef TracedValue)
+{
+	if (!NormalValue && !TracedValue)
+	{
+		Em.Error(TEXT("Neither Normal nor PathTraced inputs are provided."));
+		return;
+	}
+
+	FValueRef OutputValue;
+	if (!NormalValue)
+	{
+		OutputValue = TracedValue;
+	}
+	else if (!TracedValue)
+	{
+		OutputValue = NormalValue;
+	}
+	else
+	{
+		FValueRef ConditionValue = Em.InlineHLSL(MIR::FType::MakeBoolScalar(), ConditionHLSL);
+		OutputValue = Em.Branch(ConditionValue, TracedValue, NormalValue);
+	}
+
+	Em.Output(0, OutputValue);
+}
+
+void UMaterialExpressionRayTracingQualitySwitch::Build(MIR::FEmitter& Em)
+{
+	FValueRef NormalValue = Em.TryInput(&Normal);
+	FValueRef RayTracedValue = FDataDrivenShaderPlatformInfo::GetSupportsRayTracing(Em.GetShaderPlatform()) ? Em.TryInput(&RayTraced) : FValueRef{};
+	EmitRayOrPathTracingQualitySwitch(Em, TEXTVIEW("GetRayTracingQualitySwitch()"), NormalValue, RayTracedValue);
+}
+
+void UMaterialExpressionPathTracingQualitySwitch::Build(MIR::FEmitter& Em)
+{
+	FValueRef NormalValue = Em.TryInput(&Normal);
+	FValueRef PathTracedValue = Em.TryInput(&PathTraced);
+	EmitRayOrPathTracingQualitySwitch(Em, TEXTVIEW("GetPathTracingQualitySwitch()"), NormalValue, PathTracedValue);
+}
+
+void UMaterialExpressionPathTracingRayTypeSwitch::Build(MIR::FEmitter& Em)
+{
+	FValueRef MainValue = Em.Input(&Main);
+	UE_MIR_CHECKPOINT(Em);
+
+	// compile all arguments (its ok if some of these are not connected, the will default to using Main)
+	FValueRef ShadowValue = Em.TryInput(&Shadow);
+	FValueRef DiffuseValue = Em.TryInput(&IndirectDiffuse);
+	FValueRef SpecularValue = Em.TryInput(&IndirectSpecular);
+	FValueRef VolumeValue = Em.TryInput(&IndirectVolume);
+
+	// return Compiler->PathTracingRayTypeSwitch(ArgMain, ArgShadow, ArgDiffuse, ArgSpecular, ArgVolume);
+		
+	FValueRef OutputValue;
+	OutputValue = Em.Branch(Em.InlineHLSL(MIR::FType::MakeBoolScalar(), TEXTVIEW("GetPathTracingIsShadow()")), ShadowValue ? ShadowValue : MainValue, MainValue);
+	OutputValue = Em.Branch(Em.InlineHLSL(MIR::FType::MakeBoolScalar(), TEXTVIEW("GetPathTracingIsIndirectDiffuse()")), DiffuseValue ? DiffuseValue : MainValue, OutputValue);
+	OutputValue = Em.Branch(Em.InlineHLSL(MIR::FType::MakeBoolScalar(), TEXTVIEW("GetPathTracingIsIndirectSpecular()")), SpecularValue ? SpecularValue : MainValue, OutputValue);
+	OutputValue = Em.Branch(Em.InlineHLSL(MIR::FType::MakeBoolScalar(), TEXTVIEW("GetPathTracingIsIndirectVolume()")), VolumeValue ? VolumeValue : MainValue, OutputValue);
+	
+	Em.Output(0, OutputValue);
+}
+
+void UMaterialExpressionBindlessSwitch::Build(MIR::FEmitter& Em)
+{
+	FValueRef OutputValue;
+	switch (UE::ShaderCompiler::GetBindlessConfiguration(Em.GetShaderPlatform()))
+	{
+		case ERHIBindlessConfiguration::Disabled:
+			OutputValue = Em.Input(&Default);
+			break;
+		case ERHIBindlessConfiguration::Minimal: 
+		case ERHIBindlessConfiguration::All:
+			OutputValue = Em.Input(&Bindless);
+			break;
+		case ERHIBindlessConfiguration::RayTracing:
+			OutputValue = Em.Input(&Bindless); // @massimo.tristano todo
+			break;
+	}
+	Em.Output(0, OutputValue);
+}
+
+static FValueRef EmitSamplePhysicsField(MIR::FEmitter& Em, FExpressionInput* PositionInput, EPositionOrigin PositionOriginType, EFieldOutputType OutputType, const int32 TargetIndex)
+{
+	FValueRef PositionValue = Em.TryInput(PositionInput);
+	if (!PositionValue)
+	{
+		PositionValue = EmitWorldPosition(Em, UE::MaterialTranslatorUtils::GetWorldPositionTypeWithOrigin(PositionOriginType));
+	}
+
+	if (Em.GetFeatureLevel() == ERHIFeatureLevel::ES3_1)
+	{
+		// TODO: add physics field sampling to mobile
+		if (OutputType == Field_Output_Vector)
+		{
+			return Em.ConstantFloat3({ 0.0f, 0.0f, 0.0f });
+		}
+		else if (OutputType == Field_Output_Scalar || OutputType == Field_Output_Integer)
+		{
+			return Em.ConstantFloat(0.0f);
+		}
+	}
+
+	Em.CheckFeatureLevelIsAtLeast(ERHIFeatureLevel::SM5);
+	if (Em.CurrentExpressionHasErrors())
+	{
+		return Em.Poison();
+	}
+
+	if (PositionOriginType == EPositionOrigin::CameraRelative)
+	{
+		// LWC_TODO: support translated-world coordinates
+		PositionValue = EmitTransformVector(Em, PositionValue, MCB_TranslatedWorld, MCB_World, true, {}, {});
+	}
+
+	// Cast the position to al float3
+	PositionValue = Em.Cast(PositionValue, MIR::FType::MakeFloatVector(3));
+
+	// Make a constant integer value from the target index.
+	FValueRef TargetIndexValue = Em.ConstantInt(TargetIndex);
+
+	// LWC_TODO: LWC aware physics field
+	if (OutputType == Field_Output_Vector)
+	{
+		return Em.InlineHLSL(MIR::FType::MakeFloatVector(3), TEXTVIEW("MatPhysicsField_SamplePhysicsVectorField($0, $1)"), { PositionValue, TargetIndexValue });
+	}
+	else if (OutputType == Field_Output_Scalar)
+	{
+		return Em.InlineHLSL(MIR::FType::MakeFloatScalar(), TEXTVIEW("MatPhysicsField_SamplePhysicsScalarField($0, $1)"), { PositionValue, TargetIndexValue });
+	}
+	else if (OutputType == Field_Output_Integer)
+	{
+		return Em.InlineHLSL(MIR::FType::MakeFloatScalar(), TEXTVIEW("MatPhysicsField_SamplePhysicsIntegerField($0, $1)"), { PositionValue, TargetIndexValue });
+	}
+	else
+	{
+		checkf(false, TEXT("Unsupported output type."));
+		return Em.Poison();
+	}
+}
+
+void UMaterialExpressionSamplePhysicsIntegerField::Build(MIR::FEmitter& Em)
+{
+	Em.Output(0, EmitSamplePhysicsField(Em, &WorldPosition, WorldPositionOriginType, Field_Output_Integer, FieldTarget));
+}
+
+void UMaterialExpressionSamplePhysicsScalarField::Build(MIR::FEmitter& Em)
+{
+	Em.Output(0, EmitSamplePhysicsField(Em, &WorldPosition, WorldPositionOriginType, Field_Output_Scalar, FieldTarget));
+}
+
+void UMaterialExpressionSamplePhysicsVectorField::Build(MIR::FEmitter& Em)
+{
+	Em.Output(0, EmitSamplePhysicsField(Em, &WorldPosition, WorldPositionOriginType, Field_Output_Vector, FieldTarget));
+}
+
+void UMaterialExpressionSobol::Build(MIR::FEmitter& Em)
+{
+	// Construct the following expression:
+	// floor(Cell) + float2(SobolIndex(SobolPixel(Cell), Index) ^ (Seed * 0x10000) & 0xffff) / float(0x10000)
+
+	FValueRef CellValue   = Em.CastToFloat(Em.InputDefaultFloat2(&Cell, { 0.f, 0.f }), 2);   // float2
+	FValueRef IndexValue  = Em.InputDefaultInt(&Index, ConstIndex);                          // int
+	FValueRef SeedValue   = Em.CastToFloat(Em.InputDefaultFloat2(&Seed, {(float)ConstSeed.X, (float)ConstSeed.Y }), 2); // float2
+	FValueRef SobolBits   = Em.InlineHLSL(MIR::FType::MakeIntVector(2), TEXT("SobolIndex(SobolPixel(uint2($0)), uint($1))"), { CellValue, IndexValue });
+	FValueRef SeedUint    = Em.CastToInt(Em.Multiply(SeedValue,Em.ConstantFloat(65536.0f)), 2);
+	FValueRef MaskedBits  = Em.BitwiseAnd(Em.BitwiseXor(SobolBits, SeedUint), Em.ConstantInt(0xffff));
+	FValueRef Normalized  = Em.Divide(Em.CastToFloat(MaskedBits, 2), Em.ConstantFloat(65536.0f));
+	FValueRef Result      = Em.Add(Em.Floor(CellValue), Normalized);
+
+	Em.Output(0, Result);
+}
+
+void UMaterialExpressionSpeedTree::Build(MIR::FEmitter& Em)
+{
+	if (Em.GetMaterialInterface()->GetUsageByFlag(MATUSAGE_SkeletalMesh))
+	{
+		Em.Error(TEXT("SpeedTree node not currently supported for Skeletal Meshes, please disable usage flag."));
+		return;
+	}
+
+	FValueRef GeometryValue = Em.CastToInt(Em.InputDefaultInt(&GeometryInput, GeometryType), 1);
+	FValueRef WindValue = Em.CastToInt(Em.InputDefaultInt(&WindInput, WindType), 1);
+	FValueRef LODValue = Em.CastToInt(Em.InputDefaultInt(&LODInput, LODType), 1);
+	FValueRef ExtraBendValue = Em.CastToVector(Em.TryInput(&ExtraBendWS), 3);
+	FValueRef UseExtraBendValue = Em.ConstantBool(ExtraBendValue);
+	FValueRef BillboardThresholdValue = Em.ConstantFloat(BillboardThreshold);
+
+	if (!ExtraBendValue)
+	{
+		ExtraBendValue = Em.ConstantFloat3({ 0.f, 0.f, 0.f });
+	}
+ 
+	// Only generate previous frame's computations if required and opted-in
+	FStringView InlineCode = TEXTVIEW("Get<PREV>SpeedTreeVertexOffset(Parameters, $0, $1, $2, $3, $4, $5)");
+	if (bAccurateWindVelocities)
+	{
+		InlineCode = TEXTVIEW("GetSpeedTreeVertexOffset(Parameters, $0, $1, $2, $3, $4, $5)");
+	}
+
+	FValueRef ResultValue = Em.InlineHLSL(
+		MIR::FType::MakeFloatVector(3),
+		InlineCode,
+		{ GeometryValue, WindValue, LODValue, BillboardThresholdValue, UseExtraBendValue, ExtraBendValue },
+		MIR::EValueFlags::None, MIR::EGraphProperties::None,
+		[](FMaterialIRModule* Module, uintptr_t)
+		{
+			Module->AddEnvironmentDefine(TEXT("USES_SPEEDTREE"));
+			Module->GetStatistics().NumVertexTexCoords = FMath::Max(Module->GetStatistics().NumVertexTexCoords, 8);
+		});
+
+	TStaticArray<FValueRef, MIR::EStage::NumStages> ValuePerStage = {};
+	ValuePerStage[MIR::EStage::Stage_Vertex] = ResultValue;
+
+	Em.Output(0, Em.StageSwitch(MIR::FType::MakeFloatVector(3), ValuePerStage));
+}
+
+void UMaterialExpressionSRGBColorToWorkingColorSpace::Build(MIR::FEmitter& Em)
+{
+	FValueRef ColorValue = Em.CheckIsVector(Em.Input(&Input), 3, 4);
+	FValueRef ResultValue = ColorValue;
+
+	UE_MIR_CHECKPOINT(Em);
+
+	if (!UE::Color::FColorSpace::GetWorking().IsSRGB())
+	{
+		const UE::Color::FColorSpaceTransform& Transform = UE::Color::FColorSpaceTransform::GetSRGBToWorkingColorSpace();
+
+		FValueRef R = Em.Dot(ColorValue, Em.ConstantFloat3({ (float)Transform.M[0][0], (float)Transform.M[1][0], (float)Transform.M[2][0] }));
+		FValueRef G = Em.Dot(ColorValue, Em.ConstantFloat3({ (float)Transform.M[0][1], (float)Transform.M[1][1], (float)Transform.M[2][1] }));
+		FValueRef B = Em.Dot(ColorValue, Em.ConstantFloat3({ (float)Transform.M[0][2], (float)Transform.M[1][2], (float)Transform.M[2][2] }));
+		
+		ResultValue = ColorValue->Type.IsVector(4)
+			? Em.Vector4(R, G, B, Em.Subscript(ColorValue, 3)) // We preserve the original alpha when applicable
+			: Em.Vector3(R, G, B);
+	}
+
+	Em.Output(0, ResultValue);
+}
+
+void UMaterialExpressionTemporalSobol::Build(MIR::FEmitter& Em)
+{
+	// The logic in this function implements this expression
+	// (float2(SobolIndex(SobolPixel(uint2(Parameters.SvPosition.xy)), uint(View.StateFrameIndexMod8 + 8 * $0)) ^ (uint2($1 * 0x10000) & 0xffff)) / 0x10000)
+
+    FValueRef IndexValue = Em.CastToInt(Em.InputDefaultInt(&Index, ConstIndex), 1);
+    FValueRef SeedValue  = Em.CastToFloat(Em.InputDefaultFloat2(&Seed, { (float)ConstSeed.X, (float)ConstSeed.Y }), 2);
+    FValueRef FrameMod8  = Em.InlineHLSL(MIR::FType::MakeIntScalar(), TEXT("uint(View.StateFrameIndexMod8)"));
+
+    // View.StateFrameIndexMod8 + 8 * Index
+    FValueRef FrameIndex  = Em.Add(FrameMod8, Em.Multiply(IndexValue, Em.ConstantInt(8)));
+
+    // SobolIndex(SobolPixel(uint2(Parameters.SvPosition.xy)), FrameIndexU)
+    FValueRef Sobol = Em.InlineHLSL(MIR::FType::MakeIntVector(2), TEXT("SobolIndex(SobolPixel(uint2(Parameters.SvPosition.xy)), uint($0))"), { FrameIndex });
+
+	// uint2(SeedValue * 0x10000) & 0xffff)
+	FValueRef Scale      = Em.ConstantFloat(65536.0f); // 0x1000
+    FValueRef SeedScaled = Em.Multiply(SeedValue, Scale);
+    FValueRef SeedMasked = Em.BitwiseAnd(Em.CastToInt(SeedScaled, 2), Em.ConstantInt2({ 0xffff, 0xffff }));
+
+	Sobol = Em.CastToFloat(Em.BitwiseXor(Sobol, SeedMasked), 2);
+	Sobol = Em.Divide(Sobol, Scale);
+
+    Em.Output(0, Sobol);
+}
+
+void UMaterialExpressionNeuralNetworkInput::Build(MIR::FEmitter& Em)
+{
+	int32 CodeInput = INDEX_NONE;
+	bool bUseTextureAsInput = NeuralIndexType == ENeuralIndexType::NIT_TextureIndex;
+
+	FValueRef CoordinatesValue = Em.TryInput(&Coordinates);
+	if (CoordinatesValue)
+	{
+		if (bUseTextureAsInput)
+		{
+			CoordinatesValue = Em.Swizzle(CoordinatesValue, MIR::FSwizzleMask { MIR::EVectorComponent::Y, MIR::EVectorComponent::Z, MIR::EVectorComponent::W });
+			CoordinatesValue = Em.Vector4(Em.ConstantFloat(-1.0f), CoordinatesValue);
+		}
+	}
+	else
+	{
+		FValueRef ViewportUV = Em.InlineHLSL(TEXT("ViewportUV"));
+		float     BatchIndex = bUseTextureAsInput ? -1.0f : 0.0f;
+		CoordinatesValue     = Em.Vector4(Em.ConstantFloat2({ BatchIndex, 0.0f }), ViewportUV);
+	}
+
+	FValueRef OutputValues[] =
+	{
+		CoordinatesValue, 
+		Em.CastToFloat(Em.InputDefaultFloat3(&Input0, { 0.5f, 0.5f, 0.5f }), 3), 
+		Em.CastToFloat(Em.InputDefaultFloat(&Mask, 1.0f), 1)
+	};
+
+	Em.SetCustomOutputs(TEXTVIEW("NeuralInput"), OutputValues, MIR::MaterialOutputStage_Pixel);
+}
+
+void UMaterialExpressionNeuralNetworkOutput::Build(MIR::FEmitter& Em)
+{
+	Em.CheckMaterialDomainIsAnyOf({ MD_PostProcess });
+	UE_MIR_CHECKPOINT(Em);
+
+	FValueRef ViewportUV = Em.CastToFloat(Em.TryInput(&Coordinates), 2);
+	UE_MIR_CHECKPOINT(Em);
+
+	FValueRef BufferIndex = ViewportUV;
+
+	if (!ViewportUV)
+	{
+		ViewportUV = Em.InlineHLSL(TEXT("ViewportUV"));
+		BufferIndex = Em.Vector(Em.ConstantFloat2({ 0.0f, 0.0f }), ViewportUV);
+	}
+
+	BufferIndex = Em.CastToFloat(BufferIndex, 4);
+
+	// @massimo.tristano another note on this. do we need this in the new translator?
+	// AddEstimatedTextureSample();
+
+	static const auto EnableWithNeuralNetworks = [] (FMaterialIRModule* Module, uintptr_t)
+	{
+		Module->GetCompilationOutput().bUsedWithNeuralNetworks = true;
+	};
+
+	Em.Output(0, Em.InlineHLSL(MIR::FType::MakeFloatVector(4), TEXTVIEW("NeuralTextureOutput(Parameters,$0)"), { ViewportUV }, {}, {}, EnableWithNeuralNetworks));
+	Em.Output(1, Em.InlineHLSL(MIR::FType::MakeFloatVector(4), TEXTVIEW("NeuralBufferOutput(Parameters,$0)"), { BufferIndex }, {}, {}, EnableWithNeuralNetworks));
+}
+
+
 ////////////////////
 //  Substrate IR  //
 ////////////////////
@@ -4461,13 +4942,12 @@ static FValueRef TransformNormalFromRequestedBasisToWorld(MIR::FEmitter& Em, UMa
 	if (IsTangentSpaceNormal(Material))
 	{
 		// See TransformTangentNormalToWorld definitions in MaterialTemplate.ush
-		return EmitTransformVectorBase(Em, NormalArg, MCB_Tangent, MCB_World, false, nullptr, nullptr);
+		return EmitTransformVector(Em, NormalArg, MCB_Tangent, MCB_World, false, {}, {});
 	}
 	// Normalize
 	return NormalArg = Em.Multiply(NormalArg, Em.Rsqrt(Em.Dot(NormalArg, NormalArg)));
 }
 
-
 // The compilation of an expression can sometimes lead to a INDEX_NONE code chunk when editing material graphs 
 // or when the node is inside a material function, linked to an input pin of the material function and that input is not plugged in to anything.
 // But for normals or tangents, Substrate absolutely need a valid code chunk to de-duplicate when stored in memory. 
@@ -5111,15 +5591,15 @@ void UMaterialExpressionSubstrateHorizontalMixing::Build(MIR::FEmitter& Em)
 
 		// The tangent is optional so we treat it differently if INDEX_NONE is specified
 		FValueRef NewTangentArg;
-		if (ForegroundBasis.BasisInfo.TangentArg != nullptr && BackgroundBasis.BasisInfo.TangentArg != nullptr)
+		if (ForegroundBasis.BasisInfo.TangentArg && BackgroundBasis.BasisInfo.TangentArg)
 		{
 			NewTangentArg = SubstrateBlendNormal(Em, BackgroundBasis.BasisInfo.TangentArg, ForegroundBasis.BasisInfo.TangentArg, NormalMixArg);
 		}
-		else if (ForegroundBasis.BasisInfo.TangentArg != nullptr)
+		else if (ForegroundBasis.BasisInfo.TangentArg)
 		{
 			NewTangentArg = ForegroundBasis.BasisInfo.TangentArg;
 		}
-		else if (BackgroundBasis.BasisInfo.TangentArg != nullptr)
+		else if (BackgroundBasis.BasisInfo.TangentArg)
 		{
 			NewTangentArg = BackgroundBasis.BasisInfo.TangentArg;
 		}
@@ -5134,7 +5614,7 @@ void UMaterialExpressionSubstrateHorizontalMixing::Build(MIR::FEmitter& Em)
 			Em.Errorf(TEXT("Invalid Horizontal operator Foreground, Background, Mix or NormalMix input."));
 		}
 
-		FValueRef CameraVectorArg = EmitInlineHLSL(Em, NAME_CameraVector);
+		FValueRef CameraVectorArg = Em.InlineHLSL(NAME_CameraVector);
 		FValueRef BackgroundNoV = Em.Saturate(Em.Dot(BackgroundBasis.BasisInfo.NormalArg, CameraVectorArg));
 		FValueRef ForegroundNoV = Em.Saturate(Em.Dot(ForegroundBasis.BasisInfo.NormalArg, CameraVectorArg));
 
@@ -5258,15 +5738,15 @@ void UMaterialExpressionSubstrateVerticalLayering::Build(MIR::FEmitter& Em)
 
 		// The tangent is optional so we treat it differently if INDEX_NONE is specified
 		FValueRef NewTangentArg;
-		if (TopBasis.BasisInfo.TangentArg != nullptr && BaseBasis.BasisInfo.TangentArg != nullptr)
+		if (TopBasis.BasisInfo.TangentArg && BaseBasis.BasisInfo.TangentArg)
 		{
 			NewTangentArg = SubstrateBlendNormal(Em, BaseBasis.BasisInfo.TangentArg, TopBasis.BasisInfo.TangentArg, NormalMixArg);
 		}
-		else if (TopBasis.BasisInfo.TangentArg != nullptr)
+		else if (TopBasis.BasisInfo.TangentArg)
 		{
 			NewTangentArg = TopBasis.BasisInfo.TangentArg;
 		}
-		else if (BaseBasis.BasisInfo.TangentArg != nullptr)
+		else if (BaseBasis.BasisInfo.TangentArg)
 		{
 			NewTangentArg = BaseBasis.BasisInfo.TangentArg;
 		}
@@ -5280,7 +5760,7 @@ void UMaterialExpressionSubstrateVerticalLayering::Build(MIR::FEmitter& Em)
 			Em.Errorf(TEXT("Invalid Vertical operator Base, Top, Thickness or NormalMix input."));
 		}
 
-		FValueRef CameraVectorArg = EmitInlineHLSL(Em, NAME_CameraVector);
+		FValueRef CameraVectorArg = Em.InlineHLSL(NAME_CameraVector);
 		FValueRef TopNoV = Em.Saturate(Em.Dot(TopBasis.BasisInfo.NormalArg, CameraVectorArg));
 		FValueRef BaseNoV = Em.Saturate(Em.Dot(BaseBasis.BasisInfo.NormalArg, CameraVectorArg));
 
diff --git a/Engine/Source/Runtime/Engine/Private/Materials/MaterialIR.cpp b/Engine/Source/Runtime/Engine/Private/Materials/MaterialIR.cpp
--- a/Engine/Source/Runtime/Engine/Private/Materials/MaterialIR.cpp
+++ b/Engine/Source/Runtime/Engine/Private/Materials/MaterialIR.cpp
@@ -30,6 +30,7 @@ const TCHAR* LexToString(EValueKind Kind)
 	switch (Kind)
 	{
 		case VK_Poison: return TEXT("Poison");
+		case VK_AnalysisError: return TEXT("AnalysisError");
 		case VK_Constant: return TEXT("Constant");
 		case VK_ExternalInput: return TEXT("ExternalInput");
 		case VK_MaterialParameterCollection: return TEXT("MaterialParameterCollection");
@@ -102,6 +103,7 @@ uint32 FValue::GetSizeInBytes() const
 	switch (Kind)
 	{
 		case VK_Poison: return sizeof(FPoison);
+		case VK_AnalysisError: return sizeof(FAnalysisError);
 		case VK_Constant: return sizeof(FConstant);
 		case VK_ExternalInput: return sizeof(FExternalInput);
 		case VK_MaterialParameterCollection: return sizeof(FMaterialParameterCollection);
@@ -595,6 +597,8 @@ const TCHAR* LexToString(EExternalInput Input)
 		case EExternalInput::GlobalDistanceField: return TEXT("GlobalDistanceField");
 		case EExternalInput::DynamicParticleParameterIndex: return TEXT("DynamicParticleParameterIndex");
 		case EExternalInput::CompilingPreviousFrame: return TEXT("CompilingPreviousFrame");
+		case EExternalInput::DistanceCullFade: return TEXT("GetDistanceCullFade()");
+
 		default: UE_MIR_UNREACHABLE();
 	}
 }
@@ -611,17 +615,8 @@ FType GetExternalInputType(EExternalInput Id)
 		case EExternalInput::ViewMaterialTextureMipBias:
 		case EExternalInput::ViewMaterialTextureDerivativeMultiply:
 		case EExternalInput::GlobalDistanceField:						// Type is arbitrary, value not actually used, only as a flag for value analyzer
-			 return FType::MakeFloatScalar();
-
-		case EExternalInput::ActorPosition_Absolute:
-		case EExternalInput::ObjectPosition_Absolute:
-		case EExternalInput::WorldPosition_Absolute:
-		case EExternalInput::WorldPosition_AbsoluteNoOffsets:
-		case EExternalInput::WorldPosition_Absolute_Ddx:				// Technically these derivatives are not LWC, but the emitter needs them to be for evaluation of
-		case EExternalInput::WorldPosition_AbsoluteNoOffsets_Ddx:		// analytic derivative expressions.  So we cast them to LWC on load.  The shader compiler should
-		case EExternalInput::WorldPosition_Absolute_Ddy:				// be able to optimize the math (the tile offsets will all be zero, cancelling things out).
-		case EExternalInput::WorldPosition_AbsoluteNoOffsets_Ddy:
-			return FType::MakeDoubleVector(3);
+		case EExternalInput::DistanceCullFade:
+			return FType::MakeFloatScalar();
 
 		case EExternalInput::ActorPosition_CameraRelative:
 		case EExternalInput::ObjectPosition_CameraRelative:
@@ -645,6 +640,16 @@ FType GetExternalInputType(EExternalInput Id)
 		case EExternalInput::LocalPosition_PrimitiveNoOffsets_Ddy:
 			return FType::MakeFloatVector(3);
 
+		case EExternalInput::ActorPosition_Absolute:
+		case EExternalInput::ObjectPosition_Absolute:
+		case EExternalInput::WorldPosition_Absolute:
+		case EExternalInput::WorldPosition_AbsoluteNoOffsets:
+		case EExternalInput::WorldPosition_Absolute_Ddx:				// Technically these derivatives are not LWC, but the emitter needs them to be for evaluation of
+		case EExternalInput::WorldPosition_AbsoluteNoOffsets_Ddx:		// analytic derivative expressions.  So we cast them to LWC on load.  The shader compiler should
+		case EExternalInput::WorldPosition_Absolute_Ddy:				// be able to optimize the math (the tile offsets will all be zero, cancelling things out).
+		case EExternalInput::WorldPosition_AbsoluteNoOffsets_Ddy:
+			return FType::MakeDoubleVector(3);
+
 		case EExternalInput::DynamicParticleParameterIndex:
 			return FType::MakeIntScalar();
 
@@ -902,6 +907,7 @@ const TCHAR* LexToString(EOperator Op)
 		case BO_BitShiftRight: return TEXT("BitShiftRight");
 		case BO_BitwiseAnd: return TEXT("BitwiseAnd");
 		case BO_BitwiseOr: return TEXT("BitwiseOr");
+		case BO_BitwiseXor: return TEXT("BitwiseXor");
 		case BO_Cross: return TEXT("Cross");
 		case BO_Distance: return TEXT("Distance");
 		case BO_Divide: return TEXT("Divide");
diff --git a/Engine/Source/Runtime/Engine/Private/Materials/MaterialIREmitter.cpp b/Engine/Source/Runtime/Engine/Private/Materials/MaterialIREmitter.cpp
--- a/Engine/Source/Runtime/Engine/Private/Materials/MaterialIREmitter.cpp
+++ b/Engine/Source/Runtime/Engine/Private/Materials/MaterialIREmitter.cpp
@@ -8,6 +8,7 @@
 #include "Materials/MaterialExpressionTextureBase.h"
 #include "Materials/MaterialExternalCodeRegistry.h"
 #include "Materials/MaterialAggregate.h"
+#include "Materials/MaterialInterface.h"
 #include "Shader/Preshader2.h"
 #include "Shader/ShaderTypes.h"
 #include "MaterialShared.h"
@@ -581,9 +582,15 @@ FValueRef FEmitter::InputDefaultFloat4(const FExpressionInput* Input, UE::Math::
 	return Value ? Value : Value.To(ConstantFloat4(Default));
 }
 
+FValueRef FEmitter::InputDefault(const FExpressionInput* Input, FValueRef Default)
+{
+	FValueRef Value = TryInput(Input);
+	return Value ? Value : Default;
+}
+
 FValueRef FEmitter::CheckTypeIsKind(FValueRef Value, ETypeKind Kind)
 {
-	if (Value.IsValid() && Value->Type.Is(Kind))
+	if (Value.IsValid() && !Value->Type.Is(Kind))
 	{
 		Errorf(Value, TEXT("Expected a '%s' value, got a '%s' instead."), TypeKindToString(Kind), *Value->Type.GetSpelling());
 		return Value.ToPoison();
@@ -651,6 +658,28 @@ FValueRef FEmitter::CheckIsVector(FValueRef Value)
 	return Value;
 }
 
+FValueRef FEmitter::CheckIsVector(FValueRef Value, int32 MinNumComponents, int32 MaxNumComponents)
+{
+	Value = CheckIsVector(Value);
+	if (Value.IsValid())
+	{
+		int32 NumComponents = Value->Type.AsPrimitive()->NumComponents();
+		if (NumComponents < MinNumComponents)
+		{
+			(MinNumComponents == MaxNumComponents) 
+			? Errorf(Value, TEXT("Expected a vector value of exactly %d components, got a '%s' instead."), MinNumComponents, *Value->Type.GetSpelling())
+			: Errorf(Value, TEXT("Expected a vector value of at least %d components, got a '%s' instead."), MinNumComponents, *Value->Type.GetSpelling());
+			return Value.ToPoison();
+		}
+		if (NumComponents > MaxNumComponents)
+		{
+			Errorf(Value, TEXT("Expected a vector value of at most %d components, got a '%s' instead."), MaxNumComponents, *Value->Type.GetSpelling());
+			return Value.ToPoison();
+		}
+	}
+	return Value;
+}
+
 FValueRef FEmitter::CheckIsScalarOrVector(FValueRef Value)
 {
 	if (Value.IsValid() && (!Value->Type.AsPrimitive() || Value->Type.GetPrimitive().IsMatrix()))
@@ -700,6 +729,50 @@ FValueRef FEmitter::CheckIsAggregate(FValueRef Value, const UMaterialAggregate*
 	return Value;
 }
 
+void FEmitter::CheckMaterialDomainIsAnyOf(TConstArrayView<EMaterialDomain> Domains)
+{
+	if (Domains.Contains(BaseMaterial->MaterialDomain))
+	{
+		return;
+	}
+
+	// Report the error message.
+	FString ErrorMessage = TEXT("Expression does not support material domain '");
+	ErrorMessage += StaticEnum<EMaterialDomain>()->GetValueAsString(BaseMaterial->MaterialDomain);
+	ErrorMessage += TEXTVIEW("' (supported are ");
+
+	for (int32 i = 0; i < Domains.Num(); ++i)
+	{
+		if (i > 0)
+		{
+			ErrorMessage += TEXTVIEW(", ");
+		}
+		ErrorMessage += StaticEnum<EMaterialDomain>()->GetValueAsString(Domains[i]);
+	}
+	ErrorMessage += TEXTVIEW(")");
+
+	Error(ErrorMessage);
+}
+
+void FEmitter::CheckMaterialHasUsageFlag(EMaterialUsage UsageFlag)
+{
+	if (!MaterialInterface->GetUsageByFlag(UsageFlag))
+	{
+		Errorf(TEXT("This expression requires the material to enable usage flags %s."), *StaticEnum<EMaterialUsage>()->GetValueAsString(UsageFlag));
+	}
+}
+
+void FEmitter::CheckFeatureLevelIsAtLeast(ERHIFeatureLevel::Type MinFeatureLevel)
+{
+	if (GetFeatureLevel() < MinFeatureLevel)
+	{
+		FString FeatureLevelName, MinFeatureLevelName;
+		GetFeatureLevelName(GetFeatureLevel(), FeatureLevelName);
+		GetFeatureLevelName(MinFeatureLevel, MinFeatureLevelName);
+		Errorf(TEXT("Expression not supported with feature level %s. %s or higher required."), *FeatureLevelName, *MinFeatureLevelName);
+	}
+}
+
 bool FEmitter::ToConstantBool(FValueRef Value)
 {
 	if (!Value.IsValid())
@@ -728,15 +801,6 @@ FEmitter& FEmitter::Output(int32 OutputIndex, FValueRef Value)
 	return *this;
 }
 
-FEmitter& FEmitter::Output(const FExpressionOutput* ExpressionOutput, FValueRef Value)
-{
-	if (Value)
-	{
-		Internal::BindValueToExpressionOutput(BuilderImpl, ExpressionOutput, Value);
-	}
-	return *this;
-}
-
 static MIR::FSwizzleMask GetSwizzleMaskFromExpressionOutputMask(const FExpressionOutput& ExpressionOutput)
 {
 	MIR::FSwizzleMask SwizzleMask;
@@ -759,13 +823,27 @@ static MIR::FSwizzleMask GetSwizzleMaskFromExpressionOutputMask(const FExpressio
 	return SwizzleMask;
 }
 
-FEmitter& FEmitter::Outputs(const TConstArrayView<FExpressionOutput>& ExpressionOutputs, FValueRef Value)
+FEmitter& FEmitter::Output(const FExpressionOutput* ExpressionOutput, FValueRef Value)
+{
+	if (Value.IsValid())
+	{
+		Internal::BindValueToExpressionOutput(BuilderImpl, ExpressionOutput, Value);
+	}
+	return *this;
+}
+
+FEmitter& FEmitter::OutputsWithComponentMask(FValueRef Value)
 {
-	for (const FExpressionOutput& CurrentOutput : ExpressionOutputs)
+	for (const FExpressionOutput& ExpressionOutput : Expression->GetOutputs())
 	{
-		// Apply component swizzling for each expression output but use unmodified input value if swizzling is redundant (i.e. XYZW vector)
-		const MIR::FSwizzleMask SwizzleMask = GetSwizzleMaskFromExpressionOutputMask(CurrentOutput);
-		Output(&CurrentOutput, SwizzleMask.IsXYZW() ? Value : Swizzle(Value, SwizzleMask));
+		// Get the swizzle mask from this expression output component mask
+		FSwizzleMask SwizzleMask = GetSwizzleMaskFromExpressionOutputMask(ExpressionOutput);
+
+		// Apply the output's swizzle mask to the value
+		FValueRef OutputValue = SwizzleMask.NumComponents > 0 ? Swizzle(Value, SwizzleMask) : Value;
+
+		// Bind the resulting value to the output
+		Output(&ExpressionOutput, OutputValue);
 	}
 	return *this;
 }
@@ -982,6 +1060,14 @@ FValueRef FEmitter::Poison()
 	return FPoison::Get();
 }
 
+FValueRef FEmitter::AnalysisError(FType Type, FStringView Message)
+{
+	FAnalysisError* Error = EmitNew<FAnalysisError>(*this, Type);
+	Error->Expression = Expression;
+	Error->Message = Module->InternString(Message);
+	return Error;
+}
+
 FValueRef FEmitter::ExternalInput(EExternalInput Id, uint32 UserData)
 {
 	FExternalInput Prototype = MakePrototype<FExternalInput>(GetExternalInputType(Id));
@@ -1261,7 +1347,7 @@ FValueRef FEmitter::PreshaderParameter(FType Type, EPreshader2Opcode Opcode, FVa
 	FPreshaderParameter Prototype = MakePrototype<FPreshaderParameter>(Type);
 	Prototype.SourceParameter = SourceParameter;
 	Prototype.Opcode = Opcode;
-	Prototype.TextureIndex = Material->GetReferencedTextures().Find(SourceParameterTexture);
+	Prototype.TextureIndex = MaterialInterface->GetReferencedTextures().Find(SourceParameterTexture);
 	Prototype.Payload = MoveTemp(Payload);
 	return EmitPrototype(*this, Prototype);
 }
@@ -1392,179 +1478,164 @@ void FEmitter::SetCustomOutputs(FStringView Name, TConstArrayView<FValueRef> Arg
 	}
 }
 
-FValueRef FEmitter::Vector2(FValueRef InX, FValueRef InY)
+static void InvalidVectorConstructionError(FEmitter& Em, int32 DesiredVectorDimensions, FValueRef A, FValueRef B, FValueRef C = {}, FValueRef D = {})
 {
-	if (IsAnyNotValid(InX, InY))
+	if (C.IsValid() && D.IsValid())
 	{
-		return Poison();
+		Em.Errorf(TEXT("Cannot construct a %dD vector from concatenating primitive values of %d, %d, %d and %d dimensions"), DesiredVectorDimensions,
+				  A->Type.GetPrimitive().NumComponents(), 
+				  B->Type.GetPrimitive().NumComponents(), 
+				  C->Type.GetPrimitive().NumComponents(), 
+				  D->Type.GetPrimitive().NumComponents());
 	}
+	else if (C.IsValid())
+	{
+		Em.Errorf(TEXT("Cannot construct a %dD vector from concatenating primitive values of %d, %d and %d dimensions"), DesiredVectorDimensions,
+				  A->Type.GetPrimitive().NumComponents(), 
+				  B->Type.GetPrimitive().NumComponents(), 
+				  C->Type.GetPrimitive().NumComponents());
+	}
+	else
+	{
+		Em.Errorf(TEXT("Cannot construct a %dD vector from concatenating primitive values of %d and %d dimensions"), DesiredVectorDimensions,
+				  A->Type.GetPrimitive().NumComponents(), 
+				  B->Type.GetPrimitive().NumComponents());
+	}
+}
 
-	// TODO: replace these checks with casts instead
-	check(InX->Type.IsScalar());
-	check(InX->Type == InY->Type);
-
-	TComposite<2> Vector = MakePrototype<TComposite<2>>(FType::MakeVector(InX->Type.GetPrimitive().ScalarKind, 2));
-
-	TArrayView<FValue*> Components = Vector.GetMutableComponents();
-	Components[0] = InX;
-	Components[1] = InY;
+FValueRef FEmitter::Vector2(FValueRef X, FValueRef Y)
+{
+	// Vector2 has only one constructor from two scalars, for better error reporting, check arguments are scalars first.
+	X = CheckIsScalar(X);
+	Y = CheckIsScalar(Y);
 
-	FValueRef Value = EmitPrototype(*this, Vector);
-	
-	// If all components source inputs match, refer to the same input in the vector value too (for debugging purposes).
-	if (InX.Input == InY.Input)
+	FValueRef Result = Vector(X, Y);
+	if (Result.IsValid() && Result->Type.AsPrimitive()->NumComponents() != 2)
 	{
-		Value.Input = InX.Input;
+		InvalidVectorConstructionError(*this, 2, X, Y);
+		return Poison();
 	}
-
-	return Value;
+	return Result;
 }
 
-FValueRef FEmitter::Vector3(FValueRef InXY, FValueRef InZ)
+FValueRef FEmitter::Vector3(FValueRef A, FValueRef B, FValueRef C)
 {
-	if (IsAnyNotValid(InXY, InZ))
+	FValueRef Result = Vector(A, B, C);
+	if (Result.IsValid() && Result->Type.AsPrimitive()->NumComponents() != 3)
 	{
+		InvalidVectorConstructionError(*this, 3, A, B, C);
 		return Poison();
 	}
+	return Result;
+}
 
-	// TODO: replace these checks with casts instead
-	check(InXY->Type.IsVector() && InXY->Type.GetPrimitive().NumComponents() == 2);
-	check(InZ->Type.IsPrimitive() && InZ->Type.GetPrimitive().ScalarKind == InXY->Type.GetPrimitive().ScalarKind);
-
-	TComposite<3> Vector = MakePrototype<TComposite<3>>(FType::MakeVector(InXY->Type.GetPrimitive().ScalarKind, 3));
-	TArrayView<FValue*> Components = Vector.GetMutableComponents();
-	Components[0] = Subscript(InXY, 0);
-	Components[1] = Subscript(InXY, 1);
-	Components[2] = InZ;
-
-	FValueRef Value = EmitPrototype(*this, Vector);
-
-	// If all components source inputs match, refer to the same input in the vector value too (for debugging purposes).
-	if (InXY.Input == InZ.Input )
+FValueRef FEmitter::Vector4(FValueRef A, FValueRef B, FValueRef C, FValueRef D)
+{
+	FValueRef Result = Vector(A, B, C, D);
+	if (Result.IsValid() && Result->Type.AsPrimitive()->NumComponents() != 4)
 	{
-		Value.Input = InXY.Input;
+		InvalidVectorConstructionError(*this, 4, A, B, C, D);
+		return Poison();
 	}
-
-	return Value;
+	return Result;
 }
 
-FValueRef FEmitter::Vector3(FValueRef InX, FValueRef InY, FValueRef InZ)
+FValueRef FEmitter::Vector(FValueRef A, FValueRef B, FValueRef C, FValueRef D)
 {
-	if (IsAnyNotValid(InX, InY, InZ))
+	check(!D || C); // If D is provided, then C must be provided too.
+
+	// First verify that all given arguments are scalars or vectors
+	if (CheckIsScalarOrVector(A).IsPoison() || CheckIsScalarOrVector(B).IsPoison() || CheckIsScalarOrVector(C).IsPoison() || CheckIsScalarOrVector(D).IsPoison())
 	{
 		return Poison();
 	}
-	
-	// TODO: replace these checks with casts instead
-	check(InX->Type.IsScalar());
-	check(InX->Type == InY->Type);
-	check(InY->Type == InZ->Type);
-
-	TComposite<3> Vector = MakePrototype<TComposite<3>>(FType::MakeVector(InX->Type.GetPrimitive().ScalarKind, 3));
-	TArrayView<FValue*> Components = Vector.GetMutableComponents();
-	Components[0] = InX;
-	Components[1] = InY;
-	Components[2] = InZ;
 
-	FValueRef Value = EmitPrototype(*this, Vector);
+	FPrimitive APrimType = *A->Type.AsPrimitive();
+	FPrimitive BPrimType = *B->Type.AsPrimitive();
+	TOptional<FPrimitive> CPrimType = C.IsValid() ? C->Type.AsPrimitive() : TOptional<FPrimitive>{};
+	TOptional<FPrimitive> DPrimType = D.IsValid() ? D->Type.AsPrimitive() : TOptional<FPrimitive>{};
 
-	// If all components source inputs match, refer to the same input in the vector value too (for debugging purposes).
-	if (InX.Input == InY.Input && InX.Input == InZ.Input)
+	// Then compute the common scalar kind
+	EScalarKind CommonScalarKind = GetCommonScalarKind(APrimType.ScalarKind, BPrimType.ScalarKind);
+	if (CPrimType)
 	{
-		Value.Input = InX.Input;
+		CommonScalarKind = GetCommonScalarKind(CommonScalarKind, CPrimType->ScalarKind);
+		if (DPrimType)
+		{
+			CommonScalarKind = GetCommonScalarKind(CommonScalarKind, DPrimType->ScalarKind);
+		}
 	}
 
-	return Value;
-}
+	// Cast all given arguments to the common scalar kind
+	A = CastToScalarKind(A, CommonScalarKind);
+	B = CastToScalarKind(B, CommonScalarKind);
+	C = CastToScalarKind(C, CommonScalarKind);
+	D = CastToScalarKind(D, CommonScalarKind);
 
-FValueRef FEmitter::Vector4(FValueRef InXYZ, FValueRef InW)
-{
-	if (IsAnyNotValid(InXYZ, InW))
+	if (A.IsPoison() || B.IsPoison() || C.IsPoison() || D.IsPoison())
 	{
 		return Poison();
 	}
 
-	// TODO: replace these checks with casts instead
-	check(InXYZ->Type.IsVector() && InXYZ->Type.GetPrimitive().NumComponents() == 3);
-	check(InW->Type.IsPrimitive() && InW->Type.GetPrimitive().ScalarKind == InXYZ->Type.GetPrimitive().ScalarKind);
-
-	TComposite<4> Vector = MakePrototype<TComposite<4>>(FType::MakeVector(InXYZ->Type.GetPrimitive().ScalarKind, 4));
-	TArrayView<FValue*> Components = Vector.GetMutableComponents();
-	Components[0] = Subscript(InXYZ, 0);
-	Components[1] = Subscript(InXYZ, 1);
-	Components[2] = Subscript(InXYZ, 2);
-	Components[3] = InW;
-
-	FValueRef Value = EmitPrototype(*this, Vector);
-
-	// If all components source inputs match, refer to the same input in the vector value too (for debugging purposes).
-	if (InXYZ.Input == InW.Input)
+	int32 NumComponents = APrimType.NumComponents() + BPrimType.NumComponents();
+	if (CPrimType)
 	{
-		Value.Input = InXYZ.Input;
+		NumComponents += CPrimType->NumComponents();
+		if (DPrimType)
+		{
+			NumComponents += DPrimType->NumComponents();
+		}
 	}
-
-	return Value;
-}
-
-FValueRef FEmitter::Vector4(FValueRef InXY, FValueRef InZ, FValueRef InW)
-{
-	if (IsAnyNotValid(InXY, InZ, InW))
+	
+	if (NumComponents > 4)
 	{
+		Errorf(TEXT("Cannot create a vector of more than 4 components (%d provided)"), NumComponents);
 		return Poison();
 	}
 
-	// TODO: replace these checks with casts instead
-	check(InXY->Type.IsVector() && InXY->Type.GetPrimitive().NumComponents() == 2);
-	check(InZ->Type.IsPrimitive() && InZ->Type.GetPrimitive().ScalarKind == InXY->Type.GetPrimitive().ScalarKind);
-	check(InW->Type.IsPrimitive() && InW->Type.GetPrimitive().ScalarKind == InXY->Type.GetPrimitive().ScalarKind);
+	// Make the result vector
+	FType ResultType = FType::MakeVector(CommonScalarKind, NumComponents);
 
-	TComposite<4> Vector = MakePrototype<TComposite<4>>(FType::MakeVector(InXY->Type.GetPrimitive().ScalarKind, 4));
-	TArrayView<FValue*> Components = Vector.GetMutableComponents();
-	Components[0] = Subscript(InXY, 0);
-	Components[1] = Subscript(InXY, 1);
-	Components[2] = InZ;
-	Components[3] = InW;
+	// Allocate the result composite prototype
+	FComposite* Result = MakeCompositePrototype(*this, ResultType, NumComponents);
 
-	FValueRef Value = EmitPrototype(*this, Vector);
-
-	// If all components source inputs match, refer to the same input in the vector value too (for debugging purposes).
-	if (InXY.Input == InZ.Input && InXY.Input == InW.Input)
+	// Copy arguments components over to the result composite
+	int32 ComponentIndex = 0;
+	for (int32 i = 0; i < APrimType.NumComponents(); ++i)
 	{
-		Value.Input = InXY.Input;
+		Result->GetMutableComponents()[ComponentIndex++] = Subscript(A, i);
 	}
-
-	return Value;
-}
-
-FValueRef FEmitter::Vector4(FValueRef InX, FValueRef InY, FValueRef InZ, FValueRef InW)
-{
-	if (IsAnyNotValid(InX, InY, InZ, InW))
+	
+	for (int32 i = 0; i < BPrimType.NumComponents(); ++i)
 	{
-		return Poison();
+		Result->GetMutableComponents()[ComponentIndex++] = Subscript(B, i);
 	}
 
-	// TODO: replace these checks with casts instead
-	check(InX->Type.IsScalar());
-	check(InX->Type == InY->Type);
-	check(InY->Type == InZ->Type);
-	check(InZ->Type == InW->Type);
+	if (C.IsValid())
+	{
+		for (int32 i = 0; i < CPrimType->NumComponents(); ++i)
+		{
+			Result->GetMutableComponents()[ComponentIndex++] = Subscript(C, i);
+		}
+		if (D.IsValid())
+		{
+			for (int32 i = 0; i < DPrimType->NumComponents(); ++i)
+			{
+				Result->GetMutableComponents()[ComponentIndex++] = Subscript(D, i);
+			}
+		}
 
-	TComposite<4> Vector = MakePrototype<TComposite<4>>(FType::MakeVector(InX->Type.GetPrimitive().ScalarKind, 4));
-	TArrayView<FValue*> Components = Vector.GetMutableComponents();
-	Components[0] = InX;
-	Components[1] = InY;
-	Components[2] = InZ;
-	Components[3] = InW;
+	}
 
-	FValueRef Value = EmitPrototype(*this, Vector);
+	FValueRef ResultRef = EmitPrototype(*this, *Result);
 
-	// If all components source inputs match, refer to the same input in the vector value too (for debugging purposes).
-	if (InX.Input == InY.Input && InX.Input == InZ.Input && InX.Input == InW.Input)
+	// If all arguments come from the same expression input, propagate the input to the result value ref.
+	if (A.Input == B.Input && (!C || A.Input == C.Input) && (!D || A.Input == D.Input))
 	{
-		Value.Input = InX.Input;
+		ResultRef.Input = A.Input;
 	}
 
-	return Value;
+	return ResultRef;
 }
 
 FValueRef FEmitter::Aggregate(const UMaterialAggregate* InAggregate)
@@ -1572,7 +1643,7 @@ FValueRef FEmitter::Aggregate(const UMaterialAggregate* InAggregate)
 	return Aggregate(InAggregate, {}, TConstArrayView<FValueRef>{});
 }
 
-static FValueRef EmitAttributeDefaultValue(FEmitter& Emitter, const UMaterial* Material, const FMaterialAggregateAttribute& Attribute)
+static FValueRef EmitAttributeDefaultValue(FEmitter& Emitter, const UMaterialInterface* Material, const FMaterialAggregateAttribute& Attribute)
 {
 	switch (Attribute.Type)
 	{
@@ -1630,7 +1701,7 @@ FValueRef FEmitter::Aggregate(const UMaterialAggregate* InAggregate, FValueRef I
 		else
 		{
 			// ...otherwise construct the default value as indicated in the attribute.
-			Components[i] = EmitAttributeDefaultValue(*this, Material, InAggregate->Attributes[i]);
+			Components[i] = EmitAttributeDefaultValue(*this, MaterialInterface, InAggregate->Attributes[i]);
 		}
 	}
 
@@ -1817,6 +1888,7 @@ static T FoldScalarOperator(FEmitter& Emitter, EOperator Operator, T A, T B, T C
 			case BO_Or: return A | B;
 			case BO_BitwiseAnd: return A & B;
 			case BO_BitwiseOr: return A | B;
+			case BO_BitwiseXor: return A ^ B;
 			case BO_BitShiftLeft: return A << B;
 			case BO_BitShiftRight: return A >> B;
 			case BO_Modulo: return A % B;
@@ -1900,9 +1972,9 @@ static FValueRef TrySimplifyOperator(FEmitter& Emitter, EOperator& Op, FValueRef
 			{
 				return Emitter.Negate(B);
 			}
-			else if (A->Equals(B))
+			else if (A == B)
 			{
-				return Emitter.ConstantZero(A->Type.GetPrimitive().ScalarKind);
+				return Emitter.ConstantZero(ResultType.GetPrimitive().ScalarKind);
 			}
 			break;
 
@@ -1932,19 +2004,19 @@ static FValueRef TrySimplifyOperator(FEmitter& Emitter, EOperator& Op, FValueRef
 			{
 				return A;
 			}
-			else if (A->Equals(B))
+			else if (A == B)
 			{
-				return Emitter.ConstantOne(A->Type.GetPrimitive().ScalarKind);
+				return Emitter.ConstantOne(ResultType.GetPrimitive().ScalarKind);
 			}
 			break;
 
 		case BO_Modulo:
 			if (A->AreAllNearlyZero() || B->AreAllNearlyOne())
 			{
-				return Emitter.ConstantZero(A->Type.GetPrimitive().ScalarKind);
+				return Emitter.ConstantZero(ResultType.GetPrimitive().ScalarKind);
 			}
 			break;
-
+			
 		case BO_BitwiseAnd:
 			if (A->AreAllExactlyZero())
 			{
@@ -1954,6 +2026,10 @@ static FValueRef TrySimplifyOperator(FEmitter& Emitter, EOperator& Op, FValueRef
 			{
 				return B;
 			}
+			else if (A == B)
+			{
+				return A;
+			}
 			break;
 
 		case BO_BitwiseOr:
@@ -1965,6 +2041,25 @@ static FValueRef TrySimplifyOperator(FEmitter& Emitter, EOperator& Op, FValueRef
 			{
 				return A;
 			}
+			else if (A == B)
+			{
+				return A;
+			}
+			break;
+
+		case BO_BitwiseXor:
+			if (A->AreAllExactlyZero())
+			{
+				return B;
+			}
+			else if (B->AreAllExactlyZero())
+			{
+				return A;
+			}
+			else if (A == B)
+			{
+				return Emitter.ConstantZero(ResultType.GetPrimitive().ScalarKind);
+			}
 			break;
 
 		case BO_BitShiftLeft:
@@ -1997,7 +2092,7 @@ static FValueRef TrySimplifyOperator(FEmitter& Emitter, EOperator& Op, FValueRef
 			else if (B->AreAllNearlyZero())
 			{
 				// If the exponent is 0.
-				return Emitter.ConstantOne(A->Type.GetPrimitive().ScalarKind);
+				return Emitter.ConstantOne(ResultType.GetPrimitive().ScalarKind);
 			}
 			else if (B->AreAllExactlyOne())
 			{
@@ -2013,7 +2108,7 @@ static FValueRef TrySimplifyOperator(FEmitter& Emitter, EOperator& Op, FValueRef
 				B = {};
 				C = {};
 			}
-			else if (B->Equals(C))
+			else if (B == C)
 			{
 				return B;
 			}
@@ -2028,7 +2123,7 @@ static FValueRef TrySimplifyOperator(FEmitter& Emitter, EOperator& Op, FValueRef
 			{
 				return B;
 			}
-			else if (A->Equals(B))
+			else if (A == B)
 			{
 				return A;
 			}
@@ -2043,7 +2138,7 @@ static FValueRef TrySimplifyOperator(FEmitter& Emitter, EOperator& Op, FValueRef
 			{
 				return C;
 			}
-			else if (B->Equals(C))
+			else if (B == C)
 			{
 				return B;
 			}
@@ -2268,6 +2363,7 @@ static const FOperatorSignature* GetOperatorSignature(EOperator Op)
 		S[BO_Modulo] 					= BinaryArithmetic;
 		S[BO_BitwiseAnd]				= BinaryInteger;
 		S[BO_BitwiseOr]					= BinaryInteger;
+		S[BO_BitwiseXor]				= BinaryInteger;
 		S[BO_BitShiftLeft]				= BinaryInteger;
 		S[BO_BitShiftRight]				= BinaryInteger;
 
@@ -3380,7 +3476,9 @@ FValueRef FEmitter::StageSwitch(FType Type, TConstArrayView<FValueRef> ValuePerS
 	FStageSwitch Prototype = MakePrototype<FStageSwitch>(Type);
 	for (int i = 0; i < ValuePerStage.Num(); ++i)
 	{
-		Prototype.Args[i] = ValuePerStage[i];
+		Prototype.Args[i] = ValuePerStage[i]
+			? ValuePerStage[i]
+			: AnalysisError(Type, FString::Printf(TEXT("This operation is not supported in shader stage '%s'."), Module->GetEntryPoint(i).Name.GetData()));
 	}
 	return EmitPrototype(*this, Prototype);
 }
@@ -3906,6 +4004,7 @@ static FValue* DifferentiateOperator(FEmitter& E, FOperator* Op, EDerivativeAxis
 		case BO_Modulo:
 		case BO_BitwiseAnd:
 		case BO_BitwiseOr:
+		case BO_BitwiseXor:
 		case BO_BitShiftLeft:
 		case BO_BitShiftRight:
 		case BO_Step:
@@ -4036,13 +4135,17 @@ FValueRef FEmitter::AnalyticalPartialDerivative(FValueRef Value, EDerivativeAxis
 	}
 }
 
-static FValueRef EmitInlineHLSL(FEmitter& Emitter, FType Type, const FMaterialExternalCodeDeclaration* InExternalCodeDeclaration, FStringView Code, TConstArrayView<FValueRef> InputValues, EValueFlags ValueFlags, EGraphProperties UsedGraphProperties)
+static FValueRef EmitInlineHLSL(FEmitter& Emitter, FType Type, const FMaterialExternalCodeDeclaration* InExternalCodeDeclaration,
+								FStringView Code, TConstArrayView<FValueRef> InputValues, EValueFlags ValueFlags, EGraphProperties UsedGraphProperties,
+								FInlineHLSL::TOnAnalyzeCallback OnAnalyze = {}, uintptr_t OnAnalyzeUserData = {})
 {
 	FInlineHLSL Prototype = MakePrototype<FInlineHLSL>(Type);
 	Prototype.Type = Type;
 	Prototype.Flags = ValueFlags;
 	Prototype.GraphProperties = UsedGraphProperties;
-	
+	Prototype.OnAnalyze = OnAnalyze;
+	Prototype.OnAnalyzeUserData = OnAnalyzeUserData;
+
 	if (InExternalCodeDeclaration)
 	{
 		check(Code.IsEmpty());
@@ -4067,14 +4170,26 @@ static FValueRef EmitInlineHLSL(FEmitter& Emitter, FType Type, const FMaterialEx
 	return EmitPrototype(Emitter, Prototype);
 }
 
-FValueRef FEmitter::InlineHLSL(FType Type, FStringView Code, TConstArrayView<FValueRef> InputValues, EValueFlags ValueFlags, EGraphProperties UsedGraphProperties)
+FValueRef FEmitter::InlineHLSL(FType Type, FStringView Code, TConstArrayView<FValueRef> InputValues, EValueFlags ValueFlags, EGraphProperties UsedGraphProperties, FInlineHLSL::TOnAnalyzeCallback OnAnalyze, uintptr_t OnAnalyzeUserData)
 {
 	if (IsAnyNotValid(InputValues))
 	{
 		return Poison();
 	}
 	
-	return EmitInlineHLSL(*this, Type, nullptr, Module->InternString(Code), InputValues, ValueFlags | EValueFlags::HasDynamicHLSLCode, UsedGraphProperties);
+	return EmitInlineHLSL(*this, Type, nullptr, Module->InternString(Code), InputValues, ValueFlags | EValueFlags::HasDynamicHLSLCode, UsedGraphProperties, OnAnalyze, OnAnalyzeUserData);
+}
+
+FValueRef FEmitter::InlineHLSL(FName InExternalCodeDeclarationName, TConstArrayView<FValueRef> InputValues, EValueFlags ValueFlags, EGraphProperties UsedGraphProperties)
+{
+	const FMaterialExternalCodeDeclaration* ExternalCodeDeclaration = MaterialExternalCodeRegistry::Get().FindExternalCode(InExternalCodeDeclarationName);
+	if (!ExternalCodeDeclaration)
+	{
+		Errorf(TEXT("Missing external code declaration for '%s'"), *InExternalCodeDeclarationName.ToString());
+		return Poison();
+	}
+
+	return InlineHLSL(ExternalCodeDeclaration, InputValues, ValueFlags, UsedGraphProperties);
 }
 
 FValueRef FEmitter::InlineHLSL(const FMaterialExternalCodeDeclaration* InExternalCodeDeclaration, TConstArrayView<FValueRef> InputValues, EValueFlags ValueFlags, EGraphProperties UsedGraphProperties)
diff --git a/Engine/Source/Runtime/Engine/Private/Materials/MaterialIRModule.cpp b/Engine/Source/Runtime/Engine/Private/Materials/MaterialIRModule.cpp
--- a/Engine/Source/Runtime/Engine/Private/Materials/MaterialIRModule.cpp
+++ b/Engine/Source/Runtime/Engine/Private/Materials/MaterialIRModule.cpp
@@ -80,9 +80,9 @@ FStringView FMaterialIRModule::InternString(FStringView InString)
 	return { String, InString.Len() };
 }
 
-void FMaterialIRModule::AddError(UMaterialExpression* Expression, FString Message)
+void FMaterialIRModule::AddError(UMaterialExpression* Source, FString Message)
 {
-	Errors.Push({ Expression, MoveTemp(Message) });
+	Errors.Push({ Source, MoveTemp(Message) });
 }
 
 int32 FMaterialIRModule::FindOrAddParameterCollection(UMaterialParameterCollection* ParameterCollection)
diff --git a/Engine/Source/Runtime/Engine/Private/Materials/MaterialIRModuleBuilder.cpp b/Engine/Source/Runtime/Engine/Private/Materials/MaterialIRModuleBuilder.cpp
--- a/Engine/Source/Runtime/Engine/Private/Materials/MaterialIRModuleBuilder.cpp
+++ b/Engine/Source/Runtime/Engine/Private/Materials/MaterialIRModuleBuilder.cpp
@@ -93,7 +93,7 @@ struct FMaterialIRModuleBuilderImpl
 		this->Module = InModule;
 
 		// Setup the Builder implementation
-		ValueAnalyzer.Setup(Builder->Material, Module, &Module->CompilationOutput,  Builder->TargetInsights);
+		ValueAnalyzer.Setup(Builder->MaterialInterface, Module, &Module->CompilationOutput,  Builder->TargetInsights);
 
 		// Empty the module and set it up
 		Module->Empty();
@@ -113,9 +113,10 @@ struct FMaterialIRModuleBuilderImpl
 
 		// Setup the emitter and initialize it
 		Emitter.BuilderImpl = this;
-		Emitter.Material = Builder->Material;
+		Emitter.MaterialInterface = Builder->MaterialInterface;
+		Emitter.BaseMaterial = Builder->MaterialInterface->GetBaseMaterial();
 		Emitter.Module = Module;
-		Emitter.StaticParameterSet = &Builder->StaticParameters;
+		Emitter.StaticParameterSet = Builder->StaticParameters;
 		Emitter.SubstrateTranslatorData = &SubstrateTranslatorData;
 		Emitter.Initialize();
 
@@ -134,18 +135,18 @@ struct FMaterialIRModuleBuilderImpl
 		{
 			SubstrateTranslatorData.SubstrateCompilationConfig = *SubstrateCompilationConfig;
 		}
-		SubstrateTranslatorData.Initialize(Builder->Material->GetMaterialResource(Module->ShaderPlatform), Module->CompilationOutput, nullptr, &Emitter);
+		SubstrateTranslatorData.Initialize(Builder->MaterialInterface->GetMaterialResource(Module->ShaderPlatform), Module->CompilationOutput, nullptr, &Emitter);
 
 		FMaterialShadingModelField MaterialShadingModels;
 		// If the material gets its shading model from material expressions and we have compiled one or more shading model expressions already, 
 		// then use that shading model field instead. It's the most optimal set of shading models
-		if (Builder->Material->IsShadingModelFromMaterialExpression())// && ShadingModelsFromCompilation.IsValid())
+		if (Builder->MaterialInterface->IsShadingModelFromMaterialExpression())// && ShadingModelsFromCompilation.IsValid())
 		{
 			check(false); // SUBSTRATE_TODO add support for that case.
 		}
 		else
 		{
-			MaterialShadingModels = Builder->Material->GetShadingModels();
+			MaterialShadingModels = Builder->MaterialInterface->GetShadingModels();
 			UMaterialInterface::FilterOutPlatformShadingModels(Module->ShaderPlatform, MaterialShadingModels);
 		}
 
@@ -158,7 +159,7 @@ struct FMaterialIRModuleBuilderImpl
 	void Step_PushRootExpressionDependencies()
 	{
 		//If we are processing from the Material Attributes output node, push the dependency here.
-		if (Builder->Material->bUseMaterialAttributes)
+		if (Builder->MaterialInterface->GetMaterial()->bUseMaterialAttributes)
 		{
 			FMaterialInputDescription Input;
 			GetExpressionInputDescription(MP_MaterialAttributes, Input);
@@ -179,7 +180,7 @@ struct FMaterialIRModuleBuilderImpl
 
 		// Push the custom output expressions to the dependency queue too.
 		TArray<UMaterialExpressionCustomOutput*> CustomOutputs;
-		Builder->Material->GetAllCustomOutputExpressions(CustomOutputs);
+		Builder->MaterialInterface->GetMaterial()->GetAllCustomOutputExpressions(CustomOutputs);
 		for (UMaterialExpressionCustomOutput* Expression : CustomOutputs)
 		{
 			PushDependency(Expression);
@@ -321,7 +322,7 @@ struct FMaterialIRModuleBuilderImpl
 			PushDependency(It.Input->Expression);
 		}
 
-		// Named reroute usage nodes should make add their declaration as dependency, as they
+		// Named reroute usage nodes should add their declaration as dependency, as they
 		// do not have any direct input (they simply forward to the declaration input).
 		if (auto NamedReroute = Cast<UMaterialExpressionNamedRerouteUsage>(Emitter.Expression))
 		{
@@ -644,7 +645,7 @@ struct FMaterialIRModuleBuilderImpl
 		// First, if the material is flagged to use the material attributes aggregate, read its value now
 		// so that we can extract its individual attributes later.
 		MIR::FValue* MaterialAttributesValue = nullptr;
-		if (Builder->Material->bUseMaterialAttributes)
+		if (Builder->MaterialInterface->GetMaterial()->bUseMaterialAttributes)
 		{
 			FMaterialInputDescription InputDesc;
 			GetExpressionInputDescription(MP_MaterialAttributes, InputDesc);
@@ -741,8 +742,8 @@ struct FMaterialIRModuleBuilderImpl
 
 			// Add support for lerp to selection color for PC development builds.
 			if (Property == MP_EmissiveColor &&
-				Builder->Material->MaterialDomain != MD_Volume &&
-				UE::MaterialTranslatorUtils::IsDevelopmentFeatureEnabled(NAME_SelectionColor, Module->GetShaderPlatform(), Builder->Material))
+				Builder->MaterialInterface->GetMaterial()->MaterialDomain != MD_Volume &&
+				UE::MaterialTranslatorUtils::IsDevelopmentFeatureEnabled(NAME_SelectionColor, Module->GetShaderPlatform(), Builder->MaterialInterface->GetMaterial()))
 			{
 				FMaterialParameterMetadata ParameterMetadata{ FMaterialParameterValue(FLinearColor::Transparent) };
 				MIR::FValue* SelectionColor = Emitter.Parameter(NAME_SelectionColor, ParameterMetadata);
@@ -756,7 +757,7 @@ struct FMaterialIRModuleBuilderImpl
 			Emitter.Private_SetMaterialOutput(Property, PropertyValue);
 
 			// Finally, push this connection insight
-			PushConnectionInsight(Builder->Material, (int32)Property, InputDesc.Input->Expression, InputDesc.Input->OutputIndex, PropertyValue->Type);
+			PushConnectionInsight(Builder->MaterialInterface->GetMaterial(), (int32)Property, InputDesc.Input->Expression, InputDesc.Input->OutputIndex, PropertyValue->Type);
 		}
 
 		return Module->IsValid();
@@ -793,7 +794,7 @@ struct FMaterialIRModuleBuilderImpl
 				const bool bHasSSS = SSSProfile != nullptr;
 				const bool bHasAnisotropy = EditorOnlyData->Anisotropy.IsConnected();
 
-				FrontMaterialPropertyValue = UMaterialExpressionSubstrateShadingModels::BuildCommon(Emitter, Builder->Material,
+				FrontMaterialPropertyValue = UMaterialExpressionSubstrateShadingModels::BuildCommon(Emitter, Emitter.BaseMaterial,
 					EditorOnlyData->BaseColor, EditorOnlyData->Specular, EditorOnlyData->Metallic, EditorOnlyData->Roughness, EditorOnlyData->EmissiveColor,
 					EditorOnlyData->Opacity, EditorOnlyData->SubsurfaceColor, EditorOnlyData->ClearCoat, EditorOnlyData->ClearCoatRoughness,
 					EditorOnlyData->ShadingModelFromMaterialExpression, MaterialInterface->GetShadingModels().GetFirstShadingModel(),
@@ -912,7 +913,7 @@ struct FMaterialIRModuleBuilderImpl
 		}
 		else
 		{
-			bool bResult = Builder->Material->GetExpressionInputDescription(Property, Input);
+			bool bResult = Builder->MaterialInterface->GetMaterial()->GetExpressionInputDescription(Property, Input);
 			/**
 			* MP_SubsurfaceColor is currently hacked in the old translator to float4, 
 			* but we rely on default types (i.e.float3) for default values in the material editor output.
@@ -1056,9 +1057,9 @@ struct FMaterialIRModuleBuilderImpl
 	bool Step_PerformModuleWideAnalysis()
 	{
 		// If the material doesn't use expression shading models, or they aren't valid, initialize the shading models to the ones from the material.
-		if (!ValueAnalyzer.Material->IsShadingModelFromMaterialExpression() || !Module->ShadingModelsFromCompilation.IsValid())
+		if (!Builder->MaterialInterface->IsShadingModelFromMaterialExpression() || !Module->ShadingModelsFromCompilation.IsValid())
 		{
-			Module->ShadingModelsFromCompilation = ValueAnalyzer.Material->GetShadingModels();
+			Module->ShadingModelsFromCompilation = Builder->MaterialInterface->GetShadingModels();
 			UMaterialInterface::FilterOutPlatformShadingModels(Module->GetShaderPlatform(), Module->ShadingModelsFromCompilation);
 		}
 
@@ -1096,7 +1097,7 @@ struct FMaterialIRModuleBuilderImpl
 		// todo: We want the usage on the MaterialInstance if that exists. 
 		// See FMaterialResource::IsUsedWithInstancedStaticMeshes() which is used at the same point in the old translator.
 		// This is one case of several where the old translator works by having a FMaterial instead of just the UMaterial.
-		const bool bUsedWithInstancedStaticMeshes = ValueAnalyzer.Material->GetUsageByFlag(MATUSAGE_InstancedStaticMeshes);
+		const bool bUsedWithInstancedStaticMeshes = Builder->MaterialInterface->GetUsageByFlag(MATUSAGE_InstancedStaticMeshes);
 
 		KeepDefineConditionally(TEXT("USES_PER_INSTANCE_CUSTOM_DATA"), bUsedWithInstancedStaticMeshes);
 		KeepDefineConditionally(TEXT("NEEDS_PER_INSTANCE_RANDOM_PS"), bUsedWithInstancedStaticMeshes);
@@ -1112,7 +1113,7 @@ struct FMaterialIRModuleBuilderImpl
 		if (Module->EnvironmentDefines.Contains(TEXT("MIR.SceneDepth")))
 		{
 			// @todo-jason.hoerner - Support for material instance blend mode overrides needed.
-			if (ValueAnalyzer.Material->MaterialDomain != MD_PostProcess && !IsTranslucentBlendMode(ValueAnalyzer.Material->BlendMode))
+			if (Builder->MaterialInterface->GetMaterial()->MaterialDomain != MD_PostProcess && !IsTranslucentBlendMode(Builder->MaterialInterface->GetMaterial()->BlendMode))
 			{
 				Module->AddError(nullptr, TEXT("Only transparent or postprocess materials can read from scene depth."));
 			}
@@ -1152,17 +1153,28 @@ struct FMaterialIRModuleBuilderImpl
 		CompilationOutput.bUsesWorldPositionOffset 	= MaterialPropertyHasNonZeroValue(MP_WorldPositionOffset);
 		CompilationOutput.bUsesPixelDepthOffset 	= MaterialPropertyHasNonZeroValue(MP_PixelDepthOffset);
 		CompilationOutput.bUsesAnisotropy 			= MaterialPropertyHasNonZeroValue(MP_Anisotropy);
+		
+		// Check if displacement is actually used (disable if it was left at the default, invalid constant value of -1)
+		CompilationOutput.bUsesDisplacement 		= DoesPlatformSupportNanite(Module->GetShaderPlatform()) && Builder->MaterialInterface->IsTessellationEnabled() && MaterialPropertyHasNonZeroValue(MP_Displacement);
 
 		// bModifiesMeshPosition
-		// @todo-jason.hoerner - Fill in CompilationOutput.bUsesDisplacement.
-		// @todo-jason.hoerner - Custom output support.
-		const bool bHasFirstPersonOutput = false;  // (CustomOutputsBitmask & FMaterialAttributeDefinitionMap::GetBitmask(FirstPersonInterpolationAlphaGuid)) != 0;
+		// @massimo.tristano this kind of logic should be moved to the first person custom output expression Build() function
+		bool bHasFirstPersonOutput = false;
+		for (const MIR::FSetMaterialOutput* Output : Module->GetCustomOutputs())
+		{
+			if (Output->Name == TEXTVIEW("FirstPersonOutput"))
+			{
+				bHasFirstPersonOutput = true;
+				break;
+			}
+		}
+		
 		CompilationOutput.bModifiesMeshPosition = CompilationOutput.bUsesPixelDepthOffset || CompilationOutput.bUsesWorldPositionOffset || CompilationOutput.bUsesDisplacement || bHasFirstPersonOutput;
 		
 		// Run a final material compilation output validation and report any error reported.
 		TArray<FString> ValidationErrors;
 		UE::MaterialTranslatorUtils::FinalCompileValidation(
-			ValueAnalyzer.Material,
+			Builder->MaterialInterface->GetMaterial(),
 			CompilationOutput,
 			Module->GetCompiledShadingModels(),
 			Module->GetBlendMode(),
@@ -1176,41 +1188,6 @@ struct FMaterialIRModuleBuilderImpl
 		}
 	}
 
-	void Step_FinalizeArtifacts()
-	{
-		// Finalize the material insights debug information
-		GenerateDebugInsights();
-	}
-
-	void GenerateDebugInsights()
-	{
-		if (!Builder->Material->MaterialInsight)
-		{
-			return;
-		}
-
-		// Dump the module IR to string and store it inside the material insights.
-		Builder->Material->MaterialInsight.Get()->IRString = MIR::DebugDumpIR(Builder->Material->GetFullName(), *Module);
-		
-		// Dump the requested debugging information
-		switch (CVarMaterialIRDebugDumpLevel.GetValueOnGameThread())
-		{
-			case 2:
-			{
-				MIR::DebugDumpIRUseGraph(*Module);
-				// fallthrough
-			}
-
-			case 1:
-			{
-				// Save the dump to file
-				FString FilePath = FPaths::Combine(FPaths::ProjectSavedDir(), "Materials", TEXT("IRDump.txt"));
-				FFileHelper::SaveStringToFile(Builder->Material->MaterialInsight.Get()->IRString, *FilePath);
-				// fallthrough
-			}
-		}
-	}
-
 	void Step_LinkInstructions()
 	{
 		TArray<MIR::FInstruction*> InstructionStack{};
@@ -1289,6 +1266,41 @@ struct FMaterialIRModuleBuilderImpl
 			}
 		}
 	}
+
+	void Step_FinalizeArtifacts()
+	{
+		// Finalize the material insights debug information
+		GenerateDebugInsights();
+	}
+
+	void GenerateDebugInsights()
+	{
+		if (!Builder->TargetInsights)
+		{
+			return;
+		}
+
+		// Dump the module IR to string and store it inside the material insights.
+		Builder->TargetInsights->IRString = MIR::DebugDumpIR(Builder->MaterialInterface->GetFullName(), *Module);
+		
+		// Dump the requested debugging information
+		switch (CVarMaterialIRDebugDumpLevel.GetValueOnGameThread())
+		{
+			case 2:
+			{
+				MIR::DebugDumpIRUseGraph(*Module);
+				// fallthrough
+			}
+
+			case 1:
+			{
+				// Save the dump to file
+				FString FilePath = FPaths::Combine(FPaths::ProjectSavedDir(), "Materials", TEXT("IRDump.txt"));
+				FFileHelper::SaveStringToFile(Builder->TargetInsights->IRString, *FilePath);
+				// fallthrough
+			}
+		}
+	}
 	
 	
 	/* Auxiliary functions */
diff --git a/Engine/Source/Runtime/Engine/Private/Materials/MaterialIRToHLSLTranslator.cpp b/Engine/Source/Runtime/Engine/Private/Materials/MaterialIRToHLSLTranslator.cpp
--- a/Engine/Source/Runtime/Engine/Private/Materials/MaterialIRToHLSLTranslator.cpp
+++ b/Engine/Source/Runtime/Engine/Private/Materials/MaterialIRToHLSLTranslator.cpp
@@ -1454,6 +1454,9 @@ struct FPrivate : FMaterialIRToHLSLTranslation
 				case MIR::BO_Subtract:
 				case MIR::BO_Divide:
 				case MIR::BO_Modulo:
+				case MIR::BO_BitwiseAnd:
+				case MIR::BO_BitwiseOr:
+				case MIR::BO_BitwiseXor:
 					bOperatorIsInfix = true;
 					break;
 
@@ -1477,6 +1480,9 @@ struct FPrivate : FMaterialIRToHLSLTranslation
 				case MIR::BO_Multiply: OpString = TEXTVIEW("*"); break;
 				case MIR::BO_NotEquals: OpString = TEXTVIEW("!="); break;
 				case MIR::BO_Subtract: OpString = TEXTVIEW("-"); break;
+				case MIR::BO_BitwiseAnd: OpString = TEXTVIEW("&"); break;
+				case MIR::BO_BitwiseOr: OpString = TEXTVIEW("|"); break;
+				case MIR::BO_BitwiseXor: OpString = TEXTVIEW("^"); break;
 				default: UE_MIR_UNREACHABLE();
 			}
 
@@ -1550,6 +1556,7 @@ struct FPrivate : FMaterialIRToHLSLTranslation
 					case MIR::UO_ATan: OpString = TEXTVIEW("atan"); break;
 					case MIR::UO_ATanFast: OpString = TEXTVIEW("atanFast"); break;
 					case MIR::UO_ATanh: OpString = TEXTVIEW("atanh"); break;
+					case MIR::UO_BitwiseNot: OpString = TEXTVIEW("~"); break;
 					case MIR::UO_Ceil: OpString = TEXTVIEW("ceil"); break;
 					case MIR::UO_Cos: OpString = TEXTVIEW("cos"); break;
 					case MIR::UO_Cosh: OpString = TEXTVIEW("cosh"); break;
@@ -1565,6 +1572,7 @@ struct FPrivate : FMaterialIRToHLSLTranslation
 					case MIR::UO_Logarithm10: OpString = TEXTVIEW("log10"); break;
 					case MIR::UO_Logarithm2: OpString = TEXTVIEW("log2"); break;
 					case MIR::UO_Negate: OpString = TEXTVIEW("-"); break;
+					case MIR::UO_Not: OpString = TEXTVIEW("!"); break;
 					case MIR::UO_Reciprocal: OpString = TEXTVIEW("rcp"); break;
 					case MIR::UO_Round: OpString = TEXTVIEW("round"); break;
 					case MIR::UO_Rsqrt: OpString = TEXTVIEW("rsqrt"); break;
diff --git a/Engine/Source/Runtime/Engine/Private/Materials/MaterialIRTypes.cpp b/Engine/Source/Runtime/Engine/Private/Materials/MaterialIRTypes.cpp
--- a/Engine/Source/Runtime/Engine/Private/Materials/MaterialIRTypes.cpp
+++ b/Engine/Source/Runtime/Engine/Private/Materials/MaterialIRTypes.cpp
@@ -281,47 +281,46 @@ UE::Shader::EValueType FType::ToValueType() const
 {
 	using namespace UE::Shader;
 
-	if (AsPrimitive())
+	switch (Kind)
 	{
-		if (Primitive.IsMatrix())
-		{
-			if (Primitive.NumRows == 4 && Primitive.NumColumns == 4)
+		case ETypeKind::Primitive:
+			if (Primitive.IsMatrix())
 			{
-				if (Primitive.ScalarKind == EScalarKind::Float)
+				if (Primitive.NumRows == 4 && Primitive.NumColumns == 4)
 				{
-					return EValueType::Float4x4;
-				}
-				else if (Primitive.ScalarKind == EScalarKind::Double)
-				{
-					return Primitive.bIsLWCInverseMatrix ? EValueType::DoubleInverse4x4 : EValueType::Double4x4;
-				}
-				else
-				{
-					return EValueType::Numeric4x4;
+					if (Primitive.ScalarKind == EScalarKind::Float)
+					{
+						return EValueType::Float4x4;
+					}
+					else if (Primitive.ScalarKind == EScalarKind::Double)
+					{
+						return Primitive.bIsLWCInverseMatrix ? EValueType::DoubleInverse4x4 : EValueType::Double4x4;
+					}
+					else
+					{
+						return EValueType::Numeric4x4;
+					}
 				}
+
+				return EValueType::Any;
 			}
 
-			return EValueType::Any;
-		}
+			check(Primitive.NumRows < 4 && Primitive.NumColumns <= 4);
+
+			switch (Primitive.ScalarKind)
+			{
+				case EScalarKind::Bool: return (EValueType)((int)EValueType::Bool1 + Primitive.NumColumns - 1);
+				case EScalarKind::Int: 	return (EValueType)((int)EValueType::Int1 + Primitive.NumColumns - 1);
+				case EScalarKind::Float: return (EValueType)((int)EValueType::Float1 + Primitive.NumColumns - 1);
+				case EScalarKind::Double:	return (EValueType)((int)EValueType::Double1 + Primitive.NumColumns - 1);
+				default: UE_MIR_UNREACHABLE();
+			}
 
-		check(Primitive.NumRows < 4 && Primitive.NumColumns <= 4);
+		case ETypeKind::Aggregate:
+			return EValueType::Struct;
 
-		switch (Primitive.ScalarKind)
-		{
-			case EScalarKind::Bool: return (EValueType)((int)EValueType::Bool1 + Primitive.NumColumns - 1);
-			case EScalarKind::Int: 	return (EValueType)((int)EValueType::Int1 + Primitive.NumColumns - 1);
-			case EScalarKind::Float: return (EValueType)((int)EValueType::Float1 + Primitive.NumColumns - 1);
-			case EScalarKind::Double:	return (EValueType)((int)EValueType::Double1 + Primitive.NumColumns - 1);
-			default: UE_MIR_UNREACHABLE();
-		}
-	}
-	else if (Is(ETypeKind::Aggregate))
-	{
-		return EValueType::Struct;
-	}
-	else
-	{
-		return EValueType::Object;
+		default:
+			return EValueType::Object;
 	}
 }
 
diff --git a/Engine/Source/Runtime/Engine/Private/Materials/MaterialIRValueAnalyzer.cpp b/Engine/Source/Runtime/Engine/Private/Materials/MaterialIRValueAnalyzer.cpp
--- a/Engine/Source/Runtime/Engine/Private/Materials/MaterialIRValueAnalyzer.cpp
+++ b/Engine/Source/Runtime/Engine/Private/Materials/MaterialIRValueAnalyzer.cpp
@@ -26,6 +26,11 @@
 #include "PostProcess/PostProcessMaterialInputs.h"
 #include "Shader/Preshader2.h"
 
+static void AnalyzeAnalysisError(FMaterialIRValueAnalyzer& Analyzer, MIR::FAnalysisError* AnalysisError)
+{
+	Analyzer.Module->AddError(AnalysisError->Expression, AnalysisError->Message.ToString());
+}
+
 static void ErrorUnlessFeatureLevelSupported(FMaterialIRModule* Module, ERHIFeatureLevel::Type RequiredFeatureLevel, const TCHAR* Message)
 {
 	ERHIFeatureLevel::Type FeatureLevel = Module->GetFeatureLevel();
@@ -49,13 +54,11 @@ static void AnalyzeExternalInput(FMaterialIRValueAnalyzer& Analyzer, MIR::FExter
 		Statistics.NumVertexTexCoords = FMath::Max(Statistics.NumVertexTexCoords, TexCoordIndex + 1);
 		Statistics.NumPixelTexCoords = FMath::Max(Statistics.NumPixelTexCoords, TexCoordIndex + 1);
 	}
-
-	if (Analyzer.Material->IsPostProcessMaterial() && MIR::IsExternalInputWorldPosition(ExternalInput->Id))
+	else if (Analyzer.MaterialInterface->IsPostProcessMaterial() && MIR::IsExternalInputWorldPosition(ExternalInput->Id))
 	{
 		Analyzer.CompilationOutput->SetIsSceneTextureUsed(PPI_SceneDepth);
 	}
-
-	if (ExternalInput->Id == MIR::EExternalInput::GlobalDistanceField)
+	else if (ExternalInput->Id == MIR::EExternalInput::GlobalDistanceField)
 	{
 		if (!FDataDrivenShaderPlatformInfo::GetSupportsDistanceFields(Analyzer.Module->GetShaderPlatform()))
 		{
@@ -65,11 +68,14 @@ static void AnalyzeExternalInput(FMaterialIRValueAnalyzer& Analyzer, MIR::FExter
 
 		Analyzer.CompilationOutput->bUsesGlobalDistanceField = true;
 	}
-
-	if (ExternalInput->Id == MIR::EExternalInput::DynamicParticleParameterIndex)
+	else if (ExternalInput->Id == MIR::EExternalInput::DynamicParticleParameterIndex)
 	{
 		Analyzer.Module->GetStatistics().DynamicParticleParameterMask |= 1 << ExternalInput->UserData;
 	}
+	else if (ExternalInput->Id == MIR::EExternalInput::DistanceCullFade)
+	{
+		Analyzer.CompilationOutput->bUsesDistanceCullFade = true;
+	}
 }
 
 static void AnalyzeInStageHardwarePartialDerivative(FMaterialIRValueAnalyzer& Analyzer, MIR::FHardwarePartialDerivative* PartialDerivative, MIR::EStage Stage)
@@ -122,9 +128,9 @@ static void AnalyzeInlineHLSL(FMaterialIRValueAnalyzer& Analyzer, MIR::FInlineHL
 
 		// Validate this external code can be used for the current material domain. Empty list implies no restriction on material domains.
 		if (!InlineHLSL->ExternalCodeDeclaration->Domains.IsEmpty() &&
-			InlineHLSL->ExternalCodeDeclaration->Domains.Find(Analyzer.Material->MaterialDomain) == INDEX_NONE)
+			InlineHLSL->ExternalCodeDeclaration->Domains.Find(Analyzer.MaterialInterface->GetMaterial()->MaterialDomain) == INDEX_NONE)
 		{
-			FName AssetPathName = Analyzer.Material->GetOutermost()->GetFName();
+			FName AssetPathName = Analyzer.MaterialInterface->GetOutermost()->GetFName();
 			Analyzer.Module->AddError(nullptr, MaterialExpressionUtils::FormatUnsupportedMaterialDomainError(*InlineHLSL->ExternalCodeDeclaration, AssetPathName));
 		}
 
@@ -142,6 +148,11 @@ static void AnalyzeInlineHLSL(FMaterialIRValueAnalyzer& Analyzer, MIR::FInlineHL
 				*InlineHLSL->ExternalCodeDeclaration->Name.ToString(), *MinimumFeatureLevelName, *FeatureLevelName));
 		}
 	}
+
+	if (InlineHLSL->OnAnalyze)
+	{
+		InlineHLSL->OnAnalyze(Analyzer.Module, InlineHLSL->OnAnalyzeUserData);
+	}
 }
 
 static int32 AcquireVTStackIndex(FMaterialIRValueAnalyzer& Analyzer, MIR::FVTPageTableRead& VTPageTableRead, bool bGenerateFeedback)
@@ -181,7 +192,7 @@ static int32 RegisterTextureParameter(FMaterialIRValueAnalyzer& Analyzer, UObjec
 	const FMaterialTextureParameterInfo TextureParameterInfo
 	{
 		.ParameterInfo = ParameterInfo,
-		.TextureIndex = Analyzer.Material->GetReferencedTextures().Find(Texture),
+		.TextureIndex = Analyzer.MaterialInterface->GetReferencedTextures().Find(Texture),
 		.SamplerSource = SSM_FromTextureAsset,
 		.VirtualTextureLayerIndex = static_cast<uint8>(VTLayerIndex == INDEX_NONE ? UINT8_MAX : VTLayerIndex),
 	};
@@ -302,7 +313,8 @@ static void AnalyzeInStageVTPageTableRead(FMaterialIRValueAnalyzer& Analyzer, MI
 
 static void AnalyzeScreenTexture(FMaterialIRValueAnalyzer& Analyzer, MIR::FScreenTexture* ScreenTexture)
 {
-	const EMaterialDomain MaterialDomain = Analyzer.Material->MaterialDomain;
+	const UMaterial* Material = Analyzer.MaterialInterface->GetMaterial();
+	const EMaterialDomain MaterialDomain = Material->MaterialDomain;
 
 	switch (ScreenTexture->TextureKind)
 	{
@@ -350,7 +362,7 @@ static void AnalyzeScreenTexture(FMaterialIRValueAnalyzer& Analyzer, MIR::FScree
 
 				// Error checking from FHLSLMaterialTranslator::SceneTextureLookup
 				// Guard against using unsupported textures with SLW
-				const bool bHasSingleLayerWaterSM = Analyzer.Material->GetShadingModels().HasShadingModel(MSM_SingleLayerWater);
+				const bool bHasSingleLayerWaterSM = Analyzer.MaterialInterface->GetShadingModels().HasShadingModel(MSM_SingleLayerWater);
 				if (bHasSingleLayerWaterSM && ScreenTexture->Id != PPI_CustomDepth && ScreenTexture->Id != PPI_CustomStencil)
 				{
 					Analyzer.Module->AddError(nullptr, TEXT("Only custom depth and custom stencil can be sampled with SceneTexture when used with the Single Layer Water shading model."));
@@ -379,10 +391,11 @@ static void AnalyzeScreenTexture(FMaterialIRValueAnalyzer& Analyzer, MIR::FScree
 
 					if (ScreenTexture->Id == PPI_WorldNormal && !IsUsingDBuffers(Platform))
 					{
+						// @massimo.tristano this should not be here, as this is not the correct way to check if an attribute has a user-defined value in the new translator
 						bool bHasNormalConnected =
 							Substrate::IsSubstrateEnabled() ?
-							Analyzer.Material->IsPropertyConnected(MP_Normal) :
-							Analyzer.Material->HasNormalConnected();
+							Material->IsPropertyConnected(MP_Normal) :
+							Material->HasNormalConnected();
 
 						if (bHasNormalConnected)
 						{
@@ -430,7 +443,7 @@ static void AnalyzeScreenTexture(FMaterialIRValueAnalyzer& Analyzer, MIR::FScree
 							Analyzer.Module->AddError(nullptr, FString::Printf(TEXT("GBuffer scene textures not available in post-processing with mobile shading (platform id %d)."), Platform));
 						}
 
-						if (Analyzer.Material->IsMobileSeparateTranslucencyEnabled())
+						if (Material->IsMobileSeparateTranslucencyEnabled())
 						{
 							Analyzer.Module->AddError(nullptr, FString::Printf(TEXT("GBuffer scene textures not available for separate translucency with mobile shading (platform id %d)."), Platform));
 						}
@@ -448,7 +461,7 @@ static void AnalyzeScreenTexture(FMaterialIRValueAnalyzer& Analyzer, MIR::FScree
 				// Error checking from FHLSLMaterialTranslator::TranslateMaterial
 				if (MaterialDomain != MD_DeferredDecal && MaterialDomain != MD_PostProcess)
 				{
-					if (!Analyzer.Material->GetShadingModels().HasShadingModel(MSM_SingleLayerWater) && IsOpaqueOrMaskedBlendMode(Analyzer.Material->BlendMode))
+					if (!Analyzer.MaterialInterface->GetShadingModels().HasShadingModel(MSM_SingleLayerWater) && IsOpaqueOrMaskedBlendMode(Material->BlendMode))
 					{
 						// In opaque pass, none of the textures are available
 						Analyzer.Module->AddError(nullptr, TEXT("SceneTexture expressions cannot be used in opaque materials except if used with the Single Layer Water shading model."));
@@ -462,7 +475,7 @@ static void AnalyzeScreenTexture(FMaterialIRValueAnalyzer& Analyzer, MIR::FScree
 				if (ScreenTexture->Id == PPI_SceneDepth && bTextureLookup)
 				{
 					// Don't allow opaque and masked materials to access scene depth as the results are undefined
-					if (MaterialDomain != MD_PostProcess && !IsTranslucentBlendMode(Analyzer.Material->BlendMode))
+					if (MaterialDomain != MD_PostProcess && !IsTranslucentBlendMode(Material->BlendMode))
 					{
 						Analyzer.Module->AddError(nullptr, TEXT("Only transparent or postprocess materials can read from scene depth."));
 					}
@@ -489,13 +502,13 @@ static void AnalyzeScreenTexture(FMaterialIRValueAnalyzer& Analyzer, MIR::FScree
 				Analyzer.CompilationOutput->SetIsSceneTextureUsed(PPI_SceneDepth);
 
 				// Error checking
-				if (Analyzer.Material->IsTranslucencyWritingVelocity())
+				if (Analyzer.MaterialInterface->IsTranslucencyWritingVelocity())
 				{
 					Analyzer.Module->AddError(nullptr, TEXT("Translucenct material with 'Output Velocity' enabled will write to depth buffer, therefore cannot read from depth buffer at the same time."));
 				}
 
 				// Don't allow opaque and masked materials to access scene depth as the results are undefined
-				if (MaterialDomain != MD_PostProcess && !IsTranslucentBlendMode(Analyzer.Material->BlendMode))
+				if (MaterialDomain != MD_PostProcess && !IsTranslucentBlendMode(Material->BlendMode))
 				{
 					Analyzer.Module->AddError(nullptr, TEXT("Only transparent or postprocess materials can read from scene depth."));
 				}
@@ -506,7 +519,7 @@ static void AnalyzeScreenTexture(FMaterialIRValueAnalyzer& Analyzer, MIR::FScree
 				// Error checking (no defines or compilation outputs needed for SceneDepthWithoutWater)
 				if (MaterialDomain != MD_PostProcess)
 				{
-					if (!Analyzer.Material->GetShadingModels().HasShadingModel(MSM_SingleLayerWater))
+					if (!Analyzer.MaterialInterface->GetShadingModels().HasShadingModel(MSM_SingleLayerWater))
 					{
 						Analyzer.Module->AddError(nullptr, TEXT("Can only read scene depth below water when material Shading Model is Single Layer Water or when material Domain is PostProcess."));
 					}
@@ -796,9 +809,9 @@ static void AnalyzeMaterialParameterCollection(FMaterialIRValueAnalyzer& Analyze
 // FMaterialIRValueAnalyzer
 //////////////////////////////////////////////////////////////////////////////////////////
 
-void FMaterialIRValueAnalyzer::Setup(UMaterial* InMaterial, FMaterialIRModule* InModule, FMaterialCompilationOutput* InCompilationOutput, FMaterialInsights* InInsights)
+void FMaterialIRValueAnalyzer::Setup(UMaterialInterface* InMaterialInterface, FMaterialIRModule* InModule, FMaterialCompilationOutput* InCompilationOutput, FMaterialInsights* InInsights)
 {
-	Material = InMaterial;
+	MaterialInterface = InMaterialInterface;
 	Module = InModule;
 	Insights = InInsights;
 	CompilationOutput = InCompilationOutput;
@@ -822,6 +835,7 @@ void FMaterialIRValueAnalyzer::Analyze(MIR::FValue* Value)
 	// IMPORTANT: Before adding a case here, read the FMaterialIRValueAnalyzer documentation.
 	switch (Value->Kind)
 	{
+		EXPAND_CASE(AnalysisError);
 		EXPAND_CASE(ExternalInput);
 		EXPAND_CASE(MaterialParameterCollection);
 		EXPAND_CASE(UniformParameter);
diff --git a/Engine/Source/Runtime/Engine/Private/Materials/MaterialIRValueAnalyzer.h b/Engine/Source/Runtime/Engine/Private/Materials/MaterialIRValueAnalyzer.h
--- a/Engine/Source/Runtime/Engine/Private/Materials/MaterialIRValueAnalyzer.h
+++ b/Engine/Source/Runtime/Engine/Private/Materials/MaterialIRValueAnalyzer.h
@@ -55,7 +55,7 @@ struct FMaterialIRValueAnalyzer
 	};
 
 	// Resets the analyzer to process a new translation run.
-	void Setup(UMaterial* InMaterial, FMaterialIRModule* InModule, FMaterialCompilationOutput* InCompilationOutput, FMaterialInsights* InInsights = nullptr);
+	void Setup(UMaterialInterface* InMaterialInterface, FMaterialIRModule* InModule, FMaterialCompilationOutput* InCompilationOutput, FMaterialInsights* InInsights = nullptr);
 		
 	// Performs stage-agnostic analysis on the given value.
 	void Analyze(MIR::FValue* Value);
@@ -67,7 +67,7 @@ struct FMaterialIRValueAnalyzer
 	UE::Shader::FPreshaderData& GetPreshaderData() { return CompilationOutput->UniformExpressionSet.UniformPreshaderData; }
 
 	// The material being built.
-	UMaterial* Material{};
+	UMaterialInterface* MaterialInterface{};
 	
 	// The destination module that will contain the result of material translation.
 	FMaterialIRModule* Module{};
diff --git a/Engine/Source/Runtime/Engine/Private/Materials/MaterialShared.cpp b/Engine/Source/Runtime/Engine/Private/Materials/MaterialShared.cpp
--- a/Engine/Source/Runtime/Engine/Private/Materials/MaterialShared.cpp
+++ b/Engine/Source/Runtime/Engine/Private/Materials/MaterialShared.cpp
@@ -3479,13 +3479,13 @@ bool FMaterial::Translate_New(const FMaterialShaderMapId& InShaderMapId,
 
 	// Build the material
 	FMaterialIRModuleBuilder Builder = {
-		.Material = GetMaterialInterface()->GetMaterial(),
+		.MaterialInterface = GetMaterialInterface(),
 		.ShaderPlatform = GetShaderPlatform(),
 		.TargetPlatform = InTargetPlatform,
 		.FeatureLevel = InShaderMapId.FeatureLevel,
 		.QualityLevel = InShaderMapId.QualityLevel,
 		.BlendMode = GetBlendMode(),
-		.StaticParameters = InStaticParameters,
+		.StaticParameters = &InStaticParameters,
 		.TargetInsights = GetMaterialInterface()->MaterialInsight.Get(),
 		.PreviewExpression = GetMaterialGraphNodePreviewExpression()
 	};
diff --git a/Engine/Source/Runtime/Engine/Public/Materials/MaterialExpressionBindlessSwitch.h b/Engine/Source/Runtime/Engine/Public/Materials/MaterialExpressionBindlessSwitch.h
--- a/Engine/Source/Runtime/Engine/Public/Materials/MaterialExpressionBindlessSwitch.h
+++ b/Engine/Source/Runtime/Engine/Public/Materials/MaterialExpressionBindlessSwitch.h
@@ -20,6 +20,7 @@ public:
 
 	//~ Begin UMaterialExpression Interface
 #if WITH_EDITOR
+	virtual void Build(MIR::FEmitter& Emitter) override;
 	virtual int32 Compile(FMaterialCompiler* Compiler, int32 OutputIndex) override;
 	virtual void GetCaption(TArray<FString>& OutCaptions) const override;
 	virtual bool IsInputConnectionRequired(int32 InputIndex) const override;
diff --git a/Engine/Source/Runtime/Engine/Public/Materials/MaterialExpressionDecalMipmapLevel.h b/Engine/Source/Runtime/Engine/Public/Materials/MaterialExpressionDecalMipmapLevel.h
--- a/Engine/Source/Runtime/Engine/Public/Materials/MaterialExpressionDecalMipmapLevel.h
+++ b/Engine/Source/Runtime/Engine/Public/Materials/MaterialExpressionDecalMipmapLevel.h
@@ -7,10 +7,11 @@
 #include "UObject/ObjectMacros.h"
 #include "MaterialExpressionIO.h"
 #include "Materials/MaterialExpression.h"
+#include "Materials/MaterialExpressionExternalCodeBase.h"
 #include "MaterialExpressionDecalMipmapLevel.generated.h"
 
 UCLASS(collapsecategories, hidecategories = Object)
-class UMaterialExpressionDecalMipmapLevel : public UMaterialExpression
+class UMaterialExpressionDecalMipmapLevel : public UMaterialExpressionExternalCodeBase
 {
 	GENERATED_UCLASS_BODY()
 
@@ -21,12 +22,13 @@ class UMaterialExpressionDecalMipmapLevel : public UMaterialExpression
 	/** only used if TextureSize is not hooked up */
 	UPROPERTY(EditAnywhere, Category = MaterialExpressionDecalMipmapLevel, meta=(OverridingInputProperty = "TextureSize"))
 	float ConstWidth;
+	
 	UPROPERTY(EditAnywhere, Category = MaterialExpressionDecalMipmapLevel, meta=(OverridingInputProperty = "TextureSize"))
 	float ConstHeight;
 
-
 	//~ Begin UMaterialExpression Interface
 #if WITH_EDITOR
+	virtual void Build(MIR::FEmitter& Em) override;
 	virtual int32 Compile(class FMaterialCompiler* Compiler, int32 OutputIndex) override;
 	virtual void GetCaption(TArray<FString>& OutCaptions) const override;
 #endif
diff --git a/Engine/Source/Runtime/Engine/Public/Materials/MaterialExpressionDepthOfFieldFunction.h b/Engine/Source/Runtime/Engine/Public/Materials/MaterialExpressionDepthOfFieldFunction.h
--- a/Engine/Source/Runtime/Engine/Public/Materials/MaterialExpressionDepthOfFieldFunction.h
+++ b/Engine/Source/Runtime/Engine/Public/Materials/MaterialExpressionDepthOfFieldFunction.h
@@ -7,6 +7,7 @@
 #include "UObject/ObjectMacros.h"
 #include "MaterialExpressionIO.h"
 #include "Materials/MaterialExpression.h"
+#include "Materials/MaterialExpressionExternalCodeBase.h"
 #include "MaterialExpressionDepthOfFieldFunction.generated.h"
 
 // Note: The index is used to map the enum to different code in the shader
@@ -25,7 +26,7 @@ enum EDepthOfFieldFunctionValue : int
 };
 
 UCLASS()
-class UMaterialExpressionDepthOfFieldFunction : public UMaterialExpression
+class UMaterialExpressionDepthOfFieldFunction : public UMaterialExpressionExternalCodeBase
 {
 	GENERATED_UCLASS_BODY()
 
@@ -40,6 +41,7 @@ class UMaterialExpressionDepthOfFieldFunction : public UMaterialExpression
 
 	//~ Begin UMaterialExpression Interface
 #if WITH_EDITOR
+	virtual void Build(MIR::FEmitter& Emitter) override;
 	virtual int32 Compile(class FMaterialCompiler* Compiler, int32 OutputIndex) override;
 	virtual void GetCaption(TArray<FString>& OutCaptions) const override;
 #endif
diff --git a/Engine/Source/Runtime/Engine/Public/Materials/MaterialExpressionDistanceCullFade.h b/Engine/Source/Runtime/Engine/Public/Materials/MaterialExpressionDistanceCullFade.h
--- a/Engine/Source/Runtime/Engine/Public/Materials/MaterialExpressionDistanceCullFade.h
+++ b/Engine/Source/Runtime/Engine/Public/Materials/MaterialExpressionDistanceCullFade.h
@@ -6,10 +6,11 @@
 #include "CoreMinimal.h"
 #include "UObject/ObjectMacros.h"
 #include "Materials/MaterialExpression.h"
+#include "Materials/MaterialExpressionExternalCodeBase.h"
 #include "MaterialExpressionDistanceCullFade.generated.h"
 
 UCLASS(collapsecategories, hidecategories=Object)
-class UMaterialExpressionDistanceCullFade : public UMaterialExpression
+class UMaterialExpressionDistanceCullFade : public UMaterialExpressionExternalCodeBase
 {
 	GENERATED_BODY()
 
@@ -17,6 +18,7 @@ public:
 
 	//~ Begin UMaterialExpression Interface
 #if WITH_EDITOR
+	virtual void Build(MIR::FEmitter& Emitter) override;
 	virtual int32 Compile(class FMaterialCompiler* Compiler, int32 OutputIndex) override;
 	virtual void GetCaption(TArray<FString>& OutCaptions) const override;
 
diff --git a/Engine/Source/Runtime/Engine/Public/Materials/MaterialExpressionGIReplace.h b/Engine/Source/Runtime/Engine/Public/Materials/MaterialExpressionGIReplace.h
--- a/Engine/Source/Runtime/Engine/Public/Materials/MaterialExpressionGIReplace.h
+++ b/Engine/Source/Runtime/Engine/Public/Materials/MaterialExpressionGIReplace.h
@@ -30,6 +30,7 @@ public:
 
 	//~ Begin UMaterialExpression Interface
 #if WITH_EDITOR
+	virtual void Build(MIR::FEmitter& Emitter) override;
 	virtual int32 Compile(class FMaterialCompiler* Compiler, int32 OutputIndex) override;
 	virtual void GetCaption(TArray<FString>& OutCaptions) const override;
 #endif
diff --git a/Engine/Source/Runtime/Engine/Public/Materials/MaterialExpressionHsvToRgb.h b/Engine/Source/Runtime/Engine/Public/Materials/MaterialExpressionHsvToRgb.h
--- a/Engine/Source/Runtime/Engine/Public/Materials/MaterialExpressionHsvToRgb.h
+++ b/Engine/Source/Runtime/Engine/Public/Materials/MaterialExpressionHsvToRgb.h
@@ -16,6 +16,7 @@ public:
 
 	//~ Begin UMaterialExpressionMaterialX Interface
 #if WITH_EDITOR
+	virtual void Build(MIR::FEmitter& Emitter) override;
 	virtual int32 Compile(class FMaterialCompiler* Compiler, int32 OutputIndex) override;
 	virtual void GetCaption(TArray<FString>& OutCaptions) const override;
 	virtual void GetExpressionToolTip(TArray<FString>& OutToolTip) override;
diff --git a/Engine/Source/Runtime/Engine/Public/Materials/MaterialExpressionMaterialProxyReplace.h b/Engine/Source/Runtime/Engine/Public/Materials/MaterialExpressionMaterialProxyReplace.h
--- a/Engine/Source/Runtime/Engine/Public/Materials/MaterialExpressionMaterialProxyReplace.h
+++ b/Engine/Source/Runtime/Engine/Public/Materials/MaterialExpressionMaterialProxyReplace.h
@@ -26,6 +26,7 @@ public:
 
 	// Begin UMaterialExpression Interface
 #if WITH_EDITOR
+	virtual void Build(MIR::FEmitter& Emitter) override;
 	virtual int32 Compile(class FMaterialCompiler* Compiler, int32 OutputIndex) override;
 	virtual bool IsResultMaterialAttributes(int32 OutputIndex) override;
 	virtual EMaterialValueType GetInputValueType(int32 InputIndex) override { return MCT_Unknown; }
diff --git a/Engine/Source/Runtime/Engine/Public/Materials/MaterialExpressionNeuralPostProcessNode.h b/Engine/Source/Runtime/Engine/Public/Materials/MaterialExpressionNeuralPostProcessNode.h
--- a/Engine/Source/Runtime/Engine/Public/Materials/MaterialExpressionNeuralPostProcessNode.h
+++ b/Engine/Source/Runtime/Engine/Public/Materials/MaterialExpressionNeuralPostProcessNode.h
@@ -54,7 +54,7 @@ class UMaterialExpressionNeuralNetworkInput : public UMaterialExpressionCustomOu
 	TEnumAsByte<ENeuralIndexType> NeuralIndexType;
 	
 #if WITH_EDITOR
-
+	virtual void Build(MIR::FEmitter& Em) override;
 	virtual int32 Compile(class FMaterialCompiler* Compiler, int32 OutputIndex) override;
 	virtual void GetCaption(TArray<FString>& OutCaptions) const override;
 	virtual EMaterialValueType GetInputValueType(int32 InputIndex) override;
@@ -89,6 +89,7 @@ class UMaterialExpressionNeuralNetworkOutput : public UMaterialExpression
 	TEnumAsByte<ENeuralIndexType> NeuralIndexType;
 
 #if WITH_EDITOR
+	virtual void Build(MIR::FEmitter& Em) override;
 	virtual int32 Compile(class FMaterialCompiler* Compiler, int32 OutputIndex) override;
 	virtual void GetCaption(TArray<FString>& OutCaptions) const override;
 	virtual EMaterialValueType GetInputValueType(int32 InputIndex) override;
diff --git a/Engine/Source/Runtime/Engine/Public/Materials/MaterialExpressionOperator.h b/Engine/Source/Runtime/Engine/Public/Materials/MaterialExpressionOperator.h
--- a/Engine/Source/Runtime/Engine/Public/Materials/MaterialExpressionOperator.h
+++ b/Engine/Source/Runtime/Engine/Public/Materials/MaterialExpressionOperator.h
@@ -71,6 +71,7 @@ enum class EMaterialExpressionOperatorKind
 	Modulo,
 	BitwiseAnd,
 	BitwiseOr,
+	BitwiseXor,
 	BitShiftLeft,
 	BitShiftRight,
 	ATan2,
diff --git a/Engine/Source/Runtime/Engine/Public/Materials/MaterialExpressionPathTracingBufferTexture.h b/Engine/Source/Runtime/Engine/Public/Materials/MaterialExpressionPathTracingBufferTexture.h
--- a/Engine/Source/Runtime/Engine/Public/Materials/MaterialExpressionPathTracingBufferTexture.h
+++ b/Engine/Source/Runtime/Engine/Public/Materials/MaterialExpressionPathTracingBufferTexture.h
@@ -41,6 +41,7 @@ class UMaterialExpressionPathTracingBufferTexture : public UMaterialExpression
 
 	//~ Begin UMaterialExpression Interface
 #if WITH_EDITOR
+	virtual void Build(MIR::FEmitter& Emitter) override;
 	virtual int32 Compile(class FMaterialCompiler* Compiler, int32 OutputIndex) override;
 	virtual void GetCaption(TArray<FString>& OutCaptions) const override;
 #endif
diff --git a/Engine/Source/Runtime/Engine/Public/Materials/MaterialExpressionPathTracingQualitySwitch.h b/Engine/Source/Runtime/Engine/Public/Materials/MaterialExpressionPathTracingQualitySwitch.h
--- a/Engine/Source/Runtime/Engine/Public/Materials/MaterialExpressionPathTracingQualitySwitch.h
+++ b/Engine/Source/Runtime/Engine/Public/Materials/MaterialExpressionPathTracingQualitySwitch.h
@@ -26,6 +26,7 @@ public:
 
 	//~ Begin UMaterialExpression Interface
 #if WITH_EDITOR
+	virtual void Build(MIR::FEmitter& Em) override;
 	virtual int32 Compile(class FMaterialCompiler* Compiler, int32 OutputIndex) override;
 	virtual void GetCaption(TArray<FString>& OutCaptions) const override;
 	virtual bool IsResultMaterialAttributes(int32 OutputIndex) override;
diff --git a/Engine/Source/Runtime/Engine/Public/Materials/MaterialExpressionPathTracingRayTypeSwitch.h b/Engine/Source/Runtime/Engine/Public/Materials/MaterialExpressionPathTracingRayTypeSwitch.h
--- a/Engine/Source/Runtime/Engine/Public/Materials/MaterialExpressionPathTracingRayTypeSwitch.h
+++ b/Engine/Source/Runtime/Engine/Public/Materials/MaterialExpressionPathTracingRayTypeSwitch.h
@@ -38,6 +38,7 @@ public:
 
 	//~ Begin UMaterialExpression Interface
 #if WITH_EDITOR
+	virtual void Build(MIR::FEmitter& Em) override;
 	virtual int32 Compile(class FMaterialCompiler* Compiler, int32 OutputIndex) override;
 	virtual void GetCaption(TArray<FString>& OutCaptions) const override;
 	virtual bool IsResultMaterialAttributes(int32 OutputIndex) override;
diff --git a/Engine/Source/Runtime/Engine/Public/Materials/MaterialExpressionPostVolumeUserFlagTest.h b/Engine/Source/Runtime/Engine/Public/Materials/MaterialExpressionPostVolumeUserFlagTest.h
--- a/Engine/Source/Runtime/Engine/Public/Materials/MaterialExpressionPostVolumeUserFlagTest.h
+++ b/Engine/Source/Runtime/Engine/Public/Materials/MaterialExpressionPostVolumeUserFlagTest.h
@@ -25,6 +25,7 @@ class UMaterialExpressionPostVolumeUserFlagTest : public UMaterialExpression
 	//~ Begin UMaterialExpression Interface
 #if WITH_EDITOR
 	virtual FText GetCreationName() const override { return FText::FromString(TEXT("TestPostVolumeUserFlag")); }
+	virtual void Build(MIR::FEmitter& Emitter) override;
 	virtual int32 Compile(class FMaterialCompiler* Compiler, int32 OutputIndex) override;
 	virtual void GetCaption(TArray<FString>& OutCaptions) const override { OutCaptions.Add(FString("TestPostVolumeUserFlag")); }
 #endif // WITH_EDITOR
diff --git a/Engine/Source/Runtime/Engine/Public/Materials/MaterialExpressionPreSkinnedLocalBounds.h b/Engine/Source/Runtime/Engine/Public/Materials/MaterialExpressionPreSkinnedLocalBounds.h
--- a/Engine/Source/Runtime/Engine/Public/Materials/MaterialExpressionPreSkinnedLocalBounds.h
+++ b/Engine/Source/Runtime/Engine/Public/Materials/MaterialExpressionPreSkinnedLocalBounds.h
@@ -14,6 +14,7 @@ class UMaterialExpressionPreSkinnedLocalBounds : public UMaterialExpression
 
 #if WITH_EDITOR
 	//~ Begin UMaterialExpression Interface
+	virtual void Build(MIR::FEmitter& Emitter) override;
 	virtual int32 Compile(class FMaterialCompiler* Compiler, int32 OutputIndex) override;
 	virtual void GetCaption(TArray<FString>& OutCaptions) const override;
 	virtual void GetConnectorToolTip(int32 InputIndex, int32 OutputIndex, TArray<FString>& OutToolTip) override;
diff --git a/Engine/Source/Runtime/Engine/Public/Materials/MaterialExpressionRayTracingQualitySwitch.h b/Engine/Source/Runtime/Engine/Public/Materials/MaterialExpressionRayTracingQualitySwitch.h
--- a/Engine/Source/Runtime/Engine/Public/Materials/MaterialExpressionRayTracingQualitySwitch.h
+++ b/Engine/Source/Runtime/Engine/Public/Materials/MaterialExpressionRayTracingQualitySwitch.h
@@ -26,6 +26,7 @@ public:
 
 	//~ Begin UMaterialExpression Interface
 #if WITH_EDITOR
+	virtual void Build(MIR::FEmitter& Em) override;
 	virtual int32 Compile(class FMaterialCompiler* Compiler, int32 OutputIndex) override;
 	virtual void GetCaption(TArray<FString>& OutCaptions) const override;
 	virtual bool IsResultMaterialAttributes(int32 OutputIndex) override;
diff --git a/Engine/Source/Runtime/Engine/Public/Materials/MaterialExpressionSRGBColorToWorkingColorSpace.h b/Engine/Source/Runtime/Engine/Public/Materials/MaterialExpressionSRGBColorToWorkingColorSpace.h
--- a/Engine/Source/Runtime/Engine/Public/Materials/MaterialExpressionSRGBColorToWorkingColorSpace.h
+++ b/Engine/Source/Runtime/Engine/Public/Materials/MaterialExpressionSRGBColorToWorkingColorSpace.h
@@ -21,6 +21,7 @@ public:
 
 	//~ Begin UMaterialExpression Interface
 #if WITH_EDITOR
+	virtual void Build(MIR::FEmitter& Em) override;
 	virtual int32 Compile(class FMaterialCompiler* Compiler, int32 OutputIndex) override;
 	virtual void GetCaption(TArray<FString>& OutCaptions) const override;
 	virtual FText GetKeywords() const override {return FText::FromString(TEXT("srgb workingcolorspace"));}
diff --git a/Engine/Source/Runtime/Engine/Public/Materials/MaterialExpressionSamplePhysicsField.h b/Engine/Source/Runtime/Engine/Public/Materials/MaterialExpressionSamplePhysicsField.h
--- a/Engine/Source/Runtime/Engine/Public/Materials/MaterialExpressionSamplePhysicsField.h
+++ b/Engine/Source/Runtime/Engine/Public/Materials/MaterialExpressionSamplePhysicsField.h
@@ -35,6 +35,7 @@ public:
 #if WITH_EDITOR
 	virtual FName GetInputName(int32 InputIndex) const override;
 	virtual void PostEditChangeProperty(FPropertyChangedEvent& PropertyChangedEvent) override;
+	virtual void Build(MIR::FEmitter& Emitter) override;
 	virtual int32 Compile(class FMaterialCompiler* Compiler, int32 OutputIndex) override;
 	virtual void GetCaption(TArray<FString>& OutCaptions) const override;
 
@@ -64,6 +65,7 @@ public:
 #if WITH_EDITOR
 	virtual FName GetInputName(int32 InputIndex) const override;
 	virtual void PostEditChangeProperty(FPropertyChangedEvent& PropertyChangedEvent) override;
+	virtual void Build(MIR::FEmitter& Emitter) override;
 	virtual int32 Compile(class FMaterialCompiler* Compiler, int32 OutputIndex) override;
 	virtual void GetCaption(TArray<FString>& OutCaptions) const override;
 
@@ -93,6 +95,7 @@ public:
 #if WITH_EDITOR
 	virtual FName GetInputName(int32 InputIndex) const override;
 	virtual void PostEditChangeProperty(FPropertyChangedEvent& PropertyChangedEvent) override;
+	virtual void Build(MIR::FEmitter& Emitter) override;
 	virtual int32 Compile(class FMaterialCompiler* Compiler, int32 OutputIndex) override;
 	virtual void GetCaption(TArray<FString>& OutCaptions) const override;
 
diff --git a/Engine/Source/Runtime/Engine/Public/Materials/MaterialExpressionSobol.h b/Engine/Source/Runtime/Engine/Public/Materials/MaterialExpressionSobol.h
--- a/Engine/Source/Runtime/Engine/Public/Materials/MaterialExpressionSobol.h
+++ b/Engine/Source/Runtime/Engine/Public/Materials/MaterialExpressionSobol.h
@@ -41,6 +41,7 @@ public:
 
 	//~ Begin UMaterialExpression Interface
 #if WITH_EDITOR
+	virtual void Build(MIR::FEmitter& Emitter) override;
 	virtual int32 Compile(class FMaterialCompiler* Compiler, int32 OutputIndex) override;
 	virtual void GetCaption(TArray<FString>& OutCaptions) const override;
 #endif
diff --git a/Engine/Source/Runtime/Engine/Public/Materials/MaterialExpressionSpeedTree.h b/Engine/Source/Runtime/Engine/Public/Materials/MaterialExpressionSpeedTree.h
--- a/Engine/Source/Runtime/Engine/Public/Materials/MaterialExpressionSpeedTree.h
+++ b/Engine/Source/Runtime/Engine/Public/Materials/MaterialExpressionSpeedTree.h
@@ -80,6 +80,7 @@ public:
 
 	//~ Begin UMaterialExpression Interface
 #if WITH_EDITOR
+	virtual void Build(MIR::FEmitter& Emitter) override;
 	virtual int32 Compile(class FMaterialCompiler* Compiler, int32 OutputIndex) override;
 	virtual void GetCaption(TArray<FString>& OutCaptions) const override;
 #endif
diff --git a/Engine/Source/Runtime/Engine/Public/Materials/MaterialExpressionTemporalSobol.h b/Engine/Source/Runtime/Engine/Public/Materials/MaterialExpressionTemporalSobol.h
--- a/Engine/Source/Runtime/Engine/Public/Materials/MaterialExpressionTemporalSobol.h
+++ b/Engine/Source/Runtime/Engine/Public/Materials/MaterialExpressionTemporalSobol.h
@@ -36,6 +36,7 @@ public:
 
 	//~ Begin UMaterialExpression Interface
 #if WITH_EDITOR
+	virtual void Build(MIR::FEmitter& Em) override;
 	virtual int32 Compile(class FMaterialCompiler* Compiler, int32 OutputIndex) override;
 	virtual void GetCaption(TArray<FString>& OutCaptions) const override;
 #endif
diff --git a/Engine/Source/Runtime/Engine/Public/Materials/MaterialExpressionVolumetricAdvancedMaterialInput.h b/Engine/Source/Runtime/Engine/Public/Materials/MaterialExpressionVolumetricAdvancedMaterialInput.h
--- a/Engine/Source/Runtime/Engine/Public/Materials/MaterialExpressionVolumetricAdvancedMaterialInput.h
+++ b/Engine/Source/Runtime/Engine/Public/Materials/MaterialExpressionVolumetricAdvancedMaterialInput.h
@@ -14,6 +14,7 @@ class UMaterialExpressionVolumetricAdvancedMaterialInput : public UMaterialExpre
 
 	//~ Begin UMaterialExpression Interface
 #if WITH_EDITOR
+	virtual void Build(MIR::FEmitter& Emitter) override;
 	virtual int32 Compile(class FMaterialCompiler* Compiler, int32 OutputIndex) override;
 	virtual void GetCaption(TArray<FString>& OutCaptions) const override;
 #endif
diff --git a/Engine/Source/Runtime/Engine/Public/Materials/MaterialIR.h b/Engine/Source/Runtime/Engine/Public/Materials/MaterialIR.h
--- a/Engine/Source/Runtime/Engine/Public/Materials/MaterialIR.h
+++ b/Engine/Source/Runtime/Engine/Public/Materials/MaterialIR.h
@@ -101,6 +101,7 @@ enum EValueKind
 	// Values
 
 	VK_Poison,
+	VK_AnalysisError,
 	VK_Constant,
 	VK_ExternalInput,
 	VK_MaterialParameterCollection,
@@ -307,6 +308,20 @@ struct FPoison : TValue<VK_Poison>
 	static FPoison* Get();
 };
 
+// A utility value whose sole purpose is to report a semantic error if and when it is analyzed.
+// This value is used to mark an invalid path, for example, in combination with StageSwitch, to ensure
+// that a value is only available in a subset of stages.
+// Note: this error is reported only if this value is analyzed. If the builder determines that this value
+// does not contribute to the final output, no error is reported (since the value will not be analyzed either).
+struct FAnalysisError : TValue<VK_AnalysisError>
+{
+	// The expression originating the invalid operation.
+	UMaterialExpression* Expression;
+
+	// The error message.
+	FSimpleStringView Message;
+};
+
 // The integer type used inside MIR.
 using TInteger = int64_t;
 
@@ -447,7 +462,8 @@ enum class EExternalInput
 	GlobalDistanceField,					// Flag to mark that global distance fields are used, for value analyzer step.  Doesn't generate code.
 	DynamicParticleParameterIndex,			// UserData stores the parameter index.
 	CompilingPreviousFrame,					// Boolean value bCompilingPreviousFrame, set when compiling the vertex shader WorldPositionOffsets for the previous frame.
-
+	DistanceCullFade,
+	
 	Count,
 };
 
@@ -737,8 +753,12 @@ struct FInlineHLSL : TInstruction<VK_InlineHLSL, 0>
 	// Array of argument values.
 	FValue* Arguments[MaxNumArguments];
 
-	// Number of arguments from another expression node.
-	int32 NumArguments;
+	// Custom analysis function.
+	using TOnAnalyzeCallback = void ( *)(FMaterialIRModule*, uintptr_t UserData);
+	TOnAnalyzeCallback OnAnalyze;
+
+	// User specified data passed to the OnAnalyze callback.
+	uintptr_t OnAnalyzeUserData;
 
 	union
 	{
@@ -748,6 +768,9 @@ struct FInlineHLSL : TInstruction<VK_InlineHLSL, 0>
 		// The arbitrary inlined HLSL code snippet
 		FSimpleStringView Code;
 	};
+
+	// Number of arguments from another expression node.
+	int32 NumArguments;
 };
 
 // Material IR representation of the FSubstrateData shader parameter to transition legacy materials into Substrate.
@@ -949,6 +972,7 @@ enum EOperator
 	BO_Modulo,
 	BO_BitwiseAnd,
 	BO_BitwiseOr,
+	BO_BitwiseXor,
 	BO_BitShiftLeft,
 	BO_BitShiftRight,
 	
diff --git a/Engine/Source/Runtime/Engine/Public/Materials/MaterialIRCommon.h b/Engine/Source/Runtime/Engine/Public/Materials/MaterialIRCommon.h
--- a/Engine/Source/Runtime/Engine/Public/Materials/MaterialIRCommon.h
+++ b/Engine/Source/Runtime/Engine/Public/Materials/MaterialIRCommon.h
@@ -11,15 +11,10 @@
 
 /* Forward declarations */
 
-class FMaterial;
-class FMaterialIRModule;
-class ITargetPlatform;
-class UMaterial;
-class UMaterialAggregate;
-class UMaterialExpression;
-class UMaterialExpressionMaterialFunctionCall;
-class UTexture;
 enum EMaterialProperty : int;
+enum EMaterialDomain : int;
+enum EMaterialUsage : int;
+
 struct FExpressionInput;
 struct FExpressionOutput;
 struct FMaterialAggregateAttribute;
@@ -30,6 +25,15 @@ struct FMaterialIRModuleBuilderImpl;
 struct FShaderCompilerEnvironment;
 struct FStaticParameterSet;
 
+class FMaterial;
+class FMaterialIRModule;
+class ITargetPlatform;
+class UMaterial;
+class UMaterialAggregate;
+class UMaterialExpression;
+class UMaterialExpressionMaterialFunctionCall;
+class UTexture;
+
 namespace Substrate
 {
 	struct FSubstrateTranslatorData;
diff --git a/Engine/Source/Runtime/Engine/Public/Materials/MaterialIREmitter.h b/Engine/Source/Runtime/Engine/Public/Materials/MaterialIREmitter.h
--- a/Engine/Source/Runtime/Engine/Public/Materials/MaterialIREmitter.h
+++ b/Engine/Source/Runtime/Engine/Public/Materials/MaterialIREmitter.h
@@ -105,6 +105,10 @@ struct FValueRef
 
 	// Creates a poison version of the current reference (same input, but poison value).
 	FValueRef ToPoison() const;
+
+	// Equality/inequality operators
+	bool operator==(const FValueRef Other) const { return Value == Other.Value; }
+	bool operator!=(const FValueRef Other) const { return Value != Other.Value; }
 	
 	// The referenced value.
 	FValue* Value{};
@@ -356,6 +360,9 @@ public:
 	FValueRef InputDefaultFloat3(const FExpressionInput* Input, UE::Math::TVector<TFloat> Default);
 	FValueRef InputDefaultFloat4(const FExpressionInput* Input, UE::Math::TVector4<TFloat> Default);
 
+	// Try read the value flowing into given Input, and if not present use Default instead.
+	FValueRef InputDefault(const FExpressionInput* Input, FValueRef Default);
+
 	// Retrieves the value from the given input of given type kind.
 	// If the value is missing or its type does not have the given kind, it reports an error and returns poison.
 	FValueRef CheckTypeIsKind(FValueRef Value, ETypeKind Kind);
@@ -378,6 +385,10 @@ public:
 	// Validates that the value is a vector. If not, reports error and returns poison.
 	FValueRef CheckIsVector(FValueRef Value);
 
+	// Validates that the value is a vector of number of components between specified min/max.
+	// If not, reports error and returns poison.
+	FValueRef CheckIsVector(FValueRef Value, int32 MinNumComponents, int32 MaxNumComponents);
+
 	// Validates that the value is either a scalar or a vector. Reports error and returns poison otherwise.
 	FValueRef CheckIsScalarOrVector(FValueRef Value);
 	
@@ -391,6 +402,15 @@ public:
 	// Validates that the value is of an aggregate type. If not, reports error and returns poison.
 	FValueRef CheckIsAggregate(FValueRef Value, const UMaterialAggregate* Aggregate = nullptr);
 
+	// Checks that the material domain is any of the ones specified.
+	void CheckMaterialDomainIsAnyOf(TConstArrayView<EMaterialDomain> Domains);
+
+	// Checks that the material has the specified usage flag.
+	void CheckMaterialHasUsageFlag(EMaterialUsage UsageFlag);
+
+	// Checks that the material feature level is equal or greater than the one specified.
+	void CheckFeatureLevelIsAtLeast(ERHIFeatureLevel::Type FeatureLevel);
+
 	// Casts the value flowing into this input to a constant boolean and returns it.
 	// If the value is missing, or is not a constant boolean it reports an error and returns false.
 	bool ToConstantBool(FValueRef Value);
@@ -403,8 +423,8 @@ public:
 	// Flows given `Value` out of given expression `ExpressionOutput`.
 	FEmitter& Output(const FExpressionOutput* ExpressionOutput, FValueRef Value);
 
-	// Flows given `Value` out of given expression output list. Each output applies its respective mask.
-	FEmitter& Outputs(const TConstArrayView<FExpressionOutput>& ExpressionOutputs, FValueRef Value);
+	// Flows given `Value` out of all outputs in the current expression.
+	FEmitter& OutputsWithComponentMask(FValueRef Value);
 	
 	/*------------------------------- Constants emission -------------------------------*/
 
@@ -473,6 +493,13 @@ public:
 	// Emits the poison value.
 	FValueRef Poison();
 
+	// Emits an analysis error value.
+	// The purpose of this value is to report an error during analysis if the analyzer
+	// proves that the value is effectively used. This value is useful when combined with
+	// others, especially the StageSwitch instruction in order to gate off certain operations
+	// on specific stage.
+	FValueRef AnalysisError(FType Type, FStringView Message);
+
 	// Emits an external input value of given input Id and returns it.
 	FValueRef ExternalInput(EExternalInput Id, uint32 UserData = 0);
 
@@ -520,23 +547,21 @@ public:
 	// you cast each output argument to the desired type before calling this function.
 	void SetCustomOutputs(FStringView Name, TConstArrayView<FValueRef> Args, EMaterialOutputStage Stage);
 	
-	// Emits a 2D column vector from X and Y values.
-	FValueRef Vector2(FValueRef InX, FValueRef InY);
-
-	// Emits a 3D vector from a 2D vector (XY) and a Z value.
-	FValueRef Vector3(FValueRef InXY, FValueRef InZ);
+	// Emits a 2D column vector from X and Y scalars.
+	FValueRef Vector2(FValueRef X, FValueRef Y);
 
-	// Emits a 3D vector from separate X, Y, and Z values.
-	FValueRef Vector3(FValueRef InX, FValueRef InY, FValueRef InZ);
+	// Emits a 3D column vector concatenating A, B and C components.
+	// Note: A, B and C can be scalars or vectors, but the sum of all components from A, B and possibly C must be 3.
+	FValueRef Vector3(FValueRef A, FValueRef B, FValueRef C = {});
 
-	// Emits a 4D vector from a 3D vector (XYZ) and a W value.
-	FValueRef Vector4(FValueRef InXYZ, FValueRef InW);
+	// Emits a 4D column vector concatenating A, B, C and D components.
+	// Note: A, B, C and D can be scalars or vectors, but the sum of all components from A, B and possibly C and D must be 4.
+	FValueRef Vector4(FValueRef A, FValueRef B, FValueRef C = {}, FValueRef D = {});
 
-	// Emits a 4D vector from a 2D vector (XY), a Z value, and a W value.
-	FValueRef Vector4(FValueRef InXY, FValueRef InZ, FValueRef InW);
-
-	// Emits a 4D vector from separate X, Y, Z, and W values.
-	FValueRef Vector4(FValueRef InX, FValueRef InY, FValueRef InZ, FValueRef InW);
+	// Emits a vector putting together all the components of given arguments.
+	// Arguments can be individual scalars or vectors themselves, but the concatenation of all their components cannot exceed 4 dimensions.
+	// Example: calling this function with a (float, float2) will result in a float3.
+	FValueRef Vector(FValueRef A, FValueRef B, FValueRef C = {}, FValueRef D = {});
 
 	// Constructs an aggregate value using the provided aggregate definition.
 	// The resulting value has default-initialized attributes.
@@ -697,7 +722,13 @@ public:
 	// @param InputValues An array of input values required by the HLSL code.
 	// @param ValueFlags Flags that specify additional properties for this value (default is None).
 	// @param UsedGraphProperties Properties used during graph processing (default is None).
-	FValueRef InlineHLSL(FType Type, FStringView Code, TConstArrayView<FValueRef> InputValues = {}, EValueFlags ValueFlags = EValueFlags::None, EGraphProperties UsedGraphProperties = EGraphProperties::None);
+	FValueRef InlineHLSL(FType Type, FStringView Code, TConstArrayView<FValueRef> InputValues = {}, EValueFlags ValueFlags = EValueFlags::None, EGraphProperties UsedGraphProperties = EGraphProperties::None, FInlineHLSL::TOnAnalyzeCallback OnAnalyze = {}, uintptr_t OnAnalyzeUserData = {});
+
+	// Emit a new HLSL function inline within the material shader based on an external code declaration.
+	// @param InExternalCodeDeclarationName The name of the external HLSL code declaration to use for emission.
+	// @param InputValues An array of input values required by the HLSL code.
+	// @param UsedGraphProperties Properties used during graph processing (default is None).
+	FValueRef InlineHLSL(FName InExternalCodeDeclarationName, TConstArrayView<FValueRef> InputValues = {}, EValueFlags ValueFlags = EValueFlags::None, EGraphProperties UsedGraphProperties = EGraphProperties::None);
 
 	// Emit a new HLSL function inline within the material shader based on an external code declaration.
 	// @param InExternalCodeDeclaration The external HLSL code declaration to use for emission.
@@ -792,6 +823,7 @@ public:
 	FValueRef Divide(FValueRef A, FValueRef B) { return Operator(BO_Divide, A, B); }
 	FValueRef BitwiseAnd(FValueRef A, FValueRef B) { return Operator(BO_BitwiseAnd, A, B); }
 	FValueRef BitwiseOr(FValueRef A, FValueRef B) { return Operator(BO_BitwiseOr, A, B); }
+	FValueRef BitwiseXor(FValueRef A, FValueRef B) { return Operator(BO_BitwiseXor, A, B); }
 	FValueRef BitShiftLeft(FValueRef A, FValueRef B) { return Operator(BO_BitShiftLeft, A, B); }
 	FValueRef BitShiftRight(FValueRef A, FValueRef B) { return Operator(BO_BitShiftRight, A, B); }
 	FValueRef Fmod(FValueRef A, FValueRef B) { return Operator(BO_Fmod, A, B); }
@@ -811,11 +843,24 @@ public:
 	FValueRef Select(FValueRef A, FValueRef B, FValueRef C) { return Operator(TO_Select, A, B, C); }
 	FValueRef Smoothstep(FValueRef A, FValueRef B, FValueRef C) { return Operator(TO_Smoothstep, A, B, C); }
 
-	/*-------------------------- Pass through accessors ---------------------------*/
+	/*---------------------------- Pass through accessors -----------------------------*/
+
+	// Returns the shader platform the source material is translating for.
 	EShaderPlatform GetShaderPlatform() const;
+	
+	// Returns the target platform the source material is translating for.
 	const ITargetPlatform* GetTargetPlatform() const;
+	
+	// Returns the target feature level the source material is translating for. 
 	ERHIFeatureLevel::Type GetFeatureLevel() const;
+
+	// Returns the target quality level the source material is translating for. 
 	EMaterialQualityLevel::Type GetQualityLevel() const;
+	
+	// Returns the source material interface. 
+	const UMaterialInterface* GetMaterialInterface() const { return MaterialInterface; }
+
+	// Returns the Substrate translator data in use.
 	Substrate::FSubstrateTranslatorData* GetSubstrateTranslatorData() const;
 
 	// Internal
@@ -838,8 +883,11 @@ private:
 	// Pointer to the builder internal implementation.
 	FMaterialIRModuleBuilderImpl* BuilderImpl{};
 
-	// The material being translated.
-	UMaterial* Material{};
+	// The material interface being translated.
+	UMaterialInterface* MaterialInterface{};
+
+	// The base material of the interface being translated (could be different from the interface if the interface is a material instance).
+	UMaterial* BaseMaterial{};
 
 	// The IR module being built.
 	FMaterialIRModule* Module{};
diff --git a/Engine/Source/Runtime/Engine/Public/Materials/MaterialIRModule.h b/Engine/Source/Runtime/Engine/Public/Materials/MaterialIRModule.h
--- a/Engine/Source/Runtime/Engine/Public/Materials/MaterialIRModule.h
+++ b/Engine/Source/Runtime/Engine/Public/Materials/MaterialIRModule.h
@@ -23,10 +23,10 @@ public:
 	// Represents an error encountered during material processing.
 	struct FError
 	{
-		// The expression that caused the error.
+		// The material expression from which this error originates.
 		UMaterialExpression* Expression;
 
-		// Description of the error.
+		// A textual description of the error.
 		FString Message;
 	};
 
@@ -112,6 +112,7 @@ public:
 
 	// Retrieves the material compilation output.
 	const FMaterialCompilationOutput& GetCompilationOutput() const { return CompilationOutput; }
+	FMaterialCompilationOutput& GetCompilationOutput() { return CompilationOutput; }
 
 	// Adds a new entry point with specified name, stage and optionally the expected number of evaluated output values hint.
 	int32 AddEntryPoint(FStringView Name, MIR::EStage Stage, int32 NumOutputsHint = 0);
@@ -144,6 +145,9 @@ public:
 	// Returns a list of all environment define names this module requires to be enabled for shader compilation. @massimo.tristano refactor this
 	const TSet<FName>& GetEnvironmentDefines() const { return EnvironmentDefines; }
 
+	// Adds an environment define.
+	void AddEnvironmentDefine(FName Name) { EnvironmentDefines.Add(Name); }
+
 	// Temporary. @massimo.tristano refactor this
 	TArray<TPair<FName, int32>> GetIntegerEnvironmentDefines() const { return IntegerEnvironmentDefines; }
 
@@ -190,7 +194,7 @@ public:
 	TArrayView<const FError> GetErrors() const { return Errors; }
 
 	// Reports a translation error.
-	void AddError(UMaterialExpression* Expression, FString Message);
+	void AddError(UMaterialExpression* Source, FString Message);
 	
 	// Allocates a chunk of memory using the module's internal memory block, to be used
 	// by the user as they see fit.
@@ -262,9 +266,8 @@ private:
 	TArray<FString> UserStrings;
 
 	// Environment define names for shader compilation in this module.
+	// @massimo.tristano refactor this into a better way to carry output flags from the builder to the backend.
 	TSet<FName> EnvironmentDefines;
-
-	// Temporary array of integer defines. @massimo.tristano refactor this into a better way to carry output flags from the builder to the backend.
 	TArray<TPair<FName, int32>> IntegerEnvironmentDefines;
 
 	// Parameter collections referenced by this module.
@@ -282,6 +285,8 @@ private:
 	// List of compilation errors.
 	TArray<FError> Errors;
 
+	// @massimo.tristano remove these friends and instead add the necessary public facing
+	// non-const accessors
 	friend MIR::FEmitter;
 	friend FMaterialIRModuleBuilder;
 	friend FMaterialIRModuleBuilderImpl;
diff --git a/Engine/Source/Runtime/Engine/Public/Materials/MaterialIRModuleBuilder.h b/Engine/Source/Runtime/Engine/Public/Materials/MaterialIRModuleBuilder.h
--- a/Engine/Source/Runtime/Engine/Public/Materials/MaterialIRModuleBuilder.h
+++ b/Engine/Source/Runtime/Engine/Public/Materials/MaterialIRModuleBuilder.h
@@ -6,18 +6,17 @@
 
 #if WITH_EDITOR
 
-//
 struct FMaterialIRModuleBuilder
 {
-	UMaterial* Material;
+	UMaterialInterface* MaterialInterface;
 	EShaderPlatform ShaderPlatform;
 	const ITargetPlatform* TargetPlatform;
 	ERHIFeatureLevel::Type FeatureLevel;
 	EMaterialQualityLevel::Type QualityLevel;
 	EBlendMode BlendMode;
-	const FStaticParameterSet& StaticParameters;
+	const FStaticParameterSet* StaticParameters{};
 	FMaterialInsights* TargetInsights{};
-	UMaterialExpression* PreviewExpression;
+	UMaterialExpression* PreviewExpression{};
 
 	bool Build(FMaterialIRModule* TargetModule);
 };
diff --git a/Engine/Source/Runtime/Engine/Public/Materials/MaterialIRTypes.h b/Engine/Source/Runtime/Engine/Public/Materials/MaterialIRTypes.h
--- a/Engine/Source/Runtime/Engine/Public/Materials/MaterialIRTypes.h
+++ b/Engine/Source/Runtime/Engine/Public/Materials/MaterialIRTypes.h
@@ -41,6 +41,12 @@ enum class EScalarKind : uint8
 	Double,	// Large World Coordinates type.
 };
 
+// Returns the larger scalar kind between A and B.
+static inline constexpr EScalarKind GetCommonScalarKind(EScalarKind A, EScalarKind B)
+{
+	return FMath::Max(A, B);
+}
+
 // Returns whether the specified scalar kind supports arithmetic operators (plus, minus, etc).
 bool ScalarKindIsArithmetic(EScalarKind Kind);
 
@@ -71,10 +77,10 @@ struct FPrimitive
 	// Returns whether this type is a primitive scalar (1x1).
 	bool IsScalar() const { return NumComponents() == 1; }
 	
-	// Returns whether this primitive type is a column vector (rows > 1, columns = 1).
+	// Returns whether this primitive type is a row vector (rows = 1, columns > 1).
 	bool IsRowVector() const { return NumRows == 1 && NumColumns > 1; }
 
-	// Returns whether this primitive type is a row vector (rows = 1, columns > 1).
+	// Returns whether this primitive type is a column vector  (rows > 1, columns = 1).
 	bool IsColumnVector() const { return NumRows > 1 && NumColumns == 1; }
 	
 	// Returns whether this primitive type is a matrix (rows > 1, columns > 1).
@@ -240,6 +246,9 @@ struct FType
 	// Returns true if the given type is a vector primitive.
 	bool IsVector() const { return Is(ETypeKind::Primitive) && Primitive.IsRowVector(); }
 
+	// Returns true if the given type is a vector primitive of given components.
+	bool IsVector(int32 NumComponents) const { return Is(ETypeKind::Primitive) && Primitive.IsRowVector() && Primitive.NumComponents() == NumComponents; }
+
 	// Returns whether this type is a Material Parameter Collection object.
 	bool IsParameterCollection() const { return Kind == ETypeKind::ParameterCollection; }
 
