Fixed small warning reported by Asan in my previous change.

#Synced-CL 49801225

diff --git a/Engine/Source/Runtime/Engine/Private/Materials/MaterialIREmitter.cpp b/Engine/Source/Runtime/Engine/Private/Materials/MaterialIREmitter.cpp
--- a/Engine/Source/Runtime/Engine/Private/Materials/MaterialIREmitter.cpp
+++ b/Engine/Source/Runtime/Engine/Private/Materials/MaterialIREmitter.cpp
@@ -3476,9 +3476,13 @@ FValueRef FEmitter::StageSwitch(FType Type, TConstArrayView<FValueRef> ValuePerS
 	FStageSwitch Prototype = MakePrototype<FStageSwitch>(Type);
 	for (int i = 0; i < ValuePerStage.Num(); ++i)
 	{
-		Prototype.Args[i] = ValuePerStage[i]
-			? ValuePerStage[i]
-			: AnalysisError(Type, FString::Printf(TEXT("This operation is not supported in shader stage '%s'."), Module->GetEntryPoint(i).Name.GetData()));
+		Prototype.Args[i] = ValuePerStage[i];
+
+		// If no value was specified for this stage, create an analysis error so that if used, it will report an error accordingly.
+		if (!Prototype.Args[i])
+		{
+			Prototype.Args[i] = AnalysisError(Type, FString::Printf(TEXT("This operation is not supported in shader stage '%s'."), Module->GetEntryPoint(i).Name.GetData()));
+		}
 	}
 	return EmitPrototype(*this, Prototype);
 }
