[AsyncLoading]
Makes the ordering used in LoadOrder tests more relaxed for circular reference cases where arbitrary cycle breaking within the loader cannot be relied upon for exact match testing of callback order of packages when flushing.

#rnx

#Synced-CL 49759049

diff --git a/Engine/Source/Runtime/Engine/Private/Tests/Loading/AsyncLoadingTests_Flushes.cpp b/Engine/Source/Runtime/Engine/Private/Tests/Loading/AsyncLoadingTests_Flushes.cpp
--- a/Engine/Source/Runtime/Engine/Private/Tests/Loading/AsyncLoadingTests_Flushes.cpp
+++ b/Engine/Source/Runtime/Engine/Private/Tests/Loading/AsyncLoadingTests_Flushes.cpp
@@ -119,8 +119,8 @@ enum class EVerifyOrder : uint8
 };
 ENUM_CLASS_FLAGS(EVerifyOrder);
 
-template <typename TMutateObjectsFn>
-void TestLoadPhases(FAutomationTestBase* AutomationTest, TMutateObjectsFn& MutateObjectsFn, const TArray<const TCHAR*>& PackagesToFlush, const TArray<const TCHAR*>& ExpectedOrder, EVerifyOrder VerifyOrder)
+template <typename TMutateObjectsFn, typename TIsOrderedFn>
+void TestLoadPhases(FAutomationTestBase* AutomationTest, TMutateObjectsFn& MutateObjectsFn, const TArray<const TCHAR*>& PackagesToFlush, TIsOrderedFn& IsOrderedFn, EVerifyOrder VerifyOrder)
 {
 	TArray<const TCHAR*> SerializedPackages; SerializedPackages.Reserve(3);
 	TArray<const TCHAR*> PostLoadedPackages; PostLoadedPackages.Reserve(3);
@@ -186,36 +186,16 @@ void TestLoadPhases(FAutomationTestBase* AutomationTest, TMutateObjectsFn& Mutat
 	}
 	FlushAsyncLoading(FlushRequests);
 
-	auto VerifyLoads = [AutomationTest](const TCHAR* OrderType, const TArray<const TCHAR*>& Actual, const TArray<const TCHAR*>& Expected, bool bVerifyExactOrder)
+	auto VerifyLoads = [AutomationTest,&IsOrderedFn](const TCHAR* OrderType, const TArray<const TCHAR*>& Actual, bool bVerifyOrder)
 		{
-			bool bMatch = false;
-			if (bVerifyExactOrder)
-			{
-				bMatch = Actual == Expected;
-			}
-			else
-			{
-				bMatch = Actual.Num() == Expected.Num() && Algo::AllOf(Actual, [&Expected](const TCHAR* Element)
-				{
-					return Expected.Contains(Element);
-				});
-			}
-
+			const bool bOnlyVerifyExistence = !bVerifyOrder;
+			const bool bMatch = IsOrderedFn(Actual, bOnlyVerifyExistence);
 			if (!bMatch)
 			{
 				TStringBuilder<256> Builder;
-				Builder.Appendf(TEXT("Expected to %s packages with %s order:\n\t"), OrderType,
-					(bVerifyExactOrder ? TEXT("exact") : TEXT("any")));
-				if (!Expected.IsEmpty())
-				{
-					for (const TCHAR* Path : Expected)
-					{
-						Builder.Append(Path);
-						Builder.Append(TEXT(", "));
-					}
-					Builder.RemoveSuffix(2);
-				}
-				Builder.Append(TEXT("\nActual order:\n\t"));
+				Builder.Appendf(TEXT("Expected to %s packages with '%s' ordering for packages 1 to 4:\n\t"), OrderType,
+					(bVerifyOrder ? TEXT("custom") : TEXT("any")));
+				Builder.Append(TEXT("Actual order:\n\t"));
 				if (!Actual.IsEmpty())
 				{
 					for (const TCHAR* Path : Actual)
@@ -225,6 +205,7 @@ void TestLoadPhases(FAutomationTestBase* AutomationTest, TMutateObjectsFn& Mutat
 					}
 					Builder.RemoveSuffix(2);
 				}
+				Builder.Append(TEXT("\nFailed to find all 4 package, or failed to find the expected package ordering (refer to IsOrdered lambda in test)."));
 				AutomationTest->TestTrue(Builder.ToString(), false);
 			}
 		};
@@ -232,11 +213,25 @@ void TestLoadPhases(FAutomationTestBase* AutomationTest, TMutateObjectsFn& Mutat
 	// Currently we do not provide any Serialize order guarantee with respect to dependencies.
 	// We likely never want this so we can enable parallel preloading opportunities. However
 	// if we should want preload ordering, re-enable to ensure this is done.
-	VerifyLoads(TEXT("Serialize"), SerializedPackages, ExpectedOrder, false /*!!(VerifyOrder& EVerifyOrder::Serialize)*/);
-	VerifyLoads(TEXT("PostLoad"), PostLoadedPackages, ExpectedOrder, !!(VerifyOrder & EVerifyOrder::PostLoad));
-	VerifyLoads(TEXT("LoadCompletedCallbacks"), LoadCompletedPackages, ExpectedOrder, !!(VerifyOrder & EVerifyOrder::LoadCompleted));
+	VerifyLoads(TEXT("Serialize"), SerializedPackages, false /*!!(VerifyOrder& EVerifyOrder::Serialize)*/);
+	VerifyLoads(TEXT("PostLoad"), PostLoadedPackages, !!(VerifyOrder & EVerifyOrder::PostLoad));
+	VerifyLoads(TEXT("LoadCompletedCallbacks"), LoadCompletedPackages, !!(VerifyOrder & EVerifyOrder::LoadCompleted));
 }
 
+// The majority of the time we expect Packages 1 to 4 to be loaded 4, 3, 2, 1
+bool IsLoadOrderReverseSorted(const TArray<const TCHAR*> FlushedPackages, const bool bOnlyVerifyExistence)
+	{
+		const int IndexOfPackage1 = FlushedPackages.IndexOfByKey(FLoadingTestsScope::PackagePath1); if (IndexOfPackage1 == INDEX_NONE) return false;
+		const int IndexOfPackage2 = FlushedPackages.IndexOfByKey(FLoadingTestsScope::PackagePath2); if (IndexOfPackage2 == INDEX_NONE) return false;
+		const int IndexOfPackage3 = FlushedPackages.IndexOfByKey(FLoadingTestsScope::PackagePath3); if (IndexOfPackage3 == INDEX_NONE) return false;
+		const int IndexOfPackage4 = FlushedPackages.IndexOfByKey(FLoadingTestsScope::PackagePath4); if (IndexOfPackage4 == INDEX_NONE) return false;
+		if (bOnlyVerifyExistence)
+		{
+			return true;
+		}
+		return IndexOfPackage4 < IndexOfPackage3 && IndexOfPackage3 < IndexOfPackage2 && IndexOfPackage2 < IndexOfPackage1;
+	};
+
 IMPLEMENT_CUSTOM_SIMPLE_AUTOMATION_TEST(
 	FLoadingTests_Flush_LoadOrder_A,
 	FLoadingTests_Flush_Base,
@@ -257,8 +252,7 @@ bool FLoadingTests_Flush_LoadOrder_A::RunTest(const FString& Parameters)
 
 	bool bVerifyExactOrder = true;
 	const TArray<const TCHAR*> PackagesToFlush = { FLoadingTestsScope::PackagePath1 };
-	const TArray<const TCHAR*> ExpectedOrder = { FLoadingTestsScope::PackagePath4, FLoadingTestsScope::PackagePath3, FLoadingTestsScope::PackagePath2, FLoadingTestsScope::PackagePath1 };
-	TestLoadPhases(this, MutateObjects, PackagesToFlush, ExpectedOrder, EVerifyOrder::All);
+	TestLoadPhases(this, MutateObjects, PackagesToFlush, IsLoadOrderReverseSorted, EVerifyOrder::All);
 	return true;
 }
 
@@ -281,8 +275,7 @@ bool FLoadingTests_Flush_LoadOrder_B::RunTest(const FString& Parameters)
 
 	bool bVerifyExactOrder = true;
 	const TArray<const TCHAR*> PackagesToFlush = { FLoadingTestsScope::PackagePath3, FLoadingTestsScope::PackagePath1 };
-	const TArray<const TCHAR*> ExpectedOrder = { FLoadingTestsScope::PackagePath4, FLoadingTestsScope::PackagePath3, FLoadingTestsScope::PackagePath2, FLoadingTestsScope::PackagePath1 };
-	TestLoadPhases(this, MutateObjects, PackagesToFlush, ExpectedOrder, EVerifyOrder::All);
+	TestLoadPhases(this, MutateObjects, PackagesToFlush, IsLoadOrderReverseSorted, EVerifyOrder::All);
 	return true;
 }
 
@@ -309,15 +302,26 @@ bool FLoadingTests_Flush_LoadOrder_C::RunTest(const FString& Parameters)
 	{
 		bool bVerifyExactOrder = true;
 		const TArray<const TCHAR*> PackagesToFlush = { FLoadingTestsScope::PackagePath3, FLoadingTestsScope::PackagePath1 };
-		const TArray<const TCHAR*> ExpectedOrder = { FLoadingTestsScope::PackagePath4, FLoadingTestsScope::PackagePath3, FLoadingTestsScope::PackagePath2, FLoadingTestsScope::PackagePath1 };
-		TestLoadPhases(this, MutateObjects, PackagesToFlush, ExpectedOrder, EVerifyOrder::LoadCompleted);
+		TestLoadPhases(this, MutateObjects, PackagesToFlush, IsLoadOrderReverseSorted, EVerifyOrder::LoadCompleted);
 	}
 
 	{
 		// Swap flush order. We should still see imports before their referencers when loading
 		const TArray<const TCHAR*> PackagesToFlush = { FLoadingTestsScope::PackagePath1, FLoadingTestsScope::PackagePath3 };
-		const TArray<const TCHAR*> ExpectedOrder = { FLoadingTestsScope::PackagePath4, FLoadingTestsScope::PackagePath2, FLoadingTestsScope::PackagePath1, FLoadingTestsScope::PackagePath3 };
-		TestLoadPhases(this, MutateObjects, PackagesToFlush, ExpectedOrder, EVerifyOrder::LoadCompleted);
+		auto IsOrdered = [](const TArray<const TCHAR*> FlushedPackages, const bool bOnlyVerifyExistence)
+			{
+				const int IndexOfPackage1 = FlushedPackages.IndexOfByKey(FLoadingTestsScope::PackagePath1); if (IndexOfPackage1 == INDEX_NONE) return false;
+				const int IndexOfPackage2 = FlushedPackages.IndexOfByKey(FLoadingTestsScope::PackagePath2); if (IndexOfPackage2 == INDEX_NONE) return false;
+				const int IndexOfPackage3 = FlushedPackages.IndexOfByKey(FLoadingTestsScope::PackagePath3); if (IndexOfPackage3 == INDEX_NONE) return false;
+				const int IndexOfPackage4 = FlushedPackages.IndexOfByKey(FLoadingTestsScope::PackagePath4); if (IndexOfPackage4 == INDEX_NONE) return false;
+				if (bOnlyVerifyExistence)
+				{
+					return true;
+				}
+				// Ensure we load leaf nodes first for the first flush and then load anything outstanding in the second (4,2,1 then 3)
+				return IndexOfPackage4 < IndexOfPackage2 && IndexOfPackage2 < IndexOfPackage1 && IndexOfPackage1 < IndexOfPackage3;
+			};
+		TestLoadPhases(this, MutateObjects, PackagesToFlush, IsOrdered, EVerifyOrder::LoadCompleted);
 	}
 	return true;
 }
@@ -347,8 +351,7 @@ bool FLoadingTests_Flush_LoadOrder_CircularDependencies_A::RunTest(const FString
 		};
 
 	const TArray<const TCHAR*> PackagesToFlush = { FLoadingTestsScope::PackagePath1 };
-	const TArray<const TCHAR*> ExpectedOrder = { FLoadingTestsScope::PackagePath4, FLoadingTestsScope::PackagePath3, FLoadingTestsScope::PackagePath2, FLoadingTestsScope::PackagePath1 };
-	TestLoadPhases(this, MutateObjects, PackagesToFlush, ExpectedOrder, EVerifyOrder::LoadCompleted);
+	TestLoadPhases(this, MutateObjects, PackagesToFlush, IsLoadOrderReverseSorted, EVerifyOrder::LoadCompleted);
 	return true;
 }
 
@@ -373,8 +376,7 @@ bool FLoadingTests_Flush_LoadOrder_CircularDependencies_B::RunTest(const FString
 		};
 
 	const TArray<const TCHAR*> PackagesToFlush = { FLoadingTestsScope::PackagePath1 };
-	const TArray<const TCHAR*> ExpectedOrder = { FLoadingTestsScope::PackagePath4, FLoadingTestsScope::PackagePath3, FLoadingTestsScope::PackagePath2, FLoadingTestsScope::PackagePath1 };
-	TestLoadPhases(this, MutateObjects, PackagesToFlush, ExpectedOrder, EVerifyOrder::LoadCompleted);
+	TestLoadPhases(this, MutateObjects, PackagesToFlush, IsLoadOrderReverseSorted, EVerifyOrder::LoadCompleted);
 	return true;
 }
 
@@ -401,8 +403,7 @@ bool FLoadingTests_Flush_LoadOrder_CircularDependencies_C::RunTest(const FString
 		};
 
 	const TArray<const TCHAR*> PackagesToFlush = { FLoadingTestsScope::PackagePath1 };
-	const TArray<const TCHAR*> ExpectedOrder = { FLoadingTestsScope::PackagePath4, FLoadingTestsScope::PackagePath3, FLoadingTestsScope::PackagePath2, FLoadingTestsScope::PackagePath1 };
-	TestLoadPhases(this, MutateObjects, PackagesToFlush, ExpectedOrder, EVerifyOrder::LoadCompleted);
+	TestLoadPhases(this, MutateObjects, PackagesToFlush, IsLoadOrderReverseSorted, EVerifyOrder::LoadCompleted);
 	return true;
 }
 
@@ -432,8 +433,7 @@ bool FLoadingTests_Flush_LoadOrder_CircularDependencies_D::RunTest(const FString
 		};
 
 	const TArray<const TCHAR*> PackagesToFlush = { FLoadingTestsScope::PackagePath1 };
-	const TArray<const TCHAR*> ExpectedOrder = { FLoadingTestsScope::PackagePath4, FLoadingTestsScope::PackagePath3, FLoadingTestsScope::PackagePath2, FLoadingTestsScope::PackagePath1 };
-	TestLoadPhases(this, MutateObjects, PackagesToFlush, ExpectedOrder, EVerifyOrder::LoadCompleted);
+	TestLoadPhases(this, MutateObjects, PackagesToFlush, IsLoadOrderReverseSorted, EVerifyOrder::LoadCompleted);
 	return true;
 }
 
@@ -462,12 +462,22 @@ bool FLoadingTests_Flush_LoadOrder_CircularDependencies_E::RunTest(const FString
 			Scope.Object4->HardReference = Scope.Object2;
 		};
 
-	// Note the ordering for nodes at the same depth is coincidentally numerically ordered, we only care that depth is sorted 
-	// but for simplicity omit that as the sorting used in the implementation is stable and it's unlikely the save order will 
-	// change affecting this test. If it does, update the test to adjust to the new order or relax the ordering constraint within each depth.
 	const TArray<const TCHAR*> PackagesToFlush = { FLoadingTestsScope::PackagePath1 };
-	const TArray<const TCHAR*> ExpectedOrder = { FLoadingTestsScope::PackagePath4, FLoadingTestsScope::PackagePath3, FLoadingTestsScope::PackagePath2, FLoadingTestsScope::PackagePath1 };
-	TestLoadPhases(this, MutateObjects, PackagesToFlush, ExpectedOrder, EVerifyOrder::LoadCompleted);
+	auto IsOrdered = [](const TArray<const TCHAR*> FlushedPackages, const bool bOnlyVerifyExistence)
+		{
+			const int IndexOfPackage1 = FlushedPackages.IndexOfByKey(FLoadingTestsScope::PackagePath1); if (IndexOfPackage1 == INDEX_NONE) return false;
+			const int IndexOfPackage2 = FlushedPackages.IndexOfByKey(FLoadingTestsScope::PackagePath2); if (IndexOfPackage2 == INDEX_NONE) return false;
+			const int IndexOfPackage3 = FlushedPackages.IndexOfByKey(FLoadingTestsScope::PackagePath3); if (IndexOfPackage3 == INDEX_NONE) return false;
+			const int IndexOfPackage4 = FlushedPackages.IndexOfByKey(FLoadingTestsScope::PackagePath4); if (IndexOfPackage4 == INDEX_NONE) return false;
+			if (bOnlyVerifyExistence)
+			{
+				return true;
+			}
+
+			// We only care that depth is sorted (handle leaf nodes before the cycle first). For nodes at the same depth, the order does not need to be strictly ordered.
+			return IndexOfPackage4 == 0 && (IndexOfPackage3 == 1 || IndexOfPackage3 == 2) && (IndexOfPackage2 == 1 || IndexOfPackage2 == 2) && IndexOfPackage1 == 3;
+		};
+	TestLoadPhases(this, MutateObjects, PackagesToFlush, IsOrdered, EVerifyOrder::LoadCompleted);
 	return true;
 }
 #endif // WITH_DEV_AUTOMATION_TESTS
