Expose FPoly::DoesLineIntersect

#Synced-CL 49727259

#Synced-CL 49727259
#Synced-CL 49727259

diff --git a/Engine/Source/Runtime/Engine/Classes/Engine/Polys.h b/Engine/Source/Runtime/Engine/Classes/Engine/Polys.h
--- a/Engine/Source/Runtime/Engine/Classes/Engine/Polys.h
+++ b/Engine/Source/Runtime/Engine/Classes/Engine/Polys.h
@@ -138,7 +138,7 @@ public:
 	 * Checks to see if the specified line intersects this poly or not.  If "Intersect" is
 	 * a valid pointer, it is filled in with the intersection point.
 	 */
-	bool DoesLineIntersect( FVector Start, FVector End, FVector* Intersect = NULL );
+	ENGINE_API bool DoesLineIntersect( FVector Start, FVector End, FVector* Intersect = NULL );
 
 	/**
 	 * Checks to see if the specified vertex is on this poly.  Assumes the vertex is on the same
