Fix warning V522: Dereferencing of the null pointer 'Texture' might take place. The null pointer is passed into 'Serialize' function.
#rnx

diff --git a/Engine/Source/Runtime/Engine/Private/TextureDerivedData.cpp b/Engine/Source/Runtime/Engine/Private/TextureDerivedData.cpp
--- a/Engine/Source/Runtime/Engine/Private/TextureDerivedData.cpp
+++ b/Engine/Source/Runtime/Engine/Private/TextureDerivedData.cpp
@@ -2994,7 +2994,7 @@ static void SerializePlatformData(
 		if (bStreamable)
 	#endif
 		{
-			check(Texture->IsPossibleToStream());
+			check(Texture && Texture->IsPossibleToStream());
 
 			int32 NumNonStreamingMips = PlatformData->GetNumNonStreamingMips(/*bIsStreamingPossible*/ true);
 			// NumMips has been reduced by FirstMipToSerialize (LODBias)
