Fix TaskSyncManager to correctly handle multiple batches per world and unregistering sync points during execution.
This also fixes the PIE case and other times where batches are not correctly reset.

#Synced-CL 49784053

diff --git a/Engine/Source/Runtime/Engine/Private/TaskSyncManager.cpp b/Engine/Source/Runtime/Engine/Private/TaskSyncManager.cpp
--- a/Engine/Source/Runtime/Engine/Private/TaskSyncManager.cpp
+++ b/Engine/Source/Runtime/Engine/Private/TaskSyncManager.cpp
@@ -1113,13 +1113,19 @@ bool FTaskSyncManager::UnregisterSyncPoint(FName RegisteredName, FName SourceNam
 			{
 				if (FoundData->RegisteredPoint.SourceName == SourceName)
 				{
-					// TODO clear out active batch map properly
+					for (FBatchData& BatchData : BatchList)
+					{
+						DestroySyncPointIfValid(BatchData, *FoundId);
+					}
+
 					RegisteredDataMap.Remove(*FoundId);
 					RegisteredNameMap.Remove(RegisteredName);
 				}
 				else
 				{
 					UE_LOG(LogTaskSync, Error, TEXT("Cannot unregister sync point %s from source %s as it was added by source %s!"), *RegisteredName.ToString(), *SourceName.ToString(), *FoundData->RegisteredPoint.SourceName.ToString());
+
+					return false;
 				}
 
 				return true;
@@ -1291,6 +1297,29 @@ FActiveSyncPoint* FTaskSyncManager::GetOrCreateSyncPoint(FBatchData& BatchData,
 	return &FoundData->Get();
 }
 
+bool FTaskSyncManager::DestroySyncPointIfValid(FBatchData& BatchData, FSyncPointId::FInternalId SyncId)
+{
+	TSharedRef<FActiveSyncPoint>* FoundData = BatchData.SyncPointData.Find(SyncId);
+
+	if (FoundData)
+	{
+		// Unregister the tick function which will stop it from ticking if it's still pending
+		(*FoundData)->UnRegisterTickFunction();
+
+		if (IsTicking() && FWorldContextId(CurrentTickWorld) == BatchData.WorldContext)
+		{
+			// If we're currently ticking, we need to keep this alive for a frame
+			BatchData.SyncPointsToDelete.Add(FoundData->ToSharedPtr());
+		}
+
+		// Remove the shared ptr. If this is the only thing keeping it alive, this will destroy it immediately
+		BatchData.SyncPointData.Remove(SyncId);
+
+		return true;
+	}
+	return false;
+}
+
 FActiveSyncPoint* FTaskSyncManager::GetOrCreateTickGroupWork(FBatchData& BatchData, ETickingGroup TickGroup)
 {
 	check(TickGroup >= 0 && TickGroup < ETickingGroup::TG_MAX);
@@ -1402,6 +1431,7 @@ void FTaskSyncManager::EndFrame(const UWorld* InWorld)
 		if (BatchData.WorldContext == WorldContext)
 		{
 			BatchData.TemporaryWorkRequests.Reset();
+			BatchData.SyncPointsToDelete.Reset();
 		}
 	}
 }
@@ -1410,9 +1440,24 @@ void FTaskSyncManager::InitializeForWorld(FWorldContextId WorldContext)
 {
 	check(IsInGameThread() && WorldContext.IsValid());
 
-	FBatchContextId BatchContext = CreateNewBatch(WorldContext);
-	check(BatchList.IsValidIndex(BatchContext.BatchId));
-	InitializeBatchForNewWorld(BatchList[BatchContext.BatchId]);
+	// The batch may already exist, if so just initialize it again
+	bool bFoundBatch = false;
+	for (FBatchData& BatchData : BatchList)
+	{
+		if (BatchData.WorldContext == WorldContext)
+		{
+			bFoundBatch = true;
+			InitializeBatchForNewWorld(BatchData);
+		}
+	}
+
+	// Add the default batch if there were 0
+	if (!bFoundBatch)
+	{
+		FBatchContextId BatchContext = CreateNewBatch(WorldContext);
+		check(BatchList.IsValidIndex(BatchContext.BatchId));
+		InitializeBatchForNewWorld(BatchList[BatchContext.BatchId]);
+	}
 }
 
 void FTaskSyncManager::InitializeBatchForNewWorld(FBatchData& BatchData)
diff --git a/Engine/Source/Runtime/Engine/Public/TaskSyncManager.h b/Engine/Source/Runtime/Engine/Public/TaskSyncManager.h
--- a/Engine/Source/Runtime/Engine/Public/TaskSyncManager.h
+++ b/Engine/Source/Runtime/Engine/Public/TaskSyncManager.h
@@ -916,6 +916,9 @@ private:
 		/** Array of general gamethread work per tickgroup, these are not scheduled as real tasks and only exist in default batch */
 		TArray<TSharedPtr<FActiveSyncPoint>> TickGroupWork;
 
+		/** List of sync points that were reset during the frame, will be deleted at end */
+		TArray<TSharedPtr<FActiveSyncPoint>> SyncPointsToDelete;
+
 		/** Array of temporary work handles, dropped at end of frame */
 		TArray<FTemporaryWorkRequest> TemporaryWorkRequests;
 
@@ -923,6 +926,7 @@ private:
 		{
 			SyncPointData.Reset();
 			TickGroupWork.Reset();
+			SyncPointsToDelete.Reset();
 			TemporaryWorkRequests.Reset();
 			BatchContext = FBatchContextId();
 			WorldContext = FWorldContextId();
@@ -974,6 +978,7 @@ private:
 #endif
 
 	FActiveSyncPoint* GetOrCreateSyncPoint(FBatchData& BatchData, FRegisteredSyncPointData& SyncData);
+	bool DestroySyncPointIfValid(FBatchData& BatchData, FSyncPointId::FInternalId SyncId);
 	FActiveSyncPoint* GetOrCreateTickGroupWork(FBatchData& BatchData, ETickingGroup TickGroup);
 
 	friend FActiveSyncPoint;
