Animation - fix for ensure in UAnimInstance::TriggerQueuedMontageEvents

#Synced-CL 49766586

diff --git a/Engine/Source/Runtime/Engine/Private/Animation/AnimInstance.cpp b/Engine/Source/Runtime/Engine/Private/Animation/AnimInstance.cpp
--- a/Engine/Source/Runtime/Engine/Private/Animation/AnimInstance.cpp
+++ b/Engine/Source/Runtime/Engine/Private/Animation/AnimInstance.cpp
@@ -2469,40 +2469,55 @@ void UAnimInstance::TriggerQueuedMontageEvents()
 	bQueueMontageEvents = false;
 
 	// Trigger Montage blending out before Ended events.
+
+	// we're copying event containers in case they get changed during broadcast delegate calls 
+	// we preferred a stack copy using TInlineAllocator allocated arrays rather than a MoveTemp to avoid reallocating the original arrays
+	// also, we cannot use TMemStackAllocator, since we cannot guarantee broadcast delegate calls will not reallocate memstack containers
+
+	constexpr int32 PreallocatedEvents = 4;
+
 	if (QueuedMontageBlendingOutEvents.Num() > 0)
 	{
-		for (const FQueuedMontageBlendingOutEvent& MontageBlendingOutEvent : QueuedMontageBlendingOutEvents)
+		TArray<FQueuedMontageBlendingOutEvent, TInlineAllocator<PreallocatedEvents>> QueuedMontageBlendingOutEventsCopy(QueuedMontageBlendingOutEvents);
+		QueuedMontageBlendingOutEvents.Reset();
+
+		for (const FQueuedMontageBlendingOutEvent& MontageBlendingOutEvent : QueuedMontageBlendingOutEventsCopy)
 		{
 			TriggerMontageBlendingOutEvent(MontageBlendingOutEvent);
 		}
-		QueuedMontageBlendingOutEvents.Reset();
 	}
 
 	if (QueuedMontageBlendedInEvents.Num() > 0)
 	{
-		for (const FQueuedMontageBlendedInEvent& MontageBlendedInEvent : QueuedMontageBlendedInEvents)
+		TArray<FQueuedMontageBlendedInEvent, TInlineAllocator<PreallocatedEvents>> QueuedMontageBlendedInEventsCopy(QueuedMontageBlendedInEvents);
+		QueuedMontageBlendedInEvents.Reset();
+
+		for (const FQueuedMontageBlendedInEvent& MontageBlendedInEvent : QueuedMontageBlendedInEventsCopy)
 		{
 			TriggerMontageBlendedInEvent(MontageBlendedInEvent);
 		}
-		QueuedMontageBlendedInEvents.Reset();
 	}
 
 	if (QueuedMontageSectionChangedEvents.Num() > 0)
 	{
-		for (const FQueuedMontageSectionChangedEvent& MontageSectionChangedEvent : QueuedMontageSectionChangedEvents)
+		TArray<FQueuedMontageSectionChangedEvent, TInlineAllocator<PreallocatedEvents>> QueuedMontageSectionChangedEventsCopy(QueuedMontageSectionChangedEvents);
+		QueuedMontageSectionChangedEvents.Reset();
+
+		for (const FQueuedMontageSectionChangedEvent& MontageSectionChangedEvent : QueuedMontageSectionChangedEventsCopy)
 		{
 			TriggerMontageSectionChangedEvent(MontageSectionChangedEvent);
 		}
-		QueuedMontageSectionChangedEvents.Reset();
 	}
 	
 	if (QueuedMontageEndedEvents.Num() > 0)
 	{
-		for (const FQueuedMontageEndedEvent& MontageEndedEvent : QueuedMontageEndedEvents)
+		TArray<FQueuedMontageEndedEvent, TInlineAllocator<PreallocatedEvents>> QueuedMontageEndedEventsCopy(QueuedMontageEndedEvents);
+		QueuedMontageEndedEvents.Reset();
+
+		for (const FQueuedMontageEndedEvent& MontageEndedEvent : QueuedMontageEndedEventsCopy)
 		{
 			TriggerMontageEndedEvent(MontageEndedEvent);
 		}
-		QueuedMontageEndedEvents.Reset();
 	}
 }
 
