Substrate - NMT - make the shading model from expression path compile.

#tests PC editor dx12
#rnx
#Synced-CL 49840788

diff --git a/Engine/Source/Runtime/Engine/Private/Materials/MaterialIRModuleBuilder.cpp b/Engine/Source/Runtime/Engine/Private/Materials/MaterialIRModuleBuilder.cpp
--- a/Engine/Source/Runtime/Engine/Private/Materials/MaterialIRModuleBuilder.cpp
+++ b/Engine/Source/Runtime/Engine/Private/Materials/MaterialIRModuleBuilder.cpp
@@ -140,9 +140,11 @@ struct FMaterialIRModuleBuilderImpl
 		FMaterialShadingModelField MaterialShadingModels;
 		// If the material gets its shading model from material expressions and we have compiled one or more shading model expressions already, 
 		// then use that shading model field instead. It's the most optimal set of shading models
-		if (Builder->MaterialInterface->IsShadingModelFromMaterialExpression())// && ShadingModelsFromCompilation.IsValid())
+		if (Builder->MaterialInterface->IsShadingModelFromMaterialExpression())
 		{
-			check(false); // SUBSTRATE_TODO add support for that case.
+			// This also correctly fetches all shading models form the graphs since material rebuild the "ShadingModels" field.
+			// SUBSTRATE_TODO It would be better to use Module->ShadingModelsFromCompilation gathered from the graph using static evaluation/switches but it is evaluated later.
+			MaterialShadingModels = Builder->MaterialInterface->GetShadingModels();
 		}
 		else
 		{
