[NetPhysics] remove config from the NetworkPhysicsSettingsDataAsset UPROPERTIES

#Synced-CL 49844248

diff --git a/Engine/Source/Runtime/Engine/Public/Physics/NetworkPhysicsSettingsComponent.h b/Engine/Source/Runtime/Engine/Public/Physics/NetworkPhysicsSettingsComponent.h
--- a/Engine/Source/Runtime/Engine/Public/Physics/NetworkPhysicsSettingsComponent.h
+++ b/Engine/Source/Runtime/Engine/Public/Physics/NetworkPhysicsSettingsComponent.h
@@ -97,20 +97,20 @@ struct FNetworkPhysicsSettings
 	FNetworkPhysicsSettings() { }
 
 	// Override properties
-	UPROPERTY(config, EditDefaultsOnly, Category = "Overrides", Meta = (InlineEditConditionToggle))
+	UPROPERTY(EditDefaultsOnly, Category = "Overrides", Meta = (InlineEditConditionToggle))
 	bool bOverrideSimProxyRepMode = false;
 
 	// Override the EPhysicsReplicationMode for Actors with ENetRole::ROLE_SimulatedProxy.
-	UPROPERTY(config, EditDefaultsOnly, Category = "Overrides", Meta = (EditCondition = "bOverrideSimProxyRepMode"))
+	UPROPERTY(EditDefaultsOnly, Category = "Overrides", Meta = (EditCondition = "bOverrideSimProxyRepMode"))
 	EPhysicsReplicationMode SimProxyRepMode = EPhysicsReplicationMode::PredictiveInterpolation;
 	
 	// Register this actor's Autonomous Proxy as a focal point / focal particle in Physics Replication LOD.
-	UPROPERTY(config, EditDefaultsOnly, Category = "Overrides")
+	UPROPERTY(EditDefaultsOnly, Category = "Overrides")
 	bool bFocalParticleInPhysicsReplicationLOD = false;
 
 	// Minimum delay added, in seconds, when scheduling an event far enough in the future to ensure server and all peers that event at the same time
 	// This value should be large enough to accommodate the largest network Round Trip Time, beyond which it is understood that the simulation will suffer from corrections
-	UPROPERTY(config, EditDefaultsOnly, Category = "General")
+	UPROPERTY(EditDefaultsOnly, Category = "General")
 	float EventSchedulingMinDelaySeconds = 0.3f;
 };
 
@@ -127,32 +127,32 @@ struct FNetworkPhysicsSettingsDefaultReplication
 	{}
 
 	// Override properties
-	UPROPERTY(config, EditDefaultsOnly, Category = "Overrides", Meta = (InlineEditConditionToggle))
+	UPROPERTY(EditDefaultsOnly, Category = "Overrides", Meta = (InlineEditConditionToggle))
 	uint32 bOverrideMaxLinearHardSnapDistance : 1;
-	UPROPERTY(config, EditDefaultsOnly, Category = "Overrides", Meta = (InlineEditConditionToggle))
+	UPROPERTY(EditDefaultsOnly, Category = "Overrides", Meta = (InlineEditConditionToggle))
 	uint32 bOverrideDefaultLegacyHardsnapInPT : 1;
-	UPROPERTY(config, EditDefaultsOnly, Category = "Overrides", Meta = (InlineEditConditionToggle))
+	UPROPERTY(EditDefaultsOnly, Category = "Overrides", Meta = (InlineEditConditionToggle))
 	uint32 bOverrideCorrectConnectedBodies : 1;
-	UPROPERTY(config, EditDefaultsOnly, Category = "Overrides", Meta = (InlineEditConditionToggle))
+	UPROPERTY(EditDefaultsOnly, Category = "Overrides", Meta = (InlineEditConditionToggle))
 	uint32 bOverrideCorrectConnectedBodiesFriction : 1;
 
 	// Overrides CVar: p.MaxLinearHardSnapDistance -- Hardsnap if distance between current position and extrapolated target position is larger than this value.
-	UPROPERTY(config, EditDefaultsOnly, Category = "Overrides", Meta = (EditCondition = "bOverrideMaxLinearHardSnapDistance"))
+	UPROPERTY(EditDefaultsOnly, Category = "Overrides", Meta = (EditCondition = "bOverrideMaxLinearHardSnapDistance"))
 	float MaxLinearHardSnapDistance = 400.0f;
 	float GetMaxLinearHardSnapDistance(float DefaultValue) const { return bOverrideMaxLinearHardSnapDistance ? MaxLinearHardSnapDistance : DefaultValue; }
 
 	// Overrides CVar: p.DefaultReplication.Legacy.HardsnapInPT -- If default replication is used and it's running the legacy flow through Game Thread, allow hardsnapping to be performed on Physics Thread if async physics is enabled.
-	UPROPERTY(config, EditDefaultsOnly, Category = "Overrides", Meta = (EditCondition = "bOverrideDefaultLegacyHardsnapInPT"))
+	UPROPERTY(EditDefaultsOnly, Category = "Overrides", Meta = (EditCondition = "bOverrideDefaultLegacyHardsnapInPT"))
 	bool bHardsnapInPhysicsThread = PhysicsReplicationCVars::DefaultReplicationCVars::bHardsnapLegacyInPT;
 	bool GetHardsnapDefaultLegacyInPT() const { return bOverrideDefaultLegacyHardsnapInPT ? bHardsnapInPhysicsThread : PhysicsReplicationCVars::DefaultReplicationCVars::bHardsnapLegacyInPT; }
 
 	// Overrides CVar: p.DefaultReplication.CorrectConnectedBodies -- When true, transform corrections will also apply to any connected physics object.
-	UPROPERTY(config, EditDefaultsOnly, Category = "Overrides", Meta = (EditCondition = "bOverrideCorrectConnectedBodies"))
+	UPROPERTY(EditDefaultsOnly, Category = "Overrides", Meta = (EditCondition = "bOverrideCorrectConnectedBodies"))
 	bool bCorrectConnectedBodies = PhysicsReplicationCVars::DefaultReplicationCVars::bCorrectConnectedBodies;
 	bool GetCorrectConnectedBodies() const { return bOverrideCorrectConnectedBodies ? bCorrectConnectedBodies : PhysicsReplicationCVars::DefaultReplicationCVars::bCorrectConnectedBodies; }
 
 	// Overrides CVar: p.DefaultReplication.CorrectConnectedBodiesFriction -- When true, transform correction on any connected physics object will also recalculate their friction.
-	UPROPERTY(config, EditDefaultsOnly, Category = "Overrides", Meta = (EditCondition = "bOverrideCorrectConnectedBodiesFriction"))
+	UPROPERTY(EditDefaultsOnly, Category = "Overrides", Meta = (EditCondition = "bOverrideCorrectConnectedBodiesFriction"))
 	bool bCorrectConnectedBodiesFriction = PhysicsReplicationCVars::DefaultReplicationCVars::bCorrectConnectedBodiesFriction;
 	bool GetCorrectConnectedBodiesFriction() const { return bOverrideCorrectConnectedBodiesFriction ? bCorrectConnectedBodiesFriction : PhysicsReplicationCVars::DefaultReplicationCVars::bCorrectConnectedBodiesFriction; }
 };
@@ -184,130 +184,130 @@ struct FNetworkPhysicsSettingsPredictiveInterpolation
 	{}
 
 	// Override properties
-	UPROPERTY(config, EditDefaultsOnly, Category = "Overrides", Meta = (InlineEditConditionToggle))
+	UPROPERTY(EditDefaultsOnly, Category = "Overrides", Meta = (InlineEditConditionToggle))
 	uint32 bOverrideVelocityBased : 1;
-	UPROPERTY(config, EditDefaultsOnly, Category = "Overrides", Meta = (InlineEditConditionToggle))
+	UPROPERTY(EditDefaultsOnly, Category = "Overrides", Meta = (InlineEditConditionToggle))
 	uint32 bOverrideCorrectionAsVelocity : 1;
-	UPROPERTY(config, EditDefaultsOnly, Category = "Overrides", Meta = (InlineEditConditionToggle))
+	UPROPERTY(EditDefaultsOnly, Category = "Overrides", Meta = (InlineEditConditionToggle))
 	uint32 bOverridePosCorrectionTimeBase : 1;
-	UPROPERTY(config, EditDefaultsOnly, Category = "Overrides", Meta = (InlineEditConditionToggle))
+	UPROPERTY(EditDefaultsOnly, Category = "Overrides", Meta = (InlineEditConditionToggle))
 	uint32 bOverridePosCorrectionTimeMin : 1;
-	UPROPERTY(config, EditDefaultsOnly, Category = "Overrides", Meta = (InlineEditConditionToggle))
+	UPROPERTY(EditDefaultsOnly, Category = "Overrides", Meta = (InlineEditConditionToggle))
 	uint32 bOverridePosCorrectionTimeMultiplier : 1;
-	UPROPERTY(config, EditDefaultsOnly, Category = "Overrides", Meta = (InlineEditConditionToggle))
+	UPROPERTY(EditDefaultsOnly, Category = "Overrides", Meta = (InlineEditConditionToggle))
 	uint32 bOverrideRotCorrectionTimeBase : 1;
-	UPROPERTY(config, EditDefaultsOnly, Category = "Overrides", Meta = (InlineEditConditionToggle))
+	UPROPERTY(EditDefaultsOnly, Category = "Overrides", Meta = (InlineEditConditionToggle))
 	uint32 bOverrideRotCorrectionTimeMin : 1;
-	UPROPERTY(config, EditDefaultsOnly, Category = "Overrides", Meta = (InlineEditConditionToggle))
+	UPROPERTY(EditDefaultsOnly, Category = "Overrides", Meta = (InlineEditConditionToggle))
 	uint32 bOverrideRotCorrectionTimeMultiplier : 1;
-	UPROPERTY(config, EditDefaultsOnly, Category = "Overrides", Meta = (InlineEditConditionToggle))
+	UPROPERTY(EditDefaultsOnly, Category = "Overrides", Meta = (InlineEditConditionToggle))
 	uint32 bOverridePosInterpolationTimeMultiplier : 1;
-	UPROPERTY(config, EditDefaultsOnly, Category = "Overrides", Meta = (InlineEditConditionToggle))
+	UPROPERTY(EditDefaultsOnly, Category = "Overrides", Meta = (InlineEditConditionToggle))
 	uint32 bOverrideRotInterpolationTimeMultiplier : 1;
-	UPROPERTY(config, EditDefaultsOnly, Category = "Overrides", Meta = (InlineEditConditionToggle))
+	UPROPERTY(EditDefaultsOnly, Category = "Overrides", Meta = (InlineEditConditionToggle))
 	uint32 bOverrideSoftSnapPosStrength : 1;
-	UPROPERTY(config, EditDefaultsOnly, Category = "Overrides", Meta = (InlineEditConditionToggle))
+	UPROPERTY(EditDefaultsOnly, Category = "Overrides", Meta = (InlineEditConditionToggle))
 	uint32 bOverrideSoftSnapRotStrength : 1;
-	UPROPERTY(config, EditDefaultsOnly, Category = "Overrides", Meta = (InlineEditConditionToggle))
+	UPROPERTY(EditDefaultsOnly, Category = "Overrides", Meta = (InlineEditConditionToggle))
 	uint32 bOverrideSoftSnapToSource : 1;
-	UPROPERTY(config, EditDefaultsOnly, Category = "Overrides", Meta = (InlineEditConditionToggle))
+	UPROPERTY(EditDefaultsOnly, Category = "Overrides", Meta = (InlineEditConditionToggle))
 	uint32 bOverrideDisableSoftSnap : 1;
-	UPROPERTY(config, EditDefaultsOnly, Category = "Overrides", Meta = (InlineEditConditionToggle))
+	UPROPERTY(EditDefaultsOnly, Category = "Overrides", Meta = (InlineEditConditionToggle))
 	uint32 bOverrideSkipVelocityRepOnPosEarlyOut : 1;
-	UPROPERTY(config, EditDefaultsOnly, Category = "Overrides", Meta = (InlineEditConditionToggle))
+	UPROPERTY(EditDefaultsOnly, Category = "Overrides", Meta = (InlineEditConditionToggle))
 	uint32 bOverridePostResimWaitForUpdate : 1;
-	UPROPERTY(config, EditDefaultsOnly, Category = "Overrides", Meta = (InlineEditConditionToggle))
+	UPROPERTY(EditDefaultsOnly, Category = "Overrides", Meta = (InlineEditConditionToggle))
 	uint32 bOverrideCorrectConnectedBodies : 1;
-	UPROPERTY(config, EditDefaultsOnly, Category = "Overrides", Meta = (InlineEditConditionToggle))
+	UPROPERTY(EditDefaultsOnly, Category = "Overrides", Meta = (InlineEditConditionToggle))
 	uint32 bOverrideCorrectConnectedBodiesFriction : 1;
 
 	// Overrides CVar: np2.PredictiveInterpolation.VelocityBased -- When true, predictive interpolation replication mode will move objects replicate by applying linear velocity and angular velocity. When false, objects will move by transform shifting towards the target state.
-	UPROPERTY(config, EditDefaultsOnly, Category = "Overrides", Meta = (EditCondition = "bOverrideVelocityBased"))
+	UPROPERTY(EditDefaultsOnly, Category = "Overrides", Meta = (EditCondition = "bOverrideVelocityBased"))
 	bool bVelocityBased = PhysicsReplicationCVars::PredictiveInterpolationCVars::bVelocityBased;
 	bool GetVelocityBased() const { return bOverrideVelocityBased ? bVelocityBased : PhysicsReplicationCVars::PredictiveInterpolationCVars::bVelocityBased; }
 
 	// Overrides CVar: np2.PredictiveInterpolation.CorrectionAsVelocity -- When true, predictive interpolation will apply positional and rotational offset correction as a velocity instead of as a transform shift, while np2.PredictiveInterpolation.VelocityBased is true.
-	UPROPERTY(config, EditDefaultsOnly, Category = "Overrides", Meta = (EditCondition = "bOverrideCorrectionAsVelocity"))
+	UPROPERTY(EditDefaultsOnly, Category = "Overrides", Meta = (EditCondition = "bOverrideCorrectionAsVelocity"))
 	bool bCorrectionAsVelocity = PhysicsReplicationCVars::PredictiveInterpolationCVars::bCorrectionAsVelocity;
 	bool GetCorrectionAsVelocity() const { return bOverrideCorrectionAsVelocity ? bCorrectionAsVelocity : PhysicsReplicationCVars::PredictiveInterpolationCVars::bCorrectionAsVelocity; }
 
 	// Overrides CVar: np2.PredictiveInterpolation.PosCorrectionTimeBase -- Base time to correct positional offset over. RoundTripTime * PosCorrectionTimeMultiplier is added on top of this.
-	UPROPERTY(config, EditDefaultsOnly, Category = "Overrides", Meta = (EditCondition = "bOverridePosCorrectionTimeBase"))
+	UPROPERTY(EditDefaultsOnly, Category = "Overrides", Meta = (EditCondition = "bOverridePosCorrectionTimeBase"))
 	float PosCorrectionTimeBase = PhysicsReplicationCVars::PredictiveInterpolationCVars::PosCorrectionTimeBase;
 	float GetPosCorrectionTimeBase() const { return bOverridePosCorrectionTimeBase ? PosCorrectionTimeBase : PhysicsReplicationCVars::PredictiveInterpolationCVars::PosCorrectionTimeBase; }
 
 	// Overrides CVar: np2.PredictiveInterpolation.PosCorrectionTimeMin -- Min time to correct positional offset over. DeltaSeconds is added on top of this.
-	UPROPERTY(config, EditDefaultsOnly, Category = "Overrides", Meta = (EditCondition = "bOverridePosCorrectionTimeMin"))
+	UPROPERTY(EditDefaultsOnly, Category = "Overrides", Meta = (EditCondition = "bOverridePosCorrectionTimeMin"))
 	float PosCorrectionTimeMin = PhysicsReplicationCVars::PredictiveInterpolationCVars::PosCorrectionTimeMin;
 	float GetPosCorrectionTimeMin() const { return bOverridePosCorrectionTimeMin ? PosCorrectionTimeMin : PhysicsReplicationCVars::PredictiveInterpolationCVars::PosCorrectionTimeMin; }
 
 	// Overrides CVar: np2.PredictiveInterpolation.PosCorrectionTimeMultiplier -- Multiplier to adjust how much of RoundTripTime to add to positional offset correction.
-	UPROPERTY(config, EditDefaultsOnly, Category = "Overrides", Meta = (EditCondition = "bOverridePosCorrectionTimeMultiplier"))
+	UPROPERTY(EditDefaultsOnly, Category = "Overrides", Meta = (EditCondition = "bOverridePosCorrectionTimeMultiplier"))
 	float PosCorrectionTimeMultiplier = PhysicsReplicationCVars::PredictiveInterpolationCVars::PosCorrectionTimeMultiplier;
 	float GetPosCorrectionTimeMultiplier() const { return bOverridePosCorrectionTimeMultiplier ? PosCorrectionTimeMultiplier : PhysicsReplicationCVars::PredictiveInterpolationCVars::PosCorrectionTimeMultiplier; }
 
 	// Overrides CVar: np2.PredictiveInterpolation.RotCorrectionTimeBase -- Base time to correct rotational offset over. RoundTripTime * RotCorrectionTimeMultiplier is added on top of this.
-	UPROPERTY(config, EditDefaultsOnly, Category = "Overrides", Meta = (EditCondition = "bOverrideRotCorrectionTimeBase"))
+	UPROPERTY(EditDefaultsOnly, Category = "Overrides", Meta = (EditCondition = "bOverrideRotCorrectionTimeBase"))
 	float RotCorrectionTimeBase = PhysicsReplicationCVars::PredictiveInterpolationCVars::RotCorrectionTimeBase;
 	float GetRotCorrectionTimeBase() const { return bOverrideRotCorrectionTimeBase ? RotCorrectionTimeBase : PhysicsReplicationCVars::PredictiveInterpolationCVars::RotCorrectionTimeBase; }
 
 	// Overrides CVar: np2.PredictiveInterpolation.RotCorrectionTimeMin -- Min time to correct rotational offset over. DeltaSeconds is added on top of this.
-	UPROPERTY(config, EditDefaultsOnly, Category = "Overrides", Meta = (EditCondition = "bOverrideRotCorrectionTimeMin"))
+	UPROPERTY(EditDefaultsOnly, Category = "Overrides", Meta = (EditCondition = "bOverrideRotCorrectionTimeMin"))
 	float RotCorrectionTimeMin = PhysicsReplicationCVars::PredictiveInterpolationCVars::RotCorrectionTimeMin;
 	float GetRotCorrectionTimeMin() const { return bOverrideRotCorrectionTimeMin ? RotCorrectionTimeMin : PhysicsReplicationCVars::PredictiveInterpolationCVars::RotCorrectionTimeMin; }
 
 	// Overrides CVar: np2.PredictiveInterpolation.RotCorrectionTimeMultiplier -- Multiplier to adjust how much of RoundTripTime to add to rotational offset correction.
-	UPROPERTY(config, EditDefaultsOnly, Category = "Overrides", Meta = (EditCondition = "bOverrideRotCorrectionTimeMultiplier"))
+	UPROPERTY(EditDefaultsOnly, Category = "Overrides", Meta = (EditCondition = "bOverrideRotCorrectionTimeMultiplier"))
 	float RotCorrectionTimeMultiplier = PhysicsReplicationCVars::PredictiveInterpolationCVars::RotCorrectionTimeMultiplier;
 	float GetRotCorrectionTimeMultiplier() const { return bOverrideRotCorrectionTimeMultiplier ? RotCorrectionTimeMultiplier : PhysicsReplicationCVars::PredictiveInterpolationCVars::RotCorrectionTimeMultiplier; }
 
 	// Overrides CVar: np2.PredictiveInterpolation.InterpolationTimeMultiplier -- Multiplier to adjust the interpolation time which is based on the sendrate of state data from the server.
-	UPROPERTY(config, EditDefaultsOnly, Category = "Overrides", Meta = (EditCondition = "bOverridePosInterpolationTimeMultiplier"))
+	UPROPERTY(EditDefaultsOnly, Category = "Overrides", Meta = (EditCondition = "bOverridePosInterpolationTimeMultiplier"))
 	float PosInterpolationTimeMultiplier = PhysicsReplicationCVars::PredictiveInterpolationCVars::PosInterpolationTimeMultiplier;
 	float GetPosInterpolationTimeMultiplier() const { return bOverridePosInterpolationTimeMultiplier ? PosInterpolationTimeMultiplier : PhysicsReplicationCVars::PredictiveInterpolationCVars::PosInterpolationTimeMultiplier; }
 
 	// Overrides CVar: np2.PredictiveInterpolation.RotInterpolationTimeMultiplier -- Multiplier to adjust the rotational interpolation time which is based on the sendrate of state data from the server.
-	UPROPERTY(config, EditDefaultsOnly, Category = "Overrides", Meta = (EditCondition = "bOverrideRotInterpolationTimeMultiplier"))
+	UPROPERTY(EditDefaultsOnly, Category = "Overrides", Meta = (EditCondition = "bOverrideRotInterpolationTimeMultiplier"))
 	float RotInterpolationTimeMultiplier = PhysicsReplicationCVars::PredictiveInterpolationCVars::RotInterpolationTimeMultiplier;
 	float GetRotInterpolationTimeMultiplier() const { return bOverrideRotInterpolationTimeMultiplier ? RotInterpolationTimeMultiplier : PhysicsReplicationCVars::PredictiveInterpolationCVars::RotInterpolationTimeMultiplier; }
 
 	// Overrides CVar: np2.PredictiveInterpolation.SoftSnapPosStrength -- Value in percent between 0.0 - 1.0 representing how much to softsnap each tick of the remaining positional distance.
-	UPROPERTY(config, EditDefaultsOnly, Category = "Overrides", Meta = (EditCondition = "bOverrideSoftSnapPosStrength"))
+	UPROPERTY(EditDefaultsOnly, Category = "Overrides", Meta = (EditCondition = "bOverrideSoftSnapPosStrength"))
 	float SoftSnapPosStrength = PhysicsReplicationCVars::PredictiveInterpolationCVars::SoftSnapPosStrength;
 	float GetSoftSnapPosStrength() const { return bOverrideSoftSnapPosStrength ? SoftSnapPosStrength : PhysicsReplicationCVars::PredictiveInterpolationCVars::SoftSnapPosStrength; }
 
 	// Overrides CVar: np2.PredictiveInterpolation.SoftSnapRotStrength -- Value in percent between 0.0 - 1.0 representing how much to softsnap each tick of the remaining rotational distance.
-	UPROPERTY(config, EditDefaultsOnly, Category = "Overrides", Meta = (EditCondition = "bOverrideSoftSnapRotStrength"))
+	UPROPERTY(EditDefaultsOnly, Category = "Overrides", Meta = (EditCondition = "bOverrideSoftSnapRotStrength"))
 	float SoftSnapRotStrength = PhysicsReplicationCVars::PredictiveInterpolationCVars::SoftSnapRotStrength;
 	float GetSoftSnapRotStrength() const { return bOverrideSoftSnapRotStrength ? SoftSnapRotStrength : PhysicsReplicationCVars::PredictiveInterpolationCVars::SoftSnapRotStrength; }
 
 	// Overrides CVar: np2.PredictiveInterpolation.SoftSnapToSource -- If true, softsnap will be performed towards the source state of the current target instead of the predicted state of the current target.
-	UPROPERTY(config, EditDefaultsOnly, Category = "Overrides", Meta = (EditCondition = "bOverrideSoftSnapToSource"))
+	UPROPERTY(EditDefaultsOnly, Category = "Overrides", Meta = (EditCondition = "bOverrideSoftSnapToSource"))
 	bool bSoftSnapToSource = PhysicsReplicationCVars::PredictiveInterpolationCVars::bSoftSnapToSource;
 	bool GetSoftSnapToSource() const { return bOverrideSoftSnapToSource ? bSoftSnapToSource : PhysicsReplicationCVars::PredictiveInterpolationCVars::bSoftSnapToSource; }
 
 	// Overrides CVar: np2.PredictiveInterpolation.DisableSoftSnap -- When true, predictive interpolation will not use softsnap to correct the replication with when velocity fails. Hardsnap will still eventually kick in if replication can't reach the target.
-	UPROPERTY(config, EditDefaultsOnly, Category = "Overrides", Meta = (EditCondition = "bOverrideDisableSoftSnap"))
+	UPROPERTY(EditDefaultsOnly, Category = "Overrides", Meta = (EditCondition = "bOverrideDisableSoftSnap"))
 	bool bDisableSoftSnap = PhysicsReplicationCVars::PredictiveInterpolationCVars::bDisableSoftSnap;
 	bool GetDisableSoftSnap() const { return bOverrideDisableSoftSnap ? bDisableSoftSnap : PhysicsReplicationCVars::PredictiveInterpolationCVars::bDisableSoftSnap; }
 
 	// Overrides CVar: np2.PredictiveInterpolation.SkipVelocityRepOnPosEarlyOut -- If true, don't run linear velocity replication if position can early out but angular can't early out.
-	UPROPERTY(config, EditDefaultsOnly, Category = "Overrides", Meta = (EditCondition = "bOverrideSkipVelocityRepOnPosEarlyOut"))
+	UPROPERTY(EditDefaultsOnly, Category = "Overrides", Meta = (EditCondition = "bOverrideSkipVelocityRepOnPosEarlyOut"))
 	bool bSkipVelocityRepOnPosEarlyOut = PhysicsReplicationCVars::PredictiveInterpolationCVars::bSkipVelocityRepOnPosEarlyOut;
 	bool GetSkipVelocityRepOnPosEarlyOut() const { return bOverrideSkipVelocityRepOnPosEarlyOut ? bSkipVelocityRepOnPosEarlyOut : PhysicsReplicationCVars::PredictiveInterpolationCVars::bSkipVelocityRepOnPosEarlyOut; }
 
 	// Overrides CVar: np2.PredictiveInterpolation.PostResimWaitForUpdate -- After a resimulation, wait for replicated states that correspond to post-resim state before processing replication again.
-	UPROPERTY(config, EditDefaultsOnly, Category = "Overrides", Meta = (EditCondition = "bOverridePostResimWaitForUpdate"))
+	UPROPERTY(EditDefaultsOnly, Category = "Overrides", Meta = (EditCondition = "bOverridePostResimWaitForUpdate"))
 	bool bPostResimWaitForUpdate = PhysicsReplicationCVars::PredictiveInterpolationCVars::bPostResimWaitForUpdate;
 	bool GetPostResimWaitForUpdate() const { return bOverridePostResimWaitForUpdate ? bPostResimWaitForUpdate : PhysicsReplicationCVars::PredictiveInterpolationCVars::bPostResimWaitForUpdate; }
 
 	// Overrides CVar: np2.PredictiveInterpolation.CorrectConnectedBodies -- When true, transform corrections will also apply to any connected physics object.
-	UPROPERTY(config, EditDefaultsOnly, Category = "Overrides", Meta = (EditCondition = "bOverrideCorrectConnectedBodies"))
+	UPROPERTY(EditDefaultsOnly, Category = "Overrides", Meta = (EditCondition = "bOverrideCorrectConnectedBodies"))
 	bool bCorrectConnectedBodies = PhysicsReplicationCVars::PredictiveInterpolationCVars::bCorrectConnectedBodies;
 	bool GetCorrectConnectedBodies() const { return bOverrideCorrectConnectedBodies ? bCorrectConnectedBodies : PhysicsReplicationCVars::PredictiveInterpolationCVars::bCorrectConnectedBodies; }
 
 	// Overrides CVar: np2.PredictiveInterpolation.CorrectConnectedBodiesFriction -- When true, transform correction on any connected physics object will also recalculate their friction.
-	UPROPERTY(config, EditDefaultsOnly, Category = "Overrides", Meta = (EditCondition = "bOverrideCorrectConnectedBodiesFriction"))
+	UPROPERTY(EditDefaultsOnly, Category = "Overrides", Meta = (EditCondition = "bOverrideCorrectConnectedBodiesFriction"))
 	bool bCorrectConnectedBodiesFriction = PhysicsReplicationCVars::PredictiveInterpolationCVars::bCorrectConnectedBodiesFriction;
 	bool GetCorrectConnectedBodiesFriction() const { return bOverrideCorrectConnectedBodiesFriction ? bCorrectConnectedBodiesFriction : PhysicsReplicationCVars::PredictiveInterpolationCVars::bCorrectConnectedBodiesFriction; }
 };
@@ -339,44 +339,44 @@ struct FNetworkPhysicsSettingsResimulationErrorCorrection
 
 	/** Enable override for post-resimulation error correction settings during render interpolation
 	* NOTE: This currently does not work if the experimental p.RenderInterp.ErrorVelocityCorrection CVar is set to true (false by default) */
-	UPROPERTY(config, EditDefaultsOnly, Category = "Overrides")
+	UPROPERTY(EditDefaultsOnly, Category = "Overrides")
 	uint32 bOverrideResimErrorInterpolationSettings : 1;
 
 	/** Overrides CVar: p.RenderInterp.ErrorCorrectionDuration -- How long in seconds to apply error correction over */
-	UPROPERTY(config, EditDefaultsOnly, Category = "Overrides", Meta = (EditCondition = "bOverrideResimErrorInterpolationSettings", ClampMin = "0.0", UIMin = "0.0"))
+	UPROPERTY(EditDefaultsOnly, Category = "Overrides", Meta = (EditCondition = "bOverrideResimErrorInterpolationSettings", ClampMin = "0.0", UIMin = "0.0"))
 	float ResimErrorCorrectionDuration;
 
 	/** Overrides CVar : p.RenderInterp.MaximumErrorCorrectionBeforeSnapping -- Maximum error correction in cm before we stop interpolating and snap to target */
-	UPROPERTY(config, EditDefaultsOnly, Category = "Overrides", Meta = (EditCondition = "bOverrideResimErrorInterpolationSettings", ClampMin = "0.0", UIMin = "0.0"))
+	UPROPERTY(EditDefaultsOnly, Category = "Overrides", Meta = (EditCondition = "bOverrideResimErrorInterpolationSettings", ClampMin = "0.0", UIMin = "0.0"))
 	float ResimErrorMaximumDistanceBeforeSnapping;
 
 	/** Overrides CVar: p.RenderInterp.MaximumErrorCorrectionDesyncTimeBeforeSnapping -- Time multiplied by the particles velocity to get the distance that error correction will be performed within without snapping, disable by setting a negative value
 	* NOTE: ResimErrorMaximumDistanceBeforeSnapping will act as a lowest distance clamp. */
-	UPROPERTY(config, EditDefaultsOnly, Category = "Overrides", Meta = (EditCondition = "bOverrideResimErrorInterpolationSettings", ClampMin = "-1.0", UIMin = "-1.0"))
+	UPROPERTY(EditDefaultsOnly, Category = "Overrides", Meta = (EditCondition = "bOverrideResimErrorInterpolationSettings", ClampMin = "-1.0", UIMin = "-1.0"))
 	float ResimErrorMaximumDesyncTimeBeforeSnapping;
 
 	/** Overrides CVar: p.RenderInterp.DirectionalDecayMultiplier -- Decay error offset in the direction that the physics object is moving, value is multiplier of projected offset direction, 0.25 means a 25 % decay of the magnitude in the direction of physics travel.Deactivate by setting to 0 */
-	UPROPERTY(config, EditDefaultsOnly, Category = "Overrides", Meta = (EditCondition = "bOverrideResimErrorInterpolationSettings", ClampMin = "0.0", UIMin = "0.0", ClampMax = "1.0", UIMax = "1.0"))
+	UPROPERTY(EditDefaultsOnly, Category = "Overrides", Meta = (EditCondition = "bOverrideResimErrorInterpolationSettings", ClampMin = "0.0", UIMin = "0.0", ClampMax = "1.0", UIMax = "1.0"))
 	float ResimErrorDirectionalDecayMultiplier;
 	
 	/** Overrides CVar: p.RenderInterp.ApplyExponentialDecay -- When enabled a post-resim error will decay exponentially (instead of linearly) based on half-life time set in ExponentialDecayLinearHalfLife and ExponentialDecayAngularHalfLife. */
-	UPROPERTY(config, EditDefaultsOnly, Category = "Overrides", Meta = (EditCondition = "bOverrideResimErrorInterpolationSettings"))
+	UPROPERTY(EditDefaultsOnly, Category = "Overrides", Meta = (EditCondition = "bOverrideResimErrorInterpolationSettings"))
 	bool bRenderInterpApplyExponentialDecay;
 	
 	/** Overrides CVar: p.RenderInterp.ExponentialDecayLinearHalfLife -- Sets the positional half-life time for when bApplyExponentialDecay is enabled. */
-	UPROPERTY(config, EditDefaultsOnly, Category = "Overrides", Meta = (EditCondition = "bOverrideResimErrorInterpolationSettings", ClampMin = "0.0", UIMin = "0.0"))
+	UPROPERTY(EditDefaultsOnly, Category = "Overrides", Meta = (EditCondition = "bOverrideResimErrorInterpolationSettings", ClampMin = "0.0", UIMin = "0.0"))
 	float RenderInterpExponentialDecayLinearHalfLife;
 
 	/** Overrides CVar: p.RenderInterp.ExponentialDecayAngularHalfLife -- Sets the rotational half-life time for when bApplyExponentialDecay is enabled. */
-	UPROPERTY(config, EditDefaultsOnly, Category = "Overrides", Meta = (EditCondition = "bOverrideResimErrorInterpolationSettings", ClampMin = "0.0", UIMin = "0.0"))
+	UPROPERTY(EditDefaultsOnly, Category = "Overrides", Meta = (EditCondition = "bOverrideResimErrorInterpolationSettings", ClampMin = "0.0", UIMin = "0.0"))
 	float RenderInterpExponentialDecayAngularHalfLife;
 
 	/** Overrides CVar: p.RenderInterp.MinimumLinearThreshold -- Squared value, when the remaining render error is below this we clear it, if ApplyExponentialDecay is enabled. */
-	UPROPERTY(config, EditDefaultsOnly, Category = "Overrides", Meta = (EditCondition = "bOverrideResimErrorInterpolationSettings", ClampMin = "0.0", UIMin = "0.0"))
+	UPROPERTY(EditDefaultsOnly, Category = "Overrides", Meta = (EditCondition = "bOverrideResimErrorInterpolationSettings", ClampMin = "0.0", UIMin = "0.0"))
 	float RenderInterpMinimumLinearThreshold;
 
 	/** Overrides CVar: p.RenderInterp.MinimumAngularThreshold -- When the remaining render error angle is below this we clear it, if ApplyExponentialDecay is enabled. */
-	UPROPERTY(config, EditDefaultsOnly, Category = "Overrides", Meta = (EditCondition = "bOverrideResimErrorInterpolationSettings", ClampMin = "0.0", UIMin = "0.0"))
+	UPROPERTY(EditDefaultsOnly, Category = "Overrides", Meta = (EditCondition = "bOverrideResimErrorInterpolationSettings", ClampMin = "0.0", UIMin = "0.0"))
 	float RenderInterpMinimumAngularThreshold;
 };
 
@@ -395,104 +395,91 @@ struct FNetworkPhysicsSettingsResimulation
 		, bOverrideRotStabilityMultiplier(0)
 		, bOverrideVelStabilityMultiplier(0)
 		, bOverrideAngVelStabilityMultiplier(0)
-		, bOverrideResimulationErrorThreshold_DEPRECATED(0)
-		, ResimulationErrorThreshold_DEPRECATED(10)
 		, ResimulationErrorCorrectionSettings()
 	{};
 
 	GENERATED_BODY()
 
 	// Override properties
-	UPROPERTY(config, EditDefaultsOnly, Category = "Overrides", Meta = (InlineEditConditionToggle))
+	UPROPERTY(EditDefaultsOnly, Category = "Overrides", Meta = (InlineEditConditionToggle))
 	uint32 bOverrideResimulationErrorPositionThreshold : 1;
-	UPROPERTY(config, EditDefaultsOnly, Category = "Overrides", Meta = (InlineEditConditionToggle))
+	UPROPERTY(EditDefaultsOnly, Category = "Overrides", Meta = (InlineEditConditionToggle))
 	uint32 bOverrideResimulationErrorRotationThreshold : 1;
-	UPROPERTY(config, EditDefaultsOnly, Category = "Overrides", Meta = (InlineEditConditionToggle))
+	UPROPERTY(EditDefaultsOnly, Category = "Overrides", Meta = (InlineEditConditionToggle))
 	uint32 bOverrideResimulationErrorLinearVelocityThreshold : 1;
-	UPROPERTY(config, EditDefaultsOnly, Category = "Overrides", Meta = (InlineEditConditionToggle))
+	UPROPERTY(EditDefaultsOnly, Category = "Overrides", Meta = (InlineEditConditionToggle))
 	uint32 bOverrideResimulationErrorAngularVelocityThreshold : 1;
-	UPROPERTY(config, EditDefaultsOnly, Category = "Overrides", Meta = (InlineEditConditionToggle))
+	UPROPERTY(EditDefaultsOnly, Category = "Overrides", Meta = (InlineEditConditionToggle))
 	uint32 bOverrideRuntimeCorrectionEnabled : 1;
-	UPROPERTY(config, EditDefaultsOnly, Category = "Overrides", Meta = (InlineEditConditionToggle))
+	UPROPERTY(EditDefaultsOnly, Category = "Overrides", Meta = (InlineEditConditionToggle))
 	uint32 bOverrideRuntimeVelocityCorrection : 1;
-	UPROPERTY(config, EditDefaultsOnly, Category = "Overrides", Meta = (InlineEditConditionToggle))
+	UPROPERTY(EditDefaultsOnly, Category = "Overrides", Meta = (InlineEditConditionToggle))
 	uint32 bOverrideRuntimeCorrectConnectedBodies : 1;
-	UPROPERTY(config, EditDefaultsOnly, Category = "Overrides", Meta = (InlineEditConditionToggle))
+	UPROPERTY(EditDefaultsOnly, Category = "Overrides", Meta = (InlineEditConditionToggle))
 	uint32 bOverridePosStabilityMultiplier : 1;
-	UPROPERTY(config, EditDefaultsOnly, Category = "Overrides", Meta = (InlineEditConditionToggle))
+	UPROPERTY(EditDefaultsOnly, Category = "Overrides", Meta = (InlineEditConditionToggle))
 	uint32 bOverrideRotStabilityMultiplier : 1;
-	UPROPERTY(config, EditDefaultsOnly, Category = "Overrides", Meta = (InlineEditConditionToggle))
+	UPROPERTY(EditDefaultsOnly, Category = "Overrides", Meta = (InlineEditConditionToggle))
 	uint32 bOverrideVelStabilityMultiplier : 1;
-	UPROPERTY(config, EditDefaultsOnly, Category = "Overrides", Meta = (InlineEditConditionToggle))
+	UPROPERTY(EditDefaultsOnly, Category = "Overrides", Meta = (InlineEditConditionToggle))
 	uint32 bOverrideAngVelStabilityMultiplier : 1;
 
-	/** Deprecated UE 5.5 - bOverrideResimulationErrorThreshold has been renamed, please use bOverrideResimulationErrorPositionThreshold*/
-	UPROPERTY(config)
-	uint32 bOverrideResimulationErrorThreshold_DEPRECATED : 1;
-
-	/** Deprecated UE 5.5 - ResimulationErrorThreshold has been renamed, please use ResimulationErrorPositionThreshold */
-	UPROPERTY(config)
-	uint32 ResimulationErrorThreshold_DEPRECATED;
-
-	UE_DEPRECATED(5.5, "GetResimulationErrorThreshold has been renamed, please use GetResimulationErrorPositionThreshold.")
-	uint32 GetResimulationErrorThreshold(uint32 DefaultValue) const { return (uint32)FMath::CeilToInt(GetResimulationErrorPositionThreshold(static_cast<float>(DefaultValue))); }
-
 	// Overrides Project Settings -> Physics -> Replication -> Physics Prediction -> Resimulation Error Position Threshold -- Distance that the object is allowed to desync from the server before triggering a resimulation, within this threshold runtime correction can be performed if RuntimeCorrectionEnabled is true.
-	UPROPERTY(config, EditDefaultsOnly, Category = "Overrides", Meta = (EditCondition = "bOverrideResimulationErrorPositionThreshold"))
+	UPROPERTY(EditDefaultsOnly, Category = "Overrides", Meta = (EditCondition = "bOverrideResimulationErrorPositionThreshold"))
 	float ResimulationErrorPositionThreshold = 10.0f;
 	float GetResimulationErrorPositionThreshold(float DefaultValue) const { return bOverrideResimulationErrorPositionThreshold ? ResimulationErrorPositionThreshold : DefaultValue; }
 
 	// Overrides Project Settings -> Physics -> Replication -> Physics Prediction -> Resimulation Error Rotation Threshold -- Rotation difference in degrees that the object is allowed to desync from the server before triggering a resimulation, within this threshold runtime correction can be performed if RuntimeCorrectionEnabled is true.
-	UPROPERTY(config, EditDefaultsOnly, Category = "Overrides", Meta = (EditCondition = "bOverrideResimulationErrorRotationThreshold"))
+	UPROPERTY(EditDefaultsOnly, Category = "Overrides", Meta = (EditCondition = "bOverrideResimulationErrorRotationThreshold"))
 	float ResimulationErrorRotationThreshold = 4.0f;
 	float GetResimulationErrorRotationThreshold(float DefaultValue) const { return bOverrideResimulationErrorRotationThreshold ? ResimulationErrorRotationThreshold : DefaultValue; }
 
 	// Overrides Project Settings -> Physics -> Replication -> Physics Prediction -> Resimulation Error Linear Velocity Threshold -- Velocity difference in centimeters / second that the object is allowed to desync from the server before triggering a resimulation, within this threshold runtime correction can be performed if RuntimeCorrectionEnabled is true.
-	UPROPERTY(config, EditDefaultsOnly, Category = "Overrides", Meta = (EditCondition = "bOverrideResimulationErrorLinearVelocityThreshold"))
+	UPROPERTY(EditDefaultsOnly, Category = "Overrides", Meta = (EditCondition = "bOverrideResimulationErrorLinearVelocityThreshold"))
 	float ResimulationErrorLinearVelocityThreshold = 5.0f;
 	float GetResimulationErrorLinearVelocityThreshold(float DefaultValue) const { return bOverrideResimulationErrorLinearVelocityThreshold ? ResimulationErrorLinearVelocityThreshold : DefaultValue; }
 
 	// Overrides Project Settings -> Physics -> Replication -> Physics Prediction -> Resimulation Error Angular Velocity Threshold -- Degrees / second that the object is allowed to desync from the server before triggering a resimulation, within this threshold runtime correction can be performed if RuntimeCorrectionEnabled is true.
-	UPROPERTY(config, EditDefaultsOnly, Category = "Overrides", Meta = (EditCondition = "bOverrideResimulationErrorAngularVelocityThreshold"))
+	UPROPERTY(EditDefaultsOnly, Category = "Overrides", Meta = (EditCondition = "bOverrideResimulationErrorAngularVelocityThreshold"))
 	float ResimulationErrorAngularVelocityThreshold = 2.0f;
 	float GetResimulationErrorAngularVelocityThreshold(float DefaultValue) const { return bOverrideResimulationErrorAngularVelocityThreshold ? ResimulationErrorAngularVelocityThreshold : DefaultValue; }
 
 	// Overrides CVar: np2.Resim.RuntimeCorrectionEnabled -- Apply positional and rotational runtime corrections while within resim trigger distance.
-	UPROPERTY(config, EditDefaultsOnly, Category = "Overrides", Meta = (EditCondition = "bOverrideRuntimeCorrectionEnabled"))
+	UPROPERTY(EditDefaultsOnly, Category = "Overrides", Meta = (EditCondition = "bOverrideRuntimeCorrectionEnabled"))
 	bool bRuntimeCorrectionEnabled = PhysicsReplicationCVars::ResimulationCVars::bRuntimeCorrectionEnabled;
 	bool GetRuntimeCorrectionEnabled() const { return bOverrideRuntimeCorrectionEnabled ? bRuntimeCorrectionEnabled : PhysicsReplicationCVars::ResimulationCVars::bRuntimeCorrectionEnabled; }
 
 	// Overrides CVar: np2.Resim.RuntimeVelocityCorrection -- Apply linear and angular velocity corrections in runtime while within resim trigger distance. Used if RuntimeCorrectionEnabled is true.
-	UPROPERTY(config, EditDefaultsOnly, Category = "Overrides", Meta = (EditCondition = "bOverrideRuntimeVelocityCorrection"))
+	UPROPERTY(EditDefaultsOnly, Category = "Overrides", Meta = (EditCondition = "bOverrideRuntimeVelocityCorrection"))
 	bool bRuntimeVelocityCorrection = PhysicsReplicationCVars::ResimulationCVars::bRuntimeVelocityCorrection;
 	bool GetRuntimeVelocityCorrectionEnabled() const { return bOverrideRuntimeVelocityCorrection ? bRuntimeVelocityCorrection : PhysicsReplicationCVars::ResimulationCVars::bRuntimeVelocityCorrection; }
 
 	// Overrides CVar: np2.Resim.RuntimeCorrectConnectedBodies -- If true runtime position and rotation correction will also shift transform of any connected physics objects. Used if RuntimeCorrectionEnabled is true.
-	UPROPERTY(config, EditDefaultsOnly, Category = "Overrides", Meta = (EditCondition = "bOverrideRuntimeCorrectConnectedBodies"))
+	UPROPERTY(EditDefaultsOnly, Category = "Overrides", Meta = (EditCondition = "bOverrideRuntimeCorrectConnectedBodies"))
 	bool bRuntimeCorrectConnectedBodies = PhysicsReplicationCVars::ResimulationCVars::bRuntimeCorrectConnectedBodies;
 	bool GetRuntimeCorrectConnectedBodies() const { return bOverrideRuntimeCorrectConnectedBodies ? bRuntimeCorrectConnectedBodies : PhysicsReplicationCVars::ResimulationCVars::bRuntimeCorrectConnectedBodies; }
 
 	// Overrides CVar: np2.Resim.PosStabilityMultiplier -- Recommended range between 0.0-1.0. Lower value means more stable positional corrections.
-	UPROPERTY(config, EditDefaultsOnly, Category = "Overrides", Meta = (EditCondition = "bOverridePosStabilityMultiplier"))
+	UPROPERTY(EditDefaultsOnly, Category = "Overrides", Meta = (EditCondition = "bOverridePosStabilityMultiplier"))
 	float PosStabilityMultiplier = PhysicsReplicationCVars::ResimulationCVars::PosStabilityMultiplier;
 	float GetPosStabilityMultiplier() { return bOverridePosStabilityMultiplier ? PosStabilityMultiplier : PhysicsReplicationCVars::ResimulationCVars::PosStabilityMultiplier; }
 
 	// Overrides CVar: np2.Resim.RotStabilityMultiplier -- Recommended range between 0.0-1.0. Lower value means more stable rotational corrections.
-	UPROPERTY(config, EditDefaultsOnly, Category = "Overrides", Meta = (EditCondition = "bOverrideRotStabilityMultiplier"))
+	UPROPERTY(EditDefaultsOnly, Category = "Overrides", Meta = (EditCondition = "bOverrideRotStabilityMultiplier"))
 	float RotStabilityMultiplier = PhysicsReplicationCVars::ResimulationCVars::RotStabilityMultiplier;
 	float GetRotStabilityMultiplier() const { return bOverrideRotStabilityMultiplier ? RotStabilityMultiplier : PhysicsReplicationCVars::ResimulationCVars::RotStabilityMultiplier; }
 
 	// Overrides CVar: np2.Resim.VelStabilityMultiplier -- Recommended range between 0.0-1.0. Lower value means more stable linear velocity corrections.
-	UPROPERTY(config, EditDefaultsOnly, Category = "Overrides", Meta = (EditCondition = "bOverrideVelStabilityMultiplier"))
+	UPROPERTY(EditDefaultsOnly, Category = "Overrides", Meta = (EditCondition = "bOverrideVelStabilityMultiplier"))
 	float VelStabilityMultiplier = PhysicsReplicationCVars::ResimulationCVars::VelStabilityMultiplier;
 	float GetVelStabilityMultiplier() const { return bOverrideVelStabilityMultiplier ? VelStabilityMultiplier : PhysicsReplicationCVars::ResimulationCVars::VelStabilityMultiplier; }
 
 	// Overrides CVar: np2.Resim.AngVelStabilityMultiplier -- Recommended range between 0.0-1.0. Lower value means more stable angular velocity corrections.
-	UPROPERTY(config, EditDefaultsOnly, Category = "Overrides", Meta = (EditCondition = "bOverrideAngVelStabilityMultiplier"))
+	UPROPERTY(EditDefaultsOnly, Category = "Overrides", Meta = (EditCondition = "bOverrideAngVelStabilityMultiplier"))
 	float AngVelStabilityMultiplier = PhysicsReplicationCVars::ResimulationCVars::AngVelStabilityMultiplier;
 	float GetAngVelStabilityMultiplier() const { return bOverrideAngVelStabilityMultiplier ? AngVelStabilityMultiplier : PhysicsReplicationCVars::ResimulationCVars::AngVelStabilityMultiplier; }
 
-	UPROPERTY(config, EditDefaultsOnly, Category = "Overrides", meta = (DisplayName = "Post-Resimulation Error Correction Render Interpolation Settings"))
+	UPROPERTY(EditDefaultsOnly, Category = "Overrides", meta = (DisplayName = "Post-Resimulation Error Correction Render Interpolation Settings"))
 	FNetworkPhysicsSettingsResimulationErrorCorrection ResimulationErrorCorrectionSettings;
 };
 
@@ -537,49 +524,49 @@ struct FNetworkPhysicsSettingsNetworkPhysicsComponent
 	}
 
 	// Override properties
-	UPROPERTY(config, EditDefaultsOnly, Category = "Overrides", Meta = (InlineEditConditionToggle))
+	UPROPERTY(EditDefaultsOnly, Category = "Overrides", Meta = (InlineEditConditionToggle))
 	uint32 bOverrideRedundantInputs : 1;
-	UPROPERTY(config, EditDefaultsOnly, Category = "Overrides", Meta = (InlineEditConditionToggle))
+	UPROPERTY(EditDefaultsOnly, Category = "Overrides", Meta = (InlineEditConditionToggle))
 	uint32 bOverrideRedundantRemoteInputs : 1;
-	UPROPERTY(config, EditDefaultsOnly, Category = "Overrides", Meta = (InlineEditConditionToggle))
+	UPROPERTY(EditDefaultsOnly, Category = "Overrides", Meta = (InlineEditConditionToggle))
 	uint32 bOverrideRedundantStates : 1;
-	UPROPERTY(config, EditDefaultsOnly, Category = "Overrides", Meta = (InlineEditConditionToggle))
+	UPROPERTY(EditDefaultsOnly, Category = "Overrides", Meta = (InlineEditConditionToggle))
 	uint32 bOverrideDynamicInputBufferScalingEnabled : 1;
-	UPROPERTY(config, EditDefaultsOnly, Category = "Overrides", Meta = (InlineEditConditionToggle))
+	UPROPERTY(EditDefaultsOnly, Category = "Overrides", Meta = (InlineEditConditionToggle))
 	uint32 bOverrideCompareStateToTriggerRewind : 1;
-	UPROPERTY(config, EditDefaultsOnly, Category = "Overrides", Meta = (InlineEditConditionToggle))
+	UPROPERTY(EditDefaultsOnly, Category = "Overrides", Meta = (InlineEditConditionToggle))
 	uint32 bOverridebCompareStateToTriggerRewindIncludeSimProxies : 1;
-	UPROPERTY(config, EditDefaultsOnly, Category = "Overrides", Meta = (InlineEditConditionToggle))
+	UPROPERTY(EditDefaultsOnly, Category = "Overrides", Meta = (InlineEditConditionToggle))
 	uint32 bOverrideCompareInputToTriggerRewind : 1;
-	UPROPERTY(config, EditDefaultsOnly, Category = "Overrides", Meta = (InlineEditConditionToggle))
+	UPROPERTY(EditDefaultsOnly, Category = "Overrides", Meta = (InlineEditConditionToggle))
 	uint32 bOverrideEnableUnreliableFlow : 1;
-	UPROPERTY(config, EditDefaultsOnly, Category = "Overrides", Meta = (InlineEditConditionToggle))
+	UPROPERTY(EditDefaultsOnly, Category = "Overrides", Meta = (InlineEditConditionToggle))
 	uint32 bOverrideEnableReliableFlow : 1;
-	UPROPERTY(config, EditDefaultsOnly, Category = "Overrides", Meta = (InlineEditConditionToggle))
+	UPROPERTY(EditDefaultsOnly, Category = "Overrides", Meta = (InlineEditConditionToggle))
 	uint32 bOverrideApplyDataInsteadOfMergeData : 1;
-	UPROPERTY(config, EditDefaultsOnly, Category = "Overrides", Meta = (InlineEditConditionToggle))
+	UPROPERTY(EditDefaultsOnly, Category = "Overrides", Meta = (InlineEditConditionToggle))
 	uint32 bOverrideAllowInputExtrapolation : 1;
-	UPROPERTY(config, EditDefaultsOnly, Category = "Overrides", Meta = (InlineEditConditionToggle))
+	UPROPERTY(EditDefaultsOnly, Category = "Overrides", Meta = (InlineEditConditionToggle))
 	uint32 bOverrideValidateDataOnGameThread : 1;
-	UPROPERTY(config, EditDefaultsOnly, Category = "Overrides", Meta = (InlineEditConditionToggle))
+	UPROPERTY(EditDefaultsOnly, Category = "Overrides", Meta = (InlineEditConditionToggle))
 	uint32 bOverrideTriggerResimOnInputReceive : 1;
-	UPROPERTY(config, EditDefaultsOnly, Category = "Overrides", Meta = (InlineEditConditionToggle))
+	UPROPERTY(EditDefaultsOnly, Category = "Overrides", Meta = (InlineEditConditionToggle))
 	uint32 bOverrideApplySimProxyStateAtRuntime : 1;
-	UPROPERTY(config, EditDefaultsOnly, Category = "Overrides", Meta = (InlineEditConditionToggle))
+	UPROPERTY(EditDefaultsOnly, Category = "Overrides", Meta = (InlineEditConditionToggle))
 	uint32 bOverrideApplySimProxyInputAtRuntime : 1;
-	UPROPERTY(config, EditDefaultsOnly, Category = "Overrides", Meta = (InlineEditConditionToggle))
+	UPROPERTY(EditDefaultsOnly, Category = "Overrides", Meta = (InlineEditConditionToggle))
 	uint32 bOverrideApplyInputDecaySimProxyInputAtRuntime : 1;
-	UPROPERTY(config, EditDefaultsOnly, Category = "Overrides", Meta = (InlineEditConditionToggle))
+	UPROPERTY(EditDefaultsOnly, Category = "Overrides", Meta = (InlineEditConditionToggle))
 	uint32 bOverrideInputDecaySimProxyInputAtRuntime : 1;
-	UPROPERTY(config, EditDefaultsOnly, Category = "Overrides", Meta = (InlineEditConditionToggle))
+	UPROPERTY(EditDefaultsOnly, Category = "Overrides", Meta = (InlineEditConditionToggle))
 	uint32 bOverrideEnableLagScalingSimProxyRuntimeInputDecay : 1;
-	UPROPERTY(config, EditDefaultsOnly, Category = "Overrides", Meta = (InlineEditConditionToggle))
+	UPROPERTY(EditDefaultsOnly, Category = "Overrides", Meta = (InlineEditConditionToggle))
 	uint32 bOverrideApplyInputDecayOverSetTime : 1;
-	UPROPERTY(config, EditDefaultsOnly, Category = "Overrides", Meta = (InlineEditConditionToggle))
+	UPROPERTY(EditDefaultsOnly, Category = "Overrides", Meta = (InlineEditConditionToggle))
 	uint32 bOverrideInputDecaySetTime : 1;
-	UPROPERTY(config, EditDefaultsOnly, Category = "Overrides", Meta = (InlineEditConditionToggle))
+	UPROPERTY(EditDefaultsOnly, Category = "Overrides", Meta = (InlineEditConditionToggle))
 	uint32 bOverrideEnableLagScalingInputDecay : 1;
-	UPROPERTY(config, EditDefaultsOnly, Category = "Overrides", Meta = (InlineEditConditionToggle))
+	UPROPERTY(EditDefaultsOnly, Category = "Overrides", Meta = (InlineEditConditionToggle))
 	uint32 bOverrideInputDecayReferenceLagMs : 1;
 
 	/* Overrides CVar : np2.Resim.RedundantInputs -- How many extra inputs to send with each unreliable network message, to account for packetloss.From owning client to server and server to owning client.
@@ -599,112 +586,108 @@ struct FNetworkPhysicsSettingsNetworkPhysicsComponent
 	const uint16 GetRedundantStates() const { return bOverrideRedundantStates ? RedundantStates : (uint16)PhysicsReplicationCVars::ResimulationCVars::RedundantStates; }
 
 	// Overrides CVar: np2.Resim.DynamicInputBufferScaling.Enabled -- Enable dynmic scaling of input buffer on the server.
-	UPROPERTY(config, EditDefaultsOnly, Category = "Overrides", Meta = (EditCondition = "bOverrideDynamicInputBufferScalingEnabled"))
+	UPROPERTY(EditDefaultsOnly, Category = "Overrides", Meta = (EditCondition = "bOverrideDynamicInputBufferScalingEnabled"))
 	bool bDynamicInputBufferScalingEnabled = PhysicsReplicationCVars::ResimulationCVars::bDynamicInputBufferScalingEnabled;
 	const bool GetDynamicInputBufferScalingEnabled() const { return bOverrideDynamicInputBufferScalingEnabled ? bDynamicInputBufferScalingEnabled : PhysicsReplicationCVars::ResimulationCVars::bDynamicInputBufferScalingEnabled; }
 
 	// Overrides CVar: np2.Resim.CompareStateToTriggerRewind -- When true, cache local FNetworkPhysicsData state in rewind history and compare the predicted state with incoming server state to trigger resimulations if they differ, comparison done through FNetworkPhysicsData::CompareData.
-	UPROPERTY(config, EditDefaultsOnly, Category = "Overrides", Meta = (EditCondition = "bOverrideCompareStateToTriggerRewind"))
+	UPROPERTY(EditDefaultsOnly, Category = "Overrides", Meta = (EditCondition = "bOverrideCompareStateToTriggerRewind"))
 	bool bCompareStateToTriggerRewind = PhysicsReplicationCVars::ResimulationCVars::bCompareStateToTriggerRewind;
 	const bool GetCompareStateToTriggerRewind(const bool bDefaultValue) const { return bOverrideCompareStateToTriggerRewind ? bCompareStateToTriggerRewind : bDefaultValue; }
-	UE_DEPRECATED(5.5, "Deprecated, use GetCompareStateToTriggerRewind(const bool bDefaultValue) instead and pass in the value that should apply if not overridden.")
-		const bool GetCompareStateToTriggerRewind() const { return bOverrideCompareStateToTriggerRewind ? bCompareStateToTriggerRewind : PhysicsReplicationCVars::ResimulationCVars::bCompareStateToTriggerRewind; }
 
 	// Overrides CVar: np2.Resim.CompareStateToTriggerRewind.IncludeSimProxies -- When true, include simulated proxies when np2.Resim.CompareStateToTriggerRewind is enabled.
-	UPROPERTY(config, EditDefaultsOnly, Category = "Overrides", Meta = (EditCondition = "bOverridebCompareStateToTriggerRewindIncludeSimProxies"))
+	UPROPERTY(EditDefaultsOnly, Category = "Overrides", Meta = (EditCondition = "bOverridebCompareStateToTriggerRewindIncludeSimProxies"))
 	bool bCompareStateToTriggerRewindIncludeSimProxies = PhysicsReplicationCVars::ResimulationCVars::bCompareStateToTriggerRewindIncludeSimProxies;
 	const bool GetCompareStateToTriggerRewindIncludeSimProxies(const bool bDefaultValue) const { return bOverridebCompareStateToTriggerRewindIncludeSimProxies ? bCompareStateToTriggerRewindIncludeSimProxies : bDefaultValue; }
 
 	// Overrides CVar: np2.Resim.CompareInputToTriggerRewind -- When true, compare local predicted FNetworkPhysicsData input with incoming server input to trigger resimulations if they differ, comparison done through FNetworkPhysicsData::CompareData.
-	UPROPERTY(config, EditDefaultsOnly, Category = "Overrides", Meta = (EditCondition = "bOverrideCompareInputToTriggerRewind"))
+	UPROPERTY(EditDefaultsOnly, Category = "Overrides", Meta = (EditCondition = "bOverrideCompareInputToTriggerRewind"))
 	bool bCompareInputToTriggerRewind = PhysicsReplicationCVars::ResimulationCVars::bCompareInputToTriggerRewind;
 	const bool GetCompareInputToTriggerRewind(const bool bDefaultValue) const { return bOverrideCompareInputToTriggerRewind ? bCompareInputToTriggerRewind : bDefaultValue; }
-	UE_DEPRECATED(5.5, "Deprecated, use GetCompareInputToTriggerRewind(const bool bDefaultValue) instead and pass in the value that should apply if not overridden.")
-		const bool GetCompareInputToTriggerRewind() const { return bOverrideCompareInputToTriggerRewind ? bCompareInputToTriggerRewind : PhysicsReplicationCVars::ResimulationCVars::bCompareInputToTriggerRewind; }
 
 	// Overrides CVar: np2.Resim.EnableUnreliableFlow -- When true, allow data to be sent unreliably. Also sends FNetworkPhysicsData not marked with FNetworkPhysicsData::bimportant unreliably over the network.
-	UPROPERTY(config, EditDefaultsOnly, Category = "Overrides", Meta = (EditCondition = "bOverrideEnableUnreliableFlow"))
+	UPROPERTY(EditDefaultsOnly, Category = "Overrides", Meta = (EditCondition = "bOverrideEnableUnreliableFlow"))
 	bool bEnableUnreliableFlow = PhysicsReplicationCVars::ResimulationCVars::bEnableUnreliableFlow;
 	const bool GetEnableUnreliableFlow() const { return bOverrideEnableUnreliableFlow ? bEnableUnreliableFlow : PhysicsReplicationCVars::ResimulationCVars::bEnableUnreliableFlow; }
 
 	// Overrides CVar: np2.Resim.EnableReliableFlow -- EXPERIMENTAL -- When true, allow data to be sent reliably. Also send FNetworkPhysicsData marked with FNetworkPhysicsData::bimportant reliably over the network.
-	UPROPERTY(config, EditDefaultsOnly, Category = "Overrides", Meta = (EditCondition = "bOverrideEnableReliableFlow"))
+	UPROPERTY(EditDefaultsOnly, Category = "Overrides", Meta = (EditCondition = "bOverrideEnableReliableFlow"))
 	bool bEnableReliableFlow = PhysicsReplicationCVars::ResimulationCVars::bEnableReliableFlow;
 	const bool GetEnableReliableFlow() const { return bOverrideEnableReliableFlow ? bEnableReliableFlow : PhysicsReplicationCVars::ResimulationCVars::bEnableReliableFlow; }
 
 	// Overrides CVar: np2.Resim.ApplyDataInsteadOfMergeData -- When true, call ApplyData for each data instead of MergeData when having to use multiple data entries in one frame.
-	UPROPERTY(config, EditDefaultsOnly, Category = "Overrides", Meta = (EditCondition = "bOverrideApplyDataInsteadOfMergeData"))
+	UPROPERTY(EditDefaultsOnly, Category = "Overrides", Meta = (EditCondition = "bOverrideApplyDataInsteadOfMergeData"))
 	bool bApplyDataInsteadOfMergeData = PhysicsReplicationCVars::ResimulationCVars::bApplyDataInsteadOfMergeData;
 	const bool GetApplyDataInsteadOfMergeData() const { return bOverrideApplyDataInsteadOfMergeData ? bApplyDataInsteadOfMergeData : PhysicsReplicationCVars::ResimulationCVars::bApplyDataInsteadOfMergeData; }
 
 	// Overrides CVar: np2.Resim.AllowInputExtrapolation -- When true and not locally controlled, allow inputs to be extrapolated from last known and if there is a gap allow interpolation between two known inputs.
-	UPROPERTY(config, EditDefaultsOnly, Category = "Overrides", Meta = (EditCondition = "bOverrideAllowInputExtrapolation"))
+	UPROPERTY(EditDefaultsOnly, Category = "Overrides", Meta = (EditCondition = "bOverrideAllowInputExtrapolation"))
 	bool bAllowInputExtrapolation = PhysicsReplicationCVars::ResimulationCVars::bAllowInputExtrapolation;
 	const bool GetAllowInputExtrapolation() const { return bOverrideAllowInputExtrapolation ? bAllowInputExtrapolation : PhysicsReplicationCVars::ResimulationCVars::bAllowInputExtrapolation; }
 
 	// Overrides CVar: np2.Resim.ValidateDataOnGameThread -- When true, perform server-side input validation through FNetworkPhysicsData::ValidateData on the Game Thread. If false, perform the call on the Physics Thread.
-	UPROPERTY(config, EditDefaultsOnly, Category = "Overrides", Meta = (EditCondition = "bOverrideValidateDataOnGameThread"))
+	UPROPERTY(EditDefaultsOnly, Category = "Overrides", Meta = (EditCondition = "bOverrideValidateDataOnGameThread"))
 	bool bValidateDataOnGameThread = PhysicsReplicationCVars::ResimulationCVars::bValidateDataOnGameThread;
 	const bool GetValidateDataOnGameThread() const { return bOverrideValidateDataOnGameThread ? bValidateDataOnGameThread : PhysicsReplicationCVars::ResimulationCVars::bValidateDataOnGameThread; }
 
 	// Overrides CVar: np2.Resim.TriggerResimOnInputReceive -- When true, a resimulation will be requested to the frame of the latest frame of received inputs this frame
-	UPROPERTY(config, EditDefaultsOnly, Category = "Overrides", Meta = (EditCondition = "bOverrideTriggerResimOnInputReceive"))
+	UPROPERTY(EditDefaultsOnly, Category = "Overrides", Meta = (EditCondition = "bOverrideTriggerResimOnInputReceive"))
 	bool bTriggerResimOnInputReceive = PhysicsReplicationCVars::ResimulationCVars::bTriggerResimOnInputReceive;
 	const bool GetTriggerResimOnInputReceive() const { return bOverrideTriggerResimOnInputReceive ? bTriggerResimOnInputReceive : PhysicsReplicationCVars::ResimulationCVars::bTriggerResimOnInputReceive; }
 
 	// Overrides CVar: np2.Resim.ApplySimProxyStateAtRuntime -- When true, call ApplyData on received states for simulated proxies at runtime.
-	UPROPERTY(config, EditDefaultsOnly, Category = "Overrides", Meta = (EditCondition = "bOverrideApplySimProxyStateAtRuntime"))
+	UPROPERTY(EditDefaultsOnly, Category = "Overrides", Meta = (EditCondition = "bOverrideApplySimProxyStateAtRuntime"))
 	bool bApplySimProxyStateAtRuntime = PhysicsReplicationCVars::ResimulationCVars::bApplySimProxyStateAtRuntime;
 	const bool GetApplySimProxyStateAtRuntime() const { return bOverrideApplySimProxyStateAtRuntime ? bApplySimProxyStateAtRuntime : PhysicsReplicationCVars::ResimulationCVars::bApplySimProxyStateAtRuntime; }
 
 	// Overrides CVar: np2.Resim.ApplySimProxyInputAtRuntime -- When true, call ApplyData on received inputs for simulated proxies at runtime.
-	UPROPERTY(config, EditDefaultsOnly, Category = "Overrides", Meta = (EditCondition = "bOverrideApplySimProxyInputAtRuntime"))
+	UPROPERTY(EditDefaultsOnly, Category = "Overrides", Meta = (EditCondition = "bOverrideApplySimProxyInputAtRuntime"))
 	bool bApplySimProxyInputAtRuntime = PhysicsReplicationCVars::ResimulationCVars::bApplySimProxyInputAtRuntime;
 	const bool GetApplySimProxyInputAtRuntime() const { return bOverrideApplySimProxyInputAtRuntime ? bApplySimProxyInputAtRuntime : PhysicsReplicationCVars::ResimulationCVars::bApplySimProxyInputAtRuntime; }
 
 	// Overrides CVar: np2.Resim.ApplyInputDecaySimProxyInputAtRuntime -- When true, apply input decay on inputs applied on simulated proxies at runtime (outside of resimulation)
-	UPROPERTY(config, EditDefaultsOnly, Category = "Overrides", Meta = (EditCondition = "bOverrideApplyInputDecaySimProxyInputAtRuntime"))
+	UPROPERTY(EditDefaultsOnly, Category = "Overrides", Meta = (EditCondition = "bOverrideApplyInputDecaySimProxyInputAtRuntime"))
 	bool bApplyInputDecaySimProxyInputAtRuntime = PhysicsReplicationCVars::ResimulationCVars::bApplyInputDecaySimProxyInputAtRuntime;
 	const bool GetApplyInputDecaySimProxyInputAtRuntime() const { return bOverrideApplyInputDecaySimProxyInputAtRuntime ? bApplyInputDecaySimProxyInputAtRuntime : PhysicsReplicationCVars::ResimulationCVars::bApplyInputDecaySimProxyInputAtRuntime; }
 
 	// Overrides CVar: np2.Resim.InputDecaySimProxyInputAtRuntime -- The amount of input decay to apply on simulated proxy inputs at runtime (outside of resim), if np2.Resim.ApplyInputDecaySimProxyInputAtRuntime is true
-	UPROPERTY(config, EditDefaultsOnly, Category = "Overrides", Meta = (EditCondition = "bOverrideInputDecaySimProxyInputAtRuntime"))
+	UPROPERTY(EditDefaultsOnly, Category = "Overrides", Meta = (EditCondition = "bOverrideInputDecaySimProxyInputAtRuntime"))
 	float InputDecaySimProxyInputAtRuntime = PhysicsReplicationCVars::ResimulationCVars::InputDecaySimProxyInputAtRuntime;
 	const float GetInputDecaySimProxyInputAtRuntime() const { return bOverrideInputDecaySimProxyInputAtRuntime ? InputDecaySimProxyInputAtRuntime : PhysicsReplicationCVars::ResimulationCVars::InputDecaySimProxyInputAtRuntime; }
 
 	// Overrides CVar: np2.Resim.EnableLagScalingSimProxyRuntimeInputDecay -- When true, scales the input decay applied on simulated proxies at runtime (outside of resimulation) as a function of measured lag (forward prediction time)
-	UPROPERTY(config, EditDefaultsOnly, Category = "Overrides", Meta = (EditCondition = "bOverrideEnableLagScalingSimProxyRuntimeInputDecay"))
+	UPROPERTY(EditDefaultsOnly, Category = "Overrides", Meta = (EditCondition = "bOverrideEnableLagScalingSimProxyRuntimeInputDecay"))
 	bool bEnableLagScalingSimProxyRuntimeInputDecay = PhysicsReplicationCVars::ResimulationCVars::bEnableLagScalingSimProxyRuntimeInputDecay;
 	const bool GetEnableLagScalingSimProxyRuntimeInputDecay() const { return bOverrideEnableLagScalingSimProxyRuntimeInputDecay ? bEnableLagScalingSimProxyRuntimeInputDecay : PhysicsReplicationCVars::ResimulationCVars::bEnableLagScalingSimProxyRuntimeInputDecay; }
 
 	// Overrides CVar: np2.Resim.ApplyInputDecayOverSetTime -- When true, apply the Input Decay Curve over a set amount of time instead of over the start of input prediction and end of resim which is variable each resimulation
-	UPROPERTY(config, EditDefaultsOnly, Category = "Overrides", Meta = (EditCondition = "bOverrideApplyInputDecayOverSetTime"))
+	UPROPERTY(EditDefaultsOnly, Category = "Overrides", Meta = (EditCondition = "bOverrideApplyInputDecayOverSetTime"))
 	bool bApplyInputDecayOverSetTime = PhysicsReplicationCVars::ResimulationCVars::bApplyInputDecayOverSetTime;
 	const bool GetApplyInputDecayOverSetTime() const { return bOverrideApplyInputDecayOverSetTime ? bApplyInputDecayOverSetTime : PhysicsReplicationCVars::ResimulationCVars::bApplyInputDecayOverSetTime; }
 
 	// Overrides CVar: np2.Resim.InputDecaySetTime -- Applied when np2.Resim.ApplyInputDecayOverSetTime is true, read there for more info. Set time to apply Input Decay Curve over while predicting inputs during resimulation
-	UPROPERTY(config, EditDefaultsOnly, Category = "Overrides", Meta = (EditCondition = "bOverrideInputDecaySetTime"))
+	UPROPERTY(EditDefaultsOnly, Category = "Overrides", Meta = (EditCondition = "bOverrideInputDecaySetTime"))
 	float InputDecaySetTime = PhysicsReplicationCVars::ResimulationCVars::InputDecaySetTime;
 	const float GetInputDecaySetTime() const { return bOverrideInputDecaySetTime ? InputDecaySetTime : PhysicsReplicationCVars::ResimulationCVars::InputDecaySetTime; }
 
 	/* Curve for input decay during resimulation if input is being reused. 
 	* XAxis = Alpha value from 0.0 to 1.0 where 0 is the start of reusing input and 1 is the last time we will reuse the input this resimulation.
 	* YAxis = The Input Decay value from 0.0 to 1.0 (as a percentage where 1.0 = 100% decay) for the given Alpha. */
-	UPROPERTY(config, EditDefaultsOnly, Category = "Overrides", Meta = (XAxisName="LerpAlpha", YAxisName="DecayValue"))
+	UPROPERTY(EditDefaultsOnly, Category = "Overrides", Meta = (XAxisName="LerpAlpha", YAxisName="DecayValue"))
 	FRuntimeFloatCurve InputDecayCurve;
 	const FRuntimeFloatCurve& GetInputDecayCurve() const { return InputDecayCurve; }
 
 	/* Overrides CVar: np2.Resim.EnableLagScalingInputDecay -- When true, scales input decay as a function of measured lag (i.e. forward prediction time). The greater the lag, the greater the decay. 
 	 * The decay will be exactly the tuned InputDecay when the measured lag is equal to the reference lag (InputDecayReferenceLagMs)
 	 **/
-	UPROPERTY(config, EditDefaultsOnly, Category = "Overrides", Meta = (EditCondition = "bOverrideEnableLagScalingInputDecay"))
+	UPROPERTY(EditDefaultsOnly, Category = "Overrides", Meta = (EditCondition = "bOverrideEnableLagScalingInputDecay"))
 	bool bEnableLagScalingInputDecay = PhysicsReplicationCVars::ResimulationCVars::bEnableLagScalingInputDecay;
 	const bool GetEnableLagScalingInputDecay() const { return bOverrideEnableLagScalingInputDecay ? bEnableLagScalingInputDecay : PhysicsReplicationCVars::ResimulationCVars::bEnableLagScalingInputDecay; }
 
 	/** Overrides CVar: np2.Resim.InputDecayReferenceLagMs -- Applied when np2.Resim.EnableLagScalingInputDecay is true. Reference lag at which the decay will be exactly InputDecay. Linearly interpolated below, reaching 0 decay at 0 lag.
 	*   Tries to maintain behavior above that lag, by decaying input more to travel the same distance over greater time for the same input.
 	*/
-	UPROPERTY(config, EditDefaultsOnly, Category = "Overrides", Meta = (EditCondition = "bOverrideInputDecayReferenceLagMs"))
+	UPROPERTY(EditDefaultsOnly, Category = "Overrides", Meta = (EditCondition = "bOverrideInputDecayReferenceLagMs"))
 	float InputDecayReferenceLagMs = PhysicsReplicationCVars::ResimulationCVars::InputDecayReferenceLagMs;
 	const float GetInputDecayReferenceLagMs() const { return bOverrideInputDecayReferenceLagMs ? InputDecayReferenceLagMs : PhysicsReplicationCVars::ResimulationCVars::InputDecayReferenceLagMs; }
 };
