Fix compile error with older MSVC 14.38 toolchain - error C4706: assignment within conditional expression

#rnx

#Synced-CL 49847917

diff --git a/Engine/Source/Runtime/Engine/Private/Slate/SceneViewport.cpp b/Engine/Source/Runtime/Engine/Private/Slate/SceneViewport.cpp
--- a/Engine/Source/Runtime/Engine/Private/Slate/SceneViewport.cpp
+++ b/Engine/Source/Runtime/Engine/Private/Slate/SceneViewport.cpp
@@ -1473,7 +1473,8 @@ static void EnsureValidFullscreenResolution(const FMonitorInfo& Monitor, uint32&
 	bool bFound = false;
 	for (FScreenResolutionRHI ScreenResolution : ResArray)
 	{
-		if ((bFound = ScreenResolution.Width == NewWindowSizeX && ScreenResolution.Height == NewWindowSizeY))
+		bFound = ScreenResolution.Width == NewWindowSizeX && ScreenResolution.Height == NewWindowSizeY;
+		if (bFound)
 		{
 			break;
 		}
