Entity Net Physics Replication (mesh_component)

#Synced-CL 49803900

#Synced-CL 49803900

#Synced-CL 49803900

diff --git a/Engine/Source/Runtime/Engine/Classes/Engine/ReplicatedState.h b/Engine/Source/Runtime/Engine/Classes/Engine/ReplicatedState.h
--- a/Engine/Source/Runtime/Engine/Classes/Engine/ReplicatedState.h
+++ b/Engine/Source/Runtime/Engine/Classes/Engine/ReplicatedState.h
@@ -188,6 +188,19 @@ struct FRepMovement
 		ServerFrame = InServerFrame;
 	}
 
+	void FillFrom(const FRigidBodyState& RBState, const UObject* const Object = nullptr, int32 InServerFrame = 0)
+	{
+		//TODO: Get Rebase stuff working with Entities
+		//Location = RebaseOntoZeroOrigin(RBState.Position, Actor);
+		Location = RBState.Position;
+		Rotation = RBState.Quaternion.Rotator();
+		LinearVelocity = RBState.LinVel;
+		AngularVelocity = RBState.AngVel;
+		bSimulatedPhysicSleep = (RBState.Flags & ERigidBodyFlags::Sleeping) != 0;
+		bRepPhysics = true;
+		ServerFrame = InServerFrame;
+	}
+
 	void CopyTo(FRigidBodyState& RBState, const AActor* const Actor = nullptr) const
 	{
 		RBState.Position = RebaseOntoLocalOrigin(Location, Actor);
@@ -200,6 +213,18 @@ struct FRepMovement
 			| (decltype(FRigidBodyState::Flags))(bRepPhysics ? ERigidBodyFlags::RepPhysics : ERigidBodyFlags::None);
 	}
 
+	void CopyTo(FRigidBodyState& RBState, const UObject* const Object = nullptr) const
+	{
+		RBState.Position = Location;
+		RBState.Quaternion = Rotation.Quaternion();
+		RBState.LinVel = LinearVelocity;
+		RBState.AngVel = AngularVelocity;
+		RBState.Flags =
+			(decltype(FRigidBodyState::Flags))(bSimulatedPhysicSleep ? ERigidBodyFlags::Sleeping : ERigidBodyFlags::None)
+			| ERigidBodyFlags::NeedsUpdate
+			| (decltype(FRigidBodyState::Flags))(bRepPhysics ? ERigidBodyFlags::RepPhysics : ERigidBodyFlags::None);
+	}
+
 	bool operator==(const FRepMovement& Other) const
 	{
 		if ( LinearVelocity != Other.LinearVelocity )
diff --git a/Engine/Source/Runtime/Engine/Private/PhysicsEngine/Experimental/PhysScene_Chaos.cpp b/Engine/Source/Runtime/Engine/Private/PhysicsEngine/Experimental/PhysScene_Chaos.cpp
--- a/Engine/Source/Runtime/Engine/Private/PhysicsEngine/Experimental/PhysScene_Chaos.cpp
+++ b/Engine/Source/Runtime/Engine/Private/PhysicsEngine/Experimental/PhysScene_Chaos.cpp
@@ -1440,7 +1440,22 @@ const FRigidBodyState* FPhysScene_Chaos::GetStateFromReplicationCache(UPrimitive
 	// Create the physics replication cache if not already created
 	CreatePhysicsReplicationCache();
 
-	return PhysicsReplicationCache->GetStateFromReplicationCache(RootComponent, ServerFrame);
+	if (Chaos::FPhysicsObjectHandle PhysicsObject = RootComponent->GetPhysicsObjectByName(NAME_None))
+	{
+		return GetStateFromReplicationCache(PhysicsObject, ServerFrame);
+	}
+
+	ServerFrame = -1;
+	return nullptr;
+}
+
+
+const FRigidBodyState* FPhysScene_Chaos::GetStateFromReplicationCache(Chaos::FPhysicsObjectHandle PhysicsObject, int& ServerFrame)
+{
+	// Create the physics replication cache if not already created
+	CreatePhysicsReplicationCache();
+
+	return PhysicsReplicationCache->GetStateFromReplicationCache(PhysicsObject, ServerFrame);
 }
 
 void FPhysScene_Chaos::RegisterForReplicationCache(UPrimitiveComponent* RootComponent)
diff --git a/Engine/Source/Runtime/Engine/Private/PhysicsEngine/PhysicsReplication.cpp b/Engine/Source/Runtime/Engine/Private/PhysicsEngine/PhysicsReplication.cpp
--- a/Engine/Source/Runtime/Engine/Private/PhysicsEngine/PhysicsReplication.cpp
+++ b/Engine/Source/Runtime/Engine/Private/PhysicsEngine/PhysicsReplication.cpp
@@ -409,6 +409,17 @@ void FPhysicsReplication::RemoveReplicatedTarget(UPrimitiveComponent* Component)
 	
 	// Remove from FPhysicsObject flow
 	Chaos::FConstPhysicsObjectHandle PhysicsObject = Component->GetPhysicsObjectByName(NAME_None);
+	RemoveReplicatedTarget(PhysicsObject);
+}
+
+void FPhysicsReplication::RemoveReplicatedTarget(Chaos::FConstPhysicsObjectHandle PhysicsObject)
+{
+	// Remove from FPhysicsObject flow
+	if (!PhysicsObject)
+	{
+		return;
+	}
+
 	FReplicatedPhysicsTarget Target(PhysicsObject); // This creates a new but empty target and when it tries to update the current target in the async flow it will remove it from replication since it's empty.
 	ReplicatedTargetsQueue.Add(Target);
 }
diff --git a/Engine/Source/Runtime/Engine/Private/PhysicsEngine/PhysicsReplicationCache.cpp b/Engine/Source/Runtime/Engine/Private/PhysicsEngine/PhysicsReplicationCache.cpp
--- a/Engine/Source/Runtime/Engine/Private/PhysicsEngine/PhysicsReplicationCache.cpp
+++ b/Engine/Source/Runtime/Engine/Private/PhysicsEngine/PhysicsReplicationCache.cpp
@@ -45,6 +45,22 @@ FPhysicsReplicationCache::~FPhysicsReplicationCache()
 }
 
 const FRigidBodyState* FPhysicsReplicationCache::GetStateFromReplicationCache(UPrimitiveComponent* RootComponent, int32& OutSolverFrame)
+{
+	if (!RootComponent)
+	{
+		OutSolverFrame = 0;
+		return nullptr;
+	}
+
+	if (Chaos::FPhysicsObjectHandle PhysicsObject = RootComponent->GetPhysicsObjectByName(NAME_None))
+	{
+		return GetStateFromReplicationCache(PhysicsObject, OutSolverFrame);
+	}
+	
+	return nullptr;
+}
+
+const FRigidBodyState* FPhysicsReplicationCache::GetStateFromReplicationCache(Chaos::FPhysicsObjectHandle PhysicsObject, int32& OutSolverFrame)
 {
 	FRigidBodyState* ReplicationState = nullptr;
 	OutSolverFrame = 0;
@@ -58,7 +74,7 @@ const FRigidBodyState* FPhysicsReplicationCache::GetStateFromReplicationCache(UP
 	ProcessAsyncOutput();
 
 	// Get the state for specified component or register for replication cache
-	if (Chaos::FPhysicsObjectHandle PhysicsObject = RootComponent->GetPhysicsObjectByName(NAME_None))
+	if (PhysicsObject)
 	{
 		if (FPhysicsReplicationCacheAsyncInput* AsyncInput = AsyncPhysicsReplicationCache->GetProducerInputData_External())
 		{
diff --git a/Engine/Source/Runtime/Engine/Public/Physics/Experimental/PhysScene_Chaos.h b/Engine/Source/Runtime/Engine/Public/Physics/Experimental/PhysScene_Chaos.h
--- a/Engine/Source/Runtime/Engine/Public/Physics/Experimental/PhysScene_Chaos.h
+++ b/Engine/Source/Runtime/Engine/Public/Physics/Experimental/PhysScene_Chaos.h
@@ -290,6 +290,9 @@ public:
 	/** Get cached state for replication, if no state is cached RegisterForReplicationCache() is called */
 	ENGINE_API const FRigidBodyState* GetStateFromReplicationCache(UPrimitiveComponent* RootComponent, int& ServerFrame);
 	
+	/** Get cached state for replication, if no state is cached RegisterForReplicationCache() is called */
+	ENGINE_API const FRigidBodyState* GetStateFromReplicationCache(Chaos::FPhysicsObjectHandle PhysicsObject, int& ServerFrame);
+
 	/** Register a component for physics replication state caching, the component will deregister automatically if cache is not accessed within timelimit set by CVar: np2.ReplicationCache.LingerForNSeconds */
 	ENGINE_API void RegisterForReplicationCache(UPrimitiveComponent* RootComponent);
 
diff --git a/Engine/Source/Runtime/Engine/Public/Physics/PhysicsReplicationCache.h b/Engine/Source/Runtime/Engine/Public/Physics/PhysicsReplicationCache.h
--- a/Engine/Source/Runtime/Engine/Public/Physics/PhysicsReplicationCache.h
+++ b/Engine/Source/Runtime/Engine/Public/Physics/PhysicsReplicationCache.h
@@ -22,9 +22,12 @@ public:
 	
 	virtual ~FPhysicsReplicationCache();
 
-	/** Get cached state for replication, if no state is cached RegisterForReplicationCache() is called */
+	/** Get cached state for replication */
 	const FRigidBodyState* GetStateFromReplicationCache(UPrimitiveComponent* RootComponent, int32& OutSolverFrame);
 
+	/** Get cached state for replication */
+	const FRigidBodyState* GetStateFromReplicationCache(Chaos::FPhysicsObjectHandle PhysicsObject, int32& OutSolverFrame);
+
 	/** Register a component for physics replication state caching, the component will unregister automatically if cache is not accessed within time limit set by CVar: np2.ReplicationCache.LingerForNSeconds */
 	void RegisterForReplicationCache(UPrimitiveComponent* RootComponent);
 	
diff --git a/Engine/Source/Runtime/Engine/Public/PhysicsReplication.h b/Engine/Source/Runtime/Engine/Public/PhysicsReplication.h
--- a/Engine/Source/Runtime/Engine/Public/PhysicsReplication.h
+++ b/Engine/Source/Runtime/Engine/Public/PhysicsReplication.h
@@ -326,13 +326,15 @@ public:
 
 	/** Sets the latest replicated target for a body instance */
 	ENGINE_API virtual void SetReplicatedTarget(UPrimitiveComponent* Component, FName BoneName, const FRigidBodyState& ReplicatedTarget, int32 ServerFrame) override;
-private:
+
 	ENGINE_API void SetReplicatedTarget(Chaos::FConstPhysicsObjectHandle PhysicsObject, const FRigidBodyState& ReplicatedTarget, int32 ServerFrame, EPhysicsReplicationMode ReplicationMode = EPhysicsReplicationMode::Default);
 
-public:
 	/** Remove the replicated target*/
 	ENGINE_API virtual void RemoveReplicatedTarget(UPrimitiveComponent* Component) override;
 
+	/** Remove the replicated target*/
+	ENGINE_API void RemoveReplicatedTarget(Chaos::FConstPhysicsObjectHandle PhysicsObject);
+
 	virtual int32 GetNetworkPhysicsTickOffset() const override { return NetworkPhysicsTickOffset; }
 
 protected:
