Remove IsMaterialAttributes change from CL49678800 because it corrupts existing materials which have issues(an invalid expression input in this case).
Fix crash when opening Material Editor, caused by the same invalid expression input which contained a FunctionCall with no Function.
This invalid expression gets cleared up after a Material graph is created based on the material and LinkMaterialExpressionsFromGraph gets called.

Note that this still allows the validation in editor to happen so different types won't get set as input to the same StaticSwitch.

[RN] minor | Rendering Architecture

#Synced-CL 49844456

#Synced-CL 49844456

diff --git a/Engine/Source/Runtime/Engine/Private/Materials/MaterialExpressions.cpp b/Engine/Source/Runtime/Engine/Private/Materials/MaterialExpressions.cpp
--- a/Engine/Source/Runtime/Engine/Private/Materials/MaterialExpressions.cpp
+++ b/Engine/Source/Runtime/Engine/Private/Materials/MaterialExpressions.cpp
@@ -1929,7 +1929,13 @@ uint32 UMaterialExpression::GetOutputType(int32 OutputIndex)
 		return CompositeResultOutput;
 	}
 
-	FExpressionOutput& Output = GetOutputs()[OutputIndex];
+	TArray<FExpressionOutput>& CurrentOutputs = GetOutputs();
+	if (!CurrentOutputs.IsValidIndex(OutputIndex))
+	{
+		return MCT_Unknown;
+	}
+
+	FExpressionOutput& Output = CurrentOutputs[OutputIndex];
 	if (Output.Mask)
 	{
 		int32 MaskChannelCount = (Output.MaskR ? 1 : 0)
@@ -8672,10 +8678,10 @@ bool UMaterialExpressionStaticSwitchParameter::IsResultMaterialAttributes(int32
 {
 	check(OutputIndex == 0);
 
-	// Both A and B must be connected in a parameter. 
-	// Because of this, MA output will be valid only if both are MAs
-	return ((A.GetTracedInput().Expression && A.Expression->IsResultMaterialAttributes(A.OutputIndex)) &&
-		(B.GetTracedInput().Expression && B.Expression->IsResultMaterialAttributes(B.OutputIndex)));
+	const bool bAIsMaterialAttributes = (A.GetTracedInput().Expression && A.Expression->IsResultMaterialAttributes(A.OutputIndex));
+	const bool bBIsMaterialAttributes = (B.GetTracedInput().Expression && B.Expression->IsResultMaterialAttributes(B.OutputIndex));
+
+	return bAIsMaterialAttributes || bBIsMaterialAttributes;
 }
 
 EMaterialValueType GetRealInputValueType(const FExpressionInput& InExpression)
@@ -8717,6 +8723,7 @@ EMaterialValueType UMaterialExpressionStaticSwitchParameter::GetInputValueType(i
 	return MCT_Unknown;
 }
 
+
 EMaterialValueType UMaterialExpressionStaticSwitchParameter::GetOutputValueType(int32 OutputIndex)
 {
 	return GetCompositeOutput(OutputIndex);
