Use HDR defaults from HGIG Category 2

#Synced-CL 49797776

diff --git a/Engine/Source/Runtime/D3D12RHI/Private/Windows/WindowsD3D12Device.cpp b/Engine/Source/Runtime/D3D12RHI/Private/Windows/WindowsD3D12Device.cpp
--- a/Engine/Source/Runtime/D3D12RHI/Private/Windows/WindowsD3D12Device.cpp
+++ b/Engine/Source/Runtime/D3D12RHI/Private/Windows/WindowsD3D12Device.cpp
@@ -647,10 +647,10 @@ bool FD3D12DynamicRHI::SetupDisplayHDRMetaData()
 				DisplayInformation.MinimumLuminanceInNits = OutputDesc.MinLuminance;
 				DisplayInformation.MaximumLuminanceInNits = OutputDesc.MaxLuminance ?
 					OutputDesc.MaxLuminance :
-					(bDisplaySupportsHDROutput ? 1499.f : 100.f);
+					(bDisplaySupportsHDROutput ? 1000.f : 100.f);
 				DisplayInformation.MaximumFullFrameLuminanceInNits = OutputDesc.MaxFullFrameLuminance ?
 					OutputDesc.MaxFullFrameLuminance :
-					(bDisplaySupportsHDROutput ? 799.f : 100.f);
+					(bDisplaySupportsHDROutput ? 600.f : 100.f);
 				const bool bIsGamutValid = OutputDesc.RedPrimary[0] && OutputDesc.RedPrimary[1] &&
 					OutputDesc.GreenPrimary[0] && OutputDesc.GreenPrimary[1] &&
 					OutputDesc.BluePrimary[0] && OutputDesc.BluePrimary[1] &&
diff --git a/Engine/Source/Runtime/Engine/Private/Slate/SceneViewport.cpp b/Engine/Source/Runtime/Engine/Private/Slate/SceneViewport.cpp
--- a/Engine/Source/Runtime/Engine/Private/Slate/SceneViewport.cpp
+++ b/Engine/Source/Runtime/Engine/Private/Slate/SceneViewport.cpp
@@ -420,9 +420,9 @@ static void ComputeSceneViewportHDRMetaData(FHDRMetaData& OutHDRMetaData, const
 	if (StereoRenderTargetManager && StereoRenderTargetManager->HDRGetMetaDataForStereo(OutHDRMetaData.DisplayOutputFormat, OutHDRMetaData.DisplayColorGamut, OutHDRMetaData.bHDRSupported))
 	{
 		// TODO: Get the real values for stereo.
-		OutHDRMetaData.MinimumLuminanceInNits = 0.f;
-		OutHDRMetaData.MaximumLuminanceInNits = OutHDRMetaData.bHDRSupported ? 1499.f : 100.f;
-		OutHDRMetaData.MaximumFullFrameLuminanceInNits = OutHDRMetaData.bHDRSupported ? 799.f : 100.f;
+		OutHDRMetaData.MinimumLuminanceInNits = OutHDRMetaData.bHDRSupported ? .1f : 0.f;
+		OutHDRMetaData.MaximumLuminanceInNits = OutHDRMetaData.bHDRSupported ? 1000.f : 100.f;
+		OutHDRMetaData.MaximumFullFrameLuminanceInNits = OutHDRMetaData.bHDRSupported ? 600.f : 100.f;
 		OutHDRMetaData.LimitingColorSpace = OutHDRMetaData.bHDRSupported ? UE::Color::FColorSpace::GetP3D65() : UE::Color::FColorSpace::GetSRGB();
 		return;
 	}
diff --git a/Engine/Source/Runtime/Windows/D3D11RHI/Private/Windows/WindowsD3D11Device.cpp b/Engine/Source/Runtime/Windows/D3D11RHI/Private/Windows/WindowsD3D11Device.cpp
--- a/Engine/Source/Runtime/Windows/D3D11RHI/Private/Windows/WindowsD3D11Device.cpp
+++ b/Engine/Source/Runtime/Windows/D3D11RHI/Private/Windows/WindowsD3D11Device.cpp
@@ -566,10 +566,10 @@ bool FD3D11DynamicRHI::SetupDisplayHDRMetaData()
 				DisplayInformation.MinimumLuminanceInNits = OutputDesc.MinLuminance;
 				DisplayInformation.MaximumLuminanceInNits = OutputDesc.MaxLuminance ?
 					OutputDesc.MaxLuminance :
-					(bDisplaySupportsHDROutput ? 1499.f : 100.f);
+					(bDisplaySupportsHDROutput ? 1000.f : 100.f);
 				DisplayInformation.MaximumFullFrameLuminanceInNits = OutputDesc.MaxFullFrameLuminance ?
 					OutputDesc.MaxFullFrameLuminance :
-					(bDisplaySupportsHDROutput ? 799.f : 100.f);
+					(bDisplaySupportsHDROutput ? 600.f : 100.f);
 				const bool bIsGamutValid = OutputDesc.RedPrimary[0] && OutputDesc.RedPrimary[1] &&
 					OutputDesc.GreenPrimary[0] && OutputDesc.GreenPrimary[1] &&
 					OutputDesc.BluePrimary[0] && OutputDesc.BluePrimary[1] &&
