[Deprecation] Remove deprecated AActor::EditorKeyPressed function

diff --git a/Engine/Source/Editor/UnrealEd/Private/EditorModeManager.cpp b/Engine/Source/Editor/UnrealEd/Private/EditorModeManager.cpp
--- a/Engine/Source/Editor/UnrealEd/Private/EditorModeManager.cpp
+++ b/Engine/Source/Editor/UnrealEd/Private/EditorModeManager.cpp
@@ -1532,17 +1532,6 @@ bool FEditorModeTools::InputKey(FEditorViewportClient* InViewportClient, FViewpo
 		return true;
 		});
 
-PRAGMA_DISABLE_DEPRECATION_WARNINGS // Begin AActor::EditorKeyPressed
-	// Finally, pass input to selected actors if nothing else handled the input (Deprecated in 5.4)
-	if (!bHandled)
-	{
-		GetEditorSelectionSet()->ForEachSelectedObject<AActor>([Key, Event](AActor* ActorPtr)
-			{
-				ActorPtr->EditorKeyPressed(Key, Event);
-				return true;
-			});
-	}
-PRAGMA_ENABLE_DEPRECATION_WARNINGS // End AActor::EditorKeyPressed
 	return bHandled;
 }
 
diff --git a/Engine/Source/Runtime/Engine/Classes/GameFramework/Actor.h b/Engine/Source/Runtime/Engine/Classes/GameFramework/Actor.h
--- a/Engine/Source/Runtime/Engine/Classes/GameFramework/Actor.h
+++ b/Engine/Source/Runtime/Engine/Classes/GameFramework/Actor.h
@@ -2710,10 +2710,6 @@ public:
 	/** Called before editor paste, true allow import */
 	virtual bool ShouldImport(FStringView ActorPropString, bool IsMovingLevel) { return true; }
 
-	/** Called by InputKey when an unhandled key is pressed with a selected actor */
-	UE_DEPRECATED(5.4, "Please use UI Commands and process your custom actors as necessary by extending the level editor modules' global command list.")
-	virtual void EditorKeyPressed(FKey Key, EInputEvent Event) {}
-
 	/** Called by ReplaceSelectedActors to allow a new actor to copy properties from an old actor when it is replaced */
 	ENGINE_API virtual void EditorReplacedActor(AActor* OldActor);
 
