Fixing assertion in FVisualLogger::CleanupRedirects.
Switching to using the iterator's RemoveCurrent function instead of removing by key

#Synced-CL 49844694

diff --git a/Engine/Source/Runtime/Engine/Private/VisualLogger/VisualLogger.cpp b/Engine/Source/Runtime/Engine/Private/VisualLogger/VisualLogger.cpp
--- a/Engine/Source/Runtime/Engine/Private/VisualLogger/VisualLogger.cpp
+++ b/Engine/Source/Runtime/Engine/Private/VisualLogger/VisualLogger.cpp
@@ -1128,8 +1128,7 @@ void FVisualLogger::CleanupRedirects()
 	{
 		if (!It.Value().IsValid())
 		{
-			FObjectKey Obj = It.Key();
-			ChildToOwnerMap.Remove(Obj);
+			It.RemoveCurrent();
 		}
 	}
 }
