[RuntimeHash] Introduce new virtual methods in UWorldPartitionRuntimeHash that are required for CustomHLOD and StandaloneHLOD functionality, but were previously only present in UWorldPartitionRuntimeHashSet.
This allows the removal of Cast<UWorldPartitionRuntimeHashSet> in code related to CustomHLOD/StandaloneHLOD.
No implementation is added to SpatialHash at this time, as it is not supported by CustomHLOD/StandaloneHLOD functionality.

#Synced-CL 49651271

diff --git a/Engine/Source/Editor/UnrealEd/Private/WorldPartition/WorldPartitionHLODsBuilder.cpp b/Engine/Source/Editor/UnrealEd/Private/WorldPartition/WorldPartitionHLODsBuilder.cpp
--- a/Engine/Source/Editor/UnrealEd/Private/WorldPartition/WorldPartitionHLODsBuilder.cpp
+++ b/Engine/Source/Editor/UnrealEd/Private/WorldPartition/WorldPartitionHLODsBuilder.cpp
@@ -51,7 +51,7 @@
 #include "WorldPartition/HLOD/IWorldPartitionHLODUtilitiesModule.h"
 #include "WorldPartition/HLOD/StandaloneHLODSubsystem.h"
 #include "WorldPartition/IWorldPartitionEditorModule.h"
-#include "WorldPartition/RuntimeHashSet/WorldPartitionRuntimeHashSet.h"
+#include "WorldPartition/WorldPartitionRuntimeHash.h"
 
 #include "NavSystemConfigOverride.h"
 
@@ -355,23 +355,20 @@ bool UWorldPartitionHLODsBuilder::ShouldProcessAdditionalWorlds(UWorld* InWorld,
 	}
 
 	UWorldPartition* WP = InWorld->GetWorldPartition();
-	if (WP && WP->HasStandaloneHLOD())
+	if (WP && WP->HasStandaloneHLOD() && WP->RuntimeHash)
 	{
-		if (const UWorldPartitionRuntimeHashSet* WorldPartitionRuntimeHashSet = Cast<UWorldPartitionRuntimeHashSet>(WP->RuntimeHash))
-		{
-			FString FolderPath, PackagePrefix;
-			UWorldPartitionStandaloneHLODSubsystem::GetStandaloneHLODFolderPathAndPackagePrefix(InWorld->GetPackage()->GetName(), FolderPath, PackagePrefix);
-
-			int32 HLODDepth = WorldPartitionRuntimeHashSet->ComputeHLODHierarchyDepth();
+		FString FolderPath, PackagePrefix;
+		UWorldPartitionStandaloneHLODSubsystem::GetStandaloneHLODFolderPathAndPackagePrefix(InWorld->GetPackage()->GetName(), FolderPath, PackagePrefix);
 
-			for (int32 HLODSetupIndex = 0; HLODSetupIndex < HLODDepth; HLODSetupIndex++)
-			{
-				const FString HLODLevelPackageName = FString::Printf(TEXT("%s/%s%d"), *FolderPath, *PackagePrefix, HLODSetupIndex);
-				OutPackageNames.Add(HLODLevelPackageName);
-			}
+		int32 HLODDepth = WP->RuntimeHash->ComputeHLODHierarchyDepth();
 
-			return true;
+		for (int32 HLODSetupIndex = 0; HLODSetupIndex < HLODDepth; HLODSetupIndex++)
+		{
+			const FString HLODLevelPackageName = FString::Printf(TEXT("%s/%s%d"), *FolderPath, *PackagePrefix, HLODSetupIndex);
+			OutPackageNames.Add(HLODLevelPackageName);
 		}
+
+		return true;
 	}
 	return false;
 }
diff --git a/Engine/Source/Runtime/Engine/Private/WorldPartition/HLOD/StandaloneHLODSubsystem.cpp b/Engine/Source/Runtime/Engine/Private/WorldPartition/HLOD/StandaloneHLODSubsystem.cpp
--- a/Engine/Source/Runtime/Engine/Private/WorldPartition/HLOD/StandaloneHLODSubsystem.cpp
+++ b/Engine/Source/Runtime/Engine/Private/WorldPartition/HLOD/StandaloneHLODSubsystem.cpp
@@ -13,7 +13,7 @@
 #include "WorldPartition/HLOD/StandaloneHLODActor.h"
 #include "WorldPartition/LevelInstance/LevelInstanceActorDesc.h"
 #include "WorldPartition/WorldPartition.h"
-#include "WorldPartition/RuntimeHashSet/WorldPartitionRuntimeHashSet.h"
+#include "WorldPartition/WorldPartitionRuntimeHash.h"
 #include "WorldPartition/WorldPartitionSubsystem.h"
 
 DEFINE_LOG_CATEGORY_STATIC(LogStandaloneHLODSubsystem, Log, All);
@@ -297,10 +297,10 @@ void UWorldPartitionStandaloneHLODSubsystem::UpdateStandaloneHLODActors(UWorldPa
 			SpawnParams.Name = FName(ActorLabel);
 			SpawnParams.NameMode = FActorSpawnParameters::ESpawnActorNameMode::Requested;
 
-			UWorldPartitionRuntimeHashSet* WorldHash = GetWorld()->GetWorldPartition() ? Cast<UWorldPartitionRuntimeHashSet>(GetWorld()->GetWorldPartition()->RuntimeHash) : nullptr;
+			UWorldPartitionRuntimeHash* WorldHash = GetWorld()->GetWorldPartition() ? GetWorld()->GetWorldPartition()->RuntimeHash : nullptr;
 			if (!WorldHash)
 			{
-				UE_LOG(LogStandaloneHLODSubsystem, Log, TEXT("RuntimeHashSet not found in World %s"), *GetWorld()->GetName());
+				UE_LOG(LogStandaloneHLODSubsystem, Log, TEXT("RuntimeHash not found in World %s"), *GetWorld()->GetName());
 				return;
 			}
 
diff --git a/Engine/Source/Runtime/Engine/Private/WorldPartition/RuntimeHashSet/WorldPartitionRuntimeHashSet.cpp b/Engine/Source/Runtime/Engine/Private/WorldPartition/RuntimeHashSet/WorldPartitionRuntimeHashSet.cpp
--- a/Engine/Source/Runtime/Engine/Private/WorldPartition/RuntimeHashSet/WorldPartitionRuntimeHashSet.cpp
+++ b/Engine/Source/Runtime/Engine/Private/WorldPartition/RuntimeHashSet/WorldPartitionRuntimeHashSet.cpp
@@ -1243,6 +1243,45 @@ int32 UWorldPartitionRuntimeHashSet::ComputeHLODHierarchyDepth() const
 	return MaxHLODDepth;
 }
 
+FName UWorldPartitionRuntimeHashSet::ResolveRuntimeGridNameForHLODLayer(FName RuntimeGrid, const FSoftObjectPath& HLODLayerPath) const
+{
+	if (const URuntimePartition* RuntimePartition = ResolveRuntimePartitionForHLODLayer(RuntimeGrid, HLODLayerPath))
+	{
+		if (!RuntimePartition->SupportsHLODs())
+		{
+			return NAME_None;
+		}
+
+		return FName(RuntimeGrid.ToString() + TEXT(":") + RuntimePartition->Name.ToString());
+	}
+	return NAME_None;
+}
+
+bool UWorldPartitionRuntimeHashSet::CopyGridSetupFrom(const UWorldPartitionRuntimeHash* Other)
+{
+	if (const UWorldPartitionRuntimeHashSet* OtherHashSet = Cast<UWorldPartitionRuntimeHashSet>(Other))
+	{
+		RuntimePartitions.Empty();
+
+		for (const FRuntimePartitionDesc& SourcePartitionDesc : OtherHashSet->RuntimePartitions)
+		{
+			FRuntimePartitionDesc& PartitionDesc = RuntimePartitions.AddDefaulted_GetRef();
+			PartitionDesc.Name = SourcePartitionDesc.Name;
+			PartitionDesc.MainLayer = DuplicateObject<URuntimePartition>(SourcePartitionDesc.MainLayer, this);
+
+			for (const FRuntimePartitionHLODSetup& SourceHLODSetup : SourcePartitionDesc.HLODSetups)
+			{
+				FRuntimePartitionHLODSetup& HLODSetup = PartitionDesc.HLODSetups.AddDefaulted_GetRef();
+				HLODSetup.Name = SourceHLODSetup.Name;
+				HLODSetup.HLODLayers = SourceHLODSetup.HLODLayers;
+				HLODSetup.PartitionLayer = DuplicateObject<URuntimePartition>(SourceHLODSetup.PartitionLayer, this);
+			}
+		}
+		return true;
+	}
+	return false;
+}
+
 UWorldPartitionRuntimeHashSet::FCellUniqueId UWorldPartitionRuntimeHashSet::GetCellUniqueId(const URuntimePartition::FCellDescInstance& InCellDescInstance) const
 {
 	FCellUniqueId CellUniqueId;
diff --git a/Engine/Source/Runtime/Engine/Private/WorldPartition/RuntimeHashSet/WorldPartitionRuntimeHashSetHLODGeneration.cpp b/Engine/Source/Runtime/Engine/Private/WorldPartition/RuntimeHashSet/WorldPartitionRuntimeHashSetHLODGeneration.cpp
--- a/Engine/Source/Runtime/Engine/Private/WorldPartition/RuntimeHashSet/WorldPartitionRuntimeHashSetHLODGeneration.cpp
+++ b/Engine/Source/Runtime/Engine/Private/WorldPartition/RuntimeHashSet/WorldPartitionRuntimeHashSetHLODGeneration.cpp
@@ -256,28 +256,9 @@ bool UWorldPartitionRuntimeHashSet::SetupHLODActors(const IStreamingGenerationCo
 				StandaloneHLODWorldPartition->SetDefaultHLODLayer(nullptr);
 
 				// Create RuntimePartitions setup in the Standalone HLOD world based on the source world
-				if (UWorldPartitionRuntimeHashSet* SourceWorldHash = Cast<UWorldPartitionRuntimeHashSet>(WorldPartition->RuntimeHash))
+				if (!StandaloneHLODWorldPartition->RuntimeHash->CopyGridSetupFrom(WorldPartition->RuntimeHash))
 				{
-					if (UWorldPartitionRuntimeHashSet* StandaloneHLODWorldHash = Cast<UWorldPartitionRuntimeHashSet>(StandaloneHLODWorld->GetWorldPartition()->RuntimeHash))
-					{
-						TArray<FRuntimePartitionDesc>& StandaloneHLODWorldRuntimePartitions = StandaloneHLODWorldHash->RuntimePartitions;
-						StandaloneHLODWorldRuntimePartitions.Empty();
-						
-						for (FRuntimePartitionDesc& SourcePartitionDesc : SourceWorldHash->RuntimePartitions)
-						{
-							FRuntimePartitionDesc& PartitionDesc = StandaloneHLODWorldRuntimePartitions.AddDefaulted_GetRef();
-							PartitionDesc.Name = SourcePartitionDesc.Name;
-							PartitionDesc.MainLayer = DuplicateObject<URuntimePartition>(SourcePartitionDesc.MainLayer, StandaloneHLODWorldHash);
-
-							for (FRuntimePartitionHLODSetup& SourceHLODSetup : SourcePartitionDesc.HLODSetups)
-							{
-								FRuntimePartitionHLODSetup& HLODSetup = PartitionDesc.HLODSetups.AddDefaulted_GetRef();
-								HLODSetup.Name = SourceHLODSetup.Name;
-								HLODSetup.HLODLayers = SourceHLODSetup.HLODLayers;
-								HLODSetup.PartitionLayer = DuplicateObject<URuntimePartition>(SourceHLODSetup.PartitionLayer, StandaloneHLODWorldHash);
-							}
-						}
-					}
+					UE_LOG(LogWorldPartition, Warning, TEXT("Standalone HLOD World Init - failed to copy grid setup from the source world"));
 				}
 
 				// Save World Data Layers Package
diff --git a/Engine/Source/Runtime/Engine/Private/WorldPartition/WorldPartitionStreamingGeneration.cpp b/Engine/Source/Runtime/Engine/Private/WorldPartition/WorldPartitionStreamingGeneration.cpp
--- a/Engine/Source/Runtime/Engine/Private/WorldPartition/WorldPartitionStreamingGeneration.cpp
+++ b/Engine/Source/Runtime/Engine/Private/WorldPartition/WorldPartitionStreamingGeneration.cpp
@@ -41,7 +41,6 @@
 #include "WorldPartition/HLOD/HLODLayer.h"
 #include "WorldPartition/HLOD/StandaloneHLODActor.h"
 #include "WorldPartition/HLOD/StandaloneHLODSubsystem.h"
-#include "WorldPartition/RuntimeHashSet/WorldPartitionRuntimeHashSet.h"
 #include "WorldPartition/RuntimeHashSet/RuntimePartitionPersistent.h"
 #include "WorldPartition/WorldPartitionSubsystem.h"
 #include "WorldPartition/WorldPartitionRuntimeContainerResolving.h"
@@ -830,21 +829,10 @@ class FWorldPartitionStreamingGenerator
 
 	FName GetResolvedRuntimeGridForHLODLayer(FName RuntimeGrid, const FSoftObjectPath& HLODLayerPath) const
 	{
-		if (WorldPartitionContext)
+		if (WorldPartitionContext && WorldPartitionContext->RuntimeHash)
 		{
-			if (const UWorldPartitionRuntimeHashSet* HashSet = Cast<UWorldPartitionRuntimeHashSet>(WorldPartitionContext->RuntimeHash))
-			{
-				if (const URuntimePartition* RuntimePartition = HashSet->ResolveRuntimePartitionForHLODLayer(RuntimeGrid, HLODLayerPath))
-				{
-					if (RuntimePartition->IsA<URuntimePartitionPersistent>())
-					{
-						return NAME_None;
-					}
-					
-					FName GridName = RuntimeGrid.IsNone() ? DefaultGrid : RuntimeGrid;
-					return FName(GridName.ToString() + TEXT(":") + RuntimePartition->Name.ToString());
-				}
-			}
+			FName GridName = RuntimeGrid.IsNone() ? DefaultGrid : RuntimeGrid;
+			return WorldPartitionContext->RuntimeHash->ResolveRuntimeGridNameForHLODLayer(GridName, HLODLayerPath);
 		}
 		return NAME_None;
 	}
diff --git a/Engine/Source/Runtime/Engine/Public/WorldPartition/RuntimeHashSet/WorldPartitionRuntimeHashSet.h b/Engine/Source/Runtime/Engine/Public/WorldPartition/RuntimeHashSet/WorldPartitionRuntimeHashSet.h
--- a/Engine/Source/Runtime/Engine/Public/WorldPartition/RuntimeHashSet/WorldPartitionRuntimeHashSet.h
+++ b/Engine/Source/Runtime/Engine/Public/WorldPartition/RuntimeHashSet/WorldPartitionRuntimeHashSet.h
@@ -177,8 +177,6 @@ public:
 	ENGINE_API virtual void Draw3D(const TArray<FWorldPartitionStreamingSource>& Sources) const override;
 	ENGINE_API virtual const FGuid* GetStandaloneHLODActorSourceCellOverride(const FGuid& InActorGuid) const override;
 	ENGINE_API virtual const FGuid* GetCustomHLODActorSourceCellOverride(const FGuid& InActorGuid) const override;
-	ENGINE_API void ForEachHLODLayer(TFunctionRef<bool(FName, FName, int32)> Func) const; // RuntimePartitionName, HLODSetupName, HLODSetupIndex
-	ENGINE_API int32 ComputeHLODHierarchyDepth() const;
 
 #if WITH_EDITOR
 	// Streaming generation interface
@@ -190,6 +188,10 @@ public:
 	ENGINE_API virtual bool IsValidGrid(FName GridName, const UClass* ActorClass) const override;
 	ENGINE_API virtual bool IsValidHLODLayer(FName GridName, const FSoftObjectPath& HLODLayerPath) const override;
 	ENGINE_API virtual void DumpStateLog(FHierarchicalLogArchive& Ar) const override;
+	ENGINE_API virtual void ForEachHLODLayer(TFunctionRef<bool(FName, FName, int32)> Func) const override; // RuntimePartitionName, HLODSetupName, HLODSetupIndex
+	ENGINE_API virtual int32 ComputeHLODHierarchyDepth() const override;
+	ENGINE_API virtual FName ResolveRuntimeGridNameForHLODLayer(FName RuntimeGrid, const FSoftObjectPath& HLODLayerPath) const override;
+	ENGINE_API virtual bool CopyGridSetupFrom(const UWorldPartitionRuntimeHash* Other) override;
 #endif
 
 	// Helpers
diff --git a/Engine/Source/Runtime/Engine/Public/WorldPartition/WorldPartitionRuntimeHash.h b/Engine/Source/Runtime/Engine/Public/WorldPartition/WorldPartitionRuntimeHash.h
--- a/Engine/Source/Runtime/Engine/Public/WorldPartition/WorldPartitionRuntimeHash.h
+++ b/Engine/Source/Runtime/Engine/Public/WorldPartition/WorldPartitionRuntimeHash.h
@@ -205,6 +205,10 @@ class UWorldPartitionRuntimeHash : public UObject
 	virtual FName GetDefaultGrid() const { return NAME_None; }
 	virtual bool IsValidGrid(FName GridName, const UClass* ActorClass) const { return false; }
 	virtual bool IsValidHLODLayer(FName GridName, const FSoftObjectPath& HLODLayerPath) const { return false; }
+	virtual void ForEachHLODLayer(TFunctionRef<bool(FName, FName, int32)> Func) const {}; // RuntimePartitionName, HLODSetupName, HLODSetupIndex
+	virtual int32 ComputeHLODHierarchyDepth() const { return 0; }
+	virtual FName ResolveRuntimeGridNameForHLODLayer(FName RuntimeGrid, const FSoftObjectPath& HLODLayerPath) const { return NAME_None; }
+	virtual bool CopyGridSetupFrom(const UWorldPartitionRuntimeHash* Other) { return false; }
 	virtual void DrawPreview() const {}
 
 	virtual bool HasStreamingContent() const { return false; }
