[Engine]
* Added support for SkeletalMeshUpdater queue to be flushed on render side (used by statestream since skeletal mesh only exists render side)

#Synced-CL 49810000

diff --git a/Engine/Source/Runtime/Engine/Internal/SkeletalMeshUpdater.h b/Engine/Source/Runtime/Engine/Internal/SkeletalMeshUpdater.h
--- a/Engine/Source/Runtime/Engine/Internal/SkeletalMeshUpdater.h
+++ b/Engine/Source/Runtime/Engine/Internal/SkeletalMeshUpdater.h
@@ -444,6 +444,9 @@ public:
 		PushCommandsTask = {};
 	}
 
+	// This method is added for StateStream purposes and not to change to much code
+	ENGINE_API void PushChangesToStream();
+
 	// Call to insert a task to push commands when the game thread task completes.
 	ENGINE_API UE::Tasks::FTask AddPushCommandsTask(const UE::Tasks::FTask& PrerequisiteTask);
 
diff --git a/Engine/Source/Runtime/Engine/Private/SkeletalMeshUpdater.cpp b/Engine/Source/Runtime/Engine/Private/SkeletalMeshUpdater.cpp
--- a/Engine/Source/Runtime/Engine/Private/SkeletalMeshUpdater.cpp
+++ b/Engine/Source/Runtime/Engine/Private/SkeletalMeshUpdater.cpp
@@ -6,6 +6,13 @@
 #include "Rendering/RenderCommandPipes.h"
 #include "RenderGraphBuilder.h"
 
+#if WITH_STATE_STREAM
+#define SKELETAL_MESH_UPDATER_CHECK_THREAD_ACCESS check(IsInParallelRenderingThread());
+#else
+#define SKELETAL_MESH_UPDATER_CHECK_THREAD_ACCESS check(IsInGameThread() || IsInParallelGameThread());
+#endif
+
+
 int32 GSkeletalMeshDynamicDataPoolBudget = 4 * 1024;
 static FAutoConsoleVariableRef CVarSkeletalMeshDynamicDataPoolSize(
 	TEXT("r.SkeletalMesh.DynamicDataPoolBudget"),
@@ -198,7 +205,7 @@ FSkeletalMeshUpdateChannel::~FSkeletalMeshUpdateChannel()
 
 FSkeletalMeshUpdateHandle FSkeletalMeshUpdateChannel::Create(FSkeletalMeshObject* MeshObject)
 {
-	check(IsInGameThread() || IsInParallelGameThread());
+	SKELETAL_MESH_UPDATER_CHECK_THREAD_ACCESS
 
 	FSkeletalMeshUpdateHandle Handle;
 	Handle.Index = IndexAllocator.Allocate();
@@ -219,7 +226,7 @@ FSkeletalMeshUpdateHandle FSkeletalMeshUpdateChannel::Create(FSkeletalMeshObject
 
 void FSkeletalMeshUpdateChannel::Release(FSkeletalMeshUpdateHandle&& Handle)
 {
-	check(IsInGameThread() || IsInParallelGameThread());
+	SKELETAL_MESH_UPDATER_CHECK_THREAD_ACCESS
 	check(Handle.Channel == this);
 
 	FOp Op;
@@ -239,7 +246,7 @@ void FSkeletalMeshUpdateChannel::Release(FSkeletalMeshUpdateHandle&& Handle)
 
 void FSkeletalMeshUpdateChannel::Update(const FSkeletalMeshUpdateHandle& Handle, FSkeletalMeshDynamicData* MeshDynamicData)
 {
-	check(IsInGameThread() || IsInParallelGameThread());
+	SKELETAL_MESH_UPDATER_CHECK_THREAD_ACCESS
 	check(MeshDynamicData);
 	check(Handle.Channel == this);
 
@@ -261,8 +268,7 @@ void FSkeletalMeshUpdateChannel::Shutdown()
 
 TUniquePtr<FSkeletalMeshUpdateChannel::FOpQueue> FSkeletalMeshUpdateChannel::PopFromQueue()
 {
-	check(IsInGameThread() || IsInParallelGameThread());
-
+	SKELETAL_MESH_UPDATER_CHECK_THREAD_ACCESS
 	if (OpQueue->Num.load(std::memory_order_relaxed) == 0)
 	{
 		return {};
@@ -296,6 +302,7 @@ FSkeletalMeshUpdater::FSkeletalMeshUpdater(FSceneInterface* InScene, FGPUSkinCac
 	, GPUSkinCache(InGPUSkinCache)
 	, Channels(FSkeletalMeshUpdateChannel::GetChannels())
 {
+#if !WITH_STATE_STREAM
 	DelegateHandle = UE::RenderCommandPipe::GetStopRecordingDelegate().AddLambda([this] (const FRenderCommandPipeBitArray&)
 	{
 		if (bInAsyncPushCommandsRegion)
@@ -325,6 +332,7 @@ FSkeletalMeshUpdater::FSkeletalMeshUpdater(FSceneInterface* InScene, FGPUSkinCac
 			});
 		}
 	});
+#endif
 }
 
 void FSkeletalMeshUpdater::Shutdown()
@@ -337,6 +345,25 @@ void FSkeletalMeshUpdater::Shutdown()
 	}
 }
 
+void FSkeletalMeshUpdater::PushChangesToStream()
+{
+	TArray<TPair<FSkeletalMeshUpdateChannel*, TUniquePtr<FSkeletalMeshUpdateChannel::FOpQueue>>, FConcurrentLinearArrayAllocator> ChannelsToPush;
+	ChannelsToPush.Reserve(Channels.Num());
+
+	for (FSkeletalMeshUpdateChannel& Channel : Channels)
+	{
+		if (TUniquePtr<FSkeletalMeshUpdateChannel::FOpQueue> Ops{ Channel.PopFromQueue() }; Ops.IsValid())
+		{
+			ChannelsToPush.Emplace(&Channel, MoveTemp(Ops));
+		}
+	}
+
+	for (auto& KeyValue : ChannelsToPush)
+	{
+		KeyValue.Key->PushToStream(MoveTemp(KeyValue.Value));
+	}
+}
+
 UE::Tasks::FTask FSkeletalMeshUpdater::AddPushCommandsTask(const UE::Tasks::FTask& PrerequisiteTask)
 {
 	check(bInAsyncPushCommandsRegion);
