[Engine]
* Added ReleaseResources call to newed FSkeletalMeshObjectNanite before deleted (to make sure it release things first)

#Synced-CL 49810000

diff --git a/Engine/Source/Runtime/Engine/Private/SkinnedMeshSceneProxyDesc.cpp b/Engine/Source/Runtime/Engine/Private/SkinnedMeshSceneProxyDesc.cpp
--- a/Engine/Source/Runtime/Engine/Private/SkinnedMeshSceneProxyDesc.cpp
+++ b/Engine/Source/Runtime/Engine/Private/SkinnedMeshSceneProxyDesc.cpp
@@ -232,6 +232,12 @@ FSkeletalMeshObject* FSkinnedMeshSceneProxyDesc::CreateMeshObject(const FSkinned
 		}
 		else
 		{
+			// TODO: This should run outside of statestream as well because otherwise resources are leaked and an assert happens inside
+			// SkeletalMeshUpdater. But since there are enqueued commands accessing the NaniteMeshObject, then I don't know if this will work on game thread
+			#if WITH_STATE_STREAM
+			NaniteMeshObject->ReleaseResources();
+			#endif
+
 			delete NaniteMeshObject;
 		}
 	}
