Avoid creating a temporary TObjectPtr in a ternary expression

#Synced-CL 49651271

diff --git a/Engine/Source/Runtime/Engine/Private/WorldPartition/HLOD/StandaloneHLODSubsystem.cpp b/Engine/Source/Runtime/Engine/Private/WorldPartition/HLOD/StandaloneHLODSubsystem.cpp
--- a/Engine/Source/Runtime/Engine/Private/WorldPartition/HLOD/StandaloneHLODSubsystem.cpp
+++ b/Engine/Source/Runtime/Engine/Private/WorldPartition/HLOD/StandaloneHLODSubsystem.cpp
@@ -297,7 +297,7 @@ void UWorldPartitionStandaloneHLODSubsystem::UpdateStandaloneHLODActors(UWorldPa
 			SpawnParams.Name = FName(ActorLabel);
 			SpawnParams.NameMode = FActorSpawnParameters::ESpawnActorNameMode::Requested;
 
-			UWorldPartitionRuntimeHash* WorldHash = GetWorld()->GetWorldPartition() ? GetWorld()->GetWorldPartition()->RuntimeHash : nullptr;
+			UWorldPartitionRuntimeHash* WorldHash = GetWorld()->GetWorldPartition() ? GetWorld()->GetWorldPartition()->RuntimeHash.Get() : nullptr;
 			if (!WorldHash)
 			{
 				UE_LOG(LogStandaloneHLODSubsystem, Log, TEXT("RuntimeHash not found in World %s"), *GetWorld()->GetName());
