[Engine]
* Fixed so lyra works with WITH_STATE_STREAM set

#Synced-CL 49872240

diff --git a/Engine/Source/Runtime/Engine/Classes/Components/ExponentialHeightFogComponent.h b/Engine/Source/Runtime/Engine/Classes/Components/ExponentialHeightFogComponent.h
--- a/Engine/Source/Runtime/Engine/Classes/Components/ExponentialHeightFogComponent.h
+++ b/Engine/Source/Runtime/Engine/Classes/Components/ExponentialHeightFogComponent.h
@@ -296,6 +296,7 @@ protected:
 	ENGINE_API virtual void CreateRenderState_Concurrent(FRegisterComponentContext* Context) override;
 	ENGINE_API virtual void SendRenderTransform_Concurrent() override;
 	ENGINE_API virtual void DestroyRenderState_Concurrent() override;
+	ENGINE_API virtual void OnUnregister() override;
 	//~ End UActorComponent Interface.
 
 	ENGINE_API void AddFogIfNeeded();
diff --git a/Engine/Source/Runtime/Engine/Private/Components/ExponentialHeightFogComponent.cpp b/Engine/Source/Runtime/Engine/Private/Components/ExponentialHeightFogComponent.cpp
--- a/Engine/Source/Runtime/Engine/Private/Components/ExponentialHeightFogComponent.cpp
+++ b/Engine/Source/Runtime/Engine/Private/Components/ExponentialHeightFogComponent.cpp
@@ -150,6 +150,14 @@ void UExponentialHeightFogComponent::DestroyRenderState_Concurrent()
 	GetWorld()->Scene->RemoveExponentialHeightFog(uint64(this));
 }
 
+void UExponentialHeightFogComponent::OnUnregister()
+{
+	Super::OnUnregister();
+	#if WITH_STATE_STREAM_ACTOR
+	Handle = {};
+	#endif
+}
+
 #if WITH_EDITOR
 
 bool UExponentialHeightFogComponent::CanEditChange(const FProperty* InProperty) const
