Fix for issue where replayout did not respect net.ReplicateCustomDeltaPropertiesInRepIndexOrder and always replicated CustomDeltaProperties in RepIndexOrder

#Synced-CL 49833605

diff --git a/Engine/Source/Runtime/Engine/Private/RepLayout.cpp b/Engine/Source/Runtime/Engine/Private/RepLayout.cpp
--- a/Engine/Source/Runtime/Engine/Private/RepLayout.cpp
+++ b/Engine/Source/Runtime/Engine/Private/RepLayout.cpp
@@ -6220,7 +6220,7 @@ void FRepLayout::InitFromClass(
 	UObject* Object = InObjectClass->GetDefaultObject();
 	Object->GetLifetimeReplicatedProps(LifetimeProps);
 
-	// Set default settings for properties not explicitely registered
+	// Set default settings for properties not explicitly registered
 	{
 		UE::Net::FNetBitArray RegisteredProperties(InObjectClass->ClassReps.Num());
 		for (const FLifetimeProperty& UserProp : LifetimeProps)
@@ -6267,8 +6267,9 @@ void FRepLayout::InitFromClass(
 	FLifetimeProperty DefaultLifetimeProp;
 
 	// Setup lifetime replicated properties
-	for (int32 ParentIndex=0; ParentIndex < InObjectClass->ClassReps.Num(); ++ParentIndex)
+	for (int32 i = 0; i < LifetimeProps.Num(); i++)
 	{
+		const int32 ParentIndex = LifetimeProps[i].RepIndex;
 		const FRepRecord& ReplicatedProp = InObjectClass->ClassReps[ParentIndex];
 	
 		const bool bIsValid = ensureMsgf(Parents.IsValidIndex(ParentIndex), TEXT("ClassReps RepRecord of %s point to invalid RepIndex=%i for replicated property: %s. Parents.Num() = %d"), *GetFullNameSafe(InObjectClass), ParentIndex, *GetNameSafe(ReplicatedProp.Property), Parents.Num());
@@ -6286,8 +6287,7 @@ void FRepLayout::InitFromClass(
 		});
 
 		
-
-		// Set the correct RepIndex just in case its refereced.
+		// Set the correct RepIndex just in case its referenced.
 		DefaultLifetimeProp.RepIndex = ParentIndex;
 
 		// Use Default settings if the user didn't 
