LLM: Add missing scopes for LLMTagSet::AssetClasses. AssetClasses scopes are required for ELLMTagSet::CodeOrContent to record the allocation as Content, and therefore AssetClasses scopes should be present wherever ELLMTagSe::Assets scope is present.
#rnx

#Synced-CL 49884901

diff --git a/Engine/Source/Runtime/CoreUObject/Private/Serialization/AsyncLoading2.cpp b/Engine/Source/Runtime/CoreUObject/Private/Serialization/AsyncLoading2.cpp
--- a/Engine/Source/Runtime/CoreUObject/Private/Serialization/AsyncLoading2.cpp
+++ b/Engine/Source/Runtime/CoreUObject/Private/Serialization/AsyncLoading2.cpp
@@ -7555,6 +7555,7 @@ EEventLoadNodeExecutionResult FAsyncPackage2::ProcessLinkerLoadPackageSummary(FA
 {
 	TRACE_CPUPROFILER_EVENT_SCOPE(ProcessLinkerLoadPackageSummary);
 	LLM_SCOPE_DYNAMIC_STAT_OBJECTPATH_FNAME(LinkerLoadState->Linker->GetPackagePath().GetPackageFName(), ELLMTagSet::Assets);
+	LLM_SCOPE_DYNAMIC_STAT_OBJECTPATH(UPackage::StaticClass(), ELLMTagSet::AssetClasses);
 	UE_TRACE_METADATA_SCOPE_ASSET_FNAME(NAME_None, NAME_None, LinkerLoadState->Linker->GetPackagePath().GetPackageFName());
 	SCOPED_LOADTIMER_ASSET_TEXT(*LinkerLoadState->Linker->GetDebugName());
 
diff --git a/Engine/Source/Runtime/CoreUObject/Private/UObject/UObjectGlobals.cpp b/Engine/Source/Runtime/CoreUObject/Private/UObject/UObjectGlobals.cpp
--- a/Engine/Source/Runtime/CoreUObject/Private/UObject/UObjectGlobals.cpp
+++ b/Engine/Source/Runtime/CoreUObject/Private/UObject/UObjectGlobals.cpp
@@ -1806,6 +1806,7 @@ UPackage* LoadPackage(UPackage* InOuter, const FPackagePath& PackagePath, uint32
 #endif // STATS
 
 	LLM_SCOPE_DYNAMIC_STAT_OBJECTPATH_FNAME(PackagePath.GetPackageFName(), ELLMTagSet::Assets);
+	LLM_SCOPE_DYNAMIC_STAT_OBJECTPATH(UPackage::StaticClass(), ELLMTagSet::AssetClasses);
 	UE_TRACE_METADATA_SCOPE_ASSET_FNAME(NAME_None, NAME_None, PackagePath.GetPackageFName());
 	TRACE_LOADTIME_REQUEST_GROUP_SCOPE(TEXT("SyncLoad - %s"), *PackagePath.GetDebugName());
 	
diff --git a/Engine/Source/Runtime/Engine/Private/Components/ActorComponent.cpp b/Engine/Source/Runtime/Engine/Private/Components/ActorComponent.cpp
--- a/Engine/Source/Runtime/Engine/Private/Components/ActorComponent.cpp
+++ b/Engine/Source/Runtime/Engine/Private/Components/ActorComponent.cpp
@@ -2405,6 +2405,7 @@ void UActorComponent::CreatePhysicsState(bool bAllowDeferral)
 {
 	LLM_SCOPE(ELLMTag::Chaos);
 	LLM_SCOPE_DYNAMIC_STAT_OBJECTPATH(GetPackage(), ELLMTagSet::Assets);
+	LLM_SCOPE_DYNAMIC_STAT_OBJECTPATH(UWorld::StaticClass(), ELLMTagSet::AssetClasses);
 	UE_TRACE_METADATA_SCOPE_ASSET_FNAME(NAME_None, NAME_None, GetPackage()->GetFName());
 
 	SCOPE_CYCLE_COUNTER(STAT_ComponentCreatePhysicsState);
@@ -2522,6 +2523,7 @@ void UActorComponent::ExecuteRegisterEvents(FRegisterComponentContext* Context)
 		SCOPE_CYCLE_COUNTER(STAT_ComponentCreateRenderState);
 		LLM_SCOPE(ELLMTag::SceneRender);
 		LLM_SCOPE_DYNAMIC_STAT_OBJECTPATH(GetPackage(), ELLMTagSet::Assets);
+		LLM_SCOPE_DYNAMIC_STAT_OBJECTPATH(UWorld::StaticClass(), ELLMTagSet::AssetClasses);
 		UE_TRACE_METADATA_SCOPE_ASSET_FNAME(NAME_None, NAME_None, GetPackage()->GetFName());
 		CreateRenderState_Concurrent(Context);
 		checkf(bRenderStateCreated, TEXT("Failed to route CreateRenderState_Concurrent (%s)"), *GetFullName());
@@ -2646,6 +2648,7 @@ void UActorComponent::DoDeferredRenderUpdates_Concurrent()
 {
 	LLM_SCOPE(ELLMTag::SceneRender);
 	LLM_SCOPE_DYNAMIC_STAT_OBJECTPATH(GetPackage(), ELLMTagSet::Assets);
+	LLM_SCOPE_DYNAMIC_STAT_OBJECTPATH(UWorld::StaticClass(), ELLMTagSet::AssetClasses);
 	UE_TRACE_METADATA_SCOPE_ASSET_FNAME(NAME_None, NAME_None, GetPackage()->GetFName());
 
 	checkf(!IsUnreachable(), TEXT("%s"), *GetFullName());
diff --git a/Engine/Source/Runtime/Engine/Private/Components/SkinnedMeshComponent.cpp b/Engine/Source/Runtime/Engine/Private/Components/SkinnedMeshComponent.cpp
--- a/Engine/Source/Runtime/Engine/Private/Components/SkinnedMeshComponent.cpp
+++ b/Engine/Source/Runtime/Engine/Private/Components/SkinnedMeshComponent.cpp
@@ -1162,6 +1162,7 @@ void USkinnedMeshComponent::CreateRenderState_Concurrent(FRegisterComponentConte
 	if( GetSkinnedAsset() )
 	{
 		LLM_SCOPE_DYNAMIC_STAT_OBJECTPATH(GetSkinnedAsset()->GetPackage(), ELLMTagSet::Assets);
+		LLM_SCOPE_DYNAMIC_STAT_OBJECTPATH(UWorld::StaticClass(), ELLMTagSet::AssetClasses);
 		UE_TRACE_METADATA_SCOPE_ASSET_FNAME(NAME_None, NAME_None, GetSkinnedAsset()->GetPackage()->GetFName());
 
 		// Attempting to track down UE-45505, where it looks as if somehow a skeletal mesh component's mesh has only been partially loaded, causing a mismatch in the LOD arrays
@@ -2968,6 +2969,7 @@ bool USkinnedMeshComponent::AllocateTransformData()
 	if (GetSkinnedAsset() != nullptr && LeaderPoseComponent == nullptr )
 	{
 		LLM_SCOPE_DYNAMIC_STAT_OBJECTPATH(GetSkinnedAsset()->GetPackage(), ELLMTagSet::Assets);
+		LLM_SCOPE_DYNAMIC_STAT_OBJECTPATH(UWorld::StaticClass(), ELLMTagSet::AssetClasses);
 		UE_TRACE_METADATA_SCOPE_ASSET_FNAME(NAME_None, NAME_None, GetSkinnedAsset()->GetPackage()->GetFName());
 
 		const int32 NumBones = GetSkinnedAsset()->GetRefSkeleton().GetNum();
diff --git a/Engine/Source/Runtime/Engine/Private/Texture2D.cpp b/Engine/Source/Runtime/Engine/Private/Texture2D.cpp
--- a/Engine/Source/Runtime/Engine/Private/Texture2D.cpp
+++ b/Engine/Source/Runtime/Engine/Private/Texture2D.cpp
@@ -1595,6 +1595,7 @@ void FVirtualTexture2DResource::InitRHI(FRHICommandListBase& RHICmdList)
 {
 	TRACE_CPUPROFILER_EVENT_SCOPE(FVirtualTexture2DResource::InitRHI);
 	LLM_SCOPE_DYNAMIC_STAT_OBJECTPATH_FNAME(PackageName, ELLMTagSet::Assets);
+	LLM_SCOPE_DYNAMIC_STAT_OBJECTPATH_FNAME(FRHITextureCreateDesc().GetTraceClassName(), ELLMTagSet::AssetClasses);
 	UE_TRACE_METADATA_SCOPE_ASSET_FNAME(NAME_None, NAME_None, PackageName);
 
 	uint32 MaxAnisotropy = 0;
diff --git a/Engine/Source/Runtime/Engine/Private/TextureRenderTarget2D.cpp b/Engine/Source/Runtime/Engine/Private/TextureRenderTarget2D.cpp
--- a/Engine/Source/Runtime/Engine/Private/TextureRenderTarget2D.cpp
+++ b/Engine/Source/Runtime/Engine/Private/TextureRenderTarget2D.cpp
@@ -523,6 +523,7 @@ ETextureCreateFlags FTextureRenderTarget2DResource::GetCreateFlags()
 void FTextureRenderTarget2DResource::InitRHI(FRHICommandListBase& RHICmdList)
 {
 	LLM_SCOPE_DYNAMIC_STAT_OBJECTPATH(Owner->GetPackage(), ELLMTagSet::Assets);
+	LLM_SCOPE_DYNAMIC_STAT_OBJECTPATH_FNAME(FRHITextureCreateDesc().GetTraceClassName(), ELLMTagSet::AssetClasses);
 	UE_TRACE_METADATA_SCOPE_ASSET_FNAME(NAME_None, NAME_None, Owner->GetPackage()->GetFName());
 
 	if( TargetSizeX > 0 && TargetSizeY > 0 )
@@ -619,6 +620,7 @@ void FTextureRenderTarget2DResource::ReleaseRHI()
 void FTextureRenderTarget2DResource::UpdateDeferredResource( FRHICommandListImmediate& RHICmdList, bool bClearRenderTarget/*=true*/ )
 {
 	LLM_SCOPE_DYNAMIC_STAT_OBJECTPATH(Owner->GetPackage(), ELLMTagSet::Assets);
+	LLM_SCOPE_DYNAMIC_STAT_OBJECTPATH_FNAME(FRHITextureCreateDesc().GetTraceClassName(), ELLMTagSet::AssetClasses);
 	UE_TRACE_METADATA_SCOPE_ASSET_FNAME(NAME_None, NAME_None, Owner->GetPackage()->GetFName());
 
 	SCOPED_DRAW_EVENT(RHICmdList, GPUResourceUpdate)
diff --git a/Engine/Source/Runtime/Engine/Private/TextureRenderTarget2DArray.cpp b/Engine/Source/Runtime/Engine/Private/TextureRenderTarget2DArray.cpp
--- a/Engine/Source/Runtime/Engine/Private/TextureRenderTarget2DArray.cpp
+++ b/Engine/Source/Runtime/Engine/Private/TextureRenderTarget2DArray.cpp
@@ -228,6 +228,7 @@ UTexture2DArray* UTextureRenderTarget2DArray::ConstructTexture2DArray(UObject* I
 void FTextureRenderTarget2DArrayResource::InitRHI(FRHICommandListBase& RHICmdList)
 {
 	LLM_SCOPE_DYNAMIC_STAT_OBJECTPATH(Owner->GetPackage(), ELLMTagSet::Assets);
+	LLM_SCOPE_DYNAMIC_STAT_OBJECTPATH_FNAME(FRHITextureCreateDesc().GetTraceClassName(), ELLMTagSet::AssetClasses);
 	UE_TRACE_METADATA_SCOPE_ASSET_FNAME(NAME_None, NAME_None, Owner->GetPackage()->GetFName());
 
 	if ((Owner->SizeX > 0) && (Owner->SizeY > 0) && (Owner->Slices > 0))
@@ -308,6 +309,7 @@ void FTextureRenderTarget2DArrayResource::ReleaseRHI()
 void FTextureRenderTarget2DArrayResource::UpdateDeferredResource(FRHICommandListImmediate& RHICmdList, bool bClearRenderTarget/*=true*/)
 {
 	LLM_SCOPE_DYNAMIC_STAT_OBJECTPATH(Owner->GetPackage(), ELLMTagSet::Assets);
+	LLM_SCOPE_DYNAMIC_STAT_OBJECTPATH_FNAME(FRHITextureCreateDesc().GetTraceClassName(), ELLMTagSet::AssetClasses);
 	UE_TRACE_METADATA_SCOPE_ASSET_FNAME(NAME_None, NAME_None, Owner->GetPackage()->GetFName());
 
 	RemoveFromDeferredUpdateList();
diff --git a/Engine/Source/Runtime/Engine/Private/TextureRenderTargetCube.cpp b/Engine/Source/Runtime/Engine/Private/TextureRenderTargetCube.cpp
--- a/Engine/Source/Runtime/Engine/Private/TextureRenderTargetCube.cpp
+++ b/Engine/Source/Runtime/Engine/Private/TextureRenderTargetCube.cpp
@@ -267,6 +267,7 @@ UTextureCube* UTextureRenderTargetCube::ConstructTextureCube(UObject* InOuter, c
 void FTextureRenderTargetCubeResource::InitRHI(FRHICommandListBase& RHICmdList)
 {
 	LLM_SCOPE_DYNAMIC_STAT_OBJECTPATH(Owner->GetPackage(), ELLMTagSet::Assets);
+	LLM_SCOPE_DYNAMIC_STAT_OBJECTPATH_FNAME(FRHITextureCreateDesc().GetTraceClassName(), ELLMTagSet::AssetClasses);
 	UE_TRACE_METADATA_SCOPE_ASSET_FNAME(NAME_None, NAME_None, Owner->GetPackage()->GetFName());
 
 	if(Owner->SizeX > 0)
@@ -374,6 +375,7 @@ void FTextureRenderTargetCubeResource::ReleaseRHI()
 void FTextureRenderTargetCubeResource::UpdateDeferredResource(FRHICommandListImmediate& RHICmdList, bool bClearRenderTarget/*=true*/)
 {
 	LLM_SCOPE_DYNAMIC_STAT_OBJECTPATH(Owner->GetPackage(), ELLMTagSet::Assets);
+	LLM_SCOPE_DYNAMIC_STAT_OBJECTPATH_FNAME(FRHITextureCreateDesc().GetTraceClassName(), ELLMTagSet::AssetClasses);
 	UE_TRACE_METADATA_SCOPE_ASSET_FNAME(NAME_None, NAME_None, Owner->GetPackage()->GetFName());
 
 	RemoveFromDeferredUpdateList();
diff --git a/Engine/Source/Runtime/Engine/Private/TextureRenderTargetVolume.cpp b/Engine/Source/Runtime/Engine/Private/TextureRenderTargetVolume.cpp
--- a/Engine/Source/Runtime/Engine/Private/TextureRenderTargetVolume.cpp
+++ b/Engine/Source/Runtime/Engine/Private/TextureRenderTargetVolume.cpp
@@ -237,6 +237,7 @@ UVolumeTexture* UTextureRenderTargetVolume::ConstructTextureVolume(UObject* InOu
 void FTextureRenderTargetVolumeResource::InitRHI(FRHICommandListBase& RHICmdList)
 {
 	LLM_SCOPE_DYNAMIC_STAT_OBJECTPATH(Owner->GetPackage(), ELLMTagSet::Assets);
+	LLM_SCOPE_DYNAMIC_STAT_OBJECTPATH_FNAME(FRHITextureCreateDesc().GetTraceClassName(), ELLMTagSet::AssetClasses);
 	UE_TRACE_METADATA_SCOPE_ASSET_FNAME(NAME_None, NAME_None, Owner->GetPackage()->GetFName());
 
 	if((Owner->SizeX > 0) && (Owner->SizeY > 0) && (Owner->SizeZ > 0))
@@ -308,6 +309,7 @@ void FTextureRenderTargetVolumeResource::ReleaseRHI()
 void FTextureRenderTargetVolumeResource::UpdateDeferredResource(FRHICommandListImmediate& RHICmdList, bool bClearRenderTarget/*=true*/)
 {
 	LLM_SCOPE_DYNAMIC_STAT_OBJECTPATH(Owner->GetPackage(), ELLMTagSet::Assets);
+	LLM_SCOPE_DYNAMIC_STAT_OBJECTPATH_FNAME(FRHITextureCreateDesc().GetTraceClassName(), ELLMTagSet::AssetClasses);
 	UE_TRACE_METADATA_SCOPE_ASSET_FNAME(NAME_None, NAME_None, Owner->GetPackage()->GetFName());
 
 	RemoveFromDeferredUpdateList();
diff --git a/Engine/Source/Runtime/Renderer/Private/RendererScene.cpp b/Engine/Source/Runtime/Renderer/Private/RendererScene.cpp
--- a/Engine/Source/Runtime/Renderer/Private/RendererScene.cpp
+++ b/Engine/Source/Runtime/Renderer/Private/RendererScene.cpp
@@ -1373,6 +1373,7 @@ void FScene::BatchAddPrimitivesInternal(TArrayView<T*> InPrimitives)
 	}
 #endif
 	LLM_SCOPE_DYNAMIC_STAT_OBJECTPATH_FNAME(ToUObject(InPrimitives[0]) ? InPrimitives[0]->GetOutermost()->GetFName() : NAME_None, ELLMTagSet::Assets);
+	LLM_SCOPE_DYNAMIC_STAT_OBJECTPATH_FNAME(ToUObject(InPrimitives[0]) ? FName(UPackage::StaticClass()->GetPathName()) : NAME_None, ELLMTagSet::AssetClasses);
 	UE_TRACE_METADATA_SCOPE_ASSET_FNAME(NAME_None, NAME_None, ToUObject(InPrimitives[0]) ? InPrimitives[0]->GetOutermost()->GetFName() : NAME_None);
 
 	SCOPE_CYCLE_COUNTER(STAT_AddScenePrimitiveGT);
