LLM CodeOrContent TagSet: Fix classname in UE_TRACE_METADATA_SCOPE_ASSET_FNAME to match classname in LLM_SCOPE_DYNAMIC_STAT_OBJECTPATH. Remove no-longer needed LLM_CHECK.
#rnx

#Synced-CL 49884901

diff --git a/Engine/Source/Runtime/Core/Private/HAL/LLM/LLM.cpp b/Engine/Source/Runtime/Core/Private/HAL/LLM/LLM.cpp
--- a/Engine/Source/Runtime/Core/Private/HAL/LLM/LLM.cpp
+++ b/Engine/Source/Runtime/Core/Private/HAL/LLM/LLM.cpp
@@ -3076,7 +3076,6 @@ FTagData& FLLMGlobals::RegisterTagData(FName Name, FName DisplayName, FName Pare
 	{
 		// We set LLM UniqueName = <TheEntireString> and LLM DisplayName = StatDisplayName.
 		// Stat tags do not specify their parent, and their parent is set to the CustomName aggregator.
-		LLMCheck(DisplayName.IsNone());
 		LLMCheck(ParentName.IsNone());
 		if (DisplayName.IsNone())
 		{
diff --git a/Engine/Source/Runtime/CoreUObject/Private/Serialization/AsyncLoading2.cpp b/Engine/Source/Runtime/CoreUObject/Private/Serialization/AsyncLoading2.cpp
--- a/Engine/Source/Runtime/CoreUObject/Private/Serialization/AsyncLoading2.cpp
+++ b/Engine/Source/Runtime/CoreUObject/Private/Serialization/AsyncLoading2.cpp
@@ -7556,7 +7556,7 @@ EEventLoadNodeExecutionResult FAsyncPackage2::ProcessLinkerLoadPackageSummary(FA
 	TRACE_CPUPROFILER_EVENT_SCOPE(ProcessLinkerLoadPackageSummary);
 	LLM_SCOPE_DYNAMIC_STAT_OBJECTPATH_FNAME(LinkerLoadState->Linker->GetPackagePath().GetPackageFName(), ELLMTagSet::Assets);
 	LLM_SCOPE_DYNAMIC_STAT_OBJECTPATH(UPackage::StaticClass(), ELLMTagSet::AssetClasses);
-	UE_TRACE_METADATA_SCOPE_ASSET_FNAME(NAME_None, NAME_None, LinkerLoadState->Linker->GetPackagePath().GetPackageFName());
+	UE_TRACE_METADATA_SCOPE_ASSET_FNAME(NAME_None, UPackage::StaticClass()->GetFName(), LinkerLoadState->Linker->GetPackagePath().GetPackageFName());
 	SCOPED_LOADTIMER_ASSET_TEXT(*LinkerLoadState->Linker->GetDebugName());
 
 #if ALT2_ENABLE_NEW_ARCHIVE_FOR_LINKERLOAD
diff --git a/Engine/Source/Runtime/CoreUObject/Private/UObject/UObjectGlobals.cpp b/Engine/Source/Runtime/CoreUObject/Private/UObject/UObjectGlobals.cpp
--- a/Engine/Source/Runtime/CoreUObject/Private/UObject/UObjectGlobals.cpp
+++ b/Engine/Source/Runtime/CoreUObject/Private/UObject/UObjectGlobals.cpp
@@ -1807,7 +1807,7 @@ UPackage* LoadPackage(UPackage* InOuter, const FPackagePath& PackagePath, uint32
 
 	LLM_SCOPE_DYNAMIC_STAT_OBJECTPATH_FNAME(PackagePath.GetPackageFName(), ELLMTagSet::Assets);
 	LLM_SCOPE_DYNAMIC_STAT_OBJECTPATH(UPackage::StaticClass(), ELLMTagSet::AssetClasses);
-	UE_TRACE_METADATA_SCOPE_ASSET_FNAME(NAME_None, NAME_None, PackagePath.GetPackageFName());
+	UE_TRACE_METADATA_SCOPE_ASSET_FNAME(NAME_None, UPackage::StaticClass()->GetFName(), PackagePath.GetPackageFName());
 	TRACE_LOADTIME_REQUEST_GROUP_SCOPE(TEXT("SyncLoad - %s"), *PackagePath.GetDebugName());
 	
 	// if this is a supported asset, it should be loaded fully rather than just for diffing
diff --git a/Engine/Source/Runtime/Engine/Private/Components/ActorComponent.cpp b/Engine/Source/Runtime/Engine/Private/Components/ActorComponent.cpp
--- a/Engine/Source/Runtime/Engine/Private/Components/ActorComponent.cpp
+++ b/Engine/Source/Runtime/Engine/Private/Components/ActorComponent.cpp
@@ -2406,7 +2406,7 @@ void UActorComponent::CreatePhysicsState(bool bAllowDeferral)
 	LLM_SCOPE(ELLMTag::Chaos);
 	LLM_SCOPE_DYNAMIC_STAT_OBJECTPATH(GetPackage(), ELLMTagSet::Assets);
 	LLM_SCOPE_DYNAMIC_STAT_OBJECTPATH(UWorld::StaticClass(), ELLMTagSet::AssetClasses);
-	UE_TRACE_METADATA_SCOPE_ASSET_FNAME(NAME_None, NAME_None, GetPackage()->GetFName());
+	UE_TRACE_METADATA_SCOPE_ASSET_FNAME(NAME_None, UWorld::StaticClass()->GetFName(), GetPackage()->GetFName());
 
 	SCOPE_CYCLE_COUNTER(STAT_ComponentCreatePhysicsState);
 	TRACE_CPUPROFILER_EVENT_SCOPE(UActorComponent::CreatePhysicsState);
@@ -2524,7 +2524,7 @@ void UActorComponent::ExecuteRegisterEvents(FRegisterComponentContext* Context)
 		LLM_SCOPE(ELLMTag::SceneRender);
 		LLM_SCOPE_DYNAMIC_STAT_OBJECTPATH(GetPackage(), ELLMTagSet::Assets);
 		LLM_SCOPE_DYNAMIC_STAT_OBJECTPATH(UWorld::StaticClass(), ELLMTagSet::AssetClasses);
-		UE_TRACE_METADATA_SCOPE_ASSET_FNAME(NAME_None, NAME_None, GetPackage()->GetFName());
+		UE_TRACE_METADATA_SCOPE_ASSET_FNAME(NAME_None, UWorld::StaticClass()->GetFName(), GetPackage()->GetFName());
 		CreateRenderState_Concurrent(Context);
 		checkf(bRenderStateCreated, TEXT("Failed to route CreateRenderState_Concurrent (%s)"), *GetFullName());
 	}
@@ -2649,7 +2649,7 @@ void UActorComponent::DoDeferredRenderUpdates_Concurrent()
 	LLM_SCOPE(ELLMTag::SceneRender);
 	LLM_SCOPE_DYNAMIC_STAT_OBJECTPATH(GetPackage(), ELLMTagSet::Assets);
 	LLM_SCOPE_DYNAMIC_STAT_OBJECTPATH(UWorld::StaticClass(), ELLMTagSet::AssetClasses);
-	UE_TRACE_METADATA_SCOPE_ASSET_FNAME(NAME_None, NAME_None, GetPackage()->GetFName());
+	UE_TRACE_METADATA_SCOPE_ASSET_FNAME(NAME_None, UWorld::StaticClass()->GetFName(), GetPackage()->GetFName());
 
 	checkf(!IsUnreachable(), TEXT("%s"), *GetFullName());
 	checkf(!IsTemplate(), TEXT("%s"), *GetFullName());
diff --git a/Engine/Source/Runtime/Engine/Private/Components/SkinnedMeshComponent.cpp b/Engine/Source/Runtime/Engine/Private/Components/SkinnedMeshComponent.cpp
--- a/Engine/Source/Runtime/Engine/Private/Components/SkinnedMeshComponent.cpp
+++ b/Engine/Source/Runtime/Engine/Private/Components/SkinnedMeshComponent.cpp
@@ -1163,7 +1163,7 @@ void USkinnedMeshComponent::CreateRenderState_Concurrent(FRegisterComponentConte
 	{
 		LLM_SCOPE_DYNAMIC_STAT_OBJECTPATH(GetSkinnedAsset()->GetPackage(), ELLMTagSet::Assets);
 		LLM_SCOPE_DYNAMIC_STAT_OBJECTPATH(UWorld::StaticClass(), ELLMTagSet::AssetClasses);
-		UE_TRACE_METADATA_SCOPE_ASSET_FNAME(NAME_None, NAME_None, GetSkinnedAsset()->GetPackage()->GetFName());
+		UE_TRACE_METADATA_SCOPE_ASSET_FNAME(NAME_None, UWorld::StaticClass()->GetFName(), GetSkinnedAsset()->GetPackage()->GetFName());
 
 		// Attempting to track down UE-45505, where it looks as if somehow a skeletal mesh component's mesh has only been partially loaded, causing a mismatch in the LOD arrays
 		checkf(!GetSkinnedAsset()->HasAnyFlags(RF_NeedLoad | RF_NeedPostLoad | RF_NeedPostLoadSubobjects | RF_WillBeLoaded), TEXT("Attempting to create render state for a skeletal mesh that is is not fully loaded. Mesh: %s"), *GetSkinnedAsset()->GetName());
@@ -2970,7 +2970,7 @@ bool USkinnedMeshComponent::AllocateTransformData()
 	{
 		LLM_SCOPE_DYNAMIC_STAT_OBJECTPATH(GetSkinnedAsset()->GetPackage(), ELLMTagSet::Assets);
 		LLM_SCOPE_DYNAMIC_STAT_OBJECTPATH(UWorld::StaticClass(), ELLMTagSet::AssetClasses);
-		UE_TRACE_METADATA_SCOPE_ASSET_FNAME(NAME_None, NAME_None, GetSkinnedAsset()->GetPackage()->GetFName());
+		UE_TRACE_METADATA_SCOPE_ASSET_FNAME(NAME_None, UWorld::StaticClass()->GetFName(), GetSkinnedAsset()->GetPackage()->GetFName());
 
 		const int32 NumBones = GetSkinnedAsset()->GetRefSkeleton().GetNum();
 
diff --git a/Engine/Source/Runtime/Engine/Private/Texture2D.cpp b/Engine/Source/Runtime/Engine/Private/Texture2D.cpp
--- a/Engine/Source/Runtime/Engine/Private/Texture2D.cpp
+++ b/Engine/Source/Runtime/Engine/Private/Texture2D.cpp
@@ -1596,7 +1596,7 @@ void FVirtualTexture2DResource::InitRHI(FRHICommandListBase& RHICmdList)
 	TRACE_CPUPROFILER_EVENT_SCOPE(FVirtualTexture2DResource::InitRHI);
 	LLM_SCOPE_DYNAMIC_STAT_OBJECTPATH_FNAME(PackageName, ELLMTagSet::Assets);
 	LLM_SCOPE_DYNAMIC_STAT_OBJECTPATH_FNAME(FRHITextureCreateDesc().GetTraceClassName(), ELLMTagSet::AssetClasses);
-	UE_TRACE_METADATA_SCOPE_ASSET_FNAME(NAME_None, NAME_None, PackageName);
+	UE_TRACE_METADATA_SCOPE_ASSET_FNAME(NAME_None, FRHITextureCreateDesc().GetTraceClassName(), PackageName);
 
 	uint32 MaxAnisotropy = 0;
 	if (VirtualTextureScalability::IsAnisotropicFilteringEnabled())
diff --git a/Engine/Source/Runtime/Engine/Private/TextureRenderTarget2D.cpp b/Engine/Source/Runtime/Engine/Private/TextureRenderTarget2D.cpp
--- a/Engine/Source/Runtime/Engine/Private/TextureRenderTarget2D.cpp
+++ b/Engine/Source/Runtime/Engine/Private/TextureRenderTarget2D.cpp
@@ -524,7 +524,7 @@ void FTextureRenderTarget2DResource::InitRHI(FRHICommandListBase& RHICmdList)
 {
 	LLM_SCOPE_DYNAMIC_STAT_OBJECTPATH(Owner->GetPackage(), ELLMTagSet::Assets);
 	LLM_SCOPE_DYNAMIC_STAT_OBJECTPATH_FNAME(FRHITextureCreateDesc().GetTraceClassName(), ELLMTagSet::AssetClasses);
-	UE_TRACE_METADATA_SCOPE_ASSET_FNAME(NAME_None, NAME_None, Owner->GetPackage()->GetFName());
+	UE_TRACE_METADATA_SCOPE_ASSET_FNAME(NAME_None, FRHITextureCreateDesc().GetTraceClassName(), Owner->GetPackage()->GetFName());
 
 	if( TargetSizeX > 0 && TargetSizeY > 0 )
 	{
@@ -621,7 +621,7 @@ void FTextureRenderTarget2DResource::UpdateDeferredResource( FRHICommandListImme
 {
 	LLM_SCOPE_DYNAMIC_STAT_OBJECTPATH(Owner->GetPackage(), ELLMTagSet::Assets);
 	LLM_SCOPE_DYNAMIC_STAT_OBJECTPATH_FNAME(FRHITextureCreateDesc().GetTraceClassName(), ELLMTagSet::AssetClasses);
-	UE_TRACE_METADATA_SCOPE_ASSET_FNAME(NAME_None, NAME_None, Owner->GetPackage()->GetFName());
+	UE_TRACE_METADATA_SCOPE_ASSET_FNAME(NAME_None, FRHITextureCreateDesc().GetTraceClassName(), Owner->GetPackage()->GetFName());
 
 	SCOPED_DRAW_EVENT(RHICmdList, GPUResourceUpdate)
 	RemoveFromDeferredUpdateList();
diff --git a/Engine/Source/Runtime/Engine/Private/TextureRenderTarget2DArray.cpp b/Engine/Source/Runtime/Engine/Private/TextureRenderTarget2DArray.cpp
--- a/Engine/Source/Runtime/Engine/Private/TextureRenderTarget2DArray.cpp
+++ b/Engine/Source/Runtime/Engine/Private/TextureRenderTarget2DArray.cpp
@@ -229,7 +229,7 @@ void FTextureRenderTarget2DArrayResource::InitRHI(FRHICommandListBase& RHICmdLis
 {
 	LLM_SCOPE_DYNAMIC_STAT_OBJECTPATH(Owner->GetPackage(), ELLMTagSet::Assets);
 	LLM_SCOPE_DYNAMIC_STAT_OBJECTPATH_FNAME(FRHITextureCreateDesc().GetTraceClassName(), ELLMTagSet::AssetClasses);
-	UE_TRACE_METADATA_SCOPE_ASSET_FNAME(NAME_None, NAME_None, Owner->GetPackage()->GetFName());
+	UE_TRACE_METADATA_SCOPE_ASSET_FNAME(NAME_None, FRHITextureCreateDesc().GetTraceClassName(), Owner->GetPackage()->GetFName());
 
 	if ((Owner->SizeX > 0) && (Owner->SizeY > 0) && (Owner->Slices > 0))
 	{
@@ -310,7 +310,7 @@ void FTextureRenderTarget2DArrayResource::UpdateDeferredResource(FRHICommandList
 {
 	LLM_SCOPE_DYNAMIC_STAT_OBJECTPATH(Owner->GetPackage(), ELLMTagSet::Assets);
 	LLM_SCOPE_DYNAMIC_STAT_OBJECTPATH_FNAME(FRHITextureCreateDesc().GetTraceClassName(), ELLMTagSet::AssetClasses);
-	UE_TRACE_METADATA_SCOPE_ASSET_FNAME(NAME_None, NAME_None, Owner->GetPackage()->GetFName());
+	UE_TRACE_METADATA_SCOPE_ASSET_FNAME(NAME_None, FRHITextureCreateDesc().GetTraceClassName(), Owner->GetPackage()->GetFName());
 
 	RemoveFromDeferredUpdateList();
 
diff --git a/Engine/Source/Runtime/Engine/Private/TextureRenderTargetCube.cpp b/Engine/Source/Runtime/Engine/Private/TextureRenderTargetCube.cpp
--- a/Engine/Source/Runtime/Engine/Private/TextureRenderTargetCube.cpp
+++ b/Engine/Source/Runtime/Engine/Private/TextureRenderTargetCube.cpp
@@ -268,7 +268,7 @@ void FTextureRenderTargetCubeResource::InitRHI(FRHICommandListBase& RHICmdList)
 {
 	LLM_SCOPE_DYNAMIC_STAT_OBJECTPATH(Owner->GetPackage(), ELLMTagSet::Assets);
 	LLM_SCOPE_DYNAMIC_STAT_OBJECTPATH_FNAME(FRHITextureCreateDesc().GetTraceClassName(), ELLMTagSet::AssetClasses);
-	UE_TRACE_METADATA_SCOPE_ASSET_FNAME(NAME_None, NAME_None, Owner->GetPackage()->GetFName());
+	UE_TRACE_METADATA_SCOPE_ASSET_FNAME(NAME_None, FRHITextureCreateDesc().GetTraceClassName(), Owner->GetPackage()->GetFName());
 
 	if(Owner->SizeX > 0)
 	{
@@ -376,7 +376,7 @@ void FTextureRenderTargetCubeResource::UpdateDeferredResource(FRHICommandListImm
 {
 	LLM_SCOPE_DYNAMIC_STAT_OBJECTPATH(Owner->GetPackage(), ELLMTagSet::Assets);
 	LLM_SCOPE_DYNAMIC_STAT_OBJECTPATH_FNAME(FRHITextureCreateDesc().GetTraceClassName(), ELLMTagSet::AssetClasses);
-	UE_TRACE_METADATA_SCOPE_ASSET_FNAME(NAME_None, NAME_None, Owner->GetPackage()->GetFName());
+	UE_TRACE_METADATA_SCOPE_ASSET_FNAME(NAME_None, FRHITextureCreateDesc().GetTraceClassName(), Owner->GetPackage()->GetFName());
 
 	RemoveFromDeferredUpdateList();
 
diff --git a/Engine/Source/Runtime/Engine/Private/TextureRenderTargetVolume.cpp b/Engine/Source/Runtime/Engine/Private/TextureRenderTargetVolume.cpp
--- a/Engine/Source/Runtime/Engine/Private/TextureRenderTargetVolume.cpp
+++ b/Engine/Source/Runtime/Engine/Private/TextureRenderTargetVolume.cpp
@@ -238,7 +238,7 @@ void FTextureRenderTargetVolumeResource::InitRHI(FRHICommandListBase& RHICmdList
 {
 	LLM_SCOPE_DYNAMIC_STAT_OBJECTPATH(Owner->GetPackage(), ELLMTagSet::Assets);
 	LLM_SCOPE_DYNAMIC_STAT_OBJECTPATH_FNAME(FRHITextureCreateDesc().GetTraceClassName(), ELLMTagSet::AssetClasses);
-	UE_TRACE_METADATA_SCOPE_ASSET_FNAME(NAME_None, NAME_None, Owner->GetPackage()->GetFName());
+	UE_TRACE_METADATA_SCOPE_ASSET_FNAME(NAME_None, FRHITextureCreateDesc().GetTraceClassName(), Owner->GetPackage()->GetFName());
 
 	if((Owner->SizeX > 0) && (Owner->SizeY > 0) && (Owner->SizeZ > 0))
 	{
@@ -310,7 +310,7 @@ void FTextureRenderTargetVolumeResource::UpdateDeferredResource(FRHICommandListI
 {
 	LLM_SCOPE_DYNAMIC_STAT_OBJECTPATH(Owner->GetPackage(), ELLMTagSet::Assets);
 	LLM_SCOPE_DYNAMIC_STAT_OBJECTPATH_FNAME(FRHITextureCreateDesc().GetTraceClassName(), ELLMTagSet::AssetClasses);
-	UE_TRACE_METADATA_SCOPE_ASSET_FNAME(NAME_None, NAME_None, Owner->GetPackage()->GetFName());
+	UE_TRACE_METADATA_SCOPE_ASSET_FNAME(NAME_None, FRHITextureCreateDesc().GetTraceClassName(), Owner->GetPackage()->GetFName());
 
 	RemoveFromDeferredUpdateList();
 
diff --git a/Engine/Source/Runtime/Renderer/Private/RendererScene.cpp b/Engine/Source/Runtime/Renderer/Private/RendererScene.cpp
--- a/Engine/Source/Runtime/Renderer/Private/RendererScene.cpp
+++ b/Engine/Source/Runtime/Renderer/Private/RendererScene.cpp
@@ -1374,7 +1374,8 @@ void FScene::BatchAddPrimitivesInternal(TArrayView<T*> InPrimitives)
 #endif
 	LLM_SCOPE_DYNAMIC_STAT_OBJECTPATH_FNAME(ToUObject(InPrimitives[0]) ? InPrimitives[0]->GetOutermost()->GetFName() : NAME_None, ELLMTagSet::Assets);
 	LLM_SCOPE_DYNAMIC_STAT_OBJECTPATH_FNAME(ToUObject(InPrimitives[0]) ? FName(UPackage::StaticClass()->GetPathName()) : NAME_None, ELLMTagSet::AssetClasses);
-	UE_TRACE_METADATA_SCOPE_ASSET_FNAME(NAME_None, NAME_None, ToUObject(InPrimitives[0]) ? InPrimitives[0]->GetOutermost()->GetFName() : NAME_None);
+	UE_TRACE_METADATA_SCOPE_ASSET_FNAME(NAME_None, ToUObject(InPrimitives[0]) ? FName(UPackage::StaticClass()->GetPathName()) : NAME_None,
+		ToUObject(InPrimitives[0]) ? InPrimitives[0]->GetOutermost()->GetFName() : NAME_None);
 
 	SCOPE_CYCLE_COUNTER(STAT_AddScenePrimitiveGT);
 
