Don't consider Displacement overrides as triggering static material permutations since they only drive associated primitive data.
Also remove Displacement overrides from being shown by MaterialAnalyzer where the intention is to show overrides that trigger static permutations.
#rnx
#Synced-CL 49849534

diff --git a/Engine/Plugins/Editor/MaterialAnalyzer/Source/Private/SAnalyzedMaterialNodeWidgetItem.cpp b/Engine/Plugins/Editor/MaterialAnalyzer/Source/Private/SAnalyzedMaterialNodeWidgetItem.cpp
--- a/Engine/Plugins/Editor/MaterialAnalyzer/Source/Private/SAnalyzedMaterialNodeWidgetItem.cpp
+++ b/Engine/Plugins/Editor/MaterialAnalyzer/Source/Private/SAnalyzedMaterialNodeWidgetItem.cpp
@@ -190,8 +190,7 @@ TSharedRef<SWidget> SBasePropertyOverrideWidget::CreateRowWidget(FBasePropertyOv
 {
 	FText DisplayText = FText::GetEmpty();
 
-	if (RowData->ParameterID.IsEqual(TEXT("bOverride_OpacityMaskClipValue")) ||
-		RowData->ParameterID.IsEqual(TEXT("bOverride_DisplacementScaling")))
+	if (RowData->ParameterID.IsEqual(TEXT("bOverride_OpacityMaskClipValue")))
 	{
 		DisplayText = FText::AsNumber(RowData->ParameterValue);
 	}
diff --git a/Engine/Plugins/Editor/MaterialAnalyzer/Source/Private/SMaterialAnalyzer.cpp b/Engine/Plugins/Editor/MaterialAnalyzer/Source/Private/SMaterialAnalyzer.cpp
--- a/Engine/Plugins/Editor/MaterialAnalyzer/Source/Private/SMaterialAnalyzer.cpp
+++ b/Engine/Plugins/Editor/MaterialAnalyzer/Source/Private/SMaterialAnalyzer.cpp
@@ -48,11 +48,6 @@ SMaterialAnalyzer::SMaterialAnalyzer()
 	BasePropertyOverrideNames.Add(TEXT("bOverride_OutputTranslucentVelocity"), TEXT("bOutputTranslucentVelocity"));
 	BasePropertyOverrideNames.Add(TEXT("bOverride_bHasPixelAnimation"), TEXT("bHasPixelAnimation"));
 	BasePropertyOverrideNames.Add(TEXT("bOverride_bEnableTessellation"), TEXT("bEnableTessellation"));
-	BasePropertyOverrideNames.Add(TEXT("bOverride_DisplacementScaling"), TEXT("DisplacementScaling"));
-	BasePropertyOverrideNames.Add(TEXT("bOverride_bEnableDisplacementFade"), TEXT("bEnableDisplacementFade"));
-	BasePropertyOverrideNames.Add(TEXT("bOverride_DisplacementFadeRange"), TEXT("DisplacementFadeRange"));
-	BasePropertyOverrideNames.Add(TEXT("bOverride_MaxWorldPositionOffsetDisplacement"), TEXT("MaxWorldPositionOffsetDisplacement"));
-	BasePropertyOverrideNames.Add(TEXT("bOverride_CompatibleWithLumenCardSharing"), TEXT("bCompatibleWithLumenCardSharing"));
 }
 
 SMaterialAnalyzer::~SMaterialAnalyzer()
@@ -1078,46 +1073,6 @@ void FAnalyzeMaterialTreeAsyncTask::DoWork()
 				bIsOverridden = CurrentMaterialInstance->BasePropertyOverrides.bOverride_bEnableTessellation;
 			}
 		}
-		else if (BasePropertyOverrideName.Key.IsEqual(TEXT("bOverride_DisplacementScaling")))
-		{
-			TempValue = CurrentMaterialInterface->GetDisplacementScaling().Magnitude;
-			if (CurrentMaterialInstance)
-			{
-				bIsOverridden = CurrentMaterialInstance->BasePropertyOverrides.bOverride_DisplacementScaling;
-			}
-		}
-		else if (BasePropertyOverrideName.Key.IsEqual(TEXT("bOverride_bEnableDisplacementFade")))
-		{
-			TempValue = CurrentMaterialInterface->IsDisplacementFadeEnabled();
-			if (CurrentMaterialInstance)
-			{
-				bIsOverridden = CurrentMaterialInstance->BasePropertyOverrides.bOverride_bEnableDisplacementFade;
-			}
-		}
-		else if (BasePropertyOverrideName.Key.IsEqual(TEXT("bOverride_DisplacementFadeRange")))
-		{
-			TempValue = CurrentMaterialInterface->GetDisplacementFadeRange().EndSizePixels;
-			if (CurrentMaterialInstance)
-			{
-				bIsOverridden = CurrentMaterialInstance->BasePropertyOverrides.bOverride_DisplacementFadeRange;
-			}
-		}
-		else if (BasePropertyOverrideName.Key.IsEqual(TEXT("bOverride_MaxWorldPositionOffsetDisplacement")))
-		{
-			TempValue = CurrentMaterialInterface->GetMaxWorldPositionOffsetDisplacement();
-			if (CurrentMaterialInstance)
-			{
-				bIsOverridden = CurrentMaterialInstance->BasePropertyOverrides.bOverride_MaxWorldPositionOffsetDisplacement;
-			}
-		}
-		else if (BasePropertyOverrideName.Key.IsEqual(TEXT("bOverride_CompatibleWithLumenCardSharing")))
-		{
-			TempValue = (float)CurrentMaterialInterface->IsCompatibleWithLumenCardSharing();
-			if (CurrentMaterialInstance)
-			{
-				bIsOverridden = CurrentMaterialInstance->BasePropertyOverrides.bOverride_CompatibleWithLumenCardSharing;
-			}
-		}
 
 		// Check the parent for this variable
 		FAnalyzedMaterialNodePtr Parent = CurrentMaterialNode->Parent;
diff --git a/Engine/Source/Runtime/Engine/Private/Materials/MaterialInstance.cpp b/Engine/Source/Runtime/Engine/Private/Materials/MaterialInstance.cpp
--- a/Engine/Source/Runtime/Engine/Private/Materials/MaterialInstance.cpp
+++ b/Engine/Source/Runtime/Engine/Private/Materials/MaterialInstance.cpp
@@ -4646,8 +4646,9 @@ void UMaterialInstance::UpdateStaticPermutation(const FStaticParameterSet& NewPa
 	const bool bHasBasePropertyOverrides = HasOverridenBaseProperties();
 
 	const bool bWantsStaticPermutationResource = Parent && (!CompareParameters.IsEmpty() || bHasBasePropertyOverrides);
+	const bool bRecacheShaders = (bHasStaticPermutationResource != bWantsStaticPermutationResource) || bParamsHaveChanged || (bBasePropertyOverridesHaveChanged && bWantsStaticPermutationResource) || bForceStaticPermutationUpdate;
 
-	if (bHasStaticPermutationResource != bWantsStaticPermutationResource || bParamsHaveChanged || (bBasePropertyOverridesHaveChanged && bWantsStaticPermutationResource) || bForceStaticPermutationUpdate)
+	if (bRecacheShaders)
 	{
 		// This will flush the rendering thread which is necessary before changing bHasStaticPermutationResource, since the RT is reading from that directly
 		FlushRenderingCommands();
@@ -4660,6 +4661,14 @@ void UMaterialInstance::UpdateStaticPermutation(const FStaticParameterSet& NewPa
 		CacheResourceShadersForRendering(EMaterialShaderPrecompileMode::None);
 		RecacheUniformExpressions(true);
 
+		ValidateStaticPermutationAllowed();
+	}
+
+	// Not all base property overrides trigger a static permutation but changing them still requires render state recreation.
+	const bool bAddToUpdateContext = bRecacheShaders || bBasePropertyOverridesHaveChanged;
+
+	if (bAddToUpdateContext)
+	{
 		if (MaterialUpdateContext != nullptr)
 		{
 			MaterialUpdateContext->AddMaterialInstance(this);
@@ -4670,8 +4679,6 @@ void UMaterialInstance::UpdateStaticPermutation(const FStaticParameterSet& NewPa
 			FMaterialUpdateContext LocalMaterialUpdateContext(FMaterialUpdateContext::EOptions::RecreateRenderStates);
 			LocalMaterialUpdateContext.AddMaterialInstance(this);
 		}
-
-		ValidateStaticPermutationAllowed();
 	}
 }
 
@@ -5136,6 +5143,7 @@ bool UMaterialInstance::HasOverridenBaseProperties() const
 		return false;
 	}
 
+	// Note that we only return true for BaseProperties that need to trigger a static permutation.
 	return
 		GetBlendMode() != Parent->GetBlendMode() ||
 		GetShadingModels() != Parent->GetShadingModels() ||
@@ -5146,11 +5154,7 @@ bool UMaterialInstance::HasOverridenBaseProperties() const
 		IsTranslucencyWritingVelocity() != Parent->IsTranslucencyWritingVelocity() ||
 		HasPixelAnimation() != Parent->HasPixelAnimation() ||
 		IsTessellationEnabled() != Parent->IsTessellationEnabled() ||
-		GetDisplacementScaling() != Parent->GetDisplacementScaling() ||
-		IsDisplacementFadeEnabled() != Parent->IsDisplacementFadeEnabled() ||
-		GetDisplacementFadeRange() != Parent->GetDisplacementFadeRange() ||
 		!FMath::IsNearlyEqual(GetOpacityMaskClipValue(), Parent->GetOpacityMaskClipValue()) ||
-		!FMath::IsNearlyEqual(GetMaxWorldPositionOffsetDisplacement(), Parent->GetMaxWorldPositionOffsetDisplacement()) ||
 		GetUsageFlags() != Parent->GetUsageFlags();
 }
 
@@ -5160,7 +5164,8 @@ FString UMaterialInstance::GetBasePropertyOverrideString() const
 	FString BasePropString;
 	if (HasOverridenBaseProperties())
 	{
-		BasePropString.Appendf(TEXT("bOverride_OpacityMaskClipValue_%d, "), FMath::IsNearlyEqual(GetOpacityMaskClipValue(), Parent->GetOpacityMaskClipValue()));
+		// Only append information for the BaseProperties that cause a static permutation and are checked in HasOverridenBaseProperties().
+		BasePropString.Appendf(TEXT("bOverride_OpacityMaskClipValue_%d, "), !FMath::IsNearlyEqual(GetOpacityMaskClipValue(), Parent->GetOpacityMaskClipValue()));
 		BasePropString.Appendf(TEXT("bOverride_BlendMode_%d, "), (GetBlendMode() != Parent->GetBlendMode()));
 		BasePropString.Appendf(TEXT("bOverride_ShadingModel_%d, "), (GetShadingModels() != Parent->GetShadingModels()));
 		BasePropString.Appendf(TEXT("bOverride_TwoSided_%d, "), (IsTwoSided() != Parent->IsTwoSided()));
@@ -5170,10 +5175,6 @@ FString UMaterialInstance::GetBasePropertyOverrideString() const
 		BasePropString.Appendf(TEXT("bOverride_OutputTranslucentVelocity_%d "), (IsTranslucencyWritingVelocity() != Parent->IsTranslucencyWritingVelocity()));
 		BasePropString.Appendf(TEXT("bOverride_bHasPixelAnimation_%d "), (HasPixelAnimation() != Parent->HasPixelAnimation()));
 		BasePropString.Appendf(TEXT("bOverride_bEnableTessellation_%d "), (IsTessellationEnabled() != Parent->IsTessellationEnabled()));
-		BasePropString.Appendf(TEXT("bOverride_DisplacementScaling_%d "), (GetDisplacementScaling() != Parent->GetDisplacementScaling()));
-		BasePropString.Appendf(TEXT("bOverride_bEnableDisplacementFade_%d "), (IsDisplacementFadeEnabled() != Parent->IsDisplacementFadeEnabled()));
-		BasePropString.Appendf(TEXT("bOverride_DisplacementFadeRange_%d "), (GetDisplacementFadeRange() != Parent->GetDisplacementFadeRange()));
-		BasePropString.Appendf(TEXT("bOverride_MaxWorldPositionOffsetDisplacement_%d "), (GetMaxWorldPositionOffsetDisplacement() != Parent->GetMaxWorldPositionOffsetDisplacement()));
 		BasePropString.Appendf(TEXT("bOverride_UsageFlags_%d "), (GetUsageFlags() != Parent->GetUsageFlags()));
 	}
 	return BasePropString;
