[UE Rendering] Adding 2 new instructions to the compiler used in the context of MeshPartition:
- VectorComponent() allowing to extract dynamically a component scalar from a vector
- BitfieldExtract() allowing to unpack a bitfield range from a full scalar 32 bit value

#Synced-CL 49811085

diff --git a/Engine/Source/Runtime/Engine/Private/Materials/HLSLMaterialTranslator.cpp b/Engine/Source/Runtime/Engine/Private/Materials/HLSLMaterialTranslator.cpp
--- a/Engine/Source/Runtime/Engine/Private/Materials/HLSLMaterialTranslator.cpp
+++ b/Engine/Source/Runtime/Engine/Private/Materials/HLSLMaterialTranslator.cpp
@@ -10668,6 +10668,67 @@ static FString ComponentMaskLWC(const FString& SourceString, EMaterialValueType
 	return Result;
 }
 
+int32 FHLSLMaterialTranslator::VectorComponent(int32 Vector, int32 Component)
+{
+	if ((Vector == INDEX_NONE) || (Component == INDEX_NONE))
+	{
+		return INDEX_NONE;
+	}
+	EMaterialValueType VectorType = GetParameterType(Vector);
+	EMaterialValueType ComponentType = GetParameterType(Component);
+	if (!(VectorType & (MCT_Float2 | MCT_Float3 | MCT_Float4 | MCT_UInt2 | MCT_UInt3 | MCT_UInt4)))
+	{
+		return Errorf(TEXT("Vector type is invalid (%s: %s) for VectorComponent instruction"), *GetParameterCode(Vector), DescribeType(VectorType));
+	}
+	if (!(ComponentType & (MCT_Float1 | MCT_UInt1)))
+	{
+		return Errorf(TEXT("Component type is invalid (%s: %s) for VectorComponent instruction"), *GetParameterCode(Component), DescribeType(ComponentType));
+	}
+
+	EMaterialValueType ResultType = (VectorType & MCT_Float ? MCT_Float1 : MCT_UInt1);
+
+	return AddCodeChunk(ResultType, TEXT("%s[int(%s)]"), *GetParameterCode(Vector), *GetParameterCode(Component));
+}
+
+int32 FHLSLMaterialTranslator::BitfieldExtract(int32 Value, int32 BitOffset, int32 BitCount)
+{
+	if ((Value == INDEX_NONE) || (BitOffset == INDEX_NONE) || (BitCount == INDEX_NONE))
+	{
+		return INDEX_NONE;
+	}
+	EMaterialValueType ValueType = GetParameterType(Value);
+	EMaterialValueType BitOffsetType = GetParameterType(BitOffset);
+	EMaterialValueType BitCountType = GetParameterType(BitCount);
+	if (!(ValueType & (MCT_Float1 | MCT_UInt1)))
+	{
+		return Errorf(TEXT("Value type is invalid (%s: %s) for BitfieldExtract instruction"), *GetParameterCode(Value), DescribeType(ValueType));
+	}
+	if (!(BitOffsetType & (MCT_Float1 | MCT_UInt1)))
+	{
+		return Errorf(TEXT("BitOffset type is invalid (%s: %s) for BitfieldExtract instruction"), *GetParameterCode(BitOffset), DescribeType(BitOffsetType));
+	}
+	if (!(BitCountType & (MCT_Float1 | MCT_UInt1)))
+	{
+		return Errorf(TEXT("BitCount type is invalid (%s: %s) for BitfieldExtract instruction"), *GetParameterCode(BitCount), DescribeType(BitCountType));
+	}
+
+	int32 Mask_IDX = AddCodeChunk(MCT_UInt1, TEXT("~(0xffffffff << uint(%s))"), *GetParameterCode(BitCount));
+
+	int32 IntValue_IDX = Value;
+	if (ValueType & MCT_Float1)
+	{
+		IntValue_IDX = AddCodeChunk(MCT_UInt1, TEXT("asuint(%s)"), *GetParameterCode(Value));
+	}
+
+	int32 IntBitOffset_IDX = BitOffset;
+	if (BitOffsetType & MCT_Float1)
+	{
+		IntBitOffset_IDX = AddCodeChunk(MCT_UInt1, TEXT("uint(%s)"), *GetParameterCode(BitOffset));
+	}
+
+	return AddCodeChunk(MCT_UInt1, TEXT("(%s >> %s) & %s"), *GetParameterCode(IntValue_IDX), *GetParameterCode(IntBitOffset_IDX), *GetParameterCode(Mask_IDX));
+}
+
 int32 FHLSLMaterialTranslator::ComponentMask(int32 Vector,bool R,bool G,bool B,bool A)
 {
 	if(Vector == INDEX_NONE)
diff --git a/Engine/Source/Runtime/Engine/Private/Materials/HLSLMaterialTranslator.h b/Engine/Source/Runtime/Engine/Private/Materials/HLSLMaterialTranslator.h
--- a/Engine/Source/Runtime/Engine/Private/Materials/HLSLMaterialTranslator.h
+++ b/Engine/Source/Runtime/Engine/Private/Materials/HLSLMaterialTranslator.h
@@ -1099,6 +1099,8 @@ protected:
 	virtual int32 Clamp(int32 X, int32 A, int32 B) override;
 	virtual int32 Saturate(int32 X) override;
 	virtual int32 ComponentMask(int32 Vector, bool R, bool G, bool B, bool A) override;
+	virtual int32 VectorComponent(int32 Vector, int32 Component) override;
+	virtual int32 BitfieldExtract(int32 Value, int32 BitOffset, int32 BitCount) override;
 	virtual int32 AppendVector(int32 A, int32 B) override;
 	virtual int32 HsvToRgb(int32 X) override;
 	virtual int32 RgbToHsv(int32 X) override;
diff --git a/Engine/Source/Runtime/Engine/Public/MaterialCompiler.h b/Engine/Source/Runtime/Engine/Public/MaterialCompiler.h
--- a/Engine/Source/Runtime/Engine/Public/MaterialCompiler.h
+++ b/Engine/Source/Runtime/Engine/Public/MaterialCompiler.h
@@ -566,6 +566,8 @@ public:
 	virtual int32 InvLerp(int32 X,int32 Y,int32 A) = 0;
 
 	virtual int32 ComponentMask(int32 Vector,bool R,bool G,bool B,bool A) = 0;
+	virtual int32 VectorComponent(int32 Vector, int32 Component) = 0;
+	virtual int32 BitfieldExtract(int32 Value, int32 BitOffset, int32 BitCount) = 0;
 	virtual int32 AppendVector(int32 A,int32 B) = 0;
 	virtual int32 TransformVector(EMaterialCommonBasis SourceCoordBasis, EMaterialCommonBasis DestCoordBasis, FTransformParameters& Parameters, int32 A) = 0;
 	virtual int32 TransformPosition(EMaterialCommonBasis SourceCoordBasis, EMaterialCommonBasis DestCoordBasis, FTransformParameters& Parameters, int32 A) = 0;
@@ -1149,6 +1151,8 @@ public:
 	virtual int32 InvLerp(int32 X, int32 Y, int32 A) override { return Compiler->InvLerp(X, Y, A); }
 
 	virtual int32 ComponentMask(int32 Vector, bool R, bool G, bool B, bool A) override { return Compiler->ComponentMask(Vector, R, G, B, A); }
+	virtual int32 VectorComponent(int32 Vector, int32 Component) override { return Compiler->VectorComponent(Vector, Component); }
+	virtual int32 BitfieldExtract(int32 Value, int32 BitOffset, int32 BitCount) override { return Compiler->BitfieldExtract(Value, BitOffset, BitCount); }
 	virtual int32 AppendVector(int32 A, int32 B) override { return Compiler->AppendVector(A, B); }
 	virtual int32 TransformVector(EMaterialCommonBasis SourceCoordBasis, EMaterialCommonBasis DestCoordBasis, FTransformParameters& Parameters, int32 A) override
 	{
