[VisualLogger] updated trace output device to handle entries larger than the maximum trace event size

#Synced-CL 49921777

diff --git a/Engine/Plugins/Animation/GameplayInsights/Source/RewindDebuggerVLog/Private/VisualLoggerAnalyzer.cpp b/Engine/Plugins/Animation/GameplayInsights/Source/RewindDebuggerVLog/Private/VisualLoggerAnalyzer.cpp
--- a/Engine/Plugins/Animation/GameplayInsights/Source/RewindDebuggerVLog/Private/VisualLoggerAnalyzer.cpp
+++ b/Engine/Plugins/Animation/GameplayInsights/Source/RewindDebuggerVLog/Private/VisualLoggerAnalyzer.cpp
@@ -17,7 +17,12 @@ void FVisualLoggerAnalyzer::OnAnalysisBegin(const FOnAnalysisContext& Context)
 {
 	auto& Builder = Context.InterfaceBuilder;
 
+PRAGMA_DISABLE_DEPRECATION_WARNINGS
 	Builder.RouteEvent(RouteId_VisualLogEntry, "VisualLogger", "VisualLogEntry");
+PRAGMA_ENABLE_DEPRECATION_WARNINGS
+
+	Builder.RouteEvent(RouteId_VisualLogEntryHeader, "VisualLogger", "VisualLogEntryHeader");
+	Builder.RouteEvent(RouteId_VisualLogEntryChunk, "VisualLogger", "VisualLogEntryChunk");
 }
 
 bool FVisualLoggerAnalyzer::OnEvent(uint16 RouteId, EStyle Style, const FOnEventContext& Context)
@@ -29,6 +34,34 @@ bool FVisualLoggerAnalyzer::OnEvent(uint16 RouteId, EStyle Style, const FOnEvent
 	const auto& EventData = Context.EventData;
 	switch (RouteId)
 	{
+		case RouteId_VisualLogEntryHeader:
+		{
+			FVLogEntry& LogEntry = Provider.AddLogEntry();
+			LogEntry.EntryID = EventData.GetValue<uint32>("ID");
+			const uint64 Cycle = EventData.GetValue<uint64>("Cycle");
+			LogEntry.RecordingTime = Context.EventTime.AsSeconds(Cycle);
+			LogEntry.RecordingWorldTime = EventData.GetValue<double>("RecordingTime");
+			LogEntry.OwnerID = EventData.GetValue<uint64>("OwnerID");
+			LogEntry.ChunkNum = EventData.GetValue<uint16>("TotalChunkNum");
+			LogEntry.Size = EventData.GetValue<uint32>("Size");
+			LogEntry.Data.Reserve(LogEntry.Size);
+
+			break;
+		}
+
+		case RouteId_VisualLogEntryChunk:
+		{
+			const uint32 ID = EventData.GetValue<uint32>("ID");
+			const uint16 ChunkNum = EventData.GetValue<uint16>("ChunkNum");
+			const uint16 Size = EventData.GetValue<uint16>("Size");
+			TArrayView<const uint8> Data = EventData.GetArrayView<uint8>("Data");
+
+			Provider.AddLogEntryChunk(ID, ChunkNum, Size, Data);
+
+			break;
+		}
+
+PRAGMA_DISABLE_DEPRECATION_WARNINGS
 		case RouteId_VisualLogEntry:
 		{
 			uint64 OwnerId = EventData.GetValue<uint64>("OwnerId");
@@ -42,6 +75,7 @@ bool FVisualLoggerAnalyzer::OnEvent(uint16 RouteId, EStyle Style, const FOnEvent
 			Provider.AppendVisualLogEntry(OwnerId, Context.EventTime.AsSeconds(Cycle), Entry);
 			break;
 		}
+PRAGMA_ENABLE_DEPRECATION_WARNINGS
 	}
 
 	return true;
diff --git a/Engine/Plugins/Animation/GameplayInsights/Source/RewindDebuggerVLog/Private/VisualLoggerAnalyzer.h b/Engine/Plugins/Animation/GameplayInsights/Source/RewindDebuggerVLog/Private/VisualLoggerAnalyzer.h
--- a/Engine/Plugins/Animation/GameplayInsights/Source/RewindDebuggerVLog/Private/VisualLoggerAnalyzer.h
+++ b/Engine/Plugins/Animation/GameplayInsights/Source/RewindDebuggerVLog/Private/VisualLoggerAnalyzer.h
@@ -19,7 +19,9 @@ public:
 private:
 	enum : uint16
 	{
-		RouteId_VisualLogEntry,
+		RouteId_VisualLogEntry UE_DEPRECATED(5.8, "Used only to read old recordings; now using header + chunk to handle large events"),
+		RouteId_VisualLogEntryHeader,
+		RouteId_VisualLogEntryChunk,
 	};
 
 	TraceServices::IAnalysisSession& Session;
diff --git a/Engine/Plugins/Animation/GameplayInsights/Source/RewindDebuggerVLog/Private/VisualLoggerProvider.cpp b/Engine/Plugins/Animation/GameplayInsights/Source/RewindDebuggerVLog/Private/VisualLoggerProvider.cpp
--- a/Engine/Plugins/Animation/GameplayInsights/Source/RewindDebuggerVLog/Private/VisualLoggerProvider.cpp
+++ b/Engine/Plugins/Animation/GameplayInsights/Source/RewindDebuggerVLog/Private/VisualLoggerProvider.cpp
@@ -1,6 +1,7 @@
 // Copyright Epic Games, Inc. All Rights Reserved.
 
 #include "VisualLoggerProvider.h"
+#include "Serialization/MemoryReader.h"
 
 FName FVisualLoggerProvider::ProviderName("VisualLoggerProvider");
 
@@ -31,7 +32,7 @@ bool FVisualLoggerProvider::ReadVisualLogEntryTimeline(uint64 InObjectId, TFunct
 void FVisualLoggerProvider::EnumerateCategories(TFunctionRef<void(const FName&)> Callback) const
 {
 	Session.ReadAccessCheck();
-	
+
 	for(const FName& Category : Categories)
 	{
 		Callback(Category);
@@ -56,16 +57,46 @@ void  FVisualLoggerProvider::AppendVisualLogEntry(uint64 InObjectId, double InTi
 	}
 
 	Timeline->AppendEvent(InTime, Entry);
-	
+
 	for (const FVisualLogLine& Line : Entry.LogLines)
 	{
 		Categories.AddUnique(Line.Category);
 	}
-	
+
 	for (const FVisualLogShapeElement& Shape : Entry.ElementsToDraw)
 	{
 		Categories.AddUnique(Shape.Category);
 	}
 }
 
+FVLogEntry& FVisualLoggerProvider::AddLogEntry()
+{
+	return LogEntries.AddDefaulted_GetRef();
+}
+
+void FVisualLoggerProvider::AddLogEntryChunk(const uint32 ID, const uint32 ChunkNum, const uint16 Size, const TArrayView<const uint8>& ChunkData)
+{
+	const int32 Index = LogEntries.IndexOfByPredicate([ID](const FVLogEntry& Entry) { return Entry.EntryID == ID; });
+	if (Index == INDEX_NONE)
+	{
+		return;
+	}
+
+	FVLogEntry& LogEntry = LogEntries[Index];
+	constexpr uint32 ChunkSize = TNumericLimits<uint16>::Max();
+	check(LogEntry.Data.Num() == ChunkNum * ChunkSize);
+
+	LogEntry.Data.Append(ChunkData.GetData(), Size);
+
+	if (LogEntry.Size == LogEntry.Data.Num())
+	{
+		FMemoryReaderView Archive(LogEntry.Data);
+		FVisualLogEntry Entry;
+		Archive << Entry;
+		Entry.TimeStamp = LogEntry.RecordingWorldTime;
+		AppendVisualLogEntry(LogEntry.OwnerID, LogEntry.RecordingTime, Entry);
+		LogEntries.RemoveAtSwap(Index);
+	}
+}
+
 #undef LOCTEXT_NAMESPACE
diff --git a/Engine/Plugins/Animation/GameplayInsights/Source/RewindDebuggerVLog/Private/VisualLoggerProvider.h b/Engine/Plugins/Animation/GameplayInsights/Source/RewindDebuggerVLog/Private/VisualLoggerProvider.h
--- a/Engine/Plugins/Animation/GameplayInsights/Source/RewindDebuggerVLog/Private/VisualLoggerProvider.h
+++ b/Engine/Plugins/Animation/GameplayInsights/Source/RewindDebuggerVLog/Private/VisualLoggerProvider.h
@@ -1,4 +1,4 @@
-// Copyright Epic Games, Inc. All Rights Reserved.
+﻿// Copyright Epic Games, Inc. All Rights Reserved.
 
 #pragma once
 
@@ -18,6 +18,17 @@ struct FVLogTimelineSettings
 	};
 };
 
+struct FVLogEntry
+{
+	uint32 EntryID = INDEX_NONE;
+	uint32 Size = 0;
+	uint64 OwnerID = INDEX_NONE;
+	double RecordingTime = 0;
+	double RecordingWorldTime = 0;
+	TArray<uint8> Data;
+	uint16 ChunkNum = 0;
+};
+
 class FVisualLoggerProvider : public IVisualLoggerProvider
 {
 public:
@@ -25,17 +36,20 @@ public:
 
 	explicit FVisualLoggerProvider(TraceServices::IAnalysisSession& InSession);
 
-	/** IVisualLoggerProvider interface */
-	virtual bool ReadVisualLogEntryTimeline(uint64 InObjectId, TFunctionRef<void(const VisualLogEntryTimeline&)> Callback) const;
-
+	//~ IVisualLoggerProvider interface
+	virtual bool ReadVisualLogEntryTimeline(uint64 InObjectId, TFunctionRef<void(const VisualLogEntryTimeline&)> Callback) const override;
 	virtual void EnumerateCategories(TFunctionRef<void(const FName&)> Callback) const override;
 
 	/** Add an object event message */
 	void AppendVisualLogEntry(uint64 InObjectId, double InTime, const FVisualLogEntry& Entry);
 
+	FVLogEntry& AddLogEntry();
+	void AddLogEntryChunk(uint32 ID, uint32 ChunkNum, uint16 Size, const TArrayView<const uint8>& ChunkData);
+
 private:
 	TraceServices::IAnalysisSession& Session;
 
+	TArray<FVLogEntry> LogEntries;
 	TMap<uint64, uint32> ObjectIdToLogEntryTimelines;
 	TArray<TSharedRef<TraceServices::TPointTimeline<FVisualLogEntry, FVLogTimelineSettings>>> LogEntryTimelines;
 	TArray<FName> Categories;
diff --git a/Engine/Source/Runtime/Engine/Private/VisualLogger/VisualLoggerTraceDevice.cpp b/Engine/Source/Runtime/Engine/Private/VisualLogger/VisualLoggerTraceDevice.cpp
--- a/Engine/Source/Runtime/Engine/Private/VisualLogger/VisualLoggerTraceDevice.cpp
+++ b/Engine/Source/Runtime/Engine/Private/VisualLogger/VisualLoggerTraceDevice.cpp
@@ -1,4 +1,4 @@
-// Copyright Epic Games, Inc. All Rights Reserved.
+﻿// Copyright Epic Games, Inc. All Rights Reserved.
 #include "VisualLogger/VisualLoggerTraceDevice.h"
 
 #if UE_DEBUG_RECORDING_ENABLED
@@ -11,11 +11,20 @@
 
 UE_TRACE_MINIMAL_CHANNEL_DEFINE(VisualLoggerChannel);
 
-UE_TRACE_MINIMAL_EVENT_BEGIN(VisualLogger, VisualLogEntry)
+UE_TRACE_MINIMAL_EVENT_BEGIN(VisualLogger, VisualLogEntryHeader)
+	UE_TRACE_MINIMAL_EVENT_FIELD(uint32, ID)
 	UE_TRACE_MINIMAL_EVENT_FIELD(uint64, Cycle)
 	UE_TRACE_MINIMAL_EVENT_FIELD(double, RecordingTime)
-	UE_TRACE_MINIMAL_EVENT_FIELD(uint64, OwnerId)
-	UE_TRACE_MINIMAL_EVENT_FIELD(uint8[], LogEntry)
+	UE_TRACE_MINIMAL_EVENT_FIELD(uint64, OwnerID)
+	UE_TRACE_MINIMAL_EVENT_FIELD(uint16, TotalChunkNum)
+	UE_TRACE_MINIMAL_EVENT_FIELD(uint32, Size)
+UE_TRACE_MINIMAL_EVENT_END()
+
+UE_TRACE_MINIMAL_EVENT_BEGIN(VisualLogger, VisualLogEntryChunk)
+	UE_TRACE_MINIMAL_EVENT_FIELD(uint32, ID)
+	UE_TRACE_MINIMAL_EVENT_FIELD(uint16, ChunkNum)
+	UE_TRACE_MINIMAL_EVENT_FIELD(uint16, Size)
+	UE_TRACE_MINIMAL_EVENT_FIELD(uint8[], Data)
 UE_TRACE_MINIMAL_EVENT_END()
 
 FVisualLoggerTraceDevice& FVisualLoggerTraceDevice::Get()
@@ -71,11 +80,37 @@ void FVisualLoggerTraceDevice::Serialize(const UObject* InLogOwner, const FName&
 
 		Archive << const_cast<FVisualLogEntry&>(InLogEntry);
 
-		UE_TRACE_MINIMAL_LOG(VisualLogger, VisualLogEntry, VisualLoggerChannel)
-			<< VisualLogEntry.Cycle(FPlatformTime::Cycles64())
-			<< VisualLogEntry.RecordingTime(FObjectTrace::GetWorldElapsedTime(InLogOwner->GetWorld()))
-			<< VisualLogEntry.OwnerId(FObjectTrace::GetObjectId(InLogOwner))
-			<< VisualLogEntry.LogEntry(Archive.GetData(), Archive.Num());
+		static std::atomic<uint32> EntryID = 0;
+		const uint32 ID = EntryID.fetch_add(1);
+
+		const uint32 DataSize = static_cast<uint32>(Archive.Num());
+		constexpr uint32 MaxChunkSize = TNumericLimits<uint16>::Max();
+		const uint16 ChunkNum = (DataSize + MaxChunkSize - 1) / MaxChunkSize;
+
+		UE_TRACE_MINIMAL_LOG(VisualLogger, VisualLogEntryHeader, VisualLoggerChannel)
+			<< VisualLogEntryHeader.ID(ID)
+			<< VisualLogEntryHeader.Cycle(FPlatformTime::Cycles64())
+			<< VisualLogEntryHeader.RecordingTime(FObjectTrace::GetWorldElapsedTime(InLogOwner->GetWorld()))
+			<< VisualLogEntryHeader.OwnerID(FObjectTrace::GetObjectId(InLogOwner))
+			<< VisualLogEntryHeader.TotalChunkNum(ChunkNum)
+			<< VisualLogEntryHeader.Size(DataSize);
+
+		uint32 RemainingSize = DataSize;
+		for (uint32 Index = 0; Index < ChunkNum; ++Index)
+		{
+			const uint16 Size = static_cast<uint16>(FMath::Min(RemainingSize, MaxChunkSize));
+			uint8* ChunkData = Archive.GetData() + MaxChunkSize * Index;
+
+			UE_TRACE_MINIMAL_LOG(VisualLogger, VisualLogEntryChunk, VisualLoggerChannel)
+				<< VisualLogEntryChunk.ID(ID)
+				<< VisualLogEntryChunk.ChunkNum(Index)
+				<< VisualLogEntryChunk.Size(Size)
+				<< VisualLogEntryChunk.Data(ChunkData, Size);
+
+			RemainingSize -= Size;
+		}
+
+		check(RemainingSize == 0);
 	}
 #endif
 
