Substrate - Toon - Added common toon profile definitions to share data between shader and cpp.

#tests pc editor dx12
#rnx
#Synced-CL 49923192

diff --git a/Engine/Shaders/Private/ToonProfileCommon.ush b/Engine/Shaders/Private/ToonProfileCommon.ush
--- a/Engine/Shaders/Private/ToonProfileCommon.ush
+++ b/Engine/Shaders/Private/ToonProfileCommon.ush
@@ -2,10 +2,7 @@
 
 #pragma once
 
-#define TOON_PROFILE_DIFFUSE_RAMP_Y		0
-#define TOON_PROFILE_SPECULAR_RAMP_Y	1
-#define TOON_PROFILE_PARAMETERS_Y		2
-#define TOON_PROFILE_DIFFUSE_DITHER_Y	3
+#include "/Engine/Shared/ToonProfileDefinitions.h"
 
 float3 EvaluateToonProfileRamp(in float ToonProfileId, in float NoL, in float TOON_PROFILE_RAMP)
 {
diff --git a/Engine/Shaders/Shared/ToonProfileDefinitions.h b/Engine/Shaders/Shared/ToonProfileDefinitions.h
new file mode 100644
--- /dev/null
+++ b/Engine/Shaders/Shared/ToonProfileDefinitions.h
@@ -0,0 +1,8 @@
+// Copyright Epic Games, Inc. All Rights Reserved.
+
+#pragma once
+
+#define TOON_PROFILE_DIFFUSE_RAMP_Y		0
+#define TOON_PROFILE_SPECULAR_RAMP_Y	1
+#define TOON_PROFILE_PARAMETERS_Y		2
+#define TOON_PROFILE_DIFFUSE_DITHER_Y	3
diff --git a/Engine/Source/Runtime/Engine/Private/Rendering/ToonProfile.cpp b/Engine/Source/Runtime/Engine/Private/Rendering/ToonProfile.cpp
--- a/Engine/Source/Runtime/Engine/Private/Rendering/ToonProfile.cpp
+++ b/Engine/Source/Runtime/Engine/Private/Rendering/ToonProfile.cpp
@@ -1,6 +1,7 @@
 // Copyright Epic Games, Inc. All Rights Reserved.
 
 #include "Engine/ToonProfile.h"
+#include "ToonProfileDefinitions.h"
 #include "Engine/Texture2D.h"
 #include "DataDrivenShaderPlatformInfo.h"
 #include "Math/Float16.h"
@@ -307,7 +308,7 @@ IPooledRenderTarget* FToonProfileTextureManager::GetAtlasTexture(FRDGBuilder& Gr
 
 		// See ToonProfileCommon.ush
 		const uint32 ResolutionX = ToonProfileResolution;
-		const uint32 ResolutionY =  3 + ToonProfileResolution; // row for ramps, parameters and 2d dithering patterns. // SUBSTRATE_TODO shared offset header
+		const uint32 ResolutionY = TOON_PROFILE_DIFFUSE_DITHER_Y + ToonProfileResolution; // row for ramps, parameters and 2d dithering patterns. // SUBSTRATE_TODO shared offset header
 
 		// 1. Create atlas texture
 		const FIntPoint ToonProfileResolution2D = FIntPoint(ToonProfileResolution, ResolutionY);
