[World Partition]

Tick the engine during minimap capture warmup frame cycles.

#Synced-CL 49921777
#Synced-CL 49921777

diff --git a/Engine/Source/Runtime/Engine/Private/WorldPartition/WorldPartitionMiniMapHelper.cpp b/Engine/Source/Runtime/Engine/Private/WorldPartition/WorldPartitionMiniMapHelper.cpp
--- a/Engine/Source/Runtime/Engine/Private/WorldPartition/WorldPartitionMiniMapHelper.cpp
+++ b/Engine/Source/Runtime/Engine/Private/WorldPartition/WorldPartitionMiniMapHelper.cpp
@@ -18,6 +18,7 @@
 #include "AssetCompilingManager.h"
 #include "Math/OrthoMatrix.h"
 #include "WorldPartition/WorldPartitionMiniMap.h"
+#include "WorldPartition/WorldPartitionHelpers.h"
 
 DEFINE_LOG_CATEGORY_STATIC(LogWorldPartitionMiniMapHelper, All, All);
 
@@ -108,6 +109,7 @@ void FWorldPartitionMiniMapHelper::CaptureBoundsMiniMapToTexture(UWorld* InWorld
 	for (uint32 i = 0; i < InNumWarmupFrames; i++)
 	{
 		WarmupTickCallback();
+		FWorldPartitionHelpers::FakeEngineTick(InWorld);
 		CaptureComponent->CaptureScene();
 	}
 
