Substrate - toon - added specular dither pattern.

#tests PC editor dx12
#rnx

#Synced-CL 49923192

diff --git a/Engine/Shaders/Private/Substrate/SubstrateEvaluation.ush b/Engine/Shaders/Private/Substrate/SubstrateEvaluation.ush
--- a/Engine/Shaders/Private/Substrate/SubstrateEvaluation.ush
+++ b/Engine/Shaders/Private/Substrate/SubstrateEvaluation.ush
@@ -1326,7 +1326,8 @@ FSubstrateEvaluateResult SubstrateEvaluateBSDFCommon(FSubstrateBSDFContext BSDFC
 
 			const float2 OffsetFraction = float2(0.0f, 0.0f);
 			const float2 ScreenUV = (BSDFContext.PixelCoord - View.ViewRectMinAndSize.xy) * View.ViewSizeAndInvSize.zw;
-			float ToonProfileDiffuseDither = EvaluateToonProfileDiffuseDither(ToonProfileId, ScreenUV, ToonEvalSettings);
+			float ToonProfileDiffuseDither  = EvaluateToonProfileDiffuseDither(ToonProfileId, ScreenUV, ToonEvalSettings);
+			float ToonProfileSpecularDither = EvaluateToonProfileSpecularDither(ToonProfileId, ScreenUV, ToonEvalSettings);
 
 			// Diffuse
 			Sample.DiffuseColor				= DiffuseColor * EvaluateToonProfileDiffuseRamp(ToonProfileId, saturate(BSDFContext.Context.NoL + ToonProfileDiffuseDither)); // AreaLight.NoL to handle capsule shape?
@@ -1348,7 +1349,7 @@ FSubstrateEvaluateResult SubstrateEvaluateBSDFCommon(FSubstrateBSDFContext BSDFC
 			// TOON_TODO experiment with GGX alike lob that is normalize according to peak lobe?
 			// TOON_TODO Area light, the roughness expand trick here does not make sense to use. MRP?
 
-			const float3 SpecularRamp = SpecularColor * EvaluateToonProfileSpecularRamp(ToonProfileId, Specular);
+			const float3 SpecularRamp = SpecularColor * EvaluateToonProfileSpecularRamp(ToonProfileId, saturate(Specular + ToonProfileSpecularDither));
 			Sample.SpecularPathValue = SpecularRamp;
 			Sample.IntegratedSpecularValue += ShadowTerms.SurfaceShadow * AreaLight.Falloff * AreaLight.FalloffColor * Sample.SpecularPathValue;
 
diff --git a/Engine/Shaders/Private/ToonProfile.usf b/Engine/Shaders/Private/ToonProfile.usf
--- a/Engine/Shaders/Private/ToonProfile.usf
+++ b/Engine/Shaders/Private/ToonProfile.usf
@@ -12,11 +12,17 @@ uint ToonProfileIndex;
 float DiffuseDitherPatternStrength;
 float DiffuseDitherPatternSize;
 
+float SpecularDitherPatternStrength;
+float SpecularDitherPatternSize;
+
 RWTexture2DArray<float4>	ToonProfileTexture;
 
 Texture2D<float4>			DiffuseDitherPattern;
 SamplerState				DiffuseDitherPatternSampler;
 
+Texture2D<float4>			SpecularDitherPattern;
+SamplerState				SpecularDitherPatternSampler;
+
 Buffer<float4>				DiffuseRampBufferSRV;
 Buffer<float4>				SpecularRampBufferSRV;
 
@@ -43,14 +49,14 @@ void ToonProfileUpdateAtlasTileCS(uint3 DispatchThreadId : SV_DispatchThreadID)
 		{
 			if (PixelPos.x == 0)
 			{
-				ToonProfileTexture[uint3(PixelPos, ToonProfileIndex)] = float4(DiffuseDitherPatternStrength, DiffuseDitherPatternSize, 0.0f, 0.0f);
+				ToonProfileTexture[uint3(PixelPos, ToonProfileIndex)] = float4(DiffuseDitherPatternStrength, DiffuseDitherPatternSize, SpecularDitherPatternStrength, SpecularDitherPatternSize);
 			}
 			else
 			{
 				ToonProfileTexture[uint3(PixelPos, ToonProfileIndex)] = DefaultColor;
 			}
 		}
-		else if (PixelPos.y >= TOON_PROFILE_DIFFUSE_DITHER_Y && PixelPos.y <= uint(TOON_PROFILE_DIFFUSE_DITHER_Y + ProfileResolution))
+		else if (PixelPos.y >= TOON_PROFILE_DITHERPATTERN_Y && PixelPos.y <= uint(TOON_PROFILE_DITHERPATTERN_Y + ProfileResolution))
 		{
 		/*
 			For a Toon profile of full resolution N, we render an edge of 1 texel all around to make sure the texture is repeatable, marked `e` below.
@@ -75,7 +81,7 @@ void ToonProfileUpdateAtlasTileCS(uint3 DispatchThreadId : SV_DispatchThreadID)
 			| e | e | e | e | e | e | e | e |
 		0,1 --------------------------------- 1,1
 		*/
-			const float2 FullOrigin = float2(0.0f, TOON_PROFILE_DIFFUSE_DITHER_Y);
+			const float2 FullOrigin = float2(0.0f, TOON_PROFILE_DITHERPATTERN_Y);
 			const float2 BakeOrigin = FullOrigin + 1.0f;
 
 			const float  BakeResolution = ProfileResolution - 2.0f;
@@ -84,7 +90,9 @@ void ToonProfileUpdateAtlasTileCS(uint3 DispatchThreadId : SV_DispatchThreadID)
 			const float2 FullCoord  = float2(PixelPos.xy) + 0.5f.xx;
 			const float2 FullUVs    = (FullCoord - BakeOrigin) * BakeResolutionInv;
 
-			ToonProfileTexture[uint3(PixelPos, ToonProfileIndex)] = DiffuseDitherPattern.SampleLevel(DiffuseDitherPatternSampler, FullUVs, 0);
+			const float DiffuseDitherPatternValue  = DiffuseDitherPattern.SampleLevel (DiffuseDitherPatternSampler,  FullUVs, 0.0f).x;
+			const float SpecularDitherPatternValue = SpecularDitherPattern.SampleLevel(SpecularDitherPatternSampler, FullUVs, 0.0f).y;
+			ToonProfileTexture[uint3(PixelPos, ToonProfileIndex)] = float4(DiffuseDitherPatternValue, SpecularDitherPatternValue, 0.0f, 0.0f);
 		}
 		else
 		{
diff --git a/Engine/Shaders/Private/ToonProfileCommon.ush b/Engine/Shaders/Private/ToonProfileCommon.ush
--- a/Engine/Shaders/Private/ToonProfileCommon.ush
+++ b/Engine/Shaders/Private/ToonProfileCommon.ush
@@ -29,6 +29,8 @@ struct FToonProfileEvaluateSettings
 {
 	float DiffuseDitherPatternStrength;
 	float DiffuseDitherPatternSize;
+	float SpecularDitherPatternStrength;
+	float SpecularDitherPatternSize;
 };
 
 FToonProfileEvaluateSettings GetToonProfileEvaluateSettings(in float ToonProfileId)
@@ -40,8 +42,10 @@ FToonProfileEvaluateSettings GetToonProfileEvaluateSettings(in float ToonProfile
 	float4 Data = View.ToonProfileTexture.SampleLevel(View.ToonProfileSampler, ToonProfileCoord, 0);
 
 	FToonProfileEvaluateSettings Ret;
-	Ret.DiffuseDitherPatternStrength = Data.x;
-	Ret.DiffuseDitherPatternSize = Data.y;
+	Ret.DiffuseDitherPatternStrength	= Data.x;
+	Ret.DiffuseDitherPatternSize		= Data.y;
+	Ret.SpecularDitherPatternStrength	= Data.z;
+	Ret.SpecularDitherPatternSize		= Data.w;
 	return Ret;
 }
 
@@ -55,10 +59,10 @@ float EvaluateToonProfileDiffuseDither(in float ToonProfileId, in float2 ScreenU
 
 		// See ToonProfile.usf for the baking/sampling description.
 
-		const float2 FullOrigin = float2(0.0f, TOON_PROFILE_DIFFUSE_DITHER_Y);
+		const float2 FullOrigin = float2(0.0f, TOON_PROFILE_DITHERPATTERN_Y);
 		const float2 BakeOrigin = FullOrigin + 1.0f;
 		
-		const float2 BakeCoords        = BakeOrigin + UVs * View.ToonProfileDiffuseDitherPatternSizeAndInvSize.x;
+		const float2 BakeCoords        = BakeOrigin + UVs * View.ToonProfileDitherPatternSizeAndInvSize.x;
 
 		const float3 ToonProfileCoord = float3(
 			BakeCoords.xy * View.ToonProfileTextureSizeAndInvSize.zw,
@@ -70,3 +74,28 @@ float EvaluateToonProfileDiffuseDither(in float ToonProfileId, in float2 ScreenU
 	return 0.0f;
 }
 
+float EvaluateToonProfileSpecularDither(in float ToonProfileId, in float2 ScreenUV, in FToonProfileEvaluateSettings ToonProfileEvaluateSettings)
+{
+	if (ToonProfileEvaluateSettings.SpecularDitherPatternStrength > 0.0f)
+	{
+		float2 UVs = ScreenUV * (1.0f / ToonProfileEvaluateSettings.SpecularDitherPatternSize);	// Project the pattern on screen as a percentage of it.
+		UVs.y *= View.ViewSizeAndInvSize.y * View.ViewSizeAndInvSize.z;							// Make sure the pattern square shape is enforced on screen.
+		UVs = frac(UVs);																		// Pattern wrap repetition.
+
+		// See ToonProfile.usf for the baking/sampling description.
+
+		const float2 FullOrigin = float2(0.0f, TOON_PROFILE_DITHERPATTERN_Y);
+		const float2 BakeOrigin = FullOrigin + 1.0f;
+		
+		const float2 BakeCoords = BakeOrigin + UVs * View.ToonProfileDitherPatternSizeAndInvSize.x;
+
+		const float3 ToonProfileCoord = float3(
+			BakeCoords.xy * View.ToonProfileTextureSizeAndInvSize.zw,
+			ToonProfileId
+		);
+
+		return ToonProfileEvaluateSettings.SpecularDitherPatternStrength * (View.ToonProfileTexture.SampleLevel(View.ToonProfileSampler, ToonProfileCoord, 0).y * 2.0f - 1.0f);
+	}
+	return 0.0f;
+}
+
diff --git a/Engine/Shaders/Shared/ToonProfileDefinitions.h b/Engine/Shaders/Shared/ToonProfileDefinitions.h
--- a/Engine/Shaders/Shared/ToonProfileDefinitions.h
+++ b/Engine/Shaders/Shared/ToonProfileDefinitions.h
@@ -5,4 +5,4 @@
 #define TOON_PROFILE_DIFFUSE_RAMP_Y		0
 #define TOON_PROFILE_SPECULAR_RAMP_Y	1
 #define TOON_PROFILE_PARAMETERS_Y		2
-#define TOON_PROFILE_DIFFUSE_DITHER_Y	3
+#define TOON_PROFILE_DITHERPATTERN_Y	3
diff --git a/Engine/Source/Runtime/Engine/Classes/Engine/ToonProfile.h b/Engine/Source/Runtime/Engine/Classes/Engine/ToonProfile.h
--- a/Engine/Source/Runtime/Engine/Classes/Engine/ToonProfile.h
+++ b/Engine/Source/Runtime/Engine/Classes/Engine/ToonProfile.h
@@ -54,6 +54,24 @@ struct FToonProfileStruct
 	UPROPERTY(Category = "Specular", EditAnywhere, meta = (AllowZoomOutput = "false", ShowZoomButtons = "false", ViewMinInput = "0", ViewMaxInput = "1", ViewMinOutput = "0", ViewMaxOutput = "1", TimelineLength = "1", ShowInputGridNumbers="false", ShowOutputGridNumbers="false"))
 	FRuntimeCurveLinearColor SpecularRamp;
 
+	/**
+	 * The specular texture dithering pattern. It is recommended for it to be grey scale level in [0, 1], internally remapped to [-1, 1] for positive and negative offset.
+	 */
+	UPROPERTY(Category = "Specular", EditAnywhere)
+	TObjectPtr<UTexture> SpecularDitherPattern;
+
+	/**
+	 * The strength of the offset applied by the dither pattern to the specular ramp.
+	 */
+	UPROPERTY(Category = "Specular", EditAnywhere, meta = (UIMin = 0.0, UIMax = 1.0, ClampMin = 0.0))
+	float SpecularDitherPatternStrength;
+
+	/**
+	 * The size on screen of the dithering pattern.
+	 */
+	UPROPERTY(Category = "Specular", EditAnywhere, meta = (UIMin = 0.0, UIMax = 1.0, ClampMin = 0.0))
+	float SpecularDitherPatternSize;
+
 	FToonProfileStruct();
 
 	void Invalidate()
diff --git a/Engine/Source/Runtime/Engine/Private/Rendering/ToonProfile.cpp b/Engine/Source/Runtime/Engine/Private/Rendering/ToonProfile.cpp
--- a/Engine/Source/Runtime/Engine/Private/Rendering/ToonProfile.cpp
+++ b/Engine/Source/Runtime/Engine/Private/Rendering/ToonProfile.cpp
@@ -63,12 +63,17 @@ public:
 		SHADER_PARAMETER(uint32, ToonProfileIndex)
 		SHADER_PARAMETER(float, DiffuseDitherPatternStrength)
 		SHADER_PARAMETER(float, DiffuseDitherPatternSize)
+		SHADER_PARAMETER(float, SpecularDitherPatternStrength)
+		SHADER_PARAMETER(float, SpecularDitherPatternSize)
 
 		SHADER_PARAMETER_RDG_TEXTURE_UAV(RWTexture2DArray, ToonProfileTexture)
 
 		SHADER_PARAMETER_RDG_TEXTURE(Texture2D, DiffuseDitherPattern)
 		SHADER_PARAMETER_SAMPLER(SamplerState, DiffuseDitherPatternSampler)
 
+		SHADER_PARAMETER_RDG_TEXTURE(Texture2D, SpecularDitherPattern)
+		SHADER_PARAMETER_SAMPLER(SamplerState, SpecularDitherPatternSampler)
+
 		SHADER_PARAMETER_RDG_BUFFER_SRV(Buffer, DiffuseRampBufferSRV)
 		SHADER_PARAMETER_RDG_BUFFER_SRV(Buffer, SpecularRampBufferSRV)
 	END_SHADER_PARAMETER_STRUCT()
@@ -308,7 +313,7 @@ IPooledRenderTarget* FToonProfileTextureManager::GetAtlasTexture(FRDGBuilder& Gr
 
 		// See ToonProfileCommon.ush
 		const uint32 ResolutionX = ToonProfileResolution;
-		const uint32 ResolutionY = TOON_PROFILE_DIFFUSE_DITHER_Y + ToonProfileResolution; // row for ramps, parameters and 2d dithering patterns. // SUBSTRATE_TODO shared offset header
+		const uint32 ResolutionY = TOON_PROFILE_DITHERPATTERN_Y + ToonProfileResolution; // row for ramps, parameters and 2d dithering patterns. // SUBSTRATE_TODO shared offset header
 
 		// 1. Create atlas texture
 		const FIntPoint ToonProfileResolution2D = FIntPoint(ToonProfileResolution, ResolutionY);
@@ -330,10 +335,15 @@ IPooledRenderTarget* FToonProfileTextureManager::GetAtlasTexture(FRDGBuilder& Gr
 		{
 			const FToonProfileStruct& Data = ToonProfileEntries[LayerIt].Settings;
 
-			FTextureResource* TextureResource = Data.DiffuseDitherPattern ? Data.DiffuseDitherPattern->GetResource() : nullptr;
-			FRHITexture* Texture2DRHI = TextureResource ? TextureResource->GetTexture2DRHI() : nullptr;
-			bool bHasDiffuseDitherPaternTexture = Texture2DRHI != nullptr;
-			Texture2DRHI = Texture2DRHI ? Texture2DRHI : GBlackTexture->TextureRHI->GetTexture2D();
+			FTextureResource* DiffuseDitherPatternTextureResource = Data.DiffuseDitherPattern ? Data.DiffuseDitherPattern->GetResource() : nullptr;
+			FRHITexture* DiffuseDitherPatternTexture2DRHI = DiffuseDitherPatternTextureResource ? DiffuseDitherPatternTextureResource->GetTexture2DRHI() : nullptr;
+			bool bHasDiffuseDitherPaternTexture = DiffuseDitherPatternTexture2DRHI != nullptr;
+			DiffuseDitherPatternTexture2DRHI = DiffuseDitherPatternTexture2DRHI ? DiffuseDitherPatternTexture2DRHI : GBlackTexture->TextureRHI->GetTexture2D();
+
+			FTextureResource* SpecularDitherPatternTextureResource = Data.SpecularDitherPattern ? Data.SpecularDitherPattern->GetResource() : nullptr;
+			FRHITexture* SpecularDitherPatternTexture2DRHI = SpecularDitherPatternTextureResource ? SpecularDitherPatternTextureResource->GetTexture2DRHI() : nullptr;
+			bool bHasSpecularDitherPaternTexture = SpecularDitherPatternTexture2DRHI != nullptr;
+			SpecularDitherPatternTexture2DRHI = SpecularDitherPatternTexture2DRHI ? SpecularDitherPatternTexture2DRHI : GBlackTexture->TextureRHI->GetTexture2D();
 			
 			FToonProfileUpdateAtlasTileCS::FParameters* PassParameters = GraphBuilder.AllocParameters<FToonProfileUpdateAtlasTileCS::FParameters>();
 			PassParameters->ToonProfileTexture = ToonProfileUAVSkipBarrier;
@@ -341,6 +351,8 @@ IPooledRenderTarget* FToonProfileTextureManager::GetAtlasTexture(FRDGBuilder& Gr
 			PassParameters->ToonProfileIndex = LayerIt;
 			PassParameters->DiffuseDitherPatternStrength = bHasDiffuseDitherPaternTexture ? Data.DiffuseDitherPatternStrength : 0.0f;
 			PassParameters->DiffuseDitherPatternSize = Data.DiffuseDitherPatternSize;
+			PassParameters->SpecularDitherPatternStrength = bHasSpecularDitherPaternTexture ? Data.SpecularDitherPatternStrength : 0.0f;
+			PassParameters->SpecularDitherPatternSize = Data.SpecularDitherPatternSize;
 
 			auto CreateBufferFRomCurve = [&](FRuntimeCurveLinearColor ColorCurve)
 			{
@@ -371,9 +383,12 @@ IPooledRenderTarget* FToonProfileTextureManager::GetAtlasTexture(FRDGBuilder& Gr
 			FRDGBufferRef SpecularRampBuffer = CreateBufferFRomCurve(Data.SpecularRamp);
 			PassParameters->SpecularRampBufferSRV = GraphBuilder.CreateSRV(SpecularRampBuffer, Format);
 
-			PassParameters->DiffuseDitherPattern = GraphBuilder.RegisterExternalTexture(CreateRenderTarget(Texture2DRHI, TEXT("DiffuseDitherPattern")));
+			PassParameters->DiffuseDitherPattern = GraphBuilder.RegisterExternalTexture(CreateRenderTarget(DiffuseDitherPatternTexture2DRHI, TEXT("DiffuseDitherPattern")));
 			PassParameters->DiffuseDitherPatternSampler = TStaticSamplerState<SF_Trilinear, AM_Wrap, AM_Wrap>::GetRHI();
 
+			PassParameters->SpecularDitherPattern = GraphBuilder.RegisterExternalTexture(CreateRenderTarget(SpecularDitherPatternTexture2DRHI, TEXT("SpecularDitherPattern")));
+			PassParameters->SpecularDitherPatternSampler = TStaticSamplerState<SF_Trilinear, AM_Wrap, AM_Wrap>::GetRHI();
+
 			FToonProfileUpdateAtlasTileCS::FPermutationDomain PermutationVector;
 			TShaderMapRef<FToonProfileUpdateAtlasTileCS> Shader(GlobalShaderMap, PermutationVector);
 			FComputeShaderUtils::AddPass(GraphBuilder, 
@@ -444,8 +459,13 @@ FToonProfileStruct::FToonProfileStruct()
 	AddPointToCurve(SpecularRamp, 0.51, 1.0);
 	AddPointToCurve(SpecularRamp, 1.00, 1.0);
 
+	DiffuseDitherPattern = nullptr;
 	DiffuseDitherPatternStrength = 0.0f;	// By default, skip sampling dither pattern.
 	DiffuseDitherPatternSize = 0.1;
+
+	SpecularDitherPattern = nullptr;
+	SpecularDitherPatternStrength = 0.0f;	// By default, skip sampling dither pattern.
+	SpecularDitherPatternSize = 0.1;
 }
 
 ///////////////////////////////////////////////////////////////////////////////////////////////////
diff --git a/Engine/Source/Runtime/Engine/Private/SceneManagement.cpp b/Engine/Source/Runtime/Engine/Private/SceneManagement.cpp
--- a/Engine/Source/Runtime/Engine/Private/SceneManagement.cpp
+++ b/Engine/Source/Runtime/Engine/Private/SceneManagement.cpp
@@ -1308,7 +1308,7 @@ FViewUniformShaderParameters::FViewUniformShaderParameters()
 
 	// Toon profiles
 	ToonProfileTextureSizeAndInvSize = FVector4f(1.f, 1.f, 1.f, 1.f);
-	ToonProfileDiffuseDitherPatternSizeAndInvSize = FVector2f(1.f, 1.f);
+	ToonProfileDitherPatternSizeAndInvSize = FVector2f(1.f, 1.f);
 	ToonProfileTexture = GBlackArrayTexture->TextureRHI;
 	ToonProfileSampler = TStaticSamplerState<SF_Bilinear, AM_Clamp, AM_Clamp, AM_Clamp>::GetRHI();
 
diff --git a/Engine/Source/Runtime/Engine/Public/SceneView.h b/Engine/Source/Runtime/Engine/Public/SceneView.h
--- a/Engine/Source/Runtime/Engine/Public/SceneView.h
+++ b/Engine/Source/Runtime/Engine/Public/SceneView.h
@@ -1072,7 +1072,7 @@ enum ETranslucencyVolumeCascade
 	VIEW_UNIFORM_BUFFER_MEMBER(FVector4f, SSProfilesPreIntegratedTextureSizeAndInvSize) \
 	VIEW_UNIFORM_BUFFER_MEMBER(FVector4f, SpecularProfileTextureSizeAndInvSize) \
 	VIEW_UNIFORM_BUFFER_MEMBER(FVector4f, ToonProfileTextureSizeAndInvSize) \
-	VIEW_UNIFORM_BUFFER_MEMBER(FVector2f, ToonProfileDiffuseDitherPatternSizeAndInvSize) \
+	VIEW_UNIFORM_BUFFER_MEMBER(FVector2f, ToonProfileDitherPatternSizeAndInvSize) \
 	VIEW_UNIFORM_BUFFER_MEMBER(FVector3f, PhysicsFieldClipmapCenter) \
 	VIEW_UNIFORM_BUFFER_MEMBER(float, PhysicsFieldClipmapDistance) \
 	VIEW_UNIFORM_BUFFER_MEMBER(int, PhysicsFieldClipmapResolution) \
diff --git a/Engine/Source/Runtime/Renderer/Private/SceneRendering.cpp b/Engine/Source/Runtime/Renderer/Private/SceneRendering.cpp
--- a/Engine/Source/Runtime/Renderer/Private/SceneRendering.cpp
+++ b/Engine/Source/Runtime/Renderer/Private/SceneRendering.cpp
@@ -1860,10 +1860,10 @@ void FViewInfo::SetupUniformBufferParameters(
 		ViewUniformShaderParameters.ToonProfileTextureSizeAndInvSize = FVector4f(TextureSize.X, TextureSize.Y, 1.0f / TextureSize.X, 1.0f / TextureSize.Y);
 
 		// See ToonProfile.usf for the baking/sampling description.
-		const float  ToonProfileDiffuseDitherPatternFullResolution = ViewUniformShaderParameters.ToonProfileTextureSizeAndInvSize.X;
-		const float  ToonProfileDiffuseDitherPatternBakeResolution = ToonProfileDiffuseDitherPatternFullResolution - 2.0f;
-		const float  ToonProfileDiffuseDitherPatternBakeResolutionInv = 1.0f / ToonProfileDiffuseDitherPatternBakeResolution;
-		ViewUniformShaderParameters.ToonProfileDiffuseDitherPatternSizeAndInvSize = FVector2f(ToonProfileDiffuseDitherPatternBakeResolution, ToonProfileDiffuseDitherPatternBakeResolutionInv);
+		const float  ToonProfileDitherPatternFullResolution = ViewUniformShaderParameters.ToonProfileTextureSizeAndInvSize.X;
+		const float  ToonProfileDitherPatternBakeResolution = ToonProfileDitherPatternFullResolution - 2.0f;
+		const float  ToonProfileDitherPatternBakeResolutionInv = 1.0f / ToonProfileDitherPatternBakeResolution;
+		ViewUniformShaderParameters.ToonProfileDitherPatternSizeAndInvSize = FVector2f(ToonProfileDitherPatternBakeResolution, ToonProfileDitherPatternBakeResolutionInv);
 
 		ViewUniformShaderParameters.ToonProfileTexture = Texture;
 		ViewUniformShaderParameters.ToonProfileSampler = TStaticSamplerState<SF_Bilinear, AM_Clamp, AM_Clamp, AM_Clamp>::GetRHI();
