Prevent warnings when cooking a material asset with no set Usage Flags.
The warning will still be triggered in case a Material which has usage flags and should have shaders but has an empty shader map, along with all other cases.
Only valid reason for this to be used is in case there is a dependency on a Material Asset which is not easily removable but the Material itself is not used for rendering.

[RN] minor | Rendering Architecture

#Synced-CL 49884731

diff --git a/Engine/Source/Runtime/Engine/Private/Materials/MaterialShared.cpp b/Engine/Source/Runtime/Engine/Private/Materials/MaterialShared.cpp
--- a/Engine/Source/Runtime/Engine/Private/Materials/MaterialShared.cpp
+++ b/Engine/Source/Runtime/Engine/Private/Materials/MaterialShared.cpp
@@ -1495,13 +1495,40 @@ void FMaterial::SerializeInlineShaderMap(FArchive& Ar, const FName& SerializingA
 			}
 			else
 			{
-				UE_LOG(LogMaterial, Warning, TEXT("Cooking a material resource (in %s hierarchy) that doesn't have a valid ShaderMap! %s"),
-					*GetFriendlyName(),
-					(GameThreadShaderMap == nullptr) ? TEXT("Shadermap pointer is null.") :
-						!GameThreadShaderMap->CompiledSuccessfully() ? 
-							TEXT("Shadermap exists but wasn't compiled successfully (yet?)") :
-							TEXT("Shadermap exists but has no shaders")
-					);
+				bool bShouldWarn = true;
+				FString InvalidShadermapReason;
+				
+				if (!GameThreadShaderMap)
+				{
+					InvalidShadermapReason += TEXT("Shadermap pointer is null.");
+				}
+				else
+				{
+					if (!GameThreadShaderMap->CompiledSuccessfully())
+					{
+						InvalidShadermapReason += TEXT("Shadermap exists but wasn't compiled successfully (yet?)");
+					}
+					else
+					{
+						if (HasAnyUsage())
+						{
+							InvalidShadermapReason += TEXT("Shadermap exists, has usage flags but has no shaders");
+						}
+						else
+						{
+							UE_LOG(LogMaterial, Log, TEXT("Material resource (in %s hierarchy) has no usage flags and will have an empty ShaderMap."),
+								*GetFriendlyName());
+							bShouldWarn = false;
+						}
+					}
+
+				}
+
+				if (bShouldWarn)
+				{
+					UE_LOG(LogMaterial, Warning, TEXT("Cooking a material resource (in %s hierarchy) that doesn't have a valid ShaderMap! %s"),
+						*GetFriendlyName(), *InvalidShadermapReason);
+				}
 			}
 #else
 			UE_LOG(LogMaterial, Fatal, TEXT("Internal error: cooking outside the editor is not possible."));
@@ -1866,6 +1893,11 @@ bool FMaterialResource::IsUsedWithStaticMesh() const
 	return MaterialInstance ? MaterialInstance->GetUsageByFlag(MATUSAGE_StaticMesh) : Material->GetUsageByFlag(MATUSAGE_StaticMesh);
 }
 
+bool FMaterialResource::HasAnyUsage() const
+{
+	return MaterialInstance ? (MaterialInstance->GetUsageFlags() != 0) : (Material->GetUsageFlags() != 0);
+}
+
 bool FMaterialResource::SamplesMaterialCache() const
 {
 	return Material->GetCachedExpressionData().bSamplesMaterialCache;
diff --git a/Engine/Source/Runtime/Engine/Public/MaterialShared.h b/Engine/Source/Runtime/Engine/Public/MaterialShared.h
--- a/Engine/Source/Runtime/Engine/Public/MaterialShared.h
+++ b/Engine/Source/Runtime/Engine/Public/MaterialShared.h
@@ -2411,6 +2411,7 @@ public:
 	virtual bool IsUsedWithNeuralNetworks() const { return false; }
 	virtual bool IsUsedWithMeshDeformer() const { return false; }
 	virtual bool IsUsedWithStaticMesh() const { return true; }
+	virtual bool HasAnyUsage() const { return true; }
 	virtual bool IsFullyRough() const { return false; }
 	virtual bool GetForceCompatibleWithLightFunctionAtlas() const { return false; }
 	virtual bool UseNormalCurvatureToRoughness() const { return false; }
@@ -3224,6 +3225,7 @@ public:
 	ENGINE_API virtual bool IsUsedWithVolumetricCloud() const override;
 	ENGINE_API virtual bool IsUsedWithHeterogeneousVolumes() const override;
 	ENGINE_API virtual bool IsUsedWithStaticMesh() const override;
+	ENGINE_API virtual bool HasAnyUsage() const override;
 	ENGINE_API virtual bool SamplesMaterialCache() const override;
 	ENGINE_API virtual bool HasMaterialCacheOutput() const override;
 	ENGINE_API virtual bool IsFullyRough() const override;
