Separate Subtitle Types and Delegates into 2 headers; make delayed-start subtitles replace displaying ones of equal priority.

#Synced-CL 49822045

diff --git a/Engine/Plugins/Experimental/SubtitlesAndClosedCaptions/Source/SubtitlesAndClosedCaptions/Private/MovieSceneSubtitleSection.cpp b/Engine/Plugins/Experimental/SubtitlesAndClosedCaptions/Source/SubtitlesAndClosedCaptions/Private/MovieSceneSubtitleSection.cpp
--- a/Engine/Plugins/Experimental/SubtitlesAndClosedCaptions/Source/SubtitlesAndClosedCaptions/Private/MovieSceneSubtitleSection.cpp
+++ b/Engine/Plugins/Experimental/SubtitlesAndClosedCaptions/Source/SubtitlesAndClosedCaptions/Private/MovieSceneSubtitleSection.cpp
@@ -9,6 +9,7 @@
 #include "MovieSceneSubtitleComponentTypes.h"
 #include "SubtitleDataComponent.h"
 #include "Subtitles/SubtitlesAndClosedCaptionsDelegates.h"
+#include "Subtitles/SubtitlesAndClosedCaptionsTypes.h"
 
 #include UE_INLINE_GENERATED_CPP_BY_NAME(MovieSceneSubtitleSection)
 
diff --git a/Engine/Plugins/Experimental/SubtitlesAndClosedCaptions/Source/SubtitlesAndClosedCaptions/Private/MovieSceneSubtitlesTrack.cpp b/Engine/Plugins/Experimental/SubtitlesAndClosedCaptions/Source/SubtitlesAndClosedCaptions/Private/MovieSceneSubtitlesTrack.cpp
--- a/Engine/Plugins/Experimental/SubtitlesAndClosedCaptions/Source/SubtitlesAndClosedCaptions/Private/MovieSceneSubtitlesTrack.cpp
+++ b/Engine/Plugins/Experimental/SubtitlesAndClosedCaptions/Source/SubtitlesAndClosedCaptions/Private/MovieSceneSubtitlesTrack.cpp
@@ -6,6 +6,7 @@
 #include "MovieSceneCommonHelpers.h"
 #include "MovieSceneSubtitleSection.h"
 #include "Subtitles/SubtitlesAndClosedCaptionsDelegates.h"
+#include "Subtitles/SubtitlesAndClosedCaptionsTypes.h"
 
 #include UE_INLINE_GENERATED_CPP_BY_NAME(MovieSceneSubtitlesTrack)
 
diff --git a/Engine/Plugins/Experimental/SubtitlesAndClosedCaptions/Source/SubtitlesAndClosedCaptions/Private/SubtitlesAudioSubsystem.cpp b/Engine/Plugins/Experimental/SubtitlesAndClosedCaptions/Source/SubtitlesAndClosedCaptions/Private/SubtitlesAudioSubsystem.cpp
--- a/Engine/Plugins/Experimental/SubtitlesAndClosedCaptions/Source/SubtitlesAndClosedCaptions/Private/SubtitlesAudioSubsystem.cpp
+++ b/Engine/Plugins/Experimental/SubtitlesAndClosedCaptions/Source/SubtitlesAndClosedCaptions/Private/SubtitlesAudioSubsystem.cpp
@@ -6,6 +6,7 @@
 #include "Sound/DialogueSoundWaveProxy.h"
 #include "Sound/SoundCue.h"
 #include "Sound/SoundNodeDialoguePlayer.h"
+#include "Subtitles/SubtitlesAndClosedCaptionsTypes.h"
 #include "Subtitles/SubtitlesAndClosedCaptionsDelegates.h"
 
 #include UE_INLINE_GENERATED_CPP_BY_NAME(SubtitlesAudioSubsystem)
@@ -78,6 +79,10 @@ void USubtitlesAudioSubsystem::NotifyActiveSoundDeleting(const FActiveSound& Act
 				for (const FSubtitleAssetData& Subtitle : Subtitles->Subtitles)
 				{
 					const bool bUseDurationProperty = (Subtitle.SubtitleDurationType == ESubtitleDurationType::UseDurationProperty);
+					// If not using sound duration, subtitles continue to play instead of stopping:
+					// Their duration is completely decoupled from the sound's.
+
+					// ...but if Using the sound's duration: Stop All subtitles.
 					if (!bUseDurationProperty)
 					{
 						FSubtitlesAndClosedCaptionsDelegates::StopSubtitle.ExecuteIfBound(Subtitle);
diff --git a/Engine/Plugins/Experimental/SubtitlesAndClosedCaptions/Source/SubtitlesAndClosedCaptions/Private/SubtitlesBlueprintFunctionLibrary.cpp b/Engine/Plugins/Experimental/SubtitlesAndClosedCaptions/Source/SubtitlesAndClosedCaptions/Private/SubtitlesBlueprintFunctionLibrary.cpp
--- a/Engine/Plugins/Experimental/SubtitlesAndClosedCaptions/Source/SubtitlesAndClosedCaptions/Private/SubtitlesBlueprintFunctionLibrary.cpp
+++ b/Engine/Plugins/Experimental/SubtitlesAndClosedCaptions/Source/SubtitlesAndClosedCaptions/Private/SubtitlesBlueprintFunctionLibrary.cpp
@@ -4,6 +4,7 @@
 
 #include "Engine/Engine.h"
 #include "Subsystems/SubsystemBlueprintLibrary.h"
+#include "Subtitles/SubtitlesAndClosedCaptionsDelegates.h"
 #include "SubtitlesSubsystem.h"
 
 #include UE_INLINE_GENERATED_CPP_BY_NAME(SubtitlesBlueprintFunctionLibrary)
diff --git a/Engine/Plugins/Experimental/SubtitlesAndClosedCaptions/Source/SubtitlesAndClosedCaptions/Private/SubtitlesSubsystem.cpp b/Engine/Plugins/Experimental/SubtitlesAndClosedCaptions/Source/SubtitlesAndClosedCaptions/Private/SubtitlesSubsystem.cpp
--- a/Engine/Plugins/Experimental/SubtitlesAndClosedCaptions/Source/SubtitlesAndClosedCaptions/Private/SubtitlesSubsystem.cpp
+++ b/Engine/Plugins/Experimental/SubtitlesAndClosedCaptions/Source/SubtitlesAndClosedCaptions/Private/SubtitlesSubsystem.cpp
@@ -14,6 +14,7 @@
 #include "Styling/CoreStyle.h"
 #include "SubtitlesAndClosedCaptionsModule.h"
 #include "SubtitlesSettings.h"
+#include "Subtitles/SubtitlesAndClosedCaptionsDelegates.h"
 #include "TimerManager.h"
 #include "UnrealClient.h"
 #include "UObject/ObjectMacros.h"
@@ -171,7 +172,10 @@ void USubtitlesSubsystem::AddAndDisplaySubtitle(FActiveSubtitleAssetData& NewAct
 	if (ShouldAddSubtitle)
 	{
 		ActiveSubtitles.Add(MoveTemp(NewActiveSubtitle));
-		ActiveSubtitles.StableSort([](const FActiveSubtitleAssetData& Lhs, const FActiveSubtitleAssetData& Rhs) { return Lhs.Subtitle.Priority > Rhs.Subtitle.Priority; });
+
+		// Sort by priority, but also prefer newer subtitles with a delayed start.
+		ActiveSubtitles.StableSort([](const FActiveSubtitleAssetData& Lhs, const FActiveSubtitleAssetData& Rhs)
+			{ return (Lhs.Subtitle.Priority > Rhs.Subtitle.Priority) || (Lhs.Subtitle.StartOffset > 0 && (Lhs.Subtitle.Priority == Rhs.Subtitle.Priority)); });
 
 		UpdateWidgetData();
 	}
@@ -242,9 +246,13 @@ void USubtitlesSubsystem::ReplaceWidget(const TSoftClassPtr<USubtitleWidget>& Ne
 void USubtitlesSubsystem::RemoveActiveSubtitle(const FSubtitleAssetData& Subtitle)
 {
 	bool bSuccessfullyRemoved = false;
-	int32 FirstRemovedIndex = Algo::StableRemoveIf(ActiveSubtitles, [&Subtitle](const FActiveSubtitleAssetData& ActiveSubtitle) {return ActiveSubtitle.Subtitle == Subtitle; });
 
 	FTimerManager& TimerManager = GetWorldRef().GetTimerManager();
+
+	// Remove currently-queued subtitles
+	const int32 FirstRemovedIndex = Algo::StableRemoveIf(ActiveSubtitles, 
+		[&Subtitle](const FActiveSubtitleAssetData& ActiveSubtitle) {return ActiveSubtitle.Subtitle == Subtitle; }
+	);
 	for (int32 Index = FirstRemovedIndex; Index < ActiveSubtitles.Num(); ++Index)
 	{
 		TimerManager.ClearTimer(ActiveSubtitles[Index].DurationTimerHandle);
@@ -264,19 +272,21 @@ void USubtitlesSubsystem::RemoveActiveSubtitle(const FSubtitleAssetData& Subtitl
 			SubtitleWidget->StartDisplayingSubtitle(HighestPrioritySubtitle);
 		}
 	}
-	else if(!IsValid(SubtitleWidget))
+	else if (!IsValid(SubtitleWidget))
 	{
 		UE_LOG(LogSubtitlesAndClosedCaptions, Warning, TEXT("Can't remove subtitles because there isn't a valid UMG widget."));
 	}
 
-	// Also remove delayed-start subtitles using this asset.
-	FirstRemovedIndex = Algo::RemoveIf(DelayedSubtitles, [&Subtitle](const FActiveSubtitleAssetData& DelayedSubtitle) {return DelayedSubtitle.Subtitle == Subtitle; });
+	// Remove delayed-start subtitles that haven't been queued yet that also use this asset.
+	const int32 FirstRemovedActiveIndex = Algo::RemoveIf(DelayedSubtitles,
+		[&Subtitle](const FActiveSubtitleAssetData& DelayedSubtitle) {return DelayedSubtitle.Subtitle == Subtitle; }
+	);
 
-	for (int32 Index = FirstRemovedIndex; Index < DelayedSubtitles.Num(); ++Index)
+	for (int32 Index = FirstRemovedActiveIndex; Index < DelayedSubtitles.Num(); ++Index)
 	{
 		TimerManager.ClearTimer(DelayedSubtitles[Index].DurationTimerHandle);
 	}
-	DelayedSubtitles.SetNum(FirstRemovedIndex);
+	DelayedSubtitles.SetNum(FirstRemovedActiveIndex);
 }
 
 bool USubtitlesSubsystem::TryCreateUMGWidgetFromAsset(const TSoftClassPtr<USubtitleWidget>& WidgetToUse)
diff --git a/Engine/Plugins/Experimental/SubtitlesAndClosedCaptions/Source/SubtitlesAndClosedCaptions/Public/ActiveSubtitleAssetData.h b/Engine/Plugins/Experimental/SubtitlesAndClosedCaptions/Source/SubtitlesAndClosedCaptions/Public/ActiveSubtitleAssetData.h
--- a/Engine/Plugins/Experimental/SubtitlesAndClosedCaptions/Source/SubtitlesAndClosedCaptions/Public/ActiveSubtitleAssetData.h
+++ b/Engine/Plugins/Experimental/SubtitlesAndClosedCaptions/Source/SubtitlesAndClosedCaptions/Public/ActiveSubtitleAssetData.h
@@ -5,7 +5,7 @@
 struct FSubtitleAssetData;
 
 #include "Engine/TimerHandle.h"
-#include "Subtitles/SubtitlesAndClosedCaptionsDelegates.h"
+#include "Subtitles/SubtitlesAndClosedCaptionsTypes.h"
 #include "UObject/ObjectPtr.h"
 
 #include "ActiveSubtitleAssetData.generated.h"
diff --git a/Engine/Plugins/Experimental/SubtitlesAndClosedCaptions/Source/SubtitlesAndClosedCaptions/Public/MovieSceneSubtitlesSystem.h b/Engine/Plugins/Experimental/SubtitlesAndClosedCaptions/Source/SubtitlesAndClosedCaptions/Public/MovieSceneSubtitlesSystem.h
--- a/Engine/Plugins/Experimental/SubtitlesAndClosedCaptions/Source/SubtitlesAndClosedCaptions/Public/MovieSceneSubtitlesSystem.h
+++ b/Engine/Plugins/Experimental/SubtitlesAndClosedCaptions/Source/SubtitlesAndClosedCaptions/Public/MovieSceneSubtitlesSystem.h
@@ -14,6 +14,7 @@
 #include "MovieSceneSubtitleSection.h"
 #include "SubtitleDataComponent.h"
 #include "Subtitles/SubtitlesAndClosedCaptionsDelegates.h"
+#include "Subtitles/SubtitlesAndClosedCaptionsTypes.h"
 
 #include "MovieSceneSubtitlesSystem.generated.h"
 
diff --git a/Engine/Plugins/Experimental/SubtitlesAndClosedCaptions/Source/SubtitlesAndClosedCaptions/Public/SubtitlesBlueprintFunctionLibrary.h b/Engine/Plugins/Experimental/SubtitlesAndClosedCaptions/Source/SubtitlesAndClosedCaptions/Public/SubtitlesBlueprintFunctionLibrary.h
--- a/Engine/Plugins/Experimental/SubtitlesAndClosedCaptions/Source/SubtitlesAndClosedCaptions/Public/SubtitlesBlueprintFunctionLibrary.h
+++ b/Engine/Plugins/Experimental/SubtitlesAndClosedCaptions/Source/SubtitlesAndClosedCaptions/Public/SubtitlesBlueprintFunctionLibrary.h
@@ -3,7 +3,7 @@
 #pragma once
 
 #include "Kismet/BlueprintFunctionLibrary.h"
-#include "Subtitles/SubtitlesAndClosedCaptionsDelegates.h"
+#include "Subtitles/SubtitlesAndClosedCaptionsTypes.h"
 #include "SubtitlesBlueprintFunctionLibrary.generated.h"
 
 #define UE_API SUBTITLESANDCLOSEDCAPTIONS_API
diff --git a/Engine/Plugins/Experimental/SubtitlesAndClosedCaptions/Source/SubtitlesAndClosedCaptions/Public/SubtitlesSubsystem.h b/Engine/Plugins/Experimental/SubtitlesAndClosedCaptions/Source/SubtitlesAndClosedCaptions/Public/SubtitlesSubsystem.h
--- a/Engine/Plugins/Experimental/SubtitlesAndClosedCaptions/Source/SubtitlesAndClosedCaptions/Public/SubtitlesSubsystem.h
+++ b/Engine/Plugins/Experimental/SubtitlesAndClosedCaptions/Source/SubtitlesAndClosedCaptions/Public/SubtitlesSubsystem.h
@@ -6,6 +6,7 @@
 #include "Fonts/SlateFontInfo.h"
 #include "Math/Color.h"
 #include "Subsystems/WorldSubsystem.h"
+#include "Subtitles/SubtitlesAndClosedCaptionsDelegates.h"
 
 #include "SubtitleWidget.h"
 
diff --git a/Engine/Source/Runtime/Engine/Private/Subtitles/SubtitlesAndClosedCaptionsDelegates.cpp b/Engine/Source/Runtime/Engine/Private/Subtitles/SubtitlesAndClosedCaptionsDelegates.cpp
--- a/Engine/Source/Runtime/Engine/Private/Subtitles/SubtitlesAndClosedCaptionsDelegates.cpp
+++ b/Engine/Source/Runtime/Engine/Private/Subtitles/SubtitlesAndClosedCaptionsDelegates.cpp
@@ -1,20 +1,9 @@
 // Copyright Epic Games, Inc. All Rights Reserved.
 
 #include "Subtitles/SubtitlesAndClosedCaptionsDelegates.h"
-
-#include UE_INLINE_GENERATED_CPP_BY_NAME(SubtitlesAndClosedCaptionsDelegates)
+#include "Subtitles/SubtitlesAndClosedCaptionsTypes.h"
 
 TDelegate<void(FQueueSubtitleParameters, const ESubtitleTiming)> FSubtitlesAndClosedCaptionsDelegates::QueueSubtitle;
 TDelegate<bool(FSubtitleAssetData)> FSubtitlesAndClosedCaptionsDelegates::IsSubtitleActive;
 TDelegate<void(FSubtitleAssetData)> FSubtitlesAndClosedCaptionsDelegates::StopSubtitle;
-TDelegate<void()> FSubtitlesAndClosedCaptionsDelegates::StopAllSubtitles;
-
-#if WITH_EDITORONLY_DATA
-void USubtitleAssetUserData::PostEditChangeProperty(FPropertyChangedEvent& PropertyChangedEvent)
-{
-	for (FSubtitleAssetData& Subtitle : Subtitles)
-	{
-		Subtitle.bCanEditDuration = (Subtitle.SubtitleDurationType == ESubtitleDurationType::UseDurationProperty);
-	}
-}
-#endif // WITH_EDITORONLY_DATA
+TDelegate<void()> FSubtitlesAndClosedCaptionsDelegates::StopAllSubtitles;
\ No newline at end of file
diff --git a/Engine/Source/Runtime/Engine/Private/Subtitles/SubtitlesAndClosedCaptionsTypes.cpp b/Engine/Source/Runtime/Engine/Private/Subtitles/SubtitlesAndClosedCaptionsTypes.cpp
new file mode 100644
--- /dev/null
+++ b/Engine/Source/Runtime/Engine/Private/Subtitles/SubtitlesAndClosedCaptionsTypes.cpp
@@ -0,0 +1,15 @@
+// Copyright Epic Games, Inc. All Rights Reserved.
+
+#include "Subtitles/SubtitlesAndClosedCaptionsTypes.h"
+
+#include UE_INLINE_GENERATED_CPP_BY_NAME(SubtitlesAndClosedCaptionsTypes)
+
+#if WITH_EDITOR
+void USubtitleAssetUserData::PostEditChangeProperty(FPropertyChangedEvent& PropertyChangedEvent)
+{
+	for (FSubtitleAssetData& Subtitle : Subtitles)
+	{
+		Subtitle.bCanEditDuration = (Subtitle.SubtitleDurationType == ESubtitleDurationType::UseDurationProperty);
+	}
+}
+#endif // WITH_EDITOR
\ No newline at end of file
diff --git a/Engine/Source/Runtime/Engine/Public/Subtitles/SubtitlesAndClosedCaptionsDelegates.h b/Engine/Source/Runtime/Engine/Public/Subtitles/SubtitlesAndClosedCaptionsDelegates.h
--- a/Engine/Source/Runtime/Engine/Public/Subtitles/SubtitlesAndClosedCaptionsDelegates.h
+++ b/Engine/Source/Runtime/Engine/Public/Subtitles/SubtitlesAndClosedCaptionsDelegates.h
@@ -2,152 +2,17 @@
 
 #pragma once
 #include "Delegates/Delegate.h"
-#include "Engine/AssetUserData.h"
-#include "SubtitlesAndClosedCaptionsDelegates.generated.h"
+#include "SubtitlesAndClosedCaptionsTypes.h"
 
 #define UE_API ENGINE_API
 
-class FCanvas;
-class UAssetUserData;
-
-// Externally-Timed subtitles must be manually added and removed with USubtitlesSubsystem::QueueSubtitle and ::StopSubtitle.
-// For the initial delay before becoming visible, use USubtitleAssetUserData::StartOffset instead of this enum.
-UENUM()
-enum class ESubtitleTiming : uint8
-{
-	InternallyTimed,
-	ExternallyTimed
-};
-
-// SubtitlesAudioSubsystem will automatically queue and stop subtitles embedded in USoundBase sounds using play/stop callbacks.
-// This enum controls whether the stopping should be done automatically, or if the subtitle should be timed using its existing duration.
-UENUM(BlueprintType)
-enum class ESubtitleDurationType : uint8
-{
-	UseSoundDuration UMETA(ToolTip = "Automatically stops this subtitle when the sound stops playing (only useful when attached as AssetUserData"),
-	UseDurationProperty UMETA(ToolTip = "Time this subtitle using its Duration property")
-};
-
-// Subtitle type for type-specific rendering.
-UENUM()
-enum class ESubtitleType : uint8
-{
-	Subtitle,
-	ClosedCaption,
-	AudioDescription
-};
-
-// Minimum duration to display subtitle.
-static constexpr float SubtitleMinDuration		= 0.05f;
-
-// Default value to initialize subtitle duration to. Used by SoundWaves to check if they should manually set the duration.
-static constexpr float SubtitleDefaultDuration	= 3.f;
-
-USTRUCT(MinimalAPI, BlueprintType)
-struct FSubtitleAssetData
-{
-	GENERATED_BODY()
-
-	// #SUBTITLES_PRD: carried over from FSubtitleCue, still required
-
-	// The text to appear in the subtitle.
-	UPROPERTY(EditAnywhere, BlueprintReadWrite, Category = "Subtitles", meta = (Multiline = true))
-	FText Text;
-
-	// #SUBTITLES_PRD: carried over from FSubtitleCue, still required
-
-	// Whether to automatically unqueue the subtitle using an associated Sound, or to use the duration supplied below.
-	UPROPERTY(EditAnywhere, BlueprintReadWrite, Category = "Subtitles|Duration")
-	ESubtitleDurationType SubtitleDurationType = ESubtitleDurationType::UseSoundDuration;
-
-#if WITH_EDITORONLY_DATA
-	// Editor-only helper value for SubtitleDefaultDuration's EditCondition.
-	UPROPERTY(Transient, VisibleDefaultsOnly, Category = "Subtitles|Duration")
-	bool bCanEditDuration = (SubtitleDurationType == ESubtitleDurationType::UseDurationProperty);
-#endif
-
-	// Time to display in seconds.
-	// Defaulted to 3 seconds so when adding new subtitles it's not required to enter a placeholder Duration.
-	// Duration can be be set by ingestion pipelines when importing Subtitles in bulk
-	UPROPERTY(EditAnywhere, BlueprintReadWrite, Category = "Subtitles|Duration", meta = (ClampMin = 0.05f, EditCondition = "bCanEditDuration"))
-	float Duration = SubtitleDefaultDuration;
-
-	// Some subtitles have a delay before they're allowed to be displayed (primarily from the legacy system).
-	// StartOffset measures how long in Seconds, after queuing, before the subtitle is allowed to enter the active subtitles queue.
-	// ESubtitleTiming::ExternallyTimed does not effect this initial delay.
-	UPROPERTY(EditAnywhere, BlueprintReadWrite, Category = "Subtitles", meta = (ClampMin = 0.f))
-	float StartOffset = 0.f;
-
-	// #SUBTITLES_PRD: Priority comes from USoundBase::GetSubtitlePriority, USoundCue::GetSubtitlePriority and USoundWave::GetSubtitlePriority
-	// Consolidate various subtitle properties throughout sound/audio code into this new subtitles plugin.
-
-	// The priority of the subtitle.  Defaults to 1.  Higher values will play instead of lower values.
-	UPROPERTY(EditAnywhere, BlueprintReadWrite, Category = "Subtitles")
-	float Priority = 1.f;
-
-	// Subtitle type for type-specific rendering.
-	UPROPERTY(EditAnywhere, BlueprintReadWrite, Category = "Subtitles")
-	ESubtitleType SubtitleType = ESubtitleType::Subtitle;
-
-#if WITH_EDITORONLY_DATA
-	// Author comments, editor-only.
-	UPROPERTY(EditAnywhere, Category = "Subtitles")
-	FString Comment;
-#endif // WITH_EDITORONLY_DATA
-
-	// The struct needs a custom equality test operator because FText doesn't have one overloaded: need to use EqualTo().
-	// Don't compare the Comment property because it isn't compiled outside the editor.
-	const bool operator==(const FSubtitleAssetData& Other) const
-	{
-		return ((Text.EqualTo(Other.Text) && (SubtitleDurationType == Other.SubtitleDurationType) && (Duration == Other.Duration) 
-			&& (StartOffset == Other.StartOffset) && (Priority == Other.Priority) && (SubtitleType == Other.SubtitleType)));
-	}
-};
-
-/**
- *  Base class for subtitle data being attached to assets
- */
-UCLASS(MinimalAPI, BlueprintType)
-class USubtitleAssetUserData : public UAssetUserData
-{
-	GENERATED_BODY()
-
-public:
-	UPROPERTY(EditAnywhere, BlueprintReadOnly, Category = "Subtitles")
-	TArray<FSubtitleAssetData> Subtitles;
-
-	// In seconds, the longest any subtitle in the Subtitles array will last, including any initial offset starts.
-	float GetMaximumDuration() const
-	{
-		float Duration = 0.f;
-		for (const FSubtitleAssetData& Entry : Subtitles)
-		{
-			const float EntryDuration = Entry.Duration + Entry.StartOffset;
-			Duration = FMath::Max(Duration, EntryDuration);
-		}
-
-		return Duration;
-	};
-
-#if WITH_EDITORONLY_DATA
-	// UObject::PostEditChangeProperty is in a #if WITH_EDITOR wrapper.
-	// EDITORONLY_DATA seems equivalent here and ensures that it matches the #if for bCanEditDuration above (which is indeed editor-only data).
-	virtual void PostEditChangeProperty(FPropertyChangedEvent& PropertyChangedEvent) override;
-#endif
-};
-
-
-struct FQueueSubtitleParameters
-{
-	const FSubtitleAssetData Subtitle;
-	TOptional<float> Duration;
-	TOptional<float> StartOffset;	// Scrubbing subtitles in the sequencer needs an override for StartOffset.
-};
+struct FSubtitleAssetData;
+struct FQueueSubtitleParameters;
+enum class ESubtitleTiming : uint8;
 
 class FSubtitlesAndClosedCaptionsDelegates
 {
 public:
-
 	// Have the subtitle subsystem to queue a subtitle to be displayed
 	static UE_API TDelegate<void(FQueueSubtitleParameters, const ESubtitleTiming)> QueueSubtitle;
 
diff --git a/Engine/Source/Runtime/Engine/Public/Subtitles/SubtitlesAndClosedCaptionsTypes.h b/Engine/Source/Runtime/Engine/Public/Subtitles/SubtitlesAndClosedCaptionsTypes.h
new file mode 100644
--- /dev/null
+++ b/Engine/Source/Runtime/Engine/Public/Subtitles/SubtitlesAndClosedCaptionsTypes.h
@@ -0,0 +1,143 @@
+// Copyright Epic Games, Inc. All Rights Reserved.
+
+#pragma once
+#include "Engine/AssetUserData.h"
+#include "SubtitlesAndClosedCaptionsTypes.generated.h"
+
+#define UE_API ENGINE_API
+
+class FCanvas;
+class UAssetUserData;
+
+// Externally-Timed subtitles must be manually added and removed with USubtitlesSubsystem::QueueSubtitle and ::StopSubtitle.
+// For the initial delay before becoming visible, use USubtitleAssetUserData::StartOffset instead of this enum.
+UENUM()
+enum class ESubtitleTiming : uint8
+{
+	InternallyTimed,
+	ExternallyTimed
+};
+
+// SubtitlesAudioSubsystem will automatically queue and stop subtitles embedded in USoundBase sounds using play/stop callbacks.
+// This enum controls whether the stopping should be done automatically, or if the subtitle should be timed using its existing duration.
+UENUM(BlueprintType)
+enum class ESubtitleDurationType : uint8
+{
+	UseSoundDuration UMETA(ToolTip = "Automatically stops this subtitle when the sound stops playing (only useful when attached as AssetUserData"),
+	UseDurationProperty UMETA(ToolTip = "Time this subtitle using its Duration property")
+};
+
+// Subtitle type for type-specific rendering.
+UENUM()
+enum class ESubtitleType : uint8
+{
+	Subtitle,
+	ClosedCaption,
+	AudioDescription
+};
+
+// Minimum duration to display subtitle.
+static constexpr float SubtitleMinDuration		= 0.05f;
+
+// Default value to initialize subtitle duration to. Used by SoundWaves to check if they should manually set the duration.
+static constexpr float SubtitleDefaultDuration	= 3.f;
+
+USTRUCT(MinimalAPI, BlueprintType)
+struct FSubtitleAssetData
+{
+	GENERATED_BODY()
+
+	// #SUBTITLES_PRD: carried over from FSubtitleCue, still required
+
+	// The text to appear in the subtitle.
+	UPROPERTY(EditAnywhere, BlueprintReadWrite, Category = "Subtitles", meta = (Multiline = true))
+	FText Text;
+
+	// #SUBTITLES_PRD: carried over from FSubtitleCue, still required
+
+	// Whether to automatically unqueue the subtitle using an associated Sound, or to use the duration supplied below.
+	UPROPERTY(EditAnywhere, BlueprintReadWrite, Category = "Subtitles|Duration")
+	ESubtitleDurationType SubtitleDurationType = ESubtitleDurationType::UseSoundDuration;
+
+#if WITH_EDITORONLY_DATA
+	// Editor-only helper value for SubtitleDefaultDuration's EditCondition.
+	UPROPERTY(Transient, VisibleDefaultsOnly, Category = "Subtitles|Duration")
+	bool bCanEditDuration = (SubtitleDurationType == ESubtitleDurationType::UseDurationProperty);
+#endif // WITH_EDITORONLY_DATA
+
+	// Time to display in seconds.
+	// Defaulted to 3 seconds so when adding new subtitles it's not required to enter a placeholder Duration.
+	// Duration can be be set by ingestion pipelines when importing Subtitles in bulk
+	UPROPERTY(EditAnywhere, BlueprintReadWrite, Category = "Subtitles|Duration", meta = (ClampMin = 0.05f, EditCondition = "bCanEditDuration"))
+	float Duration = SubtitleDefaultDuration;
+
+	// Some subtitles have a delay before they're allowed to be displayed (primarily from the legacy system).
+	// StartOffset measures how long in Seconds, after queuing, before the subtitle is allowed to enter the active subtitles queue.
+	// ESubtitleTiming::ExternallyTimed does not effect this initial delay.
+	UPROPERTY(EditAnywhere, BlueprintReadWrite, Category = "Subtitles", meta = (ClampMin = 0.f))
+	float StartOffset = 0.f;
+
+	// #SUBTITLES_PRD: Priority comes from USoundBase::GetSubtitlePriority, USoundCue::GetSubtitlePriority and USoundWave::GetSubtitlePriority
+	// Consolidate various subtitle properties throughout sound/audio code into this new subtitles plugin.
+
+	// The priority of the subtitle.  Defaults to 1.  Higher values will play instead of lower values.
+	UPROPERTY(EditAnywhere, BlueprintReadWrite, Category = "Subtitles")
+	float Priority = 1.f;
+
+	// Subtitle type for type-specific rendering.
+	UPROPERTY(EditAnywhere, BlueprintReadWrite, Category = "Subtitles")
+	ESubtitleType SubtitleType = ESubtitleType::Subtitle;
+
+#if WITH_EDITORONLY_DATA
+	// Author comments, editor-only.
+	UPROPERTY(EditAnywhere, Category = "Subtitles")
+	FString Comment;
+#endif // WITH_EDITORONLY_DATA
+
+	// The struct needs a custom equality test operator because FText doesn't have one overloaded: need to use EqualTo().
+	// Don't compare the Comment property because it isn't compiled outside the editor.
+	const bool operator==(const FSubtitleAssetData& Other) const
+	{
+		return ((Text.EqualTo(Other.Text) && (SubtitleDurationType == Other.SubtitleDurationType) && (Duration == Other.Duration) 
+			&& (StartOffset == Other.StartOffset) && (Priority == Other.Priority) && (SubtitleType == Other.SubtitleType)));
+	}
+};
+
+/**
+ *  Base class for subtitle data being attached to assets
+ */
+UCLASS(MinimalAPI, BlueprintType)
+class USubtitleAssetUserData : public UAssetUserData
+{
+	GENERATED_BODY()
+
+public:
+	UPROPERTY(EditAnywhere, BlueprintReadOnly, Category = "Subtitles")
+	TArray<FSubtitleAssetData> Subtitles;
+
+	// In seconds, the longest any subtitle in the Subtitles array will last, including any initial offset starts.
+	float GetMaximumDuration() const
+	{
+		float Duration = 0.f;
+		for (const FSubtitleAssetData& Entry : Subtitles)
+		{
+			const float EntryDuration = Entry.Duration + Entry.StartOffset;
+			Duration = FMath::Max(Duration, EntryDuration);
+		}
+
+		return Duration;
+	};
+
+#if WITH_EDITOR
+	virtual void PostEditChangeProperty(FPropertyChangedEvent& PropertyChangedEvent) override;
+#endif // WITH_EDITOR
+};
+
+struct FQueueSubtitleParameters
+{
+	const FSubtitleAssetData Subtitle;
+	TOptional<float> Duration;
+	TOptional<float> StartOffset;	// Scrubbing subtitles in the sequencer needs an override for StartOffset.
+};
+
+#undef UE_API
