Make profile and their renderer update the asset preview as they should after edits.

#tests PC editor
#rnx

#Synced-CL 49930030

diff --git a/Engine/Source/Editor/UnrealEd/Private/ThumbnailRendering/SpecularProfileRenderer.cpp b/Engine/Source/Editor/UnrealEd/Private/ThumbnailRendering/SpecularProfileRenderer.cpp
--- a/Engine/Source/Editor/UnrealEd/Private/ThumbnailRendering/SpecularProfileRenderer.cpp
+++ b/Engine/Source/Editor/UnrealEd/Private/ThumbnailRendering/SpecularProfileRenderer.cpp
@@ -45,5 +45,5 @@ bool USpecularProfileRenderer::CanVisualizeAsset(UObject* Object)
 
 EThumbnailRenderFrequency USpecularProfileRenderer::GetThumbnailRenderFrequency(UObject* Object) const
 {
-	return EThumbnailRenderFrequency::OnAssetSave;
+	return EThumbnailRenderFrequency::OnPropertyChange;
 }
diff --git a/Engine/Source/Editor/UnrealEd/Private/ThumbnailRendering/SubsurfaceProfileRenderer.cpp b/Engine/Source/Editor/UnrealEd/Private/ThumbnailRendering/SubsurfaceProfileRenderer.cpp
--- a/Engine/Source/Editor/UnrealEd/Private/ThumbnailRendering/SubsurfaceProfileRenderer.cpp
+++ b/Engine/Source/Editor/UnrealEd/Private/ThumbnailRendering/SubsurfaceProfileRenderer.cpp
@@ -47,5 +47,5 @@ bool USubsurfaceProfileRenderer::CanVisualizeAsset(UObject* Object)
 
 EThumbnailRenderFrequency USubsurfaceProfileRenderer::GetThumbnailRenderFrequency(UObject* Object) const
 {
-	return EThumbnailRenderFrequency::OnAssetSave;
+	return EThumbnailRenderFrequency::OnPropertyChange;
 }
diff --git a/Engine/Source/Editor/UnrealEd/Private/ThumbnailRendering/ToonProfileRenderer.cpp b/Engine/Source/Editor/UnrealEd/Private/ThumbnailRendering/ToonProfileRenderer.cpp
--- a/Engine/Source/Editor/UnrealEd/Private/ThumbnailRendering/ToonProfileRenderer.cpp
+++ b/Engine/Source/Editor/UnrealEd/Private/ThumbnailRendering/ToonProfileRenderer.cpp
@@ -53,5 +53,5 @@ bool UToonProfileRenderer::CanVisualizeAsset(UObject* Object)
 
 EThumbnailRenderFrequency UToonProfileRenderer::GetThumbnailRenderFrequency(UObject* Object) const
 {
-	return EThumbnailRenderFrequency::OnAssetSave;
+	return EThumbnailRenderFrequency::OnPropertyChange;
 }
diff --git a/Engine/Source/Runtime/Engine/Classes/Engine/SpecularProfile.h b/Engine/Source/Runtime/Engine/Classes/Engine/SpecularProfile.h
--- a/Engine/Source/Runtime/Engine/Classes/Engine/SpecularProfile.h
+++ b/Engine/Source/Runtime/Engine/Classes/Engine/SpecularProfile.h
@@ -82,8 +82,10 @@ class USpecularProfile : public UObject
 
 	//~ Begin UObject Interface
 	virtual void BeginDestroy();
-	virtual void PostEditChangeProperty(struct FPropertyChangedEvent& PropertyChangedEvent);
 	virtual void PostDuplicate(EDuplicateMode::Type DuplicateMode);
+#if WITH_EDITOR
+	virtual void PostEditChangeProperty(struct FPropertyChangedEvent& PropertyChangedEvent);
+#endif
 	//~ End UObject Interface
 };
 
diff --git a/Engine/Source/Runtime/Engine/Classes/Engine/SubsurfaceProfile.h b/Engine/Source/Runtime/Engine/Classes/Engine/SubsurfaceProfile.h
--- a/Engine/Source/Runtime/Engine/Classes/Engine/SubsurfaceProfile.h
+++ b/Engine/Source/Runtime/Engine/Classes/Engine/SubsurfaceProfile.h
@@ -178,9 +178,11 @@ class USubsurfaceProfile : public UObject
 
 	//~ Begin UObject Interface
 	virtual void BeginDestroy();
-	virtual void PostEditChangeProperty(struct FPropertyChangedEvent& PropertyChangedEvent);
 	virtual void PostDuplicate(EDuplicateMode::Type DuplicateMode);
 	virtual void Serialize(FArchive& Ar) override;
+#if WITH_EDITOR
+	virtual void PostEditChangeProperty(struct FPropertyChangedEvent& PropertyChangedEvent);
+#endif
 
 	// Upgrade parameters from Separable to Burley.
 	virtual void PostLoad();
diff --git a/Engine/Source/Runtime/Engine/Classes/Engine/ToonProfile.h b/Engine/Source/Runtime/Engine/Classes/Engine/ToonProfile.h
--- a/Engine/Source/Runtime/Engine/Classes/Engine/ToonProfile.h
+++ b/Engine/Source/Runtime/Engine/Classes/Engine/ToonProfile.h
@@ -97,8 +97,10 @@ class UToonProfile : public UObject
 
 	//~ Begin UObject Interface
 	virtual void BeginDestroy();
-	virtual void PostEditChangeProperty(struct FPropertyChangedEvent& PropertyChangedEvent);
 	virtual void PostDuplicate(EDuplicateMode::Type DuplicateMode);
+#if WITH_EDITOR
+	virtual void PostEditChangeProperty(struct FPropertyChangedEvent& PropertyChangedEvent);
+#endif
 	//~ End UObject Interface
 };
 
diff --git a/Engine/Source/Runtime/Engine/Private/Rendering/SpecularProfile.cpp b/Engine/Source/Runtime/Engine/Private/Rendering/SpecularProfile.cpp
--- a/Engine/Source/Runtime/Engine/Private/Rendering/SpecularProfile.cpp
+++ b/Engine/Source/Runtime/Engine/Private/Rendering/SpecularProfile.cpp
@@ -541,8 +541,11 @@ void USpecularProfile::BeginDestroy()
 	Super::BeginDestroy();
 }
 
+#if WITH_EDITOR
 void USpecularProfile::PostEditChangeProperty(struct FPropertyChangedEvent& PropertyChangedEvent)
 {
+	Super::PostEditChangeProperty(PropertyChangedEvent);
+
 	const FSpecularProfileStruct LocalSettings = this->Settings;
 	USpecularProfile* LocalProfile = this;
 	GetRendererModule().InvalidatePathTracedOutput();
@@ -560,6 +563,7 @@ void USpecularProfile::PostEditChangeProperty(struct FPropertyChangedEvent& Prop
 		GSpecularProfileTextureManager.UpdateProfile(LocalProfile, LocalSettings, LocalTextureResource);
 	});
 }
+#endif // WITH_EDITOR
 
 void USpecularProfile::PostDuplicate(EDuplicateMode::Type DuplicateMode)
 {
diff --git a/Engine/Source/Runtime/Engine/Private/Rendering/SubsurfaceProfile.cpp b/Engine/Source/Runtime/Engine/Private/Rendering/SubsurfaceProfile.cpp
--- a/Engine/Source/Runtime/Engine/Private/Rendering/SubsurfaceProfile.cpp
+++ b/Engine/Source/Runtime/Engine/Private/Rendering/SubsurfaceProfile.cpp
@@ -824,8 +824,11 @@ void USubsurfaceProfile::BeginDestroy()
 	Super::BeginDestroy();
 }
 
+#if WITH_EDITOR
 void USubsurfaceProfile::PostEditChangeProperty(struct FPropertyChangedEvent& PropertyChangedEvent)
 {
+	Super::PostEditChangeProperty(PropertyChangedEvent);
+
 	const FSubsurfaceProfileStruct SettingsLocal = this->Settings;
 	USubsurfaceProfile* Profile = this;
 	GetRendererModule().InvalidatePathTracedOutput();
@@ -836,6 +839,7 @@ void USubsurfaceProfile::PostEditChangeProperty(struct FPropertyChangedEvent& Pr
 			GSubsurfaceProfileTextureObject.UpdateProfile(SettingsLocal, Profile);
 		});
 }
+#endif // WITH_EDITOR
 
 void USubsurfaceProfile::PostDuplicate(EDuplicateMode::Type DuplicateMode)
 {
diff --git a/Engine/Source/Runtime/Engine/Private/Rendering/ToonProfile.cpp b/Engine/Source/Runtime/Engine/Private/Rendering/ToonProfile.cpp
--- a/Engine/Source/Runtime/Engine/Private/Rendering/ToonProfile.cpp
+++ b/Engine/Source/Runtime/Engine/Private/Rendering/ToonProfile.cpp
@@ -488,8 +488,11 @@ void UToonProfile::BeginDestroy()
 	Super::BeginDestroy();
 }
 
+#if WITH_EDITOR
 void UToonProfile::PostEditChangeProperty(struct FPropertyChangedEvent& PropertyChangedEvent)
 {
+	Super::PostEditChangeProperty(PropertyChangedEvent);
+
 	const FToonProfileStruct LocalSettings = this->Settings;
 	UToonProfile* LocalProfile = this;
 	GetRendererModule().InvalidatePathTracedOutput();
@@ -501,6 +504,7 @@ void UToonProfile::PostEditChangeProperty(struct FPropertyChangedEvent& Property
 		GToonProfileTextureManager.UpdateProfile(LocalProfile, LocalSettings);
 	});
 }
+#endif // WITH_EDITOR
 
 void UToonProfile::PostDuplicate(EDuplicateMode::Type DuplicateMode)
 {
