[NetPhysics] Mark properties with custom condition to be able to disable them at runtime without a warning log.

#Synced-CL 49881534

diff --git a/Engine/Source/Runtime/Engine/Private/PhysicsEngine/NetworkPhysicsComponent.cpp b/Engine/Source/Runtime/Engine/Private/PhysicsEngine/NetworkPhysicsComponent.cpp
--- a/Engine/Source/Runtime/Engine/Private/PhysicsEngine/NetworkPhysicsComponent.cpp
+++ b/Engine/Source/Runtime/Engine/Private/PhysicsEngine/NetworkPhysicsComponent.cpp
@@ -738,7 +738,7 @@ void UNetworkPhysicsComponent::GetLifetimeReplicatedProps(TArray< FLifetimePrope
 	ReplicatedParamsRemote.bIsPushBased = true;
 
 	FDoRepLifetimeParams ReplicatedParamsAll;
-	ReplicatedParamsAll.Condition = COND_None;
+	ReplicatedParamsAll.Condition = COND_Custom;
 	ReplicatedParamsAll.RepNotifyCondition = REPNOTIFY_Always;
 	ReplicatedParamsAll.bIsPushBased = true;
 
