Remove leftover build fix comment.
This was used to test a fix for the build while it was failing without proper message.

#Synced-CL 49929750

diff --git a/Engine/Source/Runtime/Engine/Private/Materials/MaterialIREmitter.cpp b/Engine/Source/Runtime/Engine/Private/Materials/MaterialIREmitter.cpp
--- a/Engine/Source/Runtime/Engine/Private/Materials/MaterialIREmitter.cpp
+++ b/Engine/Source/Runtime/Engine/Private/Materials/MaterialIREmitter.cpp
@@ -737,9 +737,6 @@ void FEmitter::CheckMaterialDomainIsAnyOf(TConstArrayView<EMaterialDomain> Domai
 		return;
 	}
 
-	//static const UEnum* MaterialDomainEnum = StaticEnum<EMaterialDomain>();
-	//check(MaterialDomainEnum);
-
 	// Report the error message.
 	FString ErrorMessage = TEXT("Expression does not support material domain '");
 	ErrorMessage += StaticEnum<EMaterialDomain>()->GetValueAsString(BaseMaterial->MaterialDomain);
