[World] Old DataLayer deprecation continuation part 3 - deprecate FActorDataLayer and related APIs. Remove CitySamplePlayFromHere which was a development tool but is costly to maintain.

#tests Ran the editor and played with data layers
#rnx

#Synced-CL 49949762

diff --git a/Engine/Plugins/PCG/Source/PCG/Private/Grid/PCGPartitionActorDesc.cpp b/Engine/Plugins/PCG/Source/PCG/Private/Grid/PCGPartitionActorDesc.cpp
--- a/Engine/Plugins/PCG/Source/PCG/Private/Grid/PCGPartitionActorDesc.cpp
+++ b/Engine/Plugins/PCG/Source/PCG/Private/Grid/PCGPartitionActorDesc.cpp
@@ -27,16 +27,14 @@ FPCGGridDescriptor FPCGPartitionActorDesc::GetGridDescriptor(const FWorldPartiti
 		.SetIsRuntime(false)
 		.SetHLODLayer(TSoftObjectPtr<UHLODLayer>(GetHLODLayer()));
 
-	if (IsUsingDataLayerAsset())
+	TArray<TSoftObjectPtr<UDataLayerAsset>> DataLayerAssets;
+	// Because of deprecation DataLayerAssets in ActorDescs are stored in FNames (we can convert them to SoftObjectPtrs if IsUsingDataLayerAsset() returns true)
+	// TODO: store as object pointers now that all data layers have assets
+	Algo::Transform(GetDataLayers(), DataLayerAssets, [](const FName& InDataLayerAsset) 
 	{
-		TArray<TSoftObjectPtr<UDataLayerAsset>> DataLayerAssets;
-		// Because of deprecation DataLayerAssets in ActorDescs are stored in FNames (we can convert them to SoftObjectPtrs if IsUsingDataLayerAsset() returns true)
-		Algo::Transform(GetDataLayers(), DataLayerAssets, [](const FName& InDataLayerAsset) 
-		{
-			return TSoftObjectPtr<UDataLayerAsset>(FSoftObjectPath(InDataLayerAsset.ToString()));
-		});
-		GridDescriptor.SetDataLayerAssets(DataLayerAssets);
-	}
+		return TSoftObjectPtr<UDataLayerAsset>(FSoftObjectPath(InDataLayerAsset.ToString()));
+	});
+	GridDescriptor.SetDataLayerAssets(DataLayerAssets);
 
 	return GridDescriptor;
 }
diff --git a/Engine/Source/Editor/DataLayerEditor/Private/DataLayer/DataLayerEditorSubsystem.cpp b/Engine/Source/Editor/DataLayerEditor/Private/DataLayer/DataLayerEditorSubsystem.cpp
--- a/Engine/Source/Editor/DataLayerEditor/Private/DataLayer/DataLayerEditorSubsystem.cpp
+++ b/Engine/Source/Editor/DataLayerEditor/Private/DataLayer/DataLayerEditorSubsystem.cpp
@@ -60,7 +60,6 @@
 #include "WorldPartition/DataLayer/DataLayerInstanceWithAsset.h"
 #include "WorldPartition/DataLayer/DataLayerUtils.h"
 #include "WorldPartition/DataLayer/DataLayerManager.h"
-#include "WorldPartition/DataLayer/DeprecatedDataLayerInstance.h"
 #include "WorldPartition/DataLayer/ExternalDataLayerInstance.h"
 #include "WorldPartition/DataLayer/ExternalDataLayerAsset.h"
 #include "WorldPartition/DataLayer/ExternalDataLayerManager.h"
@@ -879,13 +878,6 @@ void UDataLayerEditorSubsystem::PostUndoRedo()
 // Operations on an individual actor.
 //
 ///////////////////////////////////////////////////////////////////////////////////////////////////////////////////////
-bool UDataLayerEditorSubsystem::IsActorValidForDataLayer(AActor* Actor)
-{
-	// This is for backward compatibility
-	TSet<TSubclassOf<UDataLayerInstance>> DataLayerInstanceClasses({ UDataLayerInstance::StaticClass() });
-	return IsActorValidForDataLayerForClasses(Actor, DataLayerInstanceClasses);
-}
-
 bool UDataLayerEditorSubsystem::IsActorValidForDataLayerInstances(AActor* Actor, const TArray<UDataLayerInstance*>& DataLayerInstances)
 {
 	TSet<TSubclassOf<UDataLayerInstance>> DataLayerInstanceClasses;
@@ -1797,12 +1789,12 @@ UDataLayerInstance* UDataLayerEditorSubsystem::CreateDataLayerInstance(const FDa
 
 	if (AWorldDataLayers* WorldDataLayers = Parameters.WorldDataLayers != nullptr ? Parameters.WorldDataLayers.Get() : (GetWorld() ? GetWorld()->GetWorldDataLayers() : nullptr))
 	{
-		if (Parameters.bIsPrivate)
+		if (Parameters.bIsPrivate || !Parameters.DataLayerAsset)
 		{
+			ensureMsgf(Parameters.bIsPrivate, TEXT("When a DataLayerAsset is not provided, bIsPrivate should be set to true for clarity"));
 			NewDataLayerInstance = CreateDataLayerInstance<UDataLayerInstancePrivate>(WorldDataLayers);
 		}
-		// Don't create an instance if no valid asset is provided
-		else if (Parameters.DataLayerAsset)
+		else
 		{
 			if (UExternalDataLayerAsset* ExternalDataLayerAsset = Cast<UExternalDataLayerAsset>(Parameters.DataLayerAsset))
 			{
@@ -1848,10 +1840,6 @@ UDataLayerInstance* UDataLayerEditorSubsystem::CreateDataLayerInstance(const FDa
 				NewDataLayerInstance = CreateDataLayerInstance<UDataLayerInstanceWithAsset>(WorldDataLayers, Parameters.DataLayerAsset);
 			}
 		}
-		else
-		{
-			NewDataLayerInstance = CreateDataLayerInstance<UDeprecatedDataLayerInstance>(WorldDataLayers);
-		}
 	}
 
 	if (const UExternalDataLayerInstance* RootExternalDataLayerInstance = NewDataLayerInstance ? NewDataLayerInstance->GetDirectOuterWorldDataLayers()->GetRootExternalDataLayerInstance() : nullptr)
@@ -2182,41 +2170,4 @@ bool UDataLayerEditorSubsystem::ApplyActorEditorContextDataLayersToActors(const
 	return bSuccess;
 }
 
-//~ Begin Deprecated
-
-PRAGMA_DISABLE_DEPRECATION_WARNINGS
-
-bool UDataLayerEditorSubsystem::RenameDataLayer(UDataLayerInstance* DataLayerInstance, const FName& InDataLayerLabel)
-{
-	if (DataLayerInstance->SupportRelabeling())
-	{
-		if (DataLayerInstance->RelabelDataLayer(InDataLayerLabel))
-		{
-			BroadcastDataLayerChanged(EDataLayerAction::Rename, DataLayerInstance, "DataLayerLabel");
-			return true;
-		}	
-	}
-
-	return false;
-}
-
-UDataLayerInstance* UDataLayerEditorSubsystem::GetDataLayerFromLabel(const FName& DataLayerLabel) const
-{
-	if (AWorldDataLayers* WorldDataLayers = GetWorld()->GetWorldDataLayers())
-	{
-		const UDataLayerInstance* DataLayerInstance = WorldDataLayers->GetDataLayerFromLabel(DataLayerLabel);
-		return const_cast<UDataLayerInstance*>(DataLayerInstance);
-	}
-	return nullptr;
-}
-
-PRAGMA_ENABLE_DEPRECATION_WARNINGS
-
-UDataLayerInstance* UDataLayerEditorSubsystem::GetDataLayer(const FActorDataLayer& ActorDataLayer) const
-{
-	return GetDataLayerInstance(ActorDataLayer.Name);
-}
-
-//~ End Deprecated
-
 #undef LOCTEXT_NAMESPACE
diff --git a/Engine/Source/Editor/DataLayerEditor/Public/DataLayer/DataLayerEditorSubsystem.h b/Engine/Source/Editor/DataLayerEditor/Public/DataLayer/DataLayerEditorSubsystem.h
--- a/Engine/Source/Editor/DataLayerEditor/Public/DataLayer/DataLayerEditorSubsystem.h
+++ b/Engine/Source/Editor/DataLayerEditor/Public/DataLayer/DataLayerEditorSubsystem.h
@@ -25,7 +25,6 @@
 #include "UObject/WeakObjectPtr.h"
 #include "UObject/WeakObjectPtrTemplates.h"
 #include "WorldPartition/WorldPartitionHandle.h"
-#include "WorldPartition/DataLayer/ActorDataLayer.h"
 #include "WorldPartition/DataLayer/DataLayerInstance.h"
 #include "WorldPartition/DataLayer/WorldDataLayers.h"
 #include "LevelEditorDragDropHandler.h"
@@ -59,7 +58,7 @@ struct FDataLayerCreationParameters
 
 	UE_API FDataLayerCreationParameters();
 
-	// Required. Will assign the asset to the created instance.
+	// Required if bIsPrivate is false. Will assign the asset to the created instance.
 	UPROPERTY(EditAnywhere, BlueprintReadWrite, Category = "Data Layer")
 	TObjectPtr<UDataLayerAsset> DataLayerAsset;
 
@@ -68,7 +67,7 @@ struct FDataLayerCreationParameters
 	TObjectPtr<AWorldDataLayers> WorldDataLayers;
 
 	UPROPERTY(EditAnywhere, BlueprintReadWrite, Category = "Data Layer")
-	bool bIsPrivate;
+	bool bIsPrivate; 
 };
 
 UCLASS()
@@ -84,7 +83,7 @@ public:
 	TArray<TSoftObjectPtr<const UDataLayerAsset>> DataLayerAssets;
 };
 
-UCLASS(MinimalAPI)
+UCLASS(BlueprintType, MinimalAPI)
 class UDataLayerEditorSubsystem final : public UEditorSubsystem, public IActorEditorContextClient, public FTickableGameObject
 {
 	GENERATED_BODY()
@@ -171,14 +170,9 @@ public:
 	///////////////////////////////////////////////////////////////////////////////////////////////////////////////////////
 	// Operations on an individual actor.
 
-	/**
-	 *	Checks to see if the specified actor is in an appropriate state to interact with DataLayers
-	 *
-	 *	@param	Actor	The actor to validate
-	 */
-	UE_DEPRECATED(5.4, "Use IsActorValidForDataLayerInstances instead.")
-	UFUNCTION(BlueprintCallable, Category = DataLayers)
-	UE_API bool IsActorValidForDataLayer(AActor* Actor);
+	UE_DEPRECATED(all, "Use IsActorValidForDataLayerInstances instead.")
+	UFUNCTION(meta = (DeprecatedFunction, DeprecationMessage = "Use IsActorValidForDataLayerInstances instead."))
+	bool IsActorValidForDataLayer(AActor* Actor) { return false; }
 
 	/**
 	 *	Checks to see if the specified actor is in an appropriate state to interact with DataLayers
@@ -736,22 +730,22 @@ public:
 	UFUNCTION(BlueprintCallable, Category = DataLayers, meta = (DeprecatedFunction, DeprecationMessage = "Use ToggleDataLayersIsLoadedInEditor instead"))
 	bool ToggleDataLayersIsDynamicallyLoadedInEditor(const TArray<UDEPRECATED_DataLayer*>& DataLayers, const bool bIsFromUserChange) { return false; }
 	
-	PRAGMA_ENABLE_DEPRECATION_WARNINGS
+	UE_DEPRECATED(all, "Use GetDataLayerInstance instead")
+	UFUNCTION(meta = (DeprecatedFunction, DeprecationMessage = "Use GetDataLayerInstance instead"))
+	UDataLayerInstance* GetDataLayer(const struct FActorDataLayer& ActorDataLayer) const { return nullptr; }
 
-	UE_DEPRECATED(5.1, "Use GetDataLayerInstance instead")
-	UFUNCTION(BlueprintCallable, Category = DataLayers)
-	UE_API UDataLayerInstance* GetDataLayer(const FActorDataLayer& ActorDataLayer) const;
+	PRAGMA_ENABLE_DEPRECATION_WARNINGS
 
-	UE_DEPRECATED(5.1, "Use GetDataLayerInstance instead")
-	UFUNCTION(BlueprintCallable, Category = DataLayers)
-	UE_API UDataLayerInstance* GetDataLayerFromLabel(const FName& DataLayerLabel) const;
+	UE_DEPRECATED(all, "Use GetDataLayerInstance instead")
+	UFUNCTION(meta = (DeprecatedFunction, DeprecationMessage = "Use GetDataLayerInstance instead"))
+	UDataLayerInstance* GetDataLayerFromLabel(const FName& DataLayerLabel) const { return nullptr; }
 
-	UE_DEPRECATED(5.1, "Renaming is not permitted anymore")
-	UFUNCTION(BlueprintCallable, Category = DataLayers)
-	UE_API bool RenameDataLayer(UDataLayerInstance* DataLayer, const FName& NewDataLayerLabel);
+	UE_DEPRECATED(all, "Renaming is not permitted anymore")
+	UFUNCTION(meta = (DeprecatedFunction, DeprecationMessage = "Renaming is not permitted anymore"))
+	bool RenameDataLayer(UDataLayerInstance* DataLayer, const FName& NewDataLayerLabel) { return false; }
 
-	UE_DEPRECATED(5.1, "Use CreateDataLayerInstance with FDataLayerCreationParameters instead")
-	UFUNCTION(BlueprintCallable, Category = DataLayers)
+	UE_DEPRECATED(all, "Use CreateDataLayerInstance with FDataLayerCreationParameters instead")
+	UFUNCTION(meta = (DeprecatedFunction, DeprecationMessage = "Use CreateDataLayerInstance with FDataLayerCreationParameters instead"))
 	UDataLayerInstance* CreateDataLayer(UDataLayerInstance* ParentDataLayer = nullptr) { return nullptr; }
 
 	//~ End Deprecated
diff --git a/Engine/Source/Editor/DetailCustomizations/Private/ActorDetails.cpp b/Engine/Source/Editor/DetailCustomizations/Private/ActorDetails.cpp
--- a/Engine/Source/Editor/DetailCustomizations/Private/ActorDetails.cpp
+++ b/Engine/Source/Editor/DetailCustomizations/Private/ActorDetails.cpp
@@ -157,7 +157,6 @@ void FActorDetails::CustomizeDetails( IDetailLayoutBuilder& DetailLayout )
 		{
 			DetailLayout.HideProperty(DetailLayout.GetProperty(AActor::GetRuntimeGridPropertyName(), AActor::StaticClass()));
 			DetailLayout.HideProperty(DetailLayout.GetProperty(AActor::GetDataLayerAssetsPropertyName(), AActor::StaticClass()));
-			DetailLayout.HideProperty(DetailLayout.GetProperty(AActor::GetDataLayerPropertyName(), AActor::StaticClass()));
 			DetailLayout.HideProperty(DetailLayout.GetProperty(AActor::GetExternalDataLayerAssetPropertyName(), AActor::StaticClass()));
 			DetailLayout.HideProperty(DetailLayout.GetProperty(AActor::GetHLODLayerPropertyName(), AActor::StaticClass()));
 		}
diff --git a/Engine/Source/Editor/UnrealEd/Classes/Commandlets/WorldPartitionDataLayerToAssetCommandLet.h b/Engine/Source/Editor/UnrealEd/Classes/Commandlets/WorldPartitionDataLayerToAssetCommandLet.h
--- a/Engine/Source/Editor/UnrealEd/Classes/Commandlets/WorldPartitionDataLayerToAssetCommandLet.h
+++ b/Engine/Source/Editor/UnrealEd/Classes/Commandlets/WorldPartitionDataLayerToAssetCommandLet.h
@@ -13,7 +13,7 @@
 #include "WorldPartition/DataLayer/DataLayerAsset.h"
 #include "WorldPartition/DataLayer/DataLayerInstance.h"
 #include "WorldPartition/DataLayer/DeprecatedDataLayerInstance.h"
-#include "WorldPartition/DataLayer/ActorDataLayer.h"
+#include "WorldPartition/DataLayer/Deprecated/DeprecatedDataLayer.h"
 
 #include "WorldPartitionDataLayerToAssetCommandLet.generated.h"
 
@@ -104,7 +104,10 @@ public:
 	UDataLayerConversionInfo* GetDataLayerConversionInfo(const UDeprecatedDataLayerInstance* DataLayer) const;
 	UDataLayerConversionInfo* GetDataLayerConversionInfo(const UDataLayerAsset* DataLayerAsset) const;
 	UDataLayerConversionInfo* GetDataLayerConversionInfo(const UDataLayerInstanceWithAsset* DataLayerInstance) const;
+
+	PRAGMA_DISABLE_DEPRECATION_WARNINGS
 	UDataLayerConversionInfo* GetDataLayerConversionInfo(const FActorDataLayer& ActorDataLayer) const;
+	PRAGMA_ENABLE_DEPRECATION_WARNINGS
 
 	UDataLayerConversionInfo* StoreExistingDataLayer(FAssetData& AssetData);
 	UDataLayerConversionInfo* StoreDataLayerAssetConversion(const UDeprecatedDataLayerInstance* DataLayerToConvert, UDataLayerAsset* NewDataLayerAsset);
diff --git a/Engine/Source/Editor/UnrealEd/Private/Commandlets/WorldPartitionDataLayerToAssetCommandlet.cpp b/Engine/Source/Editor/UnrealEd/Private/Commandlets/WorldPartitionDataLayerToAssetCommandlet.cpp
--- a/Engine/Source/Editor/UnrealEd/Private/Commandlets/WorldPartitionDataLayerToAssetCommandlet.cpp
+++ b/Engine/Source/Editor/UnrealEd/Private/Commandlets/WorldPartitionDataLayerToAssetCommandlet.cpp
@@ -74,6 +74,7 @@ UDataLayerConversionInfo* UDataLayerToAssetCommandletContext::GetDataLayerConver
 	return Entry != nullptr ? Entry->Get() : nullptr;
 }
 
+PRAGMA_DISABLE_DEPRECATION_WARNINGS
 UDataLayerConversionInfo* UDataLayerToAssetCommandletContext::GetDataLayerConversionInfo(const FActorDataLayer& ActorDataLayer) const
 {
 	const TObjectPtr<UDataLayerConversionInfo>* Entry = DataLayerConversionInfo.FindByPredicate([&ActorDataLayer](const TObjectPtr<UDataLayerConversionInfo>& Other)
@@ -83,6 +84,7 @@ UDataLayerConversionInfo* UDataLayerToAssetCommandletContext::GetDataLayerConver
 	
 	return Entry != nullptr ? Entry->Get() : nullptr;
 }
+PRAGMA_ENABLE_DEPRECATION_WARNINGS
 
 UDataLayerConversionInfo* UDataLayerToAssetCommandletContext::StoreExistingDataLayer(FAssetData& AssetData)
 {
@@ -612,7 +614,9 @@ bool UDataLayerToAssetCommandlet::CreateDataLayerInstances(TStrongObjectPtr<UDat
 			CommandletContext->StoreDataLayerInstanceConversion(ConvertingInfo->DataLayerAsset, DataLayerInstance);
 		}
 
+		PRAGMA_DISABLE_DEPRECATION_WARNINGS
 		WorldDataLayers->DeprecatedDataLayerNameToDataLayerInstance.Add(ConvertingInfo->DataLayerToConvert->GetDataLayerFName(), ConvertingInfo->DataLayerInstance);
+		PRAGMA_ENABLE_DEPRECATION_WARNINGS
 	}
 
 	uint32 ErrorCount = 0;
diff --git a/Engine/Source/Editor/UnrealEd/Private/WorldPartition/WorldPartitionMiniMapBuilder.cpp b/Engine/Source/Editor/UnrealEd/Private/WorldPartition/WorldPartitionMiniMapBuilder.cpp
--- a/Engine/Source/Editor/UnrealEd/Private/WorldPartition/WorldPartitionMiniMapBuilder.cpp
+++ b/Engine/Source/Editor/UnrealEd/Private/WorldPartition/WorldPartitionMiniMapBuilder.cpp
@@ -17,7 +17,6 @@
 #include "UObject/StrongObjectPtr.h"
 
 #include "WorldPartition/ActorDescContainer.h"
-#include "WorldPartition/DataLayer/ActorDataLayer.h"
 #include "WorldPartition/DataLayer/DataLayerInstance.h"
 #include "WorldPartition/DataLayer/DataLayerManager.h"
 #include "WorldPartition/HLOD/HLODActor.h"
diff --git a/Engine/Source/Runtime/Engine/Classes/GameFramework/Actor.h b/Engine/Source/Runtime/Engine/Classes/GameFramework/Actor.h
--- a/Engine/Source/Runtime/Engine/Classes/GameFramework/Actor.h
+++ b/Engine/Source/Runtime/Engine/Classes/GameFramework/Actor.h
@@ -1097,9 +1097,13 @@ protected:
 	UPROPERTY(BluePrintReadOnly, AdvancedDisplay, Category=Actor, TextExportTransient, NonTransactional)
 	FGuid ContentBundleGuid;
 
-	/** DataLayers the actor belongs to.*/
-	UPROPERTY(VisibleAnywhere, AdvancedDisplay, Category = DataLayers)
+#if WITH_EDITORONLY_DATA
+	PRAGMA_DISABLE_DEPRECATION_WARNINGS
+	UE_DEPRECATED(5.8, "Use DataLayerAssets instead.")
+	UPROPERTY(meta = (DeprecatedProperty, DeprecationMessage = "Use DataLayerAssets instead."))
 	TArray<FActorDataLayer> DataLayers;
+	PRAGMA_ENABLE_DEPRECATION_WARNINGS
+#endif // WITH_EDITORONLY_DATA
 
 	/** DataLayers assets the actor belongs to.*/
 	UPROPERTY(EditAnywhere, AdvancedDisplay, Category = DataLayers, meta=(DisableLevelInstancePropertyOverride))
@@ -1482,7 +1486,9 @@ public:
 	ENGINE_API void FixupDataLayers(bool bRevertChangesOnLockedDataLayer = false);
 	static const FName GetDataLayerAssetsPropertyName() { return GET_MEMBER_NAME_CHECKED(AActor, DataLayerAssets); }
 	static const FName GetExternalDataLayerAssetPropertyName() { return GET_MEMBER_NAME_CHECKED(AActor, ExternalDataLayerAsset); }
-	static const FName GetDataLayerPropertyName() { return GET_MEMBER_NAME_CHECKED(AActor, DataLayers); }
+
+	UE_DEPRECATED(5.8, "FActorDataLayer is no longer in use.")
+	static const FName GetDataLayerPropertyName() { return FName(); }
 
 	ENGINE_API TArray<const UDataLayerAsset*> GetDataLayerAssets(bool bIncludeExternalDataLayerAsset = true) const;
 	ENGINE_API bool HasExternalContent() const;
@@ -1490,36 +1496,12 @@ public:
 	virtual bool ActorTypeSupportsDataLayer() const { return true; }
 	virtual bool ActorTypeSupportsExternalDataLayer() const { return true; }
 
-	//~ Begin Deprecated
-
-	UE_DEPRECATED(5.3, "Use AActor::SupportsDataLayerType(TSubclassOf<UDataLayerInstance>) instead.")
-	ENGINE_API bool SupportsDataLayer() const;
-
-	UE_DEPRECATED(5.3, "Use UDataLayerInstance::AddActor(AActor*) instead.")
-	bool AddDataLayer(const UDataLayerAsset* DataLayerAsset) { return false; }
-
-	UE_DEPRECATED(5.3, "Use UDataLayerInstance::RemoveDataLayer(AActor*) instead.")
-	bool RemoveDataLayer(const UDataLayerAsset* DataLayerAsset) { return false; }
-
-	UE_DEPRECATED(5.1, "Convert DataLayer using UDataLayerToAssetCommandlet and use AddDataLayer(UDataLayerAsset*)")
-	ENGINE_API bool AddDataLayer(const FActorDataLayer& ActorDataLayer);
-
-	UE_DEPRECATED(5.1, "Convert DataLayer using UDataLayerToAssetCommandlet and use RemoveDataLayer(UDataLayerAsset*)")
-	ENGINE_API bool RemoveDataLayer(const FActorDataLayer& ActorDataLayer);
-
-	UE_DEPRECATED(5.1, "Convert DataLayer using UDataLayerToAssetCommandlet and use ContainsDataLayer(const UDataLayerAsset*)")
-	ENGINE_API bool ContainsDataLayer(const FActorDataLayer& ActorDataLayer) const;
-
-	UE_DEPRECATED(5.1, "Convert DataLayer using UDataLayerToAssetCommandlet and use GetDataLayerAssets() instead")
+	// This getter is kept solely to enable UDataLayerToAssetCommandlet for conversions of old data
+	PRAGMA_DISABLE_DEPRECATION_WARNINGS
+	UE_DEPRECATED(all, "Convert DataLayer using UDataLayerToAssetCommandlet and use GetDataLayerAssets() instead")
 	TArray<FActorDataLayer> const& GetActorDataLayers() const { return DataLayers; }
+	PRAGMA_ENABLE_DEPRECATION_WARNINGS
 
-	UE_DEPRECATED(5.1, "Use GetDataLayerInstances() with no parameters instead")
-	TArray<const UDataLayerInstance*> GetDataLayerInstances(const AWorldDataLayers* WorldDataLayers) const { return TArray<const UDataLayerInstance*>(); }
-
-	UE_DEPRECATED(5.1, "Use HasDataLayers() instead")
-	bool HasValidDataLayers() const { return HasDataLayers(); }
-
-	//~ End Deprecated
 #endif
 
 	ENGINE_API TArray<const UDataLayerInstance*> GetDataLayerInstances() const;
diff --git a/Engine/Source/Runtime/Engine/Private/Actor.cpp b/Engine/Source/Runtime/Engine/Private/Actor.cpp
--- a/Engine/Source/Runtime/Engine/Private/Actor.cpp
+++ b/Engine/Source/Runtime/Engine/Private/Actor.cpp
@@ -7139,10 +7139,6 @@ TArray<const UDataLayerInstance*> AActor::GetDataLayerInstancesInternal(bool bUs
 		if (UDataLayerManager* DataLayerManager = bUseLevelContext ? UDataLayerManager::GetDataLayerManager(this) : UDataLayerManager::GetDataLayerManager(GetWorld()))
 		{
 			DataLayerInstances += DataLayerManager->GetDataLayerInstances(GetDataLayerAssets());
-
-			PRAGMA_DISABLE_DEPRECATION_WARNINGS
-			DataLayerInstances += DataLayerManager->GetDataLayerInstances(DataLayers);
-			PRAGMA_ENABLE_DEPRECATION_WARNINGS
 		}
 
 		if (bIncludeParentDataLayers)
diff --git a/Engine/Source/Runtime/Engine/Private/ActorEditor.cpp b/Engine/Source/Runtime/Engine/Private/ActorEditor.cpp
--- a/Engine/Source/Runtime/Engine/Private/ActorEditor.cpp
+++ b/Engine/Source/Runtime/Engine/Private/ActorEditor.cpp
@@ -26,7 +26,6 @@
 #include "Elements/Framework/EngineElementsLibrary.h"
 #include "WorldPartition/DataLayer/IDataLayerEditorModule.h"
 #include "WorldPartition/DataLayer/DataLayerInstanceWithAsset.h"
-#include "WorldPartition/DataLayer/DeprecatedDataLayerInstance.h"
 #include "WorldPartition/DataLayer/ExternalDataLayerInstance.h"
 #include "WorldPartition/DataLayer/ExternalDataLayerAsset.h"
 #include "WorldPartition/ContentBundle/ContentBundlePaths.h"
@@ -1851,7 +1850,6 @@ void AActor::FixupDataLayers(bool bRevertChangesOnLockedDataLayer /*= false*/)
 
 	if (!SupportsDataLayerType(UDataLayerInstance::StaticClass()))
 	{
-		DataLayers.Empty();
 		DataLayerAssets.Empty();
 	}
 
@@ -1860,7 +1858,7 @@ void AActor::FixupDataLayers(bool bRevertChangesOnLockedDataLayer /*= false*/)
 		ExternalDataLayerAsset = nullptr;
 	}
 
-	if (DataLayers.IsEmpty() && DataLayerAssets.IsEmpty())
+	if (DataLayerAssets.IsEmpty())
 	{
 		return;
 	}
@@ -1941,7 +1939,6 @@ void AActor::FixupDataLayers(bool bRevertChangesOnLockedDataLayer /*= false*/)
 		}
 	}
 
-	PRAGMA_DISABLE_DEPRECATION_WARNINGS
 	auto CleanupDataLayers = [this, DataLayerManager](auto& DataLayerArray)
 	{
 		using ArrayType = typename TRemoveReference<decltype(DataLayerArray)>::Type;
@@ -1971,8 +1968,6 @@ void AActor::FixupDataLayers(bool bRevertChangesOnLockedDataLayer /*= false*/)
 		DataLayerAssets.Empty(CleanedUpDataLayers.Num());
 		Algo::Transform(CleanedUpDataLayers, DataLayerAssets, [](const UDataLayerAsset* DataLayerAsset) { return TSoftObjectPtr<UDataLayerAsset>(const_cast<UDataLayerAsset*>(DataLayerAsset)); });
 	}
-	CleanupDataLayers(DataLayers);
-	PRAGMA_ENABLE_DEPRECATION_WARNINGS
 }
 
 bool AActor::IsPropertyChangedAffectingDataLayers(FPropertyChangedEvent& PropertyChangedEvent) const
@@ -2096,26 +2091,11 @@ bool AActor::CanAddDataLayer(const UDataLayerInstance* InDataLayerInstance, FTex
 	}
 	else
 	{
-		if (const UDeprecatedDataLayerInstance* DataLayerInstance = Cast<UDeprecatedDataLayerInstance>(InDataLayerInstance))
-		{
-			if (DataLayers.Contains(DataLayerInstance->GetActorDataLayer()))
-			{
-				if (OutReason)
-				{
-					*OutReason = LOCTEXT("CantAddDataLayerActorAlreadyAssignedToDataLayer", "Actor is already assigned to this data layer.");
-				}
-				return false;
-			}
-			return true;
-		}
-		else
+		if (OutReason)
 		{
-			if (OutReason)
-			{
-				*OutReason = LOCTEXT("CantAddDataLayerInvalidDataLayerAsset", "Invalid data layer asset.");
-			}
-			return false;
+			*OutReason = LOCTEXT("CantAddDataLayerInvalidDataLayerAsset", "Invalid data layer asset.");
 		}
+		return false;
 	}
 
 	return true;
@@ -2167,43 +2147,6 @@ bool FAssignActorDataLayer::RemoveDataLayerAsset(AActor* InActor, const UDataLay
 	return false;
 }
 
-//~ Begin Deprecated
-
-PRAGMA_DISABLE_DEPRECATION_WARNINGS
-
-bool AActor::AddDataLayer(const FActorDataLayer& ActorDataLayer)
-{
-	if (SupportsDataLayerType(UDataLayerInstance::StaticClass()) && !DataLayers.Contains(ActorDataLayer))
-	{
-		Modify();
-		DataLayers.Add(ActorDataLayer);
-		return true;
-	}
-
-	return false;
-}
-
-bool AActor::RemoveDataLayer(const FActorDataLayer& ActorDataLayer)
-{
-	if (DataLayers.Contains(ActorDataLayer))
-	{
-		Modify();
-		DataLayers.Remove(ActorDataLayer);
-		return true;
-	}
-
-	return false;
-}
-
-PRAGMA_ENABLE_DEPRECATION_WARNINGS
-
-bool AActor::ContainsDataLayer(const FActorDataLayer& ActorDataLayer) const
-{
-	return DataLayers.Contains(ActorDataLayer);
-}
-
-//~ End Deprecated
-
 // DataLayers (end)
 //---------------------------------------------------------------------------
 
diff --git a/Engine/Source/Runtime/Engine/Private/WorldPartition/DataLayer/DataLayerSubsystem.cpp b/Engine/Source/Runtime/Engine/Private/WorldPartition/DataLayer/DataLayerSubsystem.cpp
--- a/Engine/Source/Runtime/Engine/Private/WorldPartition/DataLayer/DataLayerSubsystem.cpp
+++ b/Engine/Source/Runtime/Engine/Private/WorldPartition/DataLayer/DataLayerSubsystem.cpp
@@ -150,17 +150,6 @@ UDataLayerInstance* UDataLayerSubsystem::GetDataLayer(const FActorDataLayer& InD
 	return GetDataLayerFromName(InDataLayer.Name);
 }
 
-UDataLayerInstance* UDataLayerSubsystem::GetDataLayerFromLabel(FName InDataLayerLabel) const
-{
-	if (AWorldDataLayers* WorldDataLayers = GetWorld()->GetWorldDataLayers())
-	{
-		const UDataLayerInstance* DataLayerInstance = WorldDataLayers->GetDataLayerFromLabel(InDataLayerLabel);
-		return const_cast<UDataLayerInstance*>(DataLayerInstance);
-	}
-
-	return nullptr;
-}
-
 EDataLayerRuntimeState UDataLayerSubsystem::GetDataLayerRuntimeStateByLabel(const FName& InDataLayerLabel) const
 {
 	return GetDataLayerRuntimeState(GetDataLayerFromLabel(InDataLayerLabel));
diff --git a/Engine/Source/Runtime/Engine/Private/WorldPartition/DataLayer/DataLayerUtils.cpp b/Engine/Source/Runtime/Engine/Private/WorldPartition/DataLayer/DataLayerUtils.cpp
--- a/Engine/Source/Runtime/Engine/Private/WorldPartition/DataLayer/DataLayerUtils.cpp
+++ b/Engine/Source/Runtime/Engine/Private/WorldPartition/DataLayer/DataLayerUtils.cpp
@@ -25,68 +25,38 @@ bool FDataLayerUtils::ForEachDataLayerInstance(const UDataLayerManager* InDataLa
 		return true;
 	}
 	
-	// DataLayers not using DataLayer Assets represent DataLayerInstanceNames
-	if (!InActorDesc->IsUsingDataLayerAsset())
+	if (InDataLayerManager && InDataLayerManager->CanResolveDataLayers())
 	{
-		if (InDataLayerManager && InDataLayerManager->CanResolveDataLayers())
+		if (const UDataLayerInstance* ExternalDataLayerInstance = InDataLayerManager->GetDataLayerInstanceFromAssetName(ActorDescExternalDataLayer))
 		{
-			for (const FName& DataLayerInstanceName : ActorDescNonExternalDataLayers)
-			{
-				if (const UDataLayerInstance* DataLayerInstance = InDataLayerManager->GetDataLayerInstance(DataLayerInstanceName))
-				{
-					Callback(DataLayerInstance, nullptr, false);
-				}
-			}
-			return true;
+			Callback(ExternalDataLayerInstance, nullptr, true);
 		}
-		// Fallback on FWorldDataLayersActorDesc
-		else if (!InWorldDataLayersActorDescs.IsEmpty() && AreWorldDataLayersActorDescsSane(InWorldDataLayersActorDescs))
+		for (const FName& DataLayerAssetPath : ActorDescNonExternalDataLayers)
 		{
-			for (const FName& DataLayerInstanceName : ActorDescNonExternalDataLayers)
+			if (const UDataLayerInstance* DataLayerInstance = InDataLayerManager->GetDataLayerInstanceFromAssetName(DataLayerAssetPath))
 			{
-				if (const FDataLayerInstanceDesc* DataLayerInstanceDesc = GetDataLayerInstanceDescFromInstanceName(InWorldDataLayersActorDescs, DataLayerInstanceName))
-				{
-					Callback(nullptr, DataLayerInstanceDesc, false);
-				}
+				Callback(DataLayerInstance, nullptr, false);
 			}
-			return true;
 		}
+		return true;
 	}
-	// ActorDesc DataLayers represents DataLayer Asset Paths
-	else
+	// Fallback on FWorldDataLayersActorDesc
+	else if (!InWorldDataLayersActorDescs.IsEmpty() && AreWorldDataLayersActorDescsSane(InWorldDataLayersActorDescs))
 	{
-		if (InDataLayerManager && InDataLayerManager->CanResolveDataLayers())
+		if (const FDataLayerInstanceDesc* ExternalDataLayerInstanceDesc = GetDataLayerInstanceDescFromAssetPath(InWorldDataLayersActorDescs, ActorDescExternalDataLayer))
 		{
-			if (const UDataLayerInstance* ExternalDataLayerInstance = InDataLayerManager->GetDataLayerInstanceFromAssetName(ActorDescExternalDataLayer))
-			{
-				Callback(ExternalDataLayerInstance, nullptr, true);
-			}
-			for (const FName& DataLayerAssetPath : ActorDescNonExternalDataLayers)
-			{
-				if (const UDataLayerInstance* DataLayerInstance = InDataLayerManager->GetDataLayerInstanceFromAssetName(DataLayerAssetPath))
-				{
-					Callback(DataLayerInstance, nullptr, false);
-				}
-			}
-			return true;
+			Callback(nullptr, ExternalDataLayerInstanceDesc, true);
 		}
-		// Fallback on FWorldDataLayersActorDesc
-		else if (!InWorldDataLayersActorDescs.IsEmpty() && AreWorldDataLayersActorDescsSane(InWorldDataLayersActorDescs))
+		for (const FName& DataLayerAssetPath : ActorDescNonExternalDataLayers)
 		{
-			if (const FDataLayerInstanceDesc* ExternalDataLayerInstanceDesc = GetDataLayerInstanceDescFromAssetPath(InWorldDataLayersActorDescs, ActorDescExternalDataLayer))
-			{
-				Callback(nullptr, ExternalDataLayerInstanceDesc, true);
-			}
-			for (const FName& DataLayerAssetPath : ActorDescNonExternalDataLayers)
+			if (const FDataLayerInstanceDesc* DataLayerInstanceDesc = GetDataLayerInstanceDescFromAssetPath(InWorldDataLayersActorDescs, DataLayerAssetPath))
 			{
-				if (const FDataLayerInstanceDesc* DataLayerInstanceDesc = GetDataLayerInstanceDescFromAssetPath(InWorldDataLayersActorDescs, DataLayerAssetPath))
-				{
-					Callback(nullptr, DataLayerInstanceDesc, false);
-				}
+				Callback(nullptr, DataLayerInstanceDesc, false);
 			}
-			return true;
 		}
+		return true;
 	}
+
 	return false;
 }
 
diff --git a/Engine/Source/Runtime/Engine/Private/WorldPartition/DataLayer/DeprecatedDataLayerInstance.cpp b/Engine/Source/Runtime/Engine/Private/WorldPartition/DataLayer/DeprecatedDataLayerInstance.cpp
--- a/Engine/Source/Runtime/Engine/Private/WorldPartition/DataLayer/DeprecatedDataLayerInstance.cpp
+++ b/Engine/Source/Runtime/Engine/Private/WorldPartition/DataLayer/DeprecatedDataLayerInstance.cpp
@@ -49,21 +49,6 @@ FActorDataLayer UDeprecatedDataLayerInstance::GetActorDataLayer() const
 	return FActorDataLayer(GetDataLayerFName());
 }
 
-bool UDeprecatedDataLayerInstance::PerformAddActor(AActor* InActor) const
-{
-	return InActor->AddDataLayer(GetActorDataLayer());
-}
-
-bool UDeprecatedDataLayerInstance::PerformRemoveActor(AActor* InActor) const
-{
-	return InActor->RemoveDataLayer(GetActorDataLayer());
-}
-
-bool UDeprecatedDataLayerInstance::RelabelDataLayer(FName InDataLayerLabel)
-{
-	return FDataLayerUtils::SetDataLayerShortName(this, InDataLayerLabel.ToString());
-}
-
 #endif // WITH_EDITOR
 
 PRAGMA_ENABLE_DEPRECATION_WARNINGS
diff --git a/Engine/Source/Runtime/Engine/Private/WorldPartition/DataLayer/WorldDataLayers.cpp b/Engine/Source/Runtime/Engine/Private/WorldPartition/DataLayer/WorldDataLayers.cpp
--- a/Engine/Source/Runtime/Engine/Private/WorldPartition/DataLayer/WorldDataLayers.cpp
+++ b/Engine/Source/Runtime/Engine/Private/WorldPartition/DataLayer/WorldDataLayers.cpp
@@ -7,7 +7,6 @@
 #include "WorldPartition/DataLayer/WorldDataLayers.h"
 #include "WorldPartition/DataLayer/DataLayerManager.h"
 #include "WorldPartition/DataLayer/DataLayerInstanceWithAsset.h"
-#include "WorldPartition/DataLayer/DeprecatedDataLayerInstance.h"
 #include "WorldPartition/DataLayer/DataLayerUtils.h"
 #include "WorldPartition/DataLayer/ExternalDataLayerInstance.h"
 #include "WorldPartition/DataLayer/ExternalDataLayerAsset.h"
@@ -711,21 +710,6 @@ int32 AWorldDataLayers::RemoveDataLayers(const TArray<UDataLayerInstance*>& InDa
 			DataLayerInstance->Modify();
 			DataLayerInstance->OnRemovedFromWorldDataLayers();
 			UsedDataLayerInstances.Remove(DataLayerInstance);
-			if (DataLayerInstance->IsA<UDeprecatedDataLayerInstance>())
-			{
-				DeprecatedDataLayerNameToDataLayerInstance.Remove(DataLayerInstance->GetDataLayerFName());
-			}
-			else if (!DeprecatedDataLayerNameToDataLayerInstance.IsEmpty())
-			{
-				for (TMap<FName, TWeakObjectPtr<UDataLayerInstance>>::TIterator It = DeprecatedDataLayerNameToDataLayerInstance.CreateIterator(); It; ++It)
-				{
-					if (It.Value().Get() == DataLayerInstance)
-					{
-						It.RemoveCurrent();
-						break;
-					}
-				}
-			}
 			if (DataLayerInstance->IsPackageExternal())
 			{
 				DataLayerInstance->MarkAsGarbage();
@@ -1001,13 +985,6 @@ const UDataLayerInstance* AWorldDataLayers::GetDataLayerInstance(const FName& In
 	}
 #endif
 
-#if DATALAYER_TO_INSTANCE_RUNTIME_CONVERSION_ENABLED
-	if (TWeakObjectPtr<UDataLayerInstance> const* FoundDataLayerInstance = DeprecatedDataLayerNameToDataLayerInstance.Find(InDataLayerInstanceName))
-	{
-		return FoundDataLayerInstance->Get();
-	}
-#endif // DATALAYER_TO_INSTANCE_RUNTIME_CONVERSION_ENABLED
-
 	return nullptr;
 }
 
@@ -1524,80 +1501,6 @@ bool AWorldDataLayers::IsEditorOnly() const
 
 #endif
 
-//~ Begin Deprecated
-
-PRAGMA_DISABLE_DEPRECATION_WARNINGS
-
-const UDataLayerInstance* AWorldDataLayers::GetDataLayerFromLabel(const FName& InDataLayerLabel) const
-{
-	const FString DataLayerLabelSanitized = *FDataLayerUtils::GetSanitizedDataLayerShortName(InDataLayerLabel.ToString());
-	const UDataLayerInstance* FoundDataLayerInstance = nullptr;
-	ForEachDataLayerInstance([&DataLayerLabelSanitized, &FoundDataLayerInstance](UDataLayerInstance* DataLayerInstance)
-	{
-		if (DataLayerInstance->GetDataLayerShortName() == DataLayerLabelSanitized)
-		{
-			FoundDataLayerInstance = DataLayerInstance;
-			return false;
-		}
-		return true;
-	});
-
-	return FoundDataLayerInstance;
-}
-
-PRAGMA_ENABLE_DEPRECATION_WARNINGS
-
-const UDataLayerInstance* AWorldDataLayers::GetDataLayerInstance(const FActorDataLayer& InActorDataLayer) const
-{
-	return GetDataLayerInstance(InActorDataLayer.Name);
-}
-
-TArray<FName> AWorldDataLayers::GetDataLayerInstanceNames(const TArray<FActorDataLayer>& InActorDataLayers) const
-{
-#if DATALAYER_TO_INSTANCE_RUNTIME_CONVERSION_ENABLED
-	TArray<FName> OutDataLayerNames;
-	OutDataLayerNames.Reserve(InActorDataLayers.Num());
-
-	for (const FActorDataLayer& ActorDataLayer : InActorDataLayers)
-	{
-		const UDataLayerInstance* DataLayerInstance = GetDataLayerInstance(ActorDataLayer.Name);
-		if (DataLayerInstance != nullptr)
-		{
-			OutDataLayerNames.Add(DataLayerInstance->GetDataLayerFName());
-		}
-	}
-
-	return OutDataLayerNames;
-#else 
-	static_assert(0, "AWorldDataLayers::GetDataLayerInstanceNames function is deprecated and needs to be deleted.");
-	return TArray<FName>();
-#endif
-}
-
-TArray<const UDataLayerInstance*> AWorldDataLayers::GetDataLayerInstances(const TArray<FActorDataLayer>& InActorDataLayers) const
-{
-#if DATALAYER_TO_INSTANCE_RUNTIME_CONVERSION_ENABLED
-	TArray<const UDataLayerInstance*> OutDataLayerInstances;
-	OutDataLayerInstances.Reserve(InActorDataLayers.Num());
-
-	for (const FActorDataLayer& ActorDataLayer : InActorDataLayers)
-	{
-		const UDataLayerInstance* DataLayerInstance = GetDataLayerInstance(ActorDataLayer.Name);
-		if (DataLayerInstance != nullptr)
-		{
-			OutDataLayerInstances.Add(DataLayerInstance);
-		}
-	}
-
-	return OutDataLayerInstances;
-#else 
-	static_assert(0, "AWorldDataLayers::GetDataLayerInstances function is deprecated and needs to be deleted.");
-	return TArray<const UDataLayerInstance*>();
-#endif
-}
-
-//~ End Deprecated
-
 FWorldDataLayersEffectiveStates::FWorldDataLayersEffectiveStates()
 	: UpdateEpoch(0)
 	, AllEffectiveActiveDataLayerNamesEpoch(MAX_int32)
diff --git a/Engine/Source/Runtime/Engine/Private/WorldPartition/HLOD/HLODActorDesc.cpp b/Engine/Source/Runtime/Engine/Private/WorldPartition/HLOD/HLODActorDesc.cpp
--- a/Engine/Source/Runtime/Engine/Private/WorldPartition/HLOD/HLODActorDesc.cpp
+++ b/Engine/Source/Runtime/Engine/Private/WorldPartition/HLOD/HLODActorDesc.cpp
@@ -206,7 +206,6 @@ bool FHLODActorDesc::EqualsUnbuilt(const FHLODActorDesc& Other) const
 		bIsSpatiallyLoaded == Other.bIsSpatiallyLoaded &&
 		bActorIsRuntimeOnly == Other.bActorIsRuntimeOnly &&
 		bActorIsMainWorldOnly == Other.bActorIsMainWorldOnly &&
-		bIsUsingDataLayerAsset == Other.bIsUsingDataLayerAsset &&
 		HLODLayer == Other.HLODLayer &&
 		ContentBundleGuid == Other.ContentBundleGuid &&
 		ComponentDescs == Other.ComponentDescs &&
diff --git a/Engine/Source/Runtime/Engine/Private/WorldPartition/WorldPartitionActorDesc.cpp b/Engine/Source/Runtime/Engine/Private/WorldPartition/WorldPartitionActorDesc.cpp
--- a/Engine/Source/Runtime/Engine/Private/WorldPartition/WorldPartitionActorDesc.cpp
+++ b/Engine/Source/Runtime/Engine/Private/WorldPartition/WorldPartitionActorDesc.cpp
@@ -154,47 +154,24 @@ void FWorldPartitionActorDesc::Init(const AActor* InActor)
 	
 	// DataLayers
 	{
-		TArray<FName> LocalDataLayerAssetPaths;
-		TArray<FName> LocalDataLayerInstanceNames;
+		DataLayers.Empty();
 
+		// Possible there is no World Partition for regular OFPA levels that haven't been converted to support Data Layers
 		if (UWorldPartition* ActorWorldPartition = FWorldPartitionHelpers::GetWorldPartition(InActor))
 		{
 			// Force load Data Layer Assets (SoftObjectPtrs) to make sure desc remains in sync with the actor
 			InActor->ForceLoadDataLayerAssets();
 			const bool bIncludeExternalDataLayerAsset = false;
 			TArray<const UDataLayerAsset*> DataLayerAssets = InActor->GetDataLayerAssets(bIncludeExternalDataLayerAsset);
-			LocalDataLayerAssetPaths.Reserve(DataLayerAssets.Num());
+			DataLayers.Reserve(DataLayerAssets.Num());
+			// Init DataLayers persistent info
 			for (const UDataLayerAsset* DataLayerAsset : DataLayerAssets)
 			{
 				if (DataLayerAsset)
 				{
-					LocalDataLayerAssetPaths.Add(*DataLayerAsset->GetPathName());
-				}
-			}
-			
-			PRAGMA_DISABLE_DEPRECATION_WARNINGS
-			// If the deprected ActorDataLayers is empty, consider the ActorDesc to be is using DataLayerAssets (with an empty array)
-			bIsUsingDataLayerAsset = (LocalDataLayerAssetPaths.Num() > 0) || (InActor->GetActorDataLayers().IsEmpty());
-			if (!bIsUsingDataLayerAsset)
-			{
-				// Use Actor's DataLayerManager since the fixup is relative to this level
-				const UDataLayerManager* DataLayerManager = ActorWorldPartition->GetDataLayerManager();
-				if (ensure(DataLayerManager))
-				{
-					// Pass Actor Level when resolving the DataLayerInstance as FWorldPartitionActorDesc always represents the state of the actor local to its outer level
-					LocalDataLayerInstanceNames = DataLayerManager->GetDataLayerInstanceNames(InActor->GetActorDataLayers());
+					DataLayers.Add(*DataLayerAsset->GetPathName());
 				}
 			}
-			PRAGMA_ENABLE_DEPRECATION_WARNINGS
-
-			// Init DataLayers persistent info
-			DataLayers = bIsUsingDataLayerAsset ? MoveTemp(LocalDataLayerAssetPaths) : MoveTemp(LocalDataLayerInstanceNames);
-		}
-		else
-		{
-			// Possible there is no World Partition for regular OFPA levels that haven't been converted to support Data Layers
-			bIsUsingDataLayerAsset = true;
-			DataLayers.Empty();
 		}
 
 		// Initialize ExternalDataLayerAsset
@@ -419,7 +396,6 @@ bool FWorldPartitionActorDesc::Equals(const FWorldPartitionActorDesc* Other) con
 		bActorIsMainWorldOnly == Other->bActorIsMainWorldOnly &&
 		bActorIsHLODRelevant == Other->bActorIsHLODRelevant &&
 		bActorIsListedInSceneOutliner == Other->bActorIsListedInSceneOutliner &&
-		bIsUsingDataLayerAsset == Other->bIsUsingDataLayerAsset &&
 		HLODLayer == Other->HLODLayer &&
 		FolderPath == Other->FolderPath &&
 		FolderGuid == Other->FolderGuid &&
@@ -525,7 +501,7 @@ TArray<FName> FWorldPartitionActorDesc::GetDataLayers(bool bIncludeExternalDataL
 
 FName FWorldPartitionActorDesc::GetExternalDataLayer() const
 {
-	const bool bHasExternalDataLayerAsset = bIsUsingDataLayerAsset && ExternalDataLayerAsset.IsValid();
+	const bool bHasExternalDataLayerAsset = ExternalDataLayerAsset.IsValid();
 	return bHasExternalDataLayerAsset ? FName(ExternalDataLayerAsset.GetAssetPath().ToString()) : NAME_None;
 }
 
@@ -815,6 +791,7 @@ void FWorldPartitionActorDesc::Serialize(FArchive& Ar)
 
 	if (Ar.CustomVer(FFortniteSeasonBranchObjectVersion::GUID) >= FFortniteSeasonBranchObjectVersion::WorldPartitionActorDescSerializeDataLayerAssets)
 	{
+		// Keep loading this value even though it is not used anymore, in case we ever need it for conversions
 		Ar << TDeltaSerialize<bool>(bIsUsingDataLayerAsset);
 	}
 
@@ -905,7 +882,7 @@ void FWorldPartitionActorDesc::Serialize(FArchive& Ar)
 	}
 
 	// Fixup redirected data layer asset paths
-	if (Ar.IsLoading() && bIsUsingDataLayerAsset)
+	if (Ar.IsLoading())
 	{
 		for (FName& DataLayer : DataLayers)
 		{
diff --git a/Engine/Source/Runtime/Engine/Private/WorldPartition/WorldPartitionSubsystem.cpp b/Engine/Source/Runtime/Engine/Private/WorldPartition/WorldPartitionSubsystem.cpp
--- a/Engine/Source/Runtime/Engine/Private/WorldPartition/WorldPartitionSubsystem.cpp
+++ b/Engine/Source/Runtime/Engine/Private/WorldPartition/WorldPartitionSubsystem.cpp
@@ -49,10 +49,12 @@
 
 #include UE_INLINE_GENERATED_CPP_BY_NAME(WorldPartitionSubsystem)
 
+
 DECLARE_CYCLE_STAT(TEXT("Update Streaming State"), STAT_WorldPartitionUpdateStreamingState, STATGROUP_WorldPartition);
 DECLARE_CYCLE_STAT(TEXT("Update Streaming State (External)"), STAT_WorldPartitionUpdateStreamingStateExternal, STATGROUP_WorldPartition);
 CSV_DEFINE_CATEGORY(WorldPartition, (!UE_BUILD_SHIPPING));
 
+
 extern int32 GBlockOnSlowStreaming;
 static const FName NAME_WorldPartitionRuntimeHash("WorldPartitionRuntimeHash");
 
@@ -390,7 +392,7 @@ TMap<FActorContainerID, TSet<FGuid>> UWorldPartitionSubsystem::GetFilteredActors
 		const TMap<FSoftObjectPath, FWorldPartitionActorFilter::FDataLayerFilter>& DataLayerFilters = DataLayerFiltersPerContainer.FindChecked(InContainerID);
 		for (FActorDescList::TConstIterator<> ActorDescIt(InContainer); ActorDescIt; ++ActorDescIt)
 		{
-			if (ActorDescIt->GetDataLayers().Num() > 0 && ActorDescIt->IsUsingDataLayerAsset())
+			if (ActorDescIt->GetDataLayers().Num() > 0)
 			{
 				bool bExcluded = false;
 				for (FName DataLayerName : ActorDescIt->GetDataLayers())
@@ -1042,6 +1044,7 @@ void UWorldPartitionSubsystem::UpdateStreamingSources()
 	}
 
 	StreamingSourcesHash = 0;
+
 	const double CurrentTime = World->GetTimeSeconds();
 	for (FWorldPartitionStreamingSource& StreamingSource : StreamingSources)
 	{
@@ -1056,6 +1059,8 @@ void UWorldPartitionSubsystem::UpdateStreamingSources()
 		StreamingSource.UpdateHash();
 		// Build hash for all streaming sources
 		StreamingSourcesHash = HashCombine(StreamingSourcesHash, StreamingSource.GetHash());
+
+
 	}
 
 	// Cleanup StreamingSourcesVelocity
diff --git a/Engine/Source/Runtime/Engine/Private/WorldPartitionBlueprintLibrary.cpp b/Engine/Source/Runtime/Engine/Private/WorldPartitionBlueprintLibrary.cpp
--- a/Engine/Source/Runtime/Engine/Private/WorldPartitionBlueprintLibrary.cpp
+++ b/Engine/Source/Runtime/Engine/Private/WorldPartitionBlueprintLibrary.cpp
@@ -44,13 +44,11 @@ FActorDesc::FActorDesc(const FWorldPartitionActorDescInstance& InActorDesc)
 	bIsSpatiallyLoaded = InActorDesc.GetIsSpatiallyLoaded();
 	bActorIsEditorOnly = InActorDesc.GetActorIsEditorOnly();
 	ActorPackage = InActorDesc.GetActorPackage();
-	if (InActorDesc.IsUsingDataLayerAsset())
+
+	DataLayerAssets.Reserve(InActorDesc.GetDataLayers().Num());
+	for (FName DataLayerAssetPath : InActorDesc.GetDataLayers())
 	{
-		DataLayerAssets.Reserve(InActorDesc.GetDataLayers().Num());
-		for (FName DataLayerAssetPath : InActorDesc.GetDataLayers())
-		{
-			DataLayerAssets.Add(FSoftObjectPath(DataLayerAssetPath.ToString()));
-		}
+		DataLayerAssets.Add(FSoftObjectPath(DataLayerAssetPath.ToString()));
 	}
 	
 	ActorPath = *InActorDesc.GetActorSoftPath().ToString();
diff --git a/Engine/Source/Runtime/Engine/Public/WorldPartition/DataLayer/ActorDataLayer.h b/Engine/Source/Runtime/Engine/Public/WorldPartition/DataLayer/ActorDataLayer.h
--- a/Engine/Source/Runtime/Engine/Public/WorldPartition/DataLayer/ActorDataLayer.h
+++ b/Engine/Source/Runtime/Engine/Public/WorldPartition/DataLayer/ActorDataLayer.h
@@ -2,41 +2,8 @@
 
 #pragma once
 
-#include "CoreMinimal.h"
+#include "Misc/CoreMiscDefines.h"
 
-#include "WorldPartition/DataLayer/DataLayerInstance.h"
+UE_DEPRECATED_HEADER(5.8, "The contents of this header have been fully deprecated. See DeprecatedDataLayer.h for details on particular cases.") 
 
-#include "ActorDataLayer.generated.h"
-
-// This class is deprecated and only present for backward compatibility purposes.
-// Instead of using FActorDatalayer, directly save the DataLayerInstance FName if the DataLayer not exposed in data.
-// If the DataLayer is exposed in Data, then use DataLayerAssets.
-USTRUCT(BlueprintType)
-struct FActorDataLayer
-{
-	GENERATED_USTRUCT_BODY()
-
-	static_assert(DATALAYER_TO_INSTANCE_RUNTIME_CONVERSION_ENABLED, "FActorDataLayer is deprecated and needs to be deleted.");
-
-	FActorDataLayer()
-	: Name(NAME_None)
-	{}
-
-	FActorDataLayer(const FName& InName)
-	: Name(InName)
-	{}
-
-	inline bool operator==(const FActorDataLayer& Other) const { return Name == Other.Name; }
-	inline bool operator<(const FActorDataLayer& Other) const { return Name.FastLess(Other.Name); }
-
-	inline operator FName () const { return Name; }
-
-	/** The name of this layer */
-	UPROPERTY(BlueprintReadOnly, VisibleAnywhere, Category = DataLayer)
-	FName Name;
-
-	friend uint32 GetTypeHash(const FActorDataLayer& Value)
-	{
-		return GetTypeHash(Value.Name);
-	}
-};
+#include "WorldPartition/DataLayer/Deprecated/DeprecatedDataLayer.h"
diff --git a/Engine/Source/Runtime/Engine/Public/WorldPartition/DataLayer/DataLayerManager.h b/Engine/Source/Runtime/Engine/Public/WorldPartition/DataLayer/DataLayerManager.h
--- a/Engine/Source/Runtime/Engine/Public/WorldPartition/DataLayer/DataLayerManager.h
+++ b/Engine/Source/Runtime/Engine/Public/WorldPartition/DataLayer/DataLayerManager.h
@@ -5,7 +5,6 @@
 #include "WorldPartition/WorldPartition.h"
 #include "WorldPartition/WorldPartitionHandle.h"
 #include "WorldPartition/WorldPartitionHelpers.h"
-#include "WorldPartition/DataLayer/ActorDataLayer.h"
 #include "WorldPartition/DataLayer/DataLayerInstance.h"
 #include "WorldPartition/DataLayer/DataLayerEditorContext.h"
 #include "WorldPartition/DataLayer/WorldDataLayers.h"
@@ -114,8 +113,14 @@ public:
 	template<class T>
 	TArray<const UDataLayerInstance*> GetDataLayerInstances(const TArray<T>& InDataLayerIdentifiers) const;
 
+#if WITH_EDITOR
 	template<class T>
 	TArray<FName> GetDataLayerInstanceNames(const TArray<T>& InDataLayerIdentifiers) const;
+#else
+	template<class T>
+	UE_DEPRECATED(5.8, "This method is Tool-only since the removal of FActorDataLayer.")
+	TArray<FName> GetDataLayerInstanceNames(const TArray<T>& InDataLayerIdentifiers) const { return {}; }
+#endif // WITH_EDITOR
 
 	ENGINE_API void ForEachDataLayerInstance(TFunctionRef<bool(UDataLayerInstance*)> Func);
 	ENGINE_API void ForEachDataLayerInstance(TFunctionRef<bool(UDataLayerInstance*)> Func) const;
@@ -263,7 +268,9 @@ TArray<const UDataLayerInstance*> UDataLayerManager::GetDataLayerInstances(const
 	AWorldDataLayers* WorldDataLayers = GetWorldDataLayers();
 	return WorldDataLayers ? WorldDataLayers->GetDataLayerInstances(InDataLayerIdentifiers) : EmptyArray;
 }
+PRAGMA_ENABLE_DEPRECATION_WARNINGS
 
+#if WITH_EDITOR
 template<class T>
 TArray<FName> UDataLayerManager::GetDataLayerInstanceNames(const TArray<T>& InDataLayerIdentifiers) const
 {
@@ -271,4 +278,4 @@ TArray<FName> UDataLayerManager::GetDataLayerInstanceNames(const TArray<T>& InDa
 	AWorldDataLayers* WorldDataLayers = GetWorldDataLayers();
 	return WorldDataLayers ? WorldDataLayers->GetDataLayerInstanceNames(InDataLayerIdentifiers) : EmptyArray;
 }
-PRAGMA_ENABLE_DEPRECATION_WARNINGS
+#endif // WITH_EDITOR
diff --git a/Engine/Source/Runtime/Engine/Public/WorldPartition/DataLayer/DataLayerSubsystem.h b/Engine/Source/Runtime/Engine/Public/WorldPartition/DataLayer/DataLayerSubsystem.h
--- a/Engine/Source/Runtime/Engine/Public/WorldPartition/DataLayer/DataLayerSubsystem.h
+++ b/Engine/Source/Runtime/Engine/Public/WorldPartition/DataLayer/DataLayerSubsystem.h
@@ -3,7 +3,6 @@
 #pragma once
 
 #include "Subsystems/WorldSubsystem.h"
-#include "WorldPartition/DataLayer/ActorDataLayer.h"
 #include "WorldPartition/DataLayer/DataLayerInstance.h"
 #include "WorldPartition/DataLayer/DataLayerManager.h"
 #include "DataLayerSubsystem.generated.h"
@@ -205,9 +204,9 @@ public:
 	UFUNCTION(BlueprintCallable, Category = DataLayers, meta = (DeprecatedFunction, DeprecationMessage = "DataLayerSubsystem is deprecated, use GetDataLayerInstanceFromAsset in DataLayerManager"))
 	ENGINE_API UDataLayerInstance* GetDataLayerFromName(FName InDataLayerName) const;
 
-	UE_DEPRECATED(5.1, "Use GetDataLayerInstanceFromAsset instead.")
-	UFUNCTION(BlueprintCallable, Category = DataLayers, meta = (DeprecatedFunction, DeprecationMessage = "DataLayerSubsystem is deprecated, use GetDataLayerInstanceFromAsset in DataLayerManager"))
-	ENGINE_API UDataLayerInstance* GetDataLayerFromLabel(FName InDataLayerLabel) const;
+	UE_DEPRECATED(all, "Use GetDataLayerInstanceFromAsset instead.")
+	UFUNCTION(meta = (DeprecatedFunction, DeprecationMessage = "DataLayerSubsystem is deprecated, use GetDataLayerInstanceFromAsset in DataLayerManager"))
+	UDataLayerInstance* GetDataLayerFromLabel(FName InDataLayerLabel) const { return nullptr; }
 
 	UE_DEPRECATED(5.1, "Use GetDataLayerInstanceRuntimeState instead.")
 	UFUNCTION(BlueprintCallable, Category = DataLayers, meta = (DeprecatedFunction, DeprecationMessage = "DataLayerSubsystem is deprecated, use GetDataLayerInstanceRuntimeState in DataLayerManager"))
diff --git a/Engine/Source/Runtime/Engine/Public/WorldPartition/DataLayer/Deprecated/DeprecatedDataLayer.h b/Engine/Source/Runtime/Engine/Public/WorldPartition/DataLayer/Deprecated/DeprecatedDataLayer.h
--- a/Engine/Source/Runtime/Engine/Public/WorldPartition/DataLayer/Deprecated/DeprecatedDataLayer.h
+++ b/Engine/Source/Runtime/Engine/Public/WorldPartition/DataLayer/Deprecated/DeprecatedDataLayer.h
@@ -11,7 +11,7 @@
 
 #include "DeprecatedDataLayer.generated.h"
 
-// These deprecated reflected types and members must not be removed for the duration of UE5, even after 2 versions post-deprecation, to allow loading and upgrading old assets.
+// These deprecated reflected types and members must not be removed for the duration of UE5, even after 2 versions post-deprecation, to prevent data losses and allow loading and upgrading old assets.
 
 UENUM(BlueprintType)
 enum class UE_DEPRECATED(all, "Use EDataLayerRuntimeState instead.") EDataLayerState : uint8
@@ -23,6 +23,36 @@ enum class UE_DEPRECATED(all, "Use EDataLayerRuntimeState instead.") EDataLayerS
 
 static_assert(EDataLayerRuntimeState::Unloaded < EDataLayerRuntimeState::Loaded && EDataLayerRuntimeState::Loaded < EDataLayerRuntimeState::Activated, "Streaming Query code is dependent on this being true");
 
+struct UE_DEPRECATED(5.8, "This class is deprecated and only present for backward compatibility purposes") FActorDataLayer;
+
+USTRUCT(BlueprintType)
+struct FActorDataLayer
+{
+	GENERATED_USTRUCT_BODY()
+
+	FActorDataLayer()
+	: Name(NAME_None)
+	{}
+
+	FActorDataLayer(const FName& InName)
+	: Name(InName)
+	{}
+
+	inline bool operator==(const FActorDataLayer& Other) const { return Name == Other.Name; }
+	inline bool operator<(const FActorDataLayer& Other) const { return Name.FastLess(Other.Name); }
+
+	inline operator FName() const { return Name; }
+
+	/** The name of this layer */
+	UPROPERTY(meta = (DeprecatedProperty, DeprecationMessage = "FActorDataLayer is deprecated."))
+	FName Name;
+
+	friend uint32 GetTypeHash(const FActorDataLayer& Value)
+	{
+		return GetTypeHash(Value.Name);
+	}
+};
+
 static_assert(DATALAYER_TO_INSTANCE_RUNTIME_CONVERSION_ENABLED, "UDEPRECATED_DataLayer class is deprecated and needs to be deleted.");
 class UE_DEPRECATED(all, "Use UDataLayerInstance & UDataLayerAsset to create DataLayers") UDEPRECATED_DataLayer;
 
@@ -33,38 +63,40 @@ class UDEPRECATED_DataLayer : public UObject
 
 	friend class UDeprecatedDataLayerInstance;
 
-	UFUNCTION(BlueprintCallable, meta = (DeprecatedFunction, DeprecationMessage = "Use UDataLayerInstance & UDataLayerAsset to create DataLayers"))
+	PRAGMA_DISABLE_DEPRECATION_WARNINGS
+	UFUNCTION(meta = (DeprecatedFunction, DeprecationMessage = "Use UDataLayerInstance & UDataLayerAsset to create DataLayers"))
 	bool Equals(const FActorDataLayer& ActorDataLayer) const { return false; }
+	PRAGMA_ENABLE_DEPRECATION_WARNINGS
 
-	UFUNCTION(BlueprintCallable, meta = (DeprecatedFunction, DeprecationMessage = "Use UDataLayerInstance & UDataLayerAsset to create DataLayers"))
+	UFUNCTION(meta = (DeprecatedFunction, DeprecationMessage = "Use UDataLayerInstance & UDataLayerAsset to create DataLayers"))
 	FName GetDataLayerLabel() const { return FName{}; }
 
-	UFUNCTION(BlueprintCallable, meta = (DeprecatedFunction, DeprecationMessage = "Use UDataLayerInstance & UDataLayerAsset to create DataLayers"))
+	UFUNCTION(meta = (DeprecatedFunction, DeprecationMessage = "Use UDataLayerInstance & UDataLayerAsset to create DataLayers"))
 	bool IsInitiallyVisible() const { return false; }
 
-	UFUNCTION(BlueprintCallable, meta = (DeprecatedFunction, DeprecationMessage = "Use UDataLayerInstance & UDataLayerAsset to create DataLayers"))
+	UFUNCTION(meta = (DeprecatedFunction, DeprecationMessage = "Use UDataLayerInstance & UDataLayerAsset to create DataLayers"))
 	bool IsVisible() const { return false; }
 
-	UFUNCTION(BlueprintCallable, meta = (DeprecatedFunction, DeprecationMessage = "Use UDataLayerInstance & UDataLayerAsset to create DataLayers"))
+	UFUNCTION(meta = (DeprecatedFunction, DeprecationMessage = "Use UDataLayerInstance & UDataLayerAsset to create DataLayers"))
 	bool IsEffectiveVisible() const { return false; }
 
-	UFUNCTION(BlueprintCallable, meta = (DeprecatedFunction, DeprecationMessage = "Use UDataLayerInstance & UDataLayerAsset to create DataLayers"))
+	UFUNCTION(meta = (DeprecatedFunction, DeprecationMessage = "Use UDataLayerInstance & UDataLayerAsset to create DataLayers"))
 	FColor GetDebugColor() const { return FColor{}; }
 
-	UFUNCTION(BlueprintCallable, meta = (DeprecatedFunction, DeprecationMessage = "Use UDataLayerInstance & UDataLayerAsset to create DataLayers"))
+	UFUNCTION(meta = (DeprecatedFunction, DeprecationMessage = "Use UDataLayerInstance & UDataLayerAsset to create DataLayers"))
 	bool IsRuntime() const { return false; }
 	
-	UFUNCTION(BlueprintCallable, meta = (DeprecatedFunction, DeprecationMessage = "Use UDataLayerInstance & UDataLayerAsset to create DataLayers"))
+	UFUNCTION(meta = (DeprecatedFunction, DeprecationMessage = "Use UDataLayerInstance & UDataLayerAsset to create DataLayers"))
 	EDataLayerRuntimeState GetInitialRuntimeState() const { return EDataLayerRuntimeState::Unloaded; }
 
-	UFUNCTION(BlueprintCallable, meta = (DeprecatedFunction, DeprecationMessage = "Use UDataLayerInstance & UDataLayerAsset to create DataLayers"))
+	UFUNCTION(meta = (DeprecatedFunction, DeprecationMessage = "Use UDataLayerInstance & UDataLayerAsset to create DataLayers"))
 	bool IsDynamicallyLoaded() const { return false; }
 
-	UFUNCTION(BlueprintCallable, meta = (DeprecatedFunction, DeprecationMessage = "Use UDataLayerInstance & UDataLayerAsset to create DataLayers"))
+	UFUNCTION(meta = (DeprecatedFunction, DeprecationMessage = "Use UDataLayerInstance & UDataLayerAsset to create DataLayers"))
 	bool IsInitiallyActive() const { return IsRuntime() && GetInitialRuntimeState() == EDataLayerRuntimeState::Activated; }
 
 	PRAGMA_DISABLE_DEPRECATION_WARNINGS
-	UFUNCTION(BlueprintCallable, meta = (DeprecatedFunction, DeprecationMessage = "Use UDataLayerInstance & UDataLayerAsset to create DataLayers"))
+	UFUNCTION(meta = (DeprecatedFunction, DeprecationMessage = "Use UDataLayerInstance & UDataLayerAsset to create DataLayers"))
 	EDataLayerState GetInitialState() const { return EDataLayerState::Unloaded; }
 	PRAGMA_ENABLE_DEPRECATION_WARNINGS
 
diff --git a/Engine/Source/Runtime/Engine/Public/WorldPartition/DataLayer/DeprecatedDataLayerInstance.h b/Engine/Source/Runtime/Engine/Public/WorldPartition/DataLayer/DeprecatedDataLayerInstance.h
--- a/Engine/Source/Runtime/Engine/Public/WorldPartition/DataLayer/DeprecatedDataLayerInstance.h
+++ b/Engine/Source/Runtime/Engine/Public/WorldPartition/DataLayer/DeprecatedDataLayerInstance.h
@@ -2,7 +2,7 @@
 #pragma once
 
 #include "WorldPartition/DataLayer/DataLayerInstance.h"
-#include "WorldPartition/DataLayer/ActorDataLayer.h"
+#include "WorldPartition/DataLayer/Deprecated/DeprecatedDataLayer.h"
 #include "DeprecatedDataLayerInstance.generated.h"
 
 // Class used for Runtime Conversion of the Deprecated UDataLayer Class to UDataLayerInstance + UDataLayerAsset.
@@ -25,9 +25,9 @@ public:
 #if WITH_EDITOR
 	static ENGINE_API FName MakeName();
 	ENGINE_API void OnCreated();
+	PRAGMA_DISABLE_DEPRECATION_WARNINGS
 	FActorDataLayer GetActorDataLayer() const;
 
-	PRAGMA_DISABLE_DEPRECATION_WARNINGS
 	UE_DEPRECATED(5.8, "UDEPRECATED_DataLayer are no longer supported")
 	static FName MakeName(const class UDEPRECATED_DataLayer* DeprecatedDataLayer) { return FName{}; }
 
@@ -41,9 +41,9 @@ public:
 private:
 #if WITH_EDITOR
 	//~ Begin UDataLayerInstance Interface
-	ENGINE_API virtual bool PerformAddActor(AActor* InActor) const override;
-	ENGINE_API virtual bool PerformRemoveActor(AActor* InActor) const override;
-	ENGINE_API virtual bool RelabelDataLayer(FName NewDataLayerLabel) override;
+	virtual bool PerformAddActor(AActor* InActor) const override { return false; }
+	virtual bool PerformRemoveActor(AActor* InActor) const override { return false; }
+	virtual bool RelabelDataLayer(FName NewDataLayerLabel) override { return false; }
 	virtual bool SupportRelabeling() const override { return true; }
 	virtual bool CanEditDataLayerShortName() const override { return true; }
 	virtual void PerformSetDataLayerShortName(const FString& InNewShortName) { Label = *InNewShortName; }
diff --git a/Engine/Source/Runtime/Engine/Public/WorldPartition/DataLayer/WorldDataLayers.h b/Engine/Source/Runtime/Engine/Public/WorldPartition/DataLayer/WorldDataLayers.h
--- a/Engine/Source/Runtime/Engine/Public/WorldPartition/DataLayer/WorldDataLayers.h
+++ b/Engine/Source/Runtime/Engine/Public/WorldPartition/DataLayer/WorldDataLayers.h
@@ -5,7 +5,6 @@
 #include "CoreMinimal.h"
 #include "UObject/ObjectMacros.h"
 #include "GameFramework/Info.h"
-#include "WorldPartition/DataLayer/ActorDataLayer.h"
 #include "WorldPartition/DataLayer/DataLayerInstance.h"
 #include "WorldPartition/DataLayer/DataLayerInstanceProviderInterface.h"
 #include "WorldPartition/DataLayer/ExternalDataLayerAsset.h"
@@ -147,9 +146,6 @@ public:
 	ENGINE_API bool CanReferenceDataLayerAsset(const UDataLayerAsset* InDataLayerAsset, FText* OutFailureReason) const;
 #endif
 
-	UE_DEPRECATED(5.4, "Use GetWorldPartitionWorldDataLayersName() instead.")
-	static FName GetWorldPartionWorldDataLayersName() { return GetWorldPartitionWorldDataLayersName(); }
-
 	static FName GetWorldPartitionWorldDataLayersName() { return FName("WorldDataLayers"); } // reserved for ULevel::WorldDataLayers
 	
 	ENGINE_API void DumpDataLayers(FOutputDevice& OutputDevice) const;
@@ -158,12 +154,6 @@ public:
 	ENGINE_API const UDataLayerInstance* GetDataLayerInstance(const UDataLayerAsset* InDataLayerAsset) const;
 	ENGINE_API const UDataLayerInstance* GetDataLayerInstanceFromAssetName(const FName& InDataLayerAssetPathName) const;
 
-	UE_DEPRECATED(5.4, "Use ForEachDataLayerInstance() instead.")
-	void ForEachDataLayer(TFunctionRef<bool(UDataLayerInstance*)> Func) { return ForEachDataLayerInstance(Func); }
-
-	UE_DEPRECATED(5.4, "Use ForEachDataLayerInstance() instead.")
-	void ForEachDataLayer(TFunctionRef<bool(UDataLayerInstance*)> Func) const { return ForEachDataLayerInstance(Func); }
-
 	ENGINE_API void ForEachDataLayerInstance(TFunctionRef<bool(UDataLayerInstance*)> Func);
 	ENGINE_API void ForEachDataLayerInstance(TFunctionRef<bool(UDataLayerInstance*)> Func) const;
 
@@ -181,25 +171,6 @@ public:
 	ENGINE_API void OnDataLayerRuntimeStateChanged(const UDataLayerInstance* InDataLayer, EDataLayerRuntimeState InState);
 	int32 GetDataLayersStateEpoch() const { return DataLayersStateEpoch; }
 
-	//~ Begin Deprecated
-
-	UE_DEPRECATED(5.1, "Use GetDataLayerInstance with FName or UDataLayerAsset instead")
-	ENGINE_API const UDataLayerInstance* GetDataLayerInstance(const FActorDataLayer& InActorDataLayer) const;
-
-	UE_DEPRECATED(5.1, "Use GetDataLayerInstanceNames with UDataLayerAsset instead")
-	ENGINE_API TArray<FName> GetDataLayerInstanceNames(const TArray<FActorDataLayer>& InActorDataLayers) const;
-
-	UE_DEPRECATED(5.1, "Use GetDataLayerInstances with FName or UDataLayerAsset instead")
-	ENGINE_API TArray<const UDataLayerInstance*> GetDataLayerInstances(const TArray<FActorDataLayer>& InActorDataLayers) const;
-
-	UE_DEPRECATED(5.1, "Convert UDataLayers to UDataLayerAsset and UDataLayerInstance using DataLayerToAssetCommandLet")
-	ENGINE_API const UDataLayerInstance* GetDataLayerFromLabel(const FName& InDataLayerLabel) const;
-
-	UE_DEPRECATED(5.3, "This function has been deprecated.")
-	void GetUserLoadedInEditorStates(TArray<FName>& OutDataLayersLoadedInEditor, TArray<FName>& OutDataLayersNotLoadedInEditor) const {}
-
-	//~ End Deprecated
-
 protected:
 	ENGINE_API void InitializeDataLayerRuntimeStates();
 	ENGINE_API void ResetDataLayerRuntimeStates(bool bFlushNet = true);
@@ -304,13 +275,13 @@ private:
 	UPROPERTY(Transient)
 	TArray<TObjectPtr<UDataLayerInstance>> TransientDataLayerInstances;
 
-	static_assert(DATALAYER_TO_INSTANCE_RUNTIME_CONVERSION_ENABLED, "DeprecatedDataLayerNameToDataLayerInstance Property is deprecated and needs to be deleted.");
+	UE_DEPRECATED(5.8, "DeprecatedDataLayerNameToDataLayerInstance Property is no longer used.")
 	UPROPERTY()
 	TMap<FName, TWeakObjectPtr<UDataLayerInstance>> DeprecatedDataLayerNameToDataLayerInstance;
 
 	PRAGMA_DISABLE_DEPRECATION_WARNINGS
 	UE_DEPRECATED(all, "Convert Data layers to DataLayerInstances and DataLayerAssets using DataLayerToAsset CommandLet and use DataLayerInstances instead.")
-	UPROPERTY(meta = (DeprecatedProperty, DeprecationMessage = "Convert Data layers to DataLayerInstances and DataLayerAssets using DataLayerToAsset CommandLet and use DataLayerInstances instead."))
+	UPROPERTY()
 	TSet<TObjectPtr<class UDEPRECATED_DataLayer>> WorldDataLayers_DEPRECATED;
 	PRAGMA_ENABLE_DEPRECATION_WARNINGS
 
@@ -425,9 +396,7 @@ void AWorldDataLayers::OverwriteDataLayerRuntimeStates(const TArray<IdentifierTy
 			}
 			else
 			{
-			PRAGMA_DISABLE_DEPRECATION_WARNINGS
 				return GetDataLayerInstance(Identifier);
-			PRAGMA_ENABLE_DEPRECATION_WARNINGS
 			}
 		};
 
diff --git a/Engine/Source/Runtime/Engine/Public/WorldPartition/WorldPartitionActorDesc.h b/Engine/Source/Runtime/Engine/Public/WorldPartition/WorldPartitionActorDesc.h
--- a/Engine/Source/Runtime/Engine/Public/WorldPartition/WorldPartitionActorDesc.h
+++ b/Engine/Source/Runtime/Engine/Public/WorldPartition/WorldPartitionActorDesc.h
@@ -348,7 +348,10 @@ public:
 	ENGINE_API FBox GetRuntimeBounds() const;
 
 	inline const FGuid& GetParentActor() const { return ParentActor; }
-	inline bool IsUsingDataLayerAsset() const { return bIsUsingDataLayerAsset; }
+
+	UE_DEPRECATED(5.8, "Actors no longer contain data layers without assets")
+	inline bool IsUsingDataLayerAsset() const { return true; }
+
 	inline void AddProperty(FName PropertyName, FName PropertyValue = NAME_None) { Properties.AddProperty(PropertyName, PropertyValue); }
 	inline bool GetProperty(FName PropertyName, FName* PropertyValue) const { return Properties.GetProperty(PropertyName, PropertyValue); }
 	inline bool HasProperty(FName PropertyName) const { return Properties.HasProperty(PropertyName); }
@@ -599,7 +602,7 @@ protected:
 	bool							bActorIsMainWorldOnly;
 	bool							bActorIsHLODRelevant;
 	bool							bActorIsListedInSceneOutliner;
-	bool							bIsUsingDataLayerAsset; // Used to know if DataLayers array represents DataLayers Asset paths or the FNames of the deprecated version of Data Layers
+	bool							bIsUsingDataLayerAsset; // No longer used, but serialized manually and needed if we ever load old data. Treat as a UPROPERTY in terms of deprecation.
 	FSoftObjectPath					HLODLayer;
 	TArray<FName>					DataLayers;
 	FSoftObjectPath					ExternalDataLayerAsset;
diff --git a/Engine/Source/Runtime/Engine/Public/WorldPartition/WorldPartitionActorDescInstance.h b/Engine/Source/Runtime/Engine/Public/WorldPartition/WorldPartitionActorDescInstance.h
--- a/Engine/Source/Runtime/Engine/Public/WorldPartition/WorldPartitionActorDescInstance.h
+++ b/Engine/Source/Runtime/Engine/Public/WorldPartition/WorldPartitionActorDescInstance.h
@@ -50,7 +50,8 @@ public:
 	ENGINE_API virtual bool IsRuntimeRelevant() const;
 	ENGINE_API virtual bool IsEditorRelevant() const;
 
-	virtual bool IsUsingDataLayerAsset() const { return ActorDesc->IsUsingDataLayerAsset(); }
+	UE_DEPRECATED(5.8, "Actors no longer contain data layers without assets")
+	virtual bool IsUsingDataLayerAsset() const { return true; }
 	virtual TArray<FName> GetDataLayers() const { return ActorDesc->GetDataLayers(); }
 
 	virtual bool GetActorIsHLODRelevant() const { return ActorDesc->GetActorIsHLODRelevant(); }
diff --git a/Engine/Source/Runtime/Engine/Public/WorldPartition/WorldPartitionActorDescInstanceView.h b/Engine/Source/Runtime/Engine/Public/WorldPartition/WorldPartitionActorDescInstanceView.h
--- a/Engine/Source/Runtime/Engine/Public/WorldPartition/WorldPartitionActorDescInstanceView.h
+++ b/Engine/Source/Runtime/Engine/Public/WorldPartition/WorldPartitionActorDescInstanceView.h
@@ -31,7 +31,8 @@ public:
 	virtual bool IsRuntimeRelevant() const override { return ActorDescInstance->IsRuntimeRelevant(); }
 	virtual bool IsEditorRelevant() const override { return ActorDescInstance->IsEditorRelevant(); }
 
-	virtual bool IsUsingDataLayerAsset() const override { return ActorDescInstance->IsUsingDataLayerAsset(); }
+	UE_DEPRECATED(5.8, "Actors no longer contain data layers without assets")
+	virtual bool IsUsingDataLayerAsset() const override { return true; }
 	virtual TArray<FName> GetDataLayers() const override { return ActorDescInstance->GetDataLayers(); }
 
 	virtual bool GetActorIsHLODRelevant() const override { return ActorDescInstance->GetActorIsHLODRelevant(); }
diff --git a/Engine/Source/Runtime/Engine/Public/WorldPartition/WorldPartitionMiniMap.h b/Engine/Source/Runtime/Engine/Public/WorldPartition/WorldPartitionMiniMap.h
--- a/Engine/Source/Runtime/Engine/Public/WorldPartition/WorldPartitionMiniMap.h
+++ b/Engine/Source/Runtime/Engine/Public/WorldPartition/WorldPartitionMiniMap.h
@@ -3,7 +3,6 @@
 
 #include "CoreMinimal.h"
 #include "GameFramework/Info.h"
-#include "WorldPartition/DataLayer/ActorDataLayer.h"
 #include "WorldPartitionMiniMap.generated.h"
 
 class UTexture2D;
diff --git a/Engine/Source/Runtime/MovieSceneTracks/Public/Sections/MovieSceneDataLayerSection.h b/Engine/Source/Runtime/MovieSceneTracks/Public/Sections/MovieSceneDataLayerSection.h
--- a/Engine/Source/Runtime/MovieSceneTracks/Public/Sections/MovieSceneDataLayerSection.h
+++ b/Engine/Source/Runtime/MovieSceneTracks/Public/Sections/MovieSceneDataLayerSection.h
@@ -4,7 +4,6 @@
 
 #include "CoreTypes.h"
 #include "UObject/ObjectMacros.h"
-#include "WorldPartition/DataLayer/ActorDataLayer.h"
 #include "WorldPartition/DataLayer/DataLayerAsset.h"
 #include "EntitySystem/IMovieSceneEntityProvider.h"
 #include "MovieSceneSection.h"
