Prevent all dependent material functions from becoming dirty when using the material parameters panel on a material

#Synced-CL 49919419

diff --git a/Engine/Source/Runtime/Engine/Private/Materials/MaterialExpressions.cpp b/Engine/Source/Runtime/Engine/Private/Materials/MaterialExpressions.cpp
--- a/Engine/Source/Runtime/Engine/Private/Materials/MaterialExpressions.cpp
+++ b/Engine/Source/Runtime/Engine/Private/Materials/MaterialExpressions.cpp
@@ -3812,10 +3812,13 @@ bool UMaterialExpressionRuntimeVirtualTextureSampleParameter::SetParameterValue(
 {
 	if (InParameterName == ParameterName)
 	{
-		VirtualTexture = InValue;
-		if (EnumHasAnyFlags(Flags, EMaterialExpressionSetParameterValueFlags::SendPostEditChangeProperty))
+		if (VirtualTexture != InValue)
 		{
-			SendPostEditChangeProperty(this, GET_MEMBER_NAME_STRING_CHECKED(ThisClass, VirtualTexture));
+			VirtualTexture = InValue;
+			if (EnumHasAnyFlags(Flags, EMaterialExpressionSetParameterValueFlags::SendPostEditChangeProperty))
+			{
+				SendPostEditChangeProperty(this, GET_MEMBER_NAME_STRING_CHECKED(ThisClass, VirtualTexture));
+			}
 		}
 		return true;
 	}
@@ -4117,10 +4120,13 @@ bool UMaterialExpressionTextureSampleParameter::SetParameterValue(FName InParame
 {
 	if (InParameterName == ParameterName)
 	{
-		Texture = InValue;
-		if (EnumHasAnyFlags(Flags, EMaterialExpressionSetParameterValueFlags::SendPostEditChangeProperty))
+		if (Texture != InValue)
 		{
-			SendPostEditChangeProperty(this, GET_MEMBER_NAME_STRING_CHECKED(ThisClass, Texture));
+			Texture = InValue;
+			if (EnumHasAnyFlags(Flags, EMaterialExpressionSetParameterValueFlags::SendPostEditChangeProperty))
+			{
+				SendPostEditChangeProperty(this, GET_MEMBER_NAME_STRING_CHECKED(ThisClass, Texture));
+			}
 		}
 		return true;
 	}
@@ -5557,20 +5563,40 @@ bool  UMaterialExpressionStaticComponentMaskParameter::SetParameterValue(FName I
 {
 	if (InParameterName == ParameterName)
 	{
-		DefaultR = InR;
-		DefaultG = InG;
-		DefaultB = InB;
-		DefaultA = InA;
-		if (!EnumHasAnyFlags(Flags, EMaterialExpressionSetParameterValueFlags::NoUpdateExpressionGuid))
-		{
-			ExpressionGUID = InExpressionGuid;
-		}
-		if (EnumHasAnyFlags(Flags, EMaterialExpressionSetParameterValueFlags::SendPostEditChangeProperty))
+		const bool bHasRChanged = (DefaultR != InR);
+		const bool bHasGChanged = (DefaultG != InG);
+		const bool bHasBChanged = (DefaultB != InB);
+		const bool bHasAChanged = (DefaultA != InA);
+
+		if (bHasRChanged || bHasGChanged || bHasBChanged || bHasAChanged)
 		{
-			SendPostEditChangeProperty(this, GET_MEMBER_NAME_STRING_CHECKED(ThisClass, DefaultR));
-			SendPostEditChangeProperty(this, GET_MEMBER_NAME_STRING_CHECKED(ThisClass, DefaultG));
-			SendPostEditChangeProperty(this, GET_MEMBER_NAME_STRING_CHECKED(ThisClass, DefaultB));
-			SendPostEditChangeProperty(this, GET_MEMBER_NAME_STRING_CHECKED(ThisClass, DefaultA));
+			DefaultR = InR;
+			DefaultG = InG;
+			DefaultB = InB;
+			DefaultA = InA;
+			if (!EnumHasAnyFlags(Flags, EMaterialExpressionSetParameterValueFlags::NoUpdateExpressionGuid))
+			{
+				ExpressionGUID = InExpressionGuid;
+			}
+			if (EnumHasAnyFlags(Flags, EMaterialExpressionSetParameterValueFlags::SendPostEditChangeProperty))
+			{
+				if (bHasRChanged)
+				{
+					SendPostEditChangeProperty(this, GET_MEMBER_NAME_STRING_CHECKED(ThisClass, DefaultR));
+				}
+				if (bHasGChanged)
+				{
+					SendPostEditChangeProperty(this, GET_MEMBER_NAME_STRING_CHECKED(ThisClass, DefaultG));
+				}
+				if (bHasBChanged)
+				{
+					SendPostEditChangeProperty(this, GET_MEMBER_NAME_STRING_CHECKED(ThisClass, DefaultB));
+				}
+				if (bHasAChanged)
+				{
+					SendPostEditChangeProperty(this, GET_MEMBER_NAME_STRING_CHECKED(ThisClass, DefaultA));
+				}
+			}
 		}
 		return true;
 	}
@@ -8185,10 +8211,13 @@ bool UMaterialExpressionVectorParameter::SetParameterValue(FName InParameterName
 {
 	if (InParameterName == ParameterName)
 	{
-		DefaultValue = InValue;
-		if (EnumHasAnyFlags(Flags, EMaterialExpressionSetParameterValueFlags::SendPostEditChangeProperty))
+		if (DefaultValue != InValue)
 		{
-			SendPostEditChangeProperty(this, TEXT("DefaultValue"));
+			DefaultValue = InValue;
+			if (EnumHasAnyFlags(Flags, EMaterialExpressionSetParameterValueFlags::SendPostEditChangeProperty))
+			{
+				SendPostEditChangeProperty(this, GET_MEMBER_NAME_STRING_CHECKED(ThisClass, DefaultValue));
+			}
 		}
 		return true;
 	}
@@ -8315,10 +8344,13 @@ bool UMaterialExpressionDoubleVectorParameter::SetParameterValue(FName InParamet
 {
 	if (InParameterName == ParameterName)
 	{
-		DefaultValue = InValue;
-		if (EnumHasAnyFlags(Flags, EMaterialExpressionSetParameterValueFlags::SendPostEditChangeProperty))
+		if (DefaultValue != InValue)
 		{
-			SendPostEditChangeProperty(this, GET_MEMBER_NAME_STRING_CHECKED(ThisClass, DefaultValue));
+			DefaultValue = InValue;
+			if (EnumHasAnyFlags(Flags, EMaterialExpressionSetParameterValueFlags::SendPostEditChangeProperty))
+			{
+				SendPostEditChangeProperty(this, GET_MEMBER_NAME_STRING_CHECKED(ThisClass, DefaultValue));
+			}
 		}
 		return true;
 	}
@@ -8357,31 +8389,34 @@ bool UMaterialExpressionChannelMaskParameter::SetParameterValue(FName InParamete
 {
 	if (InParameterName == ParameterName)
 	{
-		// Update value
-		DefaultValue = InValue;
-
-		// Update enum
-		if (DefaultValue.R > 0.0f)
+		if (DefaultValue != InValue)
 		{
-			MaskChannel = EChannelMaskParameterColor::Red;
-		}
-		else if (DefaultValue.G > 0.0f)
-		{
-			MaskChannel = EChannelMaskParameterColor::Green;
-		}
-		else if (DefaultValue.B > 0.0f)
-		{
-			MaskChannel = EChannelMaskParameterColor::Blue;
-		}
-		else
-		{
-			MaskChannel = EChannelMaskParameterColor::Alpha;
-		}
+			// Update value
+			DefaultValue = InValue;
 
-		if (EnumHasAnyFlags(Flags, EMaterialExpressionSetParameterValueFlags::SendPostEditChangeProperty))
-		{
-			SendPostEditChangeProperty(this, GET_MEMBER_NAME_STRING_CHECKED(ThisClass, DefaultValue));
-			SendPostEditChangeProperty(this, GET_MEMBER_NAME_STRING_CHECKED(ThisClass, MaskChannel));
+			// Update enum
+			if (DefaultValue.R > 0.0f)
+			{
+				MaskChannel = EChannelMaskParameterColor::Red;
+			}
+			else if (DefaultValue.G > 0.0f)
+			{
+				MaskChannel = EChannelMaskParameterColor::Green;
+			}
+			else if (DefaultValue.B > 0.0f)
+			{
+				MaskChannel = EChannelMaskParameterColor::Blue;
+			}
+			else
+			{
+				MaskChannel = EChannelMaskParameterColor::Alpha;
+			}
+
+			if (EnumHasAnyFlags(Flags, EMaterialExpressionSetParameterValueFlags::SendPostEditChangeProperty))
+			{
+				SendPostEditChangeProperty(this, GET_MEMBER_NAME_STRING_CHECKED(ThisClass, DefaultValue));
+				SendPostEditChangeProperty(this, GET_MEMBER_NAME_STRING_CHECKED(ThisClass, MaskChannel));
+			}
 		}
 
 		return true;
@@ -8599,10 +8634,13 @@ bool UMaterialExpressionScalarParameter::SetParameterValue(FName InParameterName
 {
 	if (InParameterName == ParameterName)
 	{
-		DefaultValue = InValue;
-		if (EnumHasAnyFlags(Flags, EMaterialExpressionSetParameterValueFlags::SendPostEditChangeProperty))
+		if (DefaultValue != InValue)
 		{
-			SendPostEditChangeProperty(this, GET_MEMBER_NAME_STRING_CHECKED(ThisClass, DefaultValue));
+			DefaultValue = InValue;
+			if (EnumHasAnyFlags(Flags, EMaterialExpressionSetParameterValueFlags::SendPostEditChangeProperty))
+			{
+				SendPostEditChangeProperty(this, GET_MEMBER_NAME_STRING_CHECKED(ThisClass, DefaultValue));
+			}
 		}
 		return true;
 	}
@@ -8883,14 +8921,17 @@ bool UMaterialExpressionStaticBoolParameter::SetParameterValue(FName InParameter
 {
 	if (InParameterName == ParameterName)
 	{
-		DefaultValue = InValue;
-		if (!EnumHasAnyFlags(Flags, EMaterialExpressionSetParameterValueFlags::NoUpdateExpressionGuid))
-		{
-			ExpressionGUID = InExpressionGuid;
-		}
-		if (EnumHasAnyFlags(Flags, EMaterialExpressionSetParameterValueFlags::SendPostEditChangeProperty))
+		if (DefaultValue != InValue)
 		{
-			SendPostEditChangeProperty(this, GET_MEMBER_NAME_STRING_CHECKED(ThisClass, DefaultValue));
+			DefaultValue = InValue;
+			if (!EnumHasAnyFlags(Flags, EMaterialExpressionSetParameterValueFlags::NoUpdateExpressionGuid))
+			{
+				ExpressionGUID = InExpressionGuid;
+			}
+			if (EnumHasAnyFlags(Flags, EMaterialExpressionSetParameterValueFlags::SendPostEditChangeProperty))
+			{
+				SendPostEditChangeProperty(this, GET_MEMBER_NAME_STRING_CHECKED(ThisClass, DefaultValue));
+			}
 		}
 		return true;
 	}
@@ -9790,10 +9831,13 @@ bool UMaterialExpressionTextureCollectionParameter::SetParameterValue(const FNam
 {
 	if (InParameterName == ParameterName)
 	{
-		TextureCollection = InValue;
-		if (EnumHasAnyFlags(Flags, EMaterialExpressionSetParameterValueFlags::SendPostEditChangeProperty))
+		if (TextureCollection != InValue)
 		{
-			SendPostEditChangeProperty(this, GET_MEMBER_NAME_STRING_CHECKED(ThisClass, TextureCollection));
+			TextureCollection = InValue;
+			if (EnumHasAnyFlags(Flags, EMaterialExpressionSetParameterValueFlags::SendPostEditChangeProperty))
+			{
+				SendPostEditChangeProperty(this, GET_MEMBER_NAME_STRING_CHECKED(ThisClass, TextureCollection));
+			}
 		}
 		return true;
 	}
@@ -12188,12 +12232,23 @@ bool UMaterialExpressionFontSampleParameter::SetParameterValue(FName InParameter
 {
 	if (InParameterName == ParameterName)
 	{
-		Font = InFontValue;
-		FontTexturePage = InFontPage;
-		if (EnumHasAnyFlags(Flags, EMaterialExpressionSetParameterValueFlags::SendPostEditChangeProperty))
+		const bool bFontValueChanged = (Font != InFontValue);
+		const bool bFontPageChanged = (FontTexturePage != InFontPage);
+		if (bFontValueChanged || bFontPageChanged)
 		{
-			SendPostEditChangeProperty(this, GET_MEMBER_NAME_STRING_CHECKED(ThisClass, Font));
-			SendPostEditChangeProperty(this, GET_MEMBER_NAME_STRING_CHECKED(ThisClass, FontTexturePage));
+			Font = InFontValue;
+			FontTexturePage = InFontPage;
+			if (EnumHasAnyFlags(Flags, EMaterialExpressionSetParameterValueFlags::SendPostEditChangeProperty))
+			{
+				if (bFontValueChanged)
+				{
+					SendPostEditChangeProperty(this, GET_MEMBER_NAME_STRING_CHECKED(ThisClass, Font));
+				}
+				if (bFontPageChanged)
+				{
+					SendPostEditChangeProperty(this, GET_MEMBER_NAME_STRING_CHECKED(ThisClass, FontTexturePage));
+				}
+			}
 		}
 		return true;
 	}
@@ -17276,7 +17331,14 @@ bool UMaterialExpressionCollectionParameter::SetParameterValue(const FName& Name
 	{
 		if (UMaterialParameterCollection* ParameterCollection = Cast<UMaterialParameterCollection>(Meta.Value.ParameterCollection))
 		{
-			Collection = ParameterCollection;
+			if (Collection != ParameterCollection)
+			{
+				Collection = ParameterCollection;
+				if (EnumHasAnyFlags(Flags, EMaterialExpressionSetParameterValueFlags::SendPostEditChangeProperty))
+				{
+					SendPostEditChangeProperty(this, GET_MEMBER_NAME_STRING_CHECKED(ThisClass, Collection));
+				}
+			}
 			return true;
 		}
 	}
@@ -20890,15 +20952,31 @@ bool UMaterialExpressionCurveAtlasRowParameter::SetParameterValue(const FName& N
 		Meta.Value.Type == EMaterialParameterType::Scalar &&
 		Meta.bUsedAsAtlasPosition)
 	{
-		Curve = Meta.ScalarCurve.Get();
-		Atlas = Meta.ScalarAtlas.Get();
-		DefaultValue = Meta.Value.AsScalar();
-		if (EnumHasAnyFlags(Flags, EMaterialExpressionSetParameterValueFlags::SendPostEditChangeProperty))
+		const bool bHasCurveChanged = (Curve != Meta.ScalarCurve.Get());
+		const bool bHasAtlasChanged = (Atlas != Meta.ScalarAtlas.Get());
+		const bool bHasDefaultValueChanged = (DefaultValue != Meta.Value.AsScalar());
+		if (bHasCurveChanged || bHasAtlasChanged || bHasDefaultValueChanged)
 		{
-			SendPostEditChangeProperty(this, GET_MEMBER_NAME_STRING_CHECKED(ThisClass, Curve));
-			SendPostEditChangeProperty(this, GET_MEMBER_NAME_STRING_CHECKED(ThisClass, Atlas));
-			SendPostEditChangeProperty(this, GET_MEMBER_NAME_STRING_CHECKED(ThisClass, DefaultValue));
+			Curve = Meta.ScalarCurve.Get();
+			Atlas = Meta.ScalarAtlas.Get();
+			DefaultValue = Meta.Value.AsScalar();
+			if (EnumHasAnyFlags(Flags, EMaterialExpressionSetParameterValueFlags::SendPostEditChangeProperty))
+			{
+				if (bHasCurveChanged)
+				{
+					SendPostEditChangeProperty(this, GET_MEMBER_NAME_STRING_CHECKED(ThisClass, Curve));
+				}
+				if (bHasAtlasChanged)
+				{
+					SendPostEditChangeProperty(this, GET_MEMBER_NAME_STRING_CHECKED(ThisClass, Atlas));
+				}
+				if (bHasDefaultValueChanged)
+				{
+					SendPostEditChangeProperty(this, GET_MEMBER_NAME_STRING_CHECKED(ThisClass, DefaultValue));
+				}
+			}
 		}
+
 		if (EnumHasAnyFlags(Flags, EMaterialExpressionSetParameterValueFlags::AssignGroupAndSortPriority))
 		{
 			Group = Meta.Group;
@@ -22199,10 +22277,13 @@ bool UMaterialExpressionSparseVolumeTextureSampleParameter::SetParameterValue(FN
 {
 	if (InParameterName == ParameterName)
 	{
-		SparseVolumeTexture = InValue;
-		if (EnumHasAnyFlags(Flags, EMaterialExpressionSetParameterValueFlags::SendPostEditChangeProperty))
+		if (SparseVolumeTexture != InValue)
 		{
-			SendPostEditChangeProperty(this, GET_MEMBER_NAME_STRING_CHECKED(ThisClass, SparseVolumeTexture));
+			SparseVolumeTexture = InValue;
+			if (EnumHasAnyFlags(Flags, EMaterialExpressionSetParameterValueFlags::SendPostEditChangeProperty))
+			{
+				SendPostEditChangeProperty(this, GET_MEMBER_NAME_STRING_CHECKED(ThisClass, SparseVolumeTexture));
+			}
 		}
 		return true;
 	}
