[SkinnedAssetCompiler] Add required memory estimate to the deferred task launch of skinned assets with dependecies (currently only used by [REDACTED_PERSON]).

#Synced-CL 49893724

diff --git a/Engine/Source/Runtime/Engine/Classes/Engine/SkinnedAsset.h b/Engine/Source/Runtime/Engine/Classes/Engine/SkinnedAsset.h
--- a/Engine/Source/Runtime/Engine/Classes/Engine/SkinnedAsset.h
+++ b/Engine/Source/Runtime/Engine/Classes/Engine/SkinnedAsset.h
@@ -285,6 +285,7 @@ public:
 	virtual int32 GetMaxNumOptionalLODs(const ITargetPlatform* TargetPlatform) const
 	PURE_VIRTUAL(USkinnedAsset::GetMaxNumOptionalLODs, return 0;);
 
+	virtual ENGINE_API int64 GetRequiredMemoryForCompile() const;
 	virtual bool HasCompileDependencies() const { return false; }
 	virtual bool HasAnyDependenciesCompiling() const { return false; }
 	virtual TArray<USkinnedAsset*> GetSkinnedAssetDependencies() const { return {}; }
diff --git a/Engine/Source/Runtime/Engine/Private/SkinnedAsset.cpp b/Engine/Source/Runtime/Engine/Private/SkinnedAsset.cpp
--- a/Engine/Source/Runtime/Engine/Private/SkinnedAsset.cpp
+++ b/Engine/Source/Runtime/Engine/Private/SkinnedAsset.cpp
@@ -141,10 +141,9 @@ void USkinnedAsset::PostLoad()
 		EQueuedWorkPriority BasePriority = FSkinnedAssetCompilingManager::Get().GetBasePriority(this);
 
 		AsyncTask = MakeUnique<FSkinnedAssetAsyncBuildTask>(this, MoveTemp(Context));
-		const int64 RequiredMemory = 1024ull * 1024ull * CVarMemoryForSkinnedAssetCompile.GetValueOnAnyThread();
 		if (!HasAnyDependenciesCompiling())
 		{
-			AsyncTask->StartBackgroundTask(ThreadPool, BasePriority, EQueuedWorkFlags::DoNotRunInsideBusyWait, RequiredMemory, TEXT("SkinnedAsset"));
+			AsyncTask->StartBackgroundTask(ThreadPool, BasePriority, EQueuedWorkFlags::DoNotRunInsideBusyWait, GetRequiredMemoryForCompile(), TEXT("SkinnedAsset"));
 		}
 		
 		if (HasCompileDependencies())
@@ -259,6 +258,12 @@ bool USkinnedAsset::IsCompiling() const
 {
 	return AsyncTask != nullptr || AccessedProperties.load(std::memory_order_relaxed) != 0;
 }
+
+int64 USkinnedAsset::GetRequiredMemoryForCompile() const
+{
+	return 1024ull * 1024ull * CVarMemoryForSkinnedAssetCompile.GetValueOnAnyThread();
+}
+
 #endif // WITH_EDITOR
 
 void USkinnedAsset::UpdateUVChannelData(bool bRebuildAll)
diff --git a/Engine/Source/Runtime/Engine/Private/SkinnedAssetCompiler.cpp b/Engine/Source/Runtime/Engine/Private/SkinnedAssetCompiler.cpp
--- a/Engine/Source/Runtime/Engine/Private/SkinnedAssetCompiler.cpp
+++ b/Engine/Source/Runtime/Engine/Private/SkinnedAssetCompiler.cpp
@@ -314,7 +314,8 @@ void FSkinnedAssetCompilingManager::SchedulePendingCompilations()
 		SkinnedAsset->AsyncTask->StartBackgroundTask(
 			GetThreadPool(),
 			GetBasePriority(SkinnedAsset),
-			EQueuedWorkFlags::DoNotRunInsideBusyWait
+			EQueuedWorkFlags::DoNotRunInsideBusyWait,
+			SkinnedAsset->GetRequiredMemoryForCompile()
 		);
 	}
 
