[StandaloneHLOD] Fix HLODIndex parsing

#Synced-CL 49651271

diff --git a/Engine/Source/Runtime/Engine/Private/WorldPartition/HLOD/StandaloneHLODSubsystem.cpp b/Engine/Source/Runtime/Engine/Private/WorldPartition/HLOD/StandaloneHLODSubsystem.cpp
--- a/Engine/Source/Runtime/Engine/Private/WorldPartition/HLOD/StandaloneHLODSubsystem.cpp
+++ b/Engine/Source/Runtime/Engine/Private/WorldPartition/HLOD/StandaloneHLODSubsystem.cpp
@@ -249,6 +249,8 @@ void UWorldPartitionStandaloneHLODSubsystem::UpdateStandaloneHLODActorsRecursive
 
 void UWorldPartitionStandaloneHLODSubsystem::UpdateStandaloneHLODActors(UWorldPartitionStandaloneHLODSubsystem::FStandaloneHLODActorParams InStandaloneHLODActorParams)
 {
+	const FString HLODSuffix = TEXT("_HLOD");
+
 	TArray<AWorldPartitionStandaloneHLOD*>* HLODActorsPtr = StandaloneHLODActors.Find(InStandaloneHLODActorParams.Guid);
 	if (HLODActorsPtr)
 	{
@@ -256,7 +258,8 @@ void UWorldPartitionStandaloneHLODSubsystem::UpdateStandaloneHLODActors(UWorldPa
 
 		for (AWorldPartitionStandaloneHLOD* HLODActor : HLODActors)
 		{
-			int32 Position = HLODActor->GetActorLabel().Find(TEXT("_HLOD"), ESearchCase::IgnoreCase, ESearchDir::FromEnd);
+			int32 Position = HLODActor->GetActorLabel().Find(HLODSuffix, ESearchCase::IgnoreCase, ESearchDir::FromEnd);
+			check(Position != INDEX_NONE);
 			FString NameSuffix = HLODActor->GetActorLabel().RightChop(Position);
 			FString NewActorLabel = InStandaloneHLODActorParams.ActorLabel + NameSuffix;
 
@@ -287,10 +290,11 @@ void UWorldPartitionStandaloneHLODSubsystem::UpdateStandaloneHLODActors(UWorldPa
 			FString PackageName = FPackageName::FilenameToLongPackageName(Package);
 			const FSoftObjectPath LODWorldPackagePath(FTopLevelAssetPath(FName(*PackageName), FName(*FPackageName::GetLongPackageAssetName(PackageName))));
 
-			int32 Position = PackageName.Find(TEXT("_HLOD"), ESearchCase::IgnoreCase, ESearchDir::FromEnd);
+			int32 Position = PackageName.Find(HLODSuffix, ESearchCase::IgnoreCase, ESearchDir::FromEnd);
+			check(Position != INDEX_NONE); // Guaranteed by LevelPackagePrefix returned by GetStandaloneHLODFolderPathAndPackagePrefix
 			FString NameSuffix = PackageName.RightChop(Position);
 			FString ActorLabel = InStandaloneHLODActorParams.ActorLabel + NameSuffix;
-			int32 HLODIndex = FCString::Atoi(*NameSuffix.RightChop(4));
+			int32 HLODIndex = FCString::Atoi(*NameSuffix.RightChop(HLODSuffix.Len()));
 
 			FActorSpawnParameters SpawnParams;
 			SpawnParams.ObjectFlags = RF_Transient;
