Enable/disable angular drives applies to slerp as well as twist/swing. Fix to make it behave as per comment, plus make it work the same way for named drives as all drives

#Synced-CL 49852634

diff --git a/Engine/Source/Runtime/Engine/Classes/Components/SkeletalMeshComponent.h b/Engine/Source/Runtime/Engine/Classes/Components/SkeletalMeshComponent.h
--- a/Engine/Source/Runtime/Engine/Classes/Components/SkeletalMeshComponent.h
+++ b/Engine/Source/Runtime/Engine/Classes/Components/SkeletalMeshComponent.h
@@ -2296,11 +2296,11 @@ public:
 	UFUNCTION(BlueprintCallable, Category="Physics")
 	ENGINE_API void AccumulateAllBodiesBelowPhysicsBlendWeight(const FName& InBoneName, float AddPhysicsBlendWeight, bool bSkipCustomPhysicsType = false );
 
-	/** Enable or Disable AngularPositionDrive. If motor is in SLERP mode it will be turned on if either EnableSwingDrive OR EnableTwistDrive are enabled. In Twist and Swing mode the twist and the swing can be controlled individually.*/
+	/** Enable or Disable AngularPositionDrive. If motor is in SLERP mode it will be turned on if either bEnableSwingDrive OR bEnableTwistDrive are enabled. In Twist and Swing mode the twist and the swing can be controlled individually.*/
 	UFUNCTION(BlueprintCallable, Category = "Physics")
 	ENGINE_API void SetAllMotorsAngularPositionDrive(bool bEnableSwingDrive, bool bEnableTwistDrive, bool bSkipCustomPhysicsType = false);
 
-	/** Enable or Disable AngularVelocityDrive. If motor is in SLERP mode it will be turned on if either EnableSwingDrive OR EnableTwistDrive are enabled. In Twist and Swing mode the twist and the swing can be controlled individually.*/
+	/** Enable or Disable AngularVelocityDrive. If motor is in SLERP mode it will be turned on if either bEnableSwingDrive OR bEnableTwistDrive are enabled. In Twist and Swing mode the twist and the swing can be controlled individually.*/
 	UFUNCTION(BlueprintCallable, Category = "Physics")
 	ENGINE_API void SetAllMotorsAngularVelocityDrive(bool bEnableSwingDrive, bool bEnableTwistDrive, bool bSkipCustomPhysicsType = false);
 
@@ -2323,10 +2323,10 @@ public:
 	 */
 	ENGINE_API bool GetConstraintProfilePropertiesOrDefault(FConstraintProfileProperties& OutProperties, FName JointName, FName ProfileName);
 
-	/** Enable or Disable AngularPositionDrive based on a list of bone names */
+	/** Enable or Disable AngularPositionDrive based on a list of bone names. If motor is in SLERP mode it will be turned on/off if either bEnableSwingDrive OR bEnableTwistDrive are enabled.   */
 	ENGINE_API void SetNamedMotorsAngularPositionDrive(bool bEnableSwingDrive, bool bEnableTwistDrive, const TArray<FName>& BoneNames, bool bSetOtherBodiesToComplement = false);
 
-	/** Enable or Disable AngularVelocityDrive based on a list of bone names */
+	/** Enable or Disable AngularVelocityDrive based on a list of bone names. If motor is in SLERP mode it will be turned on/off if either bEnableSwingDrive OR bEnableTwistDrive are enabled.   */
 	ENGINE_API void SetNamedMotorsAngularVelocityDrive(bool bEnableSwingDrive, bool bEnableTwistDrive, const TArray<FName>& BoneNames, bool bSetOtherBodiesToComplement = false);
 
 	ENGINE_API void GetWeldedBodies(TArray<FBodyInstance*> & OutWeldedBodies, TArray<FName> & OutChildrenLabels, bool bIncludingAutoWeld = false) override;
diff --git a/Engine/Source/Runtime/Engine/Private/SkeletalMeshComponentPhysics.cpp b/Engine/Source/Runtime/Engine/Private/SkeletalMeshComponentPhysics.cpp
--- a/Engine/Source/Runtime/Engine/Private/SkeletalMeshComponentPhysics.cpp
+++ b/Engine/Source/Runtime/Engine/Private/SkeletalMeshComponentPhysics.cpp
@@ -1536,6 +1536,7 @@ void USkeletalMeshComponent::SetAllMotorsAngularPositionDrive(bool bEnableSwingD
 		}
 
 		Constraints[i]->SetOrientationDriveTwistAndSwing(bEnableTwistDrive, bEnableSwingDrive);
+		Constraints[i]->SetOrientationDriveSLERP(bEnableTwistDrive || bEnableSwingDrive);
 	}
 }
 
@@ -1553,10 +1554,12 @@ void USkeletalMeshComponent::SetNamedMotorsAngularPositionDrive(bool bEnableSwin
 		if( BoneNames.Contains(Instance->GetChildBoneName()) )
 		{
 			Constraints[i]->SetOrientationDriveTwistAndSwing(bEnableTwistDrive, bEnableSwingDrive);
+			Constraints[i]->SetOrientationDriveSLERP(bEnableTwistDrive || bEnableSwingDrive);
 		}
 		else if( bSetOtherBodiesToComplement )
 		{
 			Constraints[i]->SetOrientationDriveTwistAndSwing(!bEnableTwistDrive, !bEnableSwingDrive);
+			Constraints[i]->SetOrientationDriveSLERP(!(bEnableTwistDrive || bEnableSwingDrive));
 		}
 	}
 }
@@ -1575,10 +1578,12 @@ void USkeletalMeshComponent::SetNamedMotorsAngularVelocityDrive(bool bEnableSwin
 		if( BoneNames.Contains(Instance->GetChildBoneName()) )
 		{
 			Constraints[i]->SetAngularVelocityDriveTwistAndSwing(bEnableTwistDrive, bEnableSwingDrive);
+			Constraints[i]->SetAngularVelocityDriveSLERP(bEnableTwistDrive || bEnableSwingDrive);
 		}
 		else if( bSetOtherBodiesToComplement )
 		{
 			Constraints[i]->SetAngularVelocityDriveTwistAndSwing(!bEnableTwistDrive, !bEnableSwingDrive);
+			Constraints[i]->SetAngularVelocityDriveSLERP(!(bEnableTwistDrive || bEnableSwingDrive));
 		}
 	}
 }
@@ -1603,6 +1608,7 @@ void USkeletalMeshComponent::SetAllMotorsAngularVelocityDrive(bool bEnableSwingD
 		}
 
 		Constraints[i]->SetAngularVelocityDriveTwistAndSwing(bEnableTwistDrive, bEnableSwingDrive);
+		Constraints[i]->SetAngularVelocityDriveSLERP(bEnableTwistDrive || bEnableSwingDrive);
 	}
 }
 
