Fix for Morph vertex buffer not being cleared after morphs are disabled.

#Synced-CL 49782588

diff --git a/Engine/Source/Runtime/Engine/Private/SkeletalRenderGPUSkin.cpp b/Engine/Source/Runtime/Engine/Private/SkeletalRenderGPUSkin.cpp
--- a/Engine/Source/Runtime/Engine/Private/SkeletalRenderGPUSkin.cpp
+++ b/Engine/Source/Runtime/Engine/Private/SkeletalRenderGPUSkin.cpp
@@ -993,26 +993,41 @@ void FSkeletalMeshObjectGPUSkin::ProcessUpdatedDynamicData(FRHICommandList& RHIC
 		PreviousRevisionNumber = bHasPreviousReferenceToLocal ? DynamicData->PreviousRevisionNumber : INDEX_NONE;
 
 		// RayTracing-specific LOD's can't have a separate morph target buffer at the moment because there is only one morph vertex buffer across the entire mesh.
-		if (GEnableMorphTargets && LODData.GetNumVertices() > 0 && (bHasWeightedActiveMorphs || bHasExternalMorphs))
+		if (GEnableMorphTargets && LODData.GetNumVertices() > 0)
 		{
-			bMorphNeedsUpdate |= GForceUpdateMorphTargets != 0;
-
-			if (bMorphNeedsUpdate)
+			if (bHasWeightedActiveMorphs || bHasExternalMorphs)
 			{
-				LOD.MorphVertexBufferPool->SetCurrentRevisionNumber(CurrentRevisionNumber);
-			}
+				bMorphNeedsUpdate |= GForceUpdateMorphTargets != 0;
 
-			MorphVertexBuffer = &LOD.MorphVertexBufferPool->GetMorphVertexBufferForWriting();
+				if (bMorphNeedsUpdate)
+				{
+					LOD.MorphVertexBufferPool->SetCurrentRevisionNumber(CurrentRevisionNumber);
+				}
 
-			// Force an update if this is the first use of the buffer
-			if (!MorphVertexBuffer->bHasBeenUpdated)
+				MorphVertexBuffer = &LOD.MorphVertexBufferPool->GetMorphVertexBufferForWriting();
+
+				// Force an update if this is the first use of the buffer
+				if (!MorphVertexBuffer->bHasBeenUpdated)
+				{
+					bMorphNeedsUpdate = true;
+				}
+
+				if (bMorphNeedsUpdate)
+				{
+					LOD.MorphVertexBufferPool->SetUpdatedFrameNumber(BoneTransformFrameNumber);
+				}
+			}
+			// The update is still set to true when transitioning from morphs to no-morphs. We need to clear out the previous buffer contents.
+			else if (bMorphNeedsUpdate && LOD.MorphVertexBufferPool->IsInitialized())
 			{
-				bMorphNeedsUpdate = true;
+				MorphVertexBuffer = &LOD.MorphVertexBufferPool->GetMorphVertexBufferForWriting();
+				MorphVertexBuffer->bNeedsInitialClear = true;
+				bMorphNeedsUpdate = false;
 			}
-
-			if (bMorphNeedsUpdate)
+			else
 			{
-				LOD.MorphVertexBufferPool->SetUpdatedFrameNumber(BoneTransformFrameNumber);
+				MorphVertexBuffer = nullptr;
+				bMorphNeedsUpdate = false;
 			}
 		}
 		else
