[StadaloneHLOD] Support multiple instances of the same level used as SubWP

#Synced-CL 49955272

diff --git a/Engine/Source/Runtime/Engine/Private/WorldPartition/HLOD/HLODActor.cpp b/Engine/Source/Runtime/Engine/Private/WorldPartition/HLOD/HLODActor.cpp
--- a/Engine/Source/Runtime/Engine/Private/WorldPartition/HLOD/HLODActor.cpp
+++ b/Engine/Source/Runtime/Engine/Private/WorldPartition/HLOD/HLODActor.cpp
@@ -111,6 +111,12 @@ bool AWorldPartitionHLOD::IsStandalone() const
 
 const FGuid& AWorldPartitionHLOD::GetStandaloneHLODGuid() const
 {
+#if WITH_EDITOR
+	if (IsStandalone())
+	{
+		return GetActorInstanceGuid();
+	}
+#endif
 	return StandaloneHLODGuid;
 }
 
@@ -358,6 +364,11 @@ void AWorldPartitionHLOD::PreSave(FObjectPreSaveContext ObjectSaveContext)
 			}
 		}
 	});
+
+	if (ObjectSaveContext.IsCooking() && IsStandalone())
+	{
+		StandaloneHLODGuid = GetActorInstanceGuid();
+	}
 }
 #endif
 
diff --git a/Engine/Source/Runtime/Engine/Private/WorldPartition/HLOD/StandaloneHLODSubsystem.cpp b/Engine/Source/Runtime/Engine/Private/WorldPartition/HLOD/StandaloneHLODSubsystem.cpp
--- a/Engine/Source/Runtime/Engine/Private/WorldPartition/HLOD/StandaloneHLODSubsystem.cpp
+++ b/Engine/Source/Runtime/Engine/Private/WorldPartition/HLOD/StandaloneHLODSubsystem.cpp
@@ -275,6 +275,8 @@ void UWorldPartitionStandaloneHLODSubsystem::UpdateStandaloneHLODActors(UWorldPa
 		TArray<FString> Packages;
 		FPackageName::FindPackagesInDirectory(Packages, FolderPath);
 
+		TMap<int32, FActorContainerID> HLODIndexToContainer;
+
 		for (const FString& Package : Packages)
 		{
 			if (!Package.Contains(LevelPackagePrefix))
@@ -343,6 +345,21 @@ void UWorldPartitionStandaloneHLODSubsystem::UpdateStandaloneHLODActors(UWorldPa
 			HLODActor->RegisterAllComponents();
 
 			StandaloneHLODActors.FindOrAdd(InStandaloneHLODActorParams.Guid).Add(HLODActor);
+
+			HLODIndexToContainer.Add(HLODIndex, HLODActor->GetLevelInstanceID().GetContainerID());
+		}
+
+		// Create a mapping from HLOD container to container with child HLOD actors
+		// e.g. container with HLOD1 actors maps to the container with HLOD0 actors
+		HLODIndexToContainer.KeySort(TLess<int32>());
+		const FActorContainerID* ChildContainer = nullptr;
+		for (const TPair<int32, FActorContainerID>& Entry : HLODIndexToContainer)
+		{
+			if (ChildContainer)
+			{
+				HLODContainerToChildHLODActorsContainer.Add(Entry.Value, *ChildContainer);
+			}
+			ChildContainer = &Entry.Value;
 		}
 	}
 }
@@ -411,4 +428,14 @@ bool UWorldPartitionStandaloneHLODSubsystem::GetStandaloneHLODFolderPathAndPacka
 
 	return true;
 }
+
+bool UWorldPartitionStandaloneHLODSubsystem::GetChildHLODActorsContainer(const FActorContainerID& InContainerID, FActorContainerID& OutChildContainerID) const
+{
+	if (const FActorContainerID* ChildContainerID = HLODContainerToChildHLODActorsContainer.Find(InContainerID))
+	{
+		OutChildContainerID = *ChildContainerID;
+		return true;
+	}
+	return false;
+}
 #endif
\ No newline at end of file
diff --git a/Engine/Source/Runtime/Engine/Private/WorldPartition/RuntimeHashSet/WorldPartitionRuntimeHashSetStreamingGeneration.cpp b/Engine/Source/Runtime/Engine/Private/WorldPartition/RuntimeHashSet/WorldPartitionRuntimeHashSetStreamingGeneration.cpp
--- a/Engine/Source/Runtime/Engine/Private/WorldPartition/RuntimeHashSet/WorldPartitionRuntimeHashSetStreamingGeneration.cpp
+++ b/Engine/Source/Runtime/Engine/Private/WorldPartition/RuntimeHashSet/WorldPartitionRuntimeHashSetStreamingGeneration.cpp
@@ -13,6 +13,7 @@
 #include "WorldPartition/HLOD/HLODActorDesc.h"
 #include "WorldPartition/HLOD/CustomHLODPlaceholderActor.h"
 #include "WorldPartition/HLOD/CustomHLODPlaceholderActorDesc.h"
+#include "WorldPartition/HLOD/StandaloneHLODSubsystem.h"
 
 #if WITH_EDITOR
 bool UWorldPartitionRuntimeHashSet::GenerateRuntimePartitionsStreamingDescs(const IStreamingGenerationContext* StreamingGenerationContext, TMap<URuntimePartition*, TArray<URuntimePartition::FCellDescInstance>>& OutRuntimeCellDescs) const
@@ -210,13 +211,14 @@ bool UWorldPartitionRuntimeHashSet::GenerateStreaming(UWorldPartitionStreamingPo
 
 				WorldPartition->ApplyRuntimeCellsTransformerStack(RuntimeCell);
 
-				// Save HLOD actor GUID -> Cell GUID mapping to use it for Standalone HLOD source cell override
+				// Save HLOD actor instance GUID -> Cell GUID mapping to use it for Standalone HLOD source cell override
 				for (const IStreamingGenerationContext::FActorInstance& ActorInstance : CellActorInstances)
 				{
 					const FStreamingGenerationActorDescView& ActorDescView = ActorInstance.GetActorDescView();
 					if (ActorDescView.GetActorNativeClass()->IsChildOf<AWorldPartitionHLOD>())
 					{
-						StandaloneHLODActorToCell.Add(ActorDescView.GetGuid(), RuntimeCell->GetGuid());
+						FGuid ActorInstanceGuid = ActorDescView.GetContainerInstance()->GetContainerID().GetActorGuid(ActorDescView.GetGuid());
+						StandaloneHLODActorToCell.Add(ActorInstanceGuid, RuntimeCell->GetGuid());
 					}
 				}
 
@@ -276,21 +278,36 @@ bool UWorldPartitionRuntimeHashSet::GenerateStreaming(UWorldPartitionStreamingPo
 	// Standalone HLOD actors are embedded in the main world streaming cells, so for Standalone HLOD actors that are using other Standalone HLOD actors as source actors,
 	// source cell GUID won't be correct. Save the actual source cell GUID, so that we can use that override at runtime.
 	StandaloneHLODActorToSourceCellsMap.Empty();
-	StreamingGenerationContext->ForEachActorSetContainerInstance([this, &StandaloneHLODActorToCell](const IStreamingGenerationContext::FActorSetContainerInstance& ActorSetContainerInstance)
+	StreamingGenerationContext->ForEachActorSetContainerInstance([this, &StandaloneHLODActorToCell, World](const IStreamingGenerationContext::FActorSetContainerInstance& ActorSetContainerInstance)
 	{
 		const FStreamingGenerationActorDescViewMap* ActorDescViewMap = ActorSetContainerInstance.ActorDescViewMap;
-		ActorDescViewMap->ForEachActorDescView([this, &StandaloneHLODActorToCell](const FStreamingGenerationActorDescView& ActorDescView)
+		ActorDescViewMap->ForEachActorDescView([this, &StandaloneHLODActorToCell, World](const FStreamingGenerationActorDescView& ActorDescView)
 		{
 			if (ActorDescView.GetActorNativeClass()->IsChildOf<AWorldPartitionHLOD>())
 			{
 				const FHLODActorDesc* ActorDesc = static_cast<const FHLODActorDesc*>(ActorDescView.GetActorDesc());
 				for (const FGuid& ExternalChildHLODActorGuid : ActorDesc->GetExternalChildHLODActors())
 				{
-					if (FGuid* CellGuidPtr = StandaloneHLODActorToCell.Find(ExternalChildHLODActorGuid))
+					// Find External Child HLOD Actor Instance Guid, as there can be multiple instances of the same level used as SubWP
+					FGuid ExternalChildHLODActorInstanceGuid = ExternalChildHLODActorGuid;
+					if (World)
+					{
+						if (UWorldPartitionStandaloneHLODSubsystem* StandaloneHLODSubsystem = UWorld::GetSubsystem<UWorldPartitionStandaloneHLODSubsystem>(World))
+						{
+							FActorContainerID ExternalChildHLODActorContainerID;
+							if (StandaloneHLODSubsystem->GetChildHLODActorsContainer(ActorDescView.GetContainerInstance()->GetContainerID(), ExternalChildHLODActorContainerID))
+							{
+								ExternalChildHLODActorInstanceGuid = ExternalChildHLODActorContainerID.GetActorGuid(ExternalChildHLODActorGuid);
+							}
+						}
+					}
+
+					if (FGuid* CellGuidPtr = StandaloneHLODActorToCell.Find(ExternalChildHLODActorInstanceGuid))
 					{
 						FGuid CellGuid = *CellGuidPtr;
+						FGuid ActorInstanceGuid = ActorDescView.GetContainerInstance()->GetContainerID().GetActorGuid(ActorDescView.GetGuid());
 
-						if (FGuid* SourceCellGuidPtr = StandaloneHLODActorToSourceCellsMap.Find(ActorDescView.GetGuid()))
+						if (FGuid* SourceCellGuidPtr = StandaloneHLODActorToSourceCellsMap.Find(ActorInstanceGuid))
 						{
 							if (CellGuid != *SourceCellGuidPtr)
 							{
@@ -299,7 +316,7 @@ bool UWorldPartitionRuntimeHashSet::GenerateStreaming(UWorldPartitionStreamingPo
 						}
 						else
 						{
-							StandaloneHLODActorToSourceCellsMap.Add(ActorDescView.GetGuid(), CellGuid);
+							StandaloneHLODActorToSourceCellsMap.Add(ActorInstanceGuid, CellGuid);
 						}
 					}
 				}
diff --git a/Engine/Source/Runtime/Engine/Public/WorldPartition/HLOD/StandaloneHLODSubsystem.h b/Engine/Source/Runtime/Engine/Public/WorldPartition/HLOD/StandaloneHLODSubsystem.h
--- a/Engine/Source/Runtime/Engine/Public/WorldPartition/HLOD/StandaloneHLODSubsystem.h
+++ b/Engine/Source/Runtime/Engine/Public/WorldPartition/HLOD/StandaloneHLODSubsystem.h
@@ -49,6 +49,8 @@ public:
 
 	static ENGINE_API bool GetStandaloneHLODFolderPathAndPackagePrefix(const FString& InWorldPackageName, FString& OutFolderPath, FString& OutPackagePrefix);
 
+	ENGINE_API bool GetChildHLODActorsContainer(const FActorContainerID& InContainerID, FActorContainerID& OutChildContainerID) const;
+
 private:
 	void OnWorldPartitionInitialized(UWorldPartition* InWorldPartition);
 	void OnWorldPartitionUninitialized(UWorldPartition* InWorldPartition);
@@ -63,6 +65,7 @@ private:
 	void OnLevelAddedToWorld(ULevel* Level, UWorld* World);
 
 	TMap<FGuid, TArray<AWorldPartitionStandaloneHLOD*>> StandaloneHLODActors;
+	TMap<FActorContainerID, FActorContainerID> HLODContainerToChildHLODActorsContainer;
 
 	TMap<FName, TMap<int32, FName>> CachedHLODSetups;
 	bool bRefreshCachedHLODSetups;
