Fixed build error and added validation of RepIndex passed in LifeTimeReplicatedProps

#Synced-CL 49833605

diff --git a/Engine/Source/Runtime/Engine/Private/RepLayout.cpp b/Engine/Source/Runtime/Engine/Private/RepLayout.cpp
--- a/Engine/Source/Runtime/Engine/Private/RepLayout.cpp
+++ b/Engine/Source/Runtime/Engine/Private/RepLayout.cpp
@@ -6270,6 +6270,12 @@ void FRepLayout::InitFromClass(
 	for (int32 i = 0; i < LifetimeProps.Num(); i++)
 	{
 		const int32 ParentIndex = LifetimeProps[i].RepIndex;
+		const bool bIsValidParentIndex = ensureMsgf(InObjectClass->ClassReps.IsValidIndex(ParentIndex), TEXT("LifeTimeProps[%d].RepIndex references non existing classRep %d"), i, ParentIndex);
+		if (!bIsValidParentIndex)
+		{
+			continue;
+		}
+
 		const FRepRecord& ReplicatedProp = InObjectClass->ClassReps[ParentIndex];
 	
 		const bool bIsValid = ensureMsgf(Parents.IsValidIndex(ParentIndex), TEXT("ClassReps RepRecord of %s point to invalid RepIndex=%i for replicated property: %s. Parents.Num() = %d"), *GetFullNameSafe(InObjectClass), ParentIndex, *GetNameSafe(ReplicatedProp.Property), Parents.Num());
