Fixed ISKM on Android / Mobile when PlatformGPUSceneUsesUniformBufferView is true.

#Synced-CL 49959273

diff --git a/Engine/Source/Runtime/Engine/Private/GPUSkinVertexFactory.cpp b/Engine/Source/Runtime/Engine/Private/GPUSkinVertexFactory.cpp
--- a/Engine/Source/Runtime/Engine/Private/GPUSkinVertexFactory.cpp
+++ b/Engine/Source/Runtime/Engine/Private/GPUSkinVertexFactory.cpp
@@ -842,6 +842,11 @@ void TGPUSkinVertexFactory<BoneInfluenceType>::ModifyCompilationEnvironment(cons
 
 	// Whether the material supports morph targets
 	OutEnvironment.SetDefine(TEXT("GPUSKIN_MORPH_BLEND"), Parameters.MaterialParameters.bIsUsedWithMorphTargets || Parameters.MaterialParameters.bIsSpecialEngineMaterial);
+
+	if (UseGPUScene(Parameters.Platform, MaxSupportedFeatureLevel))
+	{
+		OutEnvironment.SetDefine(TEXT("USE_INSTANCE_CULLING"), TEXT("1"));
+	}
 }
 
 /**
diff --git a/Engine/Source/Runtime/Engine/Private/SkeletalMeshSceneProxy.cpp b/Engine/Source/Runtime/Engine/Private/SkeletalMeshSceneProxy.cpp
--- a/Engine/Source/Runtime/Engine/Private/SkeletalMeshSceneProxy.cpp
+++ b/Engine/Source/Runtime/Engine/Private/SkeletalMeshSceneProxy.cpp
@@ -342,6 +342,11 @@ FSkeletalMeshSceneProxy::FSkeletalMeshSceneProxy(const FSkinnedMeshSceneProxyDes
 
 	EnableGPUSceneSupportFlags();
 
+	if (UseGPUScene(GetScene().GetShaderPlatform(), FeatureLevel))
+	{
+		bVFRequiresPrimitiveUniformBuffer = PlatformGPUSceneUsesUniformBufferView(GetScene().GetShaderPlatform());
+	}
+
 	if (IsAllowingApproximateOcclusionQueries())
 	{		
 		bAllowApproximateOcclusion = (bAllowApproximateOcclusion || bRenderStatic);
@@ -750,6 +755,7 @@ void FSkeletalMeshSceneProxy::GetStaticMeshBatch(
 	BatchElement.IndexBuffer = LODData.MultiSizeIndexContainer.GetIndexBuffer();
 	BatchElement.bFetchInstanceCountFromScene = bInstancedSkinnedMesh;
 	BatchElement.NumInstances = GetInstanceDataHeader().NumInstances;
+	BatchElement.bForceInstanceCulling = true;
 
 	if (bUseGpuLodSelection)
 	{
@@ -803,7 +809,7 @@ void FSkeletalMeshSceneProxy::CreateBaseMeshBatch(const FSceneView* View, const
 	BatchElement.NumInstances = GetInstanceDataHeader().NumInstances;
 
 	BatchElement.PrimitiveUniformBuffer = GetUniformBuffer();
-	BatchElement.NumPrimitives = Section.NumTriangles;
+	BatchElement.bForceInstanceCulling = true;
 }
 
 uint8 FSkeletalMeshSceneProxy::GetCurrentFirstLODIdx_Internal() const
@@ -895,7 +901,6 @@ void FSkeletalMeshSceneProxy::GetDynamicElementsSection(const TArray<const FScen
 #endif // !(UE_BUILD_SHIPPING || UE_BUILD_TEST)
 #endif // WITH_EDITORONLY_DATA
 
-			BatchElement.MinVertexIndex = Section.BaseVertexIndex;
 			Mesh.ReverseCulling = IsLocalToWorldDeterminantNegative();
 			Mesh.CastShadow = SectionElementInfo.bEnableShadowCasting;
 			Mesh.bCanApplyViewModeOverrides = true;
@@ -943,15 +948,7 @@ void FSkeletalMeshSceneProxy::GetDynamicElementsSection(const TArray<const FScen
 				if (bAddOverlay)
 				{
 					FMeshBatch& OverlayMeshBatch = Collector.AllocateMesh();
-					OverlayMeshBatch = Mesh;
-					OverlayMeshBatch.bOverlayMaterial = true;
-					OverlayMeshBatch.CastShadow = false;
-					OverlayMeshBatch.bSelectable = false;
-					OverlayMeshBatch.MaterialRenderProxy = ActiveOverlayMaterial->GetRenderProxy();
-					// make sure overlay is always rendered on top of base mesh
-					OverlayMeshBatch.MeshIdInPrimitive += LODData.RenderSections.Num();
-					OverlayMeshBatch.Elements[0].bFetchInstanceCountFromScene = bInstancedSkinnedMesh;
-					OverlayMeshBatch.Elements[0].NumInstances = GetInstanceDataHeader().NumInstances;
+					GetOverlayMeshBatch(LODData, ActiveOverlayMaterial, Mesh, OverlayMeshBatch);
 					Collector.AddMesh(ViewIndex, OverlayMeshBatch);
 				
 					INC_DWORD_STAT_BY(STAT_SkelMeshTriangles, OverlayMeshBatch.GetNumPrimitives());
