Fixed missing NumPrimitives in the skeletal mesh proxy dynamic mesh path.

#Synced-CL 49959273

diff --git a/Engine/Source/Runtime/Engine/Private/SkeletalMeshSceneProxy.cpp b/Engine/Source/Runtime/Engine/Private/SkeletalMeshSceneProxy.cpp
--- a/Engine/Source/Runtime/Engine/Private/SkeletalMeshSceneProxy.cpp
+++ b/Engine/Source/Runtime/Engine/Private/SkeletalMeshSceneProxy.cpp
@@ -807,7 +807,7 @@ void FSkeletalMeshSceneProxy::CreateBaseMeshBatch(const FSceneView* View, const
 	BatchElement.MaxVertexIndex = Section.GetVertexBufferIndex() + Section.GetNumVertices() - 1;
 	BatchElement.bFetchInstanceCountFromScene = bInstancedSkinnedMesh;
 	BatchElement.NumInstances = GetInstanceDataHeader().NumInstances;
-
+	BatchElement.NumPrimitives = Section.NumTriangles;
 	BatchElement.PrimitiveUniformBuffer = GetUniformBuffer();
 	BatchElement.bForceInstanceCulling = true;
 }
