NetDriver
* Added FNetDriverReplicationSystemConfig.MinConnectionsForParallelTick to be able to control how many connections are required before we start running UNetConnections::Tick in parallel.
* Default is set to wait for two connections

#Synced-CL 49948957

diff --git a/Engine/Source/Runtime/Engine/Classes/Engine/NetDriver.h b/Engine/Source/Runtime/Engine/Classes/Engine/NetDriver.h
--- a/Engine/Source/Runtime/Engine/Classes/Engine/NetDriver.h
+++ b/Engine/Source/Runtime/Engine/Classes/Engine/NetDriver.h
@@ -689,12 +689,20 @@ struct FNetDriverReplicationSystemConfig
 	UPROPERTY()
 	uint32 MaxNetObjectGroupCount = 0;
 
+	/** 
+	 * Sets the minimum amount of connected clients before we start executing NetConnection::Tick in parallel.
+	 * This limits the overhead of kicking off parallel tasks when concurrency gains would be limited.
+	 * Default is to wait for 2 connections.
+	 */
+	UPROPERTY()
+	int32 MinConnectionsForParallelTick = 2;
+
 	/** 
 	*   When true, allow subsystems to run parallel workloads, such as the PollAndCopy step running several asynchronous tasks to speed up game thread execution time.
 	*   Only supported when bIsServer = true and bAllowObjectReplication = true
 	*/
 	UPROPERTY()
-	bool bAllowParallelTasks = false;	
+	bool bAllowParallelTasks = false;
 };
 
 //
@@ -1116,10 +1124,10 @@ public:
 	/** Change the NetDriver's NetDriverName. This will also reinit packet simulation settings so that settings can be qualified to a specific driver. */
 	void SetNetDriverName(FName NewNetDriverNamed);
 
-	/** Set the NetDriver's NetDriverDefintion. */
+	/** Set the NetDriver's NetDriverDefinition. */
 	void SetNetDriverDefinition(FName NewNetDriverDefinition);
 
-	/** Get the NetDriver's NetDriverDefintion. */
+	/** Get the NetDriver's NetDriverDefinition. */
 	FName GetNetDriverDefinition() const { return NetDriverDefinition; }
 
 	/** Callback after the engine created the NetDriver and set our name for the first time */
diff --git a/Engine/Source/Runtime/Engine/Private/NetDriver.cpp b/Engine/Source/Runtime/Engine/Private/NetDriver.cpp
--- a/Engine/Source/Runtime/Engine/Private/NetDriver.cpp
+++ b/Engine/Source/Runtime/Engine/Private/NetDriver.cpp
@@ -1257,8 +1257,14 @@ void UNetDriver::TickFlush(float DeltaSeconds)
 		ProcessLocalServerPackets();
 	}
 
-	// Server UNetConnection::Tick can be run in parallel if running Iris and this NetDriver supports it
-	if (IsServer() && bIsUsingParallelNetTick && ReplicationSystem && ReplicationSystemConfigServer.bAllowParallelTasks && UE::Net::Private::bAllowParallelNetTickAtRuntime)
+	// Check if this NetDriver is allowed to execute UNetConnection::Tick via multithread tasks
+	bool bCanRunClientTickInParallel = 
+		IsServer() && bIsUsingParallelNetTick && ReplicationSystem && ReplicationSystemConfigServer.bAllowParallelTasks && UE::Net::Private::bAllowParallelNetTickAtRuntime;
+
+	// Do we have enough connections to run in multithread
+	bCanRunClientTickInParallel = bCanRunClientTickInParallel && (ClientConnections.Num() >= ReplicationSystemConfigServer.MinConnectionsForParallelTick);
+
+	if (bCanRunClientTickInParallel)
 	{
 		QUICK_SCOPE_CYCLE_COUNTER(STAT_NetDriver_TickClientConnections)
 
@@ -1293,6 +1299,8 @@ void UNetDriver::TickFlush(float DeltaSeconds)
 	}
 	else
 	{
+		QUICK_SCOPE_CYCLE_COUNTER(STAT_NetDriver_TickClientConnections)
+
 		for (UNetConnection* Connection : ClientConnections)
 		{
 			Connection->Tick(DeltaSeconds);
