[Insights] Trace service support for secure connections

- Add endpoints to control secure connections.
- Add 'capabilities' array to the settings response object. This allows target to communicate which features it supports.

#Synced-CL 49918891

diff --git a/Engine/Source/Runtime/Engine/Private/TraceService.cpp b/Engine/Source/Runtime/Engine/Private/TraceService.cpp
--- a/Engine/Source/Runtime/Engine/Private/TraceService.cpp
+++ b/Engine/Source/Runtime/Engine/Private/TraceService.cpp
@@ -6,6 +6,9 @@
 #include "Misc/App.h"
 #include "ProfilingDebugging/TraceAuxiliary.h"
 #include "ProfilingDebugging/TraceScreenshot.h"
+#include "Runtime/EngineMessages/Public/TraceControlMessages.h"
+#include "Serialization/ArrayWriter.h"
+#include "Serialization/MemoryArchive.h"
 #include "Trace/Trace.h"
 #include "TraceControlMessages.h"
 
@@ -34,6 +37,9 @@ FTraceService::FTraceService(const TSharedPtr<IMessageBus>& InBus)
 			.Handling<FTraceControlBookmark>(this, &FTraceService::OnBookmark)
 #if UE_SCREENSHOT_TRACE_ENABLED
 			.Handling<FTraceControlScreenshot>(this, &FTraceService::OnScreenshot)
+#endif
+#if UE_TRACE_ALLOW_SECURE_TRACING 
+			.Handling<FTraceControlSecure>(this, &FTraceService::OnSecure)
 #endif
 			.Handling<FTraceControlSetStatNamedEvents>(this, &FTraceService::OnSetStatNamedEvents)
 			.Handling<FTraceControlStatusPing>(this, &FTraceService::OnStatusPing)
@@ -49,19 +55,22 @@ FTraceService::FTraceService(const TSharedPtr<IMessageBus>& InBus)
 		MessageEndpoint->Subscribe<FTraceControlSettingsPing>();
 		MessageEndpoint->Subscribe<FTraceControlDiscoveryPing>();
 		MessageEndpoint->Subscribe<FTraceControlChannelsPing>();
-		MessageEndpoint->Subscribe<FTraceControlStop>();
-		MessageEndpoint->Subscribe<FTraceControlSend>();
-		MessageEndpoint->Subscribe<FTraceControlChannelsSet>();
-		MessageEndpoint->Subscribe<FTraceControlFile>();
-		MessageEndpoint->Subscribe<FTraceControlSnapshotSend>();
-		MessageEndpoint->Subscribe<FTraceControlSnapshotSend>();
-		MessageEndpoint->Subscribe<FTraceControlPause>();
-		MessageEndpoint->Subscribe<FTraceControlResume>();
-		MessageEndpoint->Subscribe<FTraceControlBookmark>();
+		// MessageEndpoint->Subscribe<FTraceControlStop>();
+		// MessageEndpoint->Subscribe<FTraceControlSend>();
+		// MessageEndpoint->Subscribe<FTraceControlChannelsSet>();
+		// MessageEndpoint->Subscribe<FTraceControlFile>();
+		// MessageEndpoint->Subscribe<FTraceControlSnapshotSend>();
+		// MessageEndpoint->Subscribe<FTraceControlSnapshotSend>();
+		// MessageEndpoint->Subscribe<FTraceControlPause>();
+		// MessageEndpoint->Subscribe<FTraceControlResume>();
+		// MessageEndpoint->Subscribe<FTraceControlBookmark>();
 #if UE_SCREENSHOT_TRACE_ENABLED
 		MessageEndpoint->Subscribe<FTraceControlScreenshot>();
 #endif
 		MessageEndpoint->Subscribe<FTraceControlSetStatNamedEvents>();
+#if UE_TRACE_ALLOW_SECURE_TRACING 
+		// MessageEndpoint->Subscribe<FTraceControlSecure>();
+#endif
 	}
 }
 
@@ -262,6 +271,10 @@ void FTraceService::OnSettingsPing(const FTraceControlSettingsPing& Message, con
 	FTraceAuxiliary::EnumerateFixedChannelPresets(AddPreset);
 	FTraceAuxiliary::EnumerateChannelPresetsFromSettings(AddPreset);
 
+#if UE_TRACE_ALLOW_SECURE_TRACING 
+	Response->Capabilities.Add(FTraceCapability{ TEXT("SecureConnection") });
+#endif
+
 	MessageEndpoint->Send(Response, Context->GetSender());
 }
 
@@ -282,3 +295,47 @@ void FTraceService::OnDiscoveryPing(const FTraceControlDiscoveryPing& Message, c
 	}
 }
 
+
+void FTraceService::OnSecure(const FTraceControlSecure& Message, const TSharedRef<IMessageContext>& Context)
+{
+#if UE_TRACE_ALLOW_SECURE_TRACING 
+
+	// Generate a random authorization token
+	FGuid AuthToken = FGuid::NewGuid();
+	FArrayWriter AuthData;
+	AuthData << AuthToken;
+	uint32 AuthTokenSize = AuthData.Num();
+	
+	// Setup the security options
+	FString Identity = FGuid::NewGuid().ToString(); // Random identity
+	FTraceAuxiliary::FSecureOptions Options;
+	Options.AuthorizationToken = AuthData;
+	Options.TimeoutSeconds = 30.0f;
+	Options.Identity = TCHAR_TO_ANSI(*Identity);
+
+	const bool bSuccess = FTraceAuxiliary::StartSecure(Options);
+
+	if (!bSuccess)
+	{
+		return;
+	}
+
+	const auto Response = FMessageEndpoint::MakeMessage<FTraceControlSecureResponse>();
+	Response->Host = FTraceAuxiliary::GetSecureHost();
+	Response->Port = FTraceAuxiliary::GetSecurePort();
+	Response->AuthorizationToken = AuthData;
+	Response->CertificateAuthority = FTraceAuxiliary::GetSecureCertificate();
+	if (!Response->CertificateAuthority.IsEmpty())
+	{
+		Response->Mode = ETraceSecureMode::SelfSignedCertAndAuthToken;
+	}
+	else
+	{
+		Response->Mode = ETraceSecureMode::AuthToken;
+		Response->Identity = Identity;
+	}
+
+	MessageEndpoint->Send(Response, Context->GetSender());
+#endif
+}
+
diff --git a/Engine/Source/Runtime/Engine/Public/TraceService.h b/Engine/Source/Runtime/Engine/Public/TraceService.h
--- a/Engine/Source/Runtime/Engine/Public/TraceService.h
+++ b/Engine/Source/Runtime/Engine/Public/TraceService.h
@@ -31,6 +31,7 @@ private:
 	void OnBookmark(const FTraceControlBookmark& Message, const TSharedRef<IMessageContext>& Context);
 	void OnScreenshot(const FTraceControlScreenshot& Message, const TSharedRef<IMessageContext>& Context);
 	void OnSetStatNamedEvents(const FTraceControlSetStatNamedEvents& Message, const TSharedRef<IMessageContext>& Context);
+	void OnSecure(const FTraceControlSecure& Message, const TSharedRef<IMessageContext>& Context);
 
 protected:
 	/**
diff --git a/Engine/Source/Runtime/EngineMessages/Public/TraceControlMessages.h b/Engine/Source/Runtime/EngineMessages/Public/TraceControlMessages.h
--- a/Engine/Source/Runtime/EngineMessages/Public/TraceControlMessages.h
+++ b/Engine/Source/Runtime/EngineMessages/Public/TraceControlMessages.h
@@ -3,6 +3,7 @@
 #pragma once
 
 #include "CoreMinimal.h"
+#include "Misc/Variant.h"
 #include "UObject/ObjectMacros.h"
 #include "Misc/Guid.h"
 #include "Containers/BitArray.h"
@@ -99,6 +100,15 @@ struct FTraceChannelPreset
 	bool bIsReadOnly = false;
 };
 
+USTRUCT()
+struct FTraceCapability
+{
+	GENERATED_USTRUCT_BODY()
+
+	UPROPERTY(EditAnywhere, Category="Message")
+	FString Name;
+};
+
 /**
  * Response for FTraceControlSettingsPing containing
  * TraceLog and system settings.
@@ -119,6 +129,9 @@ struct FTraceControlSettings
 
 	UPROPERTY(EditAnywhere, Category = "Message")
 	TArray<FTraceChannelPreset> ChannelPresets;
+
+	UPROPERTY(EditAnywhere, Category = "Message")
+	TArray<FTraceCapability> Capabilities;
 };
 
 /**
@@ -280,6 +293,59 @@ struct FTraceControlFile : public FTraceControlStartCommon
 	bool bTruncateFile = false;
 };
 
+/**
+ * Start tracing to a host using a secure connection (if supported on the target)
+ */
+USTRUCT()
+struct FTraceControlSecure : public FTraceControlStartCommon
+{
+	GENERATED_USTRUCT_BODY()
+};
+
+
+/**
+ * Secure connection mode.
+ * See TraceAuxiliary for details.
+ */
+UENUM()
+enum class ETraceSecureMode : uint8
+{
+	SelfSignedCertAndAuthToken,
+	CertFileAndAuthToken,
+	AuthToken
+};
+
+/**
+ * Response to a secure start request. Specifies port to connect
+ * to, optional authorization token, and optional certificate if
+ * using self-signed certificates.
+ */
+USTRUCT()
+struct FTraceControlSecureResponse
+{
+	GENERATED_USTRUCT_BODY()
+
+	UPROPERTY(EditAnywhere, Category="Message")
+	FString Host;
+
+	UPROPERTY(EditAnywhere, Category="Message")
+	uint16 Port;
+
+	UPROPERTY(EditAnywhere, Category="Message")
+	ETraceSecureMode Mode;
+
+	UPROPERTY(EditAnywhere, Category="Message")
+	TArray<uint8> AuthorizationToken;
+
+	UPROPERTY(EditAnywhere, Category="Message")
+	TArray<uint8> CertificateAuthority;
+
+	UPROPERTY(EditAnywhere, Category="Message")
+	FString Identity;
+};
+
+
+
 /**
  * Pause all tracing (i.e. mute all channels)
  */
