New material translator:
- Some code reorganization, moving Substrate related code to the end of related files.
- Added missing operators in the IR to HLSL translator

#rnx

#Synced-CL 49989773

#Synced-CL 49989773

diff --git a/Engine/Shaders/Private/MaterialTemplate.ush b/Engine/Shaders/Private/MaterialTemplate.ush
--- a/Engine/Shaders/Private/MaterialTemplate.ush
+++ b/Engine/Shaders/Private/MaterialTemplate.ush
@@ -361,6 +361,9 @@ SHADER_DISABLE_WARNINGS
 /** FMaterialAttributes utilities */
 %{material_attributes_utilities}
 
+/** Custom output declarations */
+%{custom_outputs}
+
 SHADER_POP_WARNINGS_STATE
 
 /** 
diff --git a/Engine/Source/Runtime/Engine/Private/Materials/MaterialExpressionsIR.cpp b/Engine/Source/Runtime/Engine/Private/Materials/MaterialExpressionsIR.cpp
--- a/Engine/Source/Runtime/Engine/Private/Materials/MaterialExpressionsIR.cpp
+++ b/Engine/Source/Runtime/Engine/Private/Materials/MaterialExpressionsIR.cpp
@@ -270,7 +270,6 @@
 #include "Materials/MaterialIRModule.h"
 
 using FValueRef = MIR::FValueRef;
-using namespace Substrate;
 
 /* Helper functions to emit inline HLSL instructions */
 
@@ -3833,7 +3832,7 @@ void UMaterialExpressionCustom::Build(MIR::FEmitter& Em)
 		// Is this argument an unused scene texture sample?
 		if (SceneTextureInfo.IsValidIndex(i) && SceneTextureInfo[i] == -1)
 		{
-			// I this parameter samples an unused scene texture, skip the parameter, but still
+			// If this parameter samples an unused scene texture, skip the parameter, but still
 			// make sure the scene-texture sample is analyzed.
 			InputArgs[FuncDesc.NumInputOnlyParams] = Em.Nop(InputArgs[FuncDesc.NumInputOnlyParams]);
 		}
@@ -4823,16 +4822,24 @@ void UMaterialExpressionNeuralNetworkOutput::Build(MIR::FEmitter& Em)
 	Em.Output(1, Em.InlineHLSL(MIR::FType::MakeFloatVector(4), TEXTVIEW("NeuralBufferOutput(Parameters,$0)"), { BufferIndex }, {}, {}, EnableWithNeuralNetworks));
 }
 
+#endif // WITH_EDITOR
+
+
+
+///////////////////////////////////////////////////////////////////////////////
+//                       Substrate MIR implementation
+///////////////////////////////////////////////////////////////////////////////
+
+#if WITH_EDITOR
 
-////////////////////
-//  Substrate IR  //
-////////////////////
+using namespace Substrate;
 
 struct FSubstrateRegisteredSharedLocalBasisSearchResult
 {
 	bool bFound = false;
 	FSubstrateSharedLocalBasesInfo BasisInfo;
 };
+
 FSubstrateRegisteredSharedLocalBasisSearchResult SubstrateCompilationFindExistingSharedLocalBasis(Substrate::FSubstrateTranslatorData* SubstrateTranslatorData, uint32 NormalCodeChunkHash, uint32 TangentCodeChunkHash)
 {
 	FSubstrateRegisteredSharedLocalBasisSearchResult Result;
@@ -4964,6 +4971,7 @@ static FValueRef CompileWithDefaultFloat1(MIR::FEmitter& Em, FExpressionInput& I
 	ResultArg = Em.CastToFloatKind(ResultArg);	// Move to non LWC
 	return !ResultArg.IsValid() ? DefaultResultArg : ResultArg;
 }
+
 static FValueRef CompileWithDefaultFloat2(MIR::FEmitter& Em, FExpressionInput& Input, float X, float Y, const FVector2MaterialInput* RootNodeInput = nullptr)
 {
 	FValueRef DefaultResultArg = Em.ConstantFloat2(FVector2f(X, Y));
@@ -4975,6 +4983,7 @@ static FValueRef CompileWithDefaultFloat2(MIR::FEmitter& Em, FExpressionInput& I
 	ResultArg = Em.CastToFloatKind(ResultArg);	// Move to non LWC
 	return !ResultArg.IsValid() ? DefaultResultArg : ResultArg;
 }
+
 static FValueRef CompileWithDefaultFloat3(MIR::FEmitter& Em, FExpressionInput& Input, float X, float Y, float Z, const FColorMaterialInput* RootNodeInput = nullptr)
 {
 	FValueRef DefaultResultArg = Em.ConstantFloat3(FVector3f(X, Y, Z));
@@ -4987,7 +4996,6 @@ static FValueRef CompileWithDefaultFloat3(MIR::FEmitter& Em, FExpressionInput& I
 	return !ResultArg.IsValid() ? DefaultResultArg : ResultArg;
 }
 
-
 static FValueRef CompileWithDefaultNormalWS(MIR::FEmitter& Em, UMaterial* Material, FExpressionInput& NormalInput, bool bConvertToRequestedSpace = true)
 {
 	FSubstrateTranslatorData* SubstrateTranslatorData = Em.GetSubstrateTranslatorData();
@@ -5818,7 +5826,6 @@ void UMaterialExpressionSubstrateVerticalLayering::Build(MIR::FEmitter& Em)
 	}
 }
 
-
 void UMaterialExpressionSubstrateTransmittanceToMFP::Build(MIR::FEmitter& Em)
 {
 	FSubstrateTranslatorData* SubstrateTranslatorData = Em.GetSubstrateTranslatorData();
@@ -5846,4 +5853,4 @@ void UMaterialExpressionSubstrateTransmittanceToMFP::Build(MIR::FEmitter& Em)
 	Em.Output(1, ThicknessArg);
 }
 
-#endif // WITH_EDITOR
+#endif // WITH_EDITOR
\ No newline at end of file
diff --git a/Engine/Source/Runtime/Engine/Private/Materials/MaterialIR.cpp b/Engine/Source/Runtime/Engine/Private/Materials/MaterialIR.cpp
--- a/Engine/Source/Runtime/Engine/Private/Materials/MaterialIR.cpp
+++ b/Engine/Source/Runtime/Engine/Private/Materials/MaterialIR.cpp
@@ -93,7 +93,7 @@ bool FValue::HasSubgraphProperties(EGraphProperties Properties) const
 	return (GraphProperties & Properties) == Properties;
 }
 
-void FValue::UseSubgraphProperties(EGraphProperties Properties)
+void FValue::SetSubgraphProperties(EGraphProperties Properties)
 {
 	GraphProperties |= Properties;
 }
@@ -216,16 +216,19 @@ TConstArrayView<FValue*> FValue::GetUses() const
 			auto This = static_cast<const FSubstrateSlab*>(this);
 			return { &This->Normal, FSubstrateSlab::NumStaticUses };
 		}
+
 		case VK_SubstrateShadingModels:
 		{
 			auto This = static_cast<const FSubstrateShadingModels*>(this);
 			return { &This->BaseColor, FSubstrateShadingModels::NumStaticUses };
 		}
+
 		case VK_SubstrateHorizontalMixing:
 		{
 			auto This = static_cast<const FSubstrateHorizontalMixing*>(this);
 			return { &This->Background, FSubstrateHorizontalMixing::NumStaticUses };
 		}
+
 		case VK_SubstrateVerticalLayering:
 		{
 			auto This = static_cast<const FSubstrateVerticalLayering*>(this);
diff --git a/Engine/Source/Runtime/Engine/Private/Materials/MaterialIREmitter.cpp b/Engine/Source/Runtime/Engine/Private/Materials/MaterialIREmitter.cpp
--- a/Engine/Source/Runtime/Engine/Private/Materials/MaterialIREmitter.cpp
+++ b/Engine/Source/Runtime/Engine/Private/Materials/MaterialIREmitter.cpp
@@ -129,7 +129,7 @@
 	// Registers a new value in the module.
 	static void PushValueToModule(FEmitter& Emitter, FValue* Value)
 	{
-		Emitter.Module->Values.Add(Value);
+		Emitter.Module->AddValue(Value);
 		Emitter.ValueSet.Add(Value);
 	}
 };
@@ -1457,7 +1457,7 @@
 	}
 
 	// Declare a new custom output group to include all generated individual outputs
-	Module->AddCustomOutputGroup(Module->InternString(Name), Args.Num());
+	Module->BeginCustomOutputGroup(Module->InternString(Name), Args.Num());
 
 	// Allocate a small buffer to make room for the custom output name plus index
 	MIR::TTemporaryArray<TCHAR> TempNameBuffer{ Name.Len() + 8 };
@@ -4209,6 +4209,97 @@
 	return EmitInlineHLSL(*this, ReturnType, InExternalCodeDeclaration, {}, InputValues, ValueFlags, UsedGraphProperties);
 }
 
+static void SetFunctionFromDesc(FFunction& Function, FFunctionKind Kind, const FFunctionDesc& Desc)
+{
+	Function.Kind = Kind;
+	Function.Name = Desc.Name;
+	Function.NumInputOnlyParams = Desc.NumInputOnlyParams;
+	Function.NumInputAndOutputParams = Desc.NumInputOnlyParams + Desc.NumInputOutputParams;
+	Function.NumParameters = Function.NumInputAndOutputParams + Desc.NumOutputOnlyParams;
+	Function.ReturnType = Desc.ReturnType;
+
+	// Copy over the parameter declarations from the description to the function prototype
+	for (uint32 i = 0; i < Function.NumParameters; ++i)
+	{
+		Function.Parameters[i] = Desc.Parameters[i];
+	} 
+}
+
+const FFunction* FEmitter::FunctionHLSL(const FFunctionHLSLDesc& Desc)
+{
+	check(Desc.GetNumParameters() <= MIR::MaxNumFunctionParameters);
+
+	FFunctionHLSL Prototype = {};
+	SetFunctionFromDesc(Prototype, FFunctionKind::HLSL, Desc);
+	Prototype.Code = Desc.Code;
+	Prototype.Defines = Desc.Defines;
+	Prototype.Includes = Desc.Includes;
+	
+	// Create the new MIR HLSL function instance and set it up
+	FFunctionHLSL* Function = new (FPrivate::Allocate(*this, sizeof(FFunctionHLSL), alignof(FFunctionHLSL))) FFunctionHLSL { Prototype };
+	Function->Name = Module->InternString(Function->Name);
+	Function->Code = Module->InternString(Desc.Code);
+	Function->Defines = MakeArrayCopy(*this, Desc.Defines);
+	Function->Includes = MakeArrayCopy(*this, Desc.Includes);
+
+	// Add it to the module list
+	Module->FunctionHLSLs.Push(Function);
+
+	return Function;
+}
+
+FValueRef FEmitter::Call(const FFunction* Function, TConstArrayView<FValueRef> InputArguments)
+{
+	if (!Function)
+	{
+		return Poison();
+	}
+
+	if (InputArguments.Num() != Function->NumInputAndOutputParams)
+	{
+		Errorf(TEXT("Function called with incorrect number of arguments. Expected %d but got %d."), Function->NumInputAndOutputParams, InputArguments.Num());
+		return Poison();
+	}
+
+	FCall Call = MakePrototype<FCall>(Function->ReturnType);
+	Call.Function = Function;
+	Call.NumArguments = InputArguments.Num();
+
+	for (int i = 0; i < InputArguments.Num(); ++i)
+	{
+		Call.Arguments[i] = InputArguments[i];
+	}
+
+	return EmitPrototype(*this, Call);
+}
+
+FValueRef FEmitter::CallParameterOutput(FValueRef InCall, uint32 ParameterIndex)
+{
+	if (IsAnyNotValid(InCall))
+	{
+		return InCall.ToPoison();
+	}
+
+	FCall* Call = InCall->As<FCall>();
+	if (!Call)
+	{
+		Errorf(TEXT("Expected function call, found a '%s' value instead."), MIR::LexToString(InCall->Kind));
+		return InCall.ToPoison();
+	}
+	
+	if (ParameterIndex >= Call->Function->GetNumOutputParameters())
+	{
+		Errorf(InCall, TEXT("Invalid output index %d. Function has %d outputs."), ParameterIndex, Call->Function->GetNumOutputParameters());
+		return InCall.ToPoison();
+	}
+
+	FCallParameterOutput Proto = MakePrototype<FCallParameterOutput>(Call->Function->GetOutputParameter(ParameterIndex).Type);
+	Proto.Call = Call;
+	Proto.Index = ParameterIndex;
+
+	return EmitPrototype(*this, Proto);
+}
+
 FValueRef FEmitter::SubstrateDefaultSlab()
 {
 	FSubstrateDefaultSlab Prototype = MakePrototype<FSubstrateDefaultSlab>(FType::MakeSubstrateData());
@@ -4391,109 +4482,6 @@
 	return EmitPrototype(*this, Prototype);
 }
 
-static void SetFunctionFromDesc(FFunction& Function, FFunctionKind Kind, const FFunctionDesc& Desc)
-{
-	Function.Kind = Kind;
-	Function.Name = Desc.Name;
-	Function.NumInputOnlyParams = Desc.NumInputOnlyParams;
-	Function.NumInputAndOutputParams = Desc.NumInputOnlyParams + Desc.NumInputOutputParams;
-	Function.NumParameters = Function.NumInputAndOutputParams + Desc.NumOutputOnlyParams;
-	Function.ReturnType = Desc.ReturnType;
-
-	// Copy over the parameter declarations from the description to the function prototype
-	for (uint32 i = 0; i < Function.NumParameters; ++i)
-	{
-		Function.Parameters[i] = Desc.Parameters[i];
-	} 
-}
-
-const FFunction* FEmitter::FunctionHLSL(const FFunctionHLSLDesc& Desc)
-{
-	check(Desc.GetNumParameters() <= MIR::MaxNumFunctionParameters);
-
-	FFunctionHLSL Prototype = {};
-	SetFunctionFromDesc(Prototype, FFunctionKind::HLSL, Desc);
-	Prototype.Code = Desc.Code;
-	Prototype.Defines = Desc.Defines;
-	Prototype.Includes = Desc.Includes;
-	
-	// Create the new MIR HLSL function instance and set it up
-	FFunctionHLSL* Function = new (FPrivate::Allocate(*this, sizeof(FFunctionHLSL), alignof(FFunctionHLSL))) FFunctionHLSL { Prototype };
-	Function->Name = Module->InternString(Function->Name);
-	Function->Code = Module->InternString(Desc.Code);
-	Function->Defines = MakeArrayCopy(*this, Desc.Defines);
-	Function->Includes = MakeArrayCopy(*this, Desc.Includes);
-
-	// Add it to the module list
-	Module->FunctionHLSLs.Push(Function);
-
-	return Function;
-}
-
-FValueRef FEmitter::Call(const FFunction* Function, TConstArrayView<FValueRef> InputArguments)
-{
-	if (!Function)
-	{
-		return Poison();
-	}
-
-	if (InputArguments.Num() != Function->NumInputAndOutputParams)
-	{
-		Errorf(TEXT("Function called with incorrect number of arguments. Expected %d but got %d."), Function->NumInputAndOutputParams, InputArguments.Num());
-		return Poison();
-	}
-
-	FCall Call = MakePrototype<FCall>(Function->ReturnType);
-	Call.Function = Function;
-	Call.NumArguments = InputArguments.Num();
-
-	for (int i = 0; i < InputArguments.Num(); ++i)
-	{
-		Call.Arguments[i] = InputArguments[i];
-	}
-
-	return EmitPrototype(*this, Call);
-}
-
-FValueRef FEmitter::CallParameterOutput(FValueRef InCall, uint32 ParameterIndex)
-{
-	if (IsAnyNotValid(InCall))
-	{
-		return InCall.ToPoison();
-	}
-
-	FCall* Call = InCall->As<FCall>();
-	if (!Call)
-	{
-		Errorf(TEXT("Expected function call, found a '%s' value instead."), MIR::LexToString(InCall->Kind));
-		return InCall.ToPoison();
-	}
-	
-	if (ParameterIndex >= Call->Function->GetNumOutputParameters())
-	{
-		Errorf(InCall, TEXT("Invalid output index %d. Function has %d outputs."), ParameterIndex, Call->Function->GetNumOutputParameters());
-		return InCall.ToPoison();
-	}
-
-	FCallParameterOutput Proto = MakePrototype<FCallParameterOutput>(Call->Function->GetOutputParameter(ParameterIndex).Type);
-	Proto.Call = Call;
-	Proto.Index = ParameterIndex;
-
-	return EmitPrototype(*this, Proto);
-}
-
-void FEmitter::Initialize()
-{
-	// Create and reference the true/false constants.
-	FConstant Temp = MakePrototype<FConstant>(FType::MakeBoolScalar());
-
-	Temp.Boolean = true;
-	TrueConstant = EmitPrototype(*this, Temp);
-
-	Temp.Boolean = false;
-	FalseConstant = EmitPrototype(*this, Temp);
-}
-
 EShaderPlatform FEmitter::GetShaderPlatform() const
 {
 	return Module->GetShaderPlatform();
@@ -4519,6 +4507,18 @@
 	return SubstrateTranslatorData;
 }
 
+void FEmitter::Initialize()
+{
+	// Create and reference the true/false constants.
+	FConstant Temp = MakePrototype<FConstant>(FType::MakeBoolScalar());
+
+	Temp.Boolean = true;
+	TrueConstant = EmitPrototype(*this, Temp);
+
+	Temp.Boolean = false;
+	FalseConstant = EmitPrototype(*this, Temp);
+}
+
 bool FEmitter::FValueKeyFuncs::Matches(KeyInitType A, KeyInitType B)
 {
 	return A->Equals(B);
diff --git a/Engine/Source/Runtime/Engine/Private/Materials/MaterialIRModule.cpp b/Engine/Source/Runtime/Engine/Private/Materials/MaterialIRModule.cpp
--- a/Engine/Source/Runtime/Engine/Private/Materials/MaterialIRModule.cpp
+++ b/Engine/Source/Runtime/Engine/Private/Materials/MaterialIRModule.cpp
@@ -37,6 +37,7 @@ void FMaterialIRModule::Empty()
 	CompilationOutput = {};
 	CustomOutputs.Empty();
 	CustomOutputGroups.Empty();
+	SubstrateData = {};
 	Allocator.Flush();
 
 	// Initialize statistics used external input usage mask.
@@ -61,7 +62,7 @@ int32 FMaterialIRModule::AddEntryPoint(FStringView Name, MIR::EStage Stage, int3
 	return Index;
 }
 
-void FMaterialIRModule::AddCustomOutputGroup(FStringView Name, int32 NumOutputs)
+void FMaterialIRModule::BeginCustomOutputGroup(FStringView Name, int32 NumOutputs)
 {
 	CustomOutputGroups.Push(FCustomOutputGroup{ Name, NumOutputs, CustomOutputs.Num() });
 }
diff --git a/Engine/Source/Runtime/Engine/Private/Materials/MaterialIRToHLSLTranslator.cpp b/Engine/Source/Runtime/Engine/Private/Materials/MaterialIRToHLSLTranslator.cpp
--- a/Engine/Source/Runtime/Engine/Private/Materials/MaterialIRToHLSLTranslator.cpp
+++ b/Engine/Source/Runtime/Engine/Private/Materials/MaterialIRToHLSLTranslator.cpp
@@ -314,7 +314,7 @@
 	}
 }
 
-static FStringView GVector4SwizzleSubset[4] = { TEXTVIEW(".x"), TEXTVIEW(".xy"), TEXTVIEW(".xyz"), {} };
+static const FStringView GVector4SwizzleSubset[4] = { TEXTVIEW(".x"), TEXTVIEW(".xy"), TEXTVIEW(".xyz"), {} };
 
 struct FPrivate : FMaterialIRToHLSLTranslation
 {
@@ -522,6 +522,7 @@
 	{
 		switch (Instr->Kind)
 		{
+			case MIR::VK_Nop: LowerNop(static_cast<const MIR::FNop*>(Instr)); break;
 			case MIR::VK_Composite: LowerComposite(static_cast<const MIR::FComposite*>(Instr)); break;
 			case MIR::VK_SetMaterialOutput: LowerSetMaterialOutput(static_cast<const MIR::FSetMaterialOutput*>(Instr)); break;
 			case MIR::VK_Operator: LowerOperator(static_cast<const MIR::FOperator*>(Instr)); break;
@@ -531,18 +532,17 @@
 			case MIR::VK_TextureRead: LowerTextureRead(static_cast<const MIR::FTextureRead*>(Instr)); break;
 			case MIR::VK_VTPageTableRead: LowerVTPageTableRead(static_cast<const MIR::FVTPageTableRead*>(Instr)); break;
 			case MIR::VK_InlineHLSL: LowerInlineHLSL(static_cast<const MIR::FInlineHLSL*>(Instr)); break;
+			case MIR::VK_StageSwitch: LowerStageSwitch(static_cast<const MIR::FStageSwitch*>(Instr)); break;
+			case MIR::VK_HardwarePartialDerivative: LowerHardwarePartialDerivative(static_cast<const MIR::FHardwarePartialDerivative*>(Instr)); break;
+			case MIR::VK_Call: LowerCall(static_cast<const MIR::FCall*>(Instr)); break;
+			case MIR::VK_CallParameterOutput: LowerCallOutput(static_cast<const MIR::FCallParameterOutput*>(Instr)); break;
+			case MIR::VK_PreshaderParameter: LowerPreshaderParameter(static_cast<const MIR::FPreshaderParameter*>(Instr)); break;
 			case MIR::VK_SubstrateDefaultSlab: LowerSubstrateDefaultSlab(static_cast<const MIR::FSubstrateDefaultSlab*>(Instr)); break;
 			case MIR::VK_SubstrateSlab: LowerSubstrateSlab(static_cast<const MIR::FSubstrateSlab*>(Instr)); break;
 			case MIR::VK_SubstrateShadingModels: LowerSubstrateShadingModel(static_cast<const MIR::FSubstrateShadingModels*>(Instr)); break;
 			case MIR::VK_SubstrateHorizontalMixing: LowerSubstrateHorizontalMixing(static_cast<const MIR::FSubstrateHorizontalMixing*>(Instr)); break;
 			case MIR::VK_SubstrateVerticalLayering: LowerSubstrateVerticalLayering(static_cast<const MIR::FSubstrateVerticalLayering*>(Instr)); break;
 			case MIR::VK_SubstratePromoteToOperator: LowerSubstratePromoteToOperator(static_cast<const MIR::FSubstratePromoteToOperator*>(Instr)); break;
-			case MIR::VK_StageSwitch: LowerStageSwitch(static_cast<const MIR::FStageSwitch*>(Instr)); break;
-			case MIR::VK_HardwarePartialDerivative: LowerHardwarePartialDerivative(static_cast<const MIR::FHardwarePartialDerivative*>(Instr)); break;
-			case MIR::VK_Nop: LowerNop(static_cast<const MIR::FNop*>(Instr)); break;
-			case MIR::VK_Call: LowerCall(static_cast<const MIR::FCall*>(Instr)); break;
-			case MIR::VK_CallParameterOutput: LowerCallOutput(static_cast<const MIR::FCallParameterOutput*>(Instr)); break;
-			case MIR::VK_PreshaderParameter: LowerPreshaderParameter(static_cast<const MIR::FPreshaderParameter*>(Instr)); break;
 			default: UE_MIR_UNREACHABLE();
 		}
 	
@@ -894,157 +894,6 @@
 		}
 	}
 
-	void LowerSubstrateDefaultSlab(const MIR::FSubstrateDefaultSlab* SubstrateParameter)
-	{
-		Printer << TEXTVIEW("GetInitialisedSubstrateData();");
-
-	}
-
-	static FStringView GetSharedLocalBasesTypes(bool bFullySimplifiedInstruction)
-	{
-		return bFullySimplifiedInstruction ? TEXTVIEW("Parameters.SharedLocalBasesFullySimplified.Types") : TEXTVIEW("Parameters.SharedLocalBases.Types");
-	}
-	static FStringView GetSubstrateTree(bool bFullySimplifiedInstruction)
-	{
-		return bFullySimplifiedInstruction ? TEXTVIEW("Parameters.SubstrateTreeFullySimplified") : TEXTVIEW("Parameters.SubstrateTree");
-	}
-
-	void LowerSubstrateSlab(const MIR::FSubstrateSlab* SubstrateParameter)
-	{
-		Printer << TEXTVIEW("GetSubstrateSlabBSDF");
-		Printer << BeginArgs;
-		Printer << TEXTVIEW("Parameters.SubstratePixelFootprint") << TEXTVIEW(",");
-		Printer << LowerValue(SubstrateParameter->Normal) << TEXTVIEW(",");
-		Printer << LowerValue(SubstrateParameter->DiffuseAlbedo) << TEXTVIEW(",");
-		Printer << LowerValue(SubstrateParameter->F0) << TEXTVIEW(",");
-		Printer << LowerValue(SubstrateParameter->F90) << TEXTVIEW(",");
-		Printer << LowerValue(SubstrateParameter->Roughness) << TEXTVIEW(",");
-		Printer << LowerValue(SubstrateParameter->Anisotropy) << TEXTVIEW(",");
-		Printer << LowerValue(SubstrateParameter->SSSProfileId) << TEXTVIEW(",");
-		Printer << TEXTVIEW("false") << TEXTVIEW(","); // bSupportDefaultSSSProfile
-		Printer << LowerValue(SubstrateParameter->SSSMFP) << TEXTVIEW(",");
-		Printer << LowerValue(SubstrateParameter->SSSMFPScale) << TEXTVIEW(",");
-		Printer << LowerValue(SubstrateParameter->SSSPhaseAniso) << TEXTVIEW(",");
-		Printer << LowerValue(SubstrateParameter->SSSType) << TEXTVIEW(",");
-		Printer << LowerValue(SubstrateParameter->EmissiveColor) << TEXTVIEW(",");
-		Printer << LowerValue(SubstrateParameter->SecondRoughness) << TEXTVIEW(",");
-		Printer << LowerValue(SubstrateParameter->SecondRoughnessWeight) << TEXTVIEW(",");
-		Printer << LowerValue(SubstrateParameter->SecondRoughnessAsSimpleClearCoat) << TEXTVIEW(",");
-		Printer << LowerValue(SubstrateParameter->ClearCoatUseSecondNormal) << TEXTVIEW(",");
-		Printer << LowerValue(SubstrateParameter->ClearCoatBottomNormal) << TEXTVIEW(",");
-		Printer << LowerValue(SubstrateParameter->FuzzAmount) << TEXTVIEW(",");
-		Printer << LowerValue(SubstrateParameter->FuzzColor) << TEXTVIEW(",");
-		Printer << LowerValue(SubstrateParameter->FuzzRoughness) << TEXTVIEW(",");
-		Printer << LowerValue(SubstrateParameter->GlintValue) << TEXTVIEW(",");
-		Printer << LowerValue(SubstrateParameter->GlintUV) << TEXTVIEW(",");
-		Printer << LowerValue(SubstrateParameter->SpecularProfileId) << TEXTVIEW(",");
-		Printer << LowerValue(SubstrateParameter->Thickness) << TEXTVIEW(",");
-		Printer << LowerValue(SubstrateParameter->IsThin) << TEXTVIEW(",");
-		Printer << LowerValue(SubstrateParameter->IsAtBottom) << TEXTVIEW(",");
-		Printer << LowerValue(SubstrateParameter->LocalBasisIndex) << TEXTVIEW(",");
-		Printer << GetSharedLocalBasesTypes(SubstrateParameter->bFullySimplifiedInstruction);
-		Printer << EndArgs;
-	}
-
-	void LowerSubstrateShadingModel(const MIR::FSubstrateShadingModels* SubstrateParameter)
-	{
-		Printer << TEXTVIEW("SubstrateConvertLegacyMaterial") << (SubstrateParameter->bHasDynamicShadingModels ? TEXTVIEW("Dynamic") : TEXTVIEW("Static"));
-		Printer << BeginArgs;
-		Printer << TEXTVIEW("Parameters.SubstratePixelFootprint") << TEXTVIEW(",");
-		Printer << LowerValue(SubstrateParameter->BaseColor) << TEXTVIEW(",");
-		Printer << LowerValue(SubstrateParameter->Specular) << TEXTVIEW(",");
-		Printer << LowerValue(SubstrateParameter->Metallic) << TEXTVIEW(",");
-		Printer << LowerValue(SubstrateParameter->Roughness) << TEXTVIEW(",");
-		Printer << LowerValue(SubstrateParameter->Anisotropy) << TEXTVIEW(",");
-		Printer << LowerValue(SubstrateParameter->SubSurfaceColor) << TEXTVIEW(",");
-		Printer << LowerValue(SubstrateParameter->SubSurfaceProfileId) << TEXTVIEW(",");
-		Printer << LowerValue(SubstrateParameter->ClearCoat) << TEXTVIEW(",");
-		Printer << LowerValue(SubstrateParameter->ClearCoatRoughness) << TEXTVIEW(",");
-		Printer << LowerValue(SubstrateParameter->EmissiveColor) << TEXTVIEW(",");
-		Printer << LowerValue(SubstrateParameter->Opacity) << TEXTVIEW(",");
-		Printer << LowerValue(SubstrateParameter->ThinTranslucentTransmittanceColor) << TEXTVIEW(",");
-		Printer << LowerValue(SubstrateParameter->ThinTranslucentSurfaceCoverage) << TEXTVIEW(",");
-		Printer << LowerValue(SubstrateParameter->WaterScatteringCoefficients) << TEXTVIEW(",");
-		Printer << LowerValue(SubstrateParameter->WaterAbsorptionCoefficients) << TEXTVIEW(",");
-		Printer << LowerValue(SubstrateParameter->WaterPhaseG) << TEXTVIEW(",");
-		Printer << LowerValue(SubstrateParameter->ColorScaleBehindWater) << TEXTVIEW(",");
-		Printer << LowerValue(SubstrateParameter->ShadingModel) << TEXTVIEW(",");
-		Printer << LowerValue(SubstrateParameter->Normal) << TEXTVIEW(",");
-		Printer << LowerValue(SubstrateParameter->Tangent) << TEXTVIEW(",");
-		Printer << LowerValue(SubstrateParameter->ClearCoatNormal) << TEXTVIEW(",");
-		Printer << LowerValue(SubstrateParameter->CustomTangent) << TEXTVIEW(",");
-		Printer << LowerValue(SubstrateParameter->BasisIndexMacro) << TEXTVIEW(",");
-		Printer << LowerValue(SubstrateParameter->ClearCoat_BasisIndexMacro) << TEXTVIEW(",");
-		Printer << GetSharedLocalBasesTypes(SubstrateParameter->bFullySimplifiedInstruction);
-		Printer << EndArgs;
-	}
-
-	void LowerSubstrateHorizontalMixing(const MIR::FSubstrateHorizontalMixing* SubstrateParameter)
-	{
-		if (SubstrateParameter->bParameterBlendingEnabled)
-		{
-			Printer << TEXTVIEW("SubstrateHorizontalMixingParameterBlending");
-			Printer << BeginArgs;
-			Printer << LowerValue(SubstrateParameter->Background) << TEXTVIEW(",");
-			Printer << LowerValue(SubstrateParameter->Foreground) << TEXTVIEW(",");
-			Printer << LowerValue(SubstrateParameter->Mix) << TEXTVIEW(",");
-			Printer << LowerValue(SubstrateParameter->NormalMix) << TEXTVIEW(",");
-			Printer << LowerValue(SubstrateParameter->SharedLocalBasisIndexMacro) << TEXTVIEW(",");
-			Printer << GetSharedLocalBasesTypes(SubstrateParameter->bFullySimplifiedInstruction) << TEXTVIEW(",");
-			Printer << LowerValue(SubstrateParameter->BackgroundNoV) << TEXTVIEW(",");
-			Printer << LowerValue(SubstrateParameter->ForegroundNoV);
-			Printer << EndArgs;
-		}
-		else
-		{
-			Printer << GetSubstrateTree(SubstrateParameter->bFullySimplifiedInstruction) << TEXTVIEW(".SubstrateHorizontalMixing");
-			Printer << BeginArgs;
-			Printer << LowerValue(SubstrateParameter->Background) << TEXTVIEW(",");
-			Printer << LowerValue(SubstrateParameter->Foreground) << TEXTVIEW(",");
-			Printer << LowerValue(SubstrateParameter->Mix) << TEXTVIEW(",");
-			Printer << LowerValue(SubstrateParameter->OperatorIndex) << TEXTVIEW(",");
-			Printer << LowerValue(SubstrateParameter->MaxDistanceFromLeaves);
-			Printer << EndArgs;
-		}
-	}
-
-	void LowerSubstrateVerticalLayering(const MIR::FSubstrateVerticalLayering* SubstrateParameter)
-	{
-		if (SubstrateParameter->bParameterBlendingEnabled)
-		{
-			Printer << TEXTVIEW("SubstrateVerticalLayeringParameterBlending");
-			Printer << BeginArgs;
-			Printer << LowerValue(SubstrateParameter->Top) << TEXTVIEW(",");
-			Printer << LowerValue(SubstrateParameter->Base) << TEXTVIEW(",");
-			Printer << LowerValue(SubstrateParameter->SharedLocalBasisIndexMacro) << TEXTVIEW(",");
-			Printer << LowerValue(SubstrateParameter->TopNoV) << TEXTVIEW(",");
-			Printer << LowerValue(SubstrateParameter->BaseNoV);
-			Printer << EndArgs;
-		}
-		else
-		{
-			Printer << GetSubstrateTree(SubstrateParameter->bFullySimplifiedInstruction) << TEXTVIEW(".SubstrateVerticalLayering");
-			Printer << BeginArgs;
-			Printer << LowerValue(SubstrateParameter->Top) << TEXTVIEW(",");
-			Printer << LowerValue(SubstrateParameter->Base) << TEXTVIEW(",");
-			Printer << LowerValue(SubstrateParameter->OperatorIndex) << TEXTVIEW(",");
-			Printer << LowerValue(SubstrateParameter->MaxDistanceFromLeaves);
-			Printer << EndArgs;
-		}
-	}
-
-	void LowerSubstratePromoteToOperator(const MIR::FSubstratePromoteToOperator* SubstrateParameter)
-	{
-		Printer << GetSubstrateTree(SubstrateParameter->bFullySimplifiedInstruction) << TEXTVIEW(".PromoteParameterBlendedBSDFToOperator");
-		Printer << BeginArgs;
-		Printer << LowerValue(SubstrateParameter->SubstrateDataInput) << TEXTVIEW(",");
-		Printer << LowerValue(SubstrateParameter->OperatorIndex) << TEXTVIEW(",");
-		Printer << LowerValue(SubstrateParameter->BSDFIndex) << TEXTVIEW(",");
-		Printer << LowerValue(SubstrateParameter->LayerDepth) << TEXTVIEW(",");
-		Printer << LowerValue(SubstrateParameter->bIsBottom);
-		Printer << EndArgs;
-	}
-
 	void EmitPreshaderBufferReadDoubleVector(const MIR::FPrimitive& PrimitiveType, uint32 GlobalComponentOffset)
 	{
 		const int32 NumComponents = PrimitiveType.NumComponents();
@@ -1292,6 +1141,14 @@
 		return false;
 	}
 
+	void LowerNop(const MIR::FNop* Nop)
+	{
+		// NOP instructions are only used to analyze their argument, but have no effect, thus we compile it to a default
+		// value based on its type.
+		checkf(!Nop->Type.IsDouble(), TEXT("NOPs do not support LWC primitive type yet"));
+		Printer << TEXTVIEW("((") << Nop->Type << TEXTVIEW(")0)");
+	}
+
 	void LowerComposite(const MIR::FComposite* Composite) 
 	{
 		TConstArrayView<MIR::FValue*> Components = Composite->GetComponents();
@@ -1420,7 +1277,7 @@
 	{
 		Printer << (CurrentStage != MIR::Stage_Vertex && SetMaterialOutput->Property != MP_CustomOutput ? TEXTVIEW("PixelMaterialInputs.") : TEXTVIEW("Parameters."))
 				<< SetMaterialOutput->Name.ToView() << TEXTVIEW(" = ") << LowerValue(SetMaterialOutput->Arg);
-
+		
 		bLastInstructionWasSetMaterialOutput = true;
 	}
 
@@ -1449,6 +1306,8 @@
 				case MIR::BO_LessThanOrEquals:
 				case MIR::BO_Equals:
 				case MIR::BO_NotEquals:
+				case MIR::BO_And:
+				case MIR::BO_Or:
 				case MIR::BO_Add:
 				case MIR::BO_Multiply:
 				case MIR::BO_Subtract:
@@ -1457,6 +1316,8 @@
 				case MIR::BO_BitwiseAnd:
 				case MIR::BO_BitwiseOr:
 				case MIR::BO_BitwiseXor:
+				case MIR::BO_BitShiftLeft:
+				case MIR::BO_BitShiftRight:
 					bOperatorIsInfix = true;
 					break;
 
@@ -1483,6 +1344,8 @@
 				case MIR::BO_BitwiseAnd: OpString = TEXTVIEW("&"); break;
 				case MIR::BO_BitwiseOr: OpString = TEXTVIEW("|"); break;
 				case MIR::BO_BitwiseXor: OpString = TEXTVIEW("^"); break;
+				case MIR::BO_BitShiftLeft: OpString = TEXTVIEW("<<"); break;
+				case MIR::BO_BitShiftRight: OpString = TEXTVIEW(">>"); break;
 				default: UE_MIR_UNREACHABLE();
 			}
 
@@ -2098,14 +1961,6 @@
 		}
 	}
 
-	void LowerNop(const MIR::FNop* Nop)
-	{
-		// NOP instructions are only used to analyze their argument, but have no effect, thus we compile it to a default
-		// value based on its type.
-		checkf(!Nop->Type.IsDouble(), TEXT("NOPs do not support LWC primitive type yet"));
-		Printer << TEXTVIEW("((") << Nop->Type << TEXTVIEW(")0)");
-	}
-
 	void LowerCall(const MIR::FCall* Call)
 	{
 		int32 ParamLocal = NumLocals;
@@ -2167,6 +2022,157 @@
 		Printer << TEXT('_') << ParamIndex;
 	}
 	
+	void LowerSubstrateDefaultSlab(const MIR::FSubstrateDefaultSlab* SubstrateParameter)
+	{
+		Printer << TEXTVIEW("GetInitialisedSubstrateData();");
+	}
+
+	static FStringView GetSharedLocalBasesTypes(bool bFullySimplifiedInstruction)
+	{
+		return bFullySimplifiedInstruction ? TEXTVIEW("Parameters.SharedLocalBasesFullySimplified.Types") : TEXTVIEW("Parameters.SharedLocalBases.Types");
+	}
+
+	static FStringView GetSubstrateTree(bool bFullySimplifiedInstruction)
+	{
+		return bFullySimplifiedInstruction ? TEXTVIEW("Parameters.SubstrateTreeFullySimplified") : TEXTVIEW("Parameters.SubstrateTree");
+	}
+
+	void LowerSubstrateSlab(const MIR::FSubstrateSlab* SubstrateParameter)
+	{
+		Printer << TEXTVIEW("GetSubstrateSlabBSDF");
+		Printer << BeginArgs;
+		Printer << TEXTVIEW("Parameters.SubstratePixelFootprint") << TEXTVIEW(",");
+		Printer << LowerValue(SubstrateParameter->Normal) << TEXTVIEW(",");
+		Printer << LowerValue(SubstrateParameter->DiffuseAlbedo) << TEXTVIEW(",");
+		Printer << LowerValue(SubstrateParameter->F0) << TEXTVIEW(",");
+		Printer << LowerValue(SubstrateParameter->F90) << TEXTVIEW(",");
+		Printer << LowerValue(SubstrateParameter->Roughness) << TEXTVIEW(",");
+		Printer << LowerValue(SubstrateParameter->Anisotropy) << TEXTVIEW(",");
+		Printer << LowerValue(SubstrateParameter->SSSProfileId) << TEXTVIEW(",");
+		Printer << TEXTVIEW("false") << TEXTVIEW(","); // bSupportDefaultSSSProfile
+		Printer << LowerValue(SubstrateParameter->SSSMFP) << TEXTVIEW(",");
+		Printer << LowerValue(SubstrateParameter->SSSMFPScale) << TEXTVIEW(",");
+		Printer << LowerValue(SubstrateParameter->SSSPhaseAniso) << TEXTVIEW(",");
+		Printer << LowerValue(SubstrateParameter->SSSType) << TEXTVIEW(",");
+		Printer << LowerValue(SubstrateParameter->EmissiveColor) << TEXTVIEW(",");
+		Printer << LowerValue(SubstrateParameter->SecondRoughness) << TEXTVIEW(",");
+		Printer << LowerValue(SubstrateParameter->SecondRoughnessWeight) << TEXTVIEW(",");
+		Printer << LowerValue(SubstrateParameter->SecondRoughnessAsSimpleClearCoat) << TEXTVIEW(",");
+		Printer << LowerValue(SubstrateParameter->ClearCoatUseSecondNormal) << TEXTVIEW(",");
+		Printer << LowerValue(SubstrateParameter->ClearCoatBottomNormal) << TEXTVIEW(",");
+		Printer << LowerValue(SubstrateParameter->FuzzAmount) << TEXTVIEW(",");
+		Printer << LowerValue(SubstrateParameter->FuzzColor) << TEXTVIEW(",");
+		Printer << LowerValue(SubstrateParameter->FuzzRoughness) << TEXTVIEW(",");
+		Printer << LowerValue(SubstrateParameter->GlintValue) << TEXTVIEW(",");
+		Printer << LowerValue(SubstrateParameter->GlintUV) << TEXTVIEW(",");
+		Printer << LowerValue(SubstrateParameter->SpecularProfileId) << TEXTVIEW(",");
+		Printer << LowerValue(SubstrateParameter->Thickness) << TEXTVIEW(",");
+		Printer << LowerValue(SubstrateParameter->IsThin) << TEXTVIEW(",");
+		Printer << LowerValue(SubstrateParameter->IsAtBottom) << TEXTVIEW(",");
+		Printer << LowerValue(SubstrateParameter->LocalBasisIndex) << TEXTVIEW(",");
+		Printer << GetSharedLocalBasesTypes(SubstrateParameter->bFullySimplifiedInstruction);
+		Printer << EndArgs;
+	}
+
+	void LowerSubstrateShadingModel(const MIR::FSubstrateShadingModels* SubstrateParameter)
+	{
+		Printer << TEXTVIEW("SubstrateConvertLegacyMaterial") << (SubstrateParameter->bHasDynamicShadingModels ? TEXTVIEW("Dynamic") : TEXTVIEW("Static"));
+		Printer << BeginArgs;
+		Printer << TEXTVIEW("Parameters.SubstratePixelFootprint") << TEXTVIEW(",");
+		Printer << LowerValue(SubstrateParameter->BaseColor) << TEXTVIEW(",");
+		Printer << LowerValue(SubstrateParameter->Specular) << TEXTVIEW(",");
+		Printer << LowerValue(SubstrateParameter->Metallic) << TEXTVIEW(",");
+		Printer << LowerValue(SubstrateParameter->Roughness) << TEXTVIEW(",");
+		Printer << LowerValue(SubstrateParameter->Anisotropy) << TEXTVIEW(",");
+		Printer << LowerValue(SubstrateParameter->SubSurfaceColor) << TEXTVIEW(",");
+		Printer << LowerValue(SubstrateParameter->SubSurfaceProfileId) << TEXTVIEW(",");
+		Printer << LowerValue(SubstrateParameter->ClearCoat) << TEXTVIEW(",");
+		Printer << LowerValue(SubstrateParameter->ClearCoatRoughness) << TEXTVIEW(",");
+		Printer << LowerValue(SubstrateParameter->EmissiveColor) << TEXTVIEW(",");
+		Printer << LowerValue(SubstrateParameter->Opacity) << TEXTVIEW(",");
+		Printer << LowerValue(SubstrateParameter->ThinTranslucentTransmittanceColor) << TEXTVIEW(",");
+		Printer << LowerValue(SubstrateParameter->ThinTranslucentSurfaceCoverage) << TEXTVIEW(",");
+		Printer << LowerValue(SubstrateParameter->WaterScatteringCoefficients) << TEXTVIEW(",");
+		Printer << LowerValue(SubstrateParameter->WaterAbsorptionCoefficients) << TEXTVIEW(",");
+		Printer << LowerValue(SubstrateParameter->WaterPhaseG) << TEXTVIEW(",");
+		Printer << LowerValue(SubstrateParameter->ColorScaleBehindWater) << TEXTVIEW(",");
+		Printer << LowerValue(SubstrateParameter->ShadingModel) << TEXTVIEW(",");
+		Printer << LowerValue(SubstrateParameter->Normal) << TEXTVIEW(",");
+		Printer << LowerValue(SubstrateParameter->Tangent) << TEXTVIEW(",");
+		Printer << LowerValue(SubstrateParameter->ClearCoatNormal) << TEXTVIEW(",");
+		Printer << LowerValue(SubstrateParameter->CustomTangent) << TEXTVIEW(",");
+		Printer << LowerValue(SubstrateParameter->BasisIndexMacro) << TEXTVIEW(",");
+		Printer << LowerValue(SubstrateParameter->ClearCoat_BasisIndexMacro) << TEXTVIEW(",");
+		Printer << GetSharedLocalBasesTypes(SubstrateParameter->bFullySimplifiedInstruction);
+		Printer << EndArgs;
+	}
+
+	void LowerSubstrateHorizontalMixing(const MIR::FSubstrateHorizontalMixing* SubstrateParameter)
+	{
+		if (SubstrateParameter->bParameterBlendingEnabled)
+		{
+			Printer << TEXTVIEW("SubstrateHorizontalMixingParameterBlending");
+			Printer << BeginArgs;
+			Printer << LowerValue(SubstrateParameter->Background) << TEXTVIEW(",");
+			Printer << LowerValue(SubstrateParameter->Foreground) << TEXTVIEW(",");
+			Printer << LowerValue(SubstrateParameter->Mix) << TEXTVIEW(",");
+			Printer << LowerValue(SubstrateParameter->NormalMix) << TEXTVIEW(",");
+			Printer << LowerValue(SubstrateParameter->SharedLocalBasisIndexMacro) << TEXTVIEW(",");
+			Printer << GetSharedLocalBasesTypes(SubstrateParameter->bFullySimplifiedInstruction) << TEXTVIEW(",");
+			Printer << LowerValue(SubstrateParameter->BackgroundNoV) << TEXTVIEW(",");
+			Printer << LowerValue(SubstrateParameter->ForegroundNoV);
+			Printer << EndArgs;
+		}
+		else
+		{
+			Printer << GetSubstrateTree(SubstrateParameter->bFullySimplifiedInstruction) << TEXTVIEW(".SubstrateHorizontalMixing");
+			Printer << BeginArgs;
+			Printer << LowerValue(SubstrateParameter->Background) << TEXTVIEW(",");
+			Printer << LowerValue(SubstrateParameter->Foreground) << TEXTVIEW(",");
+			Printer << LowerValue(SubstrateParameter->Mix) << TEXTVIEW(",");
+			Printer << LowerValue(SubstrateParameter->OperatorIndex) << TEXTVIEW(",");
+			Printer << LowerValue(SubstrateParameter->MaxDistanceFromLeaves);
+			Printer << EndArgs;
+		}
+	}
+
+	void LowerSubstrateVerticalLayering(const MIR::FSubstrateVerticalLayering* SubstrateParameter)
+	{
+		if (SubstrateParameter->bParameterBlendingEnabled)
+		{
+			Printer << TEXTVIEW("SubstrateVerticalLayeringParameterBlending");
+			Printer << BeginArgs;
+			Printer << LowerValue(SubstrateParameter->Top) << TEXTVIEW(",");
+			Printer << LowerValue(SubstrateParameter->Base) << TEXTVIEW(",");
+			Printer << LowerValue(SubstrateParameter->SharedLocalBasisIndexMacro) << TEXTVIEW(",");
+			Printer << LowerValue(SubstrateParameter->TopNoV) << TEXTVIEW(",");
+			Printer << LowerValue(SubstrateParameter->BaseNoV);
+			Printer << EndArgs;
+		}
+		else
+		{
+			Printer << GetSubstrateTree(SubstrateParameter->bFullySimplifiedInstruction) << TEXTVIEW(".SubstrateVerticalLayering");
+			Printer << BeginArgs;
+			Printer << LowerValue(SubstrateParameter->Top) << TEXTVIEW(",");
+			Printer << LowerValue(SubstrateParameter->Base) << TEXTVIEW(",");
+			Printer << LowerValue(SubstrateParameter->OperatorIndex) << TEXTVIEW(",");
+			Printer << LowerValue(SubstrateParameter->MaxDistanceFromLeaves);
+			Printer << EndArgs;
+		}
+	}
+
+	void LowerSubstratePromoteToOperator(const MIR::FSubstratePromoteToOperator* SubstrateParameter)
+	{
+		Printer << GetSubstrateTree(SubstrateParameter->bFullySimplifiedInstruction) << TEXTVIEW(".PromoteParameterBlendedBSDFToOperator");
+		Printer << BeginArgs;
+		Printer << LowerValue(SubstrateParameter->SubstrateDataInput) << TEXTVIEW(",");
+		Printer << LowerValue(SubstrateParameter->OperatorIndex) << TEXTVIEW(",");
+		Printer << LowerValue(SubstrateParameter->BSDFIndex) << TEXTVIEW(",");
+		Printer << LowerValue(SubstrateParameter->LayerDepth) << TEXTVIEW(",");
+		Printer << LowerValue(SubstrateParameter->bIsBottom);
+		Printer << EndArgs;
+	}
+
 	/* Finalization */
 
 	void GenerateCustomOutputDeclarations()
@@ -2225,6 +2231,8 @@
 				}
 			}
 		}
+
+		OutputShaderTemplateParams.Add(TEXT("custom_outputs"), MoveTemp(Printer.Buffer));
 	}
 
 	// Generates the definitions of all the custom HLSL functions in the module and puts the resulting string
diff --git a/Engine/Source/Runtime/Engine/Public/Materials/MaterialIR.h b/Engine/Source/Runtime/Engine/Public/Materials/MaterialIR.h
--- a/Engine/Source/Runtime/Engine/Public/Materials/MaterialIR.h
+++ b/Engine/Source/Runtime/Engine/Public/Materials/MaterialIR.h
@@ -115,6 +115,7 @@ enum EValueKind
 
 	VK_InstructionBegin,
 
+	VK_Nop,
 	VK_SetMaterialOutput,
 	VK_Composite,
 	VK_Operator,
@@ -132,7 +133,6 @@ enum EValueKind
 	VK_SubstratePromoteToOperator,
 	VK_StageSwitch,
 	VK_HardwarePartialDerivative,
-	VK_Nop,
 	VK_Call,
 	VK_CallParameterOutput,
 	VK_PreshaderParameter,
@@ -198,7 +198,7 @@ struct FValue
 	bool HasSubgraphProperties(EGraphProperties Properties) const;
 
 	// Enables specified properties for this value.
-	void UseSubgraphProperties(EGraphProperties Properties);
+	void SetSubgraphProperties(EGraphProperties Properties);
 
 	// Returns the size in bytes of this value instance.
 	uint32 GetSizeInBytes() const;
@@ -693,6 +693,14 @@ struct TInstruction : FInstruction
 	static constexpr uint32 NumStaticUses = TNumStaticUses;
 };
 
+// A "no operation" instruction is a dummy instruction used to perform analysis on a value (and its dependencies)
+// but not push it to the list of instructions in a block.
+// A NOP instructions should effectively bahave as the default value of Arg type.
+struct FNop : TInstruction<VK_Nop, 1>
+{
+	FValue* Arg;
+};
+
 // An aggregate of other values.
 //
 // A dimensional is a fixed array of other values. This value is used to model vectors and matrices.
@@ -1213,14 +1221,6 @@ struct FHardwarePartialDerivative : TInstruction<VK_HardwarePartialDerivative, 1
 	EDerivativeSource Source;
 };
 
-// A "no operation" instruction is a dummy instruction used to perform analysis on a value (and its dependencies)
-// but not push it to the list of instructions in a block.
-// A NOP instructions should effectively bahave as the default value of Arg type.
-struct FNop : TInstruction<VK_Nop, 1>
-{
-	FValue* Arg;
-};
-
 // Kind of user-specified function.
 enum class FFunctionKind
 {
@@ -1334,6 +1334,7 @@ struct FCallParameterOutput : TInstruction<VK_CallParameterOutput, 1>
 	int32 	Index;
 };
 
+// @massimo.tristano review this value for better integration and support of preshaders within new translator
 struct FPreshaderParameterPayload
 {
 	// Uniform index for RVTs. Only used for UE::Shader::EPreshaderOpcode::RuntimeVirtualTextureUniform.
diff --git a/Engine/Source/Runtime/Engine/Public/Materials/MaterialIREmitter.h b/Engine/Source/Runtime/Engine/Public/Materials/MaterialIREmitter.h
--- a/Engine/Source/Runtime/Engine/Public/Materials/MaterialIREmitter.h
+++ b/Engine/Source/Runtime/Engine/Public/Materials/MaterialIREmitter.h
@@ -736,30 +736,6 @@ public:
 	// @param UsedGraphProperties Properties used during graph processing (default is None).
 	FValueRef InlineHLSL(const FMaterialExternalCodeDeclaration* InExternalCodeDeclaration, TConstArrayView<FValueRef> InputValues, EValueFlags ValueFlags = EValueFlags::None, EGraphProperties UsedGraphProperties = EGraphProperties::None);
 
-	// Emits an instruction that promotes a SubstrateData representing a default slab, encoded as MCT_Substrate.
-	FValueRef SubstrateDefaultSlab();
-
-	// Promote a SubstrateData BSDF to an operator so that is becomes a node of the SubstrateTree and can be combined with other operators.
-	FValueRef SubstratePromoteToOperator(FValueRef SubstrateData, FValueRef OperatorIndex, FValueRef BSDFIndex, FValueRef LayerDepth, FValueRef bIsBottom);
-
-	// Create a SubstrateData representing a legacy BSDF.
-	FValueRef SubstrateShadingModels(const FSubstrateShadingModelsDesc Desc);
-
-	// Create a SubstrateData representing a Slab BSDF.
-	FValueRef SubstrateSlab(const FSubstrateSlabDesc& Desc);
-
-	// Get the SubstrateData resulting from the horizontal mixing of the Background and Foreground SubstrateData.
-	FValueRef SubstrateHorizontalMixing(FValueRef Background, FValueRef Foreground, FValueRef Mix, FValueRef OperatorIndex, FValueRef MaxDistanceFromLeaves);
-
-	// Parameter blended version of SubstrateHorizontalMixing.
-	FValueRef SubstrateHorizontalMixingParameterBlending(FValueRef Background, FValueRef Foreground, FValueRef Mix, FValueRef NormalMix, FValueRef SharedLocalBasisIndexMacro, FValueRef BackgroundNoV, FValueRef ForegroundNoV);
-
-	// Get the SubstrateData resulting from the vertical layering of the Top over the Base SubstrateData.
-	FValueRef SubstrateVerticalLayering(FValueRef Top, FValueRef Base, FValueRef OperatorIndex, FValueRef MaxDistanceFromLeaves);
-
-	// Parameter blended version of SubstrateVerticalLayering.
-	FValueRef SubstrateVerticalLayeringParameterBlending(FValueRef Top, FValueRef Base, FValueRef SharedLocalBasisIndexMacro, FValueRef TopNoV, FValueRef BaseNoV);
-
 	// Defines a user-specified HLSL function to embed in the generated HLSL when generated IR Module is translated to HLSL.
 	// This will create a local function with the specified description. You can issue calls to the returned functions using the Call() API.
 	const FFunction* FunctionHLSL(const FFunctionHLSLDesc& Desc);
@@ -843,6 +819,32 @@ public:
 	FValueRef Select(FValueRef A, FValueRef B, FValueRef C) { return Operator(TO_Select, A, B, C); }
 	FValueRef Smoothstep(FValueRef A, FValueRef B, FValueRef C) { return Operator(TO_Smoothstep, A, B, C); }
 
+	/*----------------------------- Substrate instructions -----------------------------*/
+	
+	// Emits an instruction that promotes a SubstrateData representing a default slab, encoded as MCT_Substrate.
+	FValueRef SubstrateDefaultSlab();
+
+	// Promote a SubstrateData BSDF to an operator so that is becomes a node of the SubstrateTree and can be combined with other operators.
+	FValueRef SubstratePromoteToOperator(FValueRef SubstrateData, FValueRef OperatorIndex, FValueRef BSDFIndex, FValueRef LayerDepth, FValueRef bIsBottom);
+
+	// Create a SubstrateData representing a legacy BSDF.
+	FValueRef SubstrateShadingModels(const FSubstrateShadingModelsDesc Desc);
+
+	// Create a SubstrateData representing a Slab BSDF.
+	FValueRef SubstrateSlab(const FSubstrateSlabDesc& Desc);
+
+	// Get the SubstrateData resulting from the horizontal mixing of the Background and Foreground SubstrateData.
+	FValueRef SubstrateHorizontalMixing(FValueRef Background, FValueRef Foreground, FValueRef Mix, FValueRef OperatorIndex, FValueRef MaxDistanceFromLeaves);
+
+	// Parameter blended version of SubstrateHorizontalMixing.
+	FValueRef SubstrateHorizontalMixingParameterBlending(FValueRef Background, FValueRef Foreground, FValueRef Mix, FValueRef NormalMix, FValueRef SharedLocalBasisIndexMacro, FValueRef BackgroundNoV, FValueRef ForegroundNoV);
+
+	// Get the SubstrateData resulting from the vertical layering of the Top over the Base SubstrateData.
+	FValueRef SubstrateVerticalLayering(FValueRef Top, FValueRef Base, FValueRef OperatorIndex, FValueRef MaxDistanceFromLeaves);
+
+	// Parameter blended version of SubstrateVerticalLayering.
+	FValueRef SubstrateVerticalLayeringParameterBlending(FValueRef Top, FValueRef Base, FValueRef SharedLocalBasisIndexMacro, FValueRef TopNoV, FValueRef BaseNoV);
+
 	/*---------------------------- Pass through accessors -----------------------------*/
 
 	// Returns the shader platform the source material is translating for.
diff --git a/Engine/Source/Runtime/Engine/Public/Materials/MaterialIRModule.h b/Engine/Source/Runtime/Engine/Public/Materials/MaterialIRModule.h
--- a/Engine/Source/Runtime/Engine/Public/Materials/MaterialIRModule.h
+++ b/Engine/Source/Runtime/Engine/Public/Materials/MaterialIRModule.h
@@ -124,6 +124,9 @@ public:
 	FEntryPoint& GetEntryPoint(int32 Index) { return EntryPoints[Index]; }
 	const FEntryPoint& GetEntryPoint(int32 Index) const { return EntryPoints[Index]; }
 
+	// ADds a value to the module's list of all values.
+	void AddValue(MIR::FValue* Value) { Values.Add(Value); }
+
 	// Sets the value of a material property for future reference.
 	void SetPropertyValue(EMaterialProperty Property, MIR::FValue* Value) { PropertyValues[Property] = Value; }
 	
@@ -134,7 +137,7 @@ public:
 	TConstArrayView<const MIR::FSetMaterialOutput*> GetCustomOutputs() const { return CustomOutputs; }
 
 	//
-	void AddCustomOutputGroup(FStringView Name, int32 NumOutputs);
+	void BeginCustomOutputGroup(FStringView Name, int32 NumOutputs);
 
 	// Returns the array custom output groups.
 	TConstArrayView<FCustomOutputGroup> GetCustomOutputGroups() const { return CustomOutputGroups; }
@@ -149,7 +152,7 @@ public:
 	void AddEnvironmentDefine(FName Name) { EnvironmentDefines.Add(Name); }
 
 	// Temporary. @massimo.tristano refactor this
-	TArray<TPair<FName, int32>> GetIntegerEnvironmentDefines() const { return IntegerEnvironmentDefines; }
+	TConstArrayView<TPair<FName, int32>> GetIntegerEnvironmentDefines() const { return IntegerEnvironmentDefines; }
 
 	// Temporary. @massimo.tristano refactor this
 	void AddIntegerEnvironmentDefine(FName Name, int32 Value) { IntegerEnvironmentDefines.Emplace(Name, Value); }
