Add an explicit `FAttachmentCounter` that works under transactions.

#Synced-CL 49988233

diff --git a/Engine/Source/Runtime/Engine/Classes/Components/PrimitiveComponent.h b/Engine/Source/Runtime/Engine/Classes/Components/PrimitiveComponent.h
--- a/Engine/Source/Runtime/Engine/Classes/Components/PrimitiveComponent.h
+++ b/Engine/Source/Runtime/Engine/Classes/Components/PrimitiveComponent.h
@@ -947,7 +947,7 @@ private:
 public:	
 
 	FPrimitiveSceneInfoData& GetSceneData() { return SceneData; }
-	FThreadSafeCounter* GetAttachmentCounter() { return &GetSceneData().AttachmentCounter; }
+	FAttachmentCounter* GetAttachmentCounter() { return &GetSceneData().AttachmentCounter; }
 
 	ENGINE_API int32 GetRayTracingGroupId() const;
 
diff --git a/Engine/Source/Runtime/Engine/Public/PrimitiveSceneDesc.h b/Engine/Source/Runtime/Engine/Public/PrimitiveSceneDesc.h
--- a/Engine/Source/Runtime/Engine/Public/PrimitiveSceneDesc.h
+++ b/Engine/Source/Runtime/Engine/Public/PrimitiveSceneDesc.h
@@ -53,7 +53,7 @@ struct FPrimitiveSceneDesc
 	bool IsUnreachable() { return bIsUnreachable; }
 	bool ShouldRecreateProxyOnUpdateTransform() { return bRecreateProxyOnUpdateTransform; }
 
-	FThreadSafeCounter* GetAttachmentCounter() const { return PrimitiveSceneData ? &PrimitiveSceneData->AttachmentCounter : nullptr; }
+	FAttachmentCounter* GetAttachmentCounter() const { return PrimitiveSceneData ? &PrimitiveSceneData->AttachmentCounter : nullptr; }
 	FPrimitiveComponentId GetPrimitiveSceneId() const { return PrimitiveSceneData->PrimitiveSceneId; }
 	FPrimitiveComponentId GetLODParentId() const { return LodParentComponentId; }
 	FPrimitiveComponentId GetLightingAttachmentId() const { return LightingAttachmentComponentId; }
diff --git a/Engine/Source/Runtime/Engine/Public/PrimitiveSceneInfoData.h b/Engine/Source/Runtime/Engine/Public/PrimitiveSceneInfoData.h
--- a/Engine/Source/Runtime/Engine/Public/PrimitiveSceneInfoData.h
+++ b/Engine/Source/Runtime/Engine/Public/PrimitiveSceneInfoData.h
@@ -3,12 +3,95 @@
 #pragma once 
 
 #include "AutoRTFM.h"
+#include "CoreGlobals.h"
 #include "HAL/ThreadSafeCounter.h"
 #include "PrimitiveComponentId.h"
 
 class FPrimitiveSceneProxy;
 struct FActorLastRenderTime;
 
+// An attachment counter is a thread-safe counter that is only ever
+// incremented on the game thread, and decremented on the rendering
+// thread.
+//
+// It also supports transactions, and will undo any increments on the
+// game thread if a transaction aborts.
+struct FAttachmentCounter final
+{
+	UE_NONCOPYABLE(FAttachmentCounter);
+
+	AUTORTFM_OPEN FAttachmentCounter() = default;
+
+#if UE_AUTORTFM
+	void Increment()
+	{
+		check(IsInGameThread());
+
+		UE_AUTORTFM_OPEN
+		{
+			Payload->Payload++;
+		};
+
+		UE_AUTORTFM_ONABORT(this)
+		{
+			Payload->Payload--;
+		};
+	}
+
+	AUTORTFM_DISABLE void Decrement()
+	{
+		check(IsInRenderingThread());
+
+		Payload->Payload--;
+	}
+
+	int32 GetValue() const
+	{
+		int32 Result = 0;
+
+		UE_AUTORTFM_OPEN
+		{
+			Result = Payload->Payload;
+		};
+
+		return Result;
+	}
+private:
+
+	// Under AutoRTFM because we are creating and *sometimes* destroying
+	// `FAttachmentCounter`'s in the same transaction, we need it to be
+	// heap allocated, so that the memory doesn't just instantly go away
+	// and our on-abort handlers would then use-after-free.
+	struct FPayload final : public UE::Private::FQueryableRefCountedObject
+	{
+		std::atomic<int32> Payload = 0;
+	};
+
+	TRefCountPtr<FPayload> Payload = new FPayload();
+
+#else
+
+	void Increment()
+	{
+		check(IsInGameThread());
+		Payload++;
+	}
+
+	void Decrement()
+	{
+		check(IsInRenderingThread());
+		Payload--;
+	}
+
+	int32 GetValue() const
+	{
+		return Payload;
+	}
+private:
+	std::atomic<int32> Payload = 0;
+#endif
+};
+
 /*
  * All the necessary information for scene primitive to component feedback 
  */
@@ -49,7 +132,7 @@ struct FPrimitiveSceneInfoData
 	*                             for AttachmentCounter to be larger than 1, e.g.
 	*                             during reattachment.
 	*/
-	FThreadSafeCounter AttachmentCounter;
+	FAttachmentCounter AttachmentCounter;
 
 	/** Used by the renderer, to identify a primitive across re-registers. */
 	FPrimitiveComponentId PrimitiveSceneId;
diff --git a/Engine/Source/Runtime/Renderer/Private/RendererScene.cpp b/Engine/Source/Runtime/Renderer/Private/RendererScene.cpp
--- a/Engine/Source/Runtime/Renderer/Private/RendererScene.cpp
+++ b/Engine/Source/Runtime/Renderer/Private/RendererScene.cpp
@@ -2054,7 +2054,7 @@ void FScene::BatchRemovePrimitivesInternal(TArrayView<T*> InPrimitives)
 	{
 		FPrimitiveSceneInfo* PrimitiveSceneInfo;
 		FPrimitiveSceneProxy* PrimitiveSceneProxy;
-		FThreadSafeCounter* AttachmentCounter;
+		FAttachmentCounter* AttachmentCounter;
 	};
 
 	TArray<FDetachCommand, SceneRenderingAllocator> DestroyCommands;
