NetTest
* NetTest framework can now tell if it correctly loaded a map

#rnx

#Synced-CL 49948957

diff --git a/Engine/Source/Runtime/Engine/Private/Tests/NetTestHelpers.cpp b/Engine/Source/Runtime/Engine/Private/Tests/NetTestHelpers.cpp
--- a/Engine/Source/Runtime/Engine/Private/Tests/NetTestHelpers.cpp
+++ b/Engine/Source/Runtime/Engine/Private/Tests/NetTestHelpers.cpp
@@ -137,7 +137,8 @@ FTestWorldInstance FTestWorldInstance::CreateServer(const TCHAR* InURL, const ER
 
 	FURL LocalURL(nullptr, InURL, TRAVEL_Absolute);
 	FString BrowseError;
-	GEngine->Browse(*NewInstance.GameInstance->GetWorldContext(), LocalURL, BrowseError);
+	const EBrowseReturnVal::Type Result = GEngine->Browse(*NewInstance.GameInstance->GetWorldContext(), LocalURL, BrowseError);
+	NewInstance.bIsLoaded = (Result == EBrowseReturnVal::Success);
 
 	if (ReplicationSystem != EReplicationSystem::Default)
 	{
@@ -175,7 +176,8 @@ FTestWorldInstance FTestWorldInstance::CreateClient(int32 ServerPort, int32 NumL
 
 	FURL ClientURL(nullptr, *ClientURLString, TRAVEL_Absolute);
 	FString ClientBrowseError;
-	GEngine->Browse(*ClientWorldContext, ClientURL, ClientBrowseError);
+	const EBrowseReturnVal::Type Result = GEngine->Browse(*ClientWorldContext, ClientURL, ClientBrowseError);
+	NewInstance.bIsLoaded = (Result == EBrowseReturnVal::Success);
 
 	for (int32 LocalPlayerIndex = 1; LocalPlayerIndex < NumLocalPlayers; LocalPlayerIndex++)
 	{
@@ -218,12 +220,6 @@ FTestWorldInstance::FTestWorldInstance(const FGameInstancePIEParameters& Instanc
 	FWorldContext* WorldContext = GameInstance->GetWorldContext();
 }
 
-FTestWorldInstance::FTestWorldInstance(FTestWorldInstance&& Other)
-	: GameInstance(Other.GameInstance)
-{
-	Other.GameInstance = nullptr;
-}
-
 FTestWorldInstance& FTestWorldInstance::operator=(FTestWorldInstance&& Other)
 {
 	if (this != &Other)
@@ -231,7 +227,10 @@ FTestWorldInstance& FTestWorldInstance::operator=(FTestWorldInstance&& Other)
 		Shutdown();
 
 		GameInstance = Other.GameInstance;
+		bIsLoaded = Other.bIsLoaded;
+
 		Other.GameInstance = nullptr;
+		Other.bIsLoaded = false;
 	}
 
 	return *this;
diff --git a/Engine/Source/Runtime/Engine/Public/Tests/NetTestHelpers.h b/Engine/Source/Runtime/Engine/Public/Tests/NetTestHelpers.h
--- a/Engine/Source/Runtime/Engine/Public/Tests/NetTestHelpers.h
+++ b/Engine/Source/Runtime/Engine/Public/Tests/NetTestHelpers.h
@@ -146,9 +146,19 @@ public:
 	FTestWorldInstance(const FTestWorldInstance&) = delete;
 	FTestWorldInstance& operator=(const FTestWorldInstance&) = delete;
 
-	explicit FTestWorldInstance(bool bDelayedInit) : GameInstance(nullptr) {}
+	explicit FTestWorldInstance(bool bDelayedInit)
+	{
+		// GameInstance will be created manually later
+	}
 	
-	ENGINE_API FTestWorldInstance(FTestWorldInstance&& Other);
+	FTestWorldInstance(FTestWorldInstance&& Other)	
+		: GameInstance(Other.GameInstance)
+		, bIsLoaded(Other.bIsLoaded)
+	{
+		Other.GameInstance = nullptr;
+		Other.bIsLoaded = false;
+	}
+
 	ENGINE_API FTestWorldInstance& operator=(FTestWorldInstance&& Other);
 
 	ENGINE_API FWorldContext* GetWorldContext() const;	
@@ -165,6 +175,12 @@ public:
 	ENGINE_API void LoadStreamingLevel(FName LevelName);
 	ENGINE_API void UnloadStreamingLevel(FName LevelName);
 
+	/** Returns true if the world was correctly loaded and ready to play. */
+	bool IsLoaded() const
+	{
+		return bIsLoaded;
+	}
+
 public:
 
 	struct FContext
@@ -183,6 +199,7 @@ public:
 public:
 
 	UGameInstance* GameInstance = nullptr;
+	bool bIsLoaded = false;
 
 private:
 	explicit FTestWorldInstance(const FGameInstancePIEParameters& InstanceParams);
