[World] Deprecation build issue debug; this change is a small subset of a larger deprecation change to isolate a UHT issue happening in Incremental

#tests Local and preflight builds
#rnx

#Synced-CL 49964949

diff --git a/Engine/Source/Runtime/Engine/Public/WorldPartition/DataLayer/ActorDataLayer.h b/Engine/Source/Runtime/Engine/Public/WorldPartition/DataLayer/ActorDataLayer.h
--- a/Engine/Source/Runtime/Engine/Public/WorldPartition/DataLayer/ActorDataLayer.h
+++ b/Engine/Source/Runtime/Engine/Public/WorldPartition/DataLayer/ActorDataLayer.h
@@ -6,37 +6,5 @@
 
 #include "WorldPartition/DataLayer/DataLayerInstance.h"
 
-#include "ActorDataLayer.generated.h"
-
-// This class is deprecated and only present for backward compatibility purposes.
-// Instead of using FActorDatalayer, directly save the DataLayerInstance FName if the DataLayer not exposed in data.
-// If the DataLayer is exposed in Data, then use DataLayerAssets.
-USTRUCT(BlueprintType)
-struct FActorDataLayer
-{
-	GENERATED_USTRUCT_BODY()
-
-	static_assert(DATALAYER_TO_INSTANCE_RUNTIME_CONVERSION_ENABLED, "FActorDataLayer is deprecated and needs to be deleted.");
-
-	FActorDataLayer()
-	: Name(NAME_None)
-	{}
-
-	FActorDataLayer(const FName& InName)
-	: Name(InName)
-	{}
-
-	inline bool operator==(const FActorDataLayer& Other) const { return Name == Other.Name; }
-	inline bool operator<(const FActorDataLayer& Other) const { return Name.FastLess(Other.Name); }
-
-	inline operator FName () const { return Name; }
-
-	/** The name of this layer */
-	UPROPERTY(BlueprintReadOnly, VisibleAnywhere, Category = DataLayer)
-	FName Name;
-
-	friend uint32 GetTypeHash(const FActorDataLayer& Value)
-	{
-		return GetTypeHash(Value.Name);
-	}
-};
+// FActorDataLayer struct moved in preparation of deprecation 
+#include "WorldPartition/DataLayer/Deprecated/DeprecatedDataLayer.h"
diff --git a/Engine/Source/Runtime/Engine/Public/WorldPartition/DataLayer/Deprecated/DeprecatedDataLayer.h b/Engine/Source/Runtime/Engine/Public/WorldPartition/DataLayer/Deprecated/DeprecatedDataLayer.h
--- a/Engine/Source/Runtime/Engine/Public/WorldPartition/DataLayer/Deprecated/DeprecatedDataLayer.h
+++ b/Engine/Source/Runtime/Engine/Public/WorldPartition/DataLayer/Deprecated/DeprecatedDataLayer.h
@@ -11,7 +11,7 @@
 
 #include "DeprecatedDataLayer.generated.h"
 
-// These deprecated reflected types and members must not be removed for the duration of UE5, even after 2 versions post-deprecation, to allow loading and upgrading old assets.
+// These deprecated reflected types and members must not be removed for the duration of UE5, even after 2 versions post-deprecation, to prevent data losses and allow loading and upgrading old assets.
 
 UENUM(BlueprintType)
 enum class UE_DEPRECATED(all, "Use EDataLayerRuntimeState instead.") EDataLayerState : uint8
@@ -23,6 +23,35 @@ enum class UE_DEPRECATED(all, "Use EDataLayerRuntimeState instead.") EDataLayerS
 
 static_assert(EDataLayerRuntimeState::Unloaded < EDataLayerRuntimeState::Loaded && EDataLayerRuntimeState::Loaded < EDataLayerRuntimeState::Activated, "Streaming Query code is dependent on this being true");
 
+// Soon to be deprecated
+USTRUCT(BlueprintType)
+struct FActorDataLayer
+{
+	GENERATED_USTRUCT_BODY()
+
+	FActorDataLayer()
+	: Name(NAME_None)
+	{}
+
+	FActorDataLayer(const FName& InName)
+	: Name(InName)
+	{}
+
+	inline bool operator==(const FActorDataLayer& Other) const { return Name == Other.Name; }
+	inline bool operator<(const FActorDataLayer& Other) const { return Name.FastLess(Other.Name); }
+
+	inline operator FName() const { return Name; }
+
+	/** The name of this layer */
+	UPROPERTY(meta = (DeprecatedProperty, DeprecationMessage = "FActorDataLayer is deprecated."))
+	FName Name;
+
+	friend uint32 GetTypeHash(const FActorDataLayer& Value)
+	{
+		return GetTypeHash(Value.Name);
+	}
+};
+
 static_assert(DATALAYER_TO_INSTANCE_RUNTIME_CONVERSION_ENABLED, "UDEPRECATED_DataLayer class is deprecated and needs to be deleted.");
 class UE_DEPRECATED(all, "Use UDataLayerInstance & UDataLayerAsset to create DataLayers") UDEPRECATED_DataLayer;
 
