Hooking up the OnMIDAssigned/Removed callbacks to support verse material replication on props. Also fixing up remaining calls to OnAssigned/RemovedAsOverride to also use GetOwner()

#Synced-CL 49940788

diff --git a/Engine/Source/Runtime/Engine/Private/Components/MeshComponent.cpp b/Engine/Source/Runtime/Engine/Private/Components/MeshComponent.cpp
--- a/Engine/Source/Runtime/Engine/Private/Components/MeshComponent.cpp
+++ b/Engine/Source/Runtime/Engine/Private/Components/MeshComponent.cpp
@@ -89,7 +89,7 @@ void UMeshComponent::SetMaterial(int32 ElementIndex, UMaterialInterface* Materia
 
 			if (UMaterialInterface* PreviousMaterial = OverrideMaterials[ElementIndex].Get())
 			{
-				PreviousMaterial->OnRemovedAsOverride(this);
+				PreviousMaterial->OnRemovedAsOverride(GetOwner());
 			}
 
 			// Set the material and invalidate things
@@ -194,7 +194,7 @@ void UMeshComponent::CleanUpOverrideMaterials()
 		{
 			if (UMaterialInterface* MatInterface = OverrideMaterials[MatIndex].Get())
 			{
-				MatInterface->OnRemovedAsOverride(this);
+				MatInterface->OnRemovedAsOverride(GetOwner());
 			}
 		}
 
@@ -243,7 +243,7 @@ void UMeshComponent::EmptyOverrideMaterials()
 		{
 			if (UMaterialInterface* MatInterface = OverrideMaterials[MatIndex].Get())
 			{
-				MatInterface->OnRemovedAsOverride(this);
+				MatInterface->OnRemovedAsOverride(GetOwner());
 			}
 		}
 
@@ -499,7 +499,7 @@ void UMeshComponent::SetParameterValueOnMaterials(const FName ParameterName
 			bDynamicSlotOverlayMaterialCreated = true;
 
 			UMaterialInstanceDynamic* DynamicMaterial = UMaterialInstanceDynamic::Create(MaterialInterface, this);
-			DynamicMaterial->OnAssignedAsOverride(this);
+			DynamicMaterial->OnAssignedAsOverride(GetOwner());
 
 			MaterialSlotsOverlayMaterial[MaterialIndex] = DynamicMaterial;
 
@@ -613,7 +613,7 @@ void UMeshComponent::BeginDestroy()
 	{
 		if (UMaterialInterface* MatInterface = OverrideMaterials[MatIndex].Get())
 		{
-			MatInterface->OnRemovedAsOverride(this);
+			MatInterface->OnRemovedAsOverride(GetOwner());
 		}
 
 		OverrideMaterials[MatIndex] = nullptr;
