UE-360127 - Blueprint actor with skeletal mesh selected doesn't carry over when duplicated

PostLoad isn't called on archetype objects during duplication. In this case, PostLoad is what causes the SkeletalMeshAsset property to update correctly, which results in the missing mesh.

While this particular PostLoad restriction is a bit out of date, we want to move away from relying on its behavior during duplication. Instead, we can just use PostDuplicate to refresh the asset.

#Synced-CL 49967948

diff --git a/Engine/Source/Runtime/Engine/Classes/Components/SkeletalMeshComponent.h b/Engine/Source/Runtime/Engine/Classes/Components/SkeletalMeshComponent.h
--- a/Engine/Source/Runtime/Engine/Classes/Components/SkeletalMeshComponent.h
+++ b/Engine/Source/Runtime/Engine/Classes/Components/SkeletalMeshComponent.h
@@ -1806,6 +1806,7 @@ public:
 public:
 	//~ Begin UObject Interface.
 	ENGINE_API virtual void Serialize(FArchive& Ar) override;
+	ENGINE_API virtual void PostDuplicate(bool bDuplicateForPIE) override;
 	ENGINE_API virtual void PostLoad() override;
 	ENGINE_API virtual void PostInitProperties() override;
 
diff --git a/Engine/Source/Runtime/Engine/Private/Components/SkeletalMeshComponent.cpp b/Engine/Source/Runtime/Engine/Private/Components/SkeletalMeshComponent.cpp
--- a/Engine/Source/Runtime/Engine/Private/Components/SkeletalMeshComponent.cpp
+++ b/Engine/Source/Runtime/Engine/Private/Components/SkeletalMeshComponent.cpp
@@ -642,6 +642,13 @@ void USkeletalMeshComponent::Serialize(FArchive& Ar)
 	PRAGMA_ENABLE_DEPRECATION_WARNINGS
 }
 
+void USkeletalMeshComponent::PostDuplicate(bool bDuplicateForPIE)
+{
+	Super::PostDuplicate(bDuplicateForPIE);
+
+	RefreshSkeletalMeshAsset();
+}
+
 void USkeletalMeshComponent::PostLoad()
 {
 	Super::PostLoad();
