Fix for duplicate function definition.

#Synced-CL 49999797

diff --git a/Engine/Source/Runtime/Engine/Private/Animation/AnimBank.cpp b/Engine/Source/Runtime/Engine/Private/Animation/AnimBank.cpp
--- a/Engine/Source/Runtime/Engine/Private/Animation/AnimBank.cpp
+++ b/Engine/Source/Runtime/Engine/Private/Animation/AnimBank.cpp
@@ -1201,20 +1201,7 @@ bool UAnimBank::IsSequenceLooping(int32 SequenceIndex) const
 }
 ///////////////////////////////////////////////////////////////////////////////////////////////////
 
-inline float GetLoopedPosition(float Position, float TrackLength)
-{
-	if (TrackLength >= 0.0f)
-	{
-		Position = TrackLength == 0.0f ? 0.0f : FMath::Fmod(Position, TrackLength);
-
-		// Returns a negative range if input is negative.
-		if (Position < 0.0f)
-		{
-			Position += TrackLength;
-		}
-	}
-	return Position;
-}
+extern float GetLoopedPosition(float Position, float TrackLength);
 
 void FAnimBankTrackPackedData::Update(float SequenceLength, bool bLooping, float DeltaTime)
 {
