Fixed link error in same module.

#Synced-CL 49999797

diff --git a/Engine/Source/Runtime/Engine/Private/Animation/AnimBank.cpp b/Engine/Source/Runtime/Engine/Private/Animation/AnimBank.cpp
--- a/Engine/Source/Runtime/Engine/Private/Animation/AnimBank.cpp
+++ b/Engine/Source/Runtime/Engine/Private/Animation/AnimBank.cpp
@@ -1201,7 +1201,20 @@ bool UAnimBank::IsSequenceLooping(int32 SequenceIndex) const
 }
 ///////////////////////////////////////////////////////////////////////////////////////////////////
 
-extern float GetLoopedPosition(float Position, float TrackLength);
+static float GetAnimBankTrackLoopedPosition(float Position, float TrackLength)
+{
+	if (TrackLength >= 0.0f)
+	{
+		Position = TrackLength == 0.0f ? 0.0f : FMath::Fmod(Position, TrackLength);
+
+		// Returns a negative range if input is negative.
+		if (Position < 0.0f)
+		{
+			Position += TrackLength;
+		}
+	}
+	return Position;
+}
 
 void FAnimBankTrackPackedData::Update(float SequenceLength, bool bLooping, float DeltaTime)
 {
@@ -1220,8 +1233,8 @@ void FAnimBankTrackPackedData::Wrap(float SequenceLength, bool bLooping)
 
 	if (LoopModeEnum == EAnimSequenceTrackLoopMode::Loop)
 	{
-		Position         = GetLoopedPosition(Position, SequenceLength);
-		PreviousPosition = GetLoopedPosition(PreviousPosition, SequenceLength);
+		Position         = GetAnimBankTrackLoopedPosition(Position, SequenceLength);
+		PreviousPosition = GetAnimBankTrackLoopedPosition(PreviousPosition, SequenceLength);
 	}
 	else
 	{
diff --git a/Engine/Source/Runtime/Engine/Private/Animation/AnimSequenceTransformProviderData.cpp b/Engine/Source/Runtime/Engine/Private/Animation/AnimSequenceTransformProviderData.cpp
--- a/Engine/Source/Runtime/Engine/Private/Animation/AnimSequenceTransformProviderData.cpp
+++ b/Engine/Source/Runtime/Engine/Private/Animation/AnimSequenceTransformProviderData.cpp
@@ -29,7 +29,7 @@ static TAutoConsoleVariable<bool> CVarAnimSequenceTransformProviderBlending(
 
 ///////////////////////////////////////////////////////////////////////////////////////////////////
 
-inline float GetLoopedPosition(float Position, float TrackLength)
+static float GetAnimSequenceTrackLoopedPosition(float Position, float TrackLength)
 {
 	if (TrackLength >= 0.0f)
 	{
@@ -61,8 +61,8 @@ void FAnimSequenceTrackPackedData::Wrap(float SequenceLength)
 
 	if (LoopModeEnum == EAnimSequenceTrackLoopMode::Loop)
 	{
-		Position = GetLoopedPosition(Position, SequenceLength);
-		PreviousPosition = GetLoopedPosition(PreviousPosition, SequenceLength);
+		Position = GetAnimSequenceTrackLoopedPosition(Position, SequenceLength);
+		PreviousPosition = GetAnimSequenceTrackLoopedPosition(PreviousPosition, SequenceLength);
 	}
 	else
 	{
