Substrate - made a common function for CompileWithDefaultNormalWS and CompileWithDefaultTangentWS  implementation.

#tests PC editor dx12
#rnx

#Synced-CL 50002282

diff --git a/Engine/Source/Runtime/Engine/Private/Materials/MaterialExpressionsIR.cpp b/Engine/Source/Runtime/Engine/Private/Materials/MaterialExpressionsIR.cpp
--- a/Engine/Source/Runtime/Engine/Private/Materials/MaterialExpressionsIR.cpp
+++ b/Engine/Source/Runtime/Engine/Private/Materials/MaterialExpressionsIR.cpp
@@ -4996,68 +4996,66 @@ static FValueRef CompileWithDefaultFloat3(MIR::FEmitter& Em, FExpressionInput& I
 	return !ResultArg.IsValid() ? DefaultResultArg : ResultArg;
 }
 
-static FValueRef CompileWithDefaultNormalWS(MIR::FEmitter& Em, UMaterial* Material, FExpressionInput& NormalInput, bool bConvertToRequestedSpace = true)
+enum EDefaultDirection
+{
+	EDefaultDirectionNormal,
+	EDefaultDirectionTangent
+};
+static FValueRef CompileWithDefaultDirectionWS(MIR::FEmitter& Em, UMaterial* Material, FExpressionInput& DirectionInput, bool bConvertToRequestedSpace, EDefaultDirection DefaultDirection)
 {
 	FSubstrateTranslatorData* SubstrateTranslatorData = Em.GetSubstrateTranslatorData();
 
-	if (NormalInput.GetTracedInput().Expression != nullptr)
+	auto EmitDefaultDirection = [&]()
 	{
-		FValueRef NormalArg = Em.TryInput(&NormalInput);
-		if (!NormalArg)
-		{
-			// Nothing is plug in from the linked input, so specify world space normal the BSDF node expects.
-			NormalArg = EmitVertexNormal(Em);
-		}
-
-		// Ensure the normal has always a valid float3 type
-		NormalArg = Em.Cast(NormalArg, MIR::FType::MakeFloatVector(3));
-		NormalArg = Em.CastToFloatKind(NormalArg);	// Move to non LWC
-		if (!NormalArg.IsValid())
+		switch(DefaultDirection)
 		{
-			NormalArg = EmitVertexNormal(Em);
-		}
-
-		// Transform into world space normal if needed. BSDF nodes always expects world space normal as input.
-		if (bConvertToRequestedSpace)
-		{
-			NormalArg = TransformNormalFromRequestedBasisToWorld(Em, Material, NormalArg);
+		case EDefaultDirectionNormal:
+			return EmitVertexNormal(Em);
+			break;
+		case EDefaultDirectionTangent:
+			return EmitVertexTangent(Em);
+			break;
 		}
-		return NormalArg;
-	}
-	// Nothing is plug in on the BSDF node, so specify world space normal the node expects.
-	return EmitVertexNormal(Em);
-}
-
-static FValueRef CompileWithDefaultTangentWS(MIR::FEmitter& Em, UMaterial* Material, FExpressionInput& TangentInput, bool bConvertToRequestedSpace = true)
-{
-	FSubstrateTranslatorData* SubstrateTranslatorData = Em.GetSubstrateTranslatorData();
+		Em.Error(TEXT("Unkown EmitDefaultDirection parameter."));
+		return EmitVertexNormal(Em);
+	};
 
-	if (TangentInput.GetTracedInput().Expression != nullptr)
+	if (DirectionInput.GetTracedInput().Expression != nullptr)
 	{
-		FValueRef NormalArg = Em.TryInput(&TangentInput);
-		if (!NormalArg)
+		FValueRef DirectionArg = Em.TryInput(&DirectionInput);
+		if (!DirectionArg)
 		{
 			// Nothing is plug in from the linked input, so specify world space normal the BSDF node expects.
-			NormalArg = EmitVertexTangent(Em);
+			DirectionArg = EmitDefaultDirection();
 		}
 
 		// Ensure the normal has always a valid float3 type
-		NormalArg = Em.Cast(NormalArg, MIR::FType::MakeFloatVector(3));
-		NormalArg = Em.CastToFloatKind(NormalArg);	// Move to non LWC
-		if (!NormalArg.IsValid())
+		DirectionArg = Em.Cast(DirectionArg, MIR::FType::MakeFloatVector(3));
+		DirectionArg = Em.CastToFloatKind(DirectionArg);	// Move to non LWC
+		if (!DirectionArg.IsValid())
 		{
-			NormalArg = EmitVertexTangent(Em);
+			DirectionArg = EmitDefaultDirection();
 		}
 
 		// Transform into world space normal if needed. BSDF nodes always expects world space normal as input.
 		if (bConvertToRequestedSpace)
 		{
-			NormalArg = TransformNormalFromRequestedBasisToWorld(Em, Material, NormalArg);
+			DirectionArg = TransformNormalFromRequestedBasisToWorld(Em, Material, DirectionArg);
 		}
-		return NormalArg;
+		return DirectionArg;
 	}
 	// Nothing is plug in on the BSDF node, so specify world space normal the node expects.
-	return EmitVertexTangent(Em);
+	return EmitDefaultDirection();
+}
+
+static FValueRef CompileWithDefaultNormalWS(MIR::FEmitter& Em, UMaterial* Material, FExpressionInput& NormalInput, bool bConvertToRequestedSpace = true)
+{
+	return CompileWithDefaultDirectionWS(Em, Material, NormalInput, bConvertToRequestedSpace, EDefaultDirection::EDefaultDirectionNormal);
+}
+
+static FValueRef CompileWithDefaultTangentWS(MIR::FEmitter& Em, UMaterial* Material, FExpressionInput& TangentInput, bool bConvertToRequestedSpace = true)
+{
+	return CompileWithDefaultDirectionWS(Em, Material, TangentInput, bConvertToRequestedSpace, EDefaultDirection::EDefaultDirectionTangent);
 }
 
 void UMaterialExpressionSubstrateSlabBSDF::Build(MIR::FEmitter& Em)
