SplineComponent: Propagate input keys to FTangentSpline

Also adds support for encoding the parameterization of the closing segment of TPolyBezierSpline as a length in parameter space (called LoopKnotDelta), not an absolute knot value. This means that the number of pair knots is 1:1 with anchors and we can derive all of the information we need about the closing segment from the anchors, knot vector, and LoopKnotDelta.

#virtualized
#Synced-CL 49961239

diff --git a/Engine/Source/Editor/DetailCustomizations/Private/SplineComponentDetails.cpp b/Engine/Source/Editor/DetailCustomizations/Private/SplineComponentDetails.cpp
--- a/Engine/Source/Editor/DetailCustomizations/Private/SplineComponentDetails.cpp
+++ b/Engine/Source/Editor/DetailCustomizations/Private/SplineComponentDetails.cpp
@@ -1004,20 +1004,12 @@ bool FSplinePointDetails::CanSetInputKey() const
 		return false;
 	}
 
-PRAGMA_DISABLE_DEPRECATION_WARNINGS
-	bool bUsingSplineCurves = SplineComp->ShouldUseSplineCurves();
-PRAGMA_ENABLE_DEPRECATION_WARNINGS
-
-	return IsOnePointSelected() && bUsingSplineCurves;
+	return IsOnePointSelected();
 }
 
 void FSplinePointDetails::OnSetInputKey(float NewValue, ETextCommit::Type CommitInfo)
 {
 PRAGMA_DISABLE_DEPRECATION_WARNINGS
-	// Note: This function strongly assumes that SplineCurves is the authoritative data structure backing the spline component.
-	// I have made this assumption valid by introducing CanSetInputKey() which verifies this by checking with the selected component.
-	// This assumption is necessary because there is no interface on the component which allows us to modify input keys, we must directly write to SplineCurves.
-	
 	TryAcquireDetailsInterface();
 	
 	if ((CommitInfo != ETextCommit::OnEnter && CommitInfo != ETextCommit::OnUserMovedFocus) || !SplineComp || !SplineDetailsInterface)
@@ -1027,13 +1019,12 @@ PRAGMA_DISABLE_DEPRECATION_WARNINGS
 
 	check(SelectedKeys.Num() == 1);
 	const int32 Index = *SelectedKeys.CreateConstIterator();
-	TArray<FInterpCurvePoint<FVector>>& Positions = SplineComp->GetSplinePointsPosition().Points;
 
-	const int32 NumPoints = Positions.Num();
+	const int32 NumPoints = SplineComp->GetNumberOfSplinePoints();
 
 	bool bModifyOtherPoints = false;
-	if ((Index > 0 && NewValue <= Positions[Index - 1].InVal) ||
-		(Index < NumPoints - 1 && NewValue >= Positions[Index + 1].InVal))
+	if ((Index > 0 && NewValue <= SplineComp->GetInputKeyValueAtSplinePoint(Index - 1) ||
+		(Index < NumPoints - 1 && NewValue >= SplineComp->GetInputKeyValueAtSplinePoint(Index + 1))))
 	{
 		const FText Title(LOCTEXT("InputKeyTitle", "Input key out of range"));
 		const FText Message(LOCTEXT("InputKeyMessage", "Spline input keys must be numerically ascending. Would you like to modify other input keys in the spline in order to be able to set this value?"));
@@ -1050,39 +1041,31 @@ PRAGMA_DISABLE_DEPRECATION_WARNINGS
 	// Scope the transaction to only include the value change and none of the derived data changes that might arise from NotifyPropertyModified
 	{
 		const FScopedTransaction Transaction(LOCTEXT("SetSplinePointInputKey", "Set spline point input key"));
-		SplineComp->Modify();
 
-		TArray<FInterpCurvePoint<FQuat>>& Rotations = SplineComp->GetSplinePointsRotation().Points;
-		TArray<FInterpCurvePoint<FVector>>& Scales = SplineComp->GetSplinePointsScale().Points;
+		SplineComp->Modify();
 
 		if (bModifyOtherPoints)
 		{
 			// Shuffle the previous or next input keys down or up so the input value remains in sequence
-			if (Index > 0 && NewValue <= Positions[Index - 1].InVal)
+			if (Index > 0 && NewValue <= SplineComp->GetInputKeyValueAtSplinePoint(Index - 1))
 			{
-				float Delta = (NewValue - Positions[Index].InVal);
+				float Delta = (NewValue - SplineComp->GetInputKeyValueAtSplinePoint(Index));
 				for (int32 PrevIndex = 0; PrevIndex < Index; PrevIndex++)
 				{
-					Positions[PrevIndex].InVal += Delta;
-					Rotations[PrevIndex].InVal += Delta;
-					Scales[PrevIndex].InVal += Delta;
+					SplineComp->SetInputKeyValueAtSplinePoint(PrevIndex, SplineComp->GetInputKeyValueAtSplinePoint(PrevIndex) + Delta, false);
 				}
 			}
-			else if (Index < NumPoints - 1 && NewValue >= Positions[Index + 1].InVal)
+			else if (Index < NumPoints - 1 && NewValue >= SplineComp->GetInputKeyValueAtSplinePoint(Index + 1))
 			{
-				float Delta = (NewValue - Positions[Index].InVal);
-				for (int32 NextIndex = Index + 1; NextIndex < NumPoints; NextIndex++)
+				float Delta = (NewValue - SplineComp->GetInputKeyValueAtSplinePoint(Index));
+				for (int32 NextIndex = NumPoints - 1; NextIndex >= Index + 1; NextIndex--)
 				{
-					Positions[NextIndex].InVal += Delta;
-					Rotations[NextIndex].InVal += Delta;
-					Scales[NextIndex].InVal += Delta;
+					SplineComp->SetInputKeyValueAtSplinePoint(NextIndex, SplineComp->GetInputKeyValueAtSplinePoint(NextIndex) + Delta, false);
 				}
 			}
 		}
 
-		Positions[Index].InVal = NewValue;
-		Rotations[Index].InVal = NewValue;
-		Scales[Index].InVal = NewValue;
+		SplineComp->SetInputKeyValueAtSplinePoint(Index, NewValue, false);
 	}
 
 	SplineComp->UpdateSpline();
diff --git a/Engine/Source/Runtime/Core/Public/UObject/FortniteMainBranchObjectVersions.inl b/Engine/Source/Runtime/Core/Public/UObject/FortniteMainBranchObjectVersions.inl
--- a/Engine/Source/Runtime/Core/Public/UObject/FortniteMainBranchObjectVersions.inl
+++ b/Engine/Source/Runtime/Core/Public/UObject/FortniteMainBranchObjectVersions.inl
@@ -751,5 +751,8 @@ FORTNITE_MAIN_VERSION(CompositePassKeyerSource, 245)
 // OpenColorIO library update to v2.5.1 which transitions from ACES 1.3 to ACES 2.0 defaults
 FORTNITE_MAIN_VERSION(OpenColorIOv251, 246)
 
+// Reparameterize Spline in SplineComponent based on SplineCurves parameterization
+FORTNITE_MAIN_VERSION(SplineComponentReparameterizeOnLoad, 247)
+
 // -----<new versions can be added above this line>-------------------------------------------------
 #undef FORTNITE_MAIN_VERSION
diff --git a/Engine/Source/Runtime/Engine/Classes/Components/SplineComponent.h b/Engine/Source/Runtime/Engine/Classes/Components/SplineComponent.h
--- a/Engine/Source/Runtime/Engine/Classes/Components/SplineComponent.h
+++ b/Engine/Source/Runtime/Engine/Classes/Components/SplineComponent.h
@@ -608,10 +608,14 @@ public:
 	UFUNCTION(BlueprintCallable, Category = Spline)
 	ENGINE_API int32 GetNumberOfSplineSegments() const;
 
-	/** Get the input key (e.g. the time) of the control point of the spline at the specified index. */
+	/** Get the input key of the control point of the spline at the specified index. */
 	UFUNCTION(BlueprintCallable, Category = Spline)
 	ENGINE_API float GetInputKeyValueAtSplinePoint(int32 PointIndex) const;
 
+	/** Set the input key of the control point of the spline at the specified index. */
+	UFUNCTION(BlueprintCallable, Category = Spline)
+	ENGINE_API void SetInputKeyValueAtSplinePoint(int32 PointIndex, float InputKey, bool bUpdateSpline = true);
+
 	/** Gets the spline point of the spline at the specified index */
 	UFUNCTION(BlueprintCallable, Category = Spline)
 	ENGINE_API FSplinePoint GetSplinePointAt(int32 PointIndex, ESplineCoordinateSpace::Type CoordinateSpace) const;
diff --git a/Engine/Source/Runtime/Engine/Classes/Curves/Splines/TangentSpline.h b/Engine/Source/Runtime/Engine/Classes/Curves/Splines/TangentSpline.h
--- a/Engine/Source/Runtime/Engine/Classes/Curves/Splines/TangentSpline.h
+++ b/Engine/Source/Runtime/Engine/Classes/Curves/Splines/TangentSpline.h
@@ -53,11 +53,6 @@ public:
 		return Ar;
 	}
 
-	/**
-	 * Returns the parameter space of FTangentSpline. All public functions taking a parameter are only defined when the provided value is in this interval.
-	 * @return [0, NumSegments] when NumSegment != 0, otherwise undefined (empty interval).
-	 */
-	UE_API UE::Geometry::FInterval1f GetSegmentSpace() const;
 	UE_API float GetSegmentLength(const int32 Index, const float Param, const FVector& Scale3D = FVector(1.0f)) const;
 	UE_API float GetSplineLength();
 
@@ -92,13 +87,40 @@ public:
 	UE_API void SetSplinePointType(int32 Index, EInterpCurveMode Type);
 	UE_API EInterpCurveMode GetSplinePointType(int32 Index) const;
 
+	// Sets the parameters of all spline points according to the default policy (defined by the CVar 'TangentSpline.ParameterizationPolicy')
+	UE_API void Reparameterize();
+
+	// Sets the parameters of all spline points according to the specified policy.
+	UE_API void Reparameterize(UE::Geometry::Spline::EParameterizationPolicy ParameterizationPolicy);
+
+	/**
+	 * Sets the parameter of a spline point.
+	 * Currently, this function disallows knot vector reordering, meaning the provided parameter
+	 * must not be greater than the parameter at Index + 1 or less than the parameter at Index - 1.
+	 * 
+	 * @param Index The index of the point to update.
+	 * @param Parameter The new parameter.
+	 * @return true if the operation was successful, otherwise false.
+	 */
+	UE_API bool SetParameterAtIndex(int32 Index, float Parameter);
+
+	/**
+	 * Sets the parameters of all spline points.
+	 * InParameters must have exactly as many points as this spline.
+	 *
+	 * @param InParameters The list of new parameter values, 1:1 with the number of spline points.
+	 * @return true if the operation was successful, otherwise false.
+	 */
+	UE_API bool SetParameters(const TArray<float>& InParameters);
+
 	UE_API float GetParameterAtIndex(int32 Index) const;
 	UE_API float GetParameterAtDistance(float Distance) const;
 	UE_API float GetDistanceAtParameter(float Parameter) const;
+	UE_API int32 GetIndexAtParameter(float Parameter) const;
 
-	ENGINE_API FQuat GetOrientation(int32 Index) const;
-	ENGINE_API FQuat GetOrientation(float Param) const;
-	ENGINE_API void SetOrientation(int32 Index, const FQuat& InOrientation);
+	UE_API FQuat GetOrientation(int32 Index) const;
+	UE_API FQuat GetOrientation(float Param) const;
+	UE_API void SetOrientation(int32 Index, const FQuat& InOrientation);
 
 	UE_API virtual float FindNearest(const FVector& InLocation, float& OutSquaredDistance) const override;
 	UE_API float FindNearestOnSegment(const FVector& InLocation, int32 SegmentIndex, float& OutSquaredDistance) const;
@@ -207,13 +229,7 @@ private:
 
 private:
 
-	void Reparameterize();
-
-	/** Convert from internal spline space to [0, NumSegments] */
-	UE_API float FromInternalSplineSpace(float Parameter) const;
-
-	/** Convert from [0, NumSegments] to internal spline space */
-	UE_API float ToInternalSplineSpace(float Parameter) const;
+	void UpdateTangents();
 
 	void MarkReparamTableDirty();
 
@@ -255,7 +271,7 @@ private:
 template <typename AttrType>
 AttrType FTangentSpline::EvaluateAttribute(const FName& Name, float Parameter) const
 {
-	return UE::Geometry::Spline::TMultiSpline<UE::Geometry::Spline::FTangentBezierSpline3d>::EvaluateAttribute<AttrType>(Name, ToInternalSplineSpace(Parameter));
+	return UE::Geometry::Spline::TMultiSpline<UE::Geometry::Spline::FTangentBezierSpline3d>::EvaluateAttribute<AttrType>(Name, Parameter);
 }
 
 template <typename AttrType>
@@ -270,7 +286,7 @@ int32 FTangentSpline::AddAttributeValue(FName Name, const AttrType& Value, float
 		ControlPointType ControlPoint = ControlPointType(Value);
 		int32 NumPoints = Channel->GetNumPoints();
 
-		float ParentSpaceParameter = ToInternalSplineSpace(Parameter);
+		float ParentSpaceParameter = Parameter;
 		FInterval1f MappedChildSpace = GetMappedChildSpace(Name);
 
 		// Special case for empty spline
@@ -355,13 +371,13 @@ void FTangentSpline::RemoveAttributeValue(FName Name, int32 Index)
 			if (Index == 0)
 			{
 				FInterval1f NewMappingRange = GetMappedChildSpace(Name);
-				NewMappingRange.Min = ToInternalSplineSpace(GetAttributeParameter<AttrType>(Name, Index + 1));
+				NewMappingRange.Min = GetAttributeParameter<AttrType>(Name, Index + 1);
 				SetAttributeChannelRange(Name, NewMappingRange, EMappingRangeSpace::Parent);
 			}
 			else if (Index == Channel->GetNumPoints() - 1)
 			{
 				FInterval1f NewMappingRange = GetMappedChildSpace(Name);
-				NewMappingRange.Max = ToInternalSplineSpace(GetAttributeParameter<AttrType>(Name, Index - 1));
+				NewMappingRange.Max = GetAttributeParameter<AttrType>(Name, Index - 1);
 				SetAttributeChannelRange(Name, NewMappingRange, EMappingRangeSpace::Parent);
 			}
 		}
@@ -404,7 +420,7 @@ int32 FTangentSpline::SetAttributeParameter(const FName& Name, int32 Index, floa
 
 	if (AttrSplineType* Channel = GetTypedAttributeChannel<AttrSplineType>(Name))
 	{
-		float ParentSpaceParameter = ToInternalSplineSpace(Parameter);
+		float ParentSpaceParameter = Parameter;
 
 		// Prevent collapse of the channel space and keep it starting at 0.
 		auto SanitizeChannelSpace = [Channel]()
@@ -431,7 +447,7 @@ int32 FTangentSpline::SetAttributeParameter(const FName& Name, int32 Index, floa
 		// Uses current value of Index and ParentSpaceParameter at the time of calling to update the attribute parameter.
 		auto SetParentSpaceParameter = [this, &ParentSpaceParameter, Name, Channel, &Index]() -> int32
 			{
-				const float CurrentChildSpaceParameter = MapParameterToChildSpace(Name, ToInternalSplineSpace(GetAttributeParameter<AttrType>(Name, Index)));
+				const float CurrentChildSpaceParameter = MapParameterToChildSpace(Name, GetAttributeParameter<AttrType>(Name, Index));
 				const float DesiredChildSpaceParameter = MapParameterToChildSpace(Name, ParentSpaceParameter);
 
 				constexpr float MinStep = 2.f * UE_KINDA_SMALL_NUMBER;
@@ -482,7 +498,7 @@ int32 FTangentSpline::SetAttributeParameter(const FName& Name, int32 Index, floa
 			TArray<float> InternalParameters = CacheChannelKnots<float>(Name);
 
 			float ParentSpaceUpperBound = GetMappedChildSpace(Name).Max;
-			const float NeighborParentSpaceParameter = ToInternalSplineSpace(GetAttributeParameter<AttrType>(Name, Index + 1));
+			const float NeighborParentSpaceParameter = GetAttributeParameter<AttrType>(Name, Index + 1);
 			const bool bAttributesWillReorder = ParentSpaceParameter > NeighborParentSpaceParameter;
 
 			// If the end point is passing its neighbor, we need to actually reshuffle the points.
@@ -531,7 +547,7 @@ int32 FTangentSpline::SetAttributeParameter(const FName& Name, int32 Index, floa
 			TArray<float> InternalParameters = CacheChannelKnots<float>(Name);
 
 			float ParentSpaceLowerBound = GetMappedChildSpace(Name).Min;
-			const float NeighborParentSpaceParameter = ToInternalSplineSpace(GetAttributeParameter<AttrType>(Name, Index - 1));
+			const float NeighborParentSpaceParameter = GetAttributeParameter<AttrType>(Name, Index - 1);
 			const bool bAttributesWillReorder = ParentSpaceParameter < NeighborParentSpaceParameter;
 
 			// If the end point is passing its neighbor, we need to actually reshuffle the points.
@@ -606,7 +622,7 @@ float FTangentSpline::GetAttributeParameter(const FName& Name, int32 Index) cons
 
 	if (AttrSplineType* Channel = GetTypedAttributeChannel<AttrSplineType>(Name))
 	{
-		return FromInternalSplineSpace(MapParameterFromChildSpace(Name, Channel->GetParameter(Index)));
+		return MapParameterFromChildSpace(Name, Channel->GetParameter(Index));
 	}
 	return 0.0f;
 }
diff --git a/Engine/Source/Runtime/Engine/Private/Components/SplineComponent.cpp b/Engine/Source/Runtime/Engine/Private/Components/SplineComponent.cpp
--- a/Engine/Source/Runtime/Engine/Private/Components/SplineComponent.cpp
+++ b/Engine/Source/Runtime/Engine/Private/Components/SplineComponent.cpp
@@ -1,14 +1,11 @@
 // Copyright Epic Games, Inc. All Rights Reserved.
 
-/*=============================================================================
-	Spline.cpp
-=============================================================================*/
-
 #include "Components/SplineComponent.h"
 #include "Curves/Splines/TangentSpline.h"
 #include "Engine/Engine.h"
 #include "SceneView.h"
 #include "UObject/EditorObjectVersion.h"
+#include "UObject/FortniteMainBranchObjectVersion.h"
 #include "Math/RotationMatrix.h"
 #include "MeshElementCollector.h"
 #include "PrimitiveViewRelevance.h"
@@ -187,11 +184,19 @@ bool USplineComponent::Validate() const
 		bValid = false;
 	}
 
+	auto IsNearlyEqual = [](const double& A, const double& B, float RelTol, float AbsTol)
+	{
+		const float Tolerance = FMath::Max(AbsTol, RelTol * FMath::Max(FMath::Abs(A), FMath::Abs(B)));
+		return FMath::IsNearlyEqual(A, B, Tolerance);
+	};
+
 	const float SplineCurvesLength = WarninglessSplineCurves().GetSplineLength();
 	const float SplineLength = TangentSpline->GetDistanceAtParameter(TangentSpline->GetParameterSpace().Max);
-	if (!FMath::IsNearlyEqual(SplineCurvesLength, SplineLength))
+	if (!IsNearlyEqual(SplineCurvesLength, SplineLength, UE_KINDA_SMALL_NUMBER, UE_KINDA_SMALL_NUMBER))
 	{
 		UE_LOG(LogSplineComponent, Display, TEXT("Internally inconsistent spline lengths (%f vs %f)!"), SplineCurvesLength, SplineLength);
+		UE_LOG(LogSplineComponent, Display, TEXT("Scale: %s"), *GetComponentTransform().GetScale3D().ToCompactString());
+		
 		bValid = false;
 	}
 	UE_LOG(LogSplineComponent, Display, TEXT("Spline lengths: (%f vs %f)"), SplineCurvesLength, SplineLength);
@@ -202,6 +207,14 @@ bool USplineComponent::Validate() const
 		const auto& SplineCurvesRotationPoint = WarninglessSplineCurves().Rotation.Points[Idx];
 		const auto& SplineCurvesScalePoint = WarninglessSplineCurves().Scale.Points[Idx];
 
+		const float SplineCurvesParameter = SplineCurvesPositionPoint.InVal;
+		const float SplineParameter = TangentSpline->GetParameterAtIndex(Idx);
+		if (!IsNearlyEqual(SplineCurvesParameter, SplineParameter, UE_SMALL_NUMBER, UE_SMALL_NUMBER))
+		{
+			UE_LOG(LogSplineComponent, Display, TEXT("Internally inconsistent knot (%f vs %f) at index %d!"), SplineCurvesParameter, SplineParameter, Idx );
+			bValid = false;
+		}
+
 		const ESplinePointType::Type SplineCurvesPointType = ConvertInterpCurveModeToSplinePointType(SplineCurvesPositionPoint.InterpMode);
 		const ESplinePointType::Type SplinePointType = ConvertInterpCurveModeToSplinePointType(TangentSpline->GetSplinePointType(Idx));
 		if (SplineCurvesPointType != SplinePointType)
@@ -225,13 +238,19 @@ bool USplineComponent::Validate() const
 			UE_LOG(LogSplineComponent, Display, TEXT("Internally inconsistent arrive tangent (%s vs %s) at index %d!"), *SplineCurvesInTangent.ToCompactString(), *SplineInTangent.ToCompactString(), Idx );
 			UE_LOG(LogSplineComponent, Display, TEXT("Point types: %s and %s"), StringifySplinePointType(SplineCurvesPointType), StringifySplinePointType(SplinePointType));
 			
-			if (SplineCurvesPointType != ESplinePointType::CurveClamped || !UE::SplineComponent::GIgnoreCurveClampedPointTangentsForValidation)
+			bool bWasValid = bValid;
+			bValid = false;
+
+			if (SplineCurvesPointType == ESplinePointType::CurveClamped && UE::SplineComponent::GIgnoreCurveClampedPointTangentsForValidation)
 			{
-				bValid = false;
+				bValid = bWasValid;
+				UE_LOG(LogSplineComponent, Display, TEXT("Ignoring inconsistent arrive tangent for CurveClamped point."));
 			}
-			else
+
+			if (!IsClosedLoop() && Idx == 0 /* && UE::SplineComponent::GIgnoreEndpointInTangentMismatchForValidation */)
 			{
-				UE_LOG(LogSplineComponent, Display, TEXT("Ignoring inconsistent arrive tangent for CurveClamped point."));
+				bValid = bWasValid;
+				UE_LOG(LogSplineComponent, Display, TEXT("Ignoring inconsistent arrive tangent for first point on open spline."));
 			}
 		}
 
@@ -242,13 +261,19 @@ bool USplineComponent::Validate() const
 			UE_LOG(LogSplineComponent, Display, TEXT("Internally inconsistent leave tangent (%s vs %s) at index %d!"), *SplineCurvesOutTangent.ToCompactString(), *SplineOutTangent.ToCompactString(), Idx );
 			UE_LOG(LogSplineComponent, Display, TEXT("Point types: %s and %s"), StringifySplinePointType(SplineCurvesPointType), StringifySplinePointType(SplinePointType));
 			
-			if (SplineCurvesPointType != ESplinePointType::CurveClamped || !UE::SplineComponent::GIgnoreCurveClampedPointTangentsForValidation)
+			bool bWasValid = bValid;
+			bValid = false;
+
+			if (SplineCurvesPointType == ESplinePointType::CurveClamped && UE::SplineComponent::GIgnoreCurveClampedPointTangentsForValidation)
 			{
-				bValid = false;
+				bValid = bWasValid;
+				UE_LOG(LogSplineComponent, Display, TEXT("Ignoring inconsistent leave tangent for CurveClamped point."));
 			}
-			else
+
+			if (!IsClosedLoop() && Idx == GetNumberOfSplinePoints() - 1 /* && UE::SplineComponent::GIgnoreEndpointOutTangentMismatchForValidation */)
 			{
-				UE_LOG(LogSplineComponent, Display, TEXT("Ignoring inconsistent leave tangent for CurveClamped point."));
+				bValid = bWasValid;
+				UE_LOG(LogSplineComponent, Display, TEXT("Ignoring inconsistent leave tangent for last point on open spline."));
 			}
 		}
 
@@ -486,6 +511,7 @@ void USplineComponent::Serialize(FArchive& Ar)
 	Super::Serialize(Ar);
 
 	Ar.UsingCustomVersion(FEditorObjectVersion::GUID);
+	Ar.UsingCustomVersion(FFortniteMainBranchObjectVersion::GUID);
 
 #if WITH_EDITORONLY_DATA
 	if (Ar.IsLoading())
@@ -499,43 +525,83 @@ void USplineComponent::Serialize(FArchive& Ar)
 			WarninglessSplineCurves().ReparamTable = SplineReparamTable_DEPRECATED;
 		}
 
+		// Used to prevent redundant work.
+		bool bHasSynchronized = false;
+
 #if WITH_EDITOR
 		// If SplineComponent.UseSplineCurves is false and Spline was not previously implemented
 		// (implying it is not consistent with SplineCurves), initialize Spline from SplineCurves.
 		if (Spline.GetPreviousImplementation() == UE::Spline::ESplineType::Unimplemented && !UE::SplineComponent::GUseSplineCurves)
 		{
 			Spline = WarninglessSplineCurves();
+			bHasSynchronized = true;
 		}
 #endif
-		
-		if (UE::SplineComponent::GSynchronizeOnLoad && Spline.GetCurrentImplementation() == UE::Spline::ESplineType::Tangent && !Validate())
+
+		auto Synchronize = [this, &bHasSynchronized]() -> void
 		{
-			UE_LOG(LogSplineComponent, Display, TEXT("Internally inconsistent SplineComponent (%s) at load time!"), *GetPathName())
-			
 			switch (UE::SplineComponent::GetEffectiveLastAuthority(LastAuthority))
 			{
 			case ELastAuthority::Unset:	// Intentional fallthrough - at the time this property was added, the authority had been Spline for some time.
 			case ELastAuthority::Spline:
 				UE_LOG(LogSplineComponent, Display, TEXT("Populating SplineCurves from Spline..."))
-				{
-					FSplineCurves TemporarySplineCurves = FSplineCurves::FromSplineInterface(Spline);
-					for (int32 Idx = 0; Idx < WarninglessSplineCurves().Position.Points.Num(); ++Idx)
-					{
-						TemporarySplineCurves.Position.Points[Idx].InVal = WarninglessSplineCurves().Position.Points[Idx].InVal;
-						TemporarySplineCurves.Rotation.Points[Idx].InVal = WarninglessSplineCurves().Rotation.Points[Idx].InVal;
-						TemporarySplineCurves.Scale.Points[Idx].InVal = WarninglessSplineCurves().Scale.Points[Idx].InVal;
-					}
-					WarninglessSplineCurves() = TemporarySplineCurves;
-				}
+					WarninglessSplineCurves() = FSplineCurves::FromSplineInterface(Spline);
 				break;
 
 			case ELastAuthority::SplineCurves:
 				UE_LOG(LogSplineComponent, Display, TEXT("Populating Spline from SplineCurves..."))
-				SetSpline(WarninglessSplineCurves());
+					SetSpline(WarninglessSplineCurves());
 				break;
 			}
 
 			UpdateSpline();
+
+			bHasSynchronized = true;
+		};
+
+		// Previously, FTangentSpline's knot vector was always uniform, regardless of the input keys in SplineCurves.
+		// Once our input key interface began propagating to the knot vector, it was necessary to synchronize them on load.
+		if (!bHasSynchronized
+			&& Ar.CustomVer(FFortniteMainBranchObjectVersion::GUID) < FFortniteMainBranchObjectVersion::SplineComponentReparameterizeOnLoad)
+		{
+			if (FTangentSpline* TangentSpline = Spline.Get<UE::Spline::ESplineType::Tangent>())
+			{
+				const int32 NumTangentSplinePoints = TangentSpline->GetNumberOfPoints();
+				const int32 NumSplineCurvesPoints = WarninglessSplineCurves().Position.Points.Num();
+
+				if (NumTangentSplinePoints != NumSplineCurvesPoints)
+				{
+					UE_LOG(LogSplineComponent, Display, TEXT("Mismatched number of spline points, uniformly parameterizing TangentSpline."));
+					UE_LOG(LogSplineComponent, Display, TEXT("Path: %s"), *GetPathName());
+
+					TangentSpline->Reparameterize(UE::Geometry::Spline::EParameterizationPolicy::Uniform);
+				}
+				else
+				{
+					TArray<float> SplineCurvesParameters;
+
+					for (int32 SplineCurvesIdx = 0; SplineCurvesIdx < NumSplineCurvesPoints; ++SplineCurvesIdx)
+					{
+						SplineCurvesParameters.Add(WarninglessSplineCurves().Position.Points[SplineCurvesIdx].InVal);
+					}
+
+					TangentSpline->SetParameters(SplineCurvesParameters);
+
+					// Changing parameterization invalidates loaded auto-tangents.
+					FTangentSpline::FUpdateSplineParams Params{bClosedLoop, bStationaryEndpoints, ReparamStepsPerSegment, bLoopPositionOverride, LoopPosition, GetComponentTransform().GetScale3D()};
+					TangentSpline->UpdateSpline(Params);
+				}
+			}
+		}
+		
+		if (!bHasSynchronized
+			&& UE::SplineComponent::GSynchronizeOnLoad
+			&& Spline.GetCurrentImplementation() == UE::Spline::ESplineType::Tangent
+			&& !Validate())
+		{
+			UE_LOG(LogSplineComponent, Display, TEXT("Internally inconsistent SplineComponent (%s) at load time!"), *GetPathName())
+
+			Synchronize();
 			
 			// see if we got it right...
 			if (!Validate())
@@ -549,7 +615,7 @@ void USplineComponent::Serialize(FArchive& Ar)
 		}
 		else if (UE::SplineComponent::GValidateOnLoad && Spline.GetCurrentImplementation() == UE::Spline::ESplineType::Tangent && !Validate())
 		{
-			UE_LOG(LogSplineComponent, Display, TEXT("Internally inconsistent spline data at load time!"))
+			UE_LOG(LogSplineComponent, Warning, TEXT("Internally inconsistent spline data at load time!"))
 		}
 	}
 #endif
@@ -701,7 +767,7 @@ void FSplineCurves::UpdateSpline(bool bClosedLoop, bool bStationaryEndpoints, in
 	Scale.AutoSetTangents(0.0f, bStationaryEndpoints);
 
 	// Now initialize the spline reparam table
-	const int32 NumSegments = bClosedLoop ? NumPoints : FMath::Max(0, NumPoints - 1);
+	const int32 NumSegments = Position.bIsLooped ? NumPoints : FMath::Max(0, NumPoints - 1);
 
 	// Start by clearing it
 	ReparamTable.Points.Reset(NumSegments * ReparamStepsPerSegment + 1);
@@ -711,11 +777,11 @@ void FSplineCurves::UpdateSpline(bool bClosedLoop, bool bStationaryEndpoints, in
 		for (int32 Step = 0; Step < ReparamStepsPerSegment; ++Step)
 		{
 			const float Param = static_cast<float>(Step) / ReparamStepsPerSegment;
-			const float SegmentLength = (Step == 0) ? 0.0f : GetSegmentLength(SegmentIndex, Param, bClosedLoop, Scale3D);
+			const float SegmentLength = (Step == 0) ? 0.0f : GetSegmentLength(SegmentIndex, Param, Position.bIsLooped, Scale3D);
 
 			ReparamTable.Points.Emplace(SegmentLength + AccumulatedLength, SegmentIndex + Param, 0.0f, 0.0f, CIM_Linear);
 		}
-		AccumulatedLength += GetSegmentLength(SegmentIndex, 1.0f, bClosedLoop, Scale3D);
+		AccumulatedLength += GetSegmentLength(SegmentIndex, 1.0f, Position.bIsLooped, Scale3D);
 	}
 	ReparamTable.Points.Emplace(AccumulatedLength, static_cast<float>(NumSegments), 0.0f, 0.0f, CIM_Linear);
 	++Version;
@@ -1511,17 +1577,15 @@ void USplineComponent::AddPoint(const FSplinePoint& InSplinePoint, bool bUpdateS
 {
 	int32 Index = UpperBound(WarninglessSplineCurves().Position.Points, InSplinePoint.InputKey);
 		
-	if (Index > 0 && WarninglessSplineCurves().Position.Points[Index - 1].InVal == InSplinePoint.InputKey)
+	if (Index > 0 && FMath::IsNearlyEqual(InSplinePoint.InputKey, GetInputKeyValueAtSplinePoint(Index - 1)))
 	{
 		// Decrement in the case of a collision because we insert before the colliding element and increment that element's input key.
 		Index--;
 
 		// Increment keys only in the case of a collision to prevent invalid input keys.
-		for (int I = Index; I < WarninglessSplineCurves().Position.Points.Num(); ++I)
+		for (int I = GetNumberOfSplinePoints() - 1; I >= Index; --I)
 		{
-			WarninglessSplineCurves().Position.Points[I].InVal += 1.0f;
-			WarninglessSplineCurves().Rotation.Points[I].InVal += 1.0f;
-			WarninglessSplineCurves().Scale.Points[I].InVal += 1.0f;
+			SetInputKeyValueAtSplinePoint(I, GetInputKeyValueAtSplinePoint(I) + 1.f);
 		}
 	}
 		
@@ -1682,7 +1746,22 @@ void USplineComponent::AddSplinePointAtIndex(const FVector& Position, int32 Inde
 		NewPoint.Rotation = FQuat::Identity.Rotator();
 		NewPoint.Scale = FVector(1.f);
 
-		if (FTangentSpline* TangentSpline = Spline.Get<UE::Spline::ESplineType::Tangent>())
+		FTangentSpline* TangentSpline = Spline.Get<UE::Spline::ESplineType::Tangent>();
+
+		// Adjust subsequent points' input keys to make room for the value we are about to add
+		for (int I = NumPoints - 1; I >= Index; --I)
+		{
+			if (TangentSpline)
+			{
+				TangentSpline->SetParameterAtIndex(I, TangentSpline->GetParameterAtIndex(I) + 1.0f);
+			}
+
+			WarninglessSplineCurves().Position.Points[I].InVal += 1.0f;
+			WarninglessSplineCurves().Rotation.Points[I].InVal += 1.0f;
+			WarninglessSplineCurves().Scale.Points[I].InVal += 1.0f;
+		}
+
+		if (TangentSpline)
 		{
 			TangentSpline->InsertPoint(NewPoint, Index);
 		}
@@ -1695,15 +1774,6 @@ void USplineComponent::AddSplinePointAtIndex(const FVector& Position, int32 Inde
 		{
 			Metadata->InsertPoint(Index, 0.5f, bClosedLoop);
 		}
-		
-		// Adjust subsequent points' input keys to make room for the value just added
-		int32 WarninglessSplinePointsNum = WarninglessSplineCurves().Position.Points.Num();
-		for (int I = Index + 1; I < WarninglessSplinePointsNum; ++I)
-		{
-			WarninglessSplineCurves().Position.Points[I].InVal += 1.0f;
-			WarninglessSplineCurves().Rotation.Points[I].InVal += 1.0f;
-			WarninglessSplineCurves().Scale.Points[I].InVal += 1.0f;
-		}
 
 		if (bLoopPositionOverride)
 		{
@@ -1730,19 +1800,33 @@ void USplineComponent::RemoveSplinePoint(int32 Index, bool bUpdateSpline)
 
 	if (Index >= 0 && Index < NumPoints)
 	{
-		if (FTangentSpline* TangentSpline = Spline.Get<UE::Spline::ESplineType::Tangent>())
-		{
-			TangentSpline->RemovePoint(Index);
-		}
+		FTangentSpline* TangentSpline = Spline.Get<UE::Spline::ESplineType::Tangent>();
+
+		float SplineCurvesKnotDelta = 0.0f;
+		float SplineKnotDelta = 0.f;
 
-		float KnotDelta = 0.0f;
 		if (Index > 0)
 		{
-			KnotDelta = WarninglessSplineCurves().Position.Points[Index].InVal - WarninglessSplineCurves().Position.Points[Index - 1].InVal;
+			SplineCurvesKnotDelta = WarninglessSplineCurves().Position.Points[Index].InVal - WarninglessSplineCurves().Position.Points[Index - 1].InVal;
+
+			if (TangentSpline)
+			{
+				SplineKnotDelta = TangentSpline->GetParameterAtIndex(Index) - TangentSpline->GetParameterAtIndex(Index - 1);
+			}
 		}
 		else if (NumPoints >= 2) // e.g. Index == 0, Index < NumPoints - 1
 		{
-			KnotDelta = WarninglessSplineCurves().Position.Points[Index + 1].InVal - WarninglessSplineCurves().Position.Points[Index].InVal;
+			SplineCurvesKnotDelta = WarninglessSplineCurves().Position.Points[Index + 1].InVal - WarninglessSplineCurves().Position.Points[Index].InVal;
+
+			if (TangentSpline)
+			{
+				SplineKnotDelta = TangentSpline->GetParameterAtIndex(Index + 1) - TangentSpline->GetParameterAtIndex(Index);
+			}
+		}
+
+		if (TangentSpline)
+		{
+			TangentSpline->RemovePoint(Index);
 		}
 		
 		WarninglessSplineCurves().Position.Points.RemoveAt(Index, EAllowShrinking::No);
@@ -1759,9 +1843,15 @@ void USplineComponent::RemoveSplinePoint(int32 Index, bool bUpdateSpline)
 		// Adjust all following spline point input keys to close the gap left by the removed point
 		while (Index < NumPoints)
 		{
-			WarninglessSplineCurves().Position.Points[Index].InVal -= KnotDelta;
-			WarninglessSplineCurves().Rotation.Points[Index].InVal -= KnotDelta;
-			WarninglessSplineCurves().Scale.Points[Index].InVal -= KnotDelta;
+			if (TangentSpline)
+			{
+				TangentSpline->SetParameterAtIndex(Index, TangentSpline->GetParameterAtIndex(Index) - SplineKnotDelta);
+			}
+
+			WarninglessSplineCurves().Position.Points[Index].InVal -= SplineCurvesKnotDelta;
+			WarninglessSplineCurves().Rotation.Points[Index].InVal -= SplineCurvesKnotDelta;
+			WarninglessSplineCurves().Scale.Points[Index].InVal -= SplineCurvesKnotDelta;
+
 			Index++;
 		}
 
@@ -2095,23 +2185,53 @@ int32 USplineComponent::GetNumberOfSplinePoints() const
 int32 USplineComponent::GetNumberOfSplineSegments() const
 {
 	const int32 NumPoints = GetNumberOfSplinePoints();
-	return (bClosedLoop ? NumPoints : FMath::Max(0, NumPoints - 1));
+
+	if (NumPoints < 2)
+	{
+		return 0;
+	}
+
+	// We don't just simply read bClosedLoop because if LoopPosition holds an
+	// invalid value the spline will be open regardless of the value of bClosedLoop.
+
+	if (ShouldUseSplineCurves())
+	{
+		return WarninglessSplineCurves().Position.bIsLooped ? NumPoints : FMath::Max(0, NumPoints - 1);
+	}
+	if (const FTangentSpline* TangentSpline = Spline.Get<UE::Spline::ESplineType::Tangent>())
+	{
+		return TangentSpline->GetNumberOfSegments();
+	}
+	else
+	{
+		ensureAlwaysMsgf(false, TEXT("ShouldUseSplineCurves() returned false but Spline failed to yield an FTangentSpline!"));
+		return 0;
+	}
+
 }
 
 float USplineComponent::GetInputKeyValueAtSplinePoint(int32 PointIndex) const
 {
-	if (GetNumberOfSplinePoints() == 0)
+	const int32 NumPoints = GetNumberOfSplinePoints();
+
+	if (NumPoints == 0)
 	{
 		return 0.f;
 	}
 
 	// Special case if we are closed and PointIndex refers to the closing point. For this function only we
 	// allow users to treat it as a separate point in order to expose the parameterization of the closing segment.
-
-	if (bClosedLoop && PointIndex == GetNumberOfSplinePoints())
+	// The closing segment only exists if we have at least 2 points (otherwise there is no curve to close).
+	if (bClosedLoop && NumPoints >= 2 && PointIndex == NumPoints)
 	{
 		if (ShouldUseSplineCurves())
 		{
+			if (!WarninglessSplineCurves().Position.bIsLooped)
+			{
+				// This is a weird case that can occur when LoopPosition is invalid
+				return 0.f;
+			}
+
 			return WarninglessSplineCurves().Position.Points[PointIndex - 1].InVal + WarninglessSplineCurves().Position.LoopKeyOffset;
 		}
 		else if (const FTangentSpline* TangentSpline = Spline.Get<UE::Spline::ESplineType::Tangent>())
@@ -2144,6 +2264,40 @@ float USplineComponent::GetInputKeyValueAtSplinePoint(int32 PointIndex) const
 	}
 }
 
+void USplineComponent::SetInputKeyValueAtSplinePoint(int32 PointIndex, float InputKey, bool bUpdateSpline)
+{
+	const int32 NumPoints = GetNumberOfSplinePoints();
+
+	if (NumPoints == 0)
+	{
+		return;
+	}
+
+	PointIndex = GetClampedIndex(PointIndex);
+
+	// Prevents reordering the knot vector. If the knot vector is monotonic, we only need to check the immediate neighbors of the point.
+	if ((PointIndex + 1 < NumPoints && InputKey >= GetInputKeyValueAtSplinePoint(PointIndex + 1)) ||
+		(PointIndex - 1 >= 0 && InputKey <= GetInputKeyValueAtSplinePoint(PointIndex - 1)))
+	{
+		UE_LOG(LogSplineComponent, Warning, TEXT("Input Key change would reorder the spline, ignoring!"));
+		return;
+	}
+
+	WarninglessSplineCurves().Position.Points[PointIndex].InVal = InputKey;
+	WarninglessSplineCurves().Rotation.Points[PointIndex].InVal = InputKey;
+	WarninglessSplineCurves().Scale.Points[PointIndex].InVal = InputKey;
+
+	if (FTangentSpline* TangentSpline = Spline.Get<UE::Spline::ESplineType::Tangent>())
+	{
+		TangentSpline->SetParameterAtIndex(PointIndex, InputKey);
+	}
+
+	if (bUpdateSpline)
+	{
+		UpdateSpline();
+	}
+}
+
 FSplinePoint USplineComponent::GetSplinePointAt(int32 PointIndex, ESplineCoordinateSpace::Type CoordinateSpace) const
 {
 	PointIndex = GetClampedIndex(PointIndex);
@@ -3035,11 +3189,11 @@ bool USplineComponent::ConvertSplineToPolyline_InDistanceRange(ESplineCoordinate
 	}
 	else if (const FTangentSpline* TangentSpline = Spline.Get<UE::Spline::ESplineType::Tangent>())
 	{
-		const float StartParam = TangentSpline->GetParameterAtDistance(RangeStart);
-		const float EndParam = TangentSpline->GetParameterAtDistance(RangeEnd);
+		const float RangeStartParam = GetInputKeyValueAtDistanceAlongSpline(RangeStart);
+		const float RangeEndParam = GetInputKeyValueAtDistanceAlongSpline(RangeEnd);
 
-		SegmentStart = FMath::FloorToInt32(StartParam);
-		SegmentEnd = FMath::Min(NumSegments, FMath::CeilToInt32(EndParam));
+		SegmentStart = TangentSpline->GetIndexAtParameter(RangeStartParam);
+		SegmentEnd = FMath::Min(NumSegments, 1 + TangentSpline->GetIndexAtParameter(RangeEndParam));
 	}
 	else
 	{
@@ -3051,6 +3205,11 @@ bool USplineComponent::ConvertSplineToPolyline_InDistanceRange(ESplineCoordinate
 	TArray<double> NewDistances;
 	for (int32 SegmentIndex = SegmentStart; SegmentIndex < SegmentEnd; ++SegmentIndex)
 	{
+		if (SegmentIndex == INDEX_NONE)
+		{
+			continue;
+		}
+
 		// Get the segment range as distances, clipped with the input range
 		double StartDist = FMath::Max(RangeStart, GetDistanceAlongSplineAtSplinePoint(SegmentIndex));
 		double StopDist = FMath::Min(RangeEnd, GetDistanceAlongSplineAtSplinePoint(SegmentIndex + 1));
@@ -4087,11 +4246,15 @@ void USplineComponent::PostEditChangeChainProperty(FPropertyChangedChainEvent& P
 		static const FName StationaryEndpointsName = GET_MEMBER_NAME_CHECKED(USplineComponent, bStationaryEndpoints);
 		static const FName DefaultUpVectorName = GET_MEMBER_NAME_CHECKED(USplineComponent, DefaultUpVector);
 		static const FName ClosedLoopName = GET_MEMBER_NAME_CHECKED(USplineComponent, bClosedLoop);
+		static const FName LoopPositionOverrideName = GET_MEMBER_NAME_CHECKED(USplineComponent, bLoopPositionOverride);
+		static const FName LoopPositionName = GET_MEMBER_NAME_CHECKED(USplineComponent, LoopPosition);
 
 		if (PropertyName == ReparamStepsPerSegmentName ||
 			PropertyName == StationaryEndpointsName ||
 			PropertyName == DefaultUpVectorName ||
-			PropertyName == ClosedLoopName)
+			PropertyName == ClosedLoopName ||
+			PropertyName == LoopPositionOverrideName ||
+			PropertyName == LoopPositionName)
 		{
 			UpdateSpline();
 
diff --git a/Engine/Source/Runtime/Engine/Private/Curves/Splines/TangentSpline.cpp b/Engine/Source/Runtime/Engine/Private/Curves/Splines/TangentSpline.cpp
--- a/Engine/Source/Runtime/Engine/Private/Curves/Splines/TangentSpline.cpp
+++ b/Engine/Source/Runtime/Engine/Private/Curves/Splines/TangentSpline.cpp
@@ -153,22 +153,26 @@ FTangentSpline::FTangentSpline(const FLegacyTangentSpline& Other)
 	CreateAttributeChannel<UE::Geometry::Spline::TTangentBezierSpline<FVector>>(ScaleAttrName);
 
 	TArray<FTangentBezierControlPoint> Points;
+	TArray<float> Parameters;
 	for (int32 Idx = 0; Idx < Other.PositionCurve.Points.Num(); Idx++)
 	{
+		const FInterpCurvePoint<FVector>& Position = Other.PositionCurve.Points[Idx];
+
 		FTangentBezierControlPoint& Point = Points.AddDefaulted_GetRef();
+		Point.Position = Position.OutVal;
+		Point.TangentIn = Position.ArriveTangent;
+		Point.TangentOut = Position.LeaveTangent;
+		Point.TangentMode = ConvertInterpCurveModeToTangentMode(Position.InterpMode);
 
-		Point.Position = Other.PositionCurve.Points[Idx].OutVal;
-		Point.TangentIn = Other.PositionCurve.Points[Idx].ArriveTangent;
-		Point.TangentOut = Other.PositionCurve.Points[Idx].LeaveTangent;
-		Point.TangentMode = ConvertInterpCurveModeToTangentMode(Other.PositionCurve.Points[Idx].InterpMode);
+		Parameters.Add(Position.InVal);
 	}
 
 	// initialize spline
 	GetSpline().SetPoints(Points, UE::Geometry::Spline::EParameterizationPolicy::Uniform);
+	SetParameters(Parameters);
+
 	ResetRotation();
 	ResetScale();
-
-	// setup rotation/scale attributes
 	for (int32 Idx = 0; Idx < Other.PositionCurve.Points.Num(); Idx++)
 	{
 		if (Other.RotationCurve.Points.IsValidIndex(Idx))
@@ -181,12 +185,10 @@ FTangentSpline::FTangentSpline(const FLegacyTangentSpline& Other)
 			SetAttributeValue<FVector>(ScaleAttrName, Other.ScaleCurve.Points[Idx].OutVal, Idx);
 		}
 	}
-
 	ValidateRotScale();
 
 	FUpdateSplineParams UpdateParams;
 	UpdateParams.bClosedLoop = Other.PositionCurve.bIsLooped;
-
 	// If the source curves encode an override loop offset parameter, convert it to an absolute loop position
 	// The source curve is assumed to encode an override loop offset parameter if it
 	// 1) is closed,
@@ -207,24 +209,26 @@ FTangentSpline::FTangentSpline(const FSplineCurves& Other)
 	CreateAttributeChannel<UE::Geometry::Spline::TTangentBezierSpline<FVector>>(ScaleAttrName);
 
 	TArray<FTangentBezierControlPoint> Points;
+	TArray<float> Parameters;
 	for (int32 Idx = 0; Idx < Other.Position.Points.Num(); Idx++)
 	{
 		const FInterpCurvePoint<FVector>& Position = Other.Position.Points[Idx];
 
 		FTangentBezierControlPoint& Point = Points.AddDefaulted_GetRef();
-
 		Point.Position = Position.OutVal;
 		Point.TangentIn = Position.ArriveTangent;
 		Point.TangentOut = Position.LeaveTangent;
 		Point.TangentMode = ConvertInterpCurveModeToTangentMode(Position.InterpMode);
+
+		Parameters.Add(Position.InVal);
 	}
 
 	// initialize spline
 	GetSpline().SetPoints(Points, UE::Geometry::Spline::EParameterizationPolicy::Uniform);
+	SetParameters(Parameters);
+
 	ResetRotation();
 	ResetScale();
-
-	// setup rotation/scale attributes
 	for (int32 Idx = 0; Idx < Other.Position.Points.Num(); Idx++)
 	{
 		FQuat Rotation = FQuat::Identity;
@@ -250,12 +254,10 @@ FTangentSpline::FTangentSpline(const FSplineCurves& Other)
 		SetAttributeValue<FQuat>(RotationAttrName, Rotation, Idx);
 		SetAttributeValue<FVector>(ScaleAttrName, Scale, Idx);
 	}
-
 	ValidateRotScale();
 
 	FUpdateSplineParams UpdateParams;
 	UpdateParams.bClosedLoop = Other.Position.bIsLooped;
-
 	// If the source curves encode an override loop offset parameter, convert it to an absolute loop position
 	// The source curve is assumed to encode an override loop offset parameter if it
 	// 1) is closed,
@@ -420,8 +422,6 @@ bool FTangentSpline::Serialize(FArchive& Ar)
 
 	if (Ar.IsLoading())
 	{
-		Reparameterize();
-
 		// Do not serialize legacy curves, only generate if loading.
 		RebuildLegacyCurves();
 	}
@@ -429,11 +429,6 @@ bool FTangentSpline::Serialize(FArchive& Ar)
 	return true;
 }
 
-FInterval1f FTangentSpline::GetSegmentSpace() const
-{
-	return GetNumberOfSegments() == 0 ? FInterval1f::Empty() : FInterval1f(0.f, static_cast<float>(GetNumberOfSegments()));
-}
-
 float FTangentSpline::GetSegmentLength(const int32 Index, const float Param, const FVector& Scale3D) const
 {
 	const int32 NumPoints = GetNumberOfPoints();
@@ -498,7 +493,7 @@ float FTangentSpline::GetSegmentLength(const int32 Index, const float Param, con
 float FTangentSpline::GetSplineLength()
 {
 	// Evaluate the ParamToLength channel at the very end of the spline
-	const float MaxParameter = static_cast<float>(GetNumberOfSegments());
+	const float MaxParameter = GetParameterSpace().Max;
 	return GetDistanceAtParameter(MaxParameter);
 }
 
@@ -515,6 +510,7 @@ void FTangentSpline::SetClosedLoop(bool bInClosedLoop, bool bUpdateSpline)
 	}
 	TRACE_CPUPROFILER_EVENT_SCOPE(FTangentSpline::SetClosedLoop);
 
+	MarkReparamTableDirty();
 	MarkLegacyCurvesDirty();
 
 	GetSpline().SetClosedLoop(bInClosedLoop);
@@ -542,6 +538,9 @@ void FTangentSpline::Reset()
 	PositionCurve.Reset();
 	RotationCurve.Reset();
 	ScaleCurve.Reset();
+
+	MarkReparamTableDirty();
+	MarkLegacyCurvesDirty();
 }
 
 void FTangentSpline::ResetRotation()
@@ -558,6 +557,8 @@ void FTangentSpline::ResetRotation()
 	{
 		RotSpline->SetKnotVector(GetSpline().GetKnotVector());
 	}
+
+	MarkLegacyCurvesDirty();
 }
 
 void FTangentSpline::ResetScale()
@@ -574,6 +575,8 @@ void FTangentSpline::ResetScale()
 	{
 		ScaleSpline->SetKnotVector(GetSpline().GetKnotVector());
 	}
+
+	MarkLegacyCurvesDirty();
 }
 
 void FTangentSpline::AddPoint(const FSplinePoint& Point)
@@ -589,71 +592,81 @@ void FTangentSpline::AddPoint(const FSplinePoint& Point)
 	};
 
 	FTangentBezierControlPoint ControlPoint = ConvertToTangentBezierControlPoint(Point);
-	const int32 NumPoints = GetNumberOfPoints();
 
-	// Convert InputKey to parameter value
-	int32 Index = FMath::FloorToInt(Point.InputKey);
-	float Fraction = Point.InputKey - Index;
+	const FInterval1f ParameterSpace = GetParameterSpace();
+	const float& Parameter = Point.InputKey;
 
-	if (NumPoints == 0 || Index >= NumPoints)
+	if (ParameterSpace.IsEmpty() || Parameter > ParameterSpace.Max)
 	{
-		GetSpline().AppendPoint(ControlPoint, EParameterizationPolicy::Uniform);
+		/* Append Case */
+
+		GetSpline().AppendPoint(ControlPoint, Parameter);
+		const TArray<FKnot> NewKnotVector = GetSpline().GetKnotVector();
+
+		// We will append points to the scale and rotation splines now.
+		// We have to append with a policy because Parameter is not defined w.r.t. the child spline spaces.
+		// We cannot map it into the child space because the mapping functions do not extrapolate outside of the mapped to space.
+		// This is fine because we immediately propagate the primary knot vector which overrides whatever parameter is generated by the policy.
 
 		if (TTangentBezierSpline<FVector>* ScaleChild = GetTypedAttributeChannel<TTangentBezierSpline<FVector>>(ScaleAttrName))
 		{
 			ScaleChild->AppendPoint(TTangentBezierSpline<FVector>::FTangentBezierControlPoint(Point.Scale), EParameterizationPolicy::Uniform);
-			ScaleChild->SetKnotVector(GetSpline().GetKnotVector());
+			ScaleChild->SetKnotVector(NewKnotVector);
 		}
 
 		if (TTangentBezierSpline<FQuat>* RotChild = GetTypedAttributeChannel<TTangentBezierSpline<FQuat>>(RotationAttrName))
 		{
 			RotChild->AppendPoint(TTangentBezierSpline<FQuat>::FTangentBezierControlPoint(Point.Rotation.Quaternion()), EParameterizationPolicy::Uniform);
-			RotChild->SetKnotVector(GetSpline().GetKnotVector());
+			RotChild->SetKnotVector(NewKnotVector);
 		}
-
-		MarkLegacyCurvesDirty();
-		return;
 	}
-	else if (Index == 0)
+	else if (Parameter < ParameterSpace.Min)
 	{
-		GetSpline().PrependPoint(ControlPoint, EParameterizationPolicy::Uniform);
+		/* Prepend Case */
+
+		GetSpline().PrependPoint(ControlPoint, Parameter);
+		const TArray<FKnot> NewKnotVector = GetSpline().GetKnotVector();
+
+		// We will prepend points to the scale and rotation splines now.
+		// We have to prepend with a policy because Parameter is not defined w.r.t. the child spline spaces.
+		// We cannot map it into the child space because the mapping functions do not extrapolate outside of the mapped to space.
+		// This is fine because we immediately propagate the primary knot vector which overrides whatever parameter is generated by the policy.
 
 		if (TTangentBezierSpline<FVector>* ScaleChild = GetTypedAttributeChannel<TTangentBezierSpline<FVector>>(ScaleAttrName))
 		{
 			ScaleChild->PrependPoint(TTangentBezierSpline<FVector>::FTangentBezierControlPoint(Point.Scale), EParameterizationPolicy::Uniform);
-			ScaleChild->SetKnotVector(GetSpline().GetKnotVector());
+			ScaleChild->SetKnotVector(NewKnotVector);
 		}
 
 		if (TTangentBezierSpline<FQuat>* RotChild = GetTypedAttributeChannel<TTangentBezierSpline<FQuat>>(RotationAttrName))
 		{
 			RotChild->PrependPoint(TTangentBezierSpline<FQuat>::FTangentBezierControlPoint(Point.Rotation.Quaternion()), EParameterizationPolicy::Uniform);
-			RotChild->SetKnotVector(GetSpline().GetKnotVector());
+			RotChild->SetKnotVector(NewKnotVector);
 		}
-
-		MarkLegacyCurvesDirty();
-		return;
 	}
+	else
+	{
+		/* Insert Case */
 
-	// Insert at calculated parameter
-	const int32 SegmentIndex = Index - 1;
-	const float Parameter = GetSpline().GetSegmentParameterRange(SegmentIndex).Interpolate(Fraction);
-
-	int32 InsertedIndex = GetSpline().InsertPoint(ControlPoint, Parameter);
+		int32 InsertedIndex = GetSpline().InsertPoint(ConvertToTangentBezierControlPoint(Point), Parameter);
+		const TArray<FKnot> NewKnotVector = GetSpline().GetKnotVector();
 
-	float ScaleParam = MapParameterToChildSpace(ScaleAttrName, Parameter);
-	if (TTangentBezierSpline<FVector>* ScaleChild = GetTypedAttributeChannel<TTangentBezierSpline<FVector>>(ScaleAttrName))
-	{
-		ScaleChild->InsertPoint(TTangentBezierSpline<FVector>::FTangentBezierControlPoint(Point.Scale), ScaleParam);
-		ScaleChild->SetKnotVector(GetSpline().GetKnotVector());
-	}
+		if (TTangentBezierSpline<FVector>* ScaleChild = GetTypedAttributeChannel<TTangentBezierSpline<FVector>>(ScaleAttrName))
+		{
+			float ScaleParam = MapParameterToChildSpace(ScaleAttrName, Parameter);
+			ScaleChild->InsertPoint(TTangentBezierSpline<FVector>::FTangentBezierControlPoint(Point.Scale), ScaleParam);
+			ScaleChild->SetKnotVector(NewKnotVector);
+		}
 
-	float RotParam = MapParameterToChildSpace(RotationAttrName, Parameter);
-	if (TTangentBezierSpline<FQuat>* RotChild = GetTypedAttributeChannel<TTangentBezierSpline<FQuat>>(RotationAttrName))
-	{
-		RotChild->InsertPoint(TTangentBezierSpline<FQuat>::FTangentBezierControlPoint(Point.Rotation.Quaternion()), RotParam);
-		RotChild->SetKnotVector(GetSpline().GetKnotVector());
+		if (TTangentBezierSpline<FQuat>* RotChild = GetTypedAttributeChannel<TTangentBezierSpline<FQuat>>(RotationAttrName))
+		{
+			float RotParam = MapParameterToChildSpace(RotationAttrName, Parameter);
+			RotChild->InsertPoint(TTangentBezierSpline<FQuat>::FTangentBezierControlPoint(Point.Rotation.Quaternion()), RotParam);
+			RotChild->SetKnotVector(NewKnotVector);
+		}
 	}
 
+	MarkReparamTableDirty();
 	MarkLegacyCurvesDirty();
 }
 
@@ -661,6 +674,14 @@ void FTangentSpline::InsertPoint(const FSplinePoint& Point, int32 Index)
 {
 	using namespace UE::Geometry::Spline;
 
+	const int32 NumPoints = GetNumberOfPoints();
+
+	// We allow Index == NumPoints for appending.
+	if (Index < 0 || Index > NumPoints)
+	{
+		return;
+	}
+
 	ValidateRotScale();
 	ON_SCOPE_EXIT
 	{
@@ -670,71 +691,83 @@ void FTangentSpline::InsertPoint(const FSplinePoint& Point, int32 Index)
 	};
 
 	FTangentBezierControlPoint ControlPoint = ConvertToTangentBezierControlPoint(Point);
-	const int32 NumPoints = GetNumberOfPoints();
 
-	// Special case for empty spline
-	if (NumPoints == 0 || Index >= NumPoints)
+	if (NumPoints == 0 || Index == NumPoints)
 	{
+		/* Append Case */
+
 		GetSpline().AppendPoint(ControlPoint, EParameterizationPolicy::Uniform);
+		const TArray<FKnot> NewKnotVector = GetSpline().GetKnotVector();
 
 		if (TTangentBezierSpline<FVector>* ScaleChild = GetTypedAttributeChannel<TTangentBezierSpline<FVector>>(ScaleAttrName))
 		{
 			ScaleChild->AppendPoint(TTangentBezierSpline<FVector>::FTangentBezierControlPoint(Point.Scale), EParameterizationPolicy::Uniform);
-			ScaleChild->SetKnotVector(GetSpline().GetKnotVector());
+			ScaleChild->SetKnotVector(NewKnotVector);
 		}
 
 		if (TTangentBezierSpline<FQuat>* RotChild = GetTypedAttributeChannel<TTangentBezierSpline<FQuat>>(RotationAttrName))
 		{
 			RotChild->AppendPoint(TTangentBezierSpline<FQuat>::FTangentBezierControlPoint(Point.Rotation.Quaternion()), EParameterizationPolicy::Uniform);
-			RotChild->SetKnotVector(GetSpline().GetKnotVector());
+			RotChild->SetKnotVector(NewKnotVector);
 		}
-
-		MarkLegacyCurvesDirty();
-		return;
 	}
 	else if (Index == 0)
 	{
+		/* Prepend Case */
+
 		GetSpline().PrependPoint(ControlPoint, EParameterizationPolicy::Uniform);
+		const TArray<FKnot> NewKnotVector = GetSpline().GetKnotVector();
 
 		if (TTangentBezierSpline<FVector>* ScaleChild = GetTypedAttributeChannel<TTangentBezierSpline<FVector>>(ScaleAttrName))
 		{
 			ScaleChild->PrependPoint(TTangentBezierSpline<FVector>::FTangentBezierControlPoint(Point.Scale), EParameterizationPolicy::Uniform);
-			ScaleChild->SetKnotVector(GetSpline().GetKnotVector());
+			ScaleChild->SetKnotVector(NewKnotVector);
 		}
 
 		if (TTangentBezierSpline<FQuat>* RotChild = GetTypedAttributeChannel<TTangentBezierSpline<FQuat>>(RotationAttrName))
 		{
 			RotChild->PrependPoint(TTangentBezierSpline<FQuat>::FTangentBezierControlPoint(Point.Rotation.Quaternion()), EParameterizationPolicy::Uniform);
-			RotChild->SetKnotVector(GetSpline().GetKnotVector());
+			RotChild->SetKnotVector(NewKnotVector);
 		}
-
-		MarkLegacyCurvesDirty();
-		return;
 	}
+	else
+	{
+		/* Insert Case */
 
-	// We handle 0 and NumPoints cases above
-	Index = FMath::Clamp(Index, 1, GetNumberOfPoints() - 1);
+		// We want to assume Index is in the range [1, NumPoints - 1].
+		// Note: We can assume NumPoints >= 2 due to the above prepend/append special cases in addition to the normal bounds check on Index.
+		if (!ensureAlways(Index >= 1 && Index <= NumPoints - 1))
+		{
+			return;
+		}
 
-	const int32 SegmentIndex = Index - 1;
-	float SegmentParam;
-	const float Parameter = GetSpline().FindNearestOnSegment(Point.Position, SegmentIndex, SegmentParam);
+		// We will split the segment parameter range as opposed to finding the parameter of the nearest point on the spline
+		// so that insertion doesn't accidentally & prematurely cause a degenerate segment parameter range.
+		const int32 SplitSegmentIndex = Index - 1;
+		const float Parameter = GetSpline().GetSegmentParameterRange(SplitSegmentIndex).Interpolate(0.5f);
 
-	int32 InsertedIndex = GetSpline().InsertPoint(ControlPoint, Parameter);
+		int32 InsertedIndex = GetSpline().InsertPoint(ControlPoint, Parameter);
+		const TArray<FKnot> NewKnotVector = GetSpline().GetKnotVector();
 
-	float ScaleParam = MapParameterToChildSpace(ScaleAttrName, Parameter);
-	if (TTangentBezierSpline<FVector>* ScaleChild = GetTypedAttributeChannel<TTangentBezierSpline<FVector>>(ScaleAttrName))
-	{
-		ScaleChild->InsertPoint(TTangentBezierSpline<FVector>::FTangentBezierControlPoint(Point.Scale), ScaleParam);
-		ScaleChild->SetKnotVector(GetSpline().GetKnotVector());
-	}
+		// This would be a problem...
+		ensureAlways(Index == InsertedIndex);
 
-	float RotParam = MapParameterToChildSpace(RotationAttrName, Parameter);
-	if (TTangentBezierSpline<FQuat>* RotChild = GetTypedAttributeChannel<TTangentBezierSpline<FQuat>>(RotationAttrName))
-	{
-		RotChild->InsertPoint(TTangentBezierSpline<FQuat>::FTangentBezierControlPoint(Point.Rotation.Quaternion()), RotParam);
-		RotChild->SetKnotVector(GetSpline().GetKnotVector());
+		if (TTangentBezierSpline<FVector>* ScaleChild = GetTypedAttributeChannel<TTangentBezierSpline<FVector>>(ScaleAttrName))
+		{
+			float ScaleParam = MapParameterToChildSpace(ScaleAttrName, Parameter);
+			ScaleChild->InsertPoint(TTangentBezierSpline<FVector>::FTangentBezierControlPoint(Point.Scale), ScaleParam);
+			ScaleChild->SetKnotVector(NewKnotVector);
+		}
+
+		if (TTangentBezierSpline<FQuat>* RotChild = GetTypedAttributeChannel<TTangentBezierSpline<FQuat>>(RotationAttrName))
+		{
+			float RotParam = MapParameterToChildSpace(RotationAttrName, Parameter);
+			RotChild->InsertPoint(TTangentBezierSpline<FQuat>::FTangentBezierControlPoint(Point.Rotation.Quaternion()), RotParam);
+			RotChild->SetKnotVector(NewKnotVector);
+		}
 	}
 
+	MarkReparamTableDirty();
 	MarkLegacyCurvesDirty();
 }
 
@@ -779,19 +812,21 @@ void FTangentSpline::RemovePoint(int32 Index)
 	}
 
 	GetSpline().RemovePoint(Index);
+	const TArray<FKnot> NewKnotVector = GetSpline().GetKnotVector();
 
 	if (TTangentBezierSpline<FVector>* ScaleChild = GetTypedAttributeChannel<TTangentBezierSpline<FVector>>(ScaleAttrName))
 	{
 		ScaleChild->RemovePoint(Index);
-		ScaleChild->SetKnotVector(GetSpline().GetKnotVector());
+		ScaleChild->SetKnotVector(NewKnotVector);
 	}
 
 	if (TTangentBezierSpline<FQuat>* RotChild = GetTypedAttributeChannel<TTangentBezierSpline<FQuat>>(RotationAttrName))
 	{
 		RotChild->RemovePoint(Index);
-		RotChild->SetKnotVector(GetSpline().GetKnotVector());
+		RotChild->SetKnotVector(NewKnotVector);
 	}
 
+	MarkReparamTableDirty();
 	MarkLegacyCurvesDirty();
 }
 
@@ -938,20 +973,185 @@ EInterpCurveMode FTangentSpline::GetSplinePointType(int32 Index) const
 	{
 		return CIM_Unknown;
 	}
-	UE::Geometry::Spline::ETangentMode TangentMode = GetSpline().GetTangentMode(Index);
-	return ConvertTangentModeToInterpCurveMode(TangentMode);
+
+	return ConvertTangentModeToInterpCurveMode(GetSpline().GetTangentMode(Index));
 }
 
-float FTangentSpline::GetParameterAtIndex(int32 Index) const
+void FTangentSpline::Reparameterize()
 {
-	// Index == GetNumberOfPoints is permitted when the spline is closed to fetch the
-	// parameterization of the transient closing point which is coincident with the 0th point.
-	if (Index < 0 || (!IsClosedLoop() && Index >= GetNumberOfPoints()) || (IsClosedLoop() && Index > GetNumberOfPoints()))
+	Reparameterize(GetParameterizationPolicy());
+}
+
+void FTangentSpline::Reparameterize(UE::Geometry::Spline::EParameterizationPolicy ParameterizationPolicy)
+{
+	TRACE_CPUPROFILER_EVENT_SCOPE(FTangentSpline::Reparameterize);
+
+	using namespace UE::Geometry::Spline;
+
+	// todo: Reparameterize attribute channels such that their points remain positioned at the same point in the parameter space of its associated segment.
+
+	GetSpline().Reparameterize(ParameterizationPolicy);
+	GetSpline().UpdateTangents();
+
+	const TArray<FKnot> NewKnotVector = GetSpline().GetKnotVector();
+
+	if (TTangentBezierSpline<FVector>* ScaleSpline = GetTypedAttributeChannel<TTangentBezierSpline<FVector>>(ScaleAttrName))
 	{
-		return 0.f;
+		ScaleSpline->SetKnotVector(NewKnotVector);
+		ScaleSpline->UpdateTangents();
+	}
+
+	if (TTangentBezierSpline<FQuat>* RotSpline = GetTypedAttributeChannel<TTangentBezierSpline<FQuat>>(RotationAttrName))
+	{
+		RotSpline->SetKnotVector(NewKnotVector);
+		RotSpline->UpdateTangents();
+	}
+
+	MarkReparamTableDirty();
+	MarkLegacyCurvesDirty();
+}
+
+bool FTangentSpline::SetParameterAtIndex(int32 Index, float Parameter)
+{
+	const int32 NumPoints = GetNumberOfPoints();
+
+	if (Index < 0 || Index >= NumPoints)
+	{
+		return false;
+	}
+
+	auto ExecuteSetParameter = [this, Index, Parameter]()
+	{
+		GetSpline().SetParameter(Index, Parameter);
+
+		if (UE::Geometry::Spline::TTangentBezierSpline<FQuat>* RotSpline = GetTypedAttributeChannel<UE::Geometry::Spline::TTangentBezierSpline<FQuat>>(RotationAttrName))
+		{
+			RotSpline->SetParameter(Index, Parameter);
+		}
+
+		if (UE::Geometry::Spline::TTangentBezierSpline<FVector>* ScaleSpline = GetTypedAttributeChannel<UE::Geometry::Spline::TTangentBezierSpline<FVector>>(ScaleAttrName))
+		{
+			ScaleSpline->SetParameter(Index, Parameter);
+		}
+	};
+
+	const float CurrentPointParameter = GetParameterAtIndex(Index);
+
+	// We do this here so we can check < and > later without technically missing this case.
+	if (Parameter == CurrentPointParameter)
+	{
+		return true;
+	}
+
+	// Early out for this trivial case, no reordering is possible and no closing knot fixup is necessary.
+	if (NumPoints == 1)
+	{
+		ExecuteSetParameter();
+		return true;
+	}
+
+	const int32 PreviousIndex = Index > 0 ? Index - 1 : INDEX_NONE;
+	if (PreviousIndex != INDEX_NONE)
+	{
+		const float PreviousPointParameter = GetParameterAtIndex(PreviousIndex);
+
+		// Detect possible segment flips, early out if necessary.
+		if (Parameter <= PreviousPointParameter)
+		{
+			return false;
+		}
+	}
+
+	const int32 NextIndex = Index < NumPoints - 1 ? Index + 1 : INDEX_NONE;
+	if (NextIndex != INDEX_NONE)
+	{
+		const float NextPointParameter = GetParameterAtIndex(NextIndex);
+
+		// Detect possible segment flips, early out if necessary.
+		if (Parameter >= NextPointParameter)
+		{
+			return false;
+		}
+	}
+
+	ExecuteSetParameter();
+
+	MarkReparamTableDirty();
+	MarkLegacyCurvesDirty();
+
+	return true;
+}
+
+bool FTangentSpline::SetParameters(const TArray<float>& InParameters)
+{
+	using namespace UE::Geometry::Spline;
+
+	const int32 NumPoints = GetNumberOfPoints();
+
+	if (InParameters.Num() != NumPoints)
+	{
+		UE_LOG(LogTangentSpline, Warning, TEXT("SetParameters must be provided an array equal in size to the current number of points!"))
+		return false;
+	}
+
+	TArray<FKnot> Knots;
+
+	auto IsEndpoint = [this, NumPoints](const int32 PointIdx) -> bool
+	{
+		// If we are closed, there are no endpoints.
+		if (IsClosedLoop() && NumPoints >= 2)
+		{
+			return false;
+		}
+
+		// If we are open and the first or last point, we are an endpoint.
+		return PointIdx == 0 || PointIdx == NumPoints - 1;
+	};
+
+	for (int32 ParameterIdx = 0; ParameterIdx < InParameters.Num(); ++ParameterIdx)
+	{
+		if (ParameterIdx > 0)
+		{
+			// Detect non-monotonically increasing knot vector, which is invalid.
+			if (InParameters[ParameterIdx] <= InParameters[ParameterIdx - 1])
+			{
+				return false;
+			}
+		}
+
+		Knots.Emplace(InParameters[ParameterIdx], IsEndpoint(ParameterIdx) ? 4 : 3);
+	}
+
+	// True if we need a closing knot
+	if (IsClosedLoop() && NumPoints >= 2)
+	{
+		const float OldClosingSegmentKnotDelta = GetSpline().GetSegmentParameterRange(GetNumberOfSegments() - 1).Length();
+
+		Knots.Last().Multiplicity = 3;
+		Knots.Emplace(Knots.Last().Value + OldClosingSegmentKnotDelta, 3);
+	}
+
+	GetSpline().SetKnotVector(Knots);
+
+	if (TTangentBezierSpline<FVector>* ScaleSpline = GetTypedAttributeChannel<TTangentBezierSpline<FVector>>(ScaleAttrName))
+	{
+		ScaleSpline->SetKnotVector(Knots);
+	}
+
+	if (TTangentBezierSpline<FQuat>* RotSpline = GetTypedAttributeChannel<TTangentBezierSpline<FQuat>>(RotationAttrName))
+	{
+		RotSpline->SetKnotVector(Knots);
 	}
 
-	return FromInternalSplineSpace(GetSpline().GetParameter(Index));
+	MarkReparamTableDirty();
+	MarkLegacyCurvesDirty();
+
+	return true;
+}
+
+float FTangentSpline::GetParameterAtIndex(int32 Index) const
+{
+	return GetSpline().GetParameter(Index);
 }
 
 float FTangentSpline::GetParameterAtDistance(float Distance) const
@@ -965,7 +1165,7 @@ float FTangentSpline::GetParameterAtDistance(float Distance) const
 
 float FTangentSpline::GetDistanceAtParameter(float Parameter) const
 {
-	if (!GetSegmentSpace().Contains(Parameter))
+	if (!GetParameterSpace().Contains(Parameter))
 	{
 		return 0.f;
 	}
@@ -973,15 +1173,66 @@ float FTangentSpline::GetDistanceAtParameter(float Parameter) const
 	UpdateReparamTable();
 
 	UE::TReadScopeLock Lock(ReparamTableRWLock);
+	if (ReparamTable.Points.IsEmpty())
+	{
+		return 0.f;
+	}
+
+	if (Parameter <= ReparamTable.Points[0].OutVal)
+	{
+		return ReparamTable.Points[0].InVal;
+	}
+
+	if (Parameter >= ReparamTable.Points.Last().OutVal)
+	{
+		return ReparamTable.Points.Last().InVal;
+	}
+	const int32 NumReparamPoints = ReparamTable.Points.Num();
+	int32 Low = 0;
+	int32 High = NumReparamPoints - 1;
+	while (High - Low > 1)
+	{
+		const int32 Mid = (Low + High) / 2;
 
-	const float ParameterMax = static_cast<float>(GetNumberOfSegments());
-	const float Key = (Parameter / ParameterMax) * (ReparamTable.Points.Num() - 1);
-	const int32 LowerKey = FMath::FloorToInt32(Key);
-	ensureAlways(LowerKey >= 0 && LowerKey < ReparamTable.Points.Num());
-	const int32 UpperKey = FMath::CeilToInt32(Key);
-	ensureAlways(UpperKey >= 0 && UpperKey < ReparamTable.Points.Num());
-	const float Alpha = FMath::Frac(Key);
-	return FMath::Lerp(ReparamTable.Points[LowerKey].InVal, ReparamTable.Points[UpperKey].InVal, Alpha);
+		if (ReparamTable.Points[Mid].OutVal < Parameter)
+		{
+			Low = Mid;
+		}
+		else
+		{
+			High = Mid;
+		}
+	}
+
+	const float Alpha = FInterval1f(ReparamTable.Points[Low].OutVal, ReparamTable.Points[High].OutVal).GetT(Parameter);
+	return FMath::Lerp(ReparamTable.Points[Low].InVal, ReparamTable.Points[High].InVal, Alpha);
+}
+
+int32 FTangentSpline::GetIndexAtParameter(float Parameter) const
+{
+	if (!GetParameterSpace().Contains(Parameter))
+	{
+		return INDEX_NONE;
+	}
+
+	int32 MinIndex = 0;
+	int32 MaxIndex = GetNumberOfSegments();
+
+	while (MaxIndex - MinIndex > 1)
+	{
+		int32 MidIndex = (MinIndex + MaxIndex) / 2;
+
+		if (GetSpline().GetParameter(MidIndex) <= Parameter)
+		{
+			MinIndex = MidIndex;
+		}
+		else
+		{
+			MaxIndex = MidIndex;
+		}
+	}
+
+	return MinIndex;
 }
 
 FQuat FTangentSpline::GetOrientation(int32 Index) const
@@ -1042,12 +1293,12 @@ void FTangentSpline::SetOrientation(int32 Index, const FQuat& InOrientation)
 
 float FTangentSpline::FindNearest(const FVector& InLocation, float& OutSquaredDistance) const
 {
-	return FromInternalSplineSpace(GetSpline().FindNearest(InLocation, OutSquaredDistance));
+	return GetSpline().FindNearest(InLocation, OutSquaredDistance);
 }
 
 float FTangentSpline::FindNearestOnSegment(const FVector& InLocation, int32 SegmentIndex, float& OutSquaredDistance) const
 {
-	return FromInternalSplineSpace(GetSpline().FindNearestOnSegment(InLocation, SegmentIndex, OutSquaredDistance));
+	return GetSpline().FindNearestOnSegment(InLocation, SegmentIndex, OutSquaredDistance);
 }
 
 FVector FTangentSpline::EvaluatePosition(float Parameter) const
@@ -1059,13 +1310,12 @@ FVector FTangentSpline::EvaluatePosition(float Parameter) const
 	}
 	else
 	{
-		return Evaluate(ToInternalSplineSpace(Parameter));
+		return Evaluate(Parameter);
 	}
 }
 
 FVector FTangentSpline::EvaluateDerivative(float Parameter) const
 {
-	Parameter = ToInternalSplineSpace(Parameter);
 	return GetSpline().GetTangent(Parameter);
 }
 
@@ -1099,10 +1349,7 @@ void FTangentSpline::UpdateSpline()
 {
 	TRACE_CPUPROFILER_EVENT_SCOPE(FTangentSpline::UpdateSpline);
 
-	// Update our parameterization, forward it to the rotation and scale channels, and then recompute tangents.
-	Reparameterize();
-
-	MarkReparamTableDirty();
+	UpdateTangents();
 
 	// Updates legacy curves based on points, never evaluates internal spline.
 	RebuildLegacyCurves();
@@ -1114,7 +1361,7 @@ void FTangentSpline::UpdateSpline(const FUpdateSplineParams& Params)
 
 	TRACE_CPUPROFILER_EVENT_SCOPE(FTangentSpline::UpdateSpline);
 
-	// todo: handled unhandled params
+	CachedUpdateSplineParams = Params;
 
 	if (UE::Spline::GSkipDoubleUpdate)
 	{
@@ -1130,33 +1377,44 @@ void FTangentSpline::UpdateSpline(const FUpdateSplineParams& Params)
 	GetSpline().SetStationaryEndpoints(Params.bStationaryEndpoints);
 	ReparamStepsPerSegment = Params.ReparamStepsPerSegment;
 	
-	// Todo: enable this when hooking up spline input keys
-	constexpr bool bPropagateLoopPosition = false;
-	if constexpr (bPropagateLoopPosition)
+	if (Params.bClosedLoop)
 	{
-		if (Params.bClosedLoop)
+		const int32 NumPoints = GetNumberOfPoints();
+
+		if (NumPoints > 0)
 		{
-			const int32 NumPoints = GetNumberOfPoints();
-			if (NumPoints > 0)
+			const float LastParam = GetParameterAtIndex(NumPoints - 1);
+			const float LoopKey = Params.bLoopPositionOverride
+				? Params.LoopPosition         // user-specified absolute closing key
+				: LastParam + 1.0f;           // legacy default behaviour
+
+			if (LoopKey > LastParam)
 			{
-				const float LastParam = GetParameterAtIndex(NumPoints - 1); // FSpline-style parameter
-				const float LoopKey = Params.bLoopPositionOverride
-					? Params.LoopPosition         // user-specified absolute closing key
-					: LastParam + 1.0f;           // legacy default behaviour
+				GetSpline().SetParameter(GetNumberOfPoints(), LoopKey);
 
-				if (LoopKey > LastParam)
+				if (TTangentBezierSpline<FVector>* ScaleChild = GetTypedAttributeChannel<TTangentBezierSpline<FVector>>(ScaleAttrName))
 				{
-					GetSpline().SetParameter(GetNumberOfPoints(), LoopKey);
+					ScaleChild->SetParameter(GetNumberOfPoints(), LoopKey);
+				}
 
-					if (TTangentBezierSpline<FVector>* ScaleChild = GetTypedAttributeChannel<TTangentBezierSpline<FVector>>(ScaleAttrName))
-					{
-						ScaleChild->SetParameter(GetNumberOfPoints(), LoopKey);
-					}
+				if (TTangentBezierSpline<FQuat>* RotChild = GetTypedAttributeChannel<TTangentBezierSpline<FQuat>>(RotationAttrName))
+				{
+					RotChild->SetParameter(GetNumberOfPoints(), LoopKey);
+				}
+			}
+			else
+			{
+				// Legacy behavior: If the absolute loop position is <= LastParam, the spline is open.
 
-					if (TTangentBezierSpline<FQuat>* RotChild = GetTypedAttributeChannel<TTangentBezierSpline<FQuat>>(RotationAttrName))
-					{
-						RotChild->SetParameter(GetNumberOfPoints(), LoopKey);
-					}
+				if (UE::Spline::GSkipDoubleUpdate)
+				{
+					// Avoid updating the spline twice.
+					constexpr bool bUpdateSpline = false;
+					SetClosedLoop(false, bUpdateSpline);
+				}
+				else
+				{
+					SetClosedLoop(false);
 				}
 			}
 		}
@@ -1165,55 +1423,14 @@ void FTangentSpline::UpdateSpline(const FUpdateSplineParams& Params)
 	UpdateSpline();
 }
 
-void FTangentSpline::Reparameterize()
+void FTangentSpline::UpdateTangents()
 {
 	using namespace UE::Geometry::Spline;
 
-	TRACE_CPUPROFILER_EVENT_SCOPE(FTangentSpline::Reparameterize);
-
-	// Caching segment space attribute parameters so we can restore them (necessary because we store the mapping in parent space which are about to be invalidated).
-	TMap<FName, TPair<FInterval1f, TArray<float>>> CachedChannelParameterizations;
-	for (const FName& ChannelName : GetAttributeChannelNamesByValueType<float>())
-	{
-		if (ChannelName == RotationAttrName || ChannelName == ScaleAttrName || NumAttributeValues<float>(ChannelName) == 0)
-		{
-			continue;
-		}
-
-		ensureAlwaysMsgf(GetAttributeChannelRangeSpace(ChannelName) == EMappingRangeSpace::Parent, TEXT("Attribute channel mapping range space is not Parent!"));
-
-		FInterval1f ChannelRangeInParentSpace = GetAttributeChannelRange(ChannelName);
-		FInterval1f ChannelRangeInSegmentSpace = FInterval1f(FromInternalSplineSpace(ChannelRangeInParentSpace.Min), FromInternalSplineSpace(ChannelRangeInParentSpace.Max));
-
-		CachedChannelParameterizations.FindOrAdd(ChannelName) = TPair<FInterval1f, TArray<float>>(ChannelRangeInSegmentSpace, CacheChannelKnots<float>(ChannelName));
-
-		for (float& InternalKnot : CachedChannelParameterizations[ChannelName].Value)
-		{
-			InternalKnot = FromInternalSplineSpace(InternalKnot);
-		}
-	}
-
-	GetSpline().Reparameterize(GetParameterizationPolicy());
 	GetSpline().UpdateTangents();
 
 	for (const FName& ChannelName : GetAttributeChannelNamesByValueType<float>())
 	{
-		if (ChannelName == RotationAttrName || ChannelName == ScaleAttrName || NumAttributeValues<float>(ChannelName) == 0)
-		{
-			continue;
-		}
-
-		const FInterval1f& ChannelRangeInSegmentSpace = CachedChannelParameterizations[ChannelName].Key;
-		const FInterval1f ChannelRangeInParentSpace = FInterval1f(ToInternalSplineSpace(ChannelRangeInSegmentSpace.Min), ToInternalSplineSpace(ChannelRangeInSegmentSpace.Max));
-
-		SetAttributeChannelRange(ChannelName, ChannelRangeInParentSpace, EMappingRangeSpace::Parent);
-
-		for (float& InternalKnot : CachedChannelParameterizations[ChannelName].Value)
-		{
-			InternalKnot = ToInternalSplineSpace(InternalKnot);
-		}
-
-		RestoreChannelKnots(ChannelName, CachedChannelParameterizations[ChannelName].Value);
 		if (TTangentBezierSpline<float>* Channel = GetTypedAttributeChannel<TTangentBezierSpline<float>>(ChannelName))
 		{
 			Channel->UpdateTangents();
@@ -1222,36 +1439,16 @@ void FTangentSpline::Reparameterize()
 
 	if (TTangentBezierSpline<FVector>* ScaleSpline = GetTypedAttributeChannel<TTangentBezierSpline<FVector>>(ScaleAttrName))
 	{
-		ScaleSpline->SetKnotVector(GetSpline().GetKnotVector());
 		ScaleSpline->UpdateTangents();
 	}
 
 	if (TTangentBezierSpline<FQuat>* RotSpline = GetTypedAttributeChannel<TTangentBezierSpline<FQuat>>(RotationAttrName))
 	{
-		RotSpline->SetKnotVector(GetSpline().GetKnotVector());
 		RotSpline->UpdateTangents();
 	}
-}
-
-float FTangentSpline::FromInternalSplineSpace(float Parameter) const
-{
-	float OutLocalParam;
-	const int32 SegmentIndex = GetSpline().FindSegmentIndex(Parameter, OutLocalParam);
-	return static_cast<float>(SegmentIndex) + OutLocalParam;
-}
-float FTangentSpline::ToInternalSplineSpace(float Parameter) const
-{
-	const int32 NumSegments = GetNumberOfSegments();
-
-	if (NumSegments == 0)
-	{
-		return 0.f;
-	}
 
-	float ClampedParam = FMath::Clamp(Parameter, 0.0f, static_cast<float>(NumSegments));
-	const int32 SegmentIndex = FMath::Min(FMath::FloorToInt(ClampedParam), NumSegments - 1);
-	const float LocalT = ClampedParam - static_cast<float>(SegmentIndex); // Parameter [0,1] within segment
-	return GetSegmentParameterRange(SegmentIndex).Interpolate(LocalT);
+	MarkReparamTableDirty();
+	MarkLegacyCurvesDirty();
 }
 
 void FTangentSpline::MarkReparamTableDirty()
@@ -1287,16 +1484,20 @@ void FTangentSpline::UpdateReparamTable() const
 	float AccumulatedLength = 0.0f;
 	for (int32 SegmentIndex = 0; SegmentIndex < NumSegments; ++SegmentIndex)
 	{
-		ReparamTable.Points.Emplace(AccumulatedLength, SegmentIndex, 0.0f, 0.0f, CIM_Linear);
+		FInterval1f SegmentParameterRange = GetSegmentParameterRange(SegmentIndex);
+
+		ReparamTable.Points.Emplace(AccumulatedLength, SegmentParameterRange.Min, 0.0f, 0.0f, CIM_Linear);
+
 		for (int32 Step = 1; Step < ReparamStepsPerSegment; ++Step)
 		{
-			const float Param = static_cast<float>(Step) / ReparamStepsPerSegment;
-			ReparamTable.Points.Emplace(GetSegmentLength(SegmentIndex, Param) + AccumulatedLength, SegmentIndex + Param, 0.0f, 0.0f, CIM_Linear);
+			const float SegmentT = static_cast<float>(Step) / ReparamStepsPerSegment;
+			ReparamTable.Points.Emplace(GetSegmentLength(SegmentIndex, SegmentT, CachedUpdateSplineParams.Scale3D) + AccumulatedLength, SegmentParameterRange.Interpolate(SegmentT), 0.0f, 0.0f, CIM_Linear);
 		}
 
-		AccumulatedLength += GetSegmentLength(SegmentIndex, 1.0f);
+		AccumulatedLength += GetSegmentLength(SegmentIndex, 1.0f, CachedUpdateSplineParams.Scale3D);
 	}
-	ReparamTable.Points.Emplace(AccumulatedLength, static_cast<float>(NumSegments), 0.0f, 0.0f, CIM_Linear);
+
+	ReparamTable.Points.Emplace(AccumulatedLength, GetParameterSpace().Max, 0.0f, 0.0f, CIM_Linear);
 
 	ReparamTableVersion = CachedNextVersion;
 }
@@ -1334,13 +1535,6 @@ void FTangentSpline::RebuildLegacyCurves() const
 	RotationCurve.Points.Reset(GetNumberOfPoints());
 	ScaleCurve.Points.Reset(GetNumberOfPoints());
 
-	PositionCurve.bIsLooped = IsClosedLoop();
-	PositionCurve.LoopKeyOffset = 1.f;
-	RotationCurve.bIsLooped = IsClosedLoop();
-	RotationCurve.LoopKeyOffset = 1.f;
-	ScaleCurve.bIsLooped = IsClosedLoop();
-	ScaleCurve.LoopKeyOffset = 1.f;
-
 	using FRotationChannel = UE::Geometry::Spline::TTangentBezierSpline<FQuat>;
 	using FScaleChannel = UE::Geometry::Spline::TTangentBezierSpline<FVector>;
 
@@ -1362,6 +1556,20 @@ void FTangentSpline::RebuildLegacyCurves() const
 	Algo::SortBy(RotationCurve.Points, &FInterpCurvePoint<FQuat>::InVal);
 	Algo::SortBy(ScaleCurve.Points, &FInterpCurvePoint<FVector>::InVal);
 
+	if (IsClosedLoop())
+	{
+		const float ClosingParameter = GetParameterAtIndex(GetNumberOfPoints());
+		PositionCurve.SetLoopKey(ClosingParameter);
+		RotationCurve.SetLoopKey(ClosingParameter);
+		ScaleCurve.SetLoopKey(ClosingParameter);
+	}
+	else
+	{
+		PositionCurve.ClearLoopKey();
+		RotationCurve.ClearLoopKey();
+		ScaleCurve.ClearLoopKey();
+	}
+
 	PositionCurve.AutoSetTangents(0.f, false);
 	RotationCurve.AutoSetTangents(0.f, false);
 	ScaleCurve.AutoSetTangents(0.f, false);
diff --git a/Engine/Source/Runtime/GeometryCore/Public/Splines/BSpline.h b/Engine/Source/Runtime/GeometryCore/Public/Splines/BSpline.h
--- a/Engine/Source/Runtime/GeometryCore/Public/Splines/BSpline.h
+++ b/Engine/Source/Runtime/GeometryCore/Public/Splines/BSpline.h
@@ -37,11 +37,11 @@ struct FKnot
 
 	FKnot()
 		: Value(0.0f)
-		, Multiplicity(1)
+		, Multiplicity(0)
 	{
 	}
 
-	FKnot(float InValue, uint32 InMultiplicity = 1)
+	explicit FKnot(float InValue, uint32 InMultiplicity)
 		: Value(InValue)
 		, Multiplicity(InMultiplicity)
 	{
@@ -455,32 +455,34 @@ public:
 	bool ValidateKnotVector(const TArray<FKnot>& InKnots) const
 	{
 		const int32 RequiredKnots = GetExpectedNumKnots();
+		int32 TotalExpandedKnots = 0;
 
 		// Validate knot multiplicities and calculate total expanded knots  
-		int32 TotalExpandedKnots = 0;
 		for (const FKnot& Knot : InKnots)
 		{
-			if (Knot.Multiplicity == 0)
+			if (!ensureAlwaysMsgf(Knot.Multiplicity > 0,
+								  TEXT("Invalid knot with zero multiplicity.")))
 			{
-				UE_LOG(LogSpline, Warning, TEXT("Invalid knot with zero multiplicity."));
 				return false;
 			}
+
 			TotalExpandedKnots += Knot.Multiplicity;
 		}
 
-		if (TotalExpandedKnots < RequiredKnots)
+		if (!ensureAlwaysMsgf(TotalExpandedKnots >= RequiredKnots,
+							  TEXT("Invalid knot vector size. Need %d knots, got %d."),
+							  RequiredKnots,
+							  TotalExpandedKnots))
 		{
-			UE_LOG(LogSpline, Warning, TEXT("Invalid knot vector size. Need %d knots, got %d."),
-				   RequiredKnots, TotalExpandedKnots);
 			return false;
 		}
 
 		// Verify knot values are non-decreasing (strict comparison, no epsilon)
 		for (int32 i = 1; i < InKnots.Num(); ++i)
 		{
-			if (InKnots[i].Value < InKnots[i - 1].Value)
+			if (!ensureAlwaysMsgf(InKnots[i].Value >= InKnots[i - 1].Value,
+								  TEXT("Invalid knot vector: knot values must be non-decreasing.")))
 			{
-				UE_LOG(LogSpline, Warning, TEXT("Invalid knot vector: knot values must be non-decreasing."));
 				return false;
 			}
 		}
@@ -939,7 +941,7 @@ protected:
 			// For closed loop, generate periodic knot vector based on chord lengths
 			for (int32 i = 0; i < NumKnots; ++i)
 			{
-				PairKnots.Add(FKnot(static_cast<float>(i) * TotalLength / static_cast<float>(n + d)));
+				PairKnots.Add(FKnot(static_cast<float>(i) * TotalLength / static_cast<float>(n + d), 1));
 			}
 		}
 		else if (bClampEnds)
@@ -995,7 +997,7 @@ protected:
 			for (int32 i = 0; i < KnotCount; ++i)
 			{
 				const float Param = static_cast<float>(i) / static_cast<float>(KnotCount - 1);
-				PairKnots.Add(FKnot(Param * TotalLength));
+				PairKnots.Add(FKnot(Param * TotalLength, 1));
 			}
 		}
 		MarkFlatKnotsCacheDirty();
@@ -1047,7 +1049,7 @@ protected:
 			// Generate periodic knot vector using centripetal parameterization
 			for (int32 i = 0; i < NumKnots; ++i)
 			{
-				PairKnots.Add(FKnot(static_cast<float>(i) * TotalLength / static_cast<float>(n + d)));
+				PairKnots.Add(FKnot(static_cast<float>(i) * TotalLength / static_cast<float>(n + d), 1));
 			}
 		}
 		else if (bClampEnds)
@@ -1087,7 +1089,7 @@ protected:
 
 					// Interpolate Knots
 					const float KnotValue = FMath::Lerp(Params[LowIndex], Params[HighIndex], Alpha);
-					PairKnots.Add(FKnot(KnotValue));
+					PairKnots.Add(FKnot(KnotValue, 1));
 				}
 			}
 			// Set the last (Degree+1) knots to 1.0
@@ -1099,7 +1101,7 @@ protected:
 			for (int32 i = 0; i < KnotCount; ++i)
 			{
 				const float Param = static_cast<float>(i) / static_cast<float>(KnotCount - 1);
-				PairKnots.Add(FKnot(Param * TotalLength));
+				PairKnots.Add(FKnot(Param * TotalLength, 1));
 			}
 		}
 		MarkFlatKnotsCacheDirty();
diff --git a/Engine/Source/Runtime/GeometryCore/Public/Splines/PolyBezierSpline.h b/Engine/Source/Runtime/GeometryCore/Public/Splines/PolyBezierSpline.h
--- a/Engine/Source/Runtime/GeometryCore/Public/Splines/PolyBezierSpline.h
+++ b/Engine/Source/Runtime/GeometryCore/Public/Splines/PolyBezierSpline.h
@@ -309,6 +309,11 @@ public:
 			}
 		}
 
+		if (Ar.CustomVer(FFortniteMainBranchObjectVersion::GUID) >= FFortniteMainBranchObjectVersion::SplineComponentReparameterizeOnLoad)
+		{
+			Ar << LoopKnotDelta;
+		}
+
 		return true;
 	}
 
@@ -1067,7 +1072,8 @@ public:
 		Base::InsertValue(1, InAnchor.Anchor);
 		Base::InsertValue(2, InAnchor.OutHandle);
 
-		const int32 KnotIndex = Base::InsertKnot(FKnot(Parameter));
+		// Degree of 1 is arbitrary, it will be set correctly based on loop and clamping flags by ApplyClampedKnotsMultiplicity.
+		const int32 KnotIndex = Base::InsertKnot(FKnot(Parameter, 1));
 
 		ensureAlwaysMsgf(KnotIndex == 0, TEXT("Unexpected knot insertion index when prepending an anchor!"));
 
@@ -1114,7 +1120,8 @@ public:
 		Base::AddValue(InAnchor.Anchor);
 		Base::AddValue(InAnchor.OutHandle);
 
-		const int32 KnotIndex = Base::InsertKnot(FKnot(Parameter));
+		// Degree of 1 is arbitrary, it will be set correctly based on loop and clamping flags by ApplyClampedKnotsMultiplicity.
+		const int32 KnotIndex = Base::InsertKnot(FKnot(Parameter, 1));
 
 		ensureAlwaysMsgf(KnotIndex == InsertIndex, TEXT("Unexpected knot insertion index when appending an anchor!"));
 
diff --git a/Engine/Source/Runtime/GeometryCore/Public/Splines/TangentBezierSpline.h b/Engine/Source/Runtime/GeometryCore/Public/Splines/TangentBezierSpline.h
--- a/Engine/Source/Runtime/GeometryCore/Public/Splines/TangentBezierSpline.h
+++ b/Engine/Source/Runtime/GeometryCore/Public/Splines/TangentBezierSpline.h
@@ -432,7 +432,7 @@ public:
     
 		return Clone;
 	}
-    
+
 	/**
 	 * Replaces all points in the spline with the provided points.
 	 * @param Points - Array of control points to set
@@ -459,7 +459,7 @@ public:
 	 * @return True if the points were successfully added
 	 */
 	bool AppendPoints(const TArray<FTangentBezierControlPoint>& Points, EParameterizationPolicy Policy)
-    {
+     {
         const int32 NumPoints = Points.Num();
         
         // Need at least 2 points to add a valid segment
@@ -595,7 +595,6 @@ public:
     {
 		// Add the anchor
 		int32 NewIndex = InternalSpline.AppendAnchor(AnchorFromControlPoint(Point), Policy);
-
 		if (NewIndex != INDEX_NONE)
 		{
 			// Set tangent mode
@@ -667,13 +666,14 @@ public:
 		{
 			return;
 		}
-    
+
 		// Remove the corresponding tangent mode
-		if (TangentModes.IsValidIndex(Index))
+		if (ensureAlwaysMsgf(TangentModes.IsValidIndex(Index), TEXT("Incorrect number of tangent modes!")))
 		{
 			TangentModes.RemoveAt(Index);
 		}
 
+		// Remove the point
 		InternalSpline.RemoveAnchor(Index);
     
 		// Ensure tangent modes array size matches number of points after removal
@@ -808,7 +808,6 @@ public:
     {
 		return InternalSpline.GetNumAnchorPoints();
     }
-	
 
 	virtual int32 GetNumberOfSegments() const override
 	{
@@ -882,11 +881,6 @@ public:
 
 	virtual FInterval1f GetParameterSpace() const override
 	{
-		if (GetNumPoints() == 1)
-		{
-			return FInterval1f(0.f, 0.f);
-		}
-		
 		return InternalSpline.GetParameterSpace();
 	}
 
