SplineComponent: Fix assignment as condition warning

#Synced-CL 49961239

diff --git a/Engine/Source/Runtime/Engine/Private/Components/SplineComponent.cpp b/Engine/Source/Runtime/Engine/Private/Components/SplineComponent.cpp
--- a/Engine/Source/Runtime/Engine/Private/Components/SplineComponent.cpp
+++ b/Engine/Source/Runtime/Engine/Private/Components/SplineComponent.cpp
@@ -1153,12 +1153,13 @@ FVector USplineComponent::GetScaleAtSplineInputKey(float InKey) const
 float USplineComponent::GetDistanceAlongSplineAtSplineInputKey(float InKey) const
 {
 	int32 NumPoints = 0;
-	const FTangentSpline* TangentSpline = nullptr;
+	const FTangentSpline* TangentSpline = Spline.Get<UE::Spline::ESplineType::Tangent>();
+
 	if (ShouldUseSplineCurves())
 	{
 		NumPoints = WarninglessSplineCurves().Position.Points.Num();
 	}
-	else if ((TangentSpline = Spline.Get<UE::Spline::ESplineType::Tangent>()))
+	else if (TangentSpline)
 	{
 		NumPoints = TangentSpline->GetNumberOfPoints();
 	}
