Substrate - NMT - added support for coverage weight node.

#tests pc editor dx12
#rnx

#Synced-CL 50037646

diff --git a/Engine/Source/Runtime/Engine/Private/Materials/MaterialExpressionsIR.cpp b/Engine/Source/Runtime/Engine/Private/Materials/MaterialExpressionsIR.cpp
--- a/Engine/Source/Runtime/Engine/Private/Materials/MaterialExpressionsIR.cpp
+++ b/Engine/Source/Runtime/Engine/Private/Materials/MaterialExpressionsIR.cpp
@@ -3241,7 +3241,7 @@ void UMaterialExpressionAggregate::Build(MIR::FEmitter& Em)
 		}
 	}
 	
-	// Make hte aggregate value using the optional prototype and assignments.
+	// Make the aggregate value using the optional prototype and assignments.
 	Prototype = Em.Aggregate(Aggregate, Prototype, Assignments.Left(NumAssignments));
 
 	// Output the aggregate value
@@ -5554,7 +5554,6 @@ void UMaterialExpressionSubstrateHorizontalMixing::Build(MIR::FEmitter& Em)
 
 	FValueRef HorizontalMixArg = CompileWithDefaultFloat1(Em, Mix, 0.5f);
 
-	int32 OutputCodeChunk = INDEX_NONE;
 	FSubstrateOperator& SubstrateOperator = SubstrateTranslatorData->SubstrateCompilationGetOperator(SubstrateTranslatorData->SubstrateTreeStackGetPathUniqueId());
 	if (SubstrateOperator.bUseParameterBlending)
 	{
@@ -5690,7 +5689,6 @@ void UMaterialExpressionSubstrateVerticalLayering::Build(MIR::FEmitter& Em)
 {
 	FSubstrateTranslatorData* SubstrateTranslatorData = Em.GetSubstrateTranslatorData();
 
-
 	if (!Top.GetTracedInput().Expression)
 	{
 		return Em.Errorf(TEXT("Missing Foreground input"));
@@ -5706,7 +5704,6 @@ void UMaterialExpressionSubstrateVerticalLayering::Build(MIR::FEmitter& Em)
 
 	FValueRef ThicknessArg = CompileWithDefaultFloat1(Em, Thickness, SUBSTRATE_LAYER_DEFAULT_THICKNESS_CM);
 
-	int32 OutputCodeChunk = INDEX_NONE;
 	FSubstrateOperator& SubstrateOperator = SubstrateTranslatorData->SubstrateCompilationGetOperator(SubstrateTranslatorData->SubstrateTreeStackGetPathUniqueId());
 
 	if (SubstrateOperator.bUseParameterBlending)
@@ -5824,6 +5821,77 @@ void UMaterialExpressionSubstrateVerticalLayering::Build(MIR::FEmitter& Em)
 	}
 }
 
+void UMaterialExpressionSubstrateWeight::Build(MIR::FEmitter& Em)
+{
+	FSubstrateTranslatorData* SubstrateTranslatorData = Em.GetSubstrateTranslatorData();
+	if (!A.GetTracedInput().Expression)
+	{
+		return Em.Errorf(TEXT("Missing A input"));
+	}
+
+	// Note: we no longer do SubstrateTreeStackPush/Pop since the path guid is regenerated for every substrate node from BuildTopMaterialExpression.
+	FValueRef AInputArg = Em.TryInput(&A);
+	FValueRef WeightArg = CompileWithDefaultFloat1(Em, Weight, 1.0f);
+
+	FSubstrateOperator& SubstrateOperator = SubstrateTranslatorData->SubstrateCompilationGetOperator(SubstrateTranslatorData->SubstrateTreeStackGetPathUniqueId());
+
+	if (!AInputArg.IsValid())
+	{
+		return Em.Errorf(TEXT("A input input graphs could not be evaluated for parameter blending."));
+	}
+	if (!WeightArg.IsValid())
+	{
+		return Em.Errorf(TEXT("Weight input graphs could not be evaluated for parameter blending."));
+	}
+
+	if (SubstrateOperator.bUseParameterBlending)
+	{
+		FSubstrateOperator* ABSDFOperator = SubstrateTranslatorData->SubstrateCompilationGetOperatorFromIndex(SubstrateOperator.LeftIndex);
+		const FSubstrateRegisteredSharedLocalBasis NewRegisteredSharedLocalBasis = ABSDFOperator->BSDFRegisteredSharedLocalBasis;
+
+		const FSubstrateOperator* PromoteToOperator = SubstrateOperator.bRootOfParameterBlendingSubTree ? &SubstrateOperator : nullptr;
+		if (PromoteToOperator)
+		{
+			check(PromoteToOperator->Index != INDEX_NONE);
+			check(PromoteToOperator->BSDFIndex != INDEX_NONE);
+
+			FValueRef LayeredSlab = Em.SubstrateCoverageWeightParameterBlending(
+				AInputArg,
+				WeightArg);
+
+			FValueRef PromotedSlab = Em.SubstratePromoteToOperator(
+				LayeredSlab,
+				Em.ConstantInt(PromoteToOperator->Index),
+				Em.ConstantInt(PromoteToOperator->BSDFIndex),
+				Em.ConstantInt(PromoteToOperator->LayerDepth),
+				Em.ConstantInt(PromoteToOperator->bIsBottom ? 1 : 0));
+
+			Em.Output(0, PromotedSlab);
+		}
+		else
+		{
+			FValueRef LayeredSlab = Em.SubstrateCoverageWeightParameterBlending(
+				AInputArg,
+				WeightArg);
+
+			Em.Output(0, LayeredSlab);
+		}
+
+		// Propagate the parameter blended normal
+		SubstrateOperator.BSDFRegisteredSharedLocalBasis = NewRegisteredSharedLocalBasis;
+	}
+	else
+	{
+		FValueRef LayeredSlabs = Em.SubstrateCoverageWeight(
+			AInputArg,
+			WeightArg,
+			Em.ConstantInt(SubstrateOperator.Index),
+			Em.ConstantInt(SubstrateOperator.MaxDistanceFromLeaves));
+
+		Em.Output(0, LayeredSlabs);
+	}
+}
+
 void UMaterialExpressionSubstrateTransmittanceToMFP::Build(MIR::FEmitter& Em)
 {
 	FSubstrateTranslatorData* SubstrateTranslatorData = Em.GetSubstrateTranslatorData();
diff --git a/Engine/Source/Runtime/Engine/Private/Materials/MaterialIR.cpp b/Engine/Source/Runtime/Engine/Private/Materials/MaterialIR.cpp
--- a/Engine/Source/Runtime/Engine/Private/Materials/MaterialIR.cpp
+++ b/Engine/Source/Runtime/Engine/Private/Materials/MaterialIR.cpp
@@ -53,6 +53,7 @@ const TCHAR* LexToString(EValueKind Kind)
 		case VK_SubstrateShadingModels: return TEXT("SubstrateShadingModel");
 		case VK_SubstrateHorizontalMixing: return TEXT("SubstrateHorizontalMixing");
 		case VK_SubstrateVerticalLayering: return TEXT("SubstrateVerticalLayering");
+		case VK_SubstrateCoverageWeight: return TEXT("SubstrateCoverageWeight");
 		case VK_SubstratePromoteToOperator: return TEXT("SubstratePromoteToOperator");
 		case VK_StageSwitch: return TEXT("StageSwitch");
 		case VK_HardwarePartialDerivative: return TEXT("HardwarePartialDerivative");
@@ -126,6 +127,7 @@ uint32 FValue::GetSizeInBytes() const
 		case VK_SubstrateShadingModels: return sizeof(FSubstrateShadingModels);
 		case VK_SubstrateHorizontalMixing: return sizeof(FSubstrateHorizontalMixing);
 		case VK_SubstrateVerticalLayering: return sizeof(FSubstrateVerticalLayering);
+		case VK_SubstrateCoverageWeight: return sizeof(FSubstrateCoverageWeight);
 		case VK_SubstratePromoteToOperator: return sizeof(FSubstratePromoteToOperator);
 		case VK_StageSwitch: return sizeof(FStageSwitch);
 		case VK_HardwarePartialDerivative: return sizeof(FHardwarePartialDerivative);
@@ -234,6 +236,13 @@ TConstArrayView<FValue*> FValue::GetUses() const
 			auto This = static_cast<const FSubstrateVerticalLayering*>(this);
 			return { &This->Top, FSubstrateVerticalLayering::NumStaticUses };
 		}
+
+		case VK_SubstrateCoverageWeight:
+		{
+			auto This = static_cast<const FSubstrateCoverageWeight*>(this);
+			return { &This->A, FSubstrateCoverageWeight::NumStaticUses };
+		}
+
 		case VK_SubstratePromoteToOperator:
 		{
 			auto This = static_cast<const FSubstratePromoteToOperator*>(this);
diff --git a/Engine/Source/Runtime/Engine/Private/Materials/MaterialIREmitter.cpp b/Engine/Source/Runtime/Engine/Private/Materials/MaterialIREmitter.cpp
--- a/Engine/Source/Runtime/Engine/Private/Materials/MaterialIREmitter.cpp
+++ b/Engine/Source/Runtime/Engine/Private/Materials/MaterialIREmitter.cpp
@@ -4482,6 +4482,40 @@ FValueRef FEmitter::SubstrateVerticalLayeringParameterBlending(FValueRef Top, FV
 	return EmitPrototype(*this, Prototype);
 }
 
+FValueRef FEmitter::SubstrateCoverageWeight(FValueRef A, FValueRef Weight, FValueRef OperatorIndex, FValueRef MaxDistanceFromLeaves)
+{
+	FSubstrateCoverageWeight Prototype = MakePrototype<FSubstrateCoverageWeight>(FType::MakeSubstrateData());
+	Prototype.A = A;
+	Prototype.Weight = Weight;
+
+	// ParameterBlending = 0
+	Prototype.OperatorIndex = OperatorIndex;
+	Prototype.MaxDistanceFromLeaves = MaxDistanceFromLeaves;
+
+	// Settings
+	Prototype.bParameterBlendingEnabled = false;
+	Prototype.bFullySimplifiedInstruction = SubstrateTranslatorData->CurrentSubstrateCompilationContext == Substrate::ESubstrateCompilationContext::SCC_FullySimplified;
+
+	return EmitPrototype(*this, Prototype);
+}
+
+FValueRef FEmitter::SubstrateCoverageWeightParameterBlending(FValueRef A, FValueRef Weight)
+{
+	FSubstrateCoverageWeight Prototype = MakePrototype<FSubstrateCoverageWeight>(FType::MakeSubstrateData());
+	Prototype.A = A;
+	Prototype.Weight = Weight;
+
+	// ParameterBlending = 0
+	Prototype.OperatorIndex = nullptr;
+	Prototype.MaxDistanceFromLeaves = nullptr;
+
+	// Settings
+	Prototype.bParameterBlendingEnabled = true;
+	Prototype.bFullySimplifiedInstruction = SubstrateTranslatorData->CurrentSubstrateCompilationContext == Substrate::ESubstrateCompilationContext::SCC_FullySimplified;
+
+	return EmitPrototype(*this, Prototype);
+}
+
 EShaderPlatform FEmitter::GetShaderPlatform() const
 {
 	return Module->GetShaderPlatform();
diff --git a/Engine/Source/Runtime/Engine/Private/Materials/MaterialIRToHLSLTranslator.cpp b/Engine/Source/Runtime/Engine/Private/Materials/MaterialIRToHLSLTranslator.cpp
--- a/Engine/Source/Runtime/Engine/Private/Materials/MaterialIRToHLSLTranslator.cpp
+++ b/Engine/Source/Runtime/Engine/Private/Materials/MaterialIRToHLSLTranslator.cpp
@@ -305,6 +305,7 @@ static bool IsFoldable(const MIR::FInstruction* Instr, int32 EntryPointIndex)
 		case MIR::VK_Call:
 		case MIR::VK_SubstrateHorizontalMixing:
 		case MIR::VK_SubstrateVerticalLayering:
+		case MIR::VK_SubstrateCoverageWeight:
 		case MIR::VK_SubstratePromoteToOperator:
 		case MIR::VK_SubstrateSlab:
 		case MIR::VK_SubstrateShadingModels:
@@ -542,6 +543,7 @@ struct FPrivate : FMaterialIRToHLSLTranslation
 			case MIR::VK_SubstrateShadingModels: LowerSubstrateShadingModel(static_cast<const MIR::FSubstrateShadingModels*>(Instr)); break;
 			case MIR::VK_SubstrateHorizontalMixing: LowerSubstrateHorizontalMixing(static_cast<const MIR::FSubstrateHorizontalMixing*>(Instr)); break;
 			case MIR::VK_SubstrateVerticalLayering: LowerSubstrateVerticalLayering(static_cast<const MIR::FSubstrateVerticalLayering*>(Instr)); break;
+			case MIR::VK_SubstrateCoverageWeight: LowerSubstrateCoverageWeight(static_cast<const MIR::FSubstrateCoverageWeight*>(Instr)); break;
 			case MIR::VK_SubstratePromoteToOperator: LowerSubstratePromoteToOperator(static_cast<const MIR::FSubstratePromoteToOperator*>(Instr)); break;
 			default: UE_MIR_UNREACHABLE();
 		}
@@ -2161,6 +2163,28 @@ struct FPrivate : FMaterialIRToHLSLTranslation
 		}
 	}
 
+	void LowerSubstrateCoverageWeight(const MIR::FSubstrateCoverageWeight* SubstrateParameter)
+	{
+		if (SubstrateParameter->bParameterBlendingEnabled)
+		{
+			Printer << TEXTVIEW("SubstrateWeightParameterBlending");
+			Printer << BeginArgs;
+			Printer << LowerValue(SubstrateParameter->A) << TEXTVIEW(",");
+			Printer << LowerValue(SubstrateParameter->Weight);
+			Printer << EndArgs;
+		}
+		else
+		{
+			Printer << GetSubstrateTree(SubstrateParameter->bFullySimplifiedInstruction) << TEXTVIEW(".SubstrateWeight");
+			Printer << BeginArgs;
+			Printer << LowerValue(SubstrateParameter->A) << TEXTVIEW(",");
+			Printer << LowerValue(SubstrateParameter->Weight) << TEXTVIEW(",");
+			Printer << LowerValue(SubstrateParameter->OperatorIndex) << TEXTVIEW(",");
+			Printer << LowerValue(SubstrateParameter->MaxDistanceFromLeaves);
+			Printer << EndArgs;
+		}
+	}
+
 	void LowerSubstratePromoteToOperator(const MIR::FSubstratePromoteToOperator* SubstrateParameter)
 	{
 		Printer << GetSubstrateTree(SubstrateParameter->bFullySimplifiedInstruction) << TEXTVIEW(".PromoteParameterBlendedBSDFToOperator");
diff --git a/Engine/Source/Runtime/Engine/Public/Materials/MaterialExpressionSubstrate.h b/Engine/Source/Runtime/Engine/Public/Materials/MaterialExpressionSubstrate.h
--- a/Engine/Source/Runtime/Engine/Public/Materials/MaterialExpressionSubstrate.h
+++ b/Engine/Source/Runtime/Engine/Public/Materials/MaterialExpressionSubstrate.h
@@ -1262,6 +1262,7 @@ class UMaterialExpressionSubstrateWeight : public UMaterialExpressionSubstrateBS
 	virtual bool IsResultSubstrateMaterial(int32 OutputIndex) override;
 	virtual void GatherSubstrateMaterialInfo(FSubstrateMaterialInfo& SubstrateMaterialInfo, int32 OutputIndex) override;
 	virtual FSubstrateOperator* SubstrateGenerateMaterialTopologyTree(struct FSubstrateTranslatorDataInterface& SubstrateTranslatorData, class UMaterialExpression* Parent, int32 OutputIndex) override;
+	virtual void Build(MIR::FEmitter& Emitter) override;
 #endif
 	//~ End UMaterialExpression Interface
 };
diff --git a/Engine/Source/Runtime/Engine/Public/Materials/MaterialIR.h b/Engine/Source/Runtime/Engine/Public/Materials/MaterialIR.h
--- a/Engine/Source/Runtime/Engine/Public/Materials/MaterialIR.h
+++ b/Engine/Source/Runtime/Engine/Public/Materials/MaterialIR.h
@@ -130,6 +130,7 @@ enum EValueKind
 	VK_SubstrateShadingModels,
 	VK_SubstrateHorizontalMixing,
 	VK_SubstrateVerticalLayering,
+	VK_SubstrateCoverageWeight,
 	VK_SubstratePromoteToOperator,
 	VK_StageSwitch,
 	VK_HardwarePartialDerivative,
@@ -891,6 +892,20 @@ struct FSubstrateVerticalLayering : TInstruction<VK_SubstrateVerticalLayering, 7
 	bool bFullySimplifiedInstruction;
 };
 
+struct FSubstrateCoverageWeight : TInstruction<VK_SubstrateCoverageWeight, 4>
+{
+	FValue* A;
+	FValue* Weight;
+
+	// ParameterBlending = 0
+	FValue* OperatorIndex;
+	FValue* MaxDistanceFromLeaves;
+
+	// Settings
+	bool bParameterBlendingEnabled;
+	bool bFullySimplifiedInstruction;
+};
+
 struct FSubstratePromoteToOperator : TInstruction<VK_SubstratePromoteToOperator, 5>
 {
 	FValue* SubstrateDataInput;
diff --git a/Engine/Source/Runtime/Engine/Public/Materials/MaterialIREmitter.h b/Engine/Source/Runtime/Engine/Public/Materials/MaterialIREmitter.h
--- a/Engine/Source/Runtime/Engine/Public/Materials/MaterialIREmitter.h
+++ b/Engine/Source/Runtime/Engine/Public/Materials/MaterialIREmitter.h
@@ -845,6 +845,12 @@ public:
 	// Parameter blended version of SubstrateVerticalLayering.
 	FValueRef SubstrateVerticalLayeringParameterBlending(FValueRef Top, FValueRef Base, FValueRef SharedLocalBasisIndexMacro, FValueRef TopNoV, FValueRef BaseNoV);
 
+	// Get the SubstrateData resulting from the coverage weight of the A by Weight.
+	FValueRef SubstrateCoverageWeight(FValueRef A, FValueRef Weight, FValueRef OperatorIndex, FValueRef MaxDistanceFromLeaves);
+
+	// Parameter blended version of SubstrateCoverageWeight.
+	FValueRef SubstrateCoverageWeightParameterBlending(FValueRef A, FValueRef Weight);
+
 	/*---------------------------- Pass through accessors -----------------------------*/
 
 	// Returns the shader platform the source material is translating for.
