Substrate - NMT - added support for add node.

#tests PC editor dx12
#rnx

#Synced-CL 50037646

diff --git a/Engine/Source/Runtime/Engine/Private/Materials/MaterialExpressionsIR.cpp b/Engine/Source/Runtime/Engine/Private/Materials/MaterialExpressionsIR.cpp
--- a/Engine/Source/Runtime/Engine/Private/Materials/MaterialExpressionsIR.cpp
+++ b/Engine/Source/Runtime/Engine/Private/Materials/MaterialExpressionsIR.cpp
@@ -5579,7 +5579,7 @@ void UMaterialExpressionSubstrateHorizontalMixing::Build(MIR::FEmitter& Em)
 		FSubstrateOperator* ForegroundBSDFOperator = SubstrateTranslatorData->SubstrateCompilationGetOperatorFromIndex(SubstrateOperator.RightIndex);
 		if (!BackgroundBSDFOperator || !ForegroundBSDFOperator)
 		{
-			return Em.Errorf(TEXT("Missing input on horizontal blending node."));
+			return Em.Errorf(TEXT("Missing input on Horizontal blending node."));
 		}
 
 		FSubstrateRegisteredSharedLocalBasisSearchResult ForegroundBasis = SubstrateCompilationFindExistingSharedLocalBasis(SubstrateTranslatorData, ForegroundBSDFOperator->BSDFRegisteredSharedLocalBasis.NormalCodeChunkHash, ForegroundBSDFOperator->BSDFRegisteredSharedLocalBasis.TangentCodeChunkHash);
@@ -5724,7 +5724,7 @@ void UMaterialExpressionSubstrateVerticalLayering::Build(MIR::FEmitter& Em)
 		FSubstrateOperator* BaseBSDFOperator = SubstrateTranslatorData->SubstrateCompilationGetOperatorFromIndex(SubstrateOperator.RightIndex);
 		if (!TopBSDFOperator || !BaseBSDFOperator)
 		{
-			return Em.Errorf(TEXT("Missing input on horizontal blending node."));
+			return Em.Errorf(TEXT("Missing input on Vertical blending node."));
 		}
 
 		FSubstrateRegisteredSharedLocalBasisSearchResult BaseBasis = SubstrateCompilationFindExistingSharedLocalBasis(SubstrateTranslatorData, BaseBSDFOperator->BSDFRegisteredSharedLocalBasis.NormalCodeChunkHash, BaseBSDFOperator->BSDFRegisteredSharedLocalBasis.TangentCodeChunkHash);
@@ -5892,6 +5892,145 @@ void UMaterialExpressionSubstrateWeight::Build(MIR::FEmitter& Em)
 	}
 }
 
+void UMaterialExpressionSubstrateAdd::Build(MIR::FEmitter& Em)
+{
+	FSubstrateTranslatorData* SubstrateTranslatorData = Em.GetSubstrateTranslatorData();
+
+	if (!A.GetTracedInput().Expression)
+	{
+		return Em.Errorf(TEXT("Missing A input"));
+	}
+	if (!B.GetTracedInput().Expression)
+	{
+		return Em.Errorf(TEXT("Missing B input"));
+	}
+
+	// Note: we no longer do SubstrateTreeStackPush/Pop since the path guid is regenerated for every substrate node from BuildTopMaterialExpression.
+	FValueRef BArg = Em.TryInput(&B);
+	FValueRef AArg = Em.TryInput(&A);
+
+	FSubstrateOperator& SubstrateOperator = SubstrateTranslatorData->SubstrateCompilationGetOperator(SubstrateTranslatorData->SubstrateTreeStackGetPathUniqueId());
+	if (SubstrateOperator.bUseParameterBlending)
+	{
+		if (!AArg.IsValid())
+		{
+			return Em.Errorf(TEXT("A input graphs could not be evaluated for parameter blending."));
+		}
+		if (!BArg.IsValid())
+		{
+			return Em.Errorf(TEXT("B input graphs could not be evaluated for parameter blending."));
+		}
+
+		FValueRef NormalMixArg = Em.InlineHLSL(MIR::FType::MakeFloatScalar(), TEXTVIEW("AddParameterBlendingBSDFCoverageToNormalMix($0, $1)"),
+			{
+				BArg,
+				AArg
+			}
+		);
+
+		FSubstrateOperator* BBSDFOperator = SubstrateTranslatorData->SubstrateCompilationGetOperatorFromIndex(SubstrateOperator.LeftIndex);
+		FSubstrateOperator* ABSDFOperator = SubstrateTranslatorData->SubstrateCompilationGetOperatorFromIndex(SubstrateOperator.RightIndex);
+		if (!BBSDFOperator || !ABSDFOperator)
+		{
+			return Em.Errorf(TEXT("Missing input on Add blending node."));
+		}
+
+		FSubstrateRegisteredSharedLocalBasisSearchResult ABasis = SubstrateCompilationFindExistingSharedLocalBasis(SubstrateTranslatorData, ABSDFOperator->BSDFRegisteredSharedLocalBasis.NormalCodeChunkHash, ABSDFOperator->BSDFRegisteredSharedLocalBasis.TangentCodeChunkHash);
+		FSubstrateRegisteredSharedLocalBasisSearchResult BBasis = SubstrateCompilationFindExistingSharedLocalBasis(SubstrateTranslatorData, BBSDFOperator->BSDFRegisteredSharedLocalBasis.NormalCodeChunkHash, BBSDFOperator->BSDFRegisteredSharedLocalBasis.TangentCodeChunkHash);
+
+		if (!ABasis.bFound || !BBasis.bFound)
+		{
+			Em.Errorf(TEXT("A & B shared local basis could not be found."));
+			return;
+		}
+
+		// Compute the new Normal and Tangent resulting from the blending using code chunk
+		FValueRef NewNormalArg = SubstrateBlendNormal(Em, BBasis.BasisInfo.NormalArg, ABasis.BasisInfo.NormalArg, NormalMixArg);
+
+		// The tangent is optional so we treat it differently if INDEX_NONE is specified
+		FValueRef NewTangentArg;
+		if (ABasis.BasisInfo.TangentArg && BBasis.BasisInfo.TangentArg)
+		{
+			NewTangentArg = SubstrateBlendNormal(Em, BBasis.BasisInfo.TangentArg, ABasis.BasisInfo.TangentArg, NormalMixArg);
+		}
+		else if (ABasis.BasisInfo.TangentArg)
+		{
+			NewTangentArg = ABasis.BasisInfo.TangentArg;
+		}
+		else if (BBasis.BasisInfo.TangentArg)
+		{
+			NewTangentArg = BBasis.BasisInfo.TangentArg;
+		}
+		const FSubstrateRegisteredSharedLocalBasis NewRegisteredSharedLocalBasis = MIRSubstrateCompilationInfoCreateSharedLocalBasis2(SubstrateTranslatorData, &NewNormalArg, NewTangentArg.IsValid() ? &NewTangentArg : nullptr);
+
+		FString SharedLocalBasisName = GetSubstrateSharedLocalBasisIndexMacroInner(NewRegisteredSharedLocalBasis, SubstrateTranslatorData->CurrentSubstrateCompilationContext);
+		FValueRef SharedLocalBasisIndexMacroArg = Em.InlineHLSL(MIR::FType::MakeIntScalar(), *SharedLocalBasisName);
+
+		if (!AArg.IsValid() || !BArg.IsValid() || !NormalMixArg.IsValid())
+		{
+			Em.Errorf(TEXT("Invalid Add operator A, B, Mix or NormalMix input."));
+		}
+
+		FValueRef CameraVectorArg = Em.InlineHLSL(NAME_CameraVector);
+		FValueRef BNoV = Em.Saturate(Em.Dot(BBasis.BasisInfo.NormalArg, CameraVectorArg));
+		FValueRef ANoV = Em.Saturate(Em.Dot(ABasis.BasisInfo.NormalArg, CameraVectorArg));
+
+		const FSubstrateOperator* PromoteToOperator = SubstrateOperator.bRootOfParameterBlendingSubTree ? &SubstrateOperator : nullptr;
+		if (PromoteToOperator)
+		{
+			check(PromoteToOperator->Index != INDEX_NONE);
+			check(PromoteToOperator->BSDFIndex != INDEX_NONE);
+
+			FValueRef AddedSlabs = Em.SubstrateAddParameterBlending(
+				BArg,
+				AArg,
+				NormalMixArg,
+				SharedLocalBasisIndexMacroArg,
+				BNoV,
+				ANoV);
+
+			FValueRef PromotedSlab = Em.SubstratePromoteToOperator(
+				AddedSlabs,
+				Em.ConstantInt(PromoteToOperator->Index),
+				Em.ConstantInt(PromoteToOperator->BSDFIndex),
+				Em.ConstantInt(PromoteToOperator->LayerDepth),
+				Em.ConstantInt(PromoteToOperator->bIsBottom ? 1 : 0));
+
+			Em.Output(0, PromotedSlab);
+		}
+		else
+		{
+			FValueRef AddedSlabs = Em.SubstrateAddParameterBlending(
+				BArg,
+				AArg,
+				NormalMixArg,
+				SharedLocalBasisIndexMacroArg,
+				BNoV,
+				ANoV);
+
+			Em.Output(0, AddedSlabs);
+		}
+
+		// Propagate the parameter blended normal
+		SubstrateOperator.BSDFRegisteredSharedLocalBasis = NewRegisteredSharedLocalBasis;
+	}
+	else
+	{
+		if (!AArg.IsValid() || !BArg.IsValid())
+		{
+			Em.Errorf(TEXT("Invalid Horizontal operator A, B or Mix input."));
+		}
+
+		FValueRef AddedSlabs = Em.SubstrateAdd(
+			BArg,
+			AArg,
+			Em.ConstantInt(SubstrateOperator.Index),
+			Em.ConstantInt(SubstrateOperator.MaxDistanceFromLeaves));
+
+		Em.Output(0, AddedSlabs);
+	}
+}
+
 void UMaterialExpressionSubstrateTransmittanceToMFP::Build(MIR::FEmitter& Em)
 {
 	FSubstrateTranslatorData* SubstrateTranslatorData = Em.GetSubstrateTranslatorData();
diff --git a/Engine/Source/Runtime/Engine/Private/Materials/MaterialIR.cpp b/Engine/Source/Runtime/Engine/Private/Materials/MaterialIR.cpp
--- a/Engine/Source/Runtime/Engine/Private/Materials/MaterialIR.cpp
+++ b/Engine/Source/Runtime/Engine/Private/Materials/MaterialIR.cpp
@@ -54,6 +54,7 @@ const TCHAR* LexToString(EValueKind Kind)
 		case VK_SubstrateHorizontalMixing: return TEXT("SubstrateHorizontalMixing");
 		case VK_SubstrateVerticalLayering: return TEXT("SubstrateVerticalLayering");
 		case VK_SubstrateCoverageWeight: return TEXT("SubstrateCoverageWeight");
+		case VK_SubstrateAdd: return TEXT("SubstrateAdd");
 		case VK_SubstratePromoteToOperator: return TEXT("SubstratePromoteToOperator");
 		case VK_StageSwitch: return TEXT("StageSwitch");
 		case VK_HardwarePartialDerivative: return TEXT("HardwarePartialDerivative");
@@ -128,6 +129,7 @@ uint32 FValue::GetSizeInBytes() const
 		case VK_SubstrateHorizontalMixing: return sizeof(FSubstrateHorizontalMixing);
 		case VK_SubstrateVerticalLayering: return sizeof(FSubstrateVerticalLayering);
 		case VK_SubstrateCoverageWeight: return sizeof(FSubstrateCoverageWeight);
+		case VK_SubstrateAdd: return sizeof(FSubstrateAdd);
 		case VK_SubstratePromoteToOperator: return sizeof(FSubstratePromoteToOperator);
 		case VK_StageSwitch: return sizeof(FStageSwitch);
 		case VK_HardwarePartialDerivative: return sizeof(FHardwarePartialDerivative);
@@ -243,6 +245,12 @@ TConstArrayView<FValue*> FValue::GetUses() const
 			return { &This->A, FSubstrateCoverageWeight::NumStaticUses };
 		}
 
+		case VK_SubstrateAdd:
+		{
+			auto This = static_cast<const FSubstrateAdd*>(this);
+			return { &This->A, FSubstrateAdd::NumStaticUses };
+		}
+
 		case VK_SubstratePromoteToOperator:
 		{
 			auto This = static_cast<const FSubstratePromoteToOperator*>(this);
diff --git a/Engine/Source/Runtime/Engine/Private/Materials/MaterialIREmitter.cpp b/Engine/Source/Runtime/Engine/Private/Materials/MaterialIREmitter.cpp
--- a/Engine/Source/Runtime/Engine/Private/Materials/MaterialIREmitter.cpp
+++ b/Engine/Source/Runtime/Engine/Private/Materials/MaterialIREmitter.cpp
@@ -4516,6 +4516,52 @@ FValueRef FEmitter::SubstrateCoverageWeightParameterBlending(FValueRef A, FValue
 	return EmitPrototype(*this, Prototype);
 }
 
+FValueRef FEmitter::SubstrateAdd(FValueRef A, FValueRef B, FValueRef OperatorIndex, FValueRef MaxDistanceFromLeaves)
+{
+	FSubstrateAdd Prototype = MakePrototype<FSubstrateAdd>(FType::MakeSubstrateData());
+	Prototype.A = A;
+	Prototype.B = B;
+
+	// ParameterBlending = 0
+	Prototype.OperatorIndex = OperatorIndex;
+	Prototype.MaxDistanceFromLeaves = MaxDistanceFromLeaves;
+
+	// ParameterBlending = 1
+	Prototype.NormalMix = nullptr;
+	Prototype.SharedLocalBasisIndexMacro = nullptr;
+	Prototype.ANoV = nullptr;
+	Prototype.BNoV = nullptr;
+
+	// Settings
+	Prototype.bParameterBlendingEnabled = false;
+	Prototype.bFullySimplifiedInstruction = SubstrateTranslatorData->CurrentSubstrateCompilationContext == Substrate::ESubstrateCompilationContext::SCC_FullySimplified;
+
+	return EmitPrototype(*this, Prototype);
+}
+
+FValueRef FEmitter::SubstrateAddParameterBlending(FValueRef A, FValueRef B, FValueRef NormalMix, FValueRef SharedLocalBasisIndexMacro, FValueRef ANoV, FValueRef BNoV)
+{
+	FSubstrateAdd Prototype = MakePrototype<FSubstrateAdd>(FType::MakeSubstrateData());
+	Prototype.A = A;
+	Prototype.B = B;
+
+	// ParameterBlending = 0
+	Prototype.OperatorIndex = nullptr;
+	Prototype.MaxDistanceFromLeaves = nullptr;
+
+	// ParameterBlending = 1
+	Prototype.NormalMix = NormalMix;
+	Prototype.SharedLocalBasisIndexMacro = SharedLocalBasisIndexMacro;
+	Prototype.ANoV = ANoV;
+	Prototype.BNoV = BNoV;
+
+	// Settings
+	Prototype.bParameterBlendingEnabled = true;
+	Prototype.bFullySimplifiedInstruction = SubstrateTranslatorData->CurrentSubstrateCompilationContext == Substrate::ESubstrateCompilationContext::SCC_FullySimplified;
+
+	return EmitPrototype(*this, Prototype);
+}
+
 EShaderPlatform FEmitter::GetShaderPlatform() const
 {
 	return Module->GetShaderPlatform();
diff --git a/Engine/Source/Runtime/Engine/Private/Materials/MaterialIRToHLSLTranslator.cpp b/Engine/Source/Runtime/Engine/Private/Materials/MaterialIRToHLSLTranslator.cpp
--- a/Engine/Source/Runtime/Engine/Private/Materials/MaterialIRToHLSLTranslator.cpp
+++ b/Engine/Source/Runtime/Engine/Private/Materials/MaterialIRToHLSLTranslator.cpp
@@ -306,6 +306,7 @@ static bool IsFoldable(const MIR::FInstruction* Instr, int32 EntryPointIndex)
 		case MIR::VK_SubstrateHorizontalMixing:
 		case MIR::VK_SubstrateVerticalLayering:
 		case MIR::VK_SubstrateCoverageWeight:
+		case MIR::VK_SubstrateAdd:
 		case MIR::VK_SubstratePromoteToOperator:
 		case MIR::VK_SubstrateSlab:
 		case MIR::VK_SubstrateShadingModels:
@@ -544,6 +545,7 @@ struct FPrivate : FMaterialIRToHLSLTranslation
 			case MIR::VK_SubstrateHorizontalMixing: LowerSubstrateHorizontalMixing(static_cast<const MIR::FSubstrateHorizontalMixing*>(Instr)); break;
 			case MIR::VK_SubstrateVerticalLayering: LowerSubstrateVerticalLayering(static_cast<const MIR::FSubstrateVerticalLayering*>(Instr)); break;
 			case MIR::VK_SubstrateCoverageWeight: LowerSubstrateCoverageWeight(static_cast<const MIR::FSubstrateCoverageWeight*>(Instr)); break;
+			case MIR::VK_SubstrateAdd: LowerSubstrateAdd(static_cast<const MIR::FSubstrateAdd*>(Instr)); break;
 			case MIR::VK_SubstratePromoteToOperator: LowerSubstratePromoteToOperator(static_cast<const MIR::FSubstratePromoteToOperator*>(Instr)); break;
 			default: UE_MIR_UNREACHABLE();
 		}
@@ -2185,6 +2187,32 @@ struct FPrivate : FMaterialIRToHLSLTranslation
 		}
 	}
 
+	void LowerSubstrateAdd(const MIR::FSubstrateAdd* SubstrateParameter)
+	{
+		if (SubstrateParameter->bParameterBlendingEnabled)
+		{
+			Printer << TEXTVIEW("SubstrateAddParameterBlending");
+			Printer << BeginArgs;
+			Printer << LowerValue(SubstrateParameter->A) << TEXTVIEW(",");
+			Printer << LowerValue(SubstrateParameter->B) << TEXTVIEW(",");
+			Printer << LowerValue(SubstrateParameter->NormalMix) << TEXTVIEW(",");
+			Printer << LowerValue(SubstrateParameter->SharedLocalBasisIndexMacro) << TEXTVIEW(",");
+			Printer << LowerValue(SubstrateParameter->ANoV) << TEXTVIEW(",");
+			Printer << LowerValue(SubstrateParameter->BNoV);
+			Printer << EndArgs;
+		}
+		else
+		{
+			Printer << GetSubstrateTree(SubstrateParameter->bFullySimplifiedInstruction) << TEXTVIEW(".SubstrateAdd");
+			Printer << BeginArgs;
+			Printer << LowerValue(SubstrateParameter->A) << TEXTVIEW(",");
+			Printer << LowerValue(SubstrateParameter->B) << TEXTVIEW(",");
+			Printer << LowerValue(SubstrateParameter->OperatorIndex) << TEXTVIEW(",");
+			Printer << LowerValue(SubstrateParameter->MaxDistanceFromLeaves);
+			Printer << EndArgs;
+		}
+	}
+
 	void LowerSubstratePromoteToOperator(const MIR::FSubstratePromoteToOperator* SubstrateParameter)
 	{
 		Printer << GetSubstrateTree(SubstrateParameter->bFullySimplifiedInstruction) << TEXTVIEW(".PromoteParameterBlendedBSDFToOperator");
diff --git a/Engine/Source/Runtime/Engine/Public/Materials/MaterialExpressionSubstrate.h b/Engine/Source/Runtime/Engine/Public/Materials/MaterialExpressionSubstrate.h
--- a/Engine/Source/Runtime/Engine/Public/Materials/MaterialExpressionSubstrate.h
+++ b/Engine/Source/Runtime/Engine/Public/Materials/MaterialExpressionSubstrate.h
@@ -507,6 +507,7 @@ class UMaterialExpressionSubstrateSimpleClearCoatBSDF : public UMaterialExpressi
 	virtual FSubstrateOperator* SubstrateGenerateMaterialTopologyTree(struct FSubstrateTranslatorDataInterface& SubstrateTranslatorData, class UMaterialExpression* Parent, int32 OutputIndex) override;
 	virtual FName GetInputName(int32 InputIndex) const override;
 	virtual void GetExpressionToolTip(TArray<FString>& OutToolTip) override;
+	// SUBSTRATE_TODO virtual void Build(MIR::FEmitter& Emitter) override;
 #endif
 	//~ End UMaterialExpression Interface
 };
@@ -553,6 +554,7 @@ class UMaterialExpressionSubstrateVolumetricFogCloudBSDF : public UMaterialExpre
 	virtual bool IsResultSubstrateMaterial(int32 OutputIndex) override;
 	virtual void GatherSubstrateMaterialInfo(FSubstrateMaterialInfo& SubstrateMaterialInfo, int32 OutputIndex) override;
 	virtual FSubstrateOperator* SubstrateGenerateMaterialTopologyTree(struct FSubstrateTranslatorDataInterface& SubstrateTranslatorData, class UMaterialExpression* Parent, int32 OutputIndex) override;
+	// SUBSTRATE_TODO virtual void Build(MIR::FEmitter& Emitter) override;
 
 	UE_DEPRECATED(5.8, "Please, SubstrateGenerateMaterialTopologyTreeCommon taking the FSubstrateTranslatorData parameter.")
 	static FSubstrateOperator* SubstrateGenerateMaterialTopologyTreeCommon(
@@ -605,6 +607,7 @@ class UMaterialExpressionSubstrateUnlitBSDF : public UMaterialExpressionSubstrat
 	virtual bool IsResultSubstrateMaterial(int32 OutputIndex) override;
 	virtual void GatherSubstrateMaterialInfo(FSubstrateMaterialInfo& SubstrateMaterialInfo, int32 OutputIndex) override;
 	virtual FSubstrateOperator* SubstrateGenerateMaterialTopologyTree(struct FSubstrateTranslatorDataInterface& SubstrateTranslatorData, class UMaterialExpression* Parent, int32 OutputIndex) override;
+	// SUBSTRATE_TODO virtual void Build(MIR::FEmitter& Emitter) override;
 #endif
 	//~ End UMaterialExpression Interface
 };
@@ -665,6 +668,7 @@ class UMaterialExpressionSubstrateHairBSDF : public UMaterialExpressionSubstrate
 	virtual bool IsResultSubstrateMaterial(int32 OutputIndex) override;
 	virtual void GatherSubstrateMaterialInfo(FSubstrateMaterialInfo& SubstrateMaterialInfo, int32 OutputIndex) override;
 	virtual FSubstrateOperator* SubstrateGenerateMaterialTopologyTree(struct FSubstrateTranslatorDataInterface& SubstrateTranslatorData, class UMaterialExpression* Parent, int32 OutputIndex) override;
+	// SUBSTRATE_TODO virtual void Build(MIR::FEmitter& Emitter) override;
 #endif
 	//~ End UMaterialExpression Interface
 };
@@ -735,6 +739,7 @@ class UMaterialExpressionSubstrateEyeBSDF : public UMaterialExpressionSubstrateB
 	virtual bool IsResultSubstrateMaterial(int32 OutputIndex) override;
 	virtual void GatherSubstrateMaterialInfo(FSubstrateMaterialInfo& SubstrateMaterialInfo, int32 OutputIndex) override;
 	virtual FSubstrateOperator* SubstrateGenerateMaterialTopologyTree(struct FSubstrateTranslatorDataInterface& SubstrateTranslatorData, class UMaterialExpression* Parent, int32 OutputIndex) override;
+	// SUBSTRATE_TODO virtual void Build(MIR::FEmitter& Emitter) override;
 #endif
 	//~ End UMaterialExpression Interface
 };
@@ -819,6 +824,7 @@ class UMaterialExpressionSubstrateSingleLayerWaterBSDF : public UMaterialExpress
 	virtual bool IsResultSubstrateMaterial(int32 OutputIndex) override;
 	virtual void GatherSubstrateMaterialInfo(FSubstrateMaterialInfo& SubstrateMaterialInfo, int32 OutputIndex) override;
 	virtual FSubstrateOperator* SubstrateGenerateMaterialTopologyTree(struct FSubstrateTranslatorDataInterface& SubstrateTranslatorData, class UMaterialExpression* Parent, int32 OutputIndex) override;
+	// SUBSTRATE_TODO virtual void Build(MIR::FEmitter& Emitter) override;
 #endif
 	//~ End UMaterialExpression Interface
 };
@@ -855,6 +861,7 @@ class UMaterialExpressionSubstrateLightFunction : public UMaterialExpressionSubs
 		FExpressionInput& Color,
 		const UMaterialEditorOnlyData* EditorOnlyData = nullptr);
 
+	// SUBSTRATE_TODO virtual void Build(MIR::FEmitter& Emitter) override;
 #endif
 	//~ End UMaterialExpression Interface
 };
@@ -913,6 +920,7 @@ class UMaterialExpressionSubstrateToonBSDF : public UMaterialExpressionSubstrate
 	virtual bool IsResultSubstrateMaterial(int32 OutputIndex) override;
 	virtual void GatherSubstrateMaterialInfo(FSubstrateMaterialInfo& SubstrateMaterialInfo, int32 OutputIndex) override;
 	virtual FSubstrateOperator* SubstrateGenerateMaterialTopologyTree(struct FSubstrateTranslatorDataInterface& SubstrateTranslatorData, class UMaterialExpression* Parent, int32 OutputIndex) override;
+	// SUBSTRATE_TODO virtual void Build(MIR::FEmitter& Emitter) override;
 #endif
 	//~ End UMaterialExpression Interface
 };
@@ -954,6 +962,8 @@ class UMaterialExpressionSubstratePostProcess : public UMaterialExpressionSubstr
 	static int32 CompileCommon(class FMaterialCompiler* Compiler,
 		FExpressionInput& Color, FExpressionInput& Opacity,
 		const UMaterialEditorOnlyData* EditorOnlyData = nullptr);
+
+	// SUBSTRATE_TODO virtual void Build(MIR::FEmitter& Emitter) override;
 #endif
 	//~ End UMaterialExpression Interface
 };
@@ -995,6 +1005,8 @@ class UMaterialExpressionSubstrateUI : public UMaterialExpressionSubstrateBSDF
 	static int32 CompileCommon(class FMaterialCompiler* Compiler,
 		FExpressionInput& Color, FExpressionInput& Opacity,
 		const UMaterialEditorOnlyData* EditorOnlyData = nullptr);
+
+	// SUBSTRATE_TODO virtual void Build(MIR::FEmitter& Emitter) override;
 #endif
 	//~ End UMaterialExpression Interface
 };
@@ -1025,6 +1037,7 @@ class UMaterialExpressionSubstrateConvertToDecal : public UMaterialExpressionSub
 	virtual bool IsResultSubstrateMaterial(int32 OutputIndex) override;
 	virtual void GatherSubstrateMaterialInfo(FSubstrateMaterialInfo& SubstrateMaterialInfo, int32 OutputIndex) override;
 	virtual FSubstrateOperator* SubstrateGenerateMaterialTopologyTree(struct FSubstrateTranslatorDataInterface& SubstrateTranslatorData, class UMaterialExpression* Parent, int32 OutputIndex) override;
+	// SUBSTRATE_TODO virtual void Build(MIR::FEmitter& Emitter) override;
 #endif
 	//~ End UMaterialExpression Interface
 };
@@ -1103,6 +1116,7 @@ class UMaterialExpressionSubstrateConvertMaterialAttributes : public UMaterialEx
 		const bool bHasSSS, USubsurfaceProfile* SSSProfile,
 		FExpressionInput& ClearCoatNormal, FExpressionInput& CustomTangent);
 
+	// SUBSTRATE_TODO virtual void Build(MIR::FEmitter& Emitter) override;
 #endif // WITH_EDITOR
 	//~ End UMaterialExpression Interface
 };
@@ -1232,6 +1246,7 @@ class UMaterialExpressionSubstrateAdd : public UMaterialExpressionSubstrateBSDF
 	virtual bool IsResultSubstrateMaterial(int32 OutputIndex) override;
 	virtual void GatherSubstrateMaterialInfo(FSubstrateMaterialInfo& SubstrateMaterialInfo, int32 OutputIndex) override;
 	virtual FSubstrateOperator* SubstrateGenerateMaterialTopologyTree(struct FSubstrateTranslatorDataInterface& SubstrateTranslatorData, class UMaterialExpression* Parent, int32 OutputIndex) override;
+	virtual void Build(MIR::FEmitter& Emitter) override;
 #endif
 	//~ End UMaterialExpression Interface
 };
@@ -1311,6 +1326,7 @@ class UMaterialExpressionSubstrateSelect : public UMaterialExpressionSubstrateBS
 	virtual bool IsResultSubstrateMaterial(int32 OutputIndex) override;
 	virtual void GatherSubstrateMaterialInfo(FSubstrateMaterialInfo& SubstrateMaterialInfo, int32 OutputIndex) override;
 	virtual FSubstrateOperator* SubstrateGenerateMaterialTopologyTree(struct FSubstrateTranslatorDataInterface& SubstrateTranslatorData, class UMaterialExpression* Parent, int32 OutputIndex) override;
+	// SUBSTRATE_TODO virtual void Build(MIR::FEmitter& Emitter) override;
 #endif
 	//~ End UMaterialExpression Interface
 };
@@ -1388,6 +1404,7 @@ class UMaterialExpressionSubstrateMetalnessToDiffuseAlbedoF0 : public UMaterialE
 	virtual EMaterialValueType GetInputValueType(int32 InputIndex) override;
 	virtual void GetConnectorToolTip(int32 InputIndex, int32 OutputIndex, TArray<FString>& OutToolTip) override;
 	virtual void GetExpressionToolTip(TArray<FString>& OutToolTip) override;
+	// SUBSTRATE_TODO virtual void Build(MIR::FEmitter& Emitter) override;
 #endif
 	//~ End UMaterialExpression Interface
 };
@@ -1417,6 +1434,7 @@ class UMaterialExpressionSubstrateHazinessToSecondaryRoughness : public UMateria
 	virtual EMaterialValueType GetInputValueType(int32 InputIndex) override;
 	virtual void GetConnectorToolTip(int32 InputIndex, int32 OutputIndex, TArray<FString>& OutToolTip) override;
 	virtual void GetExpressionToolTip(TArray<FString>& OutToolTip) override;
+	// SUBSTRATE_TODO virtual void Build(MIR::FEmitter& Emitter) override;
 #endif
 	//~ End UMaterialExpression Interface
 };
@@ -1464,6 +1482,7 @@ class UMaterialExpressionSubstrateThinFilm : public UMaterialExpressionSubstrate
 	virtual EMaterialValueType GetInputValueType(int32 InputIndex) override;
 	virtual void GetConnectorToolTip(int32 InputIndex, int32 OutputIndex, TArray<FString>& OutToolTip) override;
 	virtual void GetExpressionToolTip(TArray<FString>& OutToolTip) override;
+	// SUBSTRATE_TODO virtual void Build(MIR::FEmitter& Emitter) override;
 #endif
 	//~ End UMaterialExpression Interface
 };
\ No newline at end of file
diff --git a/Engine/Source/Runtime/Engine/Public/Materials/MaterialIR.h b/Engine/Source/Runtime/Engine/Public/Materials/MaterialIR.h
--- a/Engine/Source/Runtime/Engine/Public/Materials/MaterialIR.h
+++ b/Engine/Source/Runtime/Engine/Public/Materials/MaterialIR.h
@@ -131,6 +131,7 @@ enum EValueKind
 	VK_SubstrateHorizontalMixing,
 	VK_SubstrateVerticalLayering,
 	VK_SubstrateCoverageWeight,
+	VK_SubstrateAdd,
 	VK_SubstratePromoteToOperator,
 	VK_StageSwitch,
 	VK_HardwarePartialDerivative,
@@ -906,6 +907,26 @@ struct FSubstrateCoverageWeight : TInstruction<VK_SubstrateCoverageWeight, 4>
 	bool bFullySimplifiedInstruction;
 };
 
+struct FSubstrateAdd : TInstruction<VK_SubstrateAdd, 8>
+{
+	FValue* A;
+	FValue* B;
+
+	// ParameterBlending = 0
+	FValue* OperatorIndex;
+	FValue* MaxDistanceFromLeaves;
+
+	// ParameterBlending = 1
+	FValue* NormalMix;
+	FValue* SharedLocalBasisIndexMacro;
+	FValue* ANoV;
+	FValue* BNoV;
+
+	// Settings
+	bool bParameterBlendingEnabled;
+	bool bFullySimplifiedInstruction;
+};
+
 struct FSubstratePromoteToOperator : TInstruction<VK_SubstratePromoteToOperator, 5>
 {
 	FValue* SubstrateDataInput;
diff --git a/Engine/Source/Runtime/Engine/Public/Materials/MaterialIREmitter.h b/Engine/Source/Runtime/Engine/Public/Materials/MaterialIREmitter.h
--- a/Engine/Source/Runtime/Engine/Public/Materials/MaterialIREmitter.h
+++ b/Engine/Source/Runtime/Engine/Public/Materials/MaterialIREmitter.h
@@ -851,6 +851,12 @@ public:
 	// Parameter blended version of SubstrateCoverageWeight.
 	FValueRef SubstrateCoverageWeightParameterBlending(FValueRef A, FValueRef Weight);
 
+	// Get the SubstrateData resulting from the addition of the A abd B slabs.
+	FValueRef SubstrateAdd(FValueRef A, FValueRef B, FValueRef OperatorIndex, FValueRef MaxDistanceFromLeaves);
+
+	// Parameter blended version of SubstrateAdd.
+	FValueRef SubstrateAddParameterBlending(FValueRef A, FValueRef B, FValueRef NormalMix, FValueRef SharedLocalBasisIndexMacro, FValueRef ANoV, FValueRef BNoV);
+
 	/*---------------------------- Pass through accessors -----------------------------*/
 
 	// Returns the shader platform the source material is translating for.
