Substrate - NMT - fixed inlined hlsl to emit IR instead.

#tests PC editor dx12
#rnx

#Synced-CL 50037646

diff --git a/Engine/Source/Runtime/Engine/Private/Materials/MaterialExpressionsIR.cpp b/Engine/Source/Runtime/Engine/Private/Materials/MaterialExpressionsIR.cpp
--- a/Engine/Source/Runtime/Engine/Private/Materials/MaterialExpressionsIR.cpp
+++ b/Engine/Source/Runtime/Engine/Private/Materials/MaterialExpressionsIR.cpp
@@ -5649,17 +5649,15 @@ void UMaterialExpressionSubstrateHorizontalMixing::Build(MIR::FEmitter& Em)
 		}
 		else
 		{
-			FValueRef MixedSlab = Em.InlineHLSL(MIR::FType::MakeSubstrateData(), TEXTVIEW("SubstrateHorizontalMixingParameterBlending($0, $1, $2, $3, $4, Parameters.SharedLocalBases.Types, $5, $6)"),
-				{
-					BackgroundArg,
-					ForegroundArg,
-					HorizontalMixArg,
-					NormalMixArg,
-					SharedLocalBasisIndexMacroArg,
-					BackgroundNoV,
-					ForegroundNoV
-				}
-			);
+			FValueRef MixedSlab = Em.SubstrateHorizontalMixingParameterBlending(
+				BackgroundArg,
+				ForegroundArg,
+				HorizontalMixArg,
+				NormalMixArg,
+				SharedLocalBasisIndexMacroArg,
+				BackgroundNoV,
+				ForegroundNoV);
+
 			Em.Output(0, MixedSlab);
 		}
 
