CookOnTheFlyServer TextureFormat versioning simpler by just doing all modules
Allows for big cleanup of functions GetAllTextureFormats and GetAllDefaultTextureFormats which were only used by this one little thing

#Synced-CL 50020597

diff --git a/Engine/Source/Developer/Android/AndroidTargetPlatformControls/Private/AndroidTargetPlatformControls.cpp b/Engine/Source/Developer/Android/AndroidTargetPlatformControls/Private/AndroidTargetPlatformControls.cpp
--- a/Engine/Source/Developer/Android/AndroidTargetPlatformControls/Private/AndroidTargetPlatformControls.cpp
+++ b/Engine/Source/Developer/Android/AndroidTargetPlatformControls/Private/AndroidTargetPlatformControls.cpp
@@ -538,36 +538,6 @@ void FAndroidTargetPlatformControls::GetTextureFormats(const UTexture* Texture,
 	}
 }
 
-void FAndroidTargetPlatformControls::GetAllTextureFormats(TArray<FName>& OutFormats) const
-{
-	GetAllDefaultTextureFormats(this->GetTargetPlatformSettings(), OutFormats);
-
-	for (int32 RemapIndex = 0; RemapIndex < UE_ARRAY_COUNT(AndroidTexFormat::GenericRemap); ++RemapIndex)
-	{
-		OutFormats.Remove(AndroidTexFormat::GenericRemap[RemapIndex][0]);
-	}
-
-	for (int32 RemapIndex = 0; RemapIndex < UE_ARRAY_COUNT(AndroidTexFormat::GenericRemap); ++RemapIndex)
-	{
-		OutFormats.AddUnique(AndroidTexFormat::GenericRemap[RemapIndex][1]);
-	}
-}
-
-void FAndroid_ASTCTargetPlatformControls::GetAllTextureFormats(TArray<FName>& OutFormats) const
-{
-	FAndroidTargetPlatformControls::GetAllTextureFormats(OutFormats);
-
-	for (int32 RemapIndex = 0; RemapIndex < UE_ARRAY_COUNT(AndroidTexFormat::ASTCRemap); ++RemapIndex)
-	{
-		OutFormats.Remove(AndroidTexFormat::ASTCRemap[RemapIndex][0]);
-	}
-
-	// ASTC for compressed textures
-	OutFormats.Add(AndroidTexFormat::NameAutoASTC);
-	// ETC for ETC2_R11
-	OutFormats.Add(AndroidTexFormat::NameAutoETC2);
-}
-
 void FAndroid_ASTCTargetPlatformControls::GetTextureFormats(const UTexture* Texture, TArray< TArray<FName> >& OutFormats) const
 {
 	FAndroidTargetPlatformControls::GetTextureFormats(Texture, OutFormats);
@@ -633,19 +603,6 @@ void FAndroid_DXTTargetPlatformControls::GetTextureFormats(const UTexture* Textu
 	}
 }
 
-void FAndroid_ETC2TargetPlatformControls::GetAllTextureFormats(TArray<FName>& OutFormats) const
-{
-	FAndroidTargetPlatformControls::GetAllTextureFormats(OutFormats);
-
-	for (int32 RemapIndex = 0; RemapIndex < UE_ARRAY_COUNT(AndroidTexFormat::ETCRemap); ++RemapIndex)
-	{
-		OutFormats.Remove(AndroidTexFormat::ETCRemap[RemapIndex][0]);
-	}
-
-	// support only ETC for compressed textures
-	OutFormats.Add(AndroidTexFormat::NameAutoETC2);
-}
-
 void FAndroid_ETC2TargetPlatformControls::GetTextureFormats(const UTexture* Texture, TArray< TArray<FName> >& OutFormats) const
 {
 	FAndroidTargetPlatformControls::GetTextureFormats(Texture, OutFormats);
@@ -688,19 +645,6 @@ void FAndroid_MultiTargetPlatformControls::GetTextureFormats(const UTexture* Tex
 	}
 }
 
-void FAndroid_MultiTargetPlatformControls::GetAllTextureFormats(TArray<FName>& OutFormats) const
-{
-	// Ask each platform variant to choose texture formats
-	for (ITargetPlatformControls* Platform : FormatTargetPlatforms)
-	{
-		TArray<FName> PlatformFormats;
-		Platform->GetAllTextureFormats(PlatformFormats);
-		for (FName Format : PlatformFormats)
-		{
-			OutFormats.AddUnique(Format);
-		}
-	}
-}
 #endif
 
 FText FAndroid_MultiTargetPlatformControls::DisplayName() const
diff --git a/Engine/Source/Developer/Android/AndroidTargetPlatformControls/Private/AndroidTargetPlatformControls.h b/Engine/Source/Developer/Android/AndroidTargetPlatformControls/Private/AndroidTargetPlatformControls.h
--- a/Engine/Source/Developer/Android/AndroidTargetPlatformControls/Private/AndroidTargetPlatformControls.h
+++ b/Engine/Source/Developer/Android/AndroidTargetPlatformControls/Private/AndroidTargetPlatformControls.h
@@ -98,7 +98,6 @@ public:
 #if WITH_ENGINE
 	ANDROIDTARGETPLATFORMCONTROLS_API virtual void GetTextureFormats(const UTexture* InTexture, TArray< TArray<FName> >& OutFormats) const override;
 	ANDROIDTARGETPLATFORMCONTROLS_API virtual FName FinalizeVirtualTextureLayerFormat(FName Format) const override;
-	ANDROIDTARGETPLATFORMCONTROLS_API virtual void GetAllTextureFormats(TArray<FName>& OutFormats) const override;
 #endif //WITH_ENGINE
 
 	ANDROIDTARGETPLATFORMCONTROLS_API virtual bool SupportsVariants() const override;
@@ -203,7 +202,6 @@ public:
 
 #if WITH_ENGINE
 	virtual void GetTextureFormats(const UTexture* Texture, TArray< TArray<FName> >& OutFormats) const;
-	virtual void GetAllTextureFormats(TArray<FName>& OutFormats) const override;
 #endif
 };
 
@@ -239,7 +237,6 @@ public:
 
 #if WITH_ENGINE
 	virtual void GetTextureFormats(const UTexture* Texture, TArray< TArray<FName> >& OutFormats) const;
-	virtual void GetAllTextureFormats(TArray<FName>& OutFormats) const override;
 #endif
 };
 
@@ -326,7 +323,6 @@ public:
 
 #if WITH_ENGINE
 	virtual void GetTextureFormats(const UTexture* Texture, TArray< TArray<FName> >& OutFormats) const;
-	virtual void GetAllTextureFormats(TArray<FName>& OutFormats) const override;
 #endif	
 	virtual FText DisplayName() const override;
 };
diff --git a/Engine/Source/Developer/CookedEditor/Public/CookedEditorTargetPlatformControls.h b/Engine/Source/Developer/CookedEditor/Public/CookedEditorTargetPlatformControls.h
--- a/Engine/Source/Developer/CookedEditor/Public/CookedEditorTargetPlatformControls.h
+++ b/Engine/Source/Developer/CookedEditor/Public/CookedEditorTargetPlatformControls.h
@@ -178,11 +178,6 @@ public:
 		// no textures please
 	}
 
-	virtual void GetAllTextureFormats(TArray<FName>& OutFormats) const override
-	{
-		// no textures please
-	}
-
 	virtual bool SupportsVariants() const override
 	{
 		return false;
diff --git a/Engine/Source/Developer/IOS/IOSTargetPlatformControls/Private/IOSTargetPlatformControls.cpp b/Engine/Source/Developer/IOS/IOSTargetPlatformControls/Private/IOSTargetPlatformControls.cpp
--- a/Engine/Source/Developer/IOS/IOSTargetPlatformControls/Private/IOSTargetPlatformControls.cpp
+++ b/Engine/Source/Developer/IOS/IOSTargetPlatformControls/Private/IOSTargetPlatformControls.cpp
@@ -578,24 +578,6 @@ void FIOSTargetPlatformControls::GetTextureFormats( const UTexture* Texture, TAr
 	}
 }
 
-void FIOSTargetPlatformControls::GetAllTextureFormats(TArray<FName>& OutFormats) const 
-{
-	bool bFoundRemap = false;
-
-	GetAllDefaultTextureFormats(this->GetTargetPlatformSettings(), OutFormats);
-
-	for (int32 RemapIndex = 0; RemapIndex < UE_ARRAY_COUNT(FormatRemap); RemapIndex += 2)
-	{
-		OutFormats.Remove(FormatRemap[RemapIndex+0]);
-	}
-
-	// include the formats we want
-	for (int32 RemapIndex = 0; RemapIndex < UE_ARRAY_COUNT(FormatRemap); RemapIndex += 2)
-	{
-		OutFormats.AddUnique(FormatRemap[RemapIndex + 1]);
-	}
-}
-
 FName FIOSTargetPlatformControls::FinalizeVirtualTextureLayerFormat(FName Format) const
 {
 #if WITH_EDITOR
diff --git a/Engine/Source/Developer/IOS/IOSTargetPlatformControls/Private/IOSTargetPlatformControls.h b/Engine/Source/Developer/IOS/IOSTargetPlatformControls/Private/IOSTargetPlatformControls.h
--- a/Engine/Source/Developer/IOS/IOSTargetPlatformControls/Private/IOSTargetPlatformControls.h
+++ b/Engine/Source/Developer/IOS/IOSTargetPlatformControls/Private/IOSTargetPlatformControls.h
@@ -62,8 +62,6 @@ public:
 
 	virtual void GetTextureFormats( const UTexture* Texture, TArray< TArray<FName> >& OutFormats) const override;
 
-	virtual void GetAllTextureFormats( TArray<FName>& OutFormats) const override;
-
 	virtual FName FinalizeVirtualTextureLayerFormat(FName Format) const override;
 
 #endif // WITH_ENGINE
diff --git a/Engine/Source/Developer/Linux/LinuxTargetPlatformControls/Public/LinuxTargetPlatformControls.h b/Engine/Source/Developer/Linux/LinuxTargetPlatformControls/Public/LinuxTargetPlatformControls.h
--- a/Engine/Source/Developer/Linux/LinuxTargetPlatformControls/Public/LinuxTargetPlatformControls.h
+++ b/Engine/Source/Developer/Linux/LinuxTargetPlatformControls/Public/LinuxTargetPlatformControls.h
@@ -239,15 +239,6 @@ public:
 		}
 	}
 
-	virtual void GetAllTextureFormats(TArray<FName>& OutFormats) const override
-	{
-		if (this->AllowAudioVisualData())
-		{
-			// just use the standard texture format name for this texture
-			GetAllDefaultTextureFormats(this->GetTargetPlatformSettings(), OutFormats);
-		}
-	}
-
 #endif //WITH_ENGINE
 
 	virtual bool SupportsVariants() const override
diff --git a/Engine/Source/Developer/Mac/MacTargetPlatformControls/Public/GenericMacTargetPlatformControls.h b/Engine/Source/Developer/Mac/MacTargetPlatformControls/Public/GenericMacTargetPlatformControls.h
--- a/Engine/Source/Developer/Mac/MacTargetPlatformControls/Public/GenericMacTargetPlatformControls.h
+++ b/Engine/Source/Developer/Mac/MacTargetPlatformControls/Public/GenericMacTargetPlatformControls.h
@@ -100,15 +100,6 @@ public:
 		}
 	}
 
-	virtual void GetAllTextureFormats(TArray<FName>& OutFormats) const override
-	{
-		if (!IS_DEDICATED_SERVER)
-		{
-			// just use the standard texture format name for this texture (with DX11 support)
-			GetAllDefaultTextureFormats(this->GetTargetPlatformSettings(), OutFormats);
-		}
-	}
-
 	virtual bool SupportsLQCompressionTextureFormat() const override { return false; };
 
 	virtual bool CanSupportRemoteShaderCompile() const override
diff --git a/Engine/Source/Developer/TargetPlatform/Public/Common/TargetPlatformBase.h b/Engine/Source/Developer/TargetPlatform/Public/Common/TargetPlatformBase.h
--- a/Engine/Source/Developer/TargetPlatform/Public/Common/TargetPlatformBase.h
+++ b/Engine/Source/Developer/TargetPlatform/Public/Common/TargetPlatformBase.h
@@ -905,10 +905,6 @@ public:
 	{
 		TargetPlatformControls->GetTextureFormats(Texture, OutFormats);
 	}
-	virtual void GetAllTextureFormats(TArray<FName>& OutFormats) const override 
-	{
-		TargetPlatformControls->GetAllTextureFormats(OutFormats);
-	}
 	
 	virtual void GetTextureSizeLimits(uint64 & OutMaximumSurfaceBytes, uint64 & OutMaximumPackageBytes) const override
 	{
diff --git a/Engine/Source/Developer/TargetPlatform/Public/Interfaces/ITargetPlatformControls.h b/Engine/Source/Developer/TargetPlatform/Public/Interfaces/ITargetPlatformControls.h
--- a/Engine/Source/Developer/TargetPlatform/Public/Interfaces/ITargetPlatformControls.h
+++ b/Engine/Source/Developer/TargetPlatform/Public/Interfaces/ITargetPlatformControls.h
@@ -361,7 +361,8 @@ public:
 	 *
 	 * @param OutFormats will contain all the texture formats which are possible for this platform
 	 */
-	virtual void GetAllTextureFormats(TArray<FName>& OutFormats) const = 0;
+	UE_DEPRECATED(5.9, "GetAllTextureFormats is deprecated.")
+	virtual void GetAllTextureFormats(TArray<FName>& OutFormats) const { }
 
 	/**
 	 * GetTextureFormatModuleHints is not used and should be deleted
diff --git a/Engine/Source/Developer/TextureFormat/Private/TextureFormatManagerModule.cpp b/Engine/Source/Developer/TextureFormat/Private/TextureFormatManagerModule.cpp
--- a/Engine/Source/Developer/TextureFormat/Private/TextureFormatManagerModule.cpp
+++ b/Engine/Source/Developer/TextureFormat/Private/TextureFormatManagerModule.cpp
@@ -160,6 +160,22 @@ public:
 		check( (int)TextureFormatsInitPhase >= (int)EInitPhase::GetTextureFormatsPartialOkayToRead );
 	}
 	
+	virtual TArray<const class ITextureFormat*> GetAllTextureFormats() override
+	{
+		// note : does not call UpdateTextureFormatList , similar to FindTextureFormat
+		//	assumes list is kept up to date
+
+		TArray<const class ITextureFormat*> Ret;
+		
+		// copy member variable to local with lock held
+		{
+		FScopeLock Lock(&ModuleMutex);
+		Ret = TextureFormats;
+		}
+		
+		return MoveTemp(Ret);
+	}
+
 	virtual const ITextureFormat* FindTextureFormat(FName Name) override
 	{
 		// just pass through to FindTextureFormatAndModule
diff --git a/Engine/Source/Developer/TextureFormat/Public/Interfaces/ITextureFormatManagerModule.h b/Engine/Source/Developer/TextureFormat/Public/Interfaces/ITextureFormatManagerModule.h
--- a/Engine/Source/Developer/TextureFormat/Public/Interfaces/ITextureFormatManagerModule.h
+++ b/Engine/Source/Developer/TextureFormat/Public/Interfaces/ITextureFormatManagerModule.h
@@ -37,8 +37,7 @@ public:
 	 *
 	 * @return Collection of texture formats.
 	 */
-	UE_DEPRECATED(5.6, "Deprecated as it's not thread safe - use FindTextureFormat")
-	TArray<const class ITextureFormat*> GetTextureFormats()	{ return {}; }
+	virtual TArray<const class ITextureFormat*> GetAllTextureFormats() = 0;
 
 	/**
 	 * Invalidates the texture format manager module.
diff --git a/Engine/Source/Developer/Windows/WindowsTargetPlatfomControls/Public/GenericWindowsTargetPlatformControls.h b/Engine/Source/Developer/Windows/WindowsTargetPlatfomControls/Public/GenericWindowsTargetPlatformControls.h
--- a/Engine/Source/Developer/Windows/WindowsTargetPlatfomControls/Public/GenericWindowsTargetPlatformControls.h
+++ b/Engine/Source/Developer/Windows/WindowsTargetPlatfomControls/Public/GenericWindowsTargetPlatformControls.h
@@ -159,14 +159,6 @@ public:
 		}
 	}
 
-	virtual void GetAllTextureFormats(TArray<FName>& OutFormats) const override
-	{
-		if (!TProperties::IsServerOnly())
-		{
-			GetAllDefaultTextureFormats(GenericWindowsTargetPlatformSettings, OutFormats);
-		}
-	}
-
 #endif //WITH_ENGINE
 
 	virtual bool SupportsVariants() const override
diff --git a/Engine/Source/Editor/UnrealEd/Private/CookOnTheFlyServer.cpp b/Engine/Source/Editor/UnrealEd/Private/CookOnTheFlyServer.cpp
--- a/Engine/Source/Editor/UnrealEd/Private/CookOnTheFlyServer.cpp
+++ b/Engine/Source/Editor/UnrealEd/Private/CookOnTheFlyServer.cpp
@@ -150,6 +150,7 @@
 #include "String/ParseTokens.h"
 #include "TargetDomain/TargetDomainUtils.h"
 #include "Templates/UnrealTemplate.h"
+#include "TextureFormatManager.h"
 #include "ThumbnailExternalCache.h"
 #include "UnrealEdGlobals.h"
 #include "UObject/ArchiveCookContext.h"
@@ -7222,32 +7223,16 @@ void GetAdditionalCurrentIniVersionStrings( const UCookOnTheFlyServer* CookOnThe
 	}
 
 	{
-
-		#if 1
-		// this is the only place that TargetPlatform::GetAllTextureFormats is used
-		// instead use ITextureFormatManagerModule::GetTextureFormats ?
-		//	then GetAllTextureFormats can be removed completely
-
-		// get all texture formats for this target platform, then find the modules that encode them
-		TArray<FName> AllTextureFormats;
-		TargetPlatform->GetAllTextureFormats(AllTextureFormats);
-		TArray<const ITextureFormat*> SupportedTextureFormats;
-		for (const auto& TextureName : AllTextureFormats)
+		// GetModule , not Load.  Should be loaded by TargetPlatform or Program startup on main thread.
+		//  cannot load modules from threads so must be done before any thread needs this module
+		ITextureFormatManagerModule* TFMM = FModuleManager::GetModulePtr<ITextureFormatManagerModule>("TextureFormat");
+		if ( TFMM == nullptr )
 		{
-			const ITextureFormat* TextureFormat = TPM->FindTextureFormat(TextureName);
-			if ( TextureFormat )
-			{
-				SupportedTextureFormats.AddUnique(TextureFormat);
-			}
-			else
-			{
-				UE_LOG(LogCook, Warning, TEXT("Unable to find texture format \"%s\" which is required by \"%s\""), *TextureName.ToString(), *TargetPlatform->PlatformName());
-			}
+			UE_LOG(LogInit, Error, TEXT("Texture format manager Ptr was requested but not available."));
 		}
-		#else
+
 		//note: this gets All ITextureFormat modules in the Engine, not just ones relevant to this TargetPlatform
-		const TArray<const ITextureFormat*> & SupportedTextureFormats = TPM->GetTextureFormats();
-		#endif
+		const TArray<const ITextureFormat*> & SupportedTextureFormats = TFMM->GetAllTextureFormats();
 		
 		GetVersionFormatNumbersForIniVersionStrings(IniVersionMap, TEXT("TextureFormat"), SupportedTextureFormats);
 	}
diff --git a/Engine/Source/Runtime/Engine/Private/Texture.cpp b/Engine/Source/Runtime/Engine/Private/Texture.cpp
--- a/Engine/Source/Runtime/Engine/Private/Texture.cpp
+++ b/Engine/Source/Runtime/Engine/Private/Texture.cpp
@@ -4172,10 +4172,6 @@ FName GetDefaultTextureFormatName( const ITargetPlatformSettings* TargetPlatform
 {
 	FName TextureFormatName = NAME_None;
 
-	/**
-	 * IF you add a format to this function don't forget to update GetAllDefaultTextureFormats 
-	 */
-
 #if WITH_EDITOR
 	// Supported texture format names.
 	static FName NameDXT1("DXT1");
@@ -4466,6 +4462,14 @@ FName GetDefaultTextureFormatName(const ITargetPlatform* TargetPlatform, const U
 	return GetDefaultTextureFormatName(TargetPlatform->GetTargetPlatformSettings(), TargetPlatform->GetTargetPlatformControls(), Texture, LayerIndex, bSupportCompressedVolumeTexture, Unused_BlockSize, bSupportFilteredFloat32Textures);
 }
 
+void GetAllDefaultTextureFormats( const class ITargetPlatform* TargetPlatform, TArray<FName>& OutFormats)
+{
+}
+
+void GetAllDefaultTextureFormats( const class ITargetPlatformSettings* TargetPlatformSettings, TArray<FName>& OutFormats)
+{
+}
+
 #if WITH_EDITOR
 // this is the entry point API to get the texture format name
 void UTexture::GetPlatformTextureFormatNamesWithPrefix(const class ITargetPlatform* TargetPlatform,TArray< TArray<FName> >& OutFormats) const
@@ -4548,69 +4552,6 @@ void GetDefaultTextureFormatNamePerLayer(TArray<FName>& OutFormatNames, const cl
 	GetDefaultTextureFormatNamePerLayer(OutFormatNames, TargetPlatform->GetTargetPlatformSettings(), TargetPlatform->GetTargetPlatformControls(), Texture, bSupportCompressedVolumeTexture, Unused_BlockSize, bSupportFilteredFloat32Textures);
 }
 
-void GetAllDefaultTextureFormats(const class ITargetPlatformSettings* TargetPlatformSettings, TArray<FName>& OutFormats)
-{
-	// this is only used by CookOnTheFlyServer, it could be removed entirely
-
-#if WITH_EDITOR
-	static FName NameDXT1("DXT1");
-	static FName NameDXT5("DXT5");
-	static FName NameDXT5n("DXT5n");
-	static FName NameAutoDXT("AutoDXT");
-	static FName NameBC4("BC4");
-	static FName NameBC5("BC5");
-	static FName NameBGRA8("BGRA8");
-	static FName NameXGXR8("XGXR8");
-	static FName NameG8("G8");
-	static FName NameG16("G16");
-	static FName NameVU8("VU8");
-	static FName NameRGBA16F("RGBA16F");
-	static FName NameRGBA32F("RGBA32F");
-	static FName NameR16F("R16F");
-	static FName NameR32F("R32F");
-	static FName NameBC6H("BC6H");
-	static FName NameBC7("BC7");
-
-	OutFormats.Add(NameDXT1);
-	OutFormats.Add(NameDXT5);
-	OutFormats.Add(NameDXT5n);
-	OutFormats.Add(NameAutoDXT);
-	OutFormats.Add(NameBC4);
-	OutFormats.Add(NameBC5);
-	OutFormats.Add(NameBGRA8);
-	OutFormats.Add(NameXGXR8);
-	OutFormats.Add(NameG8);
-	OutFormats.Add(NameG16);
-	OutFormats.Add(NameVU8);
-	OutFormats.Add(NameRGBA16F);
-	OutFormats.Add(NameRGBA32F);
-	OutFormats.Add(NameR16F);
-	OutFormats.Add(NameR32F);
-	OutFormats.Add(NameBC6H);
-	OutFormats.Add(NameBC7);
-	// is there any drawback to just adding the 32F textures here even if we don't want them? -> no.
-	//	what is this list even used for?
-	//  AFAICT it's only used by CookOnTheFlyServer for GetVersionFormatNumbersForIniVersionStrings
-
-	// go over the original base formats only, and possibly add on to the end of the array if there is a prefix needed
-	int NumBaseFormats = OutFormats.Num();
-	for (int Index = 0; Index < NumBaseFormats; Index++)
-	{
-		OutFormats.Add(ConditionalGetPrefixedFormat(OutFormats[Index], TargetPlatformSettings, true));
-		OutFormats.Add(ConditionalGetPrefixedFormat(OutFormats[Index], TargetPlatformSettings, false));
-	}
-	
-	// make unique:		
-	OutFormats.Sort( FNameFastLess() );
-	OutFormats.SetNum( Algo::Unique( OutFormats ) );
-#endif
-}
-
-void GetAllDefaultTextureFormats(const class ITargetPlatform* TargetPlatform, TArray<FName>& OutFormats)
-{
-	GetAllDefaultTextureFormats(TargetPlatform->GetTargetPlatformSettings(), OutFormats);
-}
-
 #if WITH_EDITOR
 
 void UTexture::NotifyMaterials(const ENotifyMaterialsEffectOnShaders EffectOnShaders)
diff --git a/Engine/Source/Runtime/Engine/Public/TextureResource.h b/Engine/Source/Runtime/Engine/Public/TextureResource.h
--- a/Engine/Source/Runtime/Engine/Public/TextureResource.h
+++ b/Engine/Source/Runtime/Engine/Public/TextureResource.h
@@ -690,6 +690,8 @@ ENGINE_API FName GetDefaultTextureFormatName( const class ITargetPlatformSetting
 ENGINE_API void GetDefaultTextureFormatNamePerLayer(TArray<FName>& OutFormatNames, const class ITargetPlatform* TargetPlatform, const class UTexture* Texture, bool bSupportCompressedVolumeTexture, int32 Unused_BlockSize, bool bSupportFilteredFloat32Textures);
 ENGINE_API void GetDefaultTextureFormatNamePerLayer(TArray<FName>& OutFormatNames, const class ITargetPlatformSettings* TargetPlatformSettings, const class ITargetPlatformControls* TargetPlatformControls, const class UTexture* Texture, bool bSupportCompressedVolumeTexture, int32 Unused_BlockSize, bool bSupportFilteredFloat32Textures);
 
-// returns all the texture formats which can be returned by GetDefaultTextureFormatName
+UE_DEPRECATED(5.9, "GetAllDefaultTextureFormats is deprecated.")
 ENGINE_API void GetAllDefaultTextureFormats( const class ITargetPlatform* TargetPlatform, TArray<FName>& OutFormats);
+
+UE_DEPRECATED(5.9, "GetAllDefaultTextureFormats is deprecated.")
 ENGINE_API void GetAllDefaultTextureFormats( const class ITargetPlatformSettings* TargetPlatformSettings, TArray<FName>& OutFormats);
