Substrate - NMT - added support for MetalnessToDiffuseAlbedoF0, HazinessToSecondaryRoughness, ThinFilm.

#rnx
#tests PC editor dx12

#Synced-CL 50037646

diff --git a/Engine/Source/Runtime/Engine/Private/Materials/MaterialExpressionsIR.cpp b/Engine/Source/Runtime/Engine/Private/Materials/MaterialExpressionsIR.cpp
--- a/Engine/Source/Runtime/Engine/Private/Materials/MaterialExpressionsIR.cpp
+++ b/Engine/Source/Runtime/Engine/Private/Materials/MaterialExpressionsIR.cpp
@@ -6056,4 +6056,83 @@ void UMaterialExpressionSubstrateTransmittanceToMFP::Build(MIR::FEmitter& Em)
 	Em.Output(1, ThicknessArg);
 }
 
+void UMaterialExpressionSubstrateMetalnessToDiffuseAlbedoF0::Build(MIR::FEmitter& Em)
+{
+	FValueRef BaseColorArg	= BaseColor.GetTracedInput().Expression	? Em.TryInput(&BaseColor)	: Em.ConstantFloat(0.18f);
+	FValueRef SpecularArg	= Specular.GetTracedInput().Expression	? Em.TryInput(&Specular)	: Em.ConstantFloat(0.5f);
+	FValueRef MetallicArg	= Metallic.GetTracedInput().Expression	? Em.TryInput(&Metallic)	: Em.ConstantFloat(0.0f);
+
+	MIR::FValueRef DiffuseAlbedoArg = Em.InlineHLSL(
+		MIR::FType::MakeFloatVector(3),
+		TEXTVIEW("ComputeDiffuseAlbedo($0, saturate($1))"),
+		{ BaseColorArg, MetallicArg });
+
+	MIR::FValueRef F0Arg = Em.InlineHLSL(
+		MIR::FType::MakeFloatVector(3),
+		TEXTVIEW("ComputeF0($0, $1, saturate($2))"),
+		{ SpecularArg, BaseColorArg, MetallicArg });
+
+	Em.Output(0, DiffuseAlbedoArg);
+	Em.Output(1, F0Arg);
+}
+
+void UMaterialExpressionSubstrateHazinessToSecondaryRoughness::Build(MIR::FEmitter& Em)
+{
+	FValueRef BaseRoughnessArg	= BaseRoughness.GetTracedInput().Expression ?	Em.TryInput(&BaseRoughness)	: Em.ConstantFloat(0.1f);
+	FValueRef HazinessArg		= Haziness.GetTracedInput().Expression ?		Em.TryInput(&Haziness)		: Em.ConstantFloat(0.5f);
+
+	MIR::FValueRef HazeRoughnessArg = Em.InlineHLSL(
+		MIR::FType::MakeScalar(MIR::EScalarKind::Float),
+		TEXTVIEW("SubstrateComputeHazeRoughness(saturate($0))"),
+		{ BaseRoughnessArg });
+
+	MIR::FValueRef HazeWeightArg = Em.InlineHLSL(
+		MIR::FType::MakeScalar(MIR::EScalarKind::Float),
+		TEXTVIEW("SubstrateComputeHazeWeight(saturate($0), saturate($1))"),
+		{ BaseRoughnessArg, HazinessArg });
+
+	Em.Output(0, HazeRoughnessArg);
+	Em.Output(1, HazeWeightArg);
+}
+
+void UMaterialExpressionSubstrateThinFilm::Build(MIR::FEmitter& Em)
+{
+	UMaterial* DummyMaterial = nullptr;
+	bool bConvertToRequestedSpace = false;
+	FValueRef NormalArg		= CompileWithDefaultNormalWS(Em, DummyMaterial, Normal, bConvertToRequestedSpace);
+
+	FValueRef F0Arg			= F0.GetTracedInput().Expression ?			Em.TryInput(&F0)		: Em.ConstantFloat(0.04f);
+	FValueRef F90Arg		= F90.GetTracedInput().Expression ?			Em.TryInput(&F90)		: Em.ConstantFloat(1.0f);
+	FValueRef ThicknessArg	= Thickness.GetTracedInput().Expression ?	Em.TryInput(&Thickness)	: Em.ConstantFloat(1.0f);
+	FValueRef IORArg		= IOR.GetTracedInput().Expression ?			Em.TryInput(&IOR)		: Em.ConstantFloat(1.44f);
+
+	FValueRef CameraVector = Em.InlineHLSL(NAME_CameraVector);
+
+	MIR::FValueRef OutF0Arg = Em.InlineHLSL(
+		MIR::FType::MakeFloatVector(3),
+		TEXTVIEW("SubstrateGetThinFilmF0F90($0, $1, $2, $3, $4).F0"),
+		{ 
+			Em.Dot(NormalArg, CameraVector),
+			F0Arg,
+			F90Arg,
+			ThicknessArg,
+			IORArg
+		});
+
+	// Calling SubstrateGetThinFilmF0F90 a second time. SUBSTRATE_TODO: create a structure to pull the data from the first call instead
+	MIR::FValueRef OutF90Arg = Em.InlineHLSL(
+		MIR::FType::MakeFloatVector(3),
+		TEXTVIEW("SubstrateGetThinFilmF0F90($0, $1, $2, $3, $4).F90"),
+		{
+			Em.Dot(NormalArg, CameraVector),
+			F0Arg,
+			F90Arg,
+			ThicknessArg,
+			IORArg
+		});
+
+	Em.Output(0, OutF0Arg);
+	Em.Output(1, OutF90Arg);
+}
+
 #endif // WITH_EDITOR
\ No newline at end of file
diff --git a/Engine/Source/Runtime/Engine/Public/Materials/MaterialExpressionSubstrate.h b/Engine/Source/Runtime/Engine/Public/Materials/MaterialExpressionSubstrate.h
--- a/Engine/Source/Runtime/Engine/Public/Materials/MaterialExpressionSubstrate.h
+++ b/Engine/Source/Runtime/Engine/Public/Materials/MaterialExpressionSubstrate.h
@@ -220,6 +220,7 @@
 #if WITH_EDITOR
 	virtual FExpressionInput* GetInput(int32 InputIndex) override;
 	virtual int32 Compile(class FMaterialCompiler* Compiler, int32 OutputIndex) override;
+	virtual void Build(MIR::FEmitter& Emitter) override;
 	virtual void GetCaption(TArray<FString>& OutCaptions) const override;
 	virtual EMaterialValueType GetOutputValueType(int32 OutputIndex) override;
 	virtual EMaterialValueType GetInputValueType(int32 InputIndex) override;
@@ -229,7 +230,6 @@
 	virtual FName GetInputName(int32 InputIndex) const override;
 	virtual void GetConnectorToolTip(int32 InputIndex, int32 OutputIndex, TArray<FString>& OutToolTip) override;
 	virtual void PostEditChangeProperty(FPropertyChangedEvent& PropertyChangedEvent) override;
-	virtual void Build(MIR::FEmitter& Emitter) override;
 
 	bool HasSSS() const;
 	bool HasAnisotropy() const;
@@ -415,6 +415,7 @@
 	//~ Begin UMaterialExpression Interface
 #if WITH_EDITOR
 	virtual int32 Compile(class FMaterialCompiler* Compiler, int32 OutputIndex) override;
+	virtual void Build(MIR::FEmitter& Emitter) override;
 	virtual void GetCaption(TArray<FString>& OutCaptions) const override;
 	virtual EMaterialValueType GetOutputValueType(int32 OutputIndex) override;
 	virtual EMaterialValueType GetInputValueType(int32 InputIndex) override;
@@ -425,7 +426,6 @@
 	virtual void GetConnectorToolTip(int32 InputIndex, int32 OutputIndex, TArray<FString>& OutToolTip) override;
 	virtual void GetExpressionToolTip(TArray<FString>& OutToolTip) override;
 	virtual void PostEditChangeProperty(FPropertyChangedEvent& PropertyChangedEvent) override;
-	virtual void Build(MIR::FEmitter& Emitter) override;
 
 	bool HasEdgeColor() const;
 	bool HasFuzz() const;
@@ -499,6 +499,7 @@
 	//~ Begin UMaterialExpression Interface
 #if WITH_EDITOR
 	virtual int32 Compile(class FMaterialCompiler* Compiler, int32 OutputIndex) override;
+	// SUBSTRATE_TODO virtual void Build(MIR::FEmitter& Emitter) override;
 	virtual void GetCaption(TArray<FString>& OutCaptions) const override;
 	virtual EMaterialValueType GetOutputValueType(int32 OutputIndex) override;
 	virtual EMaterialValueType GetInputValueType(int32 InputIndex) override;
@@ -507,7 +508,6 @@
 	virtual FSubstrateOperator* SubstrateGenerateMaterialTopologyTree(struct FSubstrateTranslatorDataInterface& SubstrateTranslatorData, class UMaterialExpression* Parent, int32 OutputIndex) override;
 	virtual FName GetInputName(int32 InputIndex) const override;
 	virtual void GetExpressionToolTip(TArray<FString>& OutToolTip) override;
-	// SUBSTRATE_TODO virtual void Build(MIR::FEmitter& Emitter) override;
 #endif
 	//~ End UMaterialExpression Interface
 };
@@ -548,13 +548,13 @@
 	//~ Begin UMaterialExpression Interface
 #if WITH_EDITOR
 	virtual int32 Compile(class FMaterialCompiler* Compiler, int32 OutputIndex) override;
+	// SUBSTRATE_TODO virtual void Build(MIR::FEmitter& Emitter) override;
 	virtual void GetCaption(TArray<FString>& OutCaptions) const override;
 	virtual EMaterialValueType GetOutputValueType(int32 OutputIndex) override;
 	virtual EMaterialValueType GetInputValueType(int32 InputIndex) override;
 	virtual bool IsResultSubstrateMaterial(int32 OutputIndex) override;
 	virtual void GatherSubstrateMaterialInfo(FSubstrateMaterialInfo& SubstrateMaterialInfo, int32 OutputIndex) override;
 	virtual FSubstrateOperator* SubstrateGenerateMaterialTopologyTree(struct FSubstrateTranslatorDataInterface& SubstrateTranslatorData, class UMaterialExpression* Parent, int32 OutputIndex) override;
-	// SUBSTRATE_TODO virtual void Build(MIR::FEmitter& Emitter) override;
 
 	UE_DEPRECATED(5.8, "Please, SubstrateGenerateMaterialTopologyTreeCommon taking the FSubstrateTranslatorData parameter.")
 	static FSubstrateOperator* SubstrateGenerateMaterialTopologyTreeCommon(
@@ -601,13 +601,13 @@
 	//~ Begin UMaterialExpression Interface
 #if WITH_EDITOR
 	virtual int32 Compile(class FMaterialCompiler* Compiler, int32 OutputIndex) override;
+	// SUBSTRATE_TODO virtual void Build(MIR::FEmitter& Emitter) override;
 	virtual void GetCaption(TArray<FString>& OutCaptions) const override;
 	virtual EMaterialValueType GetOutputValueType(int32 OutputIndex) override;
 	virtual EMaterialValueType GetInputValueType(int32 InputIndex) override;
 	virtual bool IsResultSubstrateMaterial(int32 OutputIndex) override;
 	virtual void GatherSubstrateMaterialInfo(FSubstrateMaterialInfo& SubstrateMaterialInfo, int32 OutputIndex) override;
 	virtual FSubstrateOperator* SubstrateGenerateMaterialTopologyTree(struct FSubstrateTranslatorDataInterface& SubstrateTranslatorData, class UMaterialExpression* Parent, int32 OutputIndex) override;
-	// SUBSTRATE_TODO virtual void Build(MIR::FEmitter& Emitter) override;
 #endif
 	//~ End UMaterialExpression Interface
 };
@@ -662,13 +662,13 @@
 	//~ Begin UMaterialExpression Interface
 #if WITH_EDITOR
 	virtual int32 Compile(class FMaterialCompiler* Compiler, int32 OutputIndex) override;
+	// SUBSTRATE_TODO virtual void Build(MIR::FEmitter& Emitter) override;
 	virtual void GetCaption(TArray<FString>& OutCaptions) const override;
 	virtual EMaterialValueType GetOutputValueType(int32 OutputIndex) override;
 	virtual EMaterialValueType GetInputValueType(int32 InputIndex) override;
 	virtual bool IsResultSubstrateMaterial(int32 OutputIndex) override;
 	virtual void GatherSubstrateMaterialInfo(FSubstrateMaterialInfo& SubstrateMaterialInfo, int32 OutputIndex) override;
 	virtual FSubstrateOperator* SubstrateGenerateMaterialTopologyTree(struct FSubstrateTranslatorDataInterface& SubstrateTranslatorData, class UMaterialExpression* Parent, int32 OutputIndex) override;
-	// SUBSTRATE_TODO virtual void Build(MIR::FEmitter& Emitter) override;
 #endif
 	//~ End UMaterialExpression Interface
 };
@@ -733,13 +733,13 @@
 	//~ Begin UMaterialExpression Interface
 #if WITH_EDITOR
 	virtual int32 Compile(class FMaterialCompiler* Compiler, int32 OutputIndex) override;
+	// SUBSTRATE_TODO virtual void Build(MIR::FEmitter& Emitter) override;
 	virtual void GetCaption(TArray<FString>& OutCaptions) const override;
 	virtual EMaterialValueType GetOutputValueType(int32 OutputIndex) override;
 	virtual EMaterialValueType GetInputValueType(int32 InputIndex) override;
 	virtual bool IsResultSubstrateMaterial(int32 OutputIndex) override;
 	virtual void GatherSubstrateMaterialInfo(FSubstrateMaterialInfo& SubstrateMaterialInfo, int32 OutputIndex) override;
 	virtual FSubstrateOperator* SubstrateGenerateMaterialTopologyTree(struct FSubstrateTranslatorDataInterface& SubstrateTranslatorData, class UMaterialExpression* Parent, int32 OutputIndex) override;
-	// SUBSTRATE_TODO virtual void Build(MIR::FEmitter& Emitter) override;
 #endif
 	//~ End UMaterialExpression Interface
 };
@@ -818,13 +818,13 @@
 	//~ Begin UMaterialExpression Interface
 #if WITH_EDITOR
 	virtual int32 Compile(class FMaterialCompiler* Compiler, int32 OutputIndex) override;
+	// SUBSTRATE_TODO virtual void Build(MIR::FEmitter& Emitter) override;
 	virtual void GetCaption(TArray<FString>& OutCaptions) const override;
 	virtual EMaterialValueType GetOutputValueType(int32 OutputIndex) override;
 	virtual EMaterialValueType GetInputValueType(int32 InputIndex) override;
 	virtual bool IsResultSubstrateMaterial(int32 OutputIndex) override;
 	virtual void GatherSubstrateMaterialInfo(FSubstrateMaterialInfo& SubstrateMaterialInfo, int32 OutputIndex) override;
 	virtual FSubstrateOperator* SubstrateGenerateMaterialTopologyTree(struct FSubstrateTranslatorDataInterface& SubstrateTranslatorData, class UMaterialExpression* Parent, int32 OutputIndex) override;
-	// SUBSTRATE_TODO virtual void Build(MIR::FEmitter& Emitter) override;
 #endif
 	//~ End UMaterialExpression Interface
 };
@@ -843,6 +843,7 @@
 	//~ Begin UMaterialExpression Interface
 #if WITH_EDITOR
 	virtual int32 Compile(class FMaterialCompiler* Compiler, int32 OutputIndex) override;
+	// SUBSTRATE_TODO virtual void Build(MIR::FEmitter& Emitter) override;
 	virtual void GetCaption(TArray<FString>& OutCaptions) const override;
 	virtual EMaterialValueType GetOutputValueType(int32 OutputIndex) override;
 	virtual EMaterialValueType GetInputValueType(int32 InputIndex) override;
@@ -860,8 +861,6 @@
 	static int32 CompileCommon(class FMaterialCompiler* Compiler,
 		FExpressionInput& Color,
 		const UMaterialEditorOnlyData* EditorOnlyData = nullptr);
-
-	// SUBSTRATE_TODO virtual void Build(MIR::FEmitter& Emitter) override;
 #endif
 	//~ End UMaterialExpression Interface
 };
@@ -914,13 +913,13 @@
 	//~ Begin UMaterialExpression Interface
 #if WITH_EDITOR
 	virtual int32 Compile(class FMaterialCompiler* Compiler, int32 OutputIndex) override;
+	// SUBSTRATE_TODO virtual void Build(MIR::FEmitter& Emitter) override;
 	virtual void GetCaption(TArray<FString>& OutCaptions) const override;
 	virtual EMaterialValueType GetOutputValueType(int32 OutputIndex) override;
 	virtual EMaterialValueType GetInputValueType(int32 InputIndex) override;
 	virtual bool IsResultSubstrateMaterial(int32 OutputIndex) override;
 	virtual void GatherSubstrateMaterialInfo(FSubstrateMaterialInfo& SubstrateMaterialInfo, int32 OutputIndex) override;
 	virtual FSubstrateOperator* SubstrateGenerateMaterialTopologyTree(struct FSubstrateTranslatorDataInterface& SubstrateTranslatorData, class UMaterialExpression* Parent, int32 OutputIndex) override;
-	// SUBSTRATE_TODO virtual void Build(MIR::FEmitter& Emitter) override;
 #endif
 	//~ End UMaterialExpression Interface
 };
@@ -945,6 +944,7 @@
 	//~ Begin UMaterialExpression Interface
 #if WITH_EDITOR
 	virtual int32 Compile(class FMaterialCompiler* Compiler, int32 OutputIndex) override;
+	// SUBSTRATE_TODO virtual void Build(MIR::FEmitter& Emitter) override;
 	virtual void GetCaption(TArray<FString>& OutCaptions) const override;
 	virtual EMaterialValueType GetOutputValueType(int32 OutputIndex) override;
 	virtual EMaterialValueType GetInputValueType(int32 InputIndex) override;
@@ -962,8 +962,6 @@
 	static int32 CompileCommon(class FMaterialCompiler* Compiler,
 		FExpressionInput& Color, FExpressionInput& Opacity,
 		const UMaterialEditorOnlyData* EditorOnlyData = nullptr);
-
-	// SUBSTRATE_TODO virtual void Build(MIR::FEmitter& Emitter) override;
 #endif
 	//~ End UMaterialExpression Interface
 };
@@ -988,6 +986,7 @@
 	//~ Begin UMaterialExpression Interface
 #if WITH_EDITOR
 	virtual int32 Compile(class FMaterialCompiler* Compiler, int32 OutputIndex) override;
+	// SUBSTRATE_TODO virtual void Build(MIR::FEmitter& Emitter) override;
 	virtual void GetCaption(TArray<FString>& OutCaptions) const override;
 	virtual EMaterialValueType GetOutputValueType(int32 OutputIndex) override;
 	virtual EMaterialValueType GetInputValueType(int32 InputIndex) override;
@@ -1005,8 +1004,6 @@
 	static int32 CompileCommon(class FMaterialCompiler* Compiler,
 		FExpressionInput& Color, FExpressionInput& Opacity,
 		const UMaterialEditorOnlyData* EditorOnlyData = nullptr);
-
-	// SUBSTRATE_TODO virtual void Build(MIR::FEmitter& Emitter) override;
 #endif
 	//~ End UMaterialExpression Interface
 };
@@ -1031,13 +1028,13 @@
 	//~ Begin UMaterialExpression Interface
 #if WITH_EDITOR
 	virtual int32 Compile(class FMaterialCompiler* Compiler, int32 OutputIndex) override;
+	// SUBSTRATE_TODO virtual void Build(MIR::FEmitter& Emitter) override;
 	virtual void GetCaption(TArray<FString>& OutCaptions) const override;
 	virtual EMaterialValueType GetOutputValueType(int32 OutputIndex) override;
 	virtual EMaterialValueType GetInputValueType(int32 InputIndex) override;
 	virtual bool IsResultSubstrateMaterial(int32 OutputIndex) override;
 	virtual void GatherSubstrateMaterialInfo(FSubstrateMaterialInfo& SubstrateMaterialInfo, int32 OutputIndex) override;
 	virtual FSubstrateOperator* SubstrateGenerateMaterialTopologyTree(struct FSubstrateTranslatorDataInterface& SubstrateTranslatorData, class UMaterialExpression* Parent, int32 OutputIndex) override;
-	// SUBSTRATE_TODO virtual void Build(MIR::FEmitter& Emitter) override;
 #endif
 	//~ End UMaterialExpression Interface
 };
@@ -1085,6 +1082,7 @@
 #if WITH_EDITOR
 	virtual FExpressionInput* GetInput(int32 InputIndex) override;
 	virtual int32 Compile(class FMaterialCompiler* Compiler, int32 OutputIndex) override;
+	// SUBSTRATE_TODO virtual void Build(MIR::FEmitter& Emitter) override;
 	ENGINE_API static int32 CompileDefaultSlab(class FMaterialCompiler* Compiler, int32 EmissiveOverride);
 	ENGINE_API static int32 CompileDefaultSlab(class FMaterialCompiler* Compiler, FVector3f EmissiveOverride = FVector3f::Zero());
 	virtual void GetCaption(TArray<FString>& OutCaptions) const override;
@@ -1115,8 +1113,6 @@
 		FExpressionInput& WaterScatteringCoefficients, FExpressionInput& WaterAbsorptionCoefficients, FExpressionInput& WaterPhaseG, FExpressionInput& ColorScaleBehindWater,
 		const bool bHasSSS, USubsurfaceProfile* SSSProfile,
 		FExpressionInput& ClearCoatNormal, FExpressionInput& CustomTangent);
-
-	// SUBSTRATE_TODO virtual void Build(MIR::FEmitter& Emitter) override;
 #endif // WITH_EDITOR
 	//~ End UMaterialExpression Interface
 };
@@ -1158,13 +1154,13 @@
 	//~ Begin UMaterialExpression Interface
 #if WITH_EDITOR
 	virtual int32 Compile(class FMaterialCompiler* Compiler, int32 OutputIndex) override;
+	virtual void Build(MIR::FEmitter& Emitter) override;
 	virtual void GetCaption(TArray<FString>& OutCaptions) const override;
 	virtual EMaterialValueType GetOutputValueType(int32 OutputIndex) override;
 	virtual EMaterialValueType GetInputValueType(int32 InputIndex) override;
 	virtual bool IsResultSubstrateMaterial(int32 OutputIndex) override;
 	virtual void GatherSubstrateMaterialInfo(FSubstrateMaterialInfo& SubstrateMaterialInfo, int32 OutputIndex) override;
 	virtual FSubstrateOperator* SubstrateGenerateMaterialTopologyTree(struct FSubstrateTranslatorDataInterface& SubstrateTranslatorData, class UMaterialExpression* Parent, int32 OutputIndex) override;
-	virtual void Build(MIR::FEmitter& Emitter) override;
 #endif
 	//~ End UMaterialExpression Interface
 };
@@ -1202,6 +1198,7 @@
 	//~ Begin UMaterialExpression Interface
 #if WITH_EDITOR
 	virtual int32 Compile(class FMaterialCompiler* Compiler, int32 OutputIndex) override;
+	virtual void Build(MIR::FEmitter& Emitter) override;
 	virtual void GetCaption(TArray<FString>& OutCaptions) const override;
 	virtual EMaterialValueType GetOutputValueType(int32 OutputIndex) override;
 	virtual EMaterialValueType GetInputValueType(int32 InputIndex) override;
@@ -1209,7 +1206,6 @@
 	virtual void GatherSubstrateMaterialInfo(FSubstrateMaterialInfo& SubstrateMaterialInfo, int32 OutputIndex) override;
 	virtual FSubstrateOperator* SubstrateGenerateMaterialTopologyTree(struct FSubstrateTranslatorDataInterface& SubstrateTranslatorData, class UMaterialExpression* Parent, int32 OutputIndex) override;
 	virtual FName GetInputName(int32 InputIndex) const override;
-	virtual void Build(MIR::FEmitter& Emitter) override;
 #endif
 	//~ End UMaterialExpression Interface
 };
@@ -1240,13 +1236,13 @@
 	//~ Begin UMaterialExpression Interface
 #if WITH_EDITOR
 	virtual int32 Compile(class FMaterialCompiler* Compiler, int32 OutputIndex) override;
+	virtual void Build(MIR::FEmitter& Emitter) override;
 	virtual void GetCaption(TArray<FString>& OutCaptions) const override;
 	virtual EMaterialValueType GetOutputValueType(int32 OutputIndex) override;
 	virtual EMaterialValueType GetInputValueType(int32 InputIndex) override;
 	virtual bool IsResultSubstrateMaterial(int32 OutputIndex) override;
 	virtual void GatherSubstrateMaterialInfo(FSubstrateMaterialInfo& SubstrateMaterialInfo, int32 OutputIndex) override;
 	virtual FSubstrateOperator* SubstrateGenerateMaterialTopologyTree(struct FSubstrateTranslatorDataInterface& SubstrateTranslatorData, class UMaterialExpression* Parent, int32 OutputIndex) override;
-	virtual void Build(MIR::FEmitter& Emitter) override;
 #endif
 	//~ End UMaterialExpression Interface
 };
@@ -1271,13 +1267,13 @@
 	//~ Begin UMaterialExpression Interface
 #if WITH_EDITOR
 	virtual int32 Compile(class FMaterialCompiler* Compiler, int32 OutputIndex) override;
+	virtual void Build(MIR::FEmitter& Emitter) override;
 	virtual void GetCaption(TArray<FString>& OutCaptions) const override;
 	virtual EMaterialValueType GetOutputValueType(int32 OutputIndex) override;
 	virtual EMaterialValueType GetInputValueType(int32 InputIndex) override;
 	virtual bool IsResultSubstrateMaterial(int32 OutputIndex) override;
 	virtual void GatherSubstrateMaterialInfo(FSubstrateMaterialInfo& SubstrateMaterialInfo, int32 OutputIndex) override;
 	virtual FSubstrateOperator* SubstrateGenerateMaterialTopologyTree(struct FSubstrateTranslatorDataInterface& SubstrateTranslatorData, class UMaterialExpression* Parent, int32 OutputIndex) override;
-	virtual void Build(MIR::FEmitter& Emitter) override;
 #endif
 	//~ End UMaterialExpression Interface
 };
@@ -1320,13 +1316,13 @@
 	//~ Begin UMaterialExpression Interface
 #if WITH_EDITOR
 	virtual int32 Compile(class FMaterialCompiler* Compiler, int32 OutputIndex) override;
+	// SUBSTRATE_TODO virtual void Build(MIR::FEmitter& Emitter) override;
 	virtual void GetCaption(TArray<FString>& OutCaptions) const override;
 	virtual EMaterialValueType GetOutputValueType(int32 OutputIndex) override;
 	virtual EMaterialValueType GetInputValueType(int32 InputIndex) override;
 	virtual bool IsResultSubstrateMaterial(int32 OutputIndex) override;
 	virtual void GatherSubstrateMaterialInfo(FSubstrateMaterialInfo& SubstrateMaterialInfo, int32 OutputIndex) override;
 	virtual FSubstrateOperator* SubstrateGenerateMaterialTopologyTree(struct FSubstrateTranslatorDataInterface& SubstrateTranslatorData, class UMaterialExpression* Parent, int32 OutputIndex) override;
-	// SUBSTRATE_TODO virtual void Build(MIR::FEmitter& Emitter) override;
 #endif
 	//~ End UMaterialExpression Interface
 };
@@ -1363,12 +1359,12 @@
 	//~ Begin UMaterialExpression Interface
 #if WITH_EDITOR
 	virtual int32 Compile(class FMaterialCompiler* Compiler, int32 OutputIndex) override;
+	virtual void Build(MIR::FEmitter& Emitter) override;
 	virtual void GetCaption(TArray<FString>& OutCaptions) const override;
 	virtual EMaterialValueType GetOutputValueType(int32 OutputIndex) override;
 	virtual EMaterialValueType GetInputValueType(int32 InputIndex) override;
 	virtual void GetConnectorToolTip(int32 InputIndex, int32 OutputIndex, TArray<FString>& OutToolTip) override;
 	virtual void GetExpressionToolTip(TArray<FString>& OutToolTip) override;
-	virtual void Build(MIR::FEmitter& Emitter) override;
 #endif
 	//~ End UMaterialExpression Interface
 };
@@ -1399,12 +1395,12 @@
 	//~ Begin UMaterialExpression Interface
 #if WITH_EDITOR
 	virtual int32 Compile(class FMaterialCompiler* Compiler, int32 OutputIndex) override;
+	virtual void Build(MIR::FEmitter& Emitter) override;
 	virtual void GetCaption(TArray<FString>& OutCaptions) const override;
 	virtual EMaterialValueType GetOutputValueType(int32 OutputIndex) override;
 	virtual EMaterialValueType GetInputValueType(int32 InputIndex) override;
 	virtual void GetConnectorToolTip(int32 InputIndex, int32 OutputIndex, TArray<FString>& OutToolTip) override;
 	virtual void GetExpressionToolTip(TArray<FString>& OutToolTip) override;
-	// SUBSTRATE_TODO virtual void Build(MIR::FEmitter& Emitter) override;
 #endif
 	//~ End UMaterialExpression Interface
 };
@@ -1429,12 +1425,12 @@
 	//~ Begin UMaterialExpression Interface
 #if WITH_EDITOR
 	virtual int32 Compile(class FMaterialCompiler* Compiler, int32 OutputIndex) override;
+	virtual void Build(MIR::FEmitter& Emitter) override;
 	virtual void GetCaption(TArray<FString>& OutCaptions) const override;
 	virtual EMaterialValueType GetOutputValueType(int32 OutputIndex) override;
 	virtual EMaterialValueType GetInputValueType(int32 InputIndex) override;
 	virtual void GetConnectorToolTip(int32 InputIndex, int32 OutputIndex, TArray<FString>& OutToolTip) override;
 	virtual void GetExpressionToolTip(TArray<FString>& OutToolTip) override;
-	// SUBSTRATE_TODO virtual void Build(MIR::FEmitter& Emitter) override;
 #endif
 	//~ End UMaterialExpression Interface
 };
@@ -1477,12 +1473,12 @@
 	//~ Begin UMaterialExpression Interface
 #if WITH_EDITOR
 	virtual int32 Compile(class FMaterialCompiler* Compiler, int32 OutputIndex) override;
+	virtual void Build(MIR::FEmitter& Emitter) override;
 	virtual void GetCaption(TArray<FString>& OutCaptions) const override;
 	virtual EMaterialValueType GetOutputValueType(int32 OutputIndex) override;
 	virtual EMaterialValueType GetInputValueType(int32 InputIndex) override;
 	virtual void GetConnectorToolTip(int32 InputIndex, int32 OutputIndex, TArray<FString>& OutToolTip) override;
 	virtual void GetExpressionToolTip(TArray<FString>& OutToolTip) override;
-	// SUBSTRATE_TODO virtual void Build(MIR::FEmitter& Emitter) override;
 #endif
 	//~ End UMaterialExpression Interface
 };
