Engine - set anim group follower LeaderScore to always be less than any possible leader to avoid selecting any follower as leader if there's any leader in the FAnimGroupInstance (fixing blendstack behavior when blending into a new animation introduced as CanBeLeader - or possibly TransitionLeader - with EffectiveBlendWeight  set to zero, into a group of only follower with EffectiveBlendWeight greater than zero, ultimately being selected as possible leaders instead of the CanBeLeader animation)

#Synced-CL 49959264

diff --git a/Engine/Source/Runtime/Engine/Private/AnimationAsset.cpp b/Engine/Source/Runtime/Engine/Private/AnimationAsset.cpp
--- a/Engine/Source/Runtime/Engine/Private/AnimationAsset.cpp
+++ b/Engine/Source/Runtime/Engine/Private/AnimationAsset.cpp
@@ -24,6 +24,7 @@
 
 #define LOCTEXT_NAMESPACE "AnimationAsset"
 
+#define LEADERSCORE_FOLLOWER_OFFSET	-2.f // offset to score any follower after any possible leader
 #define LEADERSCORE_ALWAYSLEADER  	2.f
 #define LEADERSCORE_MONTAGE			3.f
 
@@ -84,17 +85,20 @@ void FAnimGroupInstance::TestTickRecordForLeadership(EAnimGroupRole::Type Member
 		{
 		case EAnimGroupRole::CanBeLeader:
 		case EAnimGroupRole::TransitionLeader:
-			Candidate.LeaderScore = Candidate.EffectiveBlendWeight;
+			Candidate.LeaderScore = Candidate.EffectiveBlendWeight; // LeaderScore within [0, 1] interval
 			break;
 		case EAnimGroupRole::AlwaysLeader:
 		case EAnimGroupRole::ExclusiveAlwaysLeader:
 			// Always set the leader index
-			Candidate.LeaderScore = LEADERSCORE_ALWAYSLEADER;
+			Candidate.LeaderScore = LEADERSCORE_ALWAYSLEADER;  // LeaderScore within [2, 2] interval
 			break;
 		default:
 		case EAnimGroupRole::AlwaysFollower:
 		case EAnimGroupRole::TransitionFollower:
-			// Never set the leader index; the actual tick code will handle the case of no leader by using the first element in the array
+			// making sure that the LeaderScore of any follower is always less than any possible leader, 
+			// so leaders happen to be sorted before any follower
+			// in case there's no leader the follower with highest LeaderScore will be elected as anim group leader
+			Candidate.LeaderScore = LEADERSCORE_FOLLOWER_OFFSET + Candidate.EffectiveBlendWeight; // LeaderScore within [-2, -1] interval
 			break;
 		}
 	}
