Substrate - NMT - added support for Select node operator.

#tests PC editor dx12
#rnx

#Synced-CL 50037646

diff --git a/Engine/Source/Runtime/Engine/Private/Materials/MaterialExpressionsIR.cpp b/Engine/Source/Runtime/Engine/Private/Materials/MaterialExpressionsIR.cpp
--- a/Engine/Source/Runtime/Engine/Private/Materials/MaterialExpressionsIR.cpp
+++ b/Engine/Source/Runtime/Engine/Private/Materials/MaterialExpressionsIR.cpp
@@ -6029,6 +6029,135 @@ void UMaterialExpressionSubstrateAdd::Build(MIR::FEmitter& Em)
 	}
 }
 
+void UMaterialExpressionSubstrateSelect::Build(MIR::FEmitter& Em)
+{
+	if (!A.GetTracedInput().Expression)
+	{
+		return Em.Errorf(TEXT("Missing A input"));
+	}
+	if (!B.GetTracedInput().Expression)
+	{
+		return Em.Errorf(TEXT("Missing B input"));
+	}
+
+	FSubstrateTranslatorData* SubstrateTranslatorData = Em.GetSubstrateTranslatorData();
+	FSubstrateOperator& SubstrateOperator = SubstrateTranslatorData->SubstrateCompilationGetOperator(SubstrateTranslatorData->SubstrateTreeStackGetPathUniqueId());
+	if (SubstrateOperator.bUseParameterBlending)
+	{
+		// Note: we no longer do SubstrateTreeStackPush/Pop since the path guid is regenerated for every substrate node from BuildTopMaterialExpression.
+		FValueRef BArg = Em.TryInput(&B);
+		FValueRef AArg = Em.TryInput(&A);
+
+		if (!AArg.IsValid())
+		{
+			return Em.Errorf(TEXT("A input graphs could not be evaluated for parameter blending."));
+		}
+		if (!BArg.IsValid())
+		{
+			return Em.Errorf(TEXT("B input graphs could not be evaluated for parameter blending."));
+		}
+
+		// if SelectValue is not pluggedin, need to be 0. Otherwise it must be a float value.
+		FValueRef SelectValueArg = CompileWithDefaultFloat1(Em, SelectValue, 0.0f);
+
+		FValueRef ThresholdArg = Em.ConstantFloat(Threshold);
+		FValueRef FinalSelectValueArg = Em.InlineHLSL(
+			MIR::FType::MakeScalar(MIR::EScalarKind::Float),
+			TEXTVIEW("($0 <= $1 ? 0.0f : 1.0f)"),
+			{ SelectValueArg, ThresholdArg });
+
+		if (!AArg.IsValid() || !BArg.IsValid() || !ThresholdArg.IsValid())
+		{
+			return Em.Errorf(TEXT("Invalid Select operator A, B or Threshold input."));
+		}
+
+		auto SubstrateSelectNormal = [&](FValueRef NormalA, FValueRef NormalB)
+			{
+				return Em.InlineHLSL(
+					MIR::FType::MakeFloatVector(3),
+					TEXTVIEW("($0 <= $1 ? $2 : $3)"),
+					{ SelectValueArg, ThresholdArg, NormalA, NormalB });
+			};
+
+		FSubstrateOperator* BBSDFOperator = SubstrateTranslatorData->SubstrateCompilationGetOperatorFromIndex(SubstrateOperator.LeftIndex);
+		FSubstrateOperator* ABSDFOperator = SubstrateTranslatorData->SubstrateCompilationGetOperatorFromIndex(SubstrateOperator.RightIndex);
+		if (!BBSDFOperator || !ABSDFOperator)
+		{
+			return Em.Errorf(TEXT("Missing input on Select node."));
+		}
+
+		FSubstrateRegisteredSharedLocalBasisSearchResult ABasis = SubstrateCompilationFindExistingSharedLocalBasis(SubstrateTranslatorData, ABSDFOperator->BSDFRegisteredSharedLocalBasis.NormalCodeChunkHash, ABSDFOperator->BSDFRegisteredSharedLocalBasis.TangentCodeChunkHash);
+		FSubstrateRegisteredSharedLocalBasisSearchResult BBasis = SubstrateCompilationFindExistingSharedLocalBasis(SubstrateTranslatorData, BBSDFOperator->BSDFRegisteredSharedLocalBasis.NormalCodeChunkHash, BBSDFOperator->BSDFRegisteredSharedLocalBasis.TangentCodeChunkHash);
+
+		if (!ABasis.bFound || !BBasis.bFound)
+		{
+			Em.Errorf(TEXT("A & B shared local basis could not be found."));
+			return;
+		}
+
+		// Compute the new Normal and Tangent resulting from the blending using code chunk
+		FValueRef NewNormalArg = SubstrateSelectNormal(BBasis.BasisInfo.NormalArg, ABasis.BasisInfo.NormalArg);
+
+		// The tangent is optional so we treat it differently if INDEX_NONE is specified
+		FValueRef NewTangentArg;
+		if (ABasis.BasisInfo.TangentArg && BBasis.BasisInfo.TangentArg)
+		{
+			NewTangentArg = SubstrateSelectNormal(BBasis.BasisInfo.TangentArg, ABasis.BasisInfo.TangentArg);
+		}
+		else if (ABasis.BasisInfo.TangentArg)
+		{
+			NewTangentArg = ABasis.BasisInfo.TangentArg;
+		}
+		else if (BBasis.BasisInfo.TangentArg)
+		{
+			NewTangentArg = BBasis.BasisInfo.TangentArg;
+		}
+		const FSubstrateRegisteredSharedLocalBasis NewRegisteredSharedLocalBasis = MIRSubstrateCompilationInfoCreateSharedLocalBasis2(SubstrateTranslatorData, &NewNormalArg, NewTangentArg.IsValid() ? &NewTangentArg : nullptr);
+
+		FString SharedLocalBasisName = GetSubstrateSharedLocalBasisIndexMacroInner(NewRegisteredSharedLocalBasis, SubstrateTranslatorData->CurrentSubstrateCompilationContext);
+		FValueRef SharedLocalBasisIndexMacroArg = Em.InlineHLSL(MIR::FType::MakeIntScalar(), *SharedLocalBasisName);
+
+		const FSubstrateOperator* PromoteToOperator = SubstrateOperator.bRootOfParameterBlendingSubTree ? &SubstrateOperator : nullptr;
+		if (PromoteToOperator)
+		{
+			check(PromoteToOperator->Index != INDEX_NONE);
+			check(PromoteToOperator->BSDFIndex != INDEX_NONE);
+
+			FValueRef SelectedSlab = Em.SubstrateSelectParameterBlending(
+				AArg,
+				BArg,
+				FinalSelectValueArg,
+				SharedLocalBasisIndexMacroArg);
+
+			FValueRef PromotedSlab = Em.SubstratePromoteToOperator(
+				SelectedSlab,
+				Em.ConstantInt(PromoteToOperator->Index),
+				Em.ConstantInt(PromoteToOperator->BSDFIndex),
+				Em.ConstantInt(PromoteToOperator->LayerDepth),
+				Em.ConstantInt(PromoteToOperator->bIsBottom ? 1 : 0));
+
+			Em.Output(0, PromotedSlab);
+		}
+		else
+		{
+			FValueRef SelectedSlab = Em.SubstrateSelectParameterBlending(
+				AArg,
+				BArg,
+				FinalSelectValueArg,
+				SharedLocalBasisIndexMacroArg);
+
+			Em.Output(0, SelectedSlab);
+		}
+
+		// Propagate the parameter blended normal
+		SubstrateOperator.BSDFRegisteredSharedLocalBasis = NewRegisteredSharedLocalBasis;
+	}
+	else
+	{
+		return Em.Errorf(TEXT("The Select node can only use parameter blending to only select between one of two BSDF."));
+	}
+}
+
 void UMaterialExpressionSubstrateTransmittanceToMFP::Build(MIR::FEmitter& Em)
 {
 	FSubstrateTranslatorData* SubstrateTranslatorData = Em.GetSubstrateTranslatorData();
diff --git a/Engine/Source/Runtime/Engine/Private/Materials/MaterialIR.cpp b/Engine/Source/Runtime/Engine/Private/Materials/MaterialIR.cpp
--- a/Engine/Source/Runtime/Engine/Private/Materials/MaterialIR.cpp
+++ b/Engine/Source/Runtime/Engine/Private/Materials/MaterialIR.cpp
@@ -55,6 +55,7 @@ const TCHAR* LexToString(EValueKind Kind)
 		case VK_SubstrateVerticalLayering: return TEXT("SubstrateVerticalLayering");
 		case VK_SubstrateCoverageWeight: return TEXT("SubstrateCoverageWeight");
 		case VK_SubstrateAdd: return TEXT("SubstrateAdd");
+		case VK_SubstrateSelect: return TEXT("SubstrateSelect");
 		case VK_SubstratePromoteToOperator: return TEXT("SubstratePromoteToOperator");
 		case VK_StageSwitch: return TEXT("StageSwitch");
 		case VK_HardwarePartialDerivative: return TEXT("HardwarePartialDerivative");
@@ -130,6 +131,7 @@ uint32 FValue::GetSizeInBytes() const
 		case VK_SubstrateVerticalLayering: return sizeof(FSubstrateVerticalLayering);
 		case VK_SubstrateCoverageWeight: return sizeof(FSubstrateCoverageWeight);
 		case VK_SubstrateAdd: return sizeof(FSubstrateAdd);
+		case VK_SubstrateSelect: return sizeof(FSubstrateSelect);
 		case VK_SubstratePromoteToOperator: return sizeof(FSubstratePromoteToOperator);
 		case VK_StageSwitch: return sizeof(FStageSwitch);
 		case VK_HardwarePartialDerivative: return sizeof(FHardwarePartialDerivative);
@@ -251,6 +253,12 @@ TConstArrayView<FValue*> FValue::GetUses() const
 			return { &This->A, FSubstrateAdd::NumStaticUses };
 		}
 
+		case VK_SubstrateSelect:
+		{
+			auto This = static_cast<const FSubstrateSelect*>(this);
+			return { &This->A, FSubstrateSelect::NumStaticUses };
+		}
+
 		case VK_SubstratePromoteToOperator:
 		{
 			auto This = static_cast<const FSubstratePromoteToOperator*>(this);
diff --git a/Engine/Source/Runtime/Engine/Private/Materials/MaterialIREmitter.cpp b/Engine/Source/Runtime/Engine/Private/Materials/MaterialIREmitter.cpp
--- a/Engine/Source/Runtime/Engine/Private/Materials/MaterialIREmitter.cpp
+++ b/Engine/Source/Runtime/Engine/Private/Materials/MaterialIREmitter.cpp
@@ -4562,6 +4562,27 @@ FValueRef FEmitter::SubstrateAddParameterBlending(FValueRef A, FValueRef B, FVal
 	return EmitPrototype(*this, Prototype);
 }
 
+FValueRef FEmitter::SubstrateSelectParameterBlending(FValueRef A, FValueRef B, FValueRef SelectValue, FValueRef SharedLocalBasisIndexMacro)
+{
+	FSubstrateSelect Prototype = MakePrototype<FSubstrateSelect>(FType::MakeSubstrateData());
+	Prototype.A = A;
+	Prototype.B = B;
+	Prototype.SelectValue = SelectValue;
+
+	// ParameterBlending = 0
+	Prototype.OperatorIndex = nullptr;
+	Prototype.MaxDistanceFromLeaves = nullptr;
+
+	// ParameterBlending = 1
+	Prototype.SharedLocalBasisIndexMacro = SharedLocalBasisIndexMacro;
+
+	// Settings
+	Prototype.bParameterBlendingEnabled = true;
+	Prototype.bFullySimplifiedInstruction = SubstrateTranslatorData->CurrentSubstrateCompilationContext == Substrate::ESubstrateCompilationContext::SCC_FullySimplified;
+
+	return EmitPrototype(*this, Prototype);
+}
+
 EShaderPlatform FEmitter::GetShaderPlatform() const
 {
 	return Module->GetShaderPlatform();
diff --git a/Engine/Source/Runtime/Engine/Private/Materials/MaterialIRToHLSLTranslator.cpp b/Engine/Source/Runtime/Engine/Private/Materials/MaterialIRToHLSLTranslator.cpp
--- a/Engine/Source/Runtime/Engine/Private/Materials/MaterialIRToHLSLTranslator.cpp
+++ b/Engine/Source/Runtime/Engine/Private/Materials/MaterialIRToHLSLTranslator.cpp
@@ -307,6 +307,7 @@ static bool IsFoldable(const MIR::FInstruction* Instr, int32 EntryPointIndex)
 		case MIR::VK_SubstrateVerticalLayering:
 		case MIR::VK_SubstrateCoverageWeight:
 		case MIR::VK_SubstrateAdd:
+		case MIR::VK_SubstrateSelect:
 		case MIR::VK_SubstratePromoteToOperator:
 		case MIR::VK_SubstrateSlab:
 		case MIR::VK_SubstrateShadingModels:
@@ -546,6 +547,7 @@ struct FPrivate : FMaterialIRToHLSLTranslation
 			case MIR::VK_SubstrateVerticalLayering: LowerSubstrateVerticalLayering(static_cast<const MIR::FSubstrateVerticalLayering*>(Instr)); break;
 			case MIR::VK_SubstrateCoverageWeight: LowerSubstrateCoverageWeight(static_cast<const MIR::FSubstrateCoverageWeight*>(Instr)); break;
 			case MIR::VK_SubstrateAdd: LowerSubstrateAdd(static_cast<const MIR::FSubstrateAdd*>(Instr)); break;
+			case MIR::VK_SubstrateSelect: LowerSubstrateSelect(static_cast<const MIR::FSubstrateSelect*>(Instr)); break;
 			case MIR::VK_SubstratePromoteToOperator: LowerSubstratePromoteToOperator(static_cast<const MIR::FSubstratePromoteToOperator*>(Instr)); break;
 			default: UE_MIR_UNREACHABLE();
 		}
@@ -2213,6 +2215,25 @@ struct FPrivate : FMaterialIRToHLSLTranslation
 		}
 	}
 
+	void LowerSubstrateSelect(const MIR::FSubstrateSelect* SubstrateParameter)
+	{
+		if (SubstrateParameter->bParameterBlendingEnabled)
+		{
+			Printer << TEXTVIEW("SubstrateSelectParameterBlending");
+			Printer << BeginArgs;
+			Printer << LowerValue(SubstrateParameter->A) << TEXTVIEW(",");
+			Printer << LowerValue(SubstrateParameter->B) << TEXTVIEW(",");
+			Printer << LowerValue(SubstrateParameter->SelectValue) << TEXTVIEW(",");
+			Printer << LowerValue(SubstrateParameter->SharedLocalBasisIndexMacro) << TEXTVIEW(",");
+			Printer << GetSharedLocalBasesTypes(SubstrateParameter->bFullySimplifiedInstruction);
+			Printer << EndArgs;
+		}
+		else
+		{
+			Printer << TEXTVIEW("SubstrateSelect - Parameter Blending is required for the select node.");
+		}
+	}
+
 	void LowerSubstratePromoteToOperator(const MIR::FSubstratePromoteToOperator* SubstrateParameter)
 	{
 		Printer << GetSubstrateTree(SubstrateParameter->bFullySimplifiedInstruction) << TEXTVIEW(".PromoteParameterBlendedBSDFToOperator");
diff --git a/Engine/Source/Runtime/Engine/Public/Materials/MaterialExpressionSubstrate.h b/Engine/Source/Runtime/Engine/Public/Materials/MaterialExpressionSubstrate.h
--- a/Engine/Source/Runtime/Engine/Public/Materials/MaterialExpressionSubstrate.h
+++ b/Engine/Source/Runtime/Engine/Public/Materials/MaterialExpressionSubstrate.h
@@ -1296,7 +1296,7 @@ class UMaterialExpressionSubstrateSelect : public UMaterialExpressionSubstrateBS
 	FExpressionInput B;
 
 	/**
-	 * If <=0, A is selected, otherwise B is selected. (Default: 0))
+	 * If <= Threshold, A is selected, otherwise B is selected. (Default: 0))
 	 */
 	UPROPERTY()
 	FExpressionInput SelectValue;
@@ -1316,7 +1316,7 @@ class UMaterialExpressionSubstrateSelect : public UMaterialExpressionSubstrateBS
 	//~ Begin UMaterialExpression Interface
 #if WITH_EDITOR
 	virtual int32 Compile(class FMaterialCompiler* Compiler, int32 OutputIndex) override;
-	// SUBSTRATE_TODO virtual void Build(MIR::FEmitter& Emitter) override;
+	virtual void Build(MIR::FEmitter& Emitter) override;
 	virtual void GetCaption(TArray<FString>& OutCaptions) const override;
 	virtual EMaterialValueType GetOutputValueType(int32 OutputIndex) override;
 	virtual EMaterialValueType GetInputValueType(int32 InputIndex) override;
diff --git a/Engine/Source/Runtime/Engine/Public/Materials/MaterialIR.h b/Engine/Source/Runtime/Engine/Public/Materials/MaterialIR.h
--- a/Engine/Source/Runtime/Engine/Public/Materials/MaterialIR.h
+++ b/Engine/Source/Runtime/Engine/Public/Materials/MaterialIR.h
@@ -132,6 +132,7 @@ enum EValueKind
 	VK_SubstrateVerticalLayering,
 	VK_SubstrateCoverageWeight,
 	VK_SubstrateAdd,
+	VK_SubstrateSelect,
 	VK_SubstratePromoteToOperator,
 	VK_StageSwitch,
 	VK_HardwarePartialDerivative,
@@ -927,6 +928,24 @@ struct FSubstrateAdd : TInstruction<VK_SubstrateAdd, 8>
 	bool bFullySimplifiedInstruction;
 };
 
+struct FSubstrateSelect : TInstruction<VK_SubstrateSelect, 6>
+{
+	FValue* A;
+	FValue* B;
+	FValue* SelectValue;
+
+	// ParameterBlending = 0
+	FValue* OperatorIndex;
+	FValue* MaxDistanceFromLeaves;
+
+	// ParameterBlending = 1
+	FValue* SharedLocalBasisIndexMacro;
+
+	// Settings
+	bool bParameterBlendingEnabled;
+	bool bFullySimplifiedInstruction;
+};
+
 struct FSubstratePromoteToOperator : TInstruction<VK_SubstratePromoteToOperator, 5>
 {
 	FValue* SubstrateDataInput;
diff --git a/Engine/Source/Runtime/Engine/Public/Materials/MaterialIREmitter.h b/Engine/Source/Runtime/Engine/Public/Materials/MaterialIREmitter.h
--- a/Engine/Source/Runtime/Engine/Public/Materials/MaterialIREmitter.h
+++ b/Engine/Source/Runtime/Engine/Public/Materials/MaterialIREmitter.h
@@ -857,6 +857,9 @@ public:
 	// Parameter blended version of SubstrateAdd.
 	FValueRef SubstrateAddParameterBlending(FValueRef A, FValueRef B, FValueRef NormalMix, FValueRef SharedLocalBasisIndexMacro, FValueRef ANoV, FValueRef BNoV);
 
+	// Parameter blended version of SubstrateSelect, selecting and thus enforcing a single BSDF per pixel: A and B.
+	FValueRef SubstrateSelectParameterBlending(FValueRef A, FValueRef B, FValueRef SelectValueArg, FValueRef SharedLocalBasisIndexMacro);
+
 	/*---------------------------- Pass through accessors -----------------------------*/
 
 	// Returns the shader platform the source material is translating for.
