Fix crash when expression is null

#rnx

#Synced-CL 50061279

diff --git a/Engine/Source/Runtime/Engine/Private/Materials/MaterialExpressions.cpp b/Engine/Source/Runtime/Engine/Private/Materials/MaterialExpressions.cpp
--- a/Engine/Source/Runtime/Engine/Private/Materials/MaterialExpressions.cpp
+++ b/Engine/Source/Runtime/Engine/Private/Materials/MaterialExpressions.cpp
@@ -9286,7 +9286,10 @@ void UMaterialExpressionQualitySwitch::GatherSubstrateMaterialInfo(FSubstrateMat
 	}
 	for (int i = 0; i < EMaterialQualityLevel::Num; ++i)
 	{
-		Inputs[i].GetTracedInput().Expression->GatherSubstrateMaterialInfo(SubstrateMaterialInfo, Inputs[i].OutputIndex);
+		if (Inputs[i].GetTracedInput().Expression)
+		{
+			Inputs[i].GetTracedInput().Expression->GatherSubstrateMaterialInfo(SubstrateMaterialInfo, Inputs[i].OutputIndex);
+		}
 	}
 }
 
