Update FTickableGameObject comments to more clearly indicate the thread safety issues

#Synced-CL 49929924

diff --git a/Engine/Source/Runtime/Engine/Public/Tickable.h b/Engine/Source/Runtime/Engine/Public/Tickable.h
--- a/Engine/Source/Runtime/Engine/Public/Tickable.h
+++ b/Engine/Source/Runtime/Engine/Public/Tickable.h
@@ -126,10 +126,16 @@ public:
 };
 
 /**
- * This class provides common registration for gamethread tickable objects. It is an
+ * This class provides an easy way to register gamethread-only tickable objects. It is an
  * abstract base class requiring you to implement the Tick() and GetStatId() methods.
- * Can optionally also be ticked in the Editor, allowing for an object that both ticks
+ * It is intended to be used for simple objects that cannot easily use other methods for ticking.
+ * It can optionally also be ticked in the Editor, allowing for an object that both ticks
  * during edit time and at runtime.
+ * 
+ * NOTE: Do not use this class with UObjects unless you add explicit code to start and stop ticking.
+ * To do this, construct with ETickableTickType::Never and call SetTickableTickType to enable ticking later.
+ * Ticking cannot be safely started until after object construction finishes, and should be set back to Never in BeginDestroy or earlier.
+ * UObjects are often created and destroyed on background threads so it is not safe to control ticking from the constructor/destructor.
  */
 class FTickableGameObject : public FTickableObjectBase
 {
@@ -144,7 +150,8 @@ public:
 	 * Registers this instance with the static array of tickable objects, if it can ever tick.
 	 * By default this function can only be called on the game thread, to make sure it doesn't call functions before initialization is complete.
 	 * With the default NewObject tick type, it will call GetTickableTickType at the start of the next tick to determine the desired type.
-	 * If this is something like a UObject that could be created on a different thread (like for async loading), construct with a Never tick type and enable tick later.
+	 * NOTE: This is not threadsafe with the default tick type! 
+	 * For UObjects and other classes that can be constructed on worker threads, you must construct with ETickableTickType::Never and enable ticking later.
 	 * 
 	 * @param StartingTickType The initial ticking state of this object, override it to start disabled or skip the call to GetTickableTickType
 	 */
