Enable "Mask material only in early Z-pass" renderer setting by default, and deprecate the option for future removal.

#Synced-CL 50066859

diff --git a/Engine/Source/Runtime/Engine/Classes/Engine/RendererSettings.h b/Engine/Source/Runtime/Engine/Classes/Engine/RendererSettings.h
--- a/Engine/Source/Runtime/Engine/Classes/Engine/RendererSettings.h
+++ b/Engine/Source/Runtime/Engine/Classes/Engine/RendererSettings.h
@@ -899,6 +899,7 @@ class URendererSettings : public UDeveloperSettings
 		ToolTip="Whether to use a depth only pass to initialize Z culling for the base pass."))
 	TEnumAsByte<EEarlyZPass::Type> EarlyZPass;
 
+	UE_DEPRECATED(5.8, "This option will be removed, and its behavior always enabled in a future release.")
 	UPROPERTY(config, EditAnywhere, Category = Optimizations, meta = (
 		EditCondition = "EarlyZPass == EEarlyZPass::OpaqueAndMasked",
 		ConsoleVariable = "r.EarlyZPassOnlyMaterialMasking", DisplayName = "Mask material only in early Z-pass",
diff --git a/Engine/Source/Runtime/Renderer/Private/RendererScene.cpp b/Engine/Source/Runtime/Renderer/Private/RendererScene.cpp
--- a/Engine/Source/Runtime/Renderer/Private/RendererScene.cpp
+++ b/Engine/Source/Runtime/Renderer/Private/RendererScene.cpp
@@ -130,7 +130,7 @@
 /** Affects BasePassPixelShader.usf so must relaunch editor to recompile shaders. */
 static TAutoConsoleVariable<int32> CVarEarlyZPassOnlyMaterialMasking(
 	TEXT("r.EarlyZPassOnlyMaterialMasking"),
-	0,
+	1,
 	TEXT("Whether to compute materials' mask opacity only in early Z pass. Changing this setting requires restarting the editor.\n")
 	TEXT("Note: Needs r.EarlyZPass == 2 && r.EarlyZPassMovable == 1"),
 	ECVF_RenderThreadSafe | ECVF_ReadOnly
