Substrate - NMT - Fix check leading to crash when we already handle the case when the material resource is nullptr.
+ renaming for clarity and ease of searching

#tests PC editor dx12
#rnx

#Synced-CL 50095138

diff --git a/Engine/Source/Runtime/Engine/Private/Materials/MaterialIRModuleBuilder.cpp b/Engine/Source/Runtime/Engine/Private/Materials/MaterialIRModuleBuilder.cpp
--- a/Engine/Source/Runtime/Engine/Private/Materials/MaterialIRModuleBuilder.cpp
+++ b/Engine/Source/Runtime/Engine/Private/Materials/MaterialIRModuleBuilder.cpp
@@ -769,10 +769,10 @@ struct FMaterialIRModuleBuilderImpl
 	{
 		MIR::FValue* FrontMaterialPropertyValue = Emitter.SubstrateDefaultSlab();
 
-		FMaterial* Material = SubstrateTranslatorData.GetMaterial();
-		if (Material)
+		FMaterial* CompilerMaterial = SubstrateTranslatorData.GetMaterialResource();
+		if (CompilerMaterial)
 		{
-			UMaterialInterface* MaterialInterface = Material ? Material->GetMaterialInterface() : nullptr;
+			UMaterialInterface* MaterialInterface = CompilerMaterial ? CompilerMaterial->GetMaterialInterface() : nullptr;
 			UMaterial* BaseMaterial = MaterialInterface ? MaterialInterface->GetMaterial() : nullptr;
 			UMaterialEditorOnlyData* EditorOnlyData = BaseMaterial ? BaseMaterial->GetEditorOnlyData() : nullptr;
 
@@ -790,7 +790,7 @@ struct FMaterialIRModuleBuilderImpl
 				//		bHasSSS, SSSProfile,
 				//		*SubstrateTranslatorData.GetClearCoatNormal(), *SubstrateTranslatorData.GetCustomTangent());
 			}
-			else if ((Material->GetMaterialDomain() == MD_Surface || Material->GetMaterialDomain() == MD_RuntimeVirtualTexture || Material->GetMaterialDomain() == MD_DeferredDecal) && EditorOnlyData)
+			else if ((CompilerMaterial->GetMaterialDomain() == MD_Surface || CompilerMaterial->GetMaterialDomain() == MD_RuntimeVirtualTexture || CompilerMaterial->GetMaterialDomain() == MD_DeferredDecal) && EditorOnlyData)
 			{
 				USubsurfaceProfile* SSSProfile = MaterialInterface ? MaterialInterface->GetSubsurfaceProfileRoot_Internal() : nullptr;
 				const bool bHasSSS = SSSProfile != nullptr;
@@ -808,7 +808,7 @@ struct FMaterialIRModuleBuilderImpl
 					bHasSSS, SSSProfile,
 					EditorOnlyData);
 			}
-			else if (Material->GetMaterialDomain() == MD_Volume && EditorOnlyData)
+			else if (CompilerMaterial->GetMaterialDomain() == MD_Volume && EditorOnlyData)
 			{
 				check(false); // SUBSTRATE_TODO
 				return FrontMaterialPropertyValue;
@@ -821,19 +821,19 @@ struct FMaterialIRModuleBuilderImpl
 				//		MaterialInterface->GetShadingModels().HasOnlyShadingModel(MSM_Unlit),
 				//		EditorOnlyData);
 			}
-			else if (Material->GetMaterialDomain() == MD_LightFunction && EditorOnlyData)
+			else if (CompilerMaterial->GetMaterialDomain() == MD_LightFunction && EditorOnlyData)
 			{
 				check(false); // SUBSTRATE_TODO
 				return FrontMaterialPropertyValue;
 				//	return UMaterialExpressionSubstrateLightFunction::CompileCommon(this, EditorOnlyData->EmissiveColor, EditorOnlyData);
 			}
-			else if (Material->GetMaterialDomain() == MD_PostProcess && EditorOnlyData)
+			else if (CompilerMaterial->GetMaterialDomain() == MD_PostProcess && EditorOnlyData)
 			{
 				check(false); // SUBSTRATE_TODO
 				return FrontMaterialPropertyValue;
 				//	return UMaterialExpressionSubstratePostProcess::CompileCommon(this, EditorOnlyData->EmissiveColor, EditorOnlyData->Opacity, EditorOnlyData);
 			}
-			else if (Material->GetMaterialDomain() == MD_UI && EditorOnlyData)
+			else if (CompilerMaterial->GetMaterialDomain() == MD_UI && EditorOnlyData)
 			{
 				check(false); // SUBSTRATE_TODO
 				return FrontMaterialPropertyValue;
diff --git a/Engine/Source/Runtime/Engine/Private/Materials/SubstrateTranslatorCommon.cpp b/Engine/Source/Runtime/Engine/Private/Materials/SubstrateTranslatorCommon.cpp
--- a/Engine/Source/Runtime/Engine/Private/Materials/SubstrateTranslatorCommon.cpp
+++ b/Engine/Source/Runtime/Engine/Private/Materials/SubstrateTranslatorCommon.cpp
@@ -285,7 +285,7 @@ void FSubstrateCompilationContext::Initialise()
 
 bool FSubstrateCompilationContext::SubstrateGenerateDerivedMaterialOperatorData(FSubstrateTranslatorData& SubstrateTranslatorData)
 {
-	FMaterial* CompilerMaterial = SubstrateTranslatorData.GetMaterial();
+	FMaterial* CompilerMaterial = SubstrateTranslatorData.GetMaterialResource();
 
 	if (SubstrateMaterialExpressionRegisteredOperators.IsEmpty())
 	{
@@ -1457,7 +1457,7 @@ FSubstrateTranslatorData::FSubstrateTranslatorData()
 , bSubstrateUsesConversionFromLegacy(false)
 , bSubstrateMaterialOutputOpaqueRoughRefractions(false)
 , SubstrateCompilationConfig()
-, Material(nullptr)
+, MaterialResource(nullptr)
 , MaterialCompilationOutput(nullptr)
 , CompilerLegacy(nullptr)
 , CompilerNew(nullptr)
@@ -1472,11 +1472,10 @@ FSubstrateTranslatorData::~FSubstrateTranslatorData()
 
 void FSubstrateTranslatorData::Initialize(FMaterial* InCompilerMaterial, FMaterialCompilationOutput& InMaterialCompilationOutput, FMaterialCompiler* InCompilerLegacy, MIR::FEmitter* InCompilerNew)
 {
-	check(InCompilerMaterial);
 	check(!!InCompilerLegacy != !!InCompilerNew); // XOR
 
 	FrontMaterialExpr = nullptr;
-	Material = InCompilerMaterial;
+	MaterialResource = InCompilerMaterial;	// We handle the case where MaterialResource is nullptr down the line.
 	MaterialCompilationOutput = &InMaterialCompilationOutput;
 	CompilerLegacy = InCompilerLegacy;
 	CompilerNew = InCompilerNew;
@@ -1513,7 +1512,7 @@ void FSubstrateTranslatorData::ProcessSubstrateTreeAndTopology(
 	UMaterialExpression* ExpressionToPreview,
 	bool& bDoClearFunctionStack)
 {
-	if (!Material)
+	if (!MaterialResource)
 		return;
 
 
@@ -1544,7 +1543,7 @@ void FSubstrateTranslatorData::ProcessSubstrateTreeAndTopology(
 	}
 	else
 	{
-		FSubstrateMaterialInput* FrontMaterialInput = Material->GetMaterialInterface() ? &Material->GetMaterialInterface()->GetMaterial()->GetEditorOnlyData()->FrontMaterial : nullptr;
+		FSubstrateMaterialInput* FrontMaterialInput = MaterialResource->GetMaterialInterface() ? &MaterialResource->GetMaterialInterface()->GetMaterial()->GetEditorOnlyData()->FrontMaterial : nullptr;
 		FrontMaterialExpr = FrontMaterialInput ? FrontMaterialInput->GetTracedInput().Expression : nullptr;
 		FrontMaterialOutputIndex = FrontMaterialInput ? FrontMaterialInput->OutputIndex : INDEX_NONE;
 	}
@@ -1600,7 +1599,7 @@ void FSubstrateTranslatorData::ProcessSubstrateTreeAndTopology(
 
 	if (bMaterialUsesRootNodeToSubstrateHiddenConversion)
 	{
-		UMaterialInterface* MaterialInterface = Material->GetMaterialInterface();
+		UMaterialInterface* MaterialInterface = MaterialResource->GetMaterialInterface();
 		UMaterial* BaseMaterial = MaterialInterface ? MaterialInterface->GetMaterial() : nullptr;
 		UMaterialEditorOnlyData* EditorOnlyData = BaseMaterial ? BaseMaterial->GetEditorOnlyData() : nullptr;
 
@@ -1686,7 +1685,7 @@ void FSubstrateTranslatorData::ProcessSubstrateTreeAndTopology(
 						EditorOnlyData->ShadingModelFromMaterialExpression.IsConnected(),
 						EditorOnlyData->MaterialAttributes.IsConnected(MP_EmissiveColor));
 				}
-				else if ((Material->GetMaterialDomain() == MD_Surface || Material->GetMaterialDomain() == MD_RuntimeVirtualTexture || Material->GetMaterialDomain() == MD_DeferredDecal) && EditorOnlyData)
+				else if ((MaterialResource->GetMaterialDomain() == MD_Surface || MaterialResource->GetMaterialDomain() == MD_RuntimeVirtualTexture || MaterialResource->GetMaterialDomain() == MD_DeferredDecal) && EditorOnlyData)
 				{
 					FSubstrateOperator* Op = UMaterialExpressionSubstrateShadingModels::SubstrateGenerateMaterialTopologyTreeCommon(
 						*this, RootNodeConversionGUID, nullptr /*Parent*/, 0,
@@ -1696,7 +1695,7 @@ void FSubstrateTranslatorData::ProcessSubstrateTreeAndTopology(
 						*CustomTangent,
 						EditorOnlyData->ShadingModelFromMaterialExpression);
 
-					if (Material->GetMaterialDomain() == MD_DeferredDecal)
+					if (MaterialResource->GetMaterialDomain() == MD_DeferredDecal)
 					{
 						// We notify that this is for decal using this usage flag.
 						// Otherwise, we do not need to do anything else, the ConvertToDecal node is only here to enforce parameter blending.
@@ -1704,24 +1703,24 @@ void FSubstrateTranslatorData::ProcessSubstrateTreeAndTopology(
 						Op->SubUsage |= SUBSTRATE_OPERATOR_SUBUSAGE_DECAL;
 					}
 				}
-				else if (Material->GetMaterialDomain() == MD_Volume && EditorOnlyData)
+				else if (MaterialResource->GetMaterialDomain() == MD_Volume && EditorOnlyData)
 				{
 					UMaterialExpressionSubstrateVolumetricFogCloudBSDF::SubstrateGenerateMaterialTopologyTreeCommon(
 						*this, RootNodeConversionGUID, nullptr /*Parent*/, 0,
 						EditorOnlyData->EmissiveColor,
 						EditorOnlyData->AmbientOcclusion);
 				}
-				else if (Material->GetMaterialDomain() == MD_LightFunction && EditorOnlyData)
+				else if (MaterialResource->GetMaterialDomain() == MD_LightFunction && EditorOnlyData)
 				{
 					UMaterialExpressionSubstrateLightFunction::SubstrateGenerateMaterialTopologyTreeCommon(
 						*this, RootNodeConversionGUID, nullptr /*Parent*/, 0);
 				}
-				else if (Material->GetMaterialDomain() == MD_PostProcess && EditorOnlyData)
+				else if (MaterialResource->GetMaterialDomain() == MD_PostProcess && EditorOnlyData)
 				{
 					UMaterialExpressionSubstratePostProcess::SubstrateGenerateMaterialTopologyTreeCommon(
 						*this, RootNodeConversionGUID, nullptr /*Parent*/, 0);
 				}
-				else if (Material->GetMaterialDomain() == MD_UI && EditorOnlyData)
+				else if (MaterialResource->GetMaterialDomain() == MD_UI && EditorOnlyData)
 				{
 					UMaterialExpressionSubstrateUI::SubstrateGenerateMaterialTopologyTreeCommon(
 						*this, RootNodeConversionGUID, nullptr /*Parent*/, 0);
@@ -1733,7 +1732,7 @@ void FSubstrateTranslatorData::ProcessSubstrateTreeAndTopology(
 
 				if (SubstrateCtx.bSubstrateTreeOutOfStackDepthOccurred)
 				{
-					Errorf(FString::Printf(TEXT(" %s [%s]: Substrate - Cyclic graph detected when we only support acyclic graph."), *Material->GetAssetName(), *Material->GetAssetPath().ToString()));
+					Errorf(FString::Printf(TEXT(" %s [%s]: Substrate - Cyclic graph detected when we only support acyclic graph."), *MaterialResource->GetAssetName(), *MaterialResource->GetAssetPath().ToString()));
 					return;
 				}
 				if (!SubstrateCtx.SubstrateGenerateDerivedMaterialOperatorData(*this))
@@ -1756,7 +1755,7 @@ void FSubstrateTranslatorData::ProcessSubstrateTreeAndTopology(
 			FSubstrateCompilationContext& SubstrateCtx = SubstrateCompilationContext[SubstrateCompilationContextIndex];
 			SubstrateCtx = FSubstrateCompilationContext(ESubstrateCompilationContext(SubstrateCompilationContextIndex));
 
-			FScalarMaterialInput* SurfaceThickness = Material->IsThinSurface() && Material->GetMaterialInterface() ? &Material->GetMaterialInterface()->GetMaterial()->GetEditorOnlyData()->SurfaceThickness : nullptr;
+			FScalarMaterialInput* SurfaceThickness = MaterialResource->IsThinSurface() && MaterialResource->GetMaterialInterface() ? &MaterialResource->GetMaterialInterface()->GetMaterial()->GetEditorOnlyData()->SurfaceThickness : nullptr;
 
 			SubstrateThicknessStackPush(nullptr, SurfaceThickness);
 
@@ -1767,7 +1766,7 @@ void FSubstrateTranslatorData::ProcessSubstrateTreeAndTopology(
 
 			if (SubstrateCtx.bSubstrateTreeOutOfStackDepthOccurred)
 			{
-				Errorf(FString::Printf(TEXT(" %s [%s]: Substrate - Cyclic graph detected when we only support acyclic graph."), *Material->GetAssetName(), *Material->GetAssetPath().ToString()));
+				Errorf(FString::Printf(TEXT(" %s [%s]: Substrate - Cyclic graph detected when we only support acyclic graph."), *MaterialResource->GetAssetName(), *MaterialResource->GetAssetPath().ToString()));
 				return;
 			}
 			if (!SubstrateCtx.SubstrateGenerateDerivedMaterialOperatorData(*this))
@@ -1812,11 +1811,11 @@ bool FSubstrateTranslatorData::GenerateSubstrateTreeHLSLCode(bool bSubstrateFron
 				// The material can be null when some entries are automatically generated, for instance in the material layer blending system
 				bSubstrateFrontMaterialIsValid &= SubstrateCtx.SubstrateMaterialEffectiveClosureCount > 0;
 
-				UMaterialInterface* MaterialInterface = Material->GetMaterialInterface();
+				UMaterialInterface* MaterialInterface = MaterialResource->GetMaterialInterface();
 				UMaterial* BaseMaterial = MaterialInterface ? MaterialInterface->GetMaterial() : nullptr;
 				if (!bSubstrateFrontMaterialIsValid && BaseMaterial && BaseMaterial->bUseMaterialAttributes)
 				{
-					Errorf(FString::Printf(TEXT(" %s [%s]: Substrate - Material has no BSDF: this can happen when the SubstrateConvertMaterialAttributes node is not used and when material attributes are enabled."), *Material->GetAssetName(), *Material->GetAssetPath().ToString()));
+					Errorf(FString::Printf(TEXT(" %s [%s]: Substrate - Material has no BSDF: this can happen when the SubstrateConvertMaterialAttributes node is not used and when material attributes are enabled."), *MaterialResource->GetAssetName(), *MaterialResource->GetAssetPath().ToString()));
 					return false;
 				}
 			}
@@ -2051,7 +2050,7 @@ bool FSubstrateTranslatorData::GenerateSubstrateTreeHLSLCode(bool bSubstrateFron
 					SubstrateCtx.SubstrateEvaluateSharedLocalBases(RequestedSharedLocalBasesCount, nullptr, nullptr);// CompilerNew);
 					if (RequestedSharedLocalBasesCount > SUBSTRATE_MAX_SHAREDLOCALBASES_REGISTERS)
 					{
-						Errorf(FString::Printf(TEXT(" %s [%s]: Substrate - Material has more unique normal/tangent basis than the allowed limit %d/%d."), *Material->GetAssetName(), *Material->GetAssetPath().ToString(), RequestedSharedLocalBasesCount, SUBSTRATE_MAX_SHAREDLOCALBASES_REGISTERS));
+						Errorf(FString::Printf(TEXT(" %s [%s]: Substrate - Material has more unique normal/tangent basis than the allowed limit %d/%d."), *MaterialResource->GetAssetName(), *MaterialResource->GetAssetPath().ToString(), RequestedSharedLocalBasesCount, SUBSTRATE_MAX_SHAREDLOCALBASES_REGISTERS));
 					}
 				}
 			}
@@ -2149,7 +2148,7 @@ void FSubstrateTranslatorData::GenerateMaterialCompilationOutputAndDefines(
 {
 	bSubstratePremultipliedAlphaOpacityOverridden = bMaterialIsSubstrate && bOpacityPropertyIsUsed;
 
-	bDistortionAccountForCoverage = bSubstrateEnabled && !bMaterialUsesRootNodeToSubstrateHiddenConversion && Material->GetRefractionCoverageMode() == RCM_CoverageAccountedFor;
+	bDistortionAccountForCoverage = bSubstrateEnabled && !bMaterialUsesRootNodeToSubstrateHiddenConversion && MaterialResource->GetRefractionCoverageMode() == RCM_CoverageAccountedFor;
 
 	if (bMaterialIsSubstrate)
 	{
@@ -2160,7 +2159,7 @@ void FSubstrateTranslatorData::GenerateMaterialCompilationOutputAndDefines(
 		// Unlit cannot be combined with other BSDF so we can simply pick the default strata context
 		bSubstrateOptimizedUnlit = SubstrateCompilationContext[ESubstrateCompilationContext::SCC_Default].bSubstrateMaterialIsUnlitNode;
 
-		bSubstrateRoughnessTracking = Material->HasSubstrateRoughnessTracking();
+		bSubstrateRoughnessTracking = MaterialResource->HasSubstrateRoughnessTracking();
 
 		{
 			// For now, the fully simplified mode is used for Lumen or anything else supported inlined evaluation. The export is only valid for the default case.
@@ -2385,21 +2384,21 @@ FSubstrateOperator& FSubstrateTranslatorData::SubstrateCompilationRegisterOperat
 	if (SubstrateCtx.SubstrateMaterialExpressionToOperatorIndex.Find(SubstrateExpressionGuid))
 	{
 		// It is not possible to register/use a Substrate BSDF multiple times with this same exact graph path. (that would break the Substrate tree code generation)
-		Errorf(FString::Printf(TEXT("Material %s: It is not possible to uses a Substrate BSDF (or any ouput of type SubstrateData) multiple times within a Substrate material topology with the same graph path GUID (asset: %s).\r\n"), *Material->GetAssetName(), *Material->GetAssetPath().ToString()));
+		Errorf(FString::Printf(TEXT("Material %s: It is not possible to uses a Substrate BSDF (or any ouput of type SubstrateData) multiple times within a Substrate material topology with the same graph path GUID (asset: %s).\r\n"), *MaterialResource->GetAssetName(), *MaterialResource->GetAssetPath().ToString()));
 		return DefaultOperatorOnError;
 	}
 
 	const uint32 NewOperatorIndex = SubstrateCtx.SubstrateMaterialExpressionToOperatorIndex.Num();
 	if (NewOperatorIndex >= SUBSTRATE_MAX_COMPILER_REGISTERED_OPERATOR_COUNT)
 	{
-		Errorf(FString::Printf(TEXT("Material %s have too many Substrate Operators: the compiler is failing (asset: %s).\r\n"), *Material->GetAssetName(), *Material->GetAssetPath().ToString()));
+		Errorf(FString::Printf(TEXT("Material %s have too many Substrate Operators: the compiler is failing (asset: %s).\r\n"), *MaterialResource->GetAssetName(), *MaterialResource->GetAssetPath().ToString()));
 		return DefaultOperatorOnError;
 	}
 
 	int32* ParentOperatorIndex = SubstrateCtx.SubstrateMaterialExpressionToOperatorIndex.Find(SubstrateParentExpressionGuid);
 	if (Parent != nullptr && ParentOperatorIndex == nullptr)
 	{
-		Errorf(FString::Printf(TEXT("Material %s tries to register unknown operator parents (asset: %s).\r\n"), *Material->GetAssetName(), *Material->GetAssetPath().ToString()));
+		Errorf(FString::Printf(TEXT("Material %s tries to register unknown operator parents (asset: %s).\r\n"), *MaterialResource->GetAssetName(), *MaterialResource->GetAssetPath().ToString()));
 		return DefaultOperatorOnError;
 	}
 
@@ -2465,7 +2464,7 @@ FSubstrateOperator* FSubstrateTranslatorData::SubstrateCompilationGetOperatorFro
 	int32 OperatorCount = SubstrateCtx.SubstrateMaterialExpressionToOperatorIndex.Num();
 	if (OperatorIndex < 0 || OperatorIndex >= OperatorCount)
 	{
-		Errorf(FString::Printf(TEXT("SubstrateCompilationGetOperatorFromIndex - OperatorIndex out of range %s (asset: %s).\r\n"), *Material->GetAssetName(), *Material->GetAssetPath().ToString()));
+		Errorf(FString::Printf(TEXT("SubstrateCompilationGetOperatorFromIndex - OperatorIndex out of range %s (asset: %s).\r\n"), *MaterialResource->GetAssetName(), *MaterialResource->GetAssetPath().ToString()));
 		return nullptr;
 	};
 	return &SubstrateCtx.SubstrateMaterialExpressionRegisteredOperators[OperatorIndex];
diff --git a/Engine/Source/Runtime/Engine/Private/Materials/SubstrateTranslatorCommon.h b/Engine/Source/Runtime/Engine/Private/Materials/SubstrateTranslatorCommon.h
--- a/Engine/Source/Runtime/Engine/Private/Materials/SubstrateTranslatorCommon.h
+++ b/Engine/Source/Runtime/Engine/Private/Materials/SubstrateTranslatorCommon.h
@@ -175,7 +175,7 @@ private:
 * 3- From Step_EmitSetMaterialPropertyInstructions, this is where we assigne the FrontMaterial IR using BuildCommon function when hidden conversion is enabled.
 *	The result is then assigned to SubstrateFrontMaterial and SubstrateFrontMaterialFullySimplified custom output, similarly to Step2 above.
 *	This is alsowhere we generate the HLSL backend code we need for several substrate element such as front material & local shared bases assignement, 
-*	substrate tree functions, same defines as for hte legacy translator, etc.
+*	substrate tree functions, same defines as for the legacy translator, etc.
 * 4- And from there we are DONE.
 * 
 * Notes:
@@ -247,7 +247,7 @@ struct FSubstrateTranslatorData : public FSubstrateTranslatorDataInterface
 
 	EShaderPlatform GetShaderPlatform();
 
-	FMaterial* GetMaterial() { return Material; }
+	FMaterial* GetMaterialResource() { return MaterialResource; }
 	FMaterialCompilationOutput& GetMaterialCompilationOutput() { return *MaterialCompilationOutput; }
 
 	void ProcessSubstrateTreeAndTopology(
@@ -310,7 +310,7 @@ struct FSubstrateTranslatorData : public FSubstrateTranslatorDataInterface
 
 private:
 
-	FMaterial*						Material;
+	FMaterial*						MaterialResource;
 	FMaterialCompilationOutput*		MaterialCompilationOutput;
 	FMaterialCompiler*				CompilerLegacy;
 	MIR::FEmitter*					CompilerNew;
