Fix the UE build by removing the check's in the new `FAttachmentCounter`. Apparently the original docs were wrong and I shouldn't have been asserting these.

#Synced-CL 50099782

diff --git a/Engine/Source/Runtime/Engine/Public/PrimitiveSceneInfoData.h b/Engine/Source/Runtime/Engine/Public/PrimitiveSceneInfoData.h
--- a/Engine/Source/Runtime/Engine/Public/PrimitiveSceneInfoData.h
+++ b/Engine/Source/Runtime/Engine/Public/PrimitiveSceneInfoData.h
@@ -26,7 +26,7 @@ struct FAttachmentCounter final
 #if UE_AUTORTFM
 	void Increment()
 	{
-		check(IsInGameThread());
+		// TODO: should we `check(IsInGameThread());`?
 
 		UE_AUTORTFM_OPEN
 		{
@@ -41,8 +41,7 @@ struct FAttachmentCounter final
 
 	AUTORTFM_DISABLE void Decrement()
 	{
-		check(IsInRenderingThread());
-
+		// TODO: should we `check(IsInRenderingThread());`?
 		Payload->Payload--;
 	}
 
@@ -74,13 +73,13 @@ private:
 
 	void Increment()
 	{
-		check(IsInGameThread());
+		// TODO: should we `check(IsInGameThread());`?
 		Payload++;
 	}
 
 	void Decrement()
 	{
-		check(IsInRenderingThread());
+		// TODO: should we `check(IsInRenderingThread());`?
 		Payload--;
 	}
 
