Animation Sync Group support for Rewind Debugger

- Added a RouteId_TickRecord3 that traces sync group information from a given tick record.
- Remove all the TRACE_TICK_RECORD() macro instances and moved them to _AFTER_ we tick the asset players otherwise we are displayed non-update information and no sync group info is available.
- Update SBlendWeightsView to display sync group info.

#Synced-CL 49789214

diff --git a/Engine/Plugins/Animation/GameplayInsights/Source/GameplayInsights/Private/AnimationAnalyzer.cpp b/Engine/Plugins/Animation/GameplayInsights/Source/GameplayInsights/Private/AnimationAnalyzer.cpp
--- a/Engine/Plugins/Animation/GameplayInsights/Source/GameplayInsights/Private/AnimationAnalyzer.cpp
+++ b/Engine/Plugins/Animation/GameplayInsights/Source/GameplayInsights/Private/AnimationAnalyzer.cpp
@@ -20,6 +20,7 @@ void FAnimationAnalyzer::OnAnalysisBegin(const FOnAnalysisContext& Context)
 
 	Builder.RouteEvent(RouteId_TickRecord, "Animation", "TickRecord");
 	Builder.RouteEvent(RouteId_TickRecord2, "Animation", "TickRecord2");
+	Builder.RouteEvent(RouteId_TickRecord3, "Animation", "TickRecord3");
 	Builder.RouteEvent(RouteId_SkeletalMesh, "Animation", "SkeletalMesh");
 	Builder.RouteEvent(RouteId_SkeletalMesh2, "Animation", "SkeletalMesh2");
 	Builder.RouteEvent(RouteId_SkeletalMesh3, "Animation", "SkeletalMesh3");
@@ -67,6 +68,35 @@ bool FAnimationAnalyzer::OnEvent(uint16 RouteId, EStyle Style, const FOnEventCon
 	const auto& EventData = Context.EventData;
 	switch (RouteId)
 	{
+	case RouteId_TickRecord3:
+	{
+		double RecordingTime = EventData.GetValue<double>("RecordingTime");
+		uint64 Cycle = EventData.GetValue<uint64>("Cycle");
+		uint64 AnimInstanceId = EventData.GetValue<uint64>("AnimInstanceId");
+		uint64 AssetId = EventData.GetValue<uint64>("AssetId");
+		int32 NodeId = EventData.GetValue<int32>("NodeId");
+		float BlendWeight = EventData.GetValue<float>("BlendWeight");
+		float PlaybackTime = EventData.GetValue<float>("PlaybackTime");
+		float RootMotionWeight = EventData.GetValue<float>("RootMotionWeight");
+		float PlayRate = EventData.GetValue<float>("PlayRate");
+		float BlendSpacePositionX = EventData.GetValue<float>("BlendSpacePositionX");
+		float BlendSpacePositionY = EventData.GetValue<float>("BlendSpacePositionY");
+		float BlendSpaceFilteredPositionX = EventData.GetValue<float>("BlendSpaceFilteredPositionX");
+		float BlendSpaceFilteredPositionY = EventData.GetValue<float>("BlendSpaceFilteredPositionY");
+		uint16 FrameCounter = EventData.GetValue<uint16>("FrameCounter");
+		bool bLooping = EventData.GetValue<bool>("Looping");
+		bool bIsBlendSpace = EventData.GetValue<bool>("IsBlendSpace");
+		uint32 SyncGroupNameId = EventData.GetValue<uint32>("SyncGroupNameId");
+		bool IsSyncGroupLeader = EventData.GetValue<bool>("IsSyncGroupLeader");
+		float SyncLeaderScore = EventData.GetValue<float>("SyncLeaderScore");
+		float TimeToPrevMarker = EventData.GetValue<float>("TimeToPrevMarker");
+		float TimeToNextMarker = EventData.GetValue<float>("TimeToNextMarker");
+		int32 PrevMarkerIndex = EventData.GetValue<int32>("PrevMarkerIndex");
+		int32 NextMarkerIndex = EventData.GetValue<int32>("NextMarkerIndex");
+			
+		AnimationProvider.AppendTickRecord(AnimInstanceId, Context.EventTime.AsSeconds(Cycle), RecordingTime, AssetId, NodeId, BlendWeight, PlaybackTime, RootMotionWeight, PlayRate, BlendSpacePositionX, BlendSpacePositionY, BlendSpaceFilteredPositionX, BlendSpaceFilteredPositionY, FrameCounter, bLooping, bIsBlendSpace, SyncGroupNameId, IsSyncGroupLeader, SyncLeaderScore, TimeToPrevMarker, TimeToNextMarker, PrevMarkerIndex, NextMarkerIndex);
+		break;
+	}
 	case RouteId_TickRecord2:
 	{
 		double RecordingTime = EventData.GetValue<double>("RecordingTime");
diff --git a/Engine/Plugins/Animation/GameplayInsights/Source/GameplayInsights/Private/AnimationAnalyzer.h b/Engine/Plugins/Animation/GameplayInsights/Source/GameplayInsights/Private/AnimationAnalyzer.h
--- a/Engine/Plugins/Animation/GameplayInsights/Source/GameplayInsights/Private/AnimationAnalyzer.h
+++ b/Engine/Plugins/Animation/GameplayInsights/Source/GameplayInsights/Private/AnimationAnalyzer.h
@@ -22,6 +22,7 @@ private:
 	{
 		RouteId_TickRecord,
 		RouteId_TickRecord2,
+		RouteId_TickRecord3,
 		RouteId_SkeletalMesh,
 		RouteId_SkeletalMesh2,
 		RouteId_SkeletalMesh3,
diff --git a/Engine/Plugins/Animation/GameplayInsights/Source/GameplayInsights/Private/AnimationProvider.cpp b/Engine/Plugins/Animation/GameplayInsights/Source/GameplayInsights/Private/AnimationProvider.cpp
--- a/Engine/Plugins/Animation/GameplayInsights/Source/GameplayInsights/Private/AnimationProvider.cpp
+++ b/Engine/Plugins/Animation/GameplayInsights/Source/GameplayInsights/Private/AnimationProvider.cpp
@@ -633,6 +633,59 @@ void FAnimationProvider::AppendTickRecord(uint64 InAnimInstanceId, double InProf
 	Session.UpdateDurationSeconds(InProfileTime);
 }
 
+void FAnimationProvider::AppendTickRecord(uint64 InAnimInstanceId, double InProfileTime, double InRecordingTime, uint64 InAssetId, int32 InNodeId, float InBlendWeight, float InPlaybackTime, float InRootMotionWeight, float InPlayRate, float InBlendSpacePositionX, float InBlendSpacePositionY, float InBlendSpaceFilteredPositionX, float InBlendSpaceFilteredPositionY, uint16 InFrameCounter, bool bInLooping, bool bInIsBlendSpace, uint32 InSyncGroupName, bool bInIsSyncGroupLeader, float InSyncLeaderScore, float TimeToPrevMarker, float TimeToNextMarker, int32 PrevMarkerIndex, int32 NextMarkerIndex)
+{
+	Session.WriteAccessCheck();
+
+	bHasAnyData = true;
+
+	TSharedPtr<FTickRecordTimelineStorage> TimelineStorage;
+	uint32* TimelineStorageIndexPtr = ObjectIdToTickRecordTimelineStorage.Find(InAnimInstanceId);
+	if(TimelineStorageIndexPtr != nullptr)
+	{
+		TimelineStorage = TickRecordTimelineStorage[*TimelineStorageIndexPtr];
+	}
+	else
+	{
+		ObjectIdToTickRecordTimelineStorage.Add(InAnimInstanceId, TickRecordTimelineStorage.Num());
+		TimelineStorage = TickRecordTimelineStorage.Add_GetRef(MakeShared<FTickRecordTimelineStorage>());
+		TimelineStorage->Timeline = MakeShared<TraceServices::TPointTimeline<FTickRecordMessage>>(Session.GetLinearAllocator());
+		TimelineStorage->Timeline->SetEnumerateOutsideRange(true);
+	}
+
+	check(TimelineStorage.IsValid());
+
+	TimelineStorage->AssetIdAndPlayers.Add(TTuple<uint64, int32>(InAssetId, InNodeId));
+
+	FTickRecordMessage Message;
+	Message.AnimInstanceId = InAnimInstanceId;
+	Message.AssetId = InAssetId;
+	Message.RecordingTime = InRecordingTime;
+	Message.NodeId = InNodeId;
+	Message.BlendWeight = InBlendWeight;
+	Message.PlaybackTime = InPlaybackTime;
+	Message.RootMotionWeight = InRootMotionWeight;
+	Message.PlayRate = InPlayRate;
+	Message.BlendSpacePositionX = InBlendSpacePositionX;
+	Message.BlendSpacePositionY = InBlendSpacePositionY;
+	Message.BlendSpaceFilteredPositionX = InBlendSpaceFilteredPositionX;
+	Message.BlendSpaceFilteredPositionY = InBlendSpaceFilteredPositionY;
+	Message.FrameCounter = InFrameCounter;
+	Message.bLooping = bInLooping;
+	Message.bIsBlendSpace = bInIsBlendSpace;
+	Message.bSyncGroupNameId = InSyncGroupName;
+	Message.bIsSyncGroupLeader = bInIsSyncGroupLeader;
+	Message.SyncLeaderScore = InSyncLeaderScore;
+	Message.PrevMarkerIndex = PrevMarkerIndex;
+	Message.NextMarkerIndex = NextMarkerIndex;
+	Message.TimeToNextMarker = TimeToNextMarker;
+	Message.TimeToPrevMarker = TimeToPrevMarker;
+	
+	TimelineStorage->Timeline->AppendEvent(InProfileTime, Message);
+
+	Session.UpdateDurationSeconds(InProfileTime);
+}
+
 void FAnimationProvider::AppendSkeletalMesh(uint64 InObjectId, uint64 SkeletonId, const TArrayView<const int32>& InParentIndices)
 {
 	Session.WriteAccessCheck();
diff --git a/Engine/Plugins/Animation/GameplayInsights/Source/GameplayInsights/Private/AnimationProvider.h b/Engine/Plugins/Animation/GameplayInsights/Source/GameplayInsights/Private/AnimationProvider.h
--- a/Engine/Plugins/Animation/GameplayInsights/Source/GameplayInsights/Private/AnimationProvider.h
+++ b/Engine/Plugins/Animation/GameplayInsights/Source/GameplayInsights/Private/AnimationProvider.h
@@ -63,6 +63,7 @@ public:
 
 	/** Add a tick record */
 	void AppendTickRecord(uint64 InAnimInstanceId, double InProfileTime, double InRecordingTime, uint64 InAssetId, int32 InNodeId, float InBlendWeight, float InPlaybackTime, float InRootMotionWeight, float InPlayRate, float InBlendSpacePositionX, float InBlendSpacePositionY, float InBlendSpaceFilteredPositionX, float InBlendSpaceFilteredPositionY, uint16 InFrameCounter, bool bInLooping, bool bInIsBlendSpace);
+	void AppendTickRecord(uint64 InAnimInstanceId, double InProfileTime, double InRecordingTime, uint64 InAssetId, int32 InNodeId, float InBlendWeight, float InPlaybackTime, float InRootMotionWeight, float InPlayRate, float InBlendSpacePositionX, float InBlendSpacePositionY, float InBlendSpaceFilteredPositionX, float InBlendSpaceFilteredPositionY, uint16 InFrameCounter, bool bInLooping, bool bInIsBlendSpace, uint32 InSyncGroupName, bool bInIsSyncGroupLeader, float InSyncLeaderScore, float TimeToPrevMarker, float TimeToNextMarker, int32 PrevMarkerIndex, int32 NextMarkerIndex);
 
 	/** Add a skeletal mesh */
 	void AppendSkeletalMesh(uint64 InObjectId, uint64 SkeletonId, const TArrayView<const int32>& ParentIndices);
diff --git a/Engine/Plugins/Animation/GameplayInsights/Source/GameplayInsights/Private/SBlendWeightsView.cpp b/Engine/Plugins/Animation/GameplayInsights/Source/GameplayInsights/Private/SBlendWeightsView.cpp
--- a/Engine/Plugins/Animation/GameplayInsights/Source/GameplayInsights/Private/SBlendWeightsView.cpp
+++ b/Engine/Plugins/Animation/GameplayInsights/Source/GameplayInsights/Private/SBlendWeightsView.cpp
@@ -18,9 +18,9 @@ void SBlendWeightsView::GetVariantsAtFrame(const TraceServices::FFrame& InFrame,
 	{
 		TraceServices::FAnalysisSessionReadScope SessionReadScope(*AnalysisSession);
 
-		AnimationProvider->ReadTickRecordTimeline(ObjectId, [this, &GameplayProvider, &OutVariants, &InFrame](const FAnimationProvider::TickRecordTimeline& InTimeline)
+		AnimationProvider->ReadTickRecordTimeline(ObjectId, [this, &AnimationProvider, &GameplayProvider, &OutVariants, &InFrame](const FAnimationProvider::TickRecordTimeline& InTimeline)
 		{
-			InTimeline.EnumerateEvents(InFrame.StartTime, InFrame.EndTime, [this, &GameplayProvider, &OutVariants, &InFrame](double InStartTime, double InEndTime, uint32 InDepth, const FTickRecordMessage& InMessage)
+			InTimeline.EnumerateEvents(InFrame.StartTime, InFrame.EndTime, [this, &AnimationProvider, &GameplayProvider, &OutVariants, &InFrame](double InStartTime, double InEndTime, uint32 InDepth, const FTickRecordMessage& InMessage)
 			{
 				if(InStartTime >= InFrame.StartTime && InEndTime <= InFrame.EndTime)
 				{
@@ -40,6 +40,7 @@ void SBlendWeightsView::GetVariantsAtFrame(const TraceServices::FFrame& InFrame,
 					Header->AddChild(FVariantTreeNode::MakeFloat(LOCTEXT("PlaybackTime", "Playback Time"), InMessage.PlaybackTime));
 					Header->AddChild(FVariantTreeNode::MakeFloat(LOCTEXT("RootMotionWeight", "Root Motion Weight"), InMessage.RootMotionWeight));
 					Header->AddChild(FVariantTreeNode::MakeFloat(LOCTEXT("PlayRate", "Play Rate"), InMessage.PlayRate));
+					
 					if(InMessage.bIsBlendSpace)
 					{
 						Header->AddChild(FVariantTreeNode::MakeFloat(LOCTEXT("BlendSpacePositionX", "Blend Space Position X"), InMessage.BlendSpacePositionX));
@@ -47,6 +48,22 @@ void SBlendWeightsView::GetVariantsAtFrame(const TraceServices::FFrame& InFrame,
 						Header->AddChild(FVariantTreeNode::MakeFloat(LOCTEXT("BlendSpaceFilteredPositionX", "Blend Space Filtered Position X"), InMessage.BlendSpaceFilteredPositionX));
 						Header->AddChild(FVariantTreeNode::MakeFloat(LOCTEXT("BlendSpaceFilteredPositionY", "Blend Space Filtered Position Y"), InMessage.BlendSpaceFilteredPositionY));
 					}
+
+					static FString NoGroupName(LOCTEXT("NoSyncGroupName", "None").ToString());
+					const TCHAR* CurveName = InMessage.bSyncGroupNameId == 0 ? *NoGroupName : AnimationProvider->GetName(InMessage.bSyncGroupNameId);
+					TSharedRef<FVariantTreeNode> SyncGroup = FVariantTreeNode::MakeString(LOCTEXT("SyncGroupName", "Sync Group Name"), CurveName);
+					Header->AddChild(SyncGroup);
+					
+					// @todo: I'd like to query the name of the sync markers given their index tho I'd prefer to load and cache the related asset in the respective blend weight track to avoid loading it every time here. Additionally, I'd like to add sync groups sub-tracks / curve tracks in a follow up CL.
+					if (InMessage.bSyncGroupNameId != 0)
+					{
+						SyncGroup->AddChild(FVariantTreeNode::MakeBool(LOCTEXT("IsSyncGroupLeader", "Is Sync Group Leader"), InMessage.bIsSyncGroupLeader));
+						SyncGroup->AddChild(FVariantTreeNode::MakeFloat(LOCTEXT("SyncLeaderScore", "Sync Leader Score"), InMessage.SyncLeaderScore));
+						SyncGroup->AddChild(FVariantTreeNode::MakeInt32(LOCTEXT("PreviousMarker", "Previous Marker"), InMessage.PrevMarkerIndex));
+						SyncGroup->AddChild(FVariantTreeNode::MakeFloat(LOCTEXT("PrevMarkerTime", "Time To Previous Marker"), InMessage.TimeToPrevMarker));
+						SyncGroup->AddChild(FVariantTreeNode::MakeInt32(LOCTEXT("NextMarker", "Next Marker"), InMessage.NextMarkerIndex));
+						SyncGroup->AddChild(FVariantTreeNode::MakeFloat(LOCTEXT("NextMarkerTime", "Time To Next Marker"), InMessage.TimeToNextMarker));
+					}
 				}
 				return TraceServices::EEventEnumerate::Continue;
 			});
diff --git a/Engine/Plugins/Animation/GameplayInsights/Source/GameplayInsights/Public/IAnimationProvider.h b/Engine/Plugins/Animation/GameplayInsights/Source/GameplayInsights/Public/IAnimationProvider.h
--- a/Engine/Plugins/Animation/GameplayInsights/Source/GameplayInsights/Public/IAnimationProvider.h
+++ b/Engine/Plugins/Animation/GameplayInsights/Source/GameplayInsights/Public/IAnimationProvider.h
@@ -97,6 +97,13 @@ struct FTickRecordMessage
 	uint16 FrameCounter = 0;
 	bool bLooping = false;
 	bool bIsBlendSpace = false;
+	uint32 bSyncGroupNameId = 0;
+	bool bIsSyncGroupLeader = false;
+	float SyncLeaderScore = 0.0f;
+	float TimeToPrevMarker = 0.0f;
+	float TimeToNextMarker = 0.0f;
+	int32 PrevMarkerIndex = -1;
+	int32 NextMarkerIndex = -1;
 };
 
 enum class EAnimGraphPhase : uint8
diff --git a/Engine/Source/Runtime/AnimGraphRuntime/Private/AnimNodes/AnimNode_BlendSpacePlayer.cpp b/Engine/Source/Runtime/AnimGraphRuntime/Private/AnimNodes/AnimNode_BlendSpacePlayer.cpp
--- a/Engine/Source/Runtime/AnimGraphRuntime/Private/AnimNodes/AnimNode_BlendSpacePlayer.cpp
+++ b/Engine/Source/Runtime/AnimGraphRuntime/Private/AnimNodes/AnimNode_BlendSpacePlayer.cpp
@@ -104,8 +104,6 @@ void FAnimNode_BlendSpacePlayerBase::UpdateInternal(const FAnimationUpdateContex
 
 		SyncScope.AddTickRecord(TickRecord, GetSyncParams(TickRecord.bRequestedInertialization), UE::Anim::FAnimSyncDebugInfo(Context));
 
-		TRACE_ANIM_TICK_RECORD(Context, TickRecord);
-
 #if WITH_EDITORONLY_DATA
 		if (FAnimBlueprintDebugData* DebugData = Context.AnimInstanceProxy->GetAnimBlueprintDebugData())
 		{
diff --git a/Engine/Source/Runtime/AnimGraphRuntime/Private/AnimNodes/AnimNode_RandomPlayer.cpp b/Engine/Source/Runtime/AnimGraphRuntime/Private/AnimNodes/AnimNode_RandomPlayer.cpp
--- a/Engine/Source/Runtime/AnimGraphRuntime/Private/AnimNodes/AnimNode_RandomPlayer.cpp
+++ b/Engine/Source/Runtime/AnimGraphRuntime/Private/AnimNodes/AnimNode_RandomPlayer.cpp
@@ -244,8 +244,6 @@ void FAnimNode_RandomPlayer::Update_AnyThread(const FAnimationUpdateContext& Con
 	UE::Anim::FAnimSyncGroupScope& SyncScope = Context.GetMessageChecked<UE::Anim::FAnimSyncGroupScope>();
 	SyncScope.AddTickRecord(TickRecord, UE::Anim::FAnimSyncParams(), UE::Anim::FAnimSyncDebugInfo(Context));
 
-	TRACE_ANIM_TICK_RECORD(Context, TickRecord);
-
 	if (FAnimationRuntime::HasWeight(NextData->BlendWeight))
 	{
 		FAnimTickRecord NextTickRecord(NextData->Entry->Sequence, true, NextData->PlayRate, false, NextData->BlendWeight * BlendWeight, NextData->CurrentPlayTime, NextData->MarkerTickRecord);
@@ -253,8 +251,6 @@ void FAnimNode_RandomPlayer::Update_AnyThread(const FAnimationUpdateContext& Con
 		NextTickRecord.GatherContextData(Context);
 
 		SyncScope.AddTickRecord(NextTickRecord, UE::Anim::FAnimSyncParams(), UE::Anim::FAnimSyncDebugInfo(Context));
-
-		TRACE_ANIM_TICK_RECORD(Context, NextTickRecord);
 	}
 
 	TRACE_ANIM_NODE_VALUE(Context, TEXT("Current Sequence"), CurrentData ? CurrentData->Entry->Sequence : nullptr);
diff --git a/Engine/Source/Runtime/Engine/Private/Animation/AnimInstance.cpp b/Engine/Source/Runtime/Engine/Private/Animation/AnimInstance.cpp
--- a/Engine/Source/Runtime/Engine/Private/Animation/AnimInstance.cpp
+++ b/Engine/Source/Runtime/Engine/Private/Animation/AnimInstance.cpp
@@ -587,7 +587,6 @@ void UAnimInstance::UpdateMontageSyncGroup()
 
 #if ANIM_TRACE_ENABLED
 				FAnimationUpdateContext UpdateContext(&GetProxyOnGameThread<FAnimInstanceProxy>());
-				TRACE_ANIM_TICK_RECORD(UpdateContext, TickRecord);
 #endif
 			}
 			MontageInstance->bDidUseMarkerSyncThisTick = false;
diff --git a/Engine/Source/Runtime/Engine/Private/Animation/AnimInstanceProxy.cpp b/Engine/Source/Runtime/Engine/Private/Animation/AnimInstanceProxy.cpp
--- a/Engine/Source/Runtime/Engine/Private/Animation/AnimInstanceProxy.cpp
+++ b/Engine/Source/Runtime/Engine/Private/Animation/AnimInstanceProxy.cpp
@@ -1235,9 +1235,27 @@ void FAnimInstanceProxy::UpdateAnimation()
 	}
 
 	UpdateAnimation_WithRoot(Context, RootNode, NAME_AnimGraph);
-
+	
 	// Tick syncing
 	Sync.TickAssetPlayerInstances(*this, CurrentDeltaSeconds);
+	
+	// Now that we have updated our asset players we record their tick records from read buffer (meaning they have up-to-date info).
+#if ANIM_TRACE_ENABLED
+	{
+		for (const TTuple<FName, FAnimGroupInstance>& TickRecordGroup : Sync.GetSyncGroupMapRead())
+		{
+			for (const FAnimTickRecord& TickRecord : TickRecordGroup.Value.ActivePlayers)
+			{
+				TRACE_ANIM_TICK_RECORD(Context, TickRecord)
+			}
+		}
+	
+		for (const FAnimTickRecord& TickRecord : Sync.GetUngroupedActivePlayersRead())
+		{
+			TRACE_ANIM_TICK_RECORD(Context, TickRecord)
+		}
+	}
+#endif
 }
 
 void FAnimInstanceProxy::UpdateAnimation_WithRoot(const FAnimationUpdateContext& InContext, FAnimNode_Base* InRootNode, FName InLayerName)
diff --git a/Engine/Source/Runtime/Engine/Private/Animation/AnimNode_AssetPlayerBase.cpp b/Engine/Source/Runtime/Engine/Private/Animation/AnimNode_AssetPlayerBase.cpp
--- a/Engine/Source/Runtime/Engine/Private/Animation/AnimNode_AssetPlayerBase.cpp
+++ b/Engine/Source/Runtime/Engine/Private/Animation/AnimNode_AssetPlayerBase.cpp
@@ -40,8 +40,6 @@ void FAnimNode_AssetPlayerBase::CreateTickRecordForNode(const FAnimationUpdateCo
 
 	// Add asset player to synchronizer
 	SyncScope.AddTickRecord(TickRecord, GetSyncParams(TickRecord.bRequestedInertialization), UE::Anim::FAnimSyncDebugInfo(Context));
-
-	TRACE_ANIM_TICK_RECORD(Context, TickRecord);
 }
 
 float FAnimNode_AssetPlayerBase::GetAccumulatedTime() const
diff --git a/Engine/Source/Runtime/Engine/Private/Animation/AnimSingleNodeInstanceProxy.cpp b/Engine/Source/Runtime/Engine/Private/Animation/AnimSingleNodeInstanceProxy.cpp
--- a/Engine/Source/Runtime/Engine/Private/Animation/AnimSingleNodeInstanceProxy.cpp
+++ b/Engine/Source/Runtime/Engine/Private/Animation/AnimSingleNodeInstanceProxy.cpp
@@ -588,7 +588,6 @@ void FAnimNode_SingleNode::Update_AnyThread(const FAnimationUpdateContext& Conte
 			
 			SyncScope.AddTickRecord(TickRecord);
 
-			TRACE_ANIM_TICK_RECORD(Context, TickRecord);
 #if WITH_EDITORONLY_DATA
 			PreviewBasePose = BlendSpace->PreviewBasePose;
 #endif
@@ -600,8 +599,6 @@ void FAnimNode_SingleNode::Update_AnyThread(const FAnimationUpdateContext& Conte
 			
 			SyncScope.AddTickRecord(TickRecord);
 
-			TRACE_ANIM_TICK_RECORD(Context, TickRecord);
-
 			// if it's not looping, just set play to be false when reached to end
 			if (!Proxy->bLooping)
 			{
@@ -619,8 +616,6 @@ void FAnimNode_SingleNode::Update_AnyThread(const FAnimationUpdateContext& Conte
 			
 			SyncScope.AddTickRecord(TickRecord);
 
-			TRACE_ANIM_TICK_RECORD(Context, TickRecord);
-
 			// if it's not looping, just set play to be false when reached to end
 			if (!Proxy->bLooping)
 			{
@@ -638,8 +633,6 @@ void FAnimNode_SingleNode::Update_AnyThread(const FAnimationUpdateContext& Conte
 			
 			SyncScope.AddTickRecord(TickRecord);
 
-			TRACE_ANIM_TICK_RECORD(Context, TickRecord);
-
 			// if it's not looping, just set play to be false when reached to end
 			if (!Proxy->bLooping)
 			{
@@ -679,8 +672,6 @@ void FAnimNode_SingleNode::Update_AnyThread(const FAnimationUpdateContext& Conte
 			TickRecord.DeltaTimeRecord = &(Proxy->DeltaTimeRecord);
 
 			SyncScope.AddTickRecord(TickRecord);
-
-			TRACE_ANIM_TICK_RECORD(Context, TickRecord);
 		}
 	}
 
diff --git a/Engine/Source/Runtime/Engine/Private/Animation/AnimTrace.cpp b/Engine/Source/Runtime/Engine/Private/Animation/AnimTrace.cpp
--- a/Engine/Source/Runtime/Engine/Private/Animation/AnimTrace.cpp
+++ b/Engine/Source/Runtime/Engine/Private/Animation/AnimTrace.cpp
@@ -21,7 +21,7 @@
 
 UE_TRACE_CHANNEL_DEFINE(AnimationChannel);
 
-UE_TRACE_EVENT_BEGIN(Animation, TickRecord2)
+UE_TRACE_EVENT_BEGIN(Animation, TickRecord3)
 	UE_TRACE_EVENT_FIELD(uint64, Cycle)
 	UE_TRACE_EVENT_FIELD(double, RecordingTime)
 	UE_TRACE_EVENT_FIELD(uint64, AnimInstanceId)
@@ -38,6 +38,14 @@ UE_TRACE_EVENT_BEGIN(Animation, TickRecord2)
 	UE_TRACE_EVENT_FIELD(uint16, FrameCounter)
 	UE_TRACE_EVENT_FIELD(bool, Looping)
 	UE_TRACE_EVENT_FIELD(bool, IsBlendSpace)
+	UE_TRACE_EVENT_FIELD(uint32, SyncGroupNameId)
+	UE_TRACE_EVENT_FIELD(bool, IsSyncGroupLeader)
+	UE_TRACE_EVENT_FIELD(float, SyncLeaderScore)
+	UE_TRACE_EVENT_FIELD(float, TimeToPrevMarker)
+	UE_TRACE_EVENT_FIELD(float, TimeToNextMarker)
+	UE_TRACE_EVENT_FIELD(int32, PrevMarkerIndex)
+	UE_TRACE_EVENT_FIELD(int32, NextMarkerIndex)
+
 UE_TRACE_EVENT_END()
 
 UE_TRACE_EVENT_BEGIN(Animation, SkeletalMesh3, NoSync|Important)
@@ -492,23 +500,53 @@ void FAnimTrace::OutputAnimTickRecord(const FAnimationBaseContext& InContext, co
 			BlendSpaceFilteredPositionY = InTickRecord.BlendSpace.BlendFilter->GetFilterLastOutput().Y;
 		}
 
-		UE_TRACE_LOG(Animation, TickRecord2, AnimationChannel)
-			<< TickRecord2.Cycle(FPlatformTime::Cycles64())
-			<< TickRecord2.RecordingTime(FObjectTrace::GetWorldElapsedTime(AnimInstance->GetWorld()))
-			<< TickRecord2.AnimInstanceId(FObjectTrace::GetObjectId(AnimInstance))
-			<< TickRecord2.AssetId(FObjectTrace::GetObjectId(InTickRecord.SourceAsset))
-			<< TickRecord2.NodeId(InContext.GetCurrentNodeId())
-			<< TickRecord2.BlendWeight(InTickRecord.EffectiveBlendWeight)
-			<< TickRecord2.PlaybackTime(PlaybackTime)
-			<< TickRecord2.RootMotionWeight(InTickRecord.RootMotionWeightModifier)
-			<< TickRecord2.PlayRate(InTickRecord.PlayRateMultiplier)
-			<< TickRecord2.BlendSpacePositionX(BlendSpacePositionX)
-			<< TickRecord2.BlendSpacePositionY(BlendSpacePositionY)
-			<< TickRecord2.BlendSpaceFilteredPositionX(BlendSpaceFilteredPositionX)
-			<< TickRecord2.BlendSpaceFilteredPositionY(BlendSpaceFilteredPositionY)
-			<< TickRecord2.FrameCounter(FObjectTrace::GetObjectWorldTickCounter(AnimInstance))
-			<< TickRecord2.Looping(InTickRecord.bLooping)
-			<< TickRecord2.IsBlendSpace(bIsBlendSpace);
+		// Figure out which sync group is this tick record located, if any at all.
+		const TMap<FName, FAnimGroupInstance>& GroupedTickRecords = InContext.AnimInstanceProxy->GetSyncGroupMapRead();
+		
+		static FName NoSyncGroupName(EName::None);
+		FName SyncGroupName = NoSyncGroupName;
+		bool IsSyncGroupLeader = false;
+		
+		for (const TTuple<FName, FAnimGroupInstance>& TickRecordGroup : GroupedTickRecords)
+		{
+			int32 GroupIdx = 0;
+			for (const FAnimTickRecord& TickRecord : TickRecordGroup.Value.ActivePlayers)
+			{
+				if (&TickRecord == &InTickRecord)
+				{
+					SyncGroupName = TickRecordGroup.Key;
+					IsSyncGroupLeader = GroupIdx == TickRecordGroup.Value.GroupLeaderIndex;
+				}
+				++GroupIdx;
+			}
+		}
+
+		uint32 SyncGroupNameId = FAnimTrace::OutputName(SyncGroupName);
+		
+		UE_TRACE_LOG(Animation, TickRecord3, AnimationChannel)
+			<< TickRecord3.Cycle(FPlatformTime::Cycles64())
+			<< TickRecord3.RecordingTime(FObjectTrace::GetWorldElapsedTime(AnimInstance->GetWorld()))
+			<< TickRecord3.AnimInstanceId(FObjectTrace::GetObjectId(AnimInstance))
+			<< TickRecord3.AssetId(FObjectTrace::GetObjectId(InTickRecord.SourceAsset))
+			<< TickRecord3.NodeId(InContext.GetCurrentNodeId())
+			<< TickRecord3.BlendWeight(InTickRecord.EffectiveBlendWeight)
+			<< TickRecord3.PlaybackTime(PlaybackTime)
+			<< TickRecord3.RootMotionWeight(InTickRecord.RootMotionWeightModifier)
+			<< TickRecord3.PlayRate(InTickRecord.PlayRateMultiplier)
+			<< TickRecord3.BlendSpacePositionX(BlendSpacePositionX)
+			<< TickRecord3.BlendSpacePositionY(BlendSpacePositionY)
+			<< TickRecord3.BlendSpaceFilteredPositionX(BlendSpaceFilteredPositionX)
+			<< TickRecord3.BlendSpaceFilteredPositionY(BlendSpaceFilteredPositionY)
+			<< TickRecord3.FrameCounter(FObjectTrace::GetObjectWorldTickCounter(AnimInstance))
+			<< TickRecord3.Looping(InTickRecord.bLooping)
+			<< TickRecord3.IsBlendSpace(bIsBlendSpace)
+			<< TickRecord3.SyncGroupNameId(SyncGroupNameId)
+			<< TickRecord3.IsSyncGroupLeader(IsSyncGroupLeader)
+			<< TickRecord3.SyncLeaderScore(InTickRecord.LeaderScore)
+			<< TickRecord3.PrevMarkerIndex(InTickRecord.MarkerTickRecord ? InTickRecord.MarkerTickRecord->PreviousMarker.MarkerIndex : -1)
+			<< TickRecord3.NextMarkerIndex(InTickRecord.MarkerTickRecord ? InTickRecord.MarkerTickRecord->NextMarker.MarkerIndex : -1)
+			<< TickRecord3.TimeToPrevMarker(InTickRecord.MarkerTickRecord ? InTickRecord.MarkerTickRecord->PreviousMarker.TimeToMarker : 0)
+			<< TickRecord3.TimeToNextMarker(InTickRecord.MarkerTickRecord ? InTickRecord.MarkerTickRecord->NextMarker.TimeToMarker : 0);
 	}
 }
 
