[ObjectCacheContext/ALT]
- Prevent queuing changes coming from loading thread since events are processed on the game thread and could race with ALT objects still being serialized

#Synced-CL 50044815

diff --git a/Engine/Source/Runtime/Engine/Private/ObjectCacheContext.cpp b/Engine/Source/Runtime/Engine/Private/ObjectCacheContext.cpp
--- a/Engine/Source/Runtime/Engine/Private/ObjectCacheContext.cpp
+++ b/Engine/Source/Runtime/Engine/Private/ObjectCacheContext.cpp
@@ -1293,6 +1293,12 @@ struct FObjectCacheEventSinkPrivate
 
 	static std::atomic<bool> bShouldQueueSinkEvents;
 
+	static bool ShouldQueueSinkEvents()
+	{
+		// Loading thread events can't be queued as it could cause race condition with async loading when the queue is processed
+		return bShouldQueueSinkEvents.load(std::memory_order_relaxed) && !IsInAsyncLoadingThread();
+	}
+
 	using TEventAllocator = TLockFreeClassAllocator<FNotifyEvent, PLATFORM_CACHE_LINE_SIZE>;
 	using TEventList = TLockFreePointerListFIFO<FNotifyEvent, PLATFORM_CACHE_LINE_SIZE>;
 	static TEventAllocator& GetAllocator();
@@ -1329,7 +1335,7 @@ void FObjectCacheEventSink::NotifyCompositeTextureChanged_Concurrent(UTexture* I
 
 	if (GetObjectReverseLookupMode() != EObjectReverseLookupMode::Temporary)
 	{
-		if (FObjectCacheEventSinkPrivate::bShouldQueueSinkEvents)
+		if (FObjectCacheEventSinkPrivate::ShouldQueueSinkEvents())
 		{
 			FObjectCacheEventSinkPrivate::AddNotifyEvent(FObjectCacheEventSinkPrivate::ECompositeTextureChanged, nullptr, nullptr, nullptr, nullptr, nullptr, InTexture);
 		}
@@ -1345,7 +1351,8 @@ void FObjectCacheEventSink::NotifyMaterialDestroyed_Concurrent(UMaterialInterfac
 
 	if (GetObjectReverseLookupMode() != EObjectReverseLookupMode::Temporary)
 	{
-		if (FObjectCacheEventSinkPrivate::bShouldQueueSinkEvents)
+
+		if (FObjectCacheEventSinkPrivate::ShouldQueueSinkEvents())
 		{
 			FObjectCacheEventSinkPrivate::AddNotifyEvent(FObjectCacheEventSinkPrivate::EMaterialDestroyed, MaterialInterface, nullptr, nullptr, nullptr, nullptr, nullptr);
 		}
@@ -1362,7 +1369,7 @@ void FObjectCacheEventSink::NotifyUsedMaterialsChanged_Concurrent(const IPrimiti
 
 	if (GetObjectReverseLookupMode() != EObjectReverseLookupMode::Temporary)
 	{
-		if (FObjectCacheEventSinkPrivate::bShouldQueueSinkEvents)
+		if (FObjectCacheEventSinkPrivate::ShouldQueueSinkEvents())
 		{
 			FObjectCacheEventSinkPrivate::AddNotifyEvent(FObjectCacheEventSinkPrivate::EMaterialsChanged, nullptr, PrimitiveComponent, nullptr, nullptr, &UsedMaterials, nullptr);
 		}
@@ -1381,7 +1388,7 @@ void FObjectCacheEventSink::NotifyRenderStateChanged_Concurrent(UActorComponent*
 	{
 		if (UPrimitiveComponent* Primitive = Cast<UPrimitiveComponent>(ActorComponent))
 		{
-			if (FObjectCacheEventSinkPrivate::bShouldQueueSinkEvents)
+			if (FObjectCacheEventSinkPrivate::ShouldQueueSinkEvents())
 			{
 				FObjectCacheEventSinkPrivate::AddNotifyEvent(FObjectCacheEventSinkPrivate::ERenderStateChanged, nullptr, Primitive->GetPrimitiveComponentInterface(), nullptr, nullptr, nullptr, nullptr);
 			}
@@ -1399,7 +1406,7 @@ void FObjectCacheEventSink::NotifyRenderStateChanged_Concurrent(IPrimitiveCompon
 
 	if (GetObjectReverseLookupMode() != EObjectReverseLookupMode::Temporary)
 	{
-		if (FObjectCacheEventSinkPrivate::bShouldQueueSinkEvents)
+		if (FObjectCacheEventSinkPrivate::ShouldQueueSinkEvents())
 		{
 			FObjectCacheEventSinkPrivate::AddNotifyEvent(FObjectCacheEventSinkPrivate::ERenderStateChanged, nullptr, PrimitiveComponent, nullptr, nullptr, nullptr, nullptr);
 		}
@@ -1417,7 +1424,7 @@ void FObjectCacheEventSink::NotifyMaterialChanged_Concurrent(UMaterialInterface*
 
 	if (GetObjectReverseLookupMode() != EObjectReverseLookupMode::Temporary)
 	{
-		if (FObjectCacheEventSinkPrivate::bShouldQueueSinkEvents)
+		if (FObjectCacheEventSinkPrivate::ShouldQueueSinkEvents())
 		{
 			FObjectCacheEventSinkPrivate::AddNotifyEvent(FObjectCacheEventSinkPrivate::EMaterialChanged, MaterialInterface, nullptr, nullptr, nullptr, nullptr, nullptr);
 		}
@@ -1434,7 +1441,7 @@ void FObjectCacheEventSink::NotifyStaticMeshChanged_Concurrent(IStaticMeshCompon
 
 	if (GetObjectReverseLookupMode() != EObjectReverseLookupMode::Temporary)
 	{
-		if (FObjectCacheEventSinkPrivate::bShouldQueueSinkEvents)
+		if (FObjectCacheEventSinkPrivate::ShouldQueueSinkEvents())
 		{
 			FObjectCacheEventSinkPrivate::AddNotifyEvent(FObjectCacheEventSinkPrivate::EStaticMeshChanged, nullptr, nullptr, StaticMeshComponent, nullptr, nullptr, nullptr);
 		}
@@ -1451,7 +1458,7 @@ void FObjectCacheEventSink::NotifySkinnedMeshChanged_Concurrent(ISkinnedMeshComp
 
 	if (GetObjectReverseLookupMode() != EObjectReverseLookupMode::Temporary)
 	{
-		if (FObjectCacheEventSinkPrivate::bShouldQueueSinkEvents)
+		if (FObjectCacheEventSinkPrivate::ShouldQueueSinkEvents())
 		{
 			FObjectCacheEventSinkPrivate::AddNotifyEvent(FObjectCacheEventSinkPrivate::ESkinnedAssetChanged, nullptr, nullptr, nullptr, SkinnedMeshComponent, nullptr, nullptr);
 		}
