[Texture2D/ALT]
- Fix PostLinkerChange that ended up waiting during texture loads

#Synced-CL 50044815

diff --git a/Engine/Source/Runtime/Engine/Private/Texture2D.cpp b/Engine/Source/Runtime/Engine/Private/Texture2D.cpp
--- a/Engine/Source/Runtime/Engine/Private/Texture2D.cpp
+++ b/Engine/Source/Runtime/Engine/Private/Texture2D.cpp
@@ -792,7 +792,7 @@ void UTexture2D::UpdateResourceWithParams(EUpdateResourceFlags InFlags)
 void UTexture2D::PostLinkerChange()
 {
 	// Changing the linker requires re-creating the resource to make sure streaming behavior is right.
-	if( !HasAnyFlags( RF_BeginDestroyed | RF_NeedLoad | RF_NeedPostLoad ) && !IsUnreachable() )
+	if( !HasAnyFlags( RF_BeginDestroyed | RF_NeedLoad | RF_NeedPostLoad ) && !IsUnreachable() && !HasAnyInternalFlags(EInternalObjectFlags_AsyncLoading) )
 	{
 		// Update the resource. We have to complete any builds beforehand.
 		BlockOnAnyAsyncBuild();
