Fix cooked data not having the Components/Entities array populated when using the new field attribute.
Added a new property flag ForcePostConstructLink to fix this issue.

#rnx

#Synced-CL 50101441

diff --git a/Engine/Source/Programs/Shared/EpicGames.Core/UnrealEngineTypes.cs b/Engine/Source/Programs/Shared/EpicGames.Core/UnrealEngineTypes.cs
--- a/Engine/Source/Programs/Shared/EpicGames.Core/UnrealEngineTypes.cs
+++ b/Engine/Source/Programs/Shared/EpicGames.Core/UnrealEngineTypes.cs
@@ -1252,6 +1252,11 @@ namespace EpicGames.Core
 		/// </summary>
 		AllowSelfReference = 0x1000000000000000,
 
+		/// <summary>
+		/// Use to force native UCLASS properties to be included in the PostConstructLink chain in a derived Blueprint class, regardless of whether it differs from the native class constructor-initialized value
+		/// </summary>
+		ForcePostConstructLink = 0x2000000000000000,
+
 		/// <summary>
 		/// All Native Access Specifier flags
 		/// </summary>
diff --git a/Engine/Source/Programs/Shared/EpicGames.UHT/Specifiers/UhtPropertyMemberSpecifiers.cs b/Engine/Source/Programs/Shared/EpicGames.UHT/Specifiers/UhtPropertyMemberSpecifiers.cs
--- a/Engine/Source/Programs/Shared/EpicGames.UHT/Specifiers/UhtPropertyMemberSpecifiers.cs
+++ b/Engine/Source/Programs/Shared/EpicGames.UHT/Specifiers/UhtPropertyMemberSpecifiers.cs
@@ -414,6 +414,13 @@ namespace EpicGames.UHT.Parsers
 			context.PropertySettings.PropertyFlags |= EPropertyFlags.ExperimentalExternalObjects;
 		}
 
+		[UhtSpecifier(Extends = UhtTableNames.PropertyMember, ValueType = UhtSpecifierValueType.None)]
+		private static void ForcePostConstructLinkSpecifier(UhtSpecifierContext specifierContext)
+		{
+			UhtPropertySpecifierContext context = (UhtPropertySpecifierContext)specifierContext;
+			context.PropertySettings.PropertyFlags |= EPropertyFlags.ForcePostConstructLink;
+		}
+
 		[UhtSpecifier(Extends = UhtTableNames.PropertyMember, ValueType = UhtSpecifierValueType.OptionalString)]
 		private static void GetterSpecifier(UhtSpecifierContext specifierContext, StringView? value)
 		{
diff --git a/Engine/Source/Runtime/CoreUObject/Private/UObject/Class.cpp b/Engine/Source/Runtime/CoreUObject/Private/UObject/Class.cpp
--- a/Engine/Source/Runtime/CoreUObject/Private/UObject/Class.cpp
+++ b/Engine/Source/Runtime/CoreUObject/Private/UObject/Class.cpp
@@ -5501,7 +5501,7 @@ void UClass::Link(FArchive& Ar, bool bRelinkExistingProperties)
 EStructPropertyLinkFlags UClass::GetPropertyLinkFlags(UStruct* ContainerStruct, FProperty* Property) const
 {
 	// If we aren't a native type or we have a VM script, then we will need destructor and post construct
-	if (!ContainerStruct->Script.IsEmpty() || !HasAnyClassFlags(CLASS_Native | CLASS_Intrinsic))
+	if (!ContainerStruct->Script.IsEmpty() || !HasAnyClassFlags(CLASS_Native | CLASS_Intrinsic) || Property->HasAnyPropertyFlags(CPF_ForcePostConstructLink))
 	{
 		return EStructPropertyLinkFlags::LinkDestructor | EStructPropertyLinkFlags::LinkPostConstruct;
 	}
diff --git a/Engine/Source/Runtime/CoreUObject/Public/UObject/ObjectMacros.h b/Engine/Source/Runtime/CoreUObject/Public/UObject/ObjectMacros.h
--- a/Engine/Source/Runtime/CoreUObject/Public/UObject/ObjectMacros.h
+++ b/Engine/Source/Runtime/CoreUObject/Public/UObject/ObjectMacros.h
@@ -491,6 +491,7 @@ enum EPropertyFlags : uint64
 	CPF_ExperimentalAlwaysOverriden		= 0x0400000000000000,	///< ****Experimental*** Property should never inherit from the parent when using overridable serialization
 	CPF_ExperimentalNeverOverriden		= 0x0800000000000000,	///< ****Experimental*** Property should never be overridden when using overridable serialization
 	CPF_AllowSelfReference				= 0x1000000000000000,	///< Enables the instancing graph self referencing logic, delegates and verse function are already using this
+	CPF_ForcePostConstructLink			= 0x2000000000000000,	///< Use to force native UCLASS properties to be included in the PostConstructLink chain in a derived Blueprint class, regardless of whether it differs from the native class constructor-initialized value
 
 	// !!!!!!!!!!!!!!!!!!!!!!!!!!!!!!!!!!!!!!!!!!!!!!!!!!!!!!!!!!!!!!!!!!!!!!!!!!!!!!!!!!!!!!!!!!!!!!!!!!!!!!!!!!!!!!!!!!!!!!!!!!!!!
 	// DO NOT ADD composite flags to EPropertyFlags. Debugger visualisations have trouble displaying composite flag values as text.
diff --git a/Engine/Source/Runtime/Engine/Private/BlueprintGeneratedClass.cpp b/Engine/Source/Runtime/Engine/Private/BlueprintGeneratedClass.cpp
--- a/Engine/Source/Runtime/Engine/Private/BlueprintGeneratedClass.cpp
+++ b/Engine/Source/Runtime/Engine/Private/BlueprintGeneratedClass.cpp
@@ -203,9 +203,10 @@ namespace UE::Runtime::Engine::Private
 			const UClass* OwnerClass = Property.GetOwnerClass();
 			const bool bIsConfigProperty = Property.HasAnyPropertyFlags(CPF_Config) && !(OwnerClass && OwnerClass->HasAnyClassFlags(CLASS_PerObjectConfig));
 			const bool bIsTransientProperty = Property.HasAnyPropertyFlags(CPF_Transient | CPF_DuplicateTransient | CPF_NonPIEDuplicateTransient);
+			const bool bIsForcePostConstructLink = Property.HasAnyPropertyFlags(CPF_ForcePostConstructLink);
 
-			// Skip config properties as they're already in the PostConstructLink chain. Also skip transient properties if they contain a reference to an instanced subobjects (as those should not be initialized from defaults).
-			if (!bIsConfigProperty && (!bIsTransientProperty || !Property.ContainsInstancedObjectProperty()))
+			// Skip config and forced post construct link properties as they're already in the PostConstructLink chain. Also skip transient properties if they contain a reference to an instanced subobjects (as those should not be initialized from defaults).
+			if (!bIsConfigProperty && !bIsForcePostConstructLink &&(!bIsTransientProperty || !Property.ContainsInstancedObjectProperty()))
 			{
 				return true;
 			}
