[WorldPartition] Fix assertion in SetRuntimeDataLayerInstanceNames when setting combined container data on actors attached to parents with bIsForcedEmptyNonExternalDataLayers

#Synced-CL 50059267

diff --git a/Engine/Source/Runtime/Engine/Private/WorldPartition/WorldPartitionStreamingGeneration.cpp b/Engine/Source/Runtime/Engine/Private/WorldPartition/WorldPartitionStreamingGeneration.cpp
--- a/Engine/Source/Runtime/Engine/Private/WorldPartition/WorldPartitionStreamingGeneration.cpp
+++ b/Engine/Source/Runtime/Engine/Private/WorldPartition/WorldPartitionStreamingGeneration.cpp
@@ -505,7 +505,7 @@ void FStreamingGenerationActorDescView::SetForcedNoDataLayers()
 void FStreamingGenerationActorDescView::SetRuntimeDataLayerInstanceNames(const FDataLayerInstanceNames& InRuntimeDataLayerInstanceNames)
 {
 	check(!InRuntimeDataLayerInstanceNames.IsForcedEmptyNonExternalDataLayers() ||
-		(RuntimeDataLayerInstanceNames.IsSet() && (RuntimeDataLayerInstanceNames->IsForcedEmptyNonExternalDataLayers() == InRuntimeDataLayerInstanceNames.IsForcedEmptyNonExternalDataLayers())));
+		(GetRuntimeDataLayerInstanceNames().IsForcedEmptyNonExternalDataLayers() == InRuntimeDataLayerInstanceNames.IsForcedEmptyNonExternalDataLayers()));
 	RuntimeDataLayerInstanceNames = InRuntimeDataLayerInstanceNames;
 	if (bIsForcedNoDataLayers)
 	{
