New Material Translator:
- Added support for user-specified vertex interpolated data and the UMaterialExpressionVertexInterpolator expression.
- Improvements to UV handling in the new translator.

#rnx

#Synced-CL 50155647

diff --git a/Engine/Shaders/Private/MaterialTemplate.ush b/Engine/Shaders/Private/MaterialTemplate.ush
--- a/Engine/Shaders/Private/MaterialTemplate.ush
+++ b/Engine/Shaders/Private/MaterialTemplate.ush
@@ -649,7 +649,7 @@ struct FMaterialPixelParameters
 
 	FMaterialLWCData LWCData;
 
-%{material_pixel_parameters_extras}
+%{material_pixel_parameter_decls}
 };
 
 %{material_cache_types}
@@ -721,18 +721,19 @@ struct FMaterialVertexParameters
 	float3 PreSkinnedNormal;
 
 	half4 VertexColor;
-#if NUM_MATERIAL_TEXCOORDS_VERTEX
-	float2 TexCoords[NUM_MATERIAL_TEXCOORDS_VERTEX];
+
+	// NUM_MATERIAL_TEXCOORDS_VERTEX means the number of actual UV TexCoords used by the shader
+	// NUM_TEX_COORD_INTERPOLATORS instead is the total interpolator registers to use: this is the value above plus all the user interpolators expressions in the material
+	// Using NUM_TEX_COORD_INTERPOLATORS here so that we have a single array ot hold all interpolated data (UVs plus user).
+	// Note: FMaterialVertex/PixelParameters structs are "local" to the material
+	// Note: The code that maps to these registers is generated by ComputeUserVertexInterpolatorAllocation() in the new translator
+#if NUM_TEX_COORD_INTERPOLATORS
+	float2 TexCoords[NUM_TEX_COORD_INTERPOLATORS];
 	#if FEATURE_LEVEL == FEATURE_LEVEL_ES3_1
-	float2 TexCoordOffset; // Offset for UV localization for large UV values
+		float2 TexCoordOffset; // Offset for UV localization for large UV values
 	#endif
 #endif
 
-#if NUM_CUSTOM_VERTEX_INTERPOLATORS
-	// Used by new HLSL translator
-	float2 CustomInterpolators[NUM_CUSTOM_VERTEX_INTERPOLATORS];
-#endif
-
 	/** Per-particle properties. Only valid for particle vertex factories. */
 	FMaterialParticleParameters Particle;
 
@@ -755,7 +756,7 @@ struct FMaterialVertexParameters
 
 	MaterialFloat3 WorldPositionOffset;
 	
-%{material_vertex_parameters_extras}
+%{material_vertex_parameter_decls}
 };
 
 FMaterialVertexParameters MakeInitializedMaterialVertexParameters()
@@ -768,22 +769,6 @@ FMaterialVertexParameters MakeInitializedMaterialVertexParameters()
 	return Result;
 }
 
-float MaterialReadInterpolatorComponent(FMaterialPixelParameters Parameters, int InterpolatorIndex)
-{
-#if NUM_TEX_COORD_INTERPOLATORS
-	return Parameters.TexCoords[NUM_MATERIAL_TEXCOORDS + InterpolatorIndex / 2][InterpolatorIndex & 1];
-#else
-	return 0.0f;
-#endif
-}
-
-void MaterialPackInterpolatorComponent(in out FMaterialVertexParameters Parameters, int InterpolatorIndex, float Value)
-{
-#if NUM_CUSTOM_VERTEX_INTERPOLATORS
-	Parameters.CustomInterpolators[InterpolatorIndex / 2][InterpolatorIndex & 1] = Value;
-#endif
-}
-
 #if GET_PRIMITIVE_DATA_OVERRIDE
 	FPrimitiveSceneData GetPrimitiveData(FMaterialVertexParameters Parameters);
 	FPrimitiveSceneData GetPrimitiveData(FMaterialPixelParameters Parameters);
@@ -3609,6 +3594,8 @@ SHADER_DISABLE_WARNINGS
 
 %{uniform_material_expressions}
 
+%{custom_outputs}
+
 %{custom_functions}
 
 SHADER_POP_WARNINGS_STATE
@@ -4173,6 +4160,7 @@ FMaterialRefractionData GetMaterialRefraction(FPixelMaterialInputs PixelMaterial
 
 SHADER_PUSH_WARNINGS_STATE
 SHADER_DISABLE_WARNINGS
+
 void GetMaterialCustomizedUVs(FMaterialVertexParameters Parameters, inout float2 OutTexCoords[NUM_TEX_COORD_INTERPOLATORS])
 {
 	// Always initialise to avoid later operations to potentially cause compiler failure due to read uninitialized value when get_material_customized_u_vs is not defined.
@@ -4189,6 +4177,7 @@ void GetCustomInterpolators(FMaterialVertexParameters Parameters, inout float2 O
 {
 %{get_custom_interpolators}
 }
+
 SHADER_POP_WARNINGS_STATE
 
 #endif
diff --git a/Engine/Source/Runtime/Engine/Private/Materials/MaterialExpressionsIR.cpp b/Engine/Source/Runtime/Engine/Private/Materials/MaterialExpressionsIR.cpp
--- a/Engine/Source/Runtime/Engine/Private/Materials/MaterialExpressionsIR.cpp
+++ b/Engine/Source/Runtime/Engine/Private/Materials/MaterialExpressionsIR.cpp
@@ -4703,7 +4703,7 @@ void UMaterialExpressionSpeedTree::Build(MIR::FEmitter& Em)
 		[](FMaterialIRModule* Module, uintptr_t)
 		{
 			Module->AddEnvironmentDefine(TEXT("USES_SPEEDTREE"));
-			Module->GetStatistics().NumVertexTexCoords = FMath::Max(Module->GetStatistics().NumVertexTexCoords, 8);
+			Module->GetStatistics().NumTexCoords = FMath::Max(Module->GetStatistics().NumTexCoords, 8);
 		});
 
 	TStaticArray<FValueRef, MIR::EStage::NumStages> ValuePerStage = {};
@@ -4822,6 +4822,11 @@ void UMaterialExpressionNeuralNetworkOutput::Build(MIR::FEmitter& Em)
 	Em.Output(1, Em.InlineHLSL(MIR::FType::MakeFloatVector(4), TEXTVIEW("NeuralBufferOutput(Parameters,$0)"), { BufferIndex }, {}, {}, EnableWithNeuralNetworks));
 }
 
+void UMaterialExpressionVertexInterpolator::Build(MIR::FEmitter& Em)
+{
+	Em.Output(0, Em.VertexInterpolator(Em.Input(&Input)));
+}
+
 #endif // WITH_EDITOR
 
 
diff --git a/Engine/Source/Runtime/Engine/Private/Materials/MaterialIR.cpp b/Engine/Source/Runtime/Engine/Private/Materials/MaterialIR.cpp
--- a/Engine/Source/Runtime/Engine/Private/Materials/MaterialIR.cpp
+++ b/Engine/Source/Runtime/Engine/Private/Materials/MaterialIR.cpp
@@ -39,6 +39,8 @@ const TCHAR* LexToString(EValueKind Kind)
 		case VK_TextureObject: return TEXT("TextureObject");
 		case VK_RuntimeVirtualTextureObject: return TEXT("RuntimeVirtualTextureObject");
 		case VK_UniformParameter: return TEXT("UniformParameter");
+		case VK_GetVertexInterpolator: return TEXT("GetVertexInterpolator");
+		case VK_Nop: return TEXT("Nop");
 		case VK_Composite: return TEXT("Composite");
 		case VK_SetMaterialOutput: return TEXT("SetMaterialOutput");
 		case VK_Operator: return TEXT("Operator");
@@ -59,7 +61,7 @@ const TCHAR* LexToString(EValueKind Kind)
 		case VK_SubstratePromoteToOperator: return TEXT("SubstratePromoteToOperator");
 		case VK_StageSwitch: return TEXT("StageSwitch");
 		case VK_HardwarePartialDerivative: return TEXT("HardwarePartialDerivative");
-		case VK_Nop: return TEXT("Nop");
+		case VK_SetVertexInterpolator: return TEXT("SetVertexInterpolator");
 		case VK_Call: return TEXT("Call");
 		case VK_CallParameterOutput: return TEXT("CallParameterOutput");
 		case VK_PreshaderParameter: return TEXT("PreshaderParameter");
@@ -115,6 +117,8 @@ uint32 FValue::GetSizeInBytes() const
 		case VK_TextureObject: return sizeof(FTextureObject);
 		case VK_RuntimeVirtualTextureObject: return sizeof(FRuntimeVirtualTextureObject);
 		case VK_UniformParameter: return sizeof(FUniformParameter);
+		case VK_GetVertexInterpolator: return sizeof(FGetVertexInterpolator);
+		case VK_Nop: return sizeof(FNop);
 		case VK_Composite: return sizeof(FComposite) + sizeof(FValue*) * static_cast<const FComposite*>(this)->GetComponents().Num();
 		case VK_SetMaterialOutput: return sizeof(FSetMaterialOutput);
 		case VK_Operator: return sizeof(FOperator);
@@ -135,7 +139,7 @@ uint32 FValue::GetSizeInBytes() const
 		case VK_SubstratePromoteToOperator: return sizeof(FSubstratePromoteToOperator);
 		case VK_StageSwitch: return sizeof(FStageSwitch);
 		case VK_HardwarePartialDerivative: return sizeof(FHardwarePartialDerivative);
-		case VK_Nop: return sizeof(FNop);
+		case VK_SetVertexInterpolator: return sizeof(FSetVertexInterpolator);
 		case VK_Call: return sizeof(FCall);
 		case VK_CallParameterOutput: return sizeof(FCallParameterOutput);
 		case VK_PreshaderParameter: return sizeof(FPreshaderParameter);
@@ -157,6 +161,12 @@ TConstArrayView<FValue*> FValue::GetUses() const
 
 	switch (Kind)
 	{
+		case VK_Nop:
+		{
+			auto This = static_cast<const FNop*>(this);
+			return { &This->Arg, FNop::NumStaticUses };
+		}
+
 		case VK_Composite:
 		{
 			auto This = static_cast<const FComposite*>(this);
@@ -277,10 +287,10 @@ TConstArrayView<FValue*> FValue::GetUses() const
 			return { &This->Arg, FHardwarePartialDerivative::NumStaticUses };
 		}
 
-		case VK_Nop:
+		case VK_SetVertexInterpolator:
 		{
-			auto This = static_cast<const FNop*>(this);
-			return { &This->Arg, FNop::NumStaticUses };
+			auto This = static_cast<const FSetVertexInterpolator*>(this);
+			return { &This->Arg, FSetVertexInterpolator::NumStaticUses };
 		}
 
 		case VK_Call:
@@ -837,6 +847,25 @@ FBlock* FInstruction::GetTargetBlockForUse(int32 EntryPointIndex, int32 UseIndex
 	return Linkage[EntryPointIndex].Block;
 }
 
+FStringView FInstruction::GetName() const
+{
+	if (const FSetMaterialOutput* SetMaterialOutput = MIR::As<FSetMaterialOutput>(this))
+	{
+		return SetMaterialOutput->Name;
+	}
+	return {};
+}
+
+int32 FInstruction::GetOutputPriority() const
+{
+	if (const FSetMaterialOutput* SetMaterialOutput = MIR::As<FSetMaterialOutput>(this))
+	{
+		return SetMaterialOutput->Priority;
+	}
+	// Otherwise use a large priority based on this instruction kind
+	return 1024 + Kind;
+}
+
 bool IsComparisonOperator(EOperator Op)
 {
 	switch (Op)
diff --git a/Engine/Source/Runtime/Engine/Private/Materials/MaterialIRDebug.cpp b/Engine/Source/Runtime/Engine/Private/Materials/MaterialIRDebug.cpp
--- a/Engine/Source/Runtime/Engine/Private/Materials/MaterialIRDebug.cpp
+++ b/Engine/Source/Runtime/Engine/Private/Materials/MaterialIRDebug.cpp
@@ -4,7 +4,6 @@
 #include "Materials/MaterialIRModule.h"
 #include "Materials/MaterialIR.h"
 #include "Materials/MaterialIRTypes.h"
-#include "Materials/MaterialIRInternal.h"
 #include "Materials/MaterialAttributeDefinitionMap.h"
 #include "Materials/MaterialExternalCodeRegistry.h"
 #include "Misc/Paths.h"
@@ -433,7 +432,7 @@ struct FDebugDumpIRState
 	uint32 ReferenceInstruction(const FInstruction* Instr)
 	{
 		uint32 Id;
-		if (!Internal::Find(ValueToIdMap, static_cast<const FValue*>(Instr), Id))
+		if (!MIR::Find(ValueToIdMap, static_cast<const FValue*>(Instr), Id))
 		{
 			Id = InstrIdCounter++;
 			ValueToIdMap.Add(Instr, Id);
diff --git a/Engine/Source/Runtime/Engine/Private/Materials/MaterialIREmitter.cpp b/Engine/Source/Runtime/Engine/Private/Materials/MaterialIREmitter.cpp
--- a/Engine/Source/Runtime/Engine/Private/Materials/MaterialIREmitter.cpp
+++ b/Engine/Source/Runtime/Engine/Private/Materials/MaterialIREmitter.cpp
@@ -1211,7 +1211,7 @@ FValueRef FEmitter::Parameter(FName Name, FMaterialParameterMetadata& Metadata,
 	auto RegisterParameter = [] (FMaterialIRModule* InModule, FMaterialParameterInfo InInfo, const FMaterialParameterMetadata& InMetadata) -> uint32
 	{
 		uint32 Id;
-		if (Internal::Find(InModule->ParameterInfoToId, InInfo, Id))
+		if (MIR::Find(InModule->ParameterInfoToId, InInfo, Id))
 		{
 			check(InModule->ParameterIdToData[Id].Get<FMaterialParameterMetadata>().Value == InMetadata.Value);
 			return Id;
@@ -4110,6 +4110,12 @@ FValueRef FEmitter::AnalyticalPartialDerivative(FValueRef Value, EDerivativeAxis
 			return Cast(ConstantZero(Value->Type.GetPrimitive().ScalarKind), Value->Type);
 		}
 
+		case VK_GetVertexInterpolator:
+		{
+			FGetVertexInterpolator* AsVertexInterpolator = Value->As<FGetVertexInterpolator>();
+			return VertexInterpolator(AnalyticalPartialDerivative(AsVertexInterpolator->SetInstr->Arg, Axis));
+		}
+
 		// These values are uniform (constant), thus their value is always zero.
 		case VK_Constant:
 		case VK_UniformParameter:
@@ -4140,6 +4146,26 @@ FValueRef FEmitter::AnalyticalPartialDerivative(FValueRef Value, EDerivativeAxis
 	}
 }
 
+FValueRef FEmitter::VertexInterpolator(FValueRef VertexValue)
+{
+	// Cast the incoming vertex value to a float scalar or vector
+	VertexValue = CastToFloatKind(CheckIsScalarOrVector(VertexValue));
+	if (!VertexValue.IsValid())
+	{
+		return VertexValue;
+	}
+
+	// Create the set/get pair of vertex interpolator instructions
+	FSetVertexInterpolator* SetInstr = EmitNew<FSetVertexInterpolator>(*this, FType::MakeVoid());
+	SetInstr->Arg = VertexValue;
+
+	FGetVertexInterpolator* GetInstr = EmitNew<FGetVertexInterpolator>(*this, VertexValue->Type);
+	GetInstr->SetSubgraphProperties(EGraphProperties::PixelStageOnly);
+	GetInstr->SetInstr = SetInstr;
+
+	return GetInstr;
+}
+
 static FValueRef EmitInlineHLSL(FEmitter& Emitter, FType Type, const FMaterialExternalCodeDeclaration* InExternalCodeDeclaration,
 								FStringView Code, TConstArrayView<FValueRef> InputValues, EValueFlags ValueFlags, EGraphProperties UsedGraphProperties,
 								FInlineHLSL::TOnAnalyzeCallback OnAnalyze = {}, uintptr_t OnAnalyzeUserData = {})
diff --git a/Engine/Source/Runtime/Engine/Private/Materials/MaterialIRInternal.h b/Engine/Source/Runtime/Engine/Private/Materials/MaterialIRInternal.h
--- a/Engine/Source/Runtime/Engine/Private/Materials/MaterialIRInternal.h
+++ b/Engine/Source/Runtime/Engine/Private/Materials/MaterialIRInternal.h
@@ -26,18 +26,6 @@ void BindValueToExpressionOutput(FMaterialIRModuleBuilderImpl* Builder, const FE
 // Computes the hash of a blob of memory. Note: Ptr must be 4 bytes aligned.
 uint32 HashBytes(const void* Ptr, uint32 Size);
 
-/* Other helper functions */
-
-template <typename TKey, typename TValue>
-bool Find(const TMap<TKey, TValue>& Map, const TKey& Key, TValue& OutValue)
-{
-	if (auto ValuePtr = Map.Find(Key)) {
-		OutValue = *ValuePtr;
-		return true;
-	}
-	return false;
-}
-
 } // namespace MIR::Internal
 
 #endif // #if WITH_EDITOR
diff --git a/Engine/Source/Runtime/Engine/Private/Materials/MaterialIRModuleBuilder.cpp b/Engine/Source/Runtime/Engine/Private/Materials/MaterialIRModuleBuilder.cpp
--- a/Engine/Source/Runtime/Engine/Private/Materials/MaterialIRModuleBuilder.cpp
+++ b/Engine/Source/Runtime/Engine/Private/Materials/MaterialIRModuleBuilder.cpp
@@ -20,6 +20,8 @@
 #include "Materials/MaterialExpressionNamedReroute.h"
 #include "Materials/MaterialExpressionSingleLayerWaterMaterialOutput.h"
 #include "Materials/MaterialExpressionThinTranslucentMaterialOutput.h"
+#include "Materials/MaterialExpressionCustomOutput.h"
+#include "Materials/MaterialExpressionVertexInterpolator.h"
 #include "Materials/MaterialExpressionSubstrate.h"
 #include "Materials/MaterialSharedPrivate.h"
 #include "Materials/Material.h"
@@ -181,11 +183,15 @@ struct FMaterialIRModuleBuilderImpl
 		}
 
 		// Push the custom output expressions to the dependency queue too.
-		TArray<UMaterialExpressionCustomOutput*> CustomOutputs;
-		Builder->MaterialInterface->GetMaterial()->GetAllCustomOutputExpressions(CustomOutputs);
-		for (UMaterialExpressionCustomOutput* Expression : CustomOutputs)
+		// Note: unfortunately the vertex interpolator expression derives from custom output (even if it isn't a custom output)
+		// 	     as an artifact of the old translator. We're only interested in actual custom outputs, as vertex interpolators are
+		//       handled differently. When the old translator is removed, make UMaterialExpressionVertexInterpolator a simple UMaterialExpression.
+		for (UMaterialExpression* Expression : Builder->MaterialInterface->GetMaterial()->GetExpressions())
 		{
-			PushDependency(Expression);
+			if (Cast<UMaterialExpressionCustomOutput>(Expression) && !Cast<UMaterialExpressionVertexInterpolator>(Expression))
+			{
+				PushDependency(Expression);
+			}
 		}
 	}
 
@@ -940,35 +946,35 @@ struct FMaterialIRModuleBuilderImpl
 		TSet<MIR::FValue*> VisitedValues{};
 
 		// Analyze the nodes in each entry point.
-		for (int32 EntryPointIndex = 0; EntryPointIndex < Module->GetNumEntryPoints(); ++EntryPointIndex)
+		// Note: process the vertex stage last so that vertex interpolator getters in the pixel stage are analyzed first,
+		//       adding the setters in the vertex stage entry point outputs. This is done so that only actually used vertex interpolators
+		//       are compiled in.
+		check(Module->GetEntryPoint(0).Stage == MIR::EStage::Stage_Vertex);
+		for (int i = 1; i <= Module->GetNumEntryPoints(); ++i)
 		{
 			// Reset bookkeeping to process new output subgraph
 			ValueStack.Empty(ValueStack.Max());
 			VisitedValues.Empty(Module->Values.Num());
 
+			const int EntryPointIndex = i % Module->GetNumEntryPoints();
 			FMaterialIRModule::FEntryPoint& EntryPoint = Module->GetEntryPoint(EntryPointIndex);
 
-			// Stable sort the outputs in this entry point by their indicated priority.
+			// Stable sort the outputs in this entry point by their indicated priority, then name.
 			// Note: this is needed to support situations where a final entry-point instruction needs to execute
 			//    before another, for instance we want to set the Normal attribute before all others, because the
 			//    others may read back the data. This shouldn't really be done this way, but it is currently a
-			//    restriction of how the HLSL backend works.
-			Algo::StableSort(EntryPoint.Outputs, [](const MIR::FValue* A, const MIR::FValue* B)
-				{
-					const MIR::FSetMaterialOutput* SMOA = MIR::As<MIR::FSetMaterialOutput>(A);
-					const MIR::FSetMaterialOutput* SMOB = MIR::As<MIR::FSetMaterialOutput>(B);
-					int PriorityA = SMOA ? SMOA->Priority : 0;
-					int PriorityB = SMOB ? SMOB->Priority : 0;
-					return PriorityA < PriorityB;
-				});
+			//    restriction of how the material shader works traditionally.
+			Algo::StableSort(EntryPoint.Outputs, [](const MIR::FInstruction* A, const MIR::FInstruction* B)
+			{
+				int32 PriorityA = A->GetOutputPriority();
+				int32 PriorityB = B->GetOutputPriority();
+				return PriorityA != PriorityB ? PriorityA < PriorityB : A->GetName() < B->GetName();
+			});
 
 			// Push this SetOutput instruction the value stack for processing.
 			for (MIR::FValue* Output : EntryPoint.Outputs)
 			{
-				if (Output)
-				{
-					ValueStack.Push(Output);
-				}
+				ValueStack.Push(Output);
 			}
 
 			// Process until the value stack is empty.
@@ -1141,7 +1147,7 @@ struct FMaterialIRModuleBuilderImpl
 		CompilationOutput.UniformExpressionSet.SetParameterCollections(Module->ParameterCollections);
 
 		// NumUsedUVScalars
-		CompilationOutput.NumUsedUVScalars = Module->Statistics.NumPixelTexCoords * 2;
+		CompilationOutput.NumUsedUVScalars = Module->Statistics.NumTexCoords * 2;
 
 		// NumVirtualTextureFeedbackRequests
 		int32 NumVirtualTextureFeedbackRequests = 0;
diff --git a/Engine/Source/Runtime/Engine/Private/Materials/MaterialIRToHLSLTranslator.cpp b/Engine/Source/Runtime/Engine/Private/Materials/MaterialIRToHLSLTranslator.cpp
--- a/Engine/Source/Runtime/Engine/Private/Materials/MaterialIRToHLSLTranslator.cpp
+++ b/Engine/Source/Runtime/Engine/Private/Materials/MaterialIRToHLSLTranslator.cpp
@@ -7,7 +7,6 @@
 #include "Engine/Texture2D.h"
 #include "VT/RuntimeVirtualTexture.h"
 #include "MaterialDomain.h"
-#include "MaterialIRInternal.h"
 #include "Materials/Material.h"
 #include "Materials/MaterialAggregate.h"
 #include "Materials/MaterialAttributeDefinitionMap.h"
@@ -16,6 +15,7 @@
 #include "Materials/MaterialIR.h"
 #include "Materials/MaterialIRModule.h"
 #include "Materials/MaterialIRTypes.h"
+#include "Materials/MaterialIRInternal.h"
 #include "Materials/MaterialParameterCollection.h"
 #include "MaterialShared.h"
 #include "MaterialSharedPrivate.h"
@@ -371,24 +371,13 @@ struct FPrivate : FMaterialIRToHLSLTranslation
 	// Generates the HLSL that sets the outputs of the WPO vertex stage.	
 	void GenerateVertexStageHLSL()
 	{
-		auto GenerateCustomizedUVAssignment = [this]
-		{
-			for (int32 CustomUVIndex = 0; CustomUVIndex < Module->GetStatistics().NumPixelTexCoords; CustomUVIndex++)
-			{
-				Printer << TEXTVIEW("Parameters.MaterialAttributes.CustomizedUV") << CustomUVIndex << TEXTVIEW(" = Parameters.TexCoords[") << CustomUVIndex << TEXTVIEW("].xy;\n");
-			}
-		};
-
 		// Generate current frame vertex stage HLSL
 		GenerateEntryPoint(MIR::Stage_Vertex);
-		GenerateCustomizedUVAssignment();
-	
 		OutputShaderTemplateParams.Add(TEXT("evaluate_vertex_material_attributes"), MoveTemp(Printer.Buffer));
 
 		// Generate previous frame vertex stage HLSL
 		bCompilingPreviousFrame = true;
 		GenerateEntryPoint(MIR::Stage_Vertex);
-		GenerateCustomizedUVAssignment();
 		bCompilingPreviousFrame = false;
 		OutputShaderTemplateParams.Add(TEXT("evaluate_vertex_material_attributes_prev"), MoveTemp(Printer.Buffer));
 	}
@@ -428,7 +417,7 @@ struct FPrivate : FMaterialIRToHLSLTranslation
 			// If this is a inlinable instruction type and it has only one user, instead of allocating a local
 			// and evaluating the instruction into it, simply skip this instruction for now.
 			// It will be generated inlined by its only user later.
-            if (IsFoldable(Instr, CurrentEntryPointIndex) && Instr->GetNumUsers(CurrentEntryPointIndex) == 1)
+            if (Instr->GetNumUsers(CurrentEntryPointIndex) == 1 && IsFoldable(Instr, CurrentEntryPointIndex))
 			{
                 continue;
             }
@@ -500,6 +489,7 @@ struct FPrivate : FMaterialIRToHLSLTranslation
 			else
 			{
 				// Otherwise, this instruction has already been generated before, we emit a reference to the local that stores its result.
+				check(InstrToLocalIndex.Contains(Instr));
 				Printer << TEXTVIEW("_") << InstrToLocalIndex[Instr];
 			}
 			return NoOp;
@@ -515,6 +505,7 @@ struct FPrivate : FMaterialIRToHLSLTranslation
 			case MIR::VK_TextureObject: LowerTextureObject(static_cast<const MIR::FTextureObject*>(Value)); break;
 			case MIR::VK_RuntimeVirtualTextureObject: LowerRuntimeVirtualTextureObject(static_cast<const MIR::FRuntimeVirtualTextureObject*>(Value)); break;
 			case MIR::VK_UniformParameter: LowerUniformParameter(static_cast<const MIR::FUniformParameter*>(Value)); break;
+			case MIR::VK_GetVertexInterpolator: LowerGetVertexInterpolator(static_cast<const MIR::FGetVertexInterpolator*>(Value)); break;
 			default: UE_MIR_UNREACHABLE();
 		}
 
@@ -537,6 +528,7 @@ struct FPrivate : FMaterialIRToHLSLTranslation
 			case MIR::VK_InlineHLSL: LowerInlineHLSL(static_cast<const MIR::FInlineHLSL*>(Instr)); break;
 			case MIR::VK_StageSwitch: LowerStageSwitch(static_cast<const MIR::FStageSwitch*>(Instr)); break;
 			case MIR::VK_HardwarePartialDerivative: LowerHardwarePartialDerivative(static_cast<const MIR::FHardwarePartialDerivative*>(Instr)); break;
+			case MIR::VK_SetVertexInterpolator: LowerSetVertexInterpolator(static_cast<const MIR::FSetVertexInterpolator*>(Instr)); break;
 			case MIR::VK_Call: LowerCall(static_cast<const MIR::FCall*>(Instr)); break;
 			case MIR::VK_CallParameterOutput: LowerCallOutput(static_cast<const MIR::FCallParameterOutput*>(Instr)); break;
 			case MIR::VK_PreshaderParameter: LowerPreshaderParameter(static_cast<const MIR::FPreshaderParameter*>(Instr)); break;
@@ -809,7 +801,26 @@ struct FPrivate : FMaterialIRToHLSLTranslation
 		FlushCodeSubstring(Code.Len());
 	}
 
-	
+	void LowerGetVertexInterpolator(const MIR::FGetVertexInterpolator* GetVertexInterpolator)
+	{
+		if (GetVertexInterpolator->Type.GetPrimitive().NumComponents() > 1)
+		{
+			Printer << GetVertexInterpolator->Type << TEXT('(');
+		}
+		ComputeUserVertexInterpolatorAllocation(GetVertexInterpolator->SetInstr, [this](int32 Index, int32 RegisterIndex, FStringView RegisterSwizzle, FStringView)
+		{
+			if (Index > 0)
+			{
+				Printer << TEXTVIEW(", ");
+			}
+			Printer << TEXTVIEW("Parameters.TexCoords[") << RegisterIndex << TEXTVIEW("].") << RegisterSwizzle;
+		});
+		if (GetVertexInterpolator->Type.GetPrimitive().NumComponents() > 1)
+		{
+			Printer << TEXT(')');
+		}
+	}
+
 	void LowerInlineHLSL(const MIR::FInlineHLSL* ExternalCode)
 	{
 		FStringView FinalCode;
@@ -1967,6 +1978,32 @@ struct FPrivate : FMaterialIRToHLSLTranslation
 		}
 	}
 
+	void LowerSetVertexInterpolator(const MIR::FSetVertexInterpolator* SetVertexInterpolator)
+	{
+		// If the argument is foldable, first assign it to a HLSL local variable to avoid getting it evaluated more than once should
+		// the interpolator spill its base TexCoords register to the next.
+		uint32 Local;
+		const MIR::FInstruction* ArgAsInstr = MIR::AsInstruction(SetVertexInterpolator->Arg);
+		if (!ArgAsInstr || !MIR::Find(InstrToLocalIndex, ArgAsInstr, Local))
+		{
+			Local = NumLocals++;
+			Printer << SetVertexInterpolator->Arg->Type << TEXT(" _") << Local << TEXTVIEW(" = ") << LowerValue(SetVertexInterpolator->Arg) << EndOfStatement;
+		}
+
+		ComputeUserVertexInterpolatorAllocation(SetVertexInterpolator, [this, SetVertexInterpolator, Local](int32 Index, int32 RegisterIndex, FStringView RegisterSwizzle, FStringView ArgumentSwizzle)
+		{
+			if (Index > 0)
+			{
+				Printer << EndOfStatement;
+			}
+			Printer << TEXTVIEW("Parameters.TexCoords[") << RegisterIndex << TEXTVIEW("].") << RegisterSwizzle << TEXTVIEW(" = _") << Local;
+			if (!ArgumentSwizzle.IsEmpty())
+			{
+				Printer << TEXTVIEW(".") << ArgumentSwizzle;
+			}
+		});
+	}
+
 	void LowerCall(const MIR::FCall* Call)
 	{
 		int32 ParamLocal = NumLocals;
@@ -2246,6 +2283,52 @@ struct FPrivate : FMaterialIRToHLSLTranslation
 		Printer << EndArgs;
 	}
 
+	// Helper function that generates the mappings between the TexCoords float parameters and given VertexInterpolators, taking care of spilling to subsequent registers and appropriate swizzling.
+	//  - Index is the assignment index
+	//  - RegisterIndex is the TexCoords register this mapping
+	//  - RegisterSwizzle is the TexCoords register swizzle for this mapping
+	//  - ArgumentSwizzle is the components of SetVertexInterpolator->Arg for this mapping
+	// For instance, this may generate mappings like so (Arg is a float3)
+	//    0) TexCoords[0].zw <==> Arg.xy
+	//    1) TexCoord[1].x <==> Arg.z
+	void ComputeUserVertexInterpolatorAllocation(const MIR::FSetVertexInterpolator* SetVertexInterpolator, TFunctionRef<void(int32 MappingIndex, int32 RegisterIndex, FStringView RegisterSwizzle, FStringView ArgumentSwizzle)> Function)
+	{
+		static const TCHAR* SwizzleComponents = TEXT("xyzw");
+
+		int MappingIndex = 0;
+
+		const int BaseSlot = SetVertexInterpolator->Analysis_BaseInterpolatorFloatRegister;
+		const int NumComponents = SetVertexInterpolator->Arg->Type.GetPrimitive().NumComponents();
+
+		int SlotBegin = BaseSlot;
+		const int SlotEnd = BaseSlot + NumComponents;
+
+		while (SlotBegin < SlotEnd)
+		{
+			const int RegisterIndex = Module->GetStatistics().NumTexCoords + SlotBegin / 2;
+			const int RegisterSlotOffset = SlotBegin % 2;
+
+			// How many components fit in the current float2 from this offset. The minimum between the remaining
+			// slots in the register and the those remaining to write
+			const int PlaceableSlots = FMath::Min(2 - RegisterSlotOffset, SlotEnd - SlotBegin);
+
+			FStringView RegisterSwizzle{ &SwizzleComponents[RegisterSlotOffset], PlaceableSlots };
+			FStringView ArgumentSwizzle{};
+
+			// If we are writing a subset, swizzle the RHS to match this chunk.
+			const int ArgComponentOffset = SlotBegin - BaseSlot; // 0..NumComponents-1
+			if (!(ArgComponentOffset == 0 && PlaceableSlots == NumComponents))
+			{
+				ArgumentSwizzle = { &SwizzleComponents[ArgComponentOffset], PlaceableSlots };
+			}
+
+			Function(MappingIndex, RegisterIndex, RegisterSwizzle, ArgumentSwizzle);
+
+			SlotBegin += PlaceableSlots;
+			++MappingIndex;
+		}
+	}
+
 	/* Finalization */
 
 	void GenerateCustomOutputDeclarations()
@@ -2254,16 +2337,21 @@ struct FPrivate : FMaterialIRToHLSLTranslation
 		// set by the function that calculates all material outputs (e.g. BaseColor).
 		FHLSLPrinter VertexPrinter{ 1 };
 		FHLSLPrinter PixelPrinter{ 1 };
-		for (const MIR::FSetMaterialOutput* CustomOutput : Module->GetCustomOutputs())
+
+		for (const MIR::FSetMaterialOutput* Output : Module->GetCustomOutputs())
 		{
-			FHLSLPrinter& ThePrinter = CustomOutput->Stage == MIR::MaterialOutputStage_Vertex ? VertexPrinter : PixelPrinter;
-			ThePrinter << CustomOutput->Arg->Type << TEXT(' ') << *CustomOutput->Name << EndOfStatement;
+			FHLSLPrinter& ThePrinter = Output->Stage == MIR::MaterialOutputStage_Vertex ? VertexPrinter : PixelPrinter;
+			ThePrinter << Output->Arg->Type << TEXT(' ') << *Output->Name << EndOfStatement;
 		}
-		OutputShaderTemplateParams.Add(TEXT("material_vertex_parameters_extras"), MoveTemp(VertexPrinter.Buffer));
-		OutputShaderTemplateParams.Add(TEXT("material_pixel_parameters_extras"), MoveTemp(PixelPrinter.Buffer));
+		
+		OutputShaderTemplateParams.Add(TEXT("material_vertex_parameter_decls"), MoveTemp(VertexPrinter.Buffer));
+		OutputShaderTemplateParams.Add(TEXT("material_pixel_parameter_decls"), MoveTemp(PixelPrinter.Buffer));
 
 		// The material shader expects a getter function for each custom output. Generate these getter functions
 		// for each FMaterialPixelParameters field, generated above, matching each custom output in each group.
+		// Development Note: this is really just cruft necessary for backwards compatibility with the way the
+		// 	  old translator supported custom outputs and could be removed entirely with a little refactoring
+		//    of MaterialTemplate.ush.
 		Printer = {};
 		for (const FMaterialIRModule::FCustomOutputGroup& Group : Module->GetCustomOutputGroups())
 		{
@@ -2277,11 +2365,10 @@ struct FPrivate : FMaterialIRToHLSLTranslation
 				// Select the right XXXParameters HLSL struct used by this output based on its 
 				FStringView Signature = Instr->Stage == MIR::MaterialOutputStage_Vertex
 					? TEXTVIEW("(inout FMaterialVertexParameters Parameters)")
-				: TEXTVIEW("(inout FMaterialPixelParameters Parameters)");
+					: TEXTVIEW("(inout FMaterialPixelParameters Parameters)");
 
 				// Print the non-LWC get function.
-				Printer << Instr->Arg->Type << TEXTVIEW(" Get") << Instr->Name << Signature
-					<< OpenBlock << TEXTVIEW("return ");
+				Printer << Instr->Arg->Type << TEXTVIEW(" Get") << Instr->Name << Signature << OpenBlock << TEXTVIEW("return ");
 
 				Instr->Arg->Type.IsDouble()
 					? Printer << TEXTVIEW("WSDemote(Parameters.") << Instr->Name << TEXT(')')
@@ -2293,8 +2380,7 @@ struct FPrivate : FMaterialIRToHLSLTranslation
 				if (!Instr->Arg->Type.IsSubstrateData())
 				{
 					MIR::FType ReturnType = Instr->Arg->Type.GetPrimitive().ToScalarKind(MIR::EScalarKind::Double);
-					Printer << ReturnType << TEXTVIEW(" Get") << Instr->Name << TEXTVIEW("_LWC") << Signature
-						<< OpenBlock << TEXTVIEW("return ");
+					Printer << ReturnType << TEXTVIEW(" Get") << Instr->Name << TEXTVIEW("_LWC") << Signature << OpenBlock << TEXTVIEW("return ");
 				
 					Instr->Arg->Type.IsDouble()
 						? Printer << TEXTVIEW("Parameters.") << Instr->Name
@@ -2375,7 +2461,6 @@ struct FPrivate : FMaterialIRToHLSLTranslation
 	void GenerateOtherTemplateStringParameters()
 	{
 		const FMaterialIRModule::FStatistics ModuleStatistics = Module->GetStatistics();
-
 		const FMaterialIRModule::FSubstrateData& SubstrateData = Module->GetSubstrateData();
 
 		// material_declarations
@@ -2412,23 +2497,26 @@ struct FPrivate : FMaterialIRToHLSLTranslation
 
 			OutputShaderTemplateParams.Add(TEXT("pixel_material_inputs"), MoveTemp(Printer.Buffer));
 		}
-		
+
+		const int32 NumUserVertexInterpolators = (ModuleStatistics.NumUsedVertexInterpolatorFloatRegisters + 1) / 2; // packed into float2s
+		const int32 ShaderTexCoordsArrayLength = ModuleStatistics.NumTexCoords + NumUserVertexInterpolators;
+
 		// get_material_customized_u_vs
 		{
 			Printer = { 1 };
-			for (int32 CustomUVIndex = 0; CustomUVIndex < ModuleStatistics.NumPixelTexCoords; CustomUVIndex++)
+			for (int32 i = 0; i < ShaderTexCoordsArrayLength; i++)
 			{
-				Printer << TEXTVIEW("OutTexCoords[") << CustomUVIndex << TEXTVIEW("] = Parameters.MaterialAttributes.CustomizedUV") << CustomUVIndex << EndOfStatement;
+				Printer << TEXTVIEW("OutTexCoords[") << i << TEXTVIEW("] = Parameters.TexCoords[") << i << TEXT(']') << EndOfStatement;
 			}
 			OutputShaderTemplateParams.Add(TEXT("get_material_customized_u_vs"), MoveTemp(Printer.Buffer));
 		}
 		
 		// texcoords parameters
 		{
-			OutputShaderTemplateParams.Add(TEXT("num_material_texcoords_vertex"), FString::FromInt(ModuleStatistics.NumVertexTexCoords));
-			OutputShaderTemplateParams.Add(TEXT("num_material_texcoords"), FString::FromInt(ModuleStatistics.NumPixelTexCoords));
-			OutputShaderTemplateParams.Add(TEXT("num_custom_vertex_interpolators"), FString::FromInt(0));
-			OutputShaderTemplateParams.Add(TEXT("num_tex_coord_interpolators"), FString::FromInt(ModuleStatistics.NumPixelTexCoords));
+			OutputShaderTemplateParams.Add(TEXT("num_material_texcoords"), FString::FromInt(ModuleStatistics.NumTexCoords));
+			OutputShaderTemplateParams.Add(TEXT("num_material_texcoords_vertex"), FString::FromInt(ModuleStatistics.NumTexCoords));
+			OutputShaderTemplateParams.Add(TEXT("num_custom_vertex_interpolators"), FString::FromInt(NumUserVertexInterpolators));
+			OutputShaderTemplateParams.Add(TEXT("num_tex_coord_interpolators"), FString::FromInt(ShaderTexCoordsArrayLength));
 		}
 
 		// material properties
diff --git a/Engine/Source/Runtime/Engine/Private/Materials/MaterialIRValueAnalyzer.cpp b/Engine/Source/Runtime/Engine/Private/Materials/MaterialIRValueAnalyzer.cpp
--- a/Engine/Source/Runtime/Engine/Private/Materials/MaterialIRValueAnalyzer.cpp
+++ b/Engine/Source/Runtime/Engine/Private/Materials/MaterialIRValueAnalyzer.cpp
@@ -51,8 +51,7 @@ static void AnalyzeExternalInput(FMaterialIRValueAnalyzer& Analyzer, MIR::FExter
 	if (MIR::IsExternalInputTexCoordOrPartialDerivative(ExternalInput->Id))
 	{
 		int32 TexCoordIndex = MIR::ExternalInputToTexCoordIndex(ExternalInput->Id);
-		Statistics.NumVertexTexCoords = FMath::Max(Statistics.NumVertexTexCoords, TexCoordIndex + 1);
-		Statistics.NumPixelTexCoords = FMath::Max(Statistics.NumPixelTexCoords, TexCoordIndex + 1);
+		Statistics.NumTexCoords = FMath::Max(Statistics.NumTexCoords, TexCoordIndex + 1);
 	}
 	else if (Analyzer.MaterialInterface->IsPostProcessMaterial() && MIR::IsExternalInputWorldPosition(ExternalInput->Id))
 	{
@@ -683,7 +682,7 @@ static uint32 FindOrAddDefaultValueOffset(FMaterialIRValueAnalyzer& Analyzer, co
 	}
 
 	uint32 DefaultValueOffset = 0;
-	if (!MIR::Internal::Find(Analyzer.UniformDefaultValueOffsets, DefaultValue, DefaultValueOffset))
+	if (!MIR::Find(Analyzer.UniformDefaultValueOffsets, DefaultValue, DefaultValueOffset))
 	{
 		FUniformExpressionSet& UniformExpressionSet = Analyzer.CompilationOutput->UniformExpressionSet;
 		DefaultValueOffset = UniformExpressionSet.AddDefaultParameterValue(DefaultValue);
@@ -784,12 +783,19 @@ static void AnalyzePreshaderParameter(FMaterialIRValueAnalyzer& Analyzer, MIR::F
 
 static void AnalyzeSetMaterialOutput(FMaterialIRValueAnalyzer& Analyzer, MIR::FSetMaterialOutput* SetMaterialOutput)
 {
-	if (SetMaterialOutput->Property == EMaterialProperty::MP_Normal)
+	// If this output is evaluated in the vertex stage, ensure its input is also valid there.
+	// Inputs marked "PixelStageOnly" cannot be used when generating vertex-stage code
+	if (SetMaterialOutput->Stage == MIR::MaterialOutputStage_Vertex && SetMaterialOutput->Arg->HasSubgraphProperties(MIR::EGraphProperties::PixelStageOnly))
 	{
-		if (SetMaterialOutput->HasSubgraphProperties(MIR::EGraphProperties::ReadsPixelNormal))
-		{
-			Analyzer.Module->AddError(nullptr, TEXT("Cannot set material attribute Normal to a value that depends on reading the pixel normal, as that would create a circular dependency."));
-		}
+		Analyzer.Module->AddError({}, FString::Printf(TEXT("Input to output '%s' is not available during the vertex stage."), *SetMaterialOutput->Name));
+		return;
+	}
+	
+	// Normal output cannot read the pixel normal
+	if (SetMaterialOutput->Property == EMaterialProperty::MP_Normal && SetMaterialOutput->HasSubgraphProperties(MIR::EGraphProperties::ReadsPixelNormal))
+	{
+		Analyzer.Module->AddError(nullptr, TEXT("Cannot set material attribute Normal to a value that depends on reading the pixel normal, as that would create a circular dependency."));
+		return;
 	}
 }
 
@@ -805,6 +811,34 @@ static void AnalyzeMaterialParameterCollection(FMaterialIRValueAnalyzer& Analyze
 	MaterialParameterCollection->Analysis_CollectionIndex = CollectionIndex;
 }
 
+static void AnalyzeGetVertexInterpolator(FMaterialIRValueAnalyzer& Analyzer, MIR::FGetVertexInterpolator* GetVertexInterpolator)
+{
+	// Record in the compilation output that the material is using vertex interpolators
+	Analyzer.Module->GetCompilationOutput().bUsesVertexInterpolator = true;
+
+	// Since this value is being analyzed it means it is actually used by the material to compute some output.
+	// Push its paired set instruction to the list of vertex stage entry points to make it enabled and visible, so that it gets analyzed too.
+	Analyzer.Module->GetEntryPoint(MIR::EStage::Stage_Vertex).Outputs.Push(GetVertexInterpolator->SetInstr);
+}
+
+static void AnalyzeSetVertexInterpolator(FMaterialIRValueAnalyzer& Analyzer, MIR::FSetVertexInterpolator* SetVertexInterpolator)
+{
+	int32& NumInterpRegisters = Analyzer.Module->GetStatistics().NumUsedVertexInterpolatorFloatRegisters;
+
+	// Allocate the appropriat number of slots (float scalars) used by this interpolator.
+	SetVertexInterpolator->Analysis_BaseInterpolatorFloatRegister = NumInterpRegisters;
+	NumInterpRegisters += SetVertexInterpolator->Arg->Type.GetPrimitive().NumComponents();
+
+	// Validate that not too many interpolators are being used.
+	// @massimo.tristano todo: expose a way to control the maximum number of interpolators.
+	const int32 MaxNumVertexInterpolatorSlots = 8 * 2;
+	if (NumInterpRegisters > MaxNumVertexInterpolatorSlots)
+	{
+		Analyzer.Module->AddError({}, FString::Printf(TEXT("Maximum number of vertex interpolator exceeded (%d)."), MaxNumVertexInterpolatorSlots));
+		return;
+	}
+}
+
 //////////////////////////////////////////////////////////////////////////////////////////
 // FMaterialIRValueAnalyzer
 //////////////////////////////////////////////////////////////////////////////////////////
@@ -846,6 +880,8 @@ void FMaterialIRValueAnalyzer::Analyze(MIR::FValue* Value)
 		EXPAND_CASE(InlineHLSL);
 		EXPAND_CASE(ScreenTexture);
 		EXPAND_CASE(ShadingModel);
+		EXPAND_CASE(GetVertexInterpolator);
+		EXPAND_CASE(SetVertexInterpolator);
 		default: break;
 	}
 
diff --git a/Engine/Source/Runtime/Engine/Public/Materials/MaterialExpressionVertexInterpolator.h b/Engine/Source/Runtime/Engine/Public/Materials/MaterialExpressionVertexInterpolator.h
--- a/Engine/Source/Runtime/Engine/Public/Materials/MaterialExpressionVertexInterpolator.h
+++ b/Engine/Source/Runtime/Engine/Public/Materials/MaterialExpressionVertexInterpolator.h
@@ -15,6 +15,7 @@ class UMaterialExpressionVertexInterpolator : public UMaterialExpressionCustomOu
 	FExpressionInput Input;
 
 #if WITH_EDITOR
+	virtual void Build(MIR::FEmitter& Emitter) override;
 	virtual int32 Compile(class FMaterialCompiler* Compiler, int32 OutputIndex) override;
 	ENGINE_API int32 CompileInput(class FMaterialCompiler* Compiler, int32 AssignedInterpolatorIndex);
 	virtual void GetCaption(TArray<FString>& OutCaptions) const override;
diff --git a/Engine/Source/Runtime/Engine/Public/Materials/MaterialIR.h b/Engine/Source/Runtime/Engine/Public/Materials/MaterialIR.h
--- a/Engine/Source/Runtime/Engine/Public/Materials/MaterialIR.h
+++ b/Engine/Source/Runtime/Engine/Public/Materials/MaterialIR.h
@@ -89,8 +89,11 @@ enum class EGraphProperties : uint8
 {
 	None = 0,
 
+	// This value is only meaningful and available in the pixel stage.
+	PixelStageOnly = 1 << 0,
+
 	// Some value reads the pixel normal.
-	ReadsPixelNormal = 1 << 0,
+	ReadsPixelNormal = 1 << 1,
 };
 
 ENUM_CLASS_FLAGS(EGraphProperties);
@@ -110,6 +113,7 @@ enum EValueKind
 	VK_TextureObject,
 	VK_RuntimeVirtualTextureObject,
 	VK_UniformParameter,
+	VK_GetVertexInterpolator,
 
 	// Instructions
 
@@ -136,6 +140,7 @@ enum EValueKind
 	VK_SubstratePromoteToOperator,
 	VK_StageSwitch,
 	VK_HardwarePartialDerivative,
+	VK_SetVertexInterpolator,
 	VK_Call,
 	VK_CallParameterOutput,
 	VK_PreshaderParameter,
@@ -583,6 +588,13 @@ struct FUniformParameter: TValue<VK_UniformParameter>
 	int32 Analysis_UniformParameterIndex;
 };
 
+// Value used to retrieve (in the pixel stage) a value interpolated in the vertex stage.
+struct FGetVertexInterpolator : TValue<VK_GetVertexInterpolator>
+{
+	// Associated Set- instruction create at the time of emission.
+	struct FSetVertexInterpolator* SetInstr;
+};
+
 enum class EScreenTexture : uint8
 {
 	SceneTexture,
@@ -683,6 +695,11 @@ struct FInstruction : FValue
 	// Returns the block this instruction belongs to in specified entry point.
 	const FBlock* GetBlock(int32 EntryPointIndex) const { return Linkage[EntryPointIndex].Block; }
 
+	// Retrieves a user-specified name string if this instruction provides one.
+	FStringView GetName() const;
+
+	// Retrieves the priority if this instruction is used as a entry point output.
+	int32 GetOutputPriority() const;
 };
 
 template <EValueKind TTypeKind, uint32 TNumStaticUses = 0>
@@ -1276,6 +1293,18 @@ struct FHardwarePartialDerivative : TInstruction<VK_HardwarePartialDerivative, 1
 	EDerivativeSource Source;
 };
 
+// Instruction executed in the vertex stage that sets the value of some interpolated data, available
+// in the pixel stage via its matching Get- value.
+struct FSetVertexInterpolator : TInstruction<VK_SetVertexInterpolator, 1>
+{
+	// The value being set to the interpolator.
+	FValue* Arg;
+
+	// Interpolator registers are abstracted as an array of float scalars. This field, determined at
+	// analysis time, is assigned to the index of the first float register occupied by this interpolator.
+	uint32 Analysis_BaseInterpolatorFloatRegister;
+};
+
 // Kind of user-specified function.
 enum class FFunctionKind
 {
diff --git a/Engine/Source/Runtime/Engine/Public/Materials/MaterialIRCommon.h b/Engine/Source/Runtime/Engine/Public/Materials/MaterialIRCommon.h
--- a/Engine/Source/Runtime/Engine/Public/Materials/MaterialIRCommon.h
+++ b/Engine/Source/Runtime/Engine/Public/Materials/MaterialIRCommon.h
@@ -70,7 +70,7 @@ namespace Internal
 	uint32 HashBytes(const char* Ptr, uint32 Size);
 }
 
-/* Helper structures */
+/* Utilities */
 
 // Use this to efficiently allocate a temporary array using MemStack, instead
 // of using TArray and going through the global allocator. Only declare local variable
@@ -88,6 +88,13 @@ struct TTemporaryArray : public TArrayView<T>
 		*static_cast<TArrayView<T>*>(this) = { (T*)FMemStack::Get().Alloc(sizeof(T) * Num, alignof(T)), Num };
 	}
 
+	TTemporaryArray(TConstArrayView<T> Source)
+	: TTemporaryArray { Source.Num() }
+	{
+		static_assert(std::is_trivially_copyable_v<T>);
+		FMemory::Memcpy(TArrayView<T>::GetData(), Source.GetData(), Source.NumBytes());
+	}
+
 	operator TConstArrayView<T>() const
 	{
 		return { TArrayView<T>::GetData(), TArrayView<T>::Num() };
@@ -106,8 +113,18 @@ void ZeroArray(TArrayView<T> Array)
 	FMemory::Memzero(Array.GetData(), Array.Num() * sizeof(T));
 }
 
+template <typename TKey, typename TValue>
+bool Find(const TMap<TKey, TValue>& Map, const TKey& Key, TValue& OutValue)
+{
+	if (auto ValuePtr = Map.Find(Key)) {
+		OutValue = *ValuePtr;
+		return true;
+	}
+	return false;
 }
 
+} // namespace MIR
+
 #define UE_MIR_UNREACHABLE() { check(!"Unreachable"); UE_ASSUME(false); }
 #define UE_MIR_TODO() UE_MIR_UNREACHABLE()
 
diff --git a/Engine/Source/Runtime/Engine/Public/Materials/MaterialIREmitter.h b/Engine/Source/Runtime/Engine/Public/Materials/MaterialIREmitter.h
--- a/Engine/Source/Runtime/Engine/Public/Materials/MaterialIREmitter.h
+++ b/Engine/Source/Runtime/Engine/Public/Materials/MaterialIREmitter.h
@@ -716,6 +716,11 @@ public:
 	// @return A new FValueRef representing the analytical derivative, or a poison value if invalid.
 	FValueRef AnalyticalPartialDerivative(FValueRef Value, EDerivativeAxis Axis);
 
+	// Creates a vertex to pixel interpolator for the given value.
+	// The value is computed in the vertex stage, linearly interpolated across the primitive,
+	// and exposed in the pixel stage. The returned value is thus valid only in the pixel stage.
+	FValueRef VertexInterpolator(FValueRef VertexValue);
+
 	// Emit a new HLSL function inline within the material shader.
 	// @param Type The expected output type of the HLSL code.
 	// @param Code The HLSL code to be emitted.
diff --git a/Engine/Source/Runtime/Engine/Public/Materials/MaterialIRModule.h b/Engine/Source/Runtime/Engine/Public/Materials/MaterialIRModule.h
--- a/Engine/Source/Runtime/Engine/Public/Materials/MaterialIRModule.h
+++ b/Engine/Source/Runtime/Engine/Public/Materials/MaterialIRModule.h
@@ -36,11 +36,11 @@ public:
 		// Tracks external inputs used per frequency.
 		TBitArray<> ExternalInputUsedMask[MIR::NumStages]; 
 
-		// Number of vertex texture coordinates used.
-		int32 NumVertexTexCoords;
-
-		// Number of pixel texture coordinates used.
-		int32 NumPixelTexCoords;
+		// Number of texture coordinates used by the material.
+		int32 NumTexCoords;
+		
+		// Total number of user vertex-interpolator float registers used by the material.
+		int32 NumUsedVertexInterpolatorFloatRegisters;
 
 		// Dynamic particle parameter indices used by this material.
 		uint32 DynamicParticleParameterMask;
@@ -62,7 +62,7 @@ public:
 		MIR::FBlock RootBlock;
 
 		// An array of terminal values this entry point evaluates.
-		TArray<MIR::FValue*> Outputs;
+		TArray<MIR::FInstruction*> Outputs;
 	};
 
 	// Models a group of custom outputs contained in the module.
