BodySetup: Cooking: Replace raw pointer to FFormatContainer during cooking with a WeakPtr to the SkeletalMesh that has the BodySetup that has the FormatContainer. Looking up the FFormatContainer pointer through the weak pointer (or failing to do so because the SkeletalMesh or BodySetup has been garbage collected and the WeakPtr is null) will avoid dereferencing a dangling pointer or using stale data no longer used by the SkeletalMesh.
#rnx

#Synced-CL 50149715

diff --git a/Engine/Source/Runtime/Engine/Classes/PhysicsEngine/BodySetup.h b/Engine/Source/Runtime/Engine/Classes/PhysicsEngine/BodySetup.h
--- a/Engine/Source/Runtime/Engine/Classes/PhysicsEngine/BodySetup.h
+++ b/Engine/Source/Runtime/Engine/Classes/PhysicsEngine/BodySetup.h
@@ -28,6 +28,7 @@
 #include "BodySetup.generated.h"
 
 class ITargetPlatform;
+class USkeletalMesh;
 class UPhysicalMaterial;
 class UPhysicalMaterialMask;
 class UPrimitiveComponent;
@@ -224,6 +225,11 @@ class UBodySetup : public UBodySetupCore
 
 	/** Indicates whether this setup has any cooked collision data. This is set when writing to or reading from cooked data, but not authoritative for assets in the editor. */
 	uint8 bHasCookedCollisionData:1;
+	/**
+	 * Reports whether CookedFormatProvider was set and we should give an error when trying to read CookedData if
+	 * the provider is no longer available. Transient flag set in CreateBodySetup.
+	 */
+	uint8 bHasCookedFormatProvider : 1;
 
 	/** 
 	 * Indicates that we will never use convex or trimesh shapes. This is an optimization to skip checking for binary data.
@@ -274,8 +280,12 @@ class UBodySetup : public UBodySetupCore
 	UPROPERTY(EditAnywhere, Category=Collision, meta=(FullyExpand = "true"))
 	FBodyInstance DefaultInstance;
 
-	/** Cooked physics data override. This is needed in cases where some other body setup has the cooked data and you don't want to own it or copy it. See per poly skeletal mesh */
-	FFormatContainer* CookedFormatDataOverride;
+	/**
+	 * Cooked physics data override. This is needed in cases where some other body setup has the
+	 * cooked data and you don't want to own it or copy it. See per poly skeletal mesh.
+	 * This is a transient pointer that can be set during CreateBodySetup.
+	 */
+	TWeakObjectPtr<const USkeletalMesh> CookedFormatProvider;
 
 	/** Build scale for this body setup (static mesh settings define this value) */
 	UPROPERTY()
@@ -332,7 +342,8 @@ public:
 	{
 		NoError,
 		FailedToLock,
-		FailedToCook
+		FailedToCook,
+		DataFormatProviderUnavailable,
 	};
 
 	struct FBodySetupTryCookResult
diff --git a/Engine/Source/Runtime/Engine/Private/PhysicsEngine/BodySetup.cpp b/Engine/Source/Runtime/Engine/Private/PhysicsEngine/BodySetup.cpp
--- a/Engine/Source/Runtime/Engine/Private/PhysicsEngine/BodySetup.cpp
+++ b/Engine/Source/Runtime/Engine/Private/PhysicsEngine/BodySetup.cpp
@@ -8,6 +8,7 @@
 #include "BodySetupEnums.h"
 #include "Engine/Engine.h"
 #include "Engine/StaticMesh.h"
+#include "Engine/SkeletalMesh.h"
 #include "Components/SkinnedMeshComponent.h"
 #include "Engine/World.h"
 #include "PhysicsEngine/PhysicsSettings.h"
@@ -150,7 +151,8 @@ UBodySetup::UBodySetup(const FObjectInitializer& ObjectInitializer)
 	BuildScale3D = FVector::OneVector;
 	SetFlags(RF_Transactional);
 	bSharedCookedData = false;
-	CookedFormatDataOverride = nullptr;
+	bHasCookedFormatProvider = false;
+	CookedFormatProvider = nullptr;
 	CurrentCookHelper = nullptr;
 }
 
@@ -1227,11 +1229,19 @@ bool UBodySetup::IsCachedCookedPlatformDataLoaded(const ITargetPlatform* TargetP
 	}
 	FBodySetupTryCookResult CookResult = TryGetCookedData(TargetPlatform->GetPhysicsFormat(this));
 
-	if(CookResult.CookResult == EBodySetupCookResult::FailedToLock)
+	if (CookResult.CookResult == EBodySetupCookResult::FailedToLock)
 	{
 		UE_LOG(LogPhysics, Warning, TEXT("Multiple writes to BodySetup bulk data detected during cook for setup: %s"), *GetPathName());
 		return false;
 	}
+	if (CookResult.CookResult == EBodySetupCookResult::DataFormatProviderUnavailable)
+	{
+		UE_LOG(LogPhysics, Error,
+			TEXT("This BodySetup is on a skeletal component, and is supposed to copy its FormatContainer from the body setup of its SkeletalMesh asset. But that Body setup is no longer available. %s"),
+			*GetPathName());
+		// We cannot recover from this, so report that we are done trying to load, to avoid softlocking the cook.
+		return true;
+	}
 
 	return true;
 }
@@ -1438,18 +1448,30 @@ UBodySetup::FBodySetupTryCookResult UBodySetup::TryGetCookedData(FName Format)
 		return { nullptr, CookDataResult };
 	}
 
-#if WITH_EDITOR
-	//We don't support runtime cook optimization for per poly skeletal mesh. This is an edge case we may want to support (only helps memory savings)
-	FFormatContainer* UseCookedData = CookedFormatDataOverride ? CookedFormatDataOverride : &CookedFormatData;
-#else
-	FFormatContainer* UseCookedData = CookedFormatDataOverride ? CookedFormatDataOverride : &CookedFormatData;
-#endif
+	FFormatContainer* UsedFormatData = &CookedFormatData;
+	if (bHasCookedFormatProvider)
+	{
+		UsedFormatData = nullptr;
+		const USkeletalMesh* SkeletalMeshWithFormat = CookedFormatProvider.Get();
+		if (SkeletalMeshWithFormat)
+		{
+			UBodySetup* BodySetupWithFormat = SkeletalMeshWithFormat->GetBodySetup();
+			if (BodySetupWithFormat)
+			{
+				UsedFormatData = &BodySetupWithFormat->CookedFormatData;
+			}
+		}
+		if (!UsedFormatData)
+		{
+			return { nullptr, EBodySetupCookResult::DataFormatProviderUnavailable };
+		}
+	}
 
 	FByteBulkData* Result = nullptr;
 	const bool bIsRuntime = IsRuntime(this);
 
 #if WITH_EDITOR
-	if(!UseCookedData->Contains(Format))
+	if(!UsedFormatData->Contains(Format))
 	{
 		SCOPE_CYCLE_COUNTER(STAT_PhysXCooking);
 
@@ -1470,7 +1492,7 @@ UBodySetup::FBodySetupTryCookResult UBodySetup::TryGetCookedData(FName Format)
 		// Check the lock condition again, we only want the first thread in here to actually build the data,
 		// while it's building other threads should wait and proceed after the build completes.
 		// Double-checked locking avoids always taking the lock - we only do it for unbuilt data
-		if(!UseCookedData->Contains(Format))
+		if(!UsedFormatData->Contains(Format))
 		{
 			// We do not want a FGCObject to be created to prevent garbage collection of our own UBodySetup*
 			// because the scope and lifetime is well defined and FGCObject can't be created on other threads
@@ -1478,19 +1500,19 @@ UBodySetup::FBodySetupTryCookResult UBodySetup::TryGetCookedData(FName Format)
 			const bool bUseRefHolder = false;
 			FChaosDerivedDataCooker* PhysicsDerivedCooker = new FChaosDerivedDataCooker(this, Format, bUseRefHolder);
 
-			Result = &UseCookedData->GetFormat(Format);
+			Result = &UsedFormatData->GetFormat(Format);
 			CookDataResult = TryGetDDCBuiltData(Result, *PhysicsDerivedCooker, this, bIsRuntime);
 		}
 		else
 		{
 			// We were a waiting thread while another thread built the data, just take it from the container now
-			Result = &UseCookedData->GetFormat(Format);
+			Result = &UsedFormatData->GetFormat(Format);
 		}
 	}
 	else
 #endif // #if WITH_EDITOR
 	{
-		Result = &UseCookedData->GetFormat(Format);
+		Result = &UsedFormatData->GetFormat(Format);
 	}
 
 	check(Result);
diff --git a/Engine/Source/Runtime/Engine/Private/SkeletalMeshComponentPhysics.cpp b/Engine/Source/Runtime/Engine/Private/SkeletalMeshComponentPhysics.cpp
--- a/Engine/Source/Runtime/Engine/Private/SkeletalMeshComponentPhysics.cpp
+++ b/Engine/Source/Runtime/Engine/Private/SkeletalMeshComponentPhysics.cpp
@@ -279,7 +279,8 @@ void USkeletalMeshComponent::CreateBodySetup()
 		GetSkeletalMeshAsset()->CreateBodySetup();
 		UBodySetup* OriginalBodySetup = SkeletalMeshConst->GetBodySetup();
 		BodySetup->CopyBodyPropertiesFrom(OriginalBodySetup);
-		BodySetup->CookedFormatDataOverride = &OriginalBodySetup->CookedFormatData;
+		BodySetup->CookedFormatProvider = SkeletalMeshConst;
+		BodySetup->bHasCookedFormatProvider = true;
 	}
 
 	BodySetup->CollisionTraceFlag = CTF_UseComplexAsSimple;
