ObjectTrace: add missing macros for builds with tracing disabled

#rnx

#Synced-CL 50157960

diff --git a/Engine/Source/Runtime/Engine/Public/ObjectTrace.h b/Engine/Source/Runtime/Engine/Public/ObjectTrace.h
--- a/Engine/Source/Runtime/Engine/Public/ObjectTrace.h
+++ b/Engine/Source/Runtime/Engine/Public/ObjectTrace.h
@@ -179,6 +179,9 @@ struct FObjectTrace
 #define TRACE_OBJECT_LIFETIME_BEGIN(Object) \
 	FObjectTrace::OutputObjectLifetimeBegin(Object);
 
+#define TRACE_OBJECT_LIFETIME_BEGIN_WITH_OUTER(Object, UserDefinedOuter) \
+	FObjectTrace::OutputObjectLifetimeBegin(Object, UserDefinedOuter);
+
 #define TRACE_OBJECT_LIFETIME_END(Object) \
 	FObjectTrace::OutputObjectLifetimeEnd(Object);
 
@@ -218,6 +221,7 @@ struct FObjectTrace
 #define TRACE_PAWN_POSSESS(...)
 #define TRACE_VIEW(...)
 #define TRACE_OBJECT_LIFETIME_BEGIN(...)
+#define TRACE_OBJECT_LIFETIME_BEGIN_WITH_OUTER(...)
 #define TRACE_OBJECT_LIFETIME_END(...)
 #define TRACE_INSTANCE_LIFETIME_END(...)
 #define TRACE_CHILD_ELEMENT(...)
