Allow components which aren't explicitly marked as editor-only to be part of an attachment hierarchy if their owning actor is marked editor-only.

#Synced-CL 50158515

diff --git a/Engine/Source/Runtime/Engine/Private/Components/SceneComponent.cpp b/Engine/Source/Runtime/Engine/Private/Components/SceneComponent.cpp
--- a/Engine/Source/Runtime/Engine/Private/Components/SceneComponent.cpp
+++ b/Engine/Source/Runtime/Engine/Private/Components/SceneComponent.cpp
@@ -607,7 +607,9 @@ EDataValidationResult USceneComponent::IsDataValid(FDataValidationContext& Conte
 	{
 		for (USceneComponent* ChildSceneComponent : GetAttachChildren())
 		{
-			if (ChildSceneComponent && !ChildSceneComponent->IsEditorOnly())
+			if (	ChildSceneComponent 
+				&& !ChildSceneComponent->IsEditorOnly() 
+				&& !(ChildSceneComponent->GetOwner() && ChildSceneComponent->GetOwner()->IsEditorOnly()))
 			{
 				Context.AddError(FText::Format(LOCTEXT("SceneComponent_AttachmentEditorOnlyMismatch",
 					"Component {0} is editor-only but it has an attached child {1} that is not"),
