Removed now deprecated parameter from `FInstancedSkinnedMeshComponentHelper::CreateSceneProxy` the MinLODIndex is now in the desc.

#Synced-CL 50152832

diff --git a/Engine/Source/Runtime/Engine/Public/InstancedSkinnedMeshComponentHelper.h b/Engine/Source/Runtime/Engine/Public/InstancedSkinnedMeshComponentHelper.h
--- a/Engine/Source/Runtime/Engine/Public/InstancedSkinnedMeshComponentHelper.h
+++ b/Engine/Source/Runtime/Engine/Public/InstancedSkinnedMeshComponentHelper.h
@@ -235,8 +235,7 @@ bool FInstancedSkinnedMeshComponentHelper::IsEnabled(const T& InComponent)
 template <class T>
 FPrimitiveSceneProxy* FInstancedSkinnedMeshComponentHelper::CreateSceneProxy(const T& InComponent, const FInstancedSkinnedMeshSceneProxyDesc& InSceneProxyDesc)
 {
-	const int32 MinLODIndex = FSkinnedMeshComponentHelper::ComputeMinLOD(InComponent);
 	const bool bShouldNaniteSkin = FSkinnedMeshComponentHelper::ShouldNaniteSkin(InComponent);
 	const bool bEnabled = FInstancedSkinnedMeshComponentHelper::IsEnabled(InComponent);
-	return FInstancedSkinnedMeshSceneProxyDesc::CreateSceneProxy(InSceneProxyDesc, InComponent.bHideSkin, bShouldNaniteSkin, bEnabled, MinLODIndex);
+	return FInstancedSkinnedMeshSceneProxyDesc::CreateSceneProxy(InSceneProxyDesc, InComponent.bHideSkin, bShouldNaniteSkin, bEnabled);
 }
