Clamp the LODIndex fed to MeshObject->Update against the streaming min LOD. Avoids attempting to render from a different LOD than was prepared by the GPU skinning system.

#Synced-CL 50057943

diff --git a/Engine/Plugins/Experimental/FastGeoStreaming/Source/FastGeoStreaming/Private/FastGeoSkinnedMeshComponent.cpp b/Engine/Plugins/Experimental/FastGeoStreaming/Source/FastGeoStreaming/Private/FastGeoSkinnedMeshComponent.cpp
--- a/Engine/Plugins/Experimental/FastGeoStreaming/Source/FastGeoStreaming/Private/FastGeoSkinnedMeshComponent.cpp
+++ b/Engine/Plugins/Experimental/FastGeoStreaming/Source/FastGeoStreaming/Private/FastGeoSkinnedMeshComponent.cpp
@@ -120,6 +120,7 @@ void FFastGeoSkinnedMeshComponentBase::InitializeSceneProxyDescDynamicProperties
 	SceneProxyDesc.SkinCacheUsage = SkinCacheUsage;
 	SceneProxyDesc.MaterialRelevance = GetMaterialRelevance(GetScene()->GetShaderPlatform());
 	SceneProxyDesc.MeshObject = MeshObject;
+	SceneProxyDesc.MinLODLevel = FSkinnedMeshComponentHelper::ComputeMinLOD(*this);
 }
 
 FPrimitiveSceneProxy* FFastGeoSkinnedMeshComponentBase::CreateSceneProxy(ESceneProxyCreationError* OutError)
@@ -274,8 +275,7 @@ FFastGeoSkinnedMeshComponent::FFastGeoSkinnedMeshComponent(int32 InComponentInde
 
 FPrimitiveSceneProxy* FFastGeoSkinnedMeshComponent::AllocateSceneProxy()
 {
-	int32 MinLODIndex = FSkinnedMeshComponentHelper::ComputeMinLOD(*this);
-	return FSkinnedMeshSceneProxyDesc::CreateSceneProxy(SceneProxyDesc, bHideSkin, MinLODIndex);
+	return FSkinnedMeshSceneProxyDesc::CreateSceneProxy(SceneProxyDesc, bHideSkin);
 }
 
 void FFastGeoSkinnedMeshComponent::ApplyWorldTransform(const FTransform& InTransform)
diff --git a/Engine/Source/Runtime/Engine/Private/Components/InstancedSkinnedMeshComponent.cpp b/Engine/Source/Runtime/Engine/Private/Components/InstancedSkinnedMeshComponent.cpp
--- a/Engine/Source/Runtime/Engine/Private/Components/InstancedSkinnedMeshComponent.cpp
+++ b/Engine/Source/Runtime/Engine/Private/Components/InstancedSkinnedMeshComponent.cpp
@@ -517,7 +517,7 @@ FPrimitiveSceneProxy* UInstancedSkinnedMeshComponent::CreateSceneProxy()
 
 	GetOrCreateInstanceDataSceneProxy();
 
-	Result = FInstancedSkinnedMeshSceneProxyDesc::CreateSceneProxy(FInstancedSkinnedMeshSceneProxyDesc(this), bHideSkin, ShouldNaniteSkin(), IsEnabled(), ComputeMinLOD());
+	Result = FInstancedSkinnedMeshSceneProxyDesc::CreateSceneProxy(FInstancedSkinnedMeshSceneProxyDesc(this), bHideSkin, ShouldNaniteSkin(), IsEnabled());
 
 	// Unclear exactly how this is supposed to work with a non-instanced proxy - will be interesting...
 	// If GPU-only flag set, instance data is entirely GPU driven, don't upload from CPU.
diff --git a/Engine/Source/Runtime/Engine/Private/Components/SkinnedMeshComponent.cpp b/Engine/Source/Runtime/Engine/Private/Components/SkinnedMeshComponent.cpp
--- a/Engine/Source/Runtime/Engine/Private/Components/SkinnedMeshComponent.cpp
+++ b/Engine/Source/Runtime/Engine/Private/Components/SkinnedMeshComponent.cpp
@@ -1272,12 +1272,9 @@ void USkinnedMeshComponent::CreateRenderState_Concurrent(FRegisterComponentConte
 			//	without animated, causing random skinning issues
 			// This can happen if your MinLOD is not valid anymore after loading
 			// which causes meshes to be invisible
-			UseLOD = GetPredictedLODLevel();
-			{
-				int32 MinLodIndex = ComputeMinLOD();
-				int32 MaxLODIndex = MeshObject->GetSkeletalMeshRenderData().LODRenderData.Num() - 1;
-				UseLOD = FMath::Clamp(UseLOD, MinLodIndex, MaxLODIndex);
-			}
+			const int32 MinLodIndex = ComputeMinLOD();
+			const int32 MaxLODIndex = MeshObject->GetSkeletalMeshRenderData().LODRenderData.Num() - 1;
+			UseLOD = FMath::Clamp(GetPredictedLODLevel(), MinLodIndex, MaxLODIndex);
 
 			// Clamp to loaded streaming data if available
 			if ((GetSkinnedAsset()->IsStreamable() || !IStreamingManager::Get().IsRenderAssetStreamingEnabled(EStreamableRenderAssetType::SkeletalMesh)) && MeshObject)
@@ -1417,7 +1414,10 @@ void USkinnedMeshComponent::SendRenderDynamicData_Concurrent()
 	{
 		SCOPE_CYCLE_COUNTER(STAT_MeshObjectUpdate);
 
-		int32 UseLOD = GetPredictedLODLevel();
+		const int32 MinLodIndex = ComputeMinLOD();
+		const int32 MaxLODIndex = MeshObject->GetSkeletalMeshRenderData().LODRenderData.Num() - 1;
+		int32 UseLOD = FMath::Clamp(GetPredictedLODLevel(), MinLodIndex, MaxLODIndex);
+
 		// Clamp to loaded streaming data if available
 		if ((GetSkinnedAsset()->IsStreamable() || !IStreamingManager::Get().IsRenderAssetStreamingEnabled(EStreamableRenderAssetType::SkeletalMesh)) && MeshObject)
 		{
@@ -4386,12 +4386,14 @@ int32 USkinnedMeshComponent::GetNumLODs() const
 void USkinnedMeshComponent::SetMinLOD(int32 InNewMinLOD)
 {
 	MinLodModel = GetValidMinLOD(InNewMinLOD);
+	MarkRenderStateDirty();
 }
 
 void USkinnedMeshComponent::OverrideMinLOD(int32 InNewMinLOD)
 {
 	MinLodModel = GetValidMinLOD(InNewMinLOD);
 	bOverrideMinLod = true;
+	MarkRenderStateDirty();
 }
 
 int32 USkinnedMeshComponent::ComputeMinLOD() const
diff --git a/Engine/Source/Runtime/Engine/Private/InstancedSkinnedMeshSceneProxy.cpp b/Engine/Source/Runtime/Engine/Private/InstancedSkinnedMeshSceneProxy.cpp
--- a/Engine/Source/Runtime/Engine/Private/InstancedSkinnedMeshSceneProxy.cpp
+++ b/Engine/Source/Runtime/Engine/Private/InstancedSkinnedMeshSceneProxy.cpp
@@ -60,7 +60,7 @@ FNaniteInstancedSkinnedMeshSceneProxy::FNaniteInstancedSkinnedMeshSceneProxy(con
 }
 
 FInstancedSkinnedMeshSceneProxy::FInstancedSkinnedMeshSceneProxy(const FInstancedSkinnedMeshSceneProxyDesc& InMeshDesc, FSkeletalMeshRenderData* InRenderData)
-	: Super(InMeshDesc, InRenderData, 0)
+	: Super(InMeshDesc, InRenderData)
 	, Data(InMeshDesc)
 {
 #if WITH_EDITOR
diff --git a/Engine/Source/Runtime/Engine/Private/InstancedSkinnedMeshSceneProxyDesc.cpp b/Engine/Source/Runtime/Engine/Private/InstancedSkinnedMeshSceneProxyDesc.cpp
--- a/Engine/Source/Runtime/Engine/Private/InstancedSkinnedMeshSceneProxyDesc.cpp
+++ b/Engine/Source/Runtime/Engine/Private/InstancedSkinnedMeshSceneProxyDesc.cpp
@@ -20,7 +20,7 @@ FSkeletalMeshObject* FInstancedSkinnedMeshSceneProxyDesc::CreateMeshObject(const
 	return nullptr;
 }
 
-FPrimitiveSceneProxy* FInstancedSkinnedMeshSceneProxyDesc::CreateSceneProxy(const FInstancedSkinnedMeshSceneProxyDesc& Desc, bool bHideSkin, bool bShouldNaniteSkin, bool bIsEnabled, int32 MinLODIndex)
+FPrimitiveSceneProxy* FInstancedSkinnedMeshSceneProxyDesc::CreateSceneProxy(const FInstancedSkinnedMeshSceneProxyDesc& Desc, bool bHideSkin, bool bShouldNaniteSkin, bool bIsEnabled)
 {
 	LLM_SCOPE(ELLMTag::SkeletalMesh);
 	FPrimitiveSceneProxy* Result = nullptr;
@@ -35,7 +35,7 @@ FPrimitiveSceneProxy* FInstancedSkinnedMeshSceneProxyDesc::CreateSceneProxy(cons
 		MeshObject)
 	{
 		// Only create a scene proxy if the bone count being used is supported, or if we don't have a skeleton (this is the case with destructibles)
-		int32 MaxBonesPerChunk = SkelMeshRenderData->GetMaxBonesPerSection(MinLODIndex);
+		int32 MaxBonesPerChunk = SkelMeshRenderData->GetMaxBonesPerSection(Desc.MinLODLevel);
 		int32 MaxSupportedNumBones = MeshObject->IsCPUSkinned() ? MAX_int32 : FGPUBaseSkinVertexFactory::GetMaxGPUSkinBones();
 		if (MaxBonesPerChunk <= MaxSupportedNumBones)
 		{
@@ -62,7 +62,7 @@ FPrimitiveSceneProxy* FInstancedSkinnedMeshSceneProxyDesc::CreateSceneProxy(cons
 
 			if (Result == nullptr)
 			{
-				Result = FSkinnedMeshSceneProxyDesc::CreateSceneProxy(Desc, bHideSkin, MinLODIndex);
+				Result = FSkinnedMeshSceneProxyDesc::CreateSceneProxy(Desc, bHideSkin);
 			}
 		}
 	}
diff --git a/Engine/Source/Runtime/Engine/Private/SkeletalMeshSceneProxy.cpp b/Engine/Source/Runtime/Engine/Private/SkeletalMeshSceneProxy.cpp
--- a/Engine/Source/Runtime/Engine/Private/SkeletalMeshSceneProxy.cpp
+++ b/Engine/Source/Runtime/Engine/Private/SkeletalMeshSceneProxy.cpp
@@ -86,21 +86,21 @@ bool AllowLumenCardGenerationForSkeletalMeshes(EShaderPlatform Platform)
  * @param	Component - skeletal mesh primitive being added
  */
 FSkeletalMeshSceneProxy::FSkeletalMeshSceneProxy(const USkinnedMeshComponent* Component, FSkeletalMeshRenderData* InRenderData)
-	: FSkeletalMeshSceneProxy(FSkinnedMeshSceneProxyDesc(Component), InRenderData, Component->GetValidMinLOD(Component->GetSkinnedAsset()->GetMinLodIdx()))
+	: FSkeletalMeshSceneProxy(FSkinnedMeshSceneProxyDesc(Component), InRenderData)
 { }
 
-FSkeletalMeshSceneProxy::FSkeletalMeshSceneProxy(const FSkinnedMeshSceneProxyDesc& InMeshDesc, FSkeletalMeshRenderData* InRenderData)
-	: FSkeletalMeshSceneProxy(InMeshDesc, InRenderData, 0)
+FSkeletalMeshSceneProxy::FSkeletalMeshSceneProxy(const FSkinnedMeshSceneProxyDesc& InMeshDesc, FSkeletalMeshRenderData* InRenderData, int32 InMinLODLevel)
+	: FSkeletalMeshSceneProxy(InMeshDesc, InRenderData)
 { }
 
-FSkeletalMeshSceneProxy::FSkeletalMeshSceneProxy(const FSkinnedMeshSceneProxyDesc& InMeshDesc, FSkeletalMeshRenderData* InRenderData, int32 InClampedLODIndex)
+FSkeletalMeshSceneProxy::FSkeletalMeshSceneProxy(const FSkinnedMeshSceneProxyDesc& InMeshDesc, FSkeletalMeshRenderData* InRenderData)
 		:	FPrimitiveSceneProxy(InMeshDesc, InMeshDesc.GetSkinnedAsset()->GetFName())
 		,	Owner(Cast<AActor>(InMeshDesc.GetOwner()))
 		,	MeshObject(InMeshDesc.MeshObject)
 		,	SkinnedAsset(InMeshDesc.GetSkinnedAsset())
 		,	SceneExtensionProxy(MeshObject->CreateSceneExtensionProxy(InMeshDesc.GetSkinnedAsset(), true))
 		,	RenderData(InRenderData)
-		,	ClampedLODIndex(InClampedLODIndex)
+		,	MinLODLevel(InMeshDesc.MinLODLevel)
 #if !(UE_BUILD_SHIPPING || UE_BUILD_TEST)
 		,	PhysicsAssetForDebug(InMeshDesc.GetPhysicsAsset())
 #endif
@@ -181,7 +181,7 @@ FSkeletalMeshSceneProxy::FSkeletalMeshSceneProxy(const FSkinnedMeshSceneProxyDes
 	bool bAnySectionHasMeshDeformerUsage = false;
 	LODSections.Reserve(RenderData->LODRenderData.Num());
 	LODSections.AddZeroed(RenderData->LODRenderData.Num());
-	for(int32 LODIdx=0; LODIdx < RenderData->LODRenderData.Num(); LODIdx++)
+	for(int32 LODIdx=MinLODLevel; LODIdx < RenderData->LODRenderData.Num(); LODIdx++)
 	{
 		const FSkeletalMeshLODRenderData& LODData = RenderData->LODRenderData[LODIdx];
 		const FSkeletalMeshLODInfo& Info = *(InMeshDesc.GetSkinnedAsset()->GetLODInfo(LODIdx));
@@ -466,7 +466,7 @@ HHitProxy* FSkeletalMeshSceneProxy::CreateHitProxies(UPrimitiveComponent* Compon
 	{
 		if ( LODSections.Num() > 0 )
 		{
-			for ( int32 LODIndex = 0; LODIndex < RenderData->LODRenderData.Num(); LODIndex++ )
+			for ( int32 LODIndex = MinLODLevel; LODIndex < RenderData->LODRenderData.Num(); LODIndex++ )
 			{
 				const FSkeletalMeshLODRenderData& LODData = RenderData->LODRenderData[LODIndex];
 
@@ -516,7 +516,7 @@ void FSkeletalMeshSceneProxy::DrawStaticElements(FStaticPrimitiveDrawInterface*
 	{
 		ESceneDepthPriorityGroup PrimitiveDPG = GetStaticDepthPriorityGroup();
 
-		for (int32 LODIndex = ClampedLODIndex; LODIndex < NumLODs; ++LODIndex)
+		for (int32 LODIndex = MinLODLevel; LODIndex < NumLODs; ++LODIndex)
 		{
 			const FSkeletalMeshLODRenderData& LODData = RenderData->LODRenderData[LODIndex];
 
@@ -819,7 +819,9 @@ uint8 FSkeletalMeshSceneProxy::GetCurrentFirstLODIdx_Internal() const
 
 FDesiredLODLevel FSkeletalMeshSceneProxy::GetDesiredLODLevel_RenderThread(const FSceneView* View) const
 {
-	return FDesiredLODLevel::CreateFixed(FMath::Max(ClampedLODIndex, MeshObject->GetLOD()));
+	const int32 LODIndex = MeshObject->GetLOD();
+	check(LODIndex >= MinLODLevel);
+	return FDesiredLODLevel::CreateFixed(LODIndex);
 }
 
 bool FSkeletalMeshSceneProxy::GetCachedGeometry(FRDGBuilder& GraphBuilder, FCachedGeometry& OutCachedGeometry) const 
@@ -1002,7 +1004,7 @@ TArray<FRayTracingGeometry*> FSkeletalMeshSceneProxy::GetStaticRayTracingGeometr
 	{
 		TArray<FRayTracingGeometry*> RayTracingGeometries;
 		RayTracingGeometries.AddDefaulted(RenderData->LODRenderData.Num());
-		for (int32 LODIndex = 0; LODIndex < RenderData->LODRenderData.Num(); LODIndex++)
+		for (int32 LODIndex = MinLODLevel; LODIndex < RenderData->LODRenderData.Num(); LODIndex++)
 		{
 			FSkeletalMeshLODRenderData& LODRenderData = RenderData->LODRenderData[LODIndex];
 
@@ -1027,7 +1029,7 @@ void FSkeletalMeshSceneProxy::GetStaticRayTracingInstances(FStaticRayTracingInst
 
 	OutRayTracingInstances.LODs.SetNum(NumLODs);
 
-	for (int32 LODIndex = ClampedLODIndex; LODIndex < NumLODs; ++LODIndex)
+	for (int32 LODIndex = MinLODLevel; LODIndex < NumLODs; ++LODIndex)
 	{
 		const FSkeletalMeshLODRenderData& LODData = RenderData->LODRenderData[LODIndex];
 
diff --git a/Engine/Source/Runtime/Engine/Private/SkeletalRenderStatic.cpp b/Engine/Source/Runtime/Engine/Private/SkeletalRenderStatic.cpp
--- a/Engine/Source/Runtime/Engine/Private/SkeletalRenderStatic.cpp
+++ b/Engine/Source/Runtime/Engine/Private/SkeletalRenderStatic.cpp
@@ -19,6 +19,7 @@ FSkeletalMeshObjectStatic::FSkeletalMeshObjectStatic(const USkinnedMeshComponent
 
 FSkeletalMeshObjectStatic::FSkeletalMeshObjectStatic(const FSkinnedMeshSceneProxyDesc& InMeshDesc, FSkeletalMeshRenderData* InSkelMeshRenderData, ERHIFeatureLevel::Type InFeatureLevel)
 	: FSkeletalMeshObject(InMeshDesc, InSkelMeshRenderData, InFeatureLevel)
+	, MinLODLevel(InMeshDesc.MinLODLevel)
 {
 	// create LODs to match the base mesh
 	for (int32 LODIndex = 0; LODIndex < InSkelMeshRenderData->LODRenderData.Num(); LODIndex++)
@@ -91,8 +92,7 @@ TConstArrayView<FMatrix44f> FSkeletalMeshObjectStatic::GetReferenceToLocalMatric
 
 int32 FSkeletalMeshObjectStatic::GetLOD() const
 {
-	// WorkingMinDesiredLODLevel can be a LOD that's not loaded, so need to clamp it to the first loaded LOD
-	return FMath::Max<int32>(WorkingMinDesiredLODLevel, SkeletalMeshRenderData->CurrentFirstLODIdx);
+	return FMath::Max<int32>(WorkingMinDesiredLODLevel, MinLODLevel);
 }
 
 void FSkeletalMeshObjectStatic::FSkeletalMeshObjectLOD::InitResources(const FSkelMeshComponentLODInfo* CompLODInfo)
diff --git a/Engine/Source/Runtime/Engine/Private/SkeletalRenderStatic.h b/Engine/Source/Runtime/Engine/Private/SkeletalRenderStatic.h
--- a/Engine/Source/Runtime/Engine/Private/SkeletalRenderStatic.h
+++ b/Engine/Source/Runtime/Engine/Private/SkeletalRenderStatic.h
@@ -108,5 +108,6 @@ private:
 
 	/** Render data for each LOD */
 	TArray<FSkeletalMeshObjectLOD> LODs;
+	int32 MinLODLevel = 0;
 };
 
diff --git a/Engine/Source/Runtime/Engine/Private/SkinnedMeshSceneProxyDesc.cpp b/Engine/Source/Runtime/Engine/Private/SkinnedMeshSceneProxyDesc.cpp
--- a/Engine/Source/Runtime/Engine/Private/SkinnedMeshSceneProxyDesc.cpp
+++ b/Engine/Source/Runtime/Engine/Private/SkinnedMeshSceneProxyDesc.cpp
@@ -57,6 +57,7 @@ void FSkinnedMeshSceneProxyDesc::InitializeFromSkinnedMeshComponent(const USkinn
 	CapsuleIndirectShadowMinVisibility = InComponent->CapsuleIndirectShadowMinVisibility;
 	OverlayMaterialMaxDrawDistance = InComponent->GetOverlayMaterialMaxDrawDistance();
 	PredictedLODLevel = InComponent->GetPredictedLODLevel();
+	MinLODLevel = InComponent->ComputeMinLOD();
 	MaxDistanceFactor = InComponent->MaxDistanceFactor;
 	ComponentScale = InComponent->GetComponentScale();
 	LODInfo = InComponent->LODInfo;
@@ -271,7 +272,7 @@ FSkeletalMeshObject* FSkinnedMeshSceneProxyDesc::CreateMeshObject(const FSkinned
 	return nullptr;
 }
 
-FPrimitiveSceneProxy* FSkinnedMeshSceneProxyDesc::CreateSceneProxy(const FSkinnedMeshSceneProxyDesc& Desc, bool bHideSkin, int32 MinLODIndex)
+FPrimitiveSceneProxy* FSkinnedMeshSceneProxyDesc::CreateSceneProxy(const FSkinnedMeshSceneProxyDesc& Desc, bool bHideSkin)
 {
 	LLM_SCOPE(ELLMTag::SkeletalMesh);
 	ERHIFeatureLevel::Type SceneFeatureLevel = Desc.FeatureLevel;
@@ -287,7 +288,7 @@ FPrimitiveSceneProxy* FSkinnedMeshSceneProxyDesc::CreateSceneProxy(const FSkinne
 		MeshObject)
 	{
 		// Only create a scene proxy if the bone count being used is supported, or if we don't have a skeleton (this is the case with destructibles)
-		int32 MaxBonesPerChunk = SkelMeshRenderData->GetMaxBonesPerSection(MinLODIndex);
+		int32 MaxBonesPerChunk = SkelMeshRenderData->GetMaxBonesPerSection(Desc.MinLODLevel);
 		int32 MaxSupportedNumBones = MeshObject->IsCPUSkinned() ? MAX_int32 : FGPUBaseSkinVertexFactory::GetMaxGPUSkinBones();
 		if (MaxBonesPerChunk <= MaxSupportedNumBones)
 		{
@@ -298,8 +299,7 @@ FPrimitiveSceneProxy* FSkinnedMeshSceneProxyDesc::CreateSceneProxy(const FSkinne
 			}
 			else
 			{
-				// TODO: MinLODIndex could work for InClampedLODIndex if it's the true minimum. Otherwise we could pass another parameter. Using 0 for now to ensure meshes render.
-				Result = ::new FSkeletalMeshSceneProxy(Desc, SkelMeshRenderData, 0);
+				Result = ::new FSkeletalMeshSceneProxy(Desc, SkelMeshRenderData);
 			}
 		}
 	}
diff --git a/Engine/Source/Runtime/Engine/Public/InstancedSkinnedMeshSceneProxyDesc.h b/Engine/Source/Runtime/Engine/Public/InstancedSkinnedMeshSceneProxyDesc.h
--- a/Engine/Source/Runtime/Engine/Public/InstancedSkinnedMeshSceneProxyDesc.h
+++ b/Engine/Source/Runtime/Engine/Public/InstancedSkinnedMeshSceneProxyDesc.h
@@ -10,7 +10,15 @@ class UInstancedSkinnedMeshComponent;
 struct FInstancedSkinnedMeshSceneProxyDesc : public FSkinnedMeshSceneProxyDesc
 {
 	ENGINE_API static FSkeletalMeshObject* CreateMeshObject(const FInstancedSkinnedMeshSceneProxyDesc& Desc, FSkeletalMeshRenderData* InRenderData, ERHIFeatureLevel::Type InFeatureLevel);
-	ENGINE_API static FPrimitiveSceneProxy* CreateSceneProxy(const FInstancedSkinnedMeshSceneProxyDesc& Desc, bool bHideSkin, bool bShouldNaniteSkin, bool bIsEnabled, int32 MinLODIndex);
+	ENGINE_API static FPrimitiveSceneProxy* CreateSceneProxy(const FInstancedSkinnedMeshSceneProxyDesc& Desc, bool bHideSkin, bool bShouldNaniteSkin, bool bIsEnabled);
+
+	UE_DEPRECATED(5.8, "MeshDesc now has MinLODLevel and InClampedLODIndex is unused")
+	static FPrimitiveSceneProxy* CreateSceneProxy(const FInstancedSkinnedMeshSceneProxyDesc& Desc, bool bHideSkin, bool bShouldNaniteSkin, bool bIsEnabled, int32 InClampedLODIndex)
+	{
+		FInstancedSkinnedMeshSceneProxyDesc CopyDesc = Desc;
+		CopyDesc.MinLODLevel = InClampedLODIndex;
+		return CreateSceneProxy(CopyDesc, bHideSkin, bShouldNaniteSkin, bIsEnabled);
+	}
 
 	ENGINE_API FInstancedSkinnedMeshSceneProxyDesc(const UInstancedSkinnedMeshComponent* InComponent);
 	FInstancedSkinnedMeshSceneProxyDesc() = default;
diff --git a/Engine/Source/Runtime/Engine/Public/SkeletalMeshSceneProxy.h b/Engine/Source/Runtime/Engine/Public/SkeletalMeshSceneProxy.h
--- a/Engine/Source/Runtime/Engine/Public/SkeletalMeshSceneProxy.h
+++ b/Engine/Source/Runtime/Engine/Public/SkeletalMeshSceneProxy.h
@@ -29,8 +29,9 @@ public:
 	 * @param	Component - skeletal mesh primitive being added
 	 */
 	ENGINE_API FSkeletalMeshSceneProxy(const USkinnedMeshComponent* Component, FSkeletalMeshRenderData* InSkelMeshRenderData);
-	UE_DEPRECATED(5.7, "Pass InClampedLODIndex to FSkeletalMeshSceneProxy constructor.")
 	ENGINE_API FSkeletalMeshSceneProxy(const FSkinnedMeshSceneProxyDesc& InMeshDesc, FSkeletalMeshRenderData* InSkelMeshRenderData);
+
+	UE_DEPRECATED(5.8, "MeshDesc now has MinLODLevel and InClampedLODIndex is unused")
 	ENGINE_API FSkeletalMeshSceneProxy(const FSkinnedMeshSceneProxyDesc& InMeshDesc, FSkeletalMeshRenderData* InSkelMeshRenderData, int32 InClampedLODIndex);
 
 	ENGINE_API virtual ~FSkeletalMeshSceneProxy();
@@ -157,7 +158,7 @@ protected:
 	USkinnedAsset* SkinnedAsset;
 	TUniquePtr<FSkinningSceneExtensionProxy> SceneExtensionProxy;
 	FSkeletalMeshRenderData* RenderData;
-	int32 ClampedLODIndex;
+	int32 MinLODLevel;
 	uint8 NumLODs = 1u;
 
 #if !(UE_BUILD_SHIPPING || UE_BUILD_TEST)
diff --git a/Engine/Source/Runtime/Engine/Public/SkinnedMeshSceneProxyDesc.h b/Engine/Source/Runtime/Engine/Public/SkinnedMeshSceneProxyDesc.h
--- a/Engine/Source/Runtime/Engine/Public/SkinnedMeshSceneProxyDesc.h
+++ b/Engine/Source/Runtime/Engine/Public/SkinnedMeshSceneProxyDesc.h
@@ -20,7 +20,15 @@ struct FMeshDeformerInstanceSet;
 struct FSkinnedMeshSceneProxyDesc : public FPrimitiveSceneProxyDesc
 {
 	ENGINE_API static FSkeletalMeshObject* CreateMeshObject(const FSkinnedMeshSceneProxyDesc& Desc);
-	ENGINE_API static FPrimitiveSceneProxy* CreateSceneProxy(const FSkinnedMeshSceneProxyDesc& Desc, bool bHideSkin, int32 MinLODIndex);
+	ENGINE_API static FPrimitiveSceneProxy* CreateSceneProxy(const FSkinnedMeshSceneProxyDesc& Desc, bool bHideSkin);
+
+	UE_DEPRECATED(5.8, "MeshDesc now has MinLODLevel and InClampedLODIndex is unused")
+	static FPrimitiveSceneProxy* CreateSceneProxy(const FSkinnedMeshSceneProxyDesc& Desc, bool bHideSkin, int32 InClampedLODIndex)
+	{
+		FSkinnedMeshSceneProxyDesc CopyDesc = Desc;
+		CopyDesc.MinLODLevel = InClampedLODIndex;
+		return CreateSceneProxy(CopyDesc, bHideSkin);
+	}
 
 	FSkinnedMeshSceneProxyDesc() = default;
 	ENGINE_API FSkinnedMeshSceneProxyDesc(const USkinnedMeshComponent* Component);
@@ -43,6 +51,7 @@ struct FSkinnedMeshSceneProxyDesc : public FPrimitiveSceneProxyDesc
 
 	FVector GetComponentScale() const { return ComponentScale; }
 	int32 GetPredictedLODLevel() const { return PredictedLODLevel; }
+	int32 GetMinLODLevel() const { return MinLODLevel; }
 	float GetMaxDistanceFactor() const { return MaxDistanceFactor; }
 
 	void SetLODInfo(TArrayView<const FSkelMeshComponentLODInfo> InLODInfo) { LODInfo = InLODInfo; }
@@ -73,6 +82,7 @@ struct FSkinnedMeshSceneProxyDesc : public FPrimitiveSceneProxyDesc
 	TArrayView<const TObjectPtr<UMaterialInterface>> OverrideMaterials;
 	TArrayView<const ESkinCacheUsage> SkinCacheUsage{};
 	int32 PredictedLODLevel = 0;
+	int32 MinLODLevel = 0;
 	float MaxDistanceFactor = 1.0f;
 
 	/* SkeletalMesh Archetype Data Begin */
diff --git a/Engine/Source/Runtime/Renderer/Private/StateStream/SkinnedMeshStateStreamImpl.cpp b/Engine/Source/Runtime/Renderer/Private/StateStream/SkinnedMeshStateStreamImpl.cpp
--- a/Engine/Source/Runtime/Renderer/Private/StateStream/SkinnedMeshStateStreamImpl.cpp
+++ b/Engine/Source/Runtime/Renderer/Private/StateStream/SkinnedMeshStateStreamImpl.cpp
@@ -166,7 +166,7 @@ void FSkinnedMeshStateStreamImpl::Render_OnCreate(const FSkinnedMeshStaticState&
 	//Desc.TextureStreamingTransformScale = Transform.GetMaximumAxisScale();
 	#endif
 
-	FPrimitiveSceneProxy* SceneProxy = FSkinnedMeshSceneProxyDesc::CreateSceneProxy(Desc, false, 0);
+	FPrimitiveSceneProxy* SceneProxy = FSkinnedMeshSceneProxyDesc::CreateSceneProxy(Desc, false);
 	Object.PrimitiveSceneData.SceneProxy = SceneProxy;
 
 	SceneProxy->SetPrimitiveColor(FLinearColor::White);
