Undo fix to ISKM on mobile due to BatchPrimitiveData.Data ending up with an index of 16777208.

#Synced-CL 50172231

#Synced-CL 50172231

#Synced-CL 50172231

diff --git a/Engine/Source/Runtime/Engine/Private/GPUSkinVertexFactory.cpp b/Engine/Source/Runtime/Engine/Private/GPUSkinVertexFactory.cpp
--- a/Engine/Source/Runtime/Engine/Private/GPUSkinVertexFactory.cpp
+++ b/Engine/Source/Runtime/Engine/Private/GPUSkinVertexFactory.cpp
@@ -842,11 +842,6 @@ void TGPUSkinVertexFactory<BoneInfluenceType>::ModifyCompilationEnvironment(cons
 
 	// Whether the material supports morph targets
 	OutEnvironment.SetDefine(TEXT("GPUSKIN_MORPH_BLEND"), Parameters.MaterialParameters.bIsUsedWithMorphTargets || Parameters.MaterialParameters.bIsSpecialEngineMaterial);
-
-	if (UseGPUScene(Parameters.Platform, MaxSupportedFeatureLevel))
-	{
-		OutEnvironment.SetDefine(TEXT("USE_INSTANCE_CULLING"), TEXT("1"));
-	}
 }
 
 /**
diff --git a/Engine/Source/Runtime/Engine/Private/SkeletalMeshSceneProxy.cpp b/Engine/Source/Runtime/Engine/Private/SkeletalMeshSceneProxy.cpp
--- a/Engine/Source/Runtime/Engine/Private/SkeletalMeshSceneProxy.cpp
+++ b/Engine/Source/Runtime/Engine/Private/SkeletalMeshSceneProxy.cpp
@@ -342,11 +342,6 @@ FSkeletalMeshSceneProxy::FSkeletalMeshSceneProxy(const FSkinnedMeshSceneProxyDes
 
 	EnableGPUSceneSupportFlags();
 
-	if (UseGPUScene(GetScene().GetShaderPlatform(), FeatureLevel))
-	{
-		bVFRequiresPrimitiveUniformBuffer = PlatformGPUSceneUsesUniformBufferView(GetScene().GetShaderPlatform());
-	}
-
 	if (IsAllowingApproximateOcclusionQueries())
 	{		
 		bAllowApproximateOcclusion = (bAllowApproximateOcclusion || bRenderStatic);
@@ -755,7 +750,6 @@ void FSkeletalMeshSceneProxy::GetStaticMeshBatch(
 	BatchElement.IndexBuffer = LODData.MultiSizeIndexContainer.GetIndexBuffer();
 	BatchElement.bFetchInstanceCountFromScene = bInstancedSkinnedMesh;
 	BatchElement.NumInstances = GetInstanceDataHeader().NumInstances;
-	BatchElement.bForceInstanceCulling = true;
 
 	if (bUseGpuLodSelection)
 	{
@@ -809,7 +803,7 @@ void FSkeletalMeshSceneProxy::CreateBaseMeshBatch(const FSceneView* View, const
 	BatchElement.NumInstances = GetInstanceDataHeader().NumInstances;
 	BatchElement.NumPrimitives = Section.NumTriangles;
 	BatchElement.PrimitiveUniformBuffer = GetUniformBuffer();
-	BatchElement.bForceInstanceCulling = true;
+	BatchElement.NumPrimitives = Section.NumTriangles;
 }
 
 uint8 FSkeletalMeshSceneProxy::GetCurrentFirstLODIdx_Internal() const
@@ -903,6 +897,7 @@ void FSkeletalMeshSceneProxy::GetDynamicElementsSection(const TArray<const FScen
 #endif // !(UE_BUILD_SHIPPING || UE_BUILD_TEST)
 #endif // WITH_EDITORONLY_DATA
 
+			BatchElement.MinVertexIndex = Section.BaseVertexIndex;
 			Mesh.ReverseCulling = IsLocalToWorldDeterminantNegative();
 			Mesh.CastShadow = SectionElementInfo.bEnableShadowCasting;
 			Mesh.bCanApplyViewModeOverrides = true;
@@ -950,7 +945,15 @@ void FSkeletalMeshSceneProxy::GetDynamicElementsSection(const TArray<const FScen
 				if (bAddOverlay)
 				{
 					FMeshBatch& OverlayMeshBatch = Collector.AllocateMesh();
-					GetOverlayMeshBatch(LODData, ActiveOverlayMaterial, Mesh, OverlayMeshBatch);
+					OverlayMeshBatch = Mesh;
+					OverlayMeshBatch.bOverlayMaterial = true;
+					OverlayMeshBatch.CastShadow = false;
+					OverlayMeshBatch.bSelectable = false;
+					OverlayMeshBatch.MaterialRenderProxy = ActiveOverlayMaterial->GetRenderProxy();
+					// make sure overlay is always rendered on top of base mesh
+					OverlayMeshBatch.MeshIdInPrimitive += LODData.RenderSections.Num();
+					OverlayMeshBatch.Elements[0].bFetchInstanceCountFromScene = bInstancedSkinnedMesh;
+					OverlayMeshBatch.Elements[0].NumInstances = GetInstanceDataHeader().NumInstances;
 					Collector.AddMesh(ViewIndex, OverlayMeshBatch);
 				
 					INC_DWORD_STAT_BY(STAT_SkelMeshTriangles, OverlayMeshBatch.GetNumPrimitives());
