World Partition - HLOD: Make sure that HLOD partitions that contains destructible HLODs are not flagged as only client visible
* HLOD destruction requires HLOD actors data to be replicated, so they must be loaded on the server too.
* The bClientOnlyVisible option on HLOD partitions used to be visible to the users in the World Partition settings, but that setting was hidden in the UX refactor.
* This flag is now programatically set depending on the HLOD layers settings

#Synced-CL 50167156

diff --git a/Engine/Plugins/Editor/WorldPartitionHLODUtilities/Source/Private/WorldPartition/HLOD/Modifiers/HLODModifierMeshDestruction.cpp b/Engine/Plugins/Editor/WorldPartitionHLODUtilities/Source/Private/WorldPartition/HLOD/Modifiers/HLODModifierMeshDestruction.cpp
--- a/Engine/Plugins/Editor/WorldPartitionHLODUtilities/Source/Private/WorldPartition/HLOD/Modifiers/HLODModifierMeshDestruction.cpp
+++ b/Engine/Plugins/Editor/WorldPartitionHLODUtilities/Source/Private/WorldPartition/HLOD/Modifiers/HLODModifierMeshDestruction.cpp
@@ -295,3 +295,9 @@ void UWorldPartitionHLODModifierMeshDestruction::EndHLODBuild(TArray<UActorCompo
 
 	DestructionMeshMergeExtension.Reset();
 }
+
+bool UWorldPartitionHLODModifierMeshDestruction::MustBeLoadedOnServer() const
+{
+	// As HLOD destruction needs to be replicated, HLOD actors which supports destruction must be loaded on the server.
+	return true;
+}
\ No newline at end of file
diff --git a/Engine/Plugins/Editor/WorldPartitionHLODUtilities/Source/Public/WorldPartition/HLOD/Modifiers/HLODModifierMeshDestruction.h b/Engine/Plugins/Editor/WorldPartitionHLODUtilities/Source/Public/WorldPartition/HLOD/Modifiers/HLODModifierMeshDestruction.h
--- a/Engine/Plugins/Editor/WorldPartitionHLODUtilities/Source/Public/WorldPartition/HLOD/Modifiers/HLODModifierMeshDestruction.h
+++ b/Engine/Plugins/Editor/WorldPartitionHLODUtilities/Source/Public/WorldPartition/HLOD/Modifiers/HLODModifierMeshDestruction.h
@@ -20,6 +20,7 @@ class UWorldPartitionHLODModifierMeshDestruction : public UWorldPartitionHLODMod
 	UE_API virtual bool CanModifyHLOD(TSubclassOf<UHLODBuilder> InHLODBuilderClass) const override;
 	UE_API virtual void BeginHLODBuild(const FHLODBuildContext& InHLODBuildContext) override;
 	UE_API virtual void EndHLODBuild(TArray<UActorComponent*>&InOutComponents) override;
+	UE_API virtual bool MustBeLoadedOnServer() const override;
 
 private:
 	// Gathers a list of actors included in the merged mesh
diff --git a/Engine/Source/Runtime/Engine/Private/WorldPartition/HLOD/HLODLayer.cpp b/Engine/Source/Runtime/Engine/Private/WorldPartition/HLOD/HLODLayer.cpp
--- a/Engine/Source/Runtime/Engine/Private/WorldPartition/HLOD/HLODLayer.cpp
+++ b/Engine/Source/Runtime/Engine/Private/WorldPartition/HLOD/HLODLayer.cpp
@@ -264,6 +264,11 @@ void UHLODLayer::ComputeHLODHash(FHLODHashBuilder& InHashBuilder) const
 	}
 }
 
+bool UHLODLayer::MustBeLoadedOnServer() const
+{
+	return HLODModifierClass && HLODModifierClass->GetDefaultObject<UWorldPartitionHLODModifier>()->MustBeLoadedOnServer();
+}
+
 #endif // WITH_EDITOR
 
 #undef LOCTEXT_NAMESPACE
\ No newline at end of file
diff --git a/Engine/Source/Runtime/Engine/Private/WorldPartition/RuntimeHashSet/WorldPartitionRuntimeHashSet.cpp b/Engine/Source/Runtime/Engine/Private/WorldPartition/RuntimeHashSet/WorldPartitionRuntimeHashSet.cpp
--- a/Engine/Source/Runtime/Engine/Private/WorldPartition/RuntimeHashSet/WorldPartitionRuntimeHashSet.cpp
+++ b/Engine/Source/Runtime/Engine/Private/WorldPartition/RuntimeHashSet/WorldPartitionRuntimeHashSet.cpp
@@ -1152,6 +1152,17 @@ void UWorldPartitionRuntimeHashSet::FixupHLODSetup(FRuntimePartitionDesc& Runtim
 {
 	if (RuntimePartition.MainLayer)
 	{
+		// By default, HLOD partitions are only loaded on the clients.
+		// However, some HLOD layer configurations might requires them to be loaded on the server (for example, HLOD destruction
+		// requires HLOD actors data to be replicated).
+		auto HLODPartitionMustBeLoadedOnServer = [](const FRuntimePartitionHLODSetup& HLODSetup) -> bool
+		{
+			return Algo::AnyOf(HLODSetup.HLODLayers, [](const UHLODLayer* HLODLayer)
+			{
+				return HLODLayer->MustBeLoadedOnServer();
+			});			
+		};
+
 		TArray<TObjectPtr<const UHLODLayer>> PersistentHLODLayers;
 
 		for (int32 HLODSetupIndex = 0; HLODSetupIndex < RuntimePartition.HLODSetups.Num(); )
@@ -1180,6 +1191,9 @@ void UWorldPartitionRuntimeHashSet::FixupHLODSetup(FRuntimePartitionDesc& Runtim
 
 				// Make sure that HLOD partitions settings are updated to match the main layer settings
 				HLODSetup.PartitionLayer->UpdateHLODRuntimePartitionFrom(RuntimePartition.MainLayer);
+
+				// Only load HLOD on clients unless the HLOD setup requires them on the server too.
+				HLODSetup.PartitionLayer->bClientOnlyVisible = !HLODPartitionMustBeLoadedOnServer(HLODSetup);
 			}
 
 			HLODSetupIndex++;
@@ -1195,6 +1209,7 @@ void UWorldPartitionRuntimeHashSet::FixupHLODSetup(FRuntimePartitionDesc& Runtim
 			PersistentHLODSetup.RowDisplayName = "Non-Spatially Loaded HLODs";
 			PersistentHLODSetup.HLODLayers = MoveTemp(PersistentHLODLayers);
 			PersistentHLODSetup.bIsSpatiallyLoaded = false;
+			PersistentHLODSetup.PartitionLayer->bClientOnlyVisible = !HLODPartitionMustBeLoadedOnServer(PersistentHLODSetup);
 		}
 	}
 }
diff --git a/Engine/Source/Runtime/Engine/Public/WorldPartition/HLOD/HLODLayer.h b/Engine/Source/Runtime/Engine/Public/WorldPartition/HLOD/HLODLayer.h
--- a/Engine/Source/Runtime/Engine/Public/WorldPartition/HLOD/HLODLayer.h
+++ b/Engine/Source/Runtime/Engine/Public/WorldPartition/HLOD/HLODLayer.h
@@ -62,6 +62,7 @@ public:
 	UHLODLayer* GetLinkedLayer() const { return LinkedLayer; }
 	ENGINE_API bool DoesRequireWarmup() const;
 	ENGINE_API void ComputeHLODHash(FHLODHashBuilder& InHLODHashBuilder) const;
+	ENGINE_API bool MustBeLoadedOnServer() const;
 
 	UE_DEPRECATED(5.7, "Deprecated. These streaming grid properties are now specified in the partition's settings.")
 	ENGINE_API FName GetRuntimeGrid(uint32 InHLODLevel) const;
diff --git a/Engine/Source/Runtime/Engine/Public/WorldPartition/HLOD/HLODModifier.h b/Engine/Source/Runtime/Engine/Public/WorldPartition/HLOD/HLODModifier.h
--- a/Engine/Source/Runtime/Engine/Public/WorldPartition/HLOD/HLODModifier.h
+++ b/Engine/Source/Runtime/Engine/Public/WorldPartition/HLOD/HLODModifier.h
@@ -25,4 +25,6 @@ public:
 	
 	virtual void BeginHLODBuild(const FHLODBuildContext& InHLODBuildContext) PURE_VIRTUAL(UWorldPartitionHLODModifier::BeginHLODBuild, );
 	virtual void EndHLODBuild(TArray<UActorComponent*>& InOutComponents) PURE_VIRTUAL(UWorldPartitionHLODModifier::EndHLODBuild, );
+
+	virtual bool MustBeLoadedOnServer() const { return false; }
 };
