Resubmitted mobile GPU scene fixes for GPU skin.

#Synced-CL 50172231

diff --git a/Engine/Shaders/Private/GpuSkinVertexFactory.ush b/Engine/Shaders/Private/GpuSkinVertexFactory.ush
--- a/Engine/Shaders/Private/GpuSkinVertexFactory.ush
+++ b/Engine/Shaders/Private/GpuSkinVertexFactory.ush
@@ -298,7 +298,8 @@ FMaterialVertexParameters GetMaterialVertexParameters(
 {
 	FMaterialVertexParameters Result = MakeInitializedMaterialVertexParameters();
 	Result.SceneData = Intermediates.SceneData;
-	
+	Result.PrimitiveId = Intermediates.SceneData.PrimitiveId;
+
 	Result.WorldPosition = WorldPosition;
 	if (bIsPreviousFrame)
 	{
diff --git a/Engine/Shaders/Private/SceneDataMobileLoader.ush b/Engine/Shaders/Private/SceneDataMobileLoader.ush
--- a/Engine/Shaders/Private/SceneDataMobileLoader.ush
+++ b/Engine/Shaders/Private/SceneDataMobileLoader.ush
@@ -248,6 +248,17 @@ FInstanceSceneData LoadInstanceDataUBO(uint DrawInstanceId)
 	DataIdx += 1u;
 #endif
 
+	// SkinningData
+	{
+		InstanceData.SkinningData = 0;
+		const bool bHasSkinningData = (InstanceData.Flags & INSTANCE_SCENE_DATA_FLAG_HAS_SKINNING_DATA) != 0u;
+		if (bHasSkinningData)
+		{
+			InstanceData.SkinningData = asuint(LoadInstanceDataElementUBO(DrawInstanceId, DataIdx+0).x);
+		}
+		DataIdx += 1u;
+	}
+
 	// CustomDataCount, TODO: yzw
 	{
 		float4 Data = LoadInstanceDataElementUBO(DrawInstanceId, DataIdx+0);
@@ -267,6 +278,7 @@ float LoadInstanceCustomDataElementUBO(uint DrawInstanceId, uint CustomDataIdx,
 	#if ALLOW_STATIC_LIGHTING
 		CustomDataFloat4Offset += 1u; // one float4 reserved for pre-computed lighting
 	#endif
+	CustomDataFloat4Offset += 1u; // one float4 reserved for skinning data
 
 	const uint Float4Index = (CustomDataIdx >> 2u);
 	const uint CustomDataCount = asuint(LoadInstanceDataElementUBO(DrawInstanceId, CustomDataFloat4Offset).x);
diff --git a/Engine/Shaders/Private/SceneDataMobileWriter.ush b/Engine/Shaders/Private/SceneDataMobileWriter.ush
--- a/Engine/Shaders/Private/SceneDataMobileWriter.ush
+++ b/Engine/Shaders/Private/SceneDataMobileWriter.ush
@@ -204,6 +204,17 @@ void WriteInstanceDataUBO(RWStructuredBuffer<float4> Output, uint Offset, uint I
 		DataIdx += 1u;
 	}
 #endif
+
+	// SkinningData
+	{
+		if (Offsets.SkinningData != INVALID_INSTANCE_PAYLOAD_OFFSET)
+		{
+			uint SkinningDataValue = asuint(LoadInstancePayloadDataElement(Offsets.SkinningData)).y;
+			WriteDataElementUBO(Output, Offset, DataIdx + 0, float4(asfloat(SkinningDataValue), 0.f, 0.f, 0.f));
+		}
+		DataIdx += 1u;
+	}
+
 	// CustomDataCount, TODO: yzw
 	{
 		CustomDataCount = min(CustomDataCount, (BATCHED_INSTANCE_DATA_STRIDE_FLOAT4 - DataIdx - 1u) * 4u);
diff --git a/Engine/Source/Runtime/Engine/Private/GPUSkinVertexFactory.cpp b/Engine/Source/Runtime/Engine/Private/GPUSkinVertexFactory.cpp
--- a/Engine/Source/Runtime/Engine/Private/GPUSkinVertexFactory.cpp
+++ b/Engine/Source/Runtime/Engine/Private/GPUSkinVertexFactory.cpp
@@ -842,6 +842,11 @@ void TGPUSkinVertexFactory<BoneInfluenceType>::ModifyCompilationEnvironment(cons
 
 	// Whether the material supports morph targets
 	OutEnvironment.SetDefine(TEXT("GPUSKIN_MORPH_BLEND"), Parameters.MaterialParameters.bIsUsedWithMorphTargets || Parameters.MaterialParameters.bIsSpecialEngineMaterial);
+
+	if (UseGPUScene(Parameters.Platform, MaxSupportedFeatureLevel))
+	{
+		OutEnvironment.SetDefine(TEXT("USE_INSTANCE_CULLING"), TEXT("1"));
+	}
 }
 
 /**
diff --git a/Engine/Source/Runtime/Engine/Private/SkeletalMeshSceneProxy.cpp b/Engine/Source/Runtime/Engine/Private/SkeletalMeshSceneProxy.cpp
--- a/Engine/Source/Runtime/Engine/Private/SkeletalMeshSceneProxy.cpp
+++ b/Engine/Source/Runtime/Engine/Private/SkeletalMeshSceneProxy.cpp
@@ -342,6 +342,11 @@ FSkeletalMeshSceneProxy::FSkeletalMeshSceneProxy(const FSkinnedMeshSceneProxyDes
 
 	EnableGPUSceneSupportFlags();
 
+	if (UseGPUScene(GetScene().GetShaderPlatform(), FeatureLevel))
+	{
+		bVFRequiresPrimitiveUniformBuffer = PlatformGPUSceneUsesUniformBufferView(GetScene().GetShaderPlatform());
+	}
+
 	if (IsAllowingApproximateOcclusionQueries())
 	{		
 		bAllowApproximateOcclusion = (bAllowApproximateOcclusion || bRenderStatic);
@@ -750,6 +755,7 @@ void FSkeletalMeshSceneProxy::GetStaticMeshBatch(
 	BatchElement.IndexBuffer = LODData.MultiSizeIndexContainer.GetIndexBuffer();
 	BatchElement.bFetchInstanceCountFromScene = bInstancedSkinnedMesh;
 	BatchElement.NumInstances = GetInstanceDataHeader().NumInstances;
+	BatchElement.bForceInstanceCulling = true;
 
 	if (bUseGpuLodSelection)
 	{
@@ -803,7 +809,7 @@ void FSkeletalMeshSceneProxy::CreateBaseMeshBatch(const FSceneView* View, const
 	BatchElement.NumInstances = GetInstanceDataHeader().NumInstances;
 	BatchElement.NumPrimitives = Section.NumTriangles;
 	BatchElement.PrimitiveUniformBuffer = GetUniformBuffer();
-	BatchElement.NumPrimitives = Section.NumTriangles;
+	BatchElement.bForceInstanceCulling = true;
 }
 
 uint8 FSkeletalMeshSceneProxy::GetCurrentFirstLODIdx_Internal() const
@@ -897,7 +903,6 @@ void FSkeletalMeshSceneProxy::GetDynamicElementsSection(const TArray<const FScen
 #endif // !(UE_BUILD_SHIPPING || UE_BUILD_TEST)
 #endif // WITH_EDITORONLY_DATA
 
-			BatchElement.MinVertexIndex = Section.BaseVertexIndex;
 			Mesh.ReverseCulling = IsLocalToWorldDeterminantNegative();
 			Mesh.CastShadow = SectionElementInfo.bEnableShadowCasting;
 			Mesh.bCanApplyViewModeOverrides = true;
@@ -945,15 +950,7 @@ void FSkeletalMeshSceneProxy::GetDynamicElementsSection(const TArray<const FScen
 				if (bAddOverlay)
 				{
 					FMeshBatch& OverlayMeshBatch = Collector.AllocateMesh();
-					OverlayMeshBatch = Mesh;
-					OverlayMeshBatch.bOverlayMaterial = true;
-					OverlayMeshBatch.CastShadow = false;
-					OverlayMeshBatch.bSelectable = false;
-					OverlayMeshBatch.MaterialRenderProxy = ActiveOverlayMaterial->GetRenderProxy();
-					// make sure overlay is always rendered on top of base mesh
-					OverlayMeshBatch.MeshIdInPrimitive += LODData.RenderSections.Num();
-					OverlayMeshBatch.Elements[0].bFetchInstanceCountFromScene = bInstancedSkinnedMesh;
-					OverlayMeshBatch.Elements[0].NumInstances = GetInstanceDataHeader().NumInstances;
+					GetOverlayMeshBatch(LODData, ActiveOverlayMaterial, Mesh, OverlayMeshBatch);
 					Collector.AddMesh(ViewIndex, OverlayMeshBatch);
 				
 					INC_DWORD_STAT_BY(STAT_SkelMeshTriangles, OverlayMeshBatch.GetNumPrimitives());
