Make ShouldStripUnusedDefaultTextures thread safe for compatibility with async loading.

[RN] minor | Rendering Architecture

#Synced-CL 50036996

diff --git a/Engine/Source/Runtime/Engine/Private/Materials/MaterialInterface.cpp b/Engine/Source/Runtime/Engine/Private/Materials/MaterialInterface.cpp
--- a/Engine/Source/Runtime/Engine/Private/Materials/MaterialInterface.cpp
+++ b/Engine/Source/Runtime/Engine/Private/Materials/MaterialInterface.cpp
@@ -77,9 +77,10 @@ UEnum* UMaterialInterface::SamplerTypeEnum = nullptr;
 //////////////////////////////////////////////////////////////////////////
 
 #if WITH_EDITOR
-static TAutoConsoleVariable<bool> CVarStripUnusedDefaultTextures(
+static bool GShouldStripUnusedDefaultTextures = true;
+static FAutoConsoleVariableRef CVarStripUnusedDefaultTextures(
 	TEXT("r.Material.StripUnusedDefaultTextures"),
-	true,
+	GShouldStripUnusedDefaultTextures,
 	TEXT("If enabled, unused Default textures will be discarded at cook time.\n"),
 	ECVF_ReadOnly);
 #endif
@@ -2091,7 +2092,7 @@ void UMaterialInterface::UpdateMaterialRenderProxy(FMaterialRenderProxy& Proxy)
 #if WITH_EDITOR
 bool UMaterialInterface::ShouldStripUnusedDefaultTextures()
 {
-	return CVarStripUnusedDefaultTextures.GetValueOnGameThread();
+	return GShouldStripUnusedDefaultTextures;
 }
 
 TArray<const FMaterialResource*, TInlineAllocator<4>> UMaterialInterface::GetCookedCacheMaterialResources() const
@@ -2117,12 +2118,12 @@ void UMaterialInterface::StripUnusedDefaultTextures(const bool bInKeepProjectPri
 {
 	SCOPE_CYCLE_COUNTER(STAT_MaterialInterface_StripUnusedDefaultTextures);
 
-	if (!CVarStripUnusedDefaultTextures.GetValueOnGameThread())
+	if (!ShouldStripUnusedDefaultTextures())
 	{
 		static bool bWarnedOnce = false;
 		if (!bWarnedOnce)
 		{
-			UE_LOG(LogMaterial, Warning, TEXT("Calling StripUnusedDefaultTextures even though CVarStripUnusedDefaultTextures is false."));
+			UE_LOG(LogMaterial, Warning, TEXT("Calling StripUnusedDefaultTextures even though ShouldStripUnusedDefaultTextures returns false."));
 			bWarnedOnce = true;
 		}
 
