Dstm: Chaos vehicle support 2.3: Add guards in in NetworkPhysicsComponent to avoid duplicated callback registration/unregistration.

#Synced-CL 50118602

diff --git a/Engine/Source/Runtime/Engine/Private/PhysicsEngine/NetworkPhysicsComponent.cpp b/Engine/Source/Runtime/Engine/Private/PhysicsEngine/NetworkPhysicsComponent.cpp
--- a/Engine/Source/Runtime/Engine/Private/PhysicsEngine/NetworkPhysicsComponent.cpp
+++ b/Engine/Source/Runtime/Engine/Private/PhysicsEngine/NetworkPhysicsComponent.cpp
@@ -708,17 +708,25 @@ void UNetworkPhysicsComponent::PostMigrate(const struct FUObjectMigrationContext
 
 	if (MigrationContext.MigrationSide == EObjectMigrationSide::Send)
 	{
-		// Need to uninitialize the component during migration, because the component is not GC'ed until next frame
-		// Then the call back on AsyncPhysicsNetworkComponent would cause a crash
-		UninitializePhysicsReplication();
+		// Add extra check to avoid double call back unregistration 
+		if (NetworkPhysicsComponent_Internal)
+		{
+			// Need to uninitialize the component during migration, because the component is not GC'ed until next frame
+			// Then the call back on AsyncPhysicsNetworkComponent would cause a crash
+			UninitializePhysicsReplication();
+		}
 	}
 
 	if (MigrationContext.MigrationSide == EObjectMigrationSide::Receive)
 	{
-		// Need to reinitialize the component (excluding initialization of the base)
-		// The reason is ActorComponents do not get reinitialized after migrated to another server since all data on ActorComponent
-		// is serialized and migrated over. 
-		InitializePhysicsReplication();
+		// Add extra check to avoid double call back registration 
+		if (!NetworkPhysicsComponent_Internal)
+		{
+			// Need to reinitialize the component (excluding initialization of the base)
+			// The reason is ActorComponents do not get reinitialized after migrated to another server since all data on ActorComponent
+			// is serialized and migrated over. 
+			InitializePhysicsReplication();
+		}
 	}
 }
 #endif
@@ -2620,7 +2628,7 @@ void FAsyncNetworkPhysicsComponent::OnPostSolve_Internal()
 		return;
 	}
 	const bool bIsSolverResim = RigidsSolver->GetEvolution()->IsResimming();
-		const int32 PhysicsStep = RigidsSolver->GetCurrentFrame();
+	const int32 PhysicsStep = RigidsSolver->GetCurrentFrame();
 	const int32 NextPhysicsStep = PhysicsStep + 1;
 
 	const FNetworkPhysicsSettingsNetworkPhysicsComponent& ComponentSettings = GetComponentSettings();
