Fixing the issue where a child BuildingActor that is attached to a parent BuildingActor becomes detached when initially the child is not simulating physics and it replicates once, then the parent starts simulating physics and the child is welded to the parent. The next time the child replicates the calll to AActor::GatherCurrentMovement() goes into the IsSimulatingPhysics branch of code. After exiting the function in AActor::PreReplication, AttachmentReplication will think the AttachParent/AttachComponent has changed when they actually haven't. If the AttachmentReplication is sent to the client then the client/server become out of sync and the client detaches the child. My fix simply maintains the AttachParent/AttachComponent when simulating physics so it is not removed incorrectly on the client.
-added cvar to toggle this behavior on/off

#Synced-CL 50149474

diff --git a/Engine/Source/Runtime/Engine/Private/ActorReplication.cpp b/Engine/Source/Runtime/Engine/Private/ActorReplication.cpp
--- a/Engine/Source/Runtime/Engine/Private/ActorReplication.cpp
+++ b/Engine/Source/Runtime/Engine/Private/ActorReplication.cpp
@@ -40,6 +40,9 @@ namespace ActorReplication
 	static ENetRole SavedRole;
 }
 
+static bool GbSendPhysicsAttachRep = true;
+static FAutoConsoleVariableRef CVarSendPhysicsAttachRep(TEXT("Net.RepMovement.SendPhysicsAttachRep"), GbSendPhysicsAttachRep, TEXT("Whether to replicate attachment replication info when the actor is simulating physics"), ECVF_ReadOnly);
+
 using namespace UE::Net;
 
 float AActor::GetNetPriority(const FVector& ViewPos, const FVector& ViewDir, AActor* Viewer, AActor* ViewTarget, UActorChannel* InChannel, float Time, bool bLowBandwidth)
@@ -478,6 +481,30 @@ void AActor::GatherCurrentMovement()
 				UpdateReplicatePhysicsCondition();
 				bWasRepMovementModified = true;
 			}
+
+			if (GbSendPhysicsAttachRep)
+			{
+				if (RootComponent->GetAttachParent() != nullptr)
+				{
+					// Set the attachment info here to prevent it being detached. If the actor was previously not simulating physics
+					// but now it is and is attached/welded, the attachment should remain
+					AttachmentReplication.AttachParent = RootComponent->GetAttachParentActor();
+					if (AttachmentReplication.AttachParent != nullptr)
+					{
+						if (AttachmentReplication.AttachParent != nullptr)
+						{
+							AttachmentReplication.LocationOffset = RootComponent->GetRelativeLocation();
+							AttachmentReplication.RotationOffset = RootComponent->GetRelativeRotation();
+							AttachmentReplication.RelativeScale3D = RootComponent->GetRelativeScale3D();
+							AttachmentReplication.AttachComponent = RootComponent->GetAttachParent();
+							AttachmentReplication.AttachSocket = RootComponent->GetAttachSocketName();
+
+							// Technically, the values might have stayed the same, but we'll just assume they've changed.
+							bWasAttachmentModified = true;
+						}
+					}
+				}
+			}
 		}
 		else if (RootComponent != nullptr)
 		{
