Removing duplicate condition from previous ActorReplication check in

#Synced-CL 50149474

diff --git a/Engine/Source/Runtime/Engine/Private/ActorReplication.cpp b/Engine/Source/Runtime/Engine/Private/ActorReplication.cpp
--- a/Engine/Source/Runtime/Engine/Private/ActorReplication.cpp
+++ b/Engine/Source/Runtime/Engine/Private/ActorReplication.cpp
@@ -491,17 +491,14 @@ void AActor::GatherCurrentMovement()
 					AttachmentReplication.AttachParent = RootComponent->GetAttachParentActor();
 					if (AttachmentReplication.AttachParent != nullptr)
 					{
-						if (AttachmentReplication.AttachParent != nullptr)
-						{
-							AttachmentReplication.LocationOffset = RootComponent->GetRelativeLocation();
-							AttachmentReplication.RotationOffset = RootComponent->GetRelativeRotation();
-							AttachmentReplication.RelativeScale3D = RootComponent->GetRelativeScale3D();
-							AttachmentReplication.AttachComponent = RootComponent->GetAttachParent();
-							AttachmentReplication.AttachSocket = RootComponent->GetAttachSocketName();
-
-							// Technically, the values might have stayed the same, but we'll just assume they've changed.
-							bWasAttachmentModified = true;
-						}
+						AttachmentReplication.LocationOffset = RootComponent->GetRelativeLocation();
+						AttachmentReplication.RotationOffset = RootComponent->GetRelativeRotation();
+						AttachmentReplication.RelativeScale3D = RootComponent->GetRelativeScale3D();
+						AttachmentReplication.AttachComponent = RootComponent->GetAttachParent();
+						AttachmentReplication.AttachSocket = RootComponent->GetAttachSocketName();
+
+						// Technically, the values might have stayed the same, but we'll just assume they've changed.
+						bWasAttachmentModified = true;
 					}
 				}
 			}
