[WorldPartition] Fix possible usage of deleted pointer to ActorDesc

#tests Tested on proejct that crashed, tested LI overrides
#Synced-CL 50114392

diff --git a/Engine/Source/Runtime/Engine/Private/WorldPartition/ActorDescContainer.cpp b/Engine/Source/Runtime/Engine/Private/WorldPartition/ActorDescContainer.cpp
--- a/Engine/Source/Runtime/Engine/Private/WorldPartition/ActorDescContainer.cpp
+++ b/Engine/Source/Runtime/Engine/Private/WorldPartition/ActorDescContainer.cpp
@@ -386,6 +386,28 @@ bool UActorDescContainer::IsActorDescHandled(const AActor* InActor, bool bInUseL
 	return false;
 }
 
+void UActorDescContainer::ReassignActorDesc(TUniquePtr<FWorldPartitionActorDesc>& DescToUpdate, TFunctionRef<void(TUniquePtr<FWorldPartitionActorDesc>& ActorDesc)> UpdateAction)
+{
+	// 1 - Prepare users of the actor desc
+	OnActorDescPreUpdate(DescToUpdate.Get());
+	// 2 - Apply the transform, which is a function that depends on the context of the change, and which may have created a new ActorDesc and destroyed the old one, updating the passed-in unique ptr
+	UpdateAction(DescToUpdate);
+	// 3 - Update the pointer when needed, on all container instances, before proceeding with computing actor hashes, etc. because this may involve using other containers
+	OnActorDescUpdating(DescToUpdate.Get());
+	// 4 - Carry on processing the new actor desc
+	OnActorDescUpdated(DescToUpdate.Get());
+}
+
+void UActorDescContainer::UpdateActorDesc(FWorldPartitionActorDesc* ActorDesc, TFunctionRef<void(FWorldPartitionActorDesc* ActorDesc)> UpdateAction)
+{
+	// 1 - Prepare users of the actor desc
+	OnActorDescPreUpdate(ActorDesc);
+	// 2 - Apply the transform, which is a function that depends on the context of the change, and which may have created a new ActorDesc and destroyed the old one, updating the passed-in unique ptr
+	UpdateAction(ActorDesc);
+	// 3 - Carry on processing the modified actor desc
+	OnActorDescUpdated(ActorDesc);
+}
+
 bool UActorDescContainer::IsActorDescHandled(const AActor* InActor) const
 {
 	return IsActorDescHandled(InActor, false);
@@ -519,9 +541,15 @@ void UActorDescContainer::OnObjectPreSave(UObject* Object, FObjectPreSaveContext
 				{
 					// Existing actor
 					RemoveChildActorFromParentMap(ExistingActorDesc->Get());
-					OnActorDescUpdating(ExistingActorDesc->Get());
-					FWorldPartitionActorDescUtils::UpdateActorDescriptorFromActor(Actor, *ExistingActorDesc);
-					OnActorDescUpdated(ExistingActorDesc->Get());
+
+					ReassignActorDesc(*ExistingActorDesc,
+						[Actor](TUniquePtr<FWorldPartitionActorDesc>& ActorDesc)
+						{
+							// This changes the actor desc pointer
+							FWorldPartitionActorDescUtils::UpdateActorDescriptorFromActor(Actor, ActorDesc);
+						}
+					);
+
 					AddChildActorToParentMap(ExistingActorDesc->Get());
 
 					PropagateActorToWorldUpdate(ExistingActorDesc->Get());
@@ -608,9 +636,13 @@ void UActorDescContainer::OnClassDescriptorUpdated(const FWorldPartitionActorDes
 			{
 				if (TUniquePtr<FWorldPartitionActorDesc>* ExistingActorDesc = GetActorDescriptor(NewActorDesc->GetGuid()))
 				{
-					OnActorDescUpdating(ExistingActorDesc->Get());
-					FWorldPartitionActorDescUtils::UpdateActorDescriptorFromActorDescriptor(NewActorDesc, *ExistingActorDesc);
-					OnActorDescUpdated(ExistingActorDesc->Get());
+					ReassignActorDesc(*ExistingActorDesc,
+						[&NewActorDesc](TUniquePtr<FWorldPartitionActorDesc>& ActorDesc)
+						{
+							// This changes the actor desc pointer
+							FWorldPartitionActorDescUtils::UpdateActorDescriptorFromActorDescriptor(NewActorDesc, ActorDesc);
+						}
+					);
 				}
 			}
 		}
@@ -668,13 +700,21 @@ void UActorDescContainer::OnActorDescRemoved(FWorldPartitionActorDesc* ActorDesc
 	OnActorDescRemovedEvent.Broadcast(ActorDesc);
 }
 
+void UActorDescContainer::OnActorDescPreUpdate(FWorldPartitionActorDesc* ActorDesc)
+{
+	// Pass to prepare for change to the ActorDesc
+	OnActorDescPreUpdateEvent.Broadcast(ActorDesc);
+}
+
 void UActorDescContainer::OnActorDescUpdating(FWorldPartitionActorDesc* ActorDesc)
 {
+	// Pass to update all pointers, because the previous pointer is now deleted
 	OnActorDescUpdatingEvent.Broadcast(ActorDesc);
 }
 
 void UActorDescContainer::OnActorDescUpdated(FWorldPartitionActorDesc* ActorDesc)
 {
+	// Pass to process the new ActorDesc, which may include using other containers which must also point to the updated ActorDesc
 	OnActorDescUpdatedEvent.Broadcast(ActorDesc);
 }
 
@@ -743,9 +783,12 @@ void UActorDescContainer::PropagateActorToWorldUpdate(FWorldPartitionActorDesc*
 	{
 		if (FWorldPartitionActorDesc* ChildActor = GetActorDesc(ChildActorGuid))
 		{
-			OnActorDescUpdating(ChildActor);
-			ChildActor->UpdateActorToWorld();
-			OnActorDescUpdated(ChildActor);
+			UpdateActorDesc(ChildActor,
+				[](FWorldPartitionActorDesc* ActorDesc)
+				{
+					ActorDesc->UpdateActorToWorld();
+				}
+			);
 		}
 	});
 }
diff --git a/Engine/Source/Runtime/Engine/Private/WorldPartition/ActorDescContainerInstance.cpp b/Engine/Source/Runtime/Engine/Private/WorldPartition/ActorDescContainerInstance.cpp
--- a/Engine/Source/Runtime/Engine/Private/WorldPartition/ActorDescContainerInstance.cpp
+++ b/Engine/Source/Runtime/Engine/Private/WorldPartition/ActorDescContainerInstance.cpp
@@ -270,6 +270,7 @@ void UActorDescContainerInstance::RegisterDelegates()
 	// even in Instanced worlds (Level Instances)
 	Container->OnActorDescRemovedEvent.AddUObject(this, &UActorDescContainerInstance::OnActorDescRemoved);
 		
+	Container->OnActorDescPreUpdateEvent.AddUObject(this, &UActorDescContainerInstance::OnActorDescPreUpdate);
 	Container->OnActorDescUpdatingEvent.AddUObject(this, &UActorDescContainerInstance::OnActorDescUpdating);
 	Container->OnActorDescUpdatedEvent.AddUObject(this, &UActorDescContainerInstance::OnActorDescUpdated);
 }
@@ -282,6 +283,7 @@ void UActorDescContainerInstance::UnregisterDelegates()
 	check(Container);
 	Container->OnActorDescAddedEvent.RemoveAll(this);
 	Container->OnActorDescRemovedEvent.RemoveAll(this);
+	Container->OnActorDescPreUpdateEvent.RemoveAll(this);
 	Container->OnActorDescUpdatingEvent.RemoveAll(this);
 	Container->OnActorDescUpdatedEvent.RemoveAll(this);
 
@@ -579,7 +581,7 @@ void UActorDescContainerInstance::OnActorDescRemoved(FWorldPartitionActorDesc* I
 	RemoveActor(InActorDesc->GetGuid());
 }
 
-void UActorDescContainerInstance::OnActorDescUpdating(FWorldPartitionActorDesc* InActorDesc)
+void UActorDescContainerInstance::OnActorDescPreUpdate(FWorldPartitionActorDesc* InActorDesc)
 {
 	// It is possible for a container instance with an instancing context to listen to updates but not newly added actors
 	// so it might not find the actor if it was newly added (through the edit of a different instance of the same Level Instance)
@@ -599,21 +601,30 @@ void UActorDescContainerInstance::OnActorDescUpdating(FWorldPartitionActorDesc*
 			WorldPartition->OnActorDescInstanceUpdating(ActorDescInstance);
 		}
 
-		ActorDescInstance->OnUpdating();
+		ActorDescInstance->OnPreUpdate();
 
 		OnActorDescInstanceUpdatingEvent.Broadcast(ActorDescInstance);
 	}
 }
 
-void UActorDescContainerInstance::OnActorDescUpdated(FWorldPartitionActorDesc* InActorDesc)
+void UActorDescContainerInstance::OnActorDescUpdating(FWorldPartitionActorDesc* InActorDesc)
 {
-	// See comment in UActorDescContainerInstance::OnActorDescUpdating
+	// See comment in UActorDescContainerInstance::OnActorDescPreUpdate
 	if (TUniquePtr<FWorldPartitionActorDescInstance>* ActorDescInstance = GetActorDescInstancePtr(InActorDesc->GetGuid()))
 	{
 		check(ActorDescInstance->IsValid());
 
-		// Update instance desc
+		// Update instance desc pointer
 		ActorDescInstance->Get()->UpdateActorDesc(InActorDesc);
+	}
+}
+
+void UActorDescContainerInstance::OnActorDescUpdated(FWorldPartitionActorDesc* InActorDesc)
+{
+	// See comment in UActorDescContainerInstance::OnActorDescPreUpdate
+	if (TUniquePtr<FWorldPartitionActorDescInstance>* ActorDescInstance = GetActorDescInstancePtr(InActorDesc->GetGuid()))
+	{
+		check(ActorDescInstance->IsValid());
 
 		// Re-register container
 		if (bCreateChildContainerHierarchy && ActorDescInstance->Get()->IsChildContainerInstance())
@@ -637,4 +648,4 @@ void UActorDescContainerInstance::OnActorDescUpdated(FWorldPartitionActorDesc* I
 	}
 }
 
-#endif
+#endif
\ No newline at end of file
diff --git a/Engine/Source/Runtime/Engine/Private/WorldPartition/LevelInstance/LevelInstancePropertyOverrideContainer.cpp b/Engine/Source/Runtime/Engine/Private/WorldPartition/LevelInstance/LevelInstancePropertyOverrideContainer.cpp
--- a/Engine/Source/Runtime/Engine/Private/WorldPartition/LevelInstance/LevelInstancePropertyOverrideContainer.cpp
+++ b/Engine/Source/Runtime/Engine/Private/WorldPartition/LevelInstance/LevelInstancePropertyOverrideContainer.cpp
@@ -115,11 +115,15 @@ void ULevelInstancePropertyOverrideContainer::Initialize(const FInitializeParams
 					ParentTransform = ParentActorDesc->GetActorTransform();
 				}
 			}
+
 			// Always pass parent transform (even when Identity) because we don't want UpdateActorToWorld to resolve
 			// the parent by itself, as it misses the context (InContainerPath) to properly resolve the actordesc.
-			OnActorDescUpdating(MutableActorDesc);
-			MutableActorDesc->UpdateActorToWorld(&ParentTransform);
-			OnActorDescUpdated(MutableActorDesc);
+			UpdateActorDesc(MutableActorDesc,
+				[&ParentTransform](FWorldPartitionActorDesc* ActorDesc)
+				{
+					ActorDesc->UpdateActorToWorld(&ParentTransform);
+				}
+			);
 		}
 	};
 
@@ -196,6 +200,7 @@ void ULevelInstancePropertyOverrideContainer::UnregisterBaseContainerDelegates()
 	UActorDescContainer* BaseContainer = GetBaseContainer();
 	check(BaseContainer);
 	BaseContainer->OnActorDescRemovedEvent.RemoveAll(this);
+	BaseContainer->OnActorDescPreUpdateEvent.RemoveAll(this);
 	BaseContainer->OnActorDescUpdatingEvent.RemoveAll(this);
 	BaseContainer->OnActorDescUpdatedEvent.RemoveAll(this);
 }
@@ -206,6 +211,7 @@ void ULevelInstancePropertyOverrideContainer::RegisterBaseContainerDelegates()
 	UActorDescContainer* BaseContainer = PropertyOverrideDesc->GetBaseContainer();
 	check(BaseContainer);
 	BaseContainer->OnActorDescRemovedEvent.AddUObject(this, &ULevelInstancePropertyOverrideContainer::OnBaseContainerActorDescRemoved);
+	BaseContainer->OnActorDescPreUpdateEvent.AddUObject(this, &ULevelInstancePropertyOverrideContainer::OnBaseContainerActorDescPreUpdate);
 	BaseContainer->OnActorDescUpdatingEvent.AddUObject(this, &ULevelInstancePropertyOverrideContainer::OnBaseContainerActorDescUpdating);
 	BaseContainer->OnActorDescUpdatedEvent.AddUObject(this, &ULevelInstancePropertyOverrideContainer::OnBaseContainerActorDescUpdated);
 }
@@ -215,6 +221,11 @@ void ULevelInstancePropertyOverrideContainer::OnBaseContainerActorDescRemoved(FW
 	OnActorDescRemoved(InActorDesc);
 }
 
+void ULevelInstancePropertyOverrideContainer::OnBaseContainerActorDescPreUpdate(FWorldPartitionActorDesc* InActorDesc)
+{
+	OnActorDescPreUpdate(InActorDesc);
+}
+
 void ULevelInstancePropertyOverrideContainer::OnBaseContainerActorDescUpdating(FWorldPartitionActorDesc* InActorDesc)
 {
 	OnActorDescUpdating(InActorDesc);
diff --git a/Engine/Source/Runtime/Engine/Private/WorldPartition/LevelInstance/LevelInstancePropertyOverrideContainer.h b/Engine/Source/Runtime/Engine/Private/WorldPartition/LevelInstance/LevelInstancePropertyOverrideContainer.h
--- a/Engine/Source/Runtime/Engine/Private/WorldPartition/LevelInstance/LevelInstancePropertyOverrideContainer.h
+++ b/Engine/Source/Runtime/Engine/Private/WorldPartition/LevelInstance/LevelInstancePropertyOverrideContainer.h
@@ -63,6 +63,7 @@ private:
 	void UnregisterBaseContainerDelegates();
 
 	void OnBaseContainerActorDescRemoved(FWorldPartitionActorDesc* InActorDesc);
+	void OnBaseContainerActorDescPreUpdate(FWorldPartitionActorDesc* InActorDesc);
 	void OnBaseContainerActorDescUpdating(FWorldPartitionActorDesc* InActorDesc);
 	void OnBaseContainerActorDescUpdated(FWorldPartitionActorDesc* InActorDesc);
 
diff --git a/Engine/Source/Runtime/Engine/Private/WorldPartition/WorldPartitionActorDescInstance.cpp b/Engine/Source/Runtime/Engine/Private/WorldPartition/WorldPartitionActorDescInstance.cpp
--- a/Engine/Source/Runtime/Engine/Private/WorldPartition/WorldPartitionActorDescInstance.cpp
+++ b/Engine/Source/Runtime/Engine/Private/WorldPartition/WorldPartitionActorDescInstance.cpp
@@ -198,7 +198,7 @@ FBox FWorldPartitionActorDescInstance::GetLocalEditorBounds(bool bInForceCompute
 	return ActorDesc->GetEditorBounds();
 }
 
-void FWorldPartitionActorDescInstance::OnUpdating()
+void FWorldPartitionActorDescInstance::OnPreUpdate()
 {
 	// Invalidate cached bounds
 	if (IsChildContainerInstance())
diff --git a/Engine/Source/Runtime/Engine/Public/WorldPartition/ActorDescContainer.h b/Engine/Source/Runtime/Engine/Public/WorldPartition/ActorDescContainer.h
--- a/Engine/Source/Runtime/Engine/Public/WorldPartition/ActorDescContainer.h
+++ b/Engine/Source/Runtime/Engine/Public/WorldPartition/ActorDescContainer.h
@@ -83,9 +83,15 @@ public:
 	DECLARE_EVENT_OneParam(UActorDescContainer, FActorDescRemovedEvent, FWorldPartitionActorDesc*);
 	FActorDescRemovedEvent OnActorDescRemovedEvent;
 
+	// Pre update event, before the ActorDesc gets updated
 	DECLARE_EVENT_OneParam(UActorDescContainer, FActorDescUpdatingEvent, FWorldPartitionActorDesc*);
-	FActorDescUpdatingEvent OnActorDescUpdatingEvent;
+	FActorDescUpdatingEvent OnActorDescPreUpdateEvent;
 
+	// Update event, to update to the new pointer
+	DECLARE_EVENT_OneParam(UActorDescContainer, FOnActorDescUpdatingEvent, FWorldPartitionActorDesc*);
+	FOnActorDescUpdatingEvent OnActorDescUpdatingEvent;
+
+	// Post update event, when the new pointer is now safe to use across instance containers
 	DECLARE_EVENT_OneParam(UActorDescContainer, FActorDescUpdatedEvent, FWorldPartitionActorDesc*);
 	FActorDescUpdatedEvent OnActorDescUpdatedEvent;
 
@@ -101,6 +107,7 @@ public:
 
 	ENGINE_API void OnActorDescAdded(FWorldPartitionActorDesc* NewActorDesc);
 	ENGINE_API void OnActorDescRemoved(FWorldPartitionActorDesc* ActorDesc);
+	ENGINE_API void OnActorDescPreUpdate(FWorldPartitionActorDesc* ActorDesc);
 	ENGINE_API void OnActorDescUpdating(FWorldPartitionActorDesc* ActorDesc);
 	ENGINE_API void OnActorDescUpdated(FWorldPartitionActorDesc* ActorDesc);
 
@@ -160,6 +167,13 @@ protected:
 
 	ENGINE_API bool ShouldHandleActorEvent(const AActor* Actor, bool bInUseLoadedPath) const;
 	ENGINE_API bool IsActorDescHandled(const AActor* InActor, bool bInUseLoadedPath) const;
+
+	// Updater for the ActorDesc which allows UpdateAction to change the pointer
+	void ReassignActorDesc(TUniquePtr<FWorldPartitionActorDesc>& DescToUpdate, TFunctionRef<void(TUniquePtr<FWorldPartitionActorDesc>& ActorDesc)> UpdateAction);
+
+	// Updater for the ActorDesc which only allows updating the contents of the descriptor
+	void UpdateActorDesc(FWorldPartitionActorDesc* DescToUpdate, TFunctionRef<void(FWorldPartitionActorDesc* ActorDesc)> UpdateAction);
+
 private:
 	// GetWorld() should never be called on an ActorDescContainer to avoid any confusion as it can be used as a template
 	UWorld* GetWorld() const override { return nullptr; }
diff --git a/Engine/Source/Runtime/Engine/Public/WorldPartition/ActorDescContainerInstance.h b/Engine/Source/Runtime/Engine/Public/WorldPartition/ActorDescContainerInstance.h
--- a/Engine/Source/Runtime/Engine/Public/WorldPartition/ActorDescContainerInstance.h
+++ b/Engine/Source/Runtime/Engine/Public/WorldPartition/ActorDescContainerInstance.h
@@ -174,6 +174,7 @@ private:
 
 	void OnActorDescAdded(FWorldPartitionActorDesc* InActorDesc);
 	void OnActorDescRemoved(FWorldPartitionActorDesc* InActorDesc);
+	void OnActorDescPreUpdate(FWorldPartitionActorDesc* InActorDesc);
 	void OnActorDescUpdating(FWorldPartitionActorDesc* InActorDesc);
 	void OnActorDescUpdated(FWorldPartitionActorDesc* InActorDesc);
 
diff --git a/Engine/Source/Runtime/Engine/Public/WorldPartition/WorldPartitionActorDescInstance.h b/Engine/Source/Runtime/Engine/Public/WorldPartition/WorldPartitionActorDescInstance.h
--- a/Engine/Source/Runtime/Engine/Public/WorldPartition/WorldPartitionActorDescInstance.h
+++ b/Engine/Source/Runtime/Engine/Public/WorldPartition/WorldPartitionActorDescInstance.h
@@ -177,8 +177,8 @@ protected:
 	virtual ENGINE_API void UpdateChildContainerInstance();
 
 private:
-	/** Called by UActorDescContainerInstance::OnActorDescUpdating */
-	void OnUpdating();
+	/** Called by UActorDescContainerInstance::OnActorDescPreUpdate */
+	void OnPreUpdate();
 	void InvalidateCachedBounds();
 
 protected:
