Don't add padding when vertex count is 0 to make sure we don't allocate anything when we don't have any vertices

#Synced-CL 50194444

diff --git a/Engine/Source/Runtime/Engine/Private/Rendering/StaticMeshVertexBuffer.cpp b/Engine/Source/Runtime/Engine/Private/Rendering/StaticMeshVertexBuffer.cpp
--- a/Engine/Source/Runtime/Engine/Private/Rendering/StaticMeshVertexBuffer.cpp
+++ b/Engine/Source/Runtime/Engine/Private/Rendering/StaticMeshVertexBuffer.cpp
@@ -18,7 +18,7 @@ static uint32 GetNumTexCoordInVertexBuffer(uint32 NumTexCoords, uint32 NumVertic
 	// will declare the stream as half4/float4 while it only contains a half2/float2 and the last element can in theory
 	// be out of bounds then
 	// Vertex streams are all created with half4/float4 to reduce the number of PSO permutations coming from different texcoords in the vertex declaration
-	uint32 Padding = NumTexCoords % 2;
+	uint32 Padding = NumVertices > 0 ? NumTexCoords % 2 : 0;
 	return (NumVertices + Padding) * NumTexCoords;
 }
 
