[ContentBundle] When preparing paths, don't require that those paths are available already

#tests Tested 2 projects which used to ensure
#rnx

#Synced-CL 50114392

diff --git a/Engine/Source/Runtime/Engine/Private/WorldPartition/ContentBundle/ContentBundle.cpp b/Engine/Source/Runtime/Engine/Private/WorldPartition/ContentBundle/ContentBundle.cpp
--- a/Engine/Source/Runtime/Engine/Private/WorldPartition/ContentBundle/ContentBundle.cpp
+++ b/Engine/Source/Runtime/Engine/Private/WorldPartition/ContentBundle/ContentBundle.cpp
@@ -51,7 +51,7 @@ void FContentBundle::DoInitialize()
 	}
 	else
 	{
-		UE_LOG(LogContentBundle, Log, TEXT("%s Has no content for the current World (this can happen if the ContentBundle has no Actors for the current World)."), *ContentBundle::Log::MakeDebugInfoString(*this));
+		UE_LOG(LogContentBundle, Log, TEXT("%s Has no content for the current World (this can happen if the ContentBundle has no Actors for the current World, or if the GameFeaturePlugin containing the CB is not yet mounted)."), *ContentBundle::Log::MakeDebugInfoString(*this));
 	}
 #endif
 
diff --git a/Engine/Source/Runtime/Engine/Private/WorldPartition/ContentBundle/ContentBundlePaths.cpp b/Engine/Source/Runtime/Engine/Private/WorldPartition/ContentBundle/ContentBundlePaths.cpp
--- a/Engine/Source/Runtime/Engine/Private/WorldPartition/ContentBundle/ContentBundlePaths.cpp
+++ b/Engine/Source/Runtime/Engine/Private/WorldPartition/ContentBundle/ContentBundlePaths.cpp
@@ -7,6 +7,7 @@
 #include "UObject/Package.h"
 #include "WorldPartition/ContentBundle/ContentBundleBase.h"
 #include "WorldPartition/ContentBundle/ContentBundleDescriptor.h"
+#include "WorldPartition/ContentBundle/ContentBundleLog.h"
 #include "Misc/Paths.h"
 #include "AssetRegistry/AssetData.h"
 #include "ExternalPackageHelper.h"
@@ -81,8 +82,10 @@ FString ContentBundlePaths::GetRelativeLevelFolder(const FContentBundleBase& Con
 	const UWorld* InjectedWorld = ContentBundle.GetInjectedWorld();
 	FString WorldPackageName = InjectedWorld->GetPackage()->GetName();
 
+	UE_CLOG(!FPackageName::IsValidPath(WorldPackageName), LogContentBundle, Warning, TEXT("Invalid world package name '%s' (maybe not yet mounted), Content Bundle won't work as expected"), *WorldPackageName);
+
 	FString PackageRoot, PackagePath, PackageName;
-	ensure(FPackageName::SplitLongPackageName(WorldPackageName, PackageRoot, PackagePath, PackageName));
+	FPackageName::SplitPackageName(WorldPackageName, &PackageRoot, &PackagePath, &PackageName);
 
 	TStringBuilderWithBuffer<TCHAR, NAME_SIZE> RelativeLevelFolder;
 	RelativeLevelFolder += PackagePath;
