Fix for DemoNetDriver crashes due to accessing memory in destroyed actors.

#Synced-CL 50171975

diff --git a/Engine/Source/Runtime/Engine/Classes/Engine/NetDriver.h b/Engine/Source/Runtime/Engine/Classes/Engine/NetDriver.h
--- a/Engine/Source/Runtime/Engine/Classes/Engine/NetDriver.h
+++ b/Engine/Source/Runtime/Engine/Classes/Engine/NetDriver.h
@@ -2245,6 +2245,9 @@ protected:
 	FOnConsiderListUpdate OnProcessConsiderListOverride;
 #endif
 
+	/** Version of RemoveNetworkActor where the actor pointer must be assumed to be inaccessible. */
+	ENGINE_API virtual void RemoveInvalidNetworkActor(const AActor* BadActor);
+
 	/** Used to handle any NetDriver specific cleanup once a level has been removed from the world. */
 	ENGINE_API virtual void OnLevelRemovedFromWorld(class ULevel* Level, class UWorld* World);
 
diff --git a/Engine/Source/Runtime/Engine/Classes/Engine/NetworkObjectList.h b/Engine/Source/Runtime/Engine/Classes/Engine/NetworkObjectList.h
--- a/Engine/Source/Runtime/Engine/Classes/Engine/NetworkObjectList.h
+++ b/Engine/Source/Runtime/Engine/Classes/Engine/NetworkObjectList.h
@@ -111,7 +111,7 @@ struct FNetworkObjectInfo
 /**
  * KeyFuncs to allow using the actor pointer as the comparison key in a set.
  */
-struct FNetworkObjectKeyFuncs : BaseKeyFuncs<TSharedPtr<FNetworkObjectInfo>, AActor*, false>
+struct FNetworkObjectKeyFuncs : BaseKeyFuncs<TSharedPtr<FNetworkObjectInfo>, AActor const*, false>
 {
 	/**
 	 * @return The key used to index the given element.
@@ -177,7 +177,10 @@ public:
 	}
 
 	/** Removes actor from the internal list, and any cleanup that is necessary (i.e. resetting dormancy state) */
-	ENGINE_API void Remove(AActor* const Actor);
+	ENGINE_API void Remove(const AActor* Actor);
+
+	/** Removes an actor with inaccessible memory from the internal list, and any cleanup that is necessary (i.e. resetting dormancy state) */
+	ENGINE_API void RemoveInvalid(const AActor* BadActor);
 
 	/** Marks this object as dormant for the passed in connection */
 	ENGINE_API void MarkDormant(AActor* const Actor, UNetConnection* const Connection, const int32 NumConnections, UNetDriver* NetDriver);
diff --git a/Engine/Source/Runtime/Engine/Private/DemoNetDriver.cpp b/Engine/Source/Runtime/Engine/Private/DemoNetDriver.cpp
--- a/Engine/Source/Runtime/Engine/Private/DemoNetDriver.cpp
+++ b/Engine/Source/Runtime/Engine/Private/DemoNetDriver.cpp
@@ -1968,7 +1968,12 @@ void UDemoNetDriver::TickDemoRecordFrame(float DeltaSeconds)
 			const float RecordHzWhenNotRelevant = CVarDemoRecordHzWhenNotRelevant.GetValueOnAnyThread();
 			const float UpdateDelayWhenNotRelevant = RecordHzWhenNotRelevant > 0.0f ? 1.0f / RecordHzWhenNotRelevant : 0.5f;
 
-			TArray<AActor*, TInlineAllocator<128>> ActorsToRemove;
+			struct FActorRemoveInfo
+			{
+				AActor* Actor = nullptr;
+				bool bIsValidActor = false;
+			};
+			TArray<FActorRemoveInfo, TInlineAllocator<128>> ActorsToRemove;
 
 			FDemoActorPriority DemoActorPriority;
 			FActorPriority& ActorPriority = DemoActorPriority.ActorPriority;
@@ -1991,7 +1996,7 @@ void UDemoNetDriver::TickDemoRecordFrame(float DeltaSeconds)
 
 					if (!ActorInfo->WeakActor.IsValid())
 					{
-						ActorsToRemove.Add(Actor);
+						ActorsToRemove.Add({.Actor = Actor, .bIsValidActor = false});
 						continue;
 					}
 
@@ -1999,19 +2004,19 @@ void UDemoNetDriver::TickDemoRecordFrame(float DeltaSeconds)
 					// we still need to replicate it one more time so that the recorded replay knows it's been torn-off as well.
 					if (Actor->GetRemoteRole() == ROLE_None && !Actor->GetTearOff())
 					{
-						ActorsToRemove.Add(Actor);
+						ActorsToRemove.Add({.Actor = Actor, .bIsValidActor = true});
 						continue;
 					}
 
 					if (IsDormInitialStartupActor(Actor))
 					{
-						ActorsToRemove.Add(Actor);
+						ActorsToRemove.Add({.Actor = Actor, .bIsValidActor = true});
 						continue;
 					}
 
 					if (!Actor->bRelevantForNetworkReplays)
 					{
-						ActorsToRemove.Add(Actor);
+						ActorsToRemove.Add({.Actor = Actor, .bIsValidActor = true});
 						continue;
 					}
 
@@ -2098,9 +2103,16 @@ void UDemoNetDriver::TickDemoRecordFrame(float DeltaSeconds)
 				SCOPED_NAMED_EVENT(UDemoNetDriver_PrioritizeRemoveActors, FColor::Green);
 
 				// Always remove necessary actors, don't time slice this.
-				for (AActor* Actor : ActorsToRemove)
+				for (FActorRemoveInfo& RemoveInfo : ActorsToRemove)
 				{
-					RemoveNetworkActor(Actor);
+					if (RemoveInfo.bIsValidActor)
+					{
+						RemoveNetworkActor(RemoveInfo.Actor);
+					}
+					else
+					{
+						RemoveInvalidNetworkActor(RemoveInfo.Actor);
+					}
 				}
 			}
 		}
@@ -2116,7 +2128,7 @@ void UDemoNetDriver::TickDemoRecordFrame(float DeltaSeconds)
 			}
 			else
 			{
-				UE_LOG(LogDemo, Verbose, TEXT("Consider time exhaused without processing destruction infos"));
+				UE_LOG(LogDemo, Verbose, TEXT("Consider time exhausted without processing destruction infos"));
 			}
 		}
 		else
diff --git a/Engine/Source/Runtime/Engine/Private/NetDriver.cpp b/Engine/Source/Runtime/Engine/Private/NetDriver.cpp
--- a/Engine/Source/Runtime/Engine/Private/NetDriver.cpp
+++ b/Engine/Source/Runtime/Engine/Private/NetDriver.cpp
@@ -4589,6 +4589,26 @@ void UNetDriver::RemoveNetworkActor(AActor* Actor)
 	GetNetworkObjectList().Remove(Actor);
 }
 
+void UNetDriver::RemoveInvalidNetworkActor(const AActor* BadActor)
+{
+	LLM_SCOPE_BYTAG(NetDriver);
+
+	if (ReplicationDriver)
+	{
+		ensureMsgf(false, TEXT("Not calling UReplicationDriver::RemoveNetworkActor with bad actor pointer: %p"), BadActor);
+	}
+
+	// In case a network actor is removed while the ConsiderList is in use, remove the stale entries.
+	// This can happen when bReplicateTransactionally is true and ServerReplicateActors_ForConnection aborts.
+	CurrentConsiderList.RemoveAllSwap([BadActor](const FNetworkObjectInfo* Info)
+		{
+			return Info->Actor == BadActor;
+		},
+		EAllowShrinking::No);
+
+	GetNetworkObjectList().RemoveInvalid(BadActor);
+}
+
 void UNetDriver::DeleteSubObjectOnClients(AActor* Actor, UObject* SubObject)
 {
 	if (ReplicationSystem)
diff --git a/Engine/Source/Runtime/Engine/Private/NetworkObjectList.cpp b/Engine/Source/Runtime/Engine/Private/NetworkObjectList.cpp
--- a/Engine/Source/Runtime/Engine/Private/NetworkObjectList.cpp
+++ b/Engine/Source/Runtime/Engine/Private/NetworkObjectList.cpp
@@ -122,7 +122,7 @@ TSharedPtr<FNetworkObjectInfo>* FNetworkObjectList::FindOrAdd(AActor* const Acto
 	return NetworkObjectInfo;
 }
 
-void FNetworkObjectList::Remove(AActor* const Actor)
+void FNetworkObjectList::Remove(const AActor* const Actor)
 {
 	if (Actor == nullptr)
 	{
@@ -202,6 +202,91 @@ void FNetworkObjectList::Remove(AActor* const Actor)
 	check((ActiveNetworkObjects.Num() + ObjectsDormantOnAllConnections.Num()) == AllNetworkObjects.Num());
 }
 
+void FNetworkObjectList::RemoveInvalid(const AActor* const BadActor)
+{
+	if (BadActor == nullptr)
+	{
+		return;
+	}
+
+	TSharedPtr<FNetworkObjectInfo>* NetworkObjectInfoPtr = AllNetworkObjects.Find(BadActor);
+
+	if (NetworkObjectInfoPtr == nullptr)
+	{
+		// Sanity check that we're not on the other lists either
+		check(!ActiveNetworkObjects.Contains(BadActor));
+		check(!ObjectsDormantOnAllConnections.Contains(BadActor));
+		check((ActiveNetworkObjects.Num() + ObjectsDormantOnAllConnections.Num()) == AllNetworkObjects.Num());
+		return;
+	}
+
+	FNetworkObjectInfo* NetworkObjectInfo = NetworkObjectInfoPtr->Get();
+
+	// Lower the dormant object count for each connection this object is dormant on
+	for (auto ConnectionIt = NetworkObjectInfo->DormantConnections.CreateIterator(); ConnectionIt; ++ConnectionIt)
+	{
+		UNetConnection* Connection = (*ConnectionIt).Get();
+
+		if (Connection == nullptr || Connection->GetConnectionState() == USOCK_Closed)
+		{
+			if (UE::Net::Private::bTrackDormantObjectsByLevel)
+			{
+				DormantObjectsPerConnection.Remove(Connection);
+			}
+
+			ConnectionIt.RemoveCurrent();
+			continue;
+		}
+
+		int32& NumDormantObjectsPerConnectionRef = NumDormantObjectsPerConnection.FindOrAdd(Connection);
+		check(NumDormantObjectsPerConnectionRef > 0);
+		NumDormantObjectsPerConnectionRef--;
+
+		// Remove the actor from all connection specific dormancy object sets
+		if (UE::Net::Private::bTrackDormantObjectsByLevel)
+		{
+			if (TMap<FName, FNetworkObjectSet>* DormantObjectsByLevel = DormantObjectsPerConnection.Find(Connection))
+			{
+				for (TMap<FName, FNetworkObjectSet>::TIterator It = DormantObjectsByLevel->CreateIterator(); It; ++It)
+				{
+					FNetworkObjectSet& ObjectSet = It.Value();
+					// If we found the correct object set to remove the actor from we're done
+					if (ObjectSet.Remove(BadActor) > 0)
+					{
+						// There's no point in keeping empty sets around
+						if (ObjectSet.IsEmpty())
+						{
+							It.RemoveCurrent();
+						}
+						break;
+					}
+				}
+			}
+		}
+	}
+
+	// Remove this object from all lists
+	AllNetworkObjects.Remove(BadActor);
+	ActiveNetworkObjects.Remove(BadActor);
+	ObjectsDormantOnAllConnections.Remove(BadActor);
+
+	// Remove the actor from the fully dormant object sets
+	if (UE::Net::Private::bTrackDormantObjectsByLevel)
+	{
+		for (TMap<FName, FNetworkObjectSet>::TIterator It = FullyDormantObjectsByLevel.CreateIterator(); It; ++It)
+		{
+			FNetworkObjectSet& ObjectSet = It.Value();
+			ObjectSet.Remove(BadActor);
+			if (ObjectSet.IsEmpty())
+			{
+				It.RemoveCurrent();
+			}
+		}
+	}
+
+	check((ActiveNetworkObjects.Num() + ObjectsDormantOnAllConnections.Num()) == AllNetworkObjects.Num());
+}
+
 void FNetworkObjectList::MarkDormant(AActor* const Actor, UNetConnection* const Connection, const int32 NumConnections, UNetDriver* NetDriver)
 {
 	TSharedPtr<FNetworkObjectInfo>* NetworkObjectInfoPtr = FindOrAdd(Actor, NetDriver);
