[Waveform Editor] Strip out invalid samples on load

#Synced-CL 50159815

diff --git a/Engine/Source/Runtime/Engine/Private/Audio.cpp b/Engine/Source/Runtime/Engine/Private/Audio.cpp
--- a/Engine/Source/Runtime/Engine/Private/Audio.cpp
+++ b/Engine/Source/Runtime/Engine/Private/Audio.cpp
@@ -1896,6 +1896,23 @@ bool FWaveModInfo::ReadWaveInfo( const uint8* WaveData, int32 WaveDataSize, FStr
 			RiffChunk->ChunkLen = INTEL_ORDER32(SampleDataSize);
 		}
 
+		// Truncate any bad sample data, does not affect exported data
+		// Corrupt data remains present in the raw wave data
+		if (pBlockAlign && *pBlockAlign > 0)
+		{
+			const uint16 BlockAlign = *pBlockAlign; // NumChannels * BytesPerSample
+			const uint32 ExtraBytes = SampleDataSize % BlockAlign;
+
+			if (ExtraBytes > 0)
+			{
+				UE_LOG(LogAudio, Warning, TEXT("Wave data chunk has %u invalid byte(s), truncating invalid block of compressed data"), ExtraBytes);
+
+				SampleDataEnd = SampleDataEnd - ExtraBytes;
+				SampleDataSize = SampleDataEnd - SampleDataStart;
+				RiffChunk->ChunkLen = INTEL_ORDER32(SampleDataSize);
+			}
+		}
+
 		NewDataSize = SampleDataSize;
 
 		#if !PLATFORM_LITTLE_ENDIAN
