[HLODBuilder] Batch source control operations in the HLOD Setup step into a single operation, significantly reducing source control overhead.

#Synced-CL 50093269

diff --git a/Engine/Source/Editor/UnrealEd/Private/WorldPartition/WorldPartitionBuilderHelpers.cpp b/Engine/Source/Editor/UnrealEd/Private/WorldPartition/WorldPartitionBuilderHelpers.cpp
--- a/Engine/Source/Editor/UnrealEd/Private/WorldPartition/WorldPartitionBuilderHelpers.cpp
+++ b/Engine/Source/Editor/UnrealEd/Private/WorldPartition/WorldPartitionBuilderHelpers.cpp
@@ -90,13 +90,21 @@ bool FSourceControlHelper::Checkout(UPackage* Package) const
 
 bool FSourceControlHelper::Checkout(const FString& Filename) const
 {
-	bool bCheckedOut = PackageHelper.Checkout({ Filename });
+	return Checkout(TArray<FString>({ Filename }));
+}
+
+bool FSourceControlHelper::Checkout(const TArray<FString>& Filenames) const
+{
+	bool bCheckedOut = PackageHelper.Checkout(Filenames);
 	if (bCheckedOut)
 	{
-		const bool bAdded = ModifiedFiles.Get(FBuilderModifiedFiles::EFileOperation::FileAdded).Contains(Filename);
-		if (!bAdded)
+		for (const FString& Filename : Filenames)
 		{
-			ModifiedFiles.Add(FBuilderModifiedFiles::EFileOperation::FileEdited, Filename);
+			const bool bAdded = ModifiedFiles.Get(FBuilderModifiedFiles::EFileOperation::FileAdded).Contains(Filename);
+			if (!bAdded)
+			{
+				ModifiedFiles.Add(FBuilderModifiedFiles::EFileOperation::FileEdited, Filename);
+			}
 		}
 	}
 	return bCheckedOut;
@@ -114,10 +122,15 @@ bool FSourceControlHelper::Add(UPackage* Package) const
 
 bool FSourceControlHelper::Add(const FString& Filename) const
 {
-	bool bAdded = PackageHelper.AddToSourceControl({ Filename });
+	return Add(TArray<FString>({ Filename }));
+}
+
+bool FSourceControlHelper::Add(const TArray<FString>& Filenames) const
+{
+	bool bAdded = PackageHelper.AddToSourceControl(Filenames);
 	if (bAdded)
 	{
-		ModifiedFiles.Add(FBuilderModifiedFiles::EFileOperation::FileAdded, Filename);
+		ModifiedFiles.Append(FBuilderModifiedFiles::EFileOperation::FileAdded, Filenames);
 	}
 	return bAdded;
 }
@@ -153,39 +166,82 @@ bool FSourceControlHelper::Delete(UPackage* Package) const
 	return bDeleted;
 }
 
+bool FSourceControlHelper::Delete(const TArray<UPackage*>& Packages) const
+{
+	TArray<FString> PackageNames;
+	Algo::Transform(Packages, PackageNames, [this](UPackage* Package)
+	{
+		return GetFilename(Package);
+	});
+	
+	bool bDeleted = PackageHelper.Delete(Packages);
+	if (bDeleted)
+	{
+		ModifiedFiles.Append(FBuilderModifiedFiles::FileDeleted, PackageNames);
+	}
+	return bDeleted;
+}
+
 bool FSourceControlHelper::Save(UPackage* Package) const 
 {
-	bool bFileExists = IPlatformFile::GetPlatformPhysical().FileExists(*GetFilename(Package));
+	return Save(TArray<UPackage*>({ Package }));
+}
 
-	// Checkout package
-	Package->MarkAsFullyLoaded();
+bool FSourceControlHelper::Save(const TArray<UPackage*>& Packages) const
+{
+	TArray<FString> PackagesToCheckout;
+	TArray<FString> PackagesToAdd;
 
-	if (bFileExists)
+	for (UPackage* Package : Packages)
 	{
-		if (!Checkout(Package))
+		if (!Package)
 		{
-			UE_LOG(LogWorldPartitionBuilderSourceControlHelper, Error, TEXT("Error checking out package %s."), *Package->GetName());
+			return false;
+		}
+
+		FString PackageFileName = GetFilename(Package);
+		if (IPlatformFile::GetPlatformPhysical().FileExists(*PackageFileName))
+		{
+			PackagesToCheckout.Add(PackageFileName);
+		}
+		else
+		{
+			PackagesToAdd.Add(PackageFileName);
+		}
+	}
+	
+	// Checkout packages
+	if (PackagesToCheckout.Num() > 0)
+	{
+		if (!Checkout(PackagesToCheckout))
+		{
+			UE_LOG(LogWorldPartitionBuilderSourceControlHelper, Error, TEXT("Error checking out packages"));
 			return false;
 		}
 	}
 
-	// Save package
-	FString PackageFileName = GetFilename(Package);
-	FSavePackageArgs SaveArgs;
-	SaveArgs.TopLevelFlags = RF_Standalone | RF_Public;
-	SaveArgs.SaveFlags = PackageHelper.UseSourceControl() ? ESaveFlags::SAVE_None : ESaveFlags::SAVE_Async;
-	if (!UPackage::SavePackage(Package, nullptr, *PackageFileName, SaveArgs))
+	// Save packages
+	for (UPackage* Package : Packages)
 	{
-		UE_LOG(LogWorldPartitionBuilderSourceControlHelper, Error, TEXT("Error saving package %s."), *Package->GetName());
-		return false;
+		Package->MarkAsFullyLoaded();
+
+		FString PackageFileName = GetFilename(Package);
+		FSavePackageArgs SaveArgs;
+		SaveArgs.TopLevelFlags = RF_Standalone | RF_Public;
+		SaveArgs.SaveFlags = PackageHelper.UseSourceControl() ? ESaveFlags::SAVE_None : ESaveFlags::SAVE_Async;
+		if (!UPackage::SavePackage(Package, nullptr, *PackageFileName, SaveArgs))
+		{
+			UE_LOG(LogWorldPartitionBuilderSourceControlHelper, Error, TEXT("Error saving package %s."), *Package->GetName());
+			return false;
+		}
 	}
 
-	// Add new package to source control
-	if (!bFileExists)
+	// Add new packages to source control
+	if (PackagesToAdd.Num() > 0)
 	{
-		if (!Add(Package))
+		if (!Add(PackagesToAdd))
 		{
-			UE_LOG(LogWorldPartitionBuilderSourceControlHelper, Error, TEXT("Error adding package %s to revision control."), *Package->GetName());
+			UE_LOG(LogWorldPartitionBuilderSourceControlHelper, Error, TEXT("Error adding packages to revision control."));
 			return false;
 		}
 	}
@@ -235,4 +291,4 @@ bool FSourceControlHelper::Copy(const FString& SrcFilePath, const FString& DstFi
 	}
 
 	return true;
-}
\ No newline at end of file
+}
diff --git a/Engine/Source/Editor/UnrealEd/Public/WorldPartition/WorldPartitionBuilderHelpers.h b/Engine/Source/Editor/UnrealEd/Public/WorldPartition/WorldPartitionBuilderHelpers.h
--- a/Engine/Source/Editor/UnrealEd/Public/WorldPartition/WorldPartitionBuilderHelpers.h
+++ b/Engine/Source/Editor/UnrealEd/Public/WorldPartition/WorldPartitionBuilderHelpers.h
@@ -42,12 +42,16 @@ public:
 	virtual bool Delete(const FString& PackageName) const override;
 	virtual bool Delete(const TArray<FString>& PackageNames) const override;
 	virtual bool Delete(UPackage* Package) const override;
+	virtual bool Delete(const TArray<UPackage*>& Packages) const override;
 	virtual bool Save(UPackage* Package) const override;
+	virtual bool Save(const TArray<UPackage*>& Packages) const override;
 	virtual bool Copy(const FString& SrcFilePath, const FString& DstFilePath) const override;
 
 private:
 	bool Checkout(const FString& Filename) const;
+	bool Checkout(const TArray<FString>& Filenames) const;
 	bool Add(const FString& Filename) const;
+	bool Add(const TArray<FString>& Filenames) const;
 
 private:
 	FPackageSourceControlHelper& PackageHelper;
diff --git a/Engine/Source/Runtime/Engine/Private/WorldPartition/RuntimeHashSet/WorldPartitionRuntimeHashSetHLODGeneration.cpp b/Engine/Source/Runtime/Engine/Private/WorldPartition/RuntimeHashSet/WorldPartitionRuntimeHashSetHLODGeneration.cpp
--- a/Engine/Source/Runtime/Engine/Private/WorldPartition/RuntimeHashSet/WorldPartitionRuntimeHashSetHLODGeneration.cpp
+++ b/Engine/Source/Runtime/Engine/Private/WorldPartition/RuntimeHashSet/WorldPartitionRuntimeHashSetHLODGeneration.cpp
@@ -112,6 +112,21 @@ namespace PrivateUtils
 		}
 	}
 
+	static void SavePackages(const TArray<UPackage*>& Packages, ISourceControlHelper* SourceControlHelper)
+	{
+		if (SourceControlHelper)
+		{
+			SourceControlHelper->Save(Packages);
+		}
+		else
+		{
+			for (UPackage* Package : Packages)
+			{
+				SavePackage(Package, SourceControlHelper);
+			}
+		}
+	}
+
 	static void DeletePackage(const FString& PackageName, ISourceControlHelper* SourceControlHelper)
 	{
 		FPackagePath PackagePath = FPackagePath::FromPackageNameChecked(PackageName);
@@ -127,6 +142,29 @@ namespace PrivateUtils
 		}
 	}
 
+	static void DeletePackages(const TArray<FString>& PackageNames, ISourceControlHelper* SourceControlHelper)
+	{
+		TArray<FString> PackagesFullPaths;
+
+		Algo::Transform(PackageNames, PackagesFullPaths, [](const FString& PackageName)
+		{
+			FPackagePath PackagePath = FPackagePath::FromPackageNameChecked(PackageName);
+			return PackagePath.GetLocalFullPath();
+		});
+
+		if (SourceControlHelper)
+		{
+			SourceControlHelper->Delete(PackagesFullPaths);
+		}
+		else
+		{
+			for (const FString& PackageFullPath : PackagesFullPaths)
+			{
+				FPlatformFileManager::Get().GetPlatformFile().DeleteFile(*PackageFullPath);
+			}
+		}
+	}
+
 	static void DeletePackage(UPackage* Package, ISourceControlHelper* SourceControlHelper)
 	{
 		if (SourceControlHelper)
@@ -139,6 +177,21 @@ namespace PrivateUtils
 		}
 	}
 
+	static void DeletePackages(const TArray<UPackage*>& Packages, ISourceControlHelper* SourceControlHelper)
+	{
+		if (SourceControlHelper)
+		{
+			SourceControlHelper->Delete(Packages);
+		}
+		else
+		{
+			for (UPackage* Package : Packages)
+			{
+				DeletePackage(Package, SourceControlHelper);
+			}
+		}
+	}
+
 	static void DeletePackage(UWorldPartition* WorldPartition, const FWorldPartitionHandle& Handle, ISourceControlHelper* SourceControlHelper)
 	{
 		if (Handle.IsLoaded())
@@ -152,6 +205,45 @@ namespace PrivateUtils
 			WorldPartition->RemoveActor(Handle->GetGuid());
 		}
 	}
+
+	static void DeletePackages(UWorldPartition* WorldPartition, const TArray<FWorldPartitionHandle>& Handles, ISourceControlHelper* SourceControlHelper)
+	{
+		TArray<UPackage*> LoadedPackagesToDelete;
+		TArray<FString> UnloadedPackagesToDelete;
+		TArray<FGuid> UnloadedActorsToRemove;
+
+		for (const FWorldPartitionHandle& Handle : Handles)
+		{
+			if (Handle.IsLoaded())
+			{
+				LoadedPackagesToDelete.Add(Handle.GetActor()->GetPackage());
+			}
+			else
+			{
+				UnloadedPackagesToDelete.Add(Handle->GetActorPackage().ToString());
+				UnloadedActorsToRemove.Add(Handle->GetGuid());
+			}
+		}
+		check(UnloadedPackagesToDelete.Num() == UnloadedActorsToRemove.Num());
+
+		if (LoadedPackagesToDelete.Num())
+		{
+			DeletePackages(LoadedPackagesToDelete, SourceControlHelper);
+			for (UPackage* Package : LoadedPackagesToDelete)
+			{
+				WorldPartition->OnPackageDeleted(Package);
+			}
+		}
+
+		if (UnloadedPackagesToDelete.Num())
+		{
+			DeletePackages(UnloadedPackagesToDelete, SourceControlHelper);
+			for (const FGuid& ActorGuid : UnloadedActorsToRemove)
+			{
+				WorldPartition->RemoveActor(ActorGuid);
+			}
+		}
+	}
 }
 
 bool UWorldPartitionRuntimeHashSet::SupportsHLODs() const
@@ -214,6 +306,9 @@ bool UWorldPartitionRuntimeHashSet::SetupHLODActors(const IStreamingGenerationCo
 
 	TUniquePtr<IStreamingGenerationContext> CurrentHLODStreamingGenerationContext = MakeUnique<FStreamingGenerationContextProxy>(StreamingGenerationContext);
 
+	TArray<UPackage*> PackagesToSave;
+	TArray<FWorldPartitionReference> ActorsToSaveReferences;
+
 	int32 HLODLevel = 0;
 	while (CurrentHLODStreamingGenerationContext)
 	{
@@ -445,7 +540,8 @@ bool UWorldPartitionRuntimeHashSet::SetupHLODActors(const IStreamingGenerationCo
 						{
 							if (CellHLODActor->GetPackage()->IsDirty())
 							{
-								PrivateUtils::SavePackage(CellHLODActor->GetPackage(), Params.SourceControlHelper);
+								PackagesToSave.Add(CellHLODActor->GetPackage());
+								ActorsToSaveReferences.Emplace(WorldPartition, CellHLODActor->GetActorGuid());
 							}
 						}
 					}
@@ -548,9 +644,18 @@ bool UWorldPartitionRuntimeHashSet::SetupHLODActors(const IStreamingGenerationCo
 		}
 	}
 
+	if (PackagesToSave.Num() > 0)
+	{
+		PrivateUtils::SavePackages(PackagesToSave, Params.SourceControlHelper);
+	}
+	ActorsToSaveReferences.Empty();
+
 	// Destroy all unreferenced HLOD actors
 	if (!Params.bReportOnly)
 	{
+		TArray<UPackage*> PackagesToDelete;
+		TArray<FWorldPartitionHandle> ActorHandlesToDelete;
+
 		for (const auto& HLODActorPair : HLODCreationContext.HLODActorDescs)
 		{
 			check(HLODActorPair.Value.IsValid());
@@ -572,13 +677,13 @@ bool UWorldPartitionRuntimeHashSet::SetupHLODActors(const IStreamingGenerationCo
 							{	
 								if (HLODResourcesPackage != HLODActor->GetPackage())
 								{
-									PrivateUtils::DeletePackage(HLODResourcesPackage, Params.SourceControlHelper);
+									PackagesToDelete.Add(HLODResourcesPackage);
 								}
 							}
 						}
 					}
-
-					PrivateUtils::DeletePackage(WorldPartition, HLODActorPair.Value, Params.SourceControlHelper);
+					
+					ActorHandlesToDelete.Add(HLODActorPair.Value);
 				}
 				else
 				{
@@ -592,6 +697,16 @@ bool UWorldPartitionRuntimeHashSet::SetupHLODActors(const IStreamingGenerationCo
 				}
 			}
 		}
+
+		if (PackagesToDelete.Num())
+		{
+			PrivateUtils::DeletePackages(PackagesToDelete, Params.SourceControlHelper);
+		}
+
+		if (ActorHandlesToDelete.Num())
+		{
+			PrivateUtils::DeletePackages(WorldPartition, ActorHandlesToDelete, Params.SourceControlHelper);
+		}
 	}
 
 	return false;
diff --git a/Engine/Source/Runtime/Engine/Public/WorldPartition/WorldPartition.h b/Engine/Source/Runtime/Engine/Public/WorldPartition/WorldPartition.h
--- a/Engine/Source/Runtime/Engine/Public/WorldPartition/WorldPartition.h
+++ b/Engine/Source/Runtime/Engine/Public/WorldPartition/WorldPartition.h
@@ -113,7 +113,9 @@ public:
 	virtual bool Delete(const FString& PackageName) const =0;
 	virtual bool Delete(const TArray<FString>& PackageNames) const =0;
 	virtual bool Delete(UPackage* Package) const =0;
+	virtual bool Delete(const TArray<UPackage*>& Packages) const =0;
 	virtual bool Save(UPackage* Package) const =0;
+	virtual bool Save(const TArray<UPackage*>& Packages) const =0;
 	virtual bool Copy(const FString& SrcFilePath, const FString& DstFilePath) const =0;
 };
 #endif
