Fix for crash in USkeletalMeshComponent::ForEachBodyBelow when instantiated Bodies do not match the SkeletalBodySetups in the Reference Skeleton Mesh.

#Synced-CL 50214038

diff --git a/Engine/Source/Runtime/Engine/Private/SkeletalMeshComponentPhysics.cpp b/Engine/Source/Runtime/Engine/Private/SkeletalMeshComponentPhysics.cpp
--- a/Engine/Source/Runtime/Engine/Private/SkeletalMeshComponentPhysics.cpp
+++ b/Engine/Source/Runtime/Engine/Private/SkeletalMeshComponentPhysics.cpp
@@ -2243,6 +2243,13 @@ int32 USkeletalMeshComponent::ForEachBodyBelow(FName BoneName, bool bIncludeSelf
 		int32 NumBodiesFound = 0;
 		for (int32 BodyIdx : BodyIndices)
 		{
+			// GetBodyIndicesBelow does not take into account the actual Bodies we've created.
+			if (!Bodies.IsValidIndex(BodyIdx))
+			{
+				UE_LOG(LogSkeletalMesh, Warning, TEXT("%hs: assumed we had BodyIdx %d but we only have %d Bodies"), __FUNCTION__, BodyIdx, Bodies.Num());
+				break;
+			}
+
 			FBodyInstance* BI = Bodies[BodyIdx];
 			if (bSkipCustomType)
 			{
