[HLODBuilder] During the HLOD Setup step, save HLOD actors in separate batches after each HLOD level is processed.
This is required because actors from the previous level are needed to process the next level.

#Synced-CL 50093269

diff --git a/Engine/Source/Runtime/Engine/Private/WorldPartition/RuntimeHashSet/WorldPartitionRuntimeHashSetHLODGeneration.cpp b/Engine/Source/Runtime/Engine/Private/WorldPartition/RuntimeHashSet/WorldPartitionRuntimeHashSetHLODGeneration.cpp
--- a/Engine/Source/Runtime/Engine/Private/WorldPartition/RuntimeHashSet/WorldPartitionRuntimeHashSetHLODGeneration.cpp
+++ b/Engine/Source/Runtime/Engine/Private/WorldPartition/RuntimeHashSet/WorldPartitionRuntimeHashSetHLODGeneration.cpp
@@ -552,6 +552,13 @@ bool UWorldPartitionRuntimeHashSet::SetupHLODActors(const IStreamingGenerationCo
 			}
 		}
 
+		if (PackagesToSave.Num() > 0)
+		{
+			PrivateUtils::SavePackages(PackagesToSave, Params.SourceControlHelper);
+			PackagesToSave.Empty();
+		}
+		ActorsToSaveReferences.Empty();
+
 		CurrentHLODStreamingGenerationContext.Reset();
 
 		// Build the next HLOD generation context
@@ -644,12 +651,6 @@ bool UWorldPartitionRuntimeHashSet::SetupHLODActors(const IStreamingGenerationCo
 		}
 	}
 
-	if (PackagesToSave.Num() > 0)
-	{
-		PrivateUtils::SavePackages(PackagesToSave, Params.SourceControlHelper);
-	}
-	ActorsToSaveReferences.Empty();
-
 	// Destroy all unreferenced HLOD actors
 	if (!Params.bReportOnly)
 	{
