Fix CSV profiler check fail due to stat type mismatch

#Synced-CL 49921741

diff --git a/Engine/Source/Runtime/Engine/Private/Rendering/RayTracingGeometryManager.cpp b/Engine/Source/Runtime/Engine/Private/Rendering/RayTracingGeometryManager.cpp
--- a/Engine/Source/Runtime/Engine/Private/Rendering/RayTracingGeometryManager.cpp
+++ b/Engine/Source/Runtime/Engine/Private/Rendering/RayTracingGeometryManager.cpp
@@ -618,8 +618,8 @@ void FRayTracingGeometryManager::Tick(FRHICommandList& RHICmdList)
 
 		SET_MEMORY_STAT(STAT_RayTracingGeometryReferencedMemory, 0);
 		SET_MEMORY_STAT(STAT_RayTracingGeometryRequestedMemory, 0);
-		CSV_CUSTOM_STAT(RayTracingGeometry, ReferencedSizeMB, 0, ECsvCustomStatOp::Set);
-		CSV_CUSTOM_STAT(RayTracingGeometry, RequestedSizeMB, 0, ECsvCustomStatOp::Set);
+		CSV_CUSTOM_STAT(RayTracingGeometry, ReferencedSizeMB, 0.0f, ECsvCustomStatOp::Set);
+		CSV_CUSTOM_STAT(RayTracingGeometry, RequestedSizeMB, 0.0f, ECsvCustomStatOp::Set);
 	}
 	else if (bUsingReferenceBasedResidency)
 	{
