Improved GC pause functionality
- Move logic up to Engine level so it can work with existing GC logic. This simplifies resume handling and removes the need for delegates
- Pausing GC is now handled via UEngine::PauseGarbageCollection/ResumeGarbageCollection
- Whilst GC is paused, calls to UEngine::ConditionalCollectGarbage will early-out. TrimMemory calls from LoadMap will also be suppressed
- Fix an ensure in LoadMap (CheckAssetReferences was being called in the case where we were omitting the GC)

Added CVars:
- gc.PauseGCDeferredMapLoadTrim can be used to issue a deferred Trim in the event of a map load trim whilst paused
- gc.PauseGCForceNoFullPurgeOnResume can be used to suppress Full Purge requests issued whilst GC was paused

#Synced-CL 49921741

diff --git a/Engine/Source/Runtime/CoreUObject/Private/UObject/GarbageCollection.cpp b/Engine/Source/Runtime/CoreUObject/Private/UObject/GarbageCollection.cpp
--- a/Engine/Source/Runtime/CoreUObject/Private/UObject/GarbageCollection.cpp
+++ b/Engine/Source/Runtime/CoreUObject/Private/UObject/GarbageCollection.cpp
@@ -11,7 +11,6 @@
 #include "Containers/ConsumeAllMpmcQueue.h"
 #include "Containers/Deque.h"
 #include "Containers/StaticBitArray.h"
-#include "Containers/Ticker.h"
 #include "HAL/MallocLeakDetection.h"
 #include "HAL/ThreadSafeBool.h"
 #include "Misc/TimeGuard.h"
@@ -38,7 +37,6 @@
 #include "UObject/Package.h"
 #include "Async/ParallelFor.h"
 #include "ProfilingDebugging/CsvProfiler.h"
-#include "HAL/PlatformMemoryHelpers.h"
 #include "HAL/Runnable.h"
 #include "HAL/RunnableThread.h"
 #include "HAL/ThreadHeartBeat.h"
@@ -244,23 +242,6 @@ static FAutoConsoleVariableRef CIncrementalBeginDestroyGranularity(
 	ECVF_Default
 );
 
-static float GPauseGCTimeLimitSeconds = 30.0;
-static FAutoConsoleVariableRef CPauseGCTimeLimitSeconds(
-	TEXT("gc.PauseGCTimeLimitSeconds"),
-	GPauseGCTimeLimitSeconds,
-	TEXT("If GC is paused, it will resume automatically after this duration"),
-	ECVF_Default
-);
-
-static int32 GPauseGCFreeMemThresholdMB = 200;
-static FAutoConsoleVariableRef CPauseGCFreeMemThresholdMB(
-	TEXT("gc.PauseGCFreeMemThresholdMB"),
-	GPauseGCFreeMemThresholdMB,
-	TEXT("If GC is paused, it will resume automatically if free physical memory is less than this threshold"),
-	ECVF_Default
-);
-
-
 #if WITH_VERSE_VM || defined(__INTELLISENSE__)
 bool GEnableFrankenGC = true;
 static FAutoConsoleVariableRef CEnableFrankenGC(
@@ -6325,145 +6306,6 @@ bool UnhashUnreachableObjects(bool bUseTimeLimit, double TimeLimit)
 	return bTimeLimitReached;
 }
 
-class FGarbageCollectionPause
-{
-public:
-	static FGarbageCollectionPause& Get()
-	{
-		static FGarbageCollectionPause Inst;
-		return Inst;
-	}
-	bool IsPaused() const
-	{
-		return StartTime > 0.0;
-	}
-
-	bool Resume(bool bCleanupDelegate = true, bool bAllowCatchupGC = true)
-	{
-		check(IsInGameThread());
-		if (!IsPaused())
-		{
-			// Nothing to do
-			return false;
-		}
-		UE_LOG(LogGarbage, Log, TEXT("Garbage collection resumed"));
-		StartTime = 0.0;
-		if (bCleanupDelegate)
-		{
-			FTSTicker::GetCoreTicker().RemoveTicker(TickHandle);
-		}
-
-		// If we missed any GCs while we were paused then request one now
-		if (bAllowCatchupGC && MissedGcCount>0)
-		{
-			TryCollectGarbage(GARBAGE_COLLECTION_KEEPFLAGS, bFullPurgeRequestedWhilePaused);
-		}
-		return true;
-	}
-
-	bool ShouldResume() const
-	{
-		if (GPauseGCTimeLimitSeconds > 0.0 && FPlatformTime::Seconds() - StartTime > GPauseGCTimeLimitSeconds)
-		{
-			UE_LOG(LogGarbage, Log, TEXT("Resuming GC due to time limit: %.2fs"), GPauseGCTimeLimitSeconds);
-			return true;
-		}
-		if (GPauseGCFreeMemThresholdMB > 0 && PlatformMemoryHelpers::GetFrameMemoryStats().AvailablePhysical < ((uint64)GPauseGCFreeMemThresholdMB << 20ull))
-		{
-			UE_LOG(LogGarbage, Log, TEXT("Resuming GC due to mem threshold: %dMB"), GPauseGCFreeMemThresholdMB);
-			return true;
-		}
-		return false;
-	}
-
-	void Pause()
-	{
-		check(IsInGameThread());
-		if (IsPaused())
-		{
-			// Nothing to do
-			return;
-		}
-		MissedGcCount=0;
-		bFullPurgeRequestedWhilePaused = false;
-		StartTime = FPlatformTime::Seconds();
-
-		// Check this every frame to ensure the cvars are runtime modifyable
-		TickHandle = FTSTicker::GetCoreTicker().AddTicker(FTickerDelegate::CreateLambda([this](float) 
-			{
-				if (ShouldResume())
-				{
-					Resume(false);
-					return false;
-				}
-				return true;
-			}), 0.0f);
-
-		UE_LOG(LogGarbage, Log, TEXT("Garbage collection paused"));
-		if (GPauseGCFreeMemThresholdMB > 0)
-		{
-			UE_LOG(LogGarbage, Log, TEXT("  Free mem threshold: %dMB"), GPauseGCFreeMemThresholdMB);
-		}
-		if (GPauseGCTimeLimitSeconds > 0.0)
-		{
-			UE_LOG(LogGarbage, Log, TEXT("  Time limit: %.2fs"), GPauseGCTimeLimitSeconds);
-		}
-	}
-
-	// Called from GC. Returns true if we're paused but attempts to resume first if possible
-	// Note: if KeepFlags is anything other than GARBAGE_COLLECTION_KEEPFLAGS, will resume immediately
-	bool OnGc(EObjectFlags KeepFlags, bool bFullPurge)
-	{
-		if (!IsPaused())
-		{
-			return false;
-		}
-		if (ShouldResume())
-		{
-			Resume(true, false);
-			return false;
-		}
-		if (KeepFlags != GARBAGE_COLLECTION_KEEPFLAGS)
-		{
-			UE_LOG(LogGarbage, Warning, TEXT("GC Resumed due to non-standard KeepFlags: 0x%08X"), (uint32)KeepFlags);
-			Resume(true, false);
-			return false;
-		}
-		MissedGcCount++;
-		bFullPurgeRequestedWhilePaused |= bFullPurge;
-		return true;
-	}
-
-	int32 GetMissedGcCount() const
-	{
-		return MissedGcCount;
-	}
-
-private:
-	FGarbageCollectionPause() {}
-
-	FTSTicker::FDelegateHandle TickHandle;
-	double StartTime = 0.0;
-	int32 MissedGcCount = 0;
-	bool bFullPurgeRequestedWhilePaused = false;
-};
-
-void PauseGarbageCollection()
-{
-	FGarbageCollectionPause::Get().Pause();
-}
-
-bool ResumeGarbageCollection()
-{
-	return FGarbageCollectionPause::Get().Resume();
-}
-
-bool IsGarbageCollectionPaused()
-{
-	return FGarbageCollectionPause::Get().IsPaused();
-}
-
-
 void CollectGarbage(EObjectFlags KeepFlags, bool bPerformFullPurge)
 {
 	if (GIsInitialLoad)
@@ -6473,12 +6315,6 @@ void CollectGarbage(EObjectFlags KeepFlags, bool bPerformFullPurge)
 		return;
 	}
 
-	if (FGarbageCollectionPause::Get().OnGc(KeepFlags, bPerformFullPurge))
-	{
-		UE_LOG(LogGarbage, Log, TEXT("GC is paused. Skipped CollectGarbage (Full purge: %s)"), bPerformFullPurge ? TEXT("Enabled") : TEXT("Disabled"));
-		return;
-	}
-
 	// No other thread may be performing UObject operations while we're running
 	AcquireGCLock();
 
@@ -6504,13 +6340,6 @@ bool TryCollectGarbage(EObjectFlags KeepFlags, bool bPerformFullPurge)
 		return false;
 	}
 
-	if (FGarbageCollectionPause::Get().OnGc(KeepFlags, bPerformFullPurge))
-	{
-		// GC was paused. Only log on the first skip to avoid logspam, since this is often called every frame
-		UE_CLOG(FGarbageCollectionPause::Get().GetMissedGcCount()==1, LogGarbage, Log, TEXT("GC is paused. Skipped TryCollectGarbage (Full purge: %s)"), bPerformFullPurge ? TEXT("Enabled") : TEXT("Disabled"));
-		return false;
-	}
-
 	// No other thread may be performing UObject operations while we're running so try to acquire GC lock
 	if (UE::GC::GIsIncrementalReachabilityPending)
 	{
diff --git a/Engine/Source/Runtime/CoreUObject/Public/UObject/UObjectGlobals.h b/Engine/Source/Runtime/CoreUObject/Public/UObject/UObjectGlobals.h
--- a/Engine/Source/Runtime/CoreUObject/Public/UObject/UObjectGlobals.h
+++ b/Engine/Source/Runtime/CoreUObject/Public/UObject/UObjectGlobals.h
@@ -943,28 +943,6 @@ COREUOBJECT_API void CollectGarbage(EObjectFlags KeepFlags, bool bPerformFullPur
 */
 COREUOBJECT_API bool TryCollectGarbage(EObjectFlags KeepFlags, bool bPerformFullPurge = true);
 
-/**
- * Disables garbage collection until the time limit or free memory threshold are exceeded, or ResumeGarbageCollection is called
- * Cvars gc.PauseGCTimeLimitSeconds and gc.PauseGCFreeMemThresholdMB can be used to specify time and memory thresholds for 
- * auto-resume.
- *
- */
-COREUOBJECT_API void PauseGarbageCollection();
-
-/**
- * Re-enables garbage collection following a call to PauseGarbageCollection
- * 
- * @return true if garbage collection was paused or false otherwise 
- */
-COREUOBJECT_API bool ResumeGarbageCollection();
-
-/**
- * Checks if GC is paused
- *
- * @return true if garbage collection is paused or false otherwise
- */
-COREUOBJECT_API bool IsGarbageCollectionPaused();
-
 /**
 * Calls ConditionalBeginDestroy on unreachable objects
 *
diff --git a/Engine/Source/Runtime/Engine/Classes/Engine/Engine.h b/Engine/Source/Runtime/Engine/Classes/Engine/Engine.h
--- a/Engine/Source/Runtime/Engine/Classes/Engine/Engine.h
+++ b/Engine/Source/Runtime/Engine/Classes/Engine/Engine.h
@@ -3028,6 +3028,28 @@ public:
 		return ServiceLocator.ToSharedRef();
 	}
 
+	/**
+	 * Disables garbage collection until the time limit or free memory threshold are exceeded, or ResumeGarbageCollection is called
+	 * Cvars gc.PauseGCTimeLimitSeconds and gc.PauseGCFreeMemThresholdMB can be used to specify time and memory thresholds for
+	 * auto-resume.
+	 *
+	 */
+	ENGINE_API void PauseGarbageCollection();
+
+	/**
+	 * Re-enables garbage collection following a call to PauseGarbageCollection
+	 *
+	 * @return true if garbage collection was paused or false otherwise
+	 */
+	ENGINE_API bool ResumeGarbageCollection();
+
+	/**
+	 * Checks if GC is paused
+	 *
+	 * @return true if garbage collection is paused or false otherwise
+	 */
+	ENGINE_API bool IsGarbageCollectionPaused();
+
 protected:
 
 	/** Portal RPC client. */
diff --git a/Engine/Source/Runtime/Engine/Private/UnrealEngine.cpp b/Engine/Source/Runtime/Engine/Private/UnrealEngine.cpp
--- a/Engine/Source/Runtime/Engine/Private/UnrealEngine.cpp
+++ b/Engine/Source/Runtime/Engine/Private/UnrealEngine.cpp
@@ -1778,6 +1778,209 @@ static FAutoConsoleVariableRef CVarIncrementalGCTimePerFrame(
 bool GLogStackOnForceGC = false;
 #endif
 
+static float GPauseGCTimeLimitSeconds = 30.0;
+static FAutoConsoleVariableRef CPauseGCTimeLimitSeconds(
+	TEXT("gc.PauseGCTimeLimitSeconds"),
+	GPauseGCTimeLimitSeconds,
+	TEXT("If GC is paused, it will resume automatically after this duration"),
+	ECVF_Default
+);
+
+static int32 GPauseGCFreeMemThresholdMB = 300;
+static FAutoConsoleVariableRef CPauseGCFreeMemThresholdMB(
+	TEXT("gc.PauseGCFreeMemThresholdMB"),
+	GPauseGCFreeMemThresholdMB,
+	TEXT("If GC is paused, it will resume automatically if free physical memory is less than this threshold"),
+	ECVF_Default
+);
+
+static bool GPauseGCDeferredMapLoadTrim = false;
+static FAutoConsoleVariableRef CPauseGCDeferredMapLoadTrim(
+	TEXT("gc.PauseGCDeferredMapLoadTrim"),
+	GPauseGCDeferredMapLoadTrim,
+	TEXT("If true, LoadMap TrimMemory will be replicated after GC resumes.\n")
+	TEXT("If false, LoadMap TrimMemory will be skipped but we'll do a full purge GC on resume"),
+	ECVF_Default
+);
+
+static bool GPauseGCForceNoFullPurgeOnResume = false;
+static FAutoConsoleVariableRef CPauseGCForceNoFullPurgeOnResume(
+	TEXT("gc.PauseGCForceNoFullPurgeOnResume"),
+	GPauseGCForceNoFullPurgeOnResume,
+	TEXT("If true, clear the FullPurge flag on resume; any full purge GCs requested while paused will be cleared\n"),
+	ECVF_Default
+);
+
+
+enum class EGcRequestType : uint8
+{
+	Incremental,
+	FullPurge,
+	LoadMapTrim,
+};
+
+inline const TCHAR* LexToString(EGcRequestType Type)
+{
+	switch (Type)
+	{
+	case EGcRequestType::Incremental:		return TEXT("Incremental");
+	case EGcRequestType::FullPurge:			return TEXT("FullPurge");
+	case EGcRequestType::LoadMapTrim:		return TEXT("LoadMapTrim");
+	default:								return TEXT("Unknown");
+	}
+}
+
+// Handles pausing and resuming garbage collection
+class FGarbageCollectionPause
+{
+public:
+	static FGarbageCollectionPause& Get()
+	{
+		static FGarbageCollectionPause Inst;
+		return Inst;
+	}
+	bool IsPaused() const
+	{
+		return StartTime > 0.0;
+	}
+
+	bool Resume()
+	{
+		check(IsInGameThread());
+		if (!IsPaused())
+		{
+			// Nothing to do
+			return false;
+		}
+		// If a trim was requested then do it now
+		UE_LOG(LogGarbage, Log, TEXT("Garbage collection resumed"));
+		ResumeInternal();
+		if (GPauseGCDeferredMapLoadTrim && bTrimRequested && GEngine)
+		{
+			UE_LOG(LogGarbage, Log, TEXT("TrimMemory after GC resumed"));
+			GEngine->TrimMemory();
+		}
+		return true;
+	}
+
+	void ResumeInternal()
+	{
+		StartTime = 0.0;
+	}
+
+	bool ShouldResume() const
+	{
+		if (GPauseGCTimeLimitSeconds > 0.0 && FPlatformTime::Seconds() - StartTime > GPauseGCTimeLimitSeconds)
+		{
+			UE_LOG(LogGarbage, Log, TEXT("Resuming GC due to time limit: %.2fs"), GPauseGCTimeLimitSeconds);
+			return true;
+		}
+		if (GPauseGCFreeMemThresholdMB > 0 && PlatformMemoryHelpers::GetFrameMemoryStats().AvailablePhysical < ((uint64)GPauseGCFreeMemThresholdMB << 20ull))
+		{
+			UE_LOG(LogGarbage, Log, TEXT("Resuming GC due to mem threshold: %dMB"), GPauseGCFreeMemThresholdMB);
+			return true;
+		}
+		return false;
+	}
+
+	void Pause()
+	{
+		check(IsInGameThread());
+		if (IsPaused())
+		{
+			// Nothing to do
+			return;
+		}
+		MissedGcCount = 0;
+		bTrimRequested = false;
+		StartTime = FPlatformTime::Seconds();
+
+		UE_LOG(LogGarbage, Log, TEXT("Garbage collection paused"));
+		if (GPauseGCFreeMemThresholdMB > 0)
+		{
+			UE_LOG(LogGarbage, Log, TEXT("  Free mem threshold: %dMB"), GPauseGCFreeMemThresholdMB);
+		}
+		if (GPauseGCTimeLimitSeconds > 0.0)
+		{
+			UE_LOG(LogGarbage, Log, TEXT("  Time limit: %.2fs"), GPauseGCTimeLimitSeconds);
+		}
+	}
+
+	// Called from ConditionalCollectGarbage. Returns true if the GC should proceed (usually the case if not paused)
+	// Will unpause if conditions are met
+	bool OnGc(EGcRequestType GcType)
+	{
+		if (!IsPaused())
+		{
+			return true;
+		}
+
+		// Check if we should auto-resume at this point (ie under free mem threshold/over time limit)
+		if ( ShouldResume() )
+		{
+			ResumeInternal();
+			UE_LOG(LogGarbage, Log, TEXT("Garbage collection resumed (automatic)"));
+
+			// If a trim was requested and we're not about to trim then do one one now (this prevents an in-flight FullPurge request)
+			if (GPauseGCDeferredMapLoadTrim && bTrimRequested && GcType != EGcRequestType::LoadMapTrim && GEngine)
+			{
+				UE_LOG(LogGarbage, Log, TEXT("TrimMemory after GC resumed"));
+				GEngine->TrimMemory();
+				return false;
+			}
+
+			return true;
+		}
+
+		if (GcType == EGcRequestType::LoadMapTrim)
+		{
+			bTrimRequested = true;
+		}
+
+		UE_CLOG(MissedGcCount < 8 || GcType == EGcRequestType::LoadMapTrim, LogGarbage, Log, TEXT("GC is paused. Skipped GC (Type: %s)"), LexToString(GcType));
+		MissedGcCount++;
+		return false;
+	}
+
+	int32 GetMissedGcCount() const
+	{
+		return MissedGcCount;
+	}
+
+private:
+	FGarbageCollectionPause() {}
+
+	FTSTicker::FDelegateHandle TickHandle;
+	double StartTime = 0.0;
+	int32 MissedGcCount = 0;
+	bool bTrimRequested = false;
+};
+
+void UEngine::PauseGarbageCollection()
+{
+	FGarbageCollectionPause::Get().Pause();
+}
+
+bool UEngine::ResumeGarbageCollection()
+{
+	if (FGarbageCollectionPause::Get().Resume())
+	{
+		// Clear any pending full purges on resume, if configured
+		if (GPauseGCForceNoFullPurgeOnResume)
+		{
+			bFullPurgeTriggered = false;
+		}
+		return true;
+	}
+	return false;
+}
+
+bool UEngine::IsGarbageCollectionPaused()
+{
+	return FGarbageCollectionPause::Get().IsPaused();
+}
+
+
 void UEngine::SendWorldEndOfFrameUpdates()
 {
 	// Gather worlds that need EOF updates
@@ -1954,16 +2157,19 @@ void UEngine::ConditionalCollectGarbage()
 
 			if (bFullPurgeTriggered)
 			{
-				TGuardValue<bool> _FullGC(bGCPerformingFullPurge, true);
-				if (TryCollectGarbage(GARBAGE_COLLECTION_KEEPFLAGS, true))
+				if (FGarbageCollectionPause::Get().OnGc(EGcRequestType::FullPurge))
 				{
-					ForEachObjectOfClass(UWorld::StaticClass(),[](UObject* World)
+					TGuardValue<bool> _FullGC(bGCPerformingFullPurge, true);
+					if (TryCollectGarbage(GARBAGE_COLLECTION_KEEPFLAGS, true))
 					{
-						CastChecked<UWorld>(World)->CleanupActors();
-					});
-					bFullPurgeTriggered = false;
-					bShouldDelayGarbageCollect = false;
-					TimeSinceLastPendingKillPurge = 0.0f;
+						ForEachObjectOfClass(UWorld::StaticClass(), [](UObject* World)
+							{
+								CastChecked<UWorld>(World)->CleanupActors();
+							});
+						bFullPurgeTriggered = false;
+						bShouldDelayGarbageCollect = false;
+						TimeSinceLastPendingKillPurge = 0.0f;
+					}
 				}
 			}
 			else
@@ -2001,29 +2207,32 @@ void UEngine::ConditionalCollectGarbage()
 
 					const float TimeBetweenPurgingPendingKillObjects = GetTimeBetweenGarbageCollectionPasses(bHasPlayersConnected);
 
-					// See if we should delay garbage collect for this frame
-					if (bShouldDelayGarbageCollect)
-					{
-						bShouldDelayGarbageCollect = false;
-					}
-					else if (IsIncrementalReachabilityAnalysisPending())
-					{
-						SCOPE_CYCLE_COUNTER(STAT_GCMarkTime);
-						PerformIncrementalReachabilityAnalysis(GetReachabilityAnalysisTimeLimit());
-					}
-					// Perform incremental purge update if it's pending or in progress.
-					else if (!IsIncrementalPurgePending()
-						// Purge reference to pending kill objects every now and so often.
-						&& (TimeSinceLastPendingKillPurge > TimeBetweenPurgingPendingKillObjects) && TimeBetweenPurgingPendingKillObjects > 0.f)
-					{
-						SCOPE_CYCLE_COUNTER(STAT_GCMarkTime);
-						PerformGarbageCollectionAndCleanupActors();
-					}
-					else
+					if (FGarbageCollectionPause::Get().OnGc(EGcRequestType::Incremental))
 					{
-						SCOPE_CYCLE_COUNTER(STAT_GCSweepTime);
-						float IncGCTime = GetIncrementalGCTimePerFrame();
-						IncrementalPurgeGarbage(true, IncGCTime);
+						// See if we should delay garbage collect for this frame
+						if (bShouldDelayGarbageCollect)
+						{
+							bShouldDelayGarbageCollect = false;
+						}
+						else if (IsIncrementalReachabilityAnalysisPending())
+						{
+							SCOPE_CYCLE_COUNTER(STAT_GCMarkTime);
+							PerformIncrementalReachabilityAnalysis(GetReachabilityAnalysisTimeLimit());
+						}
+						// Perform incremental purge update if it's pending or in progress.
+						else if (!IsIncrementalPurgePending()
+							// Purge reference to pending kill objects every now and so often.
+							&& (TimeSinceLastPendingKillPurge > TimeBetweenPurgingPendingKillObjects) && TimeBetweenPurgingPendingKillObjects > 0.f)
+						{
+							SCOPE_CYCLE_COUNTER(STAT_GCMarkTime);
+							PerformGarbageCollectionAndCleanupActors();
+						}
+						else
+						{
+							SCOPE_CYCLE_COUNTER(STAT_GCSweepTime);
+							float IncGCTime = GetIncrementalGCTimePerFrame();
+							IncrementalPurgeGarbage(true, IncGCTime);
+						}
 					}
 				}
 			}
@@ -15960,9 +16169,15 @@ bool UEngine::LoadMap( FWorldContext& WorldContext, FURL URL, class UPendingNetG
 	}
 
 	FMoviePlayerProxy::BlockingTick();
-	// trim memory to clear up allocations from the previous level (also flushes rendering)
-	if (GDelayTrimMemoryDuringMapLoadMode == 0)
+	const bool bGcAllowed = FGarbageCollectionPause::Get().OnGc(EGcRequestType::LoadMapTrim);
+	if (!bGcAllowed)
+	{
+		// GC is paused so don't TrimMemory, but signal a full purge at the next GC
+		ForceGarbageCollection(true);
+	}
+	else if (GDelayTrimMemoryDuringMapLoadMode == 0)
 	{
+		// trim memory to clear up allocations from the previous level (also flushes rendering)
 		TrimMemory();
 	}
 
@@ -15978,8 +16193,8 @@ bool UEngine::LoadMap( FWorldContext& WorldContext, FURL URL, class UPendingNetG
 		appDumpTextureMemoryStats(TEXT(""));
 	}
 
-	// if we aren't trimming memory above, then the world won't be fully cleaned up at this point, so don't bother checking
-	if (GDelayTrimMemoryDuringMapLoadMode == 0)
+	// if we aren't trimming memory above or GC is paused, then the world won't be fully cleaned up at this point, so don't bother checking
+	if (GDelayTrimMemoryDuringMapLoadMode == 0 && bGcAllowed)
 	{
 		// Dump info
 		CheckAndHandleStaleWorldObjectReferences(&WorldContext);
@@ -16348,7 +16563,7 @@ bool UEngine::LoadMap( FWorldContext& WorldContext, FURL URL, class UPendingNetG
 	WorldContext.OwningGameInstance->LoadComplete(StopTime - StartTime, *URL.Map);
 
 	// perform the delayed TrimMemory if desired
-	if (GDelayTrimMemoryDuringMapLoadMode != 0)
+	if (GDelayTrimMemoryDuringMapLoadMode != 0 && bGcAllowed)
 	{
 		TrimMemory();
 
@@ -16371,6 +16586,10 @@ void UEngine::TrimMemory()
 		if (GEngine != nullptr)
 		{
 			GEngine->bGCPerformingFullPurge = true;
+			
+			// Reset this to avoid a redundant full purge if there's one pending
+			GEngine->bFullPurgeTriggered = false;
+			GEngine->TimeSinceLastPendingKillPurge = 0.0f;
 		}
 
 		CollectGarbage(GARBAGE_COLLECTION_KEEPFLAGS, true);
