Fix non-linear blend space playback when blending samples of differing play rates

* Currently, GetAnimationLengthFromSampleData calculated the animation length by linearly blending the scaled animation length of each sample.
* If sample A has a playrate of 1.0 and sample B has a playrate of 0.01 and both samples use the same animation of 4s then sample A will have an effective length of 4s and sample B will have an effective length of 400s.
* If the preview sample point if 50% of A and 50% of B then the calculated animation length is 202s (0.5 * 4 + 0.5 * 400) which is non-linear between A and B as we are blending inverses.
* However, we expect a length of approx 8s.
* The current code turns each sample into an effective length and blends thise linearly but the effective length is an inverse quantity and linearly interpolative an inverse is non-linear.
* The fix is to blend how fast each sample advances and then perform the inverse.
* Given the same example, sample A contributes an effective speed of 0.25, sample B contrinues an effective speed of 0.0025. With a weighting of 50:50 we get a blended speed of 0.12625 which we take the reciprocal to get the length of 7.9s.

#Synced-CL 49849479

diff --git a/Engine/Source/Runtime/Core/Public/UObject/FortniteReleaseBranchCustomObjectVersions.inl b/Engine/Source/Runtime/Core/Public/UObject/FortniteReleaseBranchCustomObjectVersions.inl
--- a/Engine/Source/Runtime/Core/Public/UObject/FortniteReleaseBranchCustomObjectVersions.inl
+++ b/Engine/Source/Runtime/Core/Public/UObject/FortniteReleaseBranchCustomObjectVersions.inl
@@ -63,5 +63,8 @@ FORTNITE_RELEASE_VERSION(CommonUIPlatformNamingUpgradeOption, 17)
 // Add new fictitious forces modes to the cloth simulation
 FORTNITE_RELEASE_VERSION(ChaosClothAddNewFictitiousForcesModes, 18)
 
+// Updated how blend spaces calculate the effective animation length of the sample point
+FORTNITE_RELEASE_VERSION(BlendSpaceUpdatedSamplePointAnimationLengthDerivation, 19)
+
 // -----<new versions can be added above this line>-------------------------------------------------
 #undef FORTNITE_RELEASE_VERSION
\ No newline at end of file
diff --git a/Engine/Source/Runtime/Engine/Classes/Animation/BlendSpace.h b/Engine/Source/Runtime/Engine/Classes/Animation/BlendSpace.h
--- a/Engine/Source/Runtime/Engine/Classes/Animation/BlendSpace.h
+++ b/Engine/Source/Runtime/Engine/Classes/Animation/BlendSpace.h
@@ -839,6 +839,14 @@ public:
 	/** If true, all follower samples will pass the same amount of markers the leader sample has passed to match its sync phase. Otherwise, followers samples will only match their next valid sync position. */
 	UPROPERTY(EditAnywhere, Category = Animation)
 	bool bShouldMatchSyncPhases = false;
+
+	/** Enable to use the old method for calculating the sample point's effective animation length.
+	The legacy method was to calculate each sample's effective animation length and perform a weighted blend to calculate the animation length for the preview.
+	This means the result is non-linear with respect to each sample's play rate as play rate is the inverse of animation length.
+	The new method is to perform a weighted blend on the play rates directly and then derive the animation length.
+	*/
+	UPROPERTY(EditAnywhere, Category = Animation)
+	bool bUseLegacySamplePointAnimationLengthCalculations = false;
 	
 #if WITH_EDITORONLY_DATA
 	/** Preview Base pose for additive BlendSpace **/
diff --git a/Engine/Source/Runtime/Engine/Private/Animation/BlendSpace.cpp b/Engine/Source/Runtime/Engine/Private/Animation/BlendSpace.cpp
--- a/Engine/Source/Runtime/Engine/Private/Animation/BlendSpace.cpp
+++ b/Engine/Source/Runtime/Engine/Private/Animation/BlendSpace.cpp
@@ -15,6 +15,7 @@
 #include "Animation/SkeletonRemapping.h"
 #include "Animation/SkeletonRemappingRegistry.h"
 #include "Math/UnrealMathUtility.h"
+#include "UObject/FortniteReleaseBranchCustomObjectVersion.h"
 #include "UObject/FrameworkObjectVersion.h"
 #include "UObject/UE5MainStreamObjectVersion.h"
 #include "UObject/UObjectIterator.h"
@@ -65,15 +66,21 @@ void UBlendSpace::PostLoad()
 #if WITH_EDITOR	
 	// Only do this during editor time (could alter the blendspace data during runtime otherwise) 
 	ValidateSampleData();
-#endif // WITH_EDITOR
 
 	InitializePerBoneBlend();
+
+	if (GetLinkerCustomVersion(FFortniteReleaseBranchCustomObjectVersion::GUID) < FFortniteReleaseBranchCustomObjectVersion::BlendSpaceUpdatedSamplePointAnimationLengthDerivation)
+	{
+		bUseLegacySamplePointAnimationLengthCalculations = true;
+	}
+#endif // WITH_EDITOR
 }
 
 void UBlendSpace::Serialize(FArchive& Ar)
 {
 	Ar.UsingCustomVersion(FFrameworkObjectVersion::GUID);
 	Ar.UsingCustomVersion(FUE5MainStreamObjectVersion::GUID);
+	Ar.UsingCustomVersion(FFortniteReleaseBranchCustomObjectVersion::GUID);
 	Super::Serialize(Ar);
 
 #if WITH_EDITOR
@@ -2037,27 +2044,64 @@ void UBlendSpace::TickFollowerSamples(TArray<FBlendSampleData>& SampleDataList,
 
 float UBlendSpace::GetAnimationLengthFromSampleData(const TArray<FBlendSampleData>& SampleDataList) const
 {
-	float BlendAnimLength = 0.f;
-	for (int32 I = 0; I < SampleDataList.Num(); ++I)
+	if (bUseLegacySamplePointAnimationLengthCalculations)
 	{
-		const int32 SampleDataIndex = SampleDataList[I].SampleDataIndex;
-		if (SampleData.IsValidIndex(SampleDataIndex))
+		float BlendAnimLength = 0.f;
+		for (int32 I = 0; I < SampleDataList.Num(); ++I)
 		{
-			const FBlendSample& Sample = SampleData[SampleDataIndex];
-			if (Sample.Animation)
+			const int32 SampleDataIndex = SampleDataList[I].SampleDataIndex;
+			if (SampleData.IsValidIndex(SampleDataIndex))
 			{
-				//Use the SamplePlayRate from the SampleDataList, not the RateScale from SampleData as SamplePlayRate might contain
-				//Multiple samples contribution which we would otherwise lose
-				const float MultipliedSampleRateScale = Sample.Animation->RateScale * SampleDataList[I].SamplePlayRate;
-				// apply rate scale to get actual playback time
-				BlendAnimLength += (Sample.GetSamplePlayLength() / ((MultipliedSampleRateScale) != 0.0f ? FMath::Abs(MultipliedSampleRateScale) : 1.0f)) * SampleDataList[I].GetClampedWeight();
+				const FBlendSample& Sample = SampleData[SampleDataIndex];
+				if (Sample.Animation)
+				{
+					//Use the SamplePlayRate from the SampleDataList, not the RateScale from SampleData as SamplePlayRate might contain
+					//Multiple samples contribution which we would otherwise lose
+					const float MultipliedSampleRateScale = Sample.Animation->RateScale * SampleDataList[I].SamplePlayRate;
+					// apply rate scale to get actual playback time
+					BlendAnimLength += (Sample.GetSamplePlayLength() / ((MultipliedSampleRateScale) != 0.0f ? FMath::Abs(MultipliedSampleRateScale) : 1.0f)) * SampleDataList[I].GetClampedWeight();
 
-				UE_LOG(LogAnimation, Verbose, TEXT("[%d] - Sample Animation(%s) : Weight(%0.5f) "), I + 1, *Sample.Animation->GetName(), SampleDataList[I].GetClampedWeight());
+					UE_LOG(LogAnimation, Verbose, TEXT("[%d] - Sample Animation(%s) : Weight(%0.5f) "), I + 1, *Sample.Animation->GetName(), SampleDataList[I].GetClampedWeight());
+				}
 			}
 		}
+
+		return BlendAnimLength;
 	}
+	else
+	{
+		float BlendedNormalizedSpeed = 0.0f;
 
-	return BlendAnimLength;
+		for (int32 I = 0; I < SampleDataList.Num(); ++I)
+		{
+			const int32 SampleDataIndex = SampleDataList[I].SampleDataIndex;
+			if (SampleData.IsValidIndex(SampleDataIndex))
+			{
+				const FBlendSample& Sample = SampleData[SampleDataIndex];
+				if (Sample.Animation)
+				{
+					const float SamplePlayRate = Sample.Animation->RateScale * SampleDataList[I].SamplePlayRate;
+					const float SampleAnimLength = Sample.GetSamplePlayLength();
+					const float SampleWeight = SampleDataList[I].GetClampedWeight();
+
+					const float SampleNormalizedSpeed = (SampleAnimLength != 0.0f) ? (SamplePlayRate / SampleAnimLength) : 0.0f;
+					BlendedNormalizedSpeed += SampleNormalizedSpeed * SampleWeight;
+
+					UE_LOG(LogAnimation, Verbose, TEXT("[%d] - Sample Animation(%s) : Weight(%0.5f) "), I + 1, *Sample.Animation->GetName(), SampleDataList[I].GetClampedWeight());
+				}
+			}
+		}
+
+		if (BlendedNormalizedSpeed == 0.0f)
+		{
+			return 0.0f;
+		}
+		else
+		{
+			const float BlendedAnimLength = 1.0f / BlendedNormalizedSpeed;
+			return BlendedAnimLength;
+		}
+	}
 }
 
 FVector UBlendSpace::GetClampedBlendInput(const FVector& BlendInput)  const
