UE: Providing more logging info when a NaN is found in the bounds of a static mesh. Also changing distance field bounds NaN message verbosity to 'Warning'.

#Synced-CL 50221024

diff --git a/Engine/Source/Runtime/Engine/Private/StaticMesh.cpp b/Engine/Source/Runtime/Engine/Private/StaticMesh.cpp
--- a/Engine/Source/Runtime/Engine/Private/StaticMesh.cpp
+++ b/Engine/Source/Runtime/Engine/Private/StaticMesh.cpp
@@ -2652,7 +2652,7 @@ void FStaticMeshRenderData::Serialize(FArchive& Ar, UStaticMesh* Owner, bool bCo
 #if !UE_BUILD_SHIPPING
 						if (LOD.DistanceFieldData->LocalSpaceMeshBounds.ContainsNaN())
 						{
-							UE_LOG(LogStaticMesh, Display, TEXT("FStaticMeshRenderData::Serialize found NaN in DistanceFieldData of '%s'. LOD:%d"), *GetPathNameSafe(Owner), ResourceIndex);
+							UE_LOG(LogStaticMesh, Warning, TEXT("FStaticMeshRenderData::Serialize found NaN in DistanceFieldData of '%s'. LOD:%d"), *GetPathNameSafe(Owner), ResourceIndex);
 						}
 #endif // UE_BUILD_SHIPPING
 						LOD.DistanceFieldData->Serialize(Ar, Owner);
@@ -2663,6 +2663,13 @@ void FStaticMeshRenderData::Serialize(FArchive& Ar, UStaticMesh* Owner, bool bCo
 	}
 
 	Ar << Bounds;
+	
+#if !UE_BUILD_SHIPPING
+	if (Bounds.ContainsNaN())
+	{
+		UE_LOG(LogStaticMesh, Warning, TEXT("FStaticMeshRenderData::Serialize found NaN in Bounds of '%s'"), *GetPathNameSafe(Owner));
+	}
+#endif // !UE_BUILD_SHIPPING
 
 #if WITH_EDITORONLY_DATA
 	if (bHasNaniteFallbackMesh && Ar.IsCooking() && Ar.IsSaving())
