Change UMaterialInterface::GetReferencedTextures() editor implementation.
ReferencedDefaultTextures can be empty for material instances in editor and so we didn't get the parent material texture references returned.
This fixes FObjectCacheContext which, because of the bug, didn't include MICs in the Texture -> Material lookup.
#rnx
#Synced-CL 50244055

diff --git a/Engine/Source/Runtime/Engine/Private/Materials/MaterialInterface.cpp b/Engine/Source/Runtime/Engine/Private/Materials/MaterialInterface.cpp
--- a/Engine/Source/Runtime/Engine/Private/Materials/MaterialInterface.cpp
+++ b/Engine/Source/Runtime/Engine/Private/Materials/MaterialInterface.cpp
@@ -400,7 +400,11 @@ void UMaterialInterface::GetQualityLevelUsage(TArray<bool, TInlineAllocator<EMat
 
 TArrayView<const TObjectPtr<UObject>> UMaterialInterface::GetReferencedTextures() const
 {
+#if WITH_EDITORONLY_DATA
+	return GetDefaultTextures();
+#else
 	return ReferencedDefaultTextures;
+#endif
 }
 
 TConstArrayView<TObjectPtr<UTextureCollection>> UMaterialInterface::GetReferencedTextureCollections() const
