Adding logs for resolution changes

#rnx

#Synced-CL 50270914

diff --git a/Engine/Source/Runtime/Engine/Private/GameEngine.cpp b/Engine/Source/Runtime/Engine/Private/GameEngine.cpp
--- a/Engine/Source/Runtime/Engine/Private/GameEngine.cpp
+++ b/Engine/Source/Runtime/Engine/Private/GameEngine.cpp
@@ -488,6 +488,7 @@ void UGameEngine::DetermineGameWindowResolution( int32& ResolutionX, int32& Reso
 		// We need to pass the resolution back out to GameUserSettings, or it will just override it again
 		ResolutionX = DisplayMetrics.PrimaryDisplayWorkAreaRect.Right - DisplayMetrics.PrimaryDisplayWorkAreaRect.Left;
 		ResolutionY = DisplayMetrics.PrimaryDisplayWorkAreaRect.Bottom - DisplayMetrics.PrimaryDisplayWorkAreaRect.Top;
+		UE_LOG(LogEngine, Display, TEXT("UGameEngine::DetermineGameWindowResolution: calling RequestResolutionChange %dx%d %s"), ResolutionX, ResolutionY, LexToString(EWindowMode::Fullscreen));
 		FSystemResolution::RequestResolutionChange(ResolutionX, ResolutionY, EWindowMode::Fullscreen);
 	}
 
@@ -519,6 +520,7 @@ TSharedRef<SWindow> UGameEngine::CreateGameWindow()
 		if (ResX != GSystemResolution.ResX || ResY != GSystemResolution.ResY || WindowMode != GSystemResolution.WindowMode)
 		{
 			FSystemResolution::RequestResolutionChange(ResX, ResY, WindowMode);
+			UE_LOG(LogEngine, Display, TEXT("UGameEngine::CreateGameWindow: calling RequestResolutionChange %dx%d %s"), ResX, ResY, LexToString(WindowMode));
 			IConsoleManager::Get().CallAllConsoleVariableSinks();
 		}
 	}
@@ -864,6 +866,7 @@ void UGameEngine::OnViewportResized(FViewport* Viewport, uint32 Unused)
 		{
 			GSystemResolution.ResX = ViewportSize.X;
 			GSystemResolution.ResY = ViewportSize.Y;
+			UE_LOG(LogEngine, Display, TEXT("UGameEngine::OnViewportResized: calling RequestResolutionChange %dx%d %s"), GSystemResolution.ResX, GSystemResolution.ResY, LexToString(EWindowMode::Windowed));
 			FSystemResolution::RequestResolutionChange(GSystemResolution.ResX, GSystemResolution.ResY, EWindowMode::Windowed);
 
 			UGameUserSettings* Settings = GetGameUserSettings();
diff --git a/Engine/Source/Runtime/Engine/Private/GameUserSettings.cpp b/Engine/Source/Runtime/Engine/Private/GameUserSettings.cpp
--- a/Engine/Source/Runtime/Engine/Private/GameUserSettings.cpp
+++ b/Engine/Source/Runtime/Engine/Private/GameUserSettings.cpp
@@ -574,6 +574,7 @@ void UGameUserSettings::ApplyResolutionSettings(bool bCheckForCommandLineOverrid
 
 	// Request a resolution change
 	const int32 MonitorIndex = GetClosestMonitorIndex();
+	UE_LOG(LogGameUserSettings, Display, TEXT("UGameUserSettings::ApplyResolutionSettings: calling RequestResolutionChange %dx%d %s"), ResolutionSizeX, ResolutionSizeY, LexToString(NewFullscreenMode));
 	RequestResolutionChange(ResolutionSizeX, ResolutionSizeY, NewFullscreenMode, bCheckForCommandLineOverrides, MonitorIndex);
 
 	if (NewFullscreenMode == EWindowMode::Fullscreen || NewFullscreenMode == EWindowMode::WindowedFullscreen)
@@ -664,6 +665,8 @@ void UGameUserSettings::RequestResolutionChange(int32 InResolutionX, int32 InRes
 	}
 
 	FSystemResolution::SetMonitorIndex(MonitorIndex);
+
+	UE_LOG(LogGameUserSettings, Display, TEXT("UGameUserSettings::RequestResolutionChange: calling RequestResolutionChange %dx%d %s"), InResolutionX, InResolutionY, LexToString(InWindowMode));
 	FSystemResolution::RequestResolutionChange(InResolutionX, InResolutionY, InWindowMode);
 }
 
@@ -811,6 +814,7 @@ void UGameUserSettings::PreloadResolutionSettings(bool bAllowCmdLineOverrides /*
 #endif
 
 	const int32 MonitorIndex = GetClosestMonitorFromIDAndIndex(DisplayID, DisplayIndex);
+	UE_LOG(LogGameUserSettings, Display, TEXT("UGameUserSettings::PreloadResolutionSettings: calling RequestResolutionChange %dx%d %s"), ResolutionX, ResolutionY, LexToString(WindowMode));
 	RequestResolutionChange(ResolutionX, ResolutionY, WindowMode, bAllowCmdLineOverrides, MonitorIndex);
 
 	IConsoleManager::Get().CallAllConsoleVariableSinks();
diff --git a/Engine/Source/Runtime/Engine/Private/GameViewportClient.cpp b/Engine/Source/Runtime/Engine/Private/GameViewportClient.cpp
--- a/Engine/Source/Runtime/Engine/Private/GameViewportClient.cpp
+++ b/Engine/Source/Runtime/Engine/Private/GameViewportClient.cpp
@@ -3965,6 +3965,7 @@ bool UGameViewportClient::SetDisplayConfiguration(const FIntPoint* Dimensions, E
 			NewY = Dimensions->Y;
 		}
 
+		UE_LOG(LogViewport, Display, TEXT("UGameViewportClient::SetDisplayConfiguration: calling RequestResolutionChange %dx%d %s"), NewX, NewY, LexToString(WindowMode));
 		FSystemResolution::RequestResolutionChange(NewX, NewY, WindowMode);
 	}
 
@@ -4040,6 +4041,7 @@ bool UGameViewportClient::HandleToggleFullscreenCommand()
 	if (!bNewModeApplied)
 	{
 		FSystemResolution::SetMonitorIndex(MonitorIndex);
+		UE_LOG(LogViewport, Display, TEXT("UGameViewportClient::HandleToggleFullscreenCommand: calling RequestResolutionChange %dx%d %s"), ResolutionX, ResolutionY, LexToString(FullScreenMode));
 		FSystemResolution::RequestResolutionChange(ResolutionX, ResolutionY, FullScreenMode);
 	}
 
@@ -4076,6 +4078,7 @@ bool UGameViewportClient::HandleSetResCommand( const TCHAR* Cmd, FOutputDevice&
 		}
 		if( X && Y )
 		{
+			UE_LOG(LogViewport, Display, TEXT("UGameEngine::DetermineGameWindowResolution: calling RequestResolutionChange %dx%d %s"), X, Y, LexToString(WindowMode));
 			FSystemResolution::RequestResolutionChange(X, Y, WindowMode);
 		}
 	}
diff --git a/Engine/Source/Runtime/Engine/Private/UnrealEngine.cpp b/Engine/Source/Runtime/Engine/Private/UnrealEngine.cpp
--- a/Engine/Source/Runtime/Engine/Private/UnrealEngine.cpp
+++ b/Engine/Source/Runtime/Engine/Private/UnrealEngine.cpp
@@ -18868,6 +18868,14 @@ void FSystemResolution::RequestResolutionChange(int32 InResX, int32 InResY, EWin
 	CVarSystemResolution->SetWithCurrentPriority(*NewValue);
 }
 
+void FSystemResolution::ForceRefresh()
+{
+	SetMonitorIndex(MonitorIndex);
+	UE_LOG(LogEngine, Display, TEXT("FSystemResolution::ForceRefresh: calling RequestResolutionChange %dx%d %s"), ResX, ResY, LexToString(WindowMode));
+	RequestResolutionChange(ResX, ResY, WindowMode);
+	bForceRefresh = true;
+}
+
 //////////////////////////////////////////////////////////////////////////
 // STATS
 void UEngine::ExecEngineStat(UWorld* World, FCommonViewportClient* ViewportClient, const TCHAR* InName)
diff --git a/Engine/Source/Runtime/Engine/Public/UnrealEngine.h b/Engine/Source/Runtime/Engine/Public/UnrealEngine.h
--- a/Engine/Source/Runtime/Engine/Public/UnrealEngine.h
+++ b/Engine/Source/Runtime/Engine/Public/UnrealEngine.h
@@ -459,12 +459,7 @@ struct FSystemResolution
 	// when the console variable sinks are called
 	static ENGINE_API void RequestResolutionChange(int32 InResX, int32 InResY, EWindowMode::Type InWindowMode);
 
-	void ForceRefresh()
-	{
-		SetMonitorIndex(MonitorIndex);
-		RequestResolutionChange(ResX, ResY, WindowMode);
-		bForceRefresh = true;
-	}
+	ENGINE_API void ForceRefresh();
 };
 
 ENGINE_API extern FSystemResolution GSystemResolution;
