Fixed a bug where not all pso events where correctly tracked for particle systems

#rnx

#Synced-CL 50186389

diff --git a/Engine/Source/Runtime/Engine/Classes/Particles/ParticleSystem.h b/Engine/Source/Runtime/Engine/Classes/Particles/ParticleSystem.h
--- a/Engine/Source/Runtime/Engine/Classes/Particles/ParticleSystem.h
+++ b/Engine/Source/Runtime/Engine/Classes/Particles/ParticleSystem.h
@@ -150,7 +150,7 @@ protected:
 	bool PSOPrecachingLaunched = false;
 	FGraphEventRef PrecachePSOsEvent;
 	TArray<FMaterialPSOPrecacheRequestID> MaterialPSOPrecacheRequestIDs;
-	TMap<FObjectKey, FGraphEventRef> MaterialPSORequests;
+	TMap<FObjectKey, FGraphEventArray> MaterialPSORequests;
 };
 
 /**
diff --git a/Engine/Source/Runtime/Engine/Private/Particles/ParticleSystem.cpp b/Engine/Source/Runtime/Engine/Private/Particles/ParticleSystem.cpp
--- a/Engine/Source/Runtime/Engine/Private/Particles/ParticleSystem.cpp
+++ b/Engine/Source/Runtime/Engine/Private/Particles/ParticleSystem.cpp
@@ -35,18 +35,22 @@ void UFXSystemAsset::LaunchPSOPrecaching(const FMaterialInterfacePSOPrecachePara
 	FGraphEventArray PrecachePSOsEvents;
 	if (IsComponentPSOPrecachingEnabled())
 	{
-		PrecacheMaterialPSOs(PSOPrecacheParamsList, MaterialPSOPrecacheRequestIDs, PrecachePSOsEvents);
 		for (const FMaterialInterfacePSOPrecacheParams& MaterialPSOPrecacheParams : PSOPrecacheParamsList)
 		{
 			if (MaterialPSOPrecacheParams.MaterialInterface)
 			{
-				MaterialPSORequests.Add(MaterialPSOPrecacheParams.MaterialInterface, PrecachePSOsEvents.Last());
+				FGraphEventArray MaterialPSOsEvents = MaterialPSOPrecacheParams.MaterialInterface->PrecachePSOs(MaterialPSOPrecacheParams.VertexFactoryDataList, MaterialPSOPrecacheParams.PSOPrecacheParams, MaterialPSOPrecacheParams.Priority, MaterialPSOPrecacheRequestIDs);
+				if (!MaterialPSOsEvents.IsEmpty())
+				{
+					MaterialPSORequests.Add(MaterialPSOPrecacheParams.MaterialInterface, MaterialPSOsEvents);
+					PrecachePSOsEvents.Append(MaterialPSOsEvents);
+				}
 			}
 		}
 	} 
 
 	// Create task to signal that the PSO precache events are done by adding them as prerequisite to the task.
-	if (PrecachePSOsEvents.Num() > 0)
+	if (!PrecachePSOsEvents.IsEmpty())
 	{
 		struct FReleasePrecachePSOsEventTask
 		{
@@ -86,7 +90,14 @@ bool UFXSystemAsset::IsLoadingPSOCache(UMaterialInterface* Material) const
 	{
 		return false;
 	}
-	return MaterialPSORequests[Material]->IsCompleted() == false;
+	for (const FGraphEventRef& Task : MaterialPSORequests[Material])
+	{
+		if (!Task->IsCompleted())
+		{
+			return true;
+		}
+	}
+	return false;
 }
 
 ////////////////////////////////////////////////////////////////////////////////////////////////////
