Some refactoring and cleanup of new material translator's IR to HLSL code.

#rnx
#Synced-CL 50155647

diff --git a/Engine/Source/Runtime/Engine/Private/Materials/MaterialIRToHLSLTranslator.cpp b/Engine/Source/Runtime/Engine/Private/Materials/MaterialIRToHLSLTranslator.cpp
--- a/Engine/Source/Runtime/Engine/Private/Materials/MaterialIRToHLSLTranslator.cpp
+++ b/Engine/Source/Runtime/Engine/Private/Materials/MaterialIRToHLSLTranslator.cpp
@@ -26,6 +26,8 @@
 
 #if WITH_EDITOR
 
+static const FStringView GVector4SwizzleSubset[4] = { TEXTVIEW(".x"), TEXTVIEW(".xy"), TEXTVIEW(".xyz"), {} };
+
 #define TAB "    "
 
 enum ENoOp { NoOp };
@@ -229,100 +231,30 @@ static FHLSLPrinter& operator<<(FHLSLPrinter& Printer, const MIR::FType& Type)
 	return Printer;
 }
 
-static const TCHAR* GetHLSLTypeString(EMaterialValueType Type)
-{
-	switch (Type)
-	{
-		case MCT_Float1: return TEXT("MaterialFloat");
-		case MCT_Float2: return TEXT("MaterialFloat2");
-		case MCT_Float3: return TEXT("MaterialFloat3");
-		case MCT_Float4: return TEXT("MaterialFloat4");
-		case MCT_Float: return TEXT("MaterialFloat");
-		case MCT_Texture2D: return TEXT("texture2D");
-		case MCT_TextureCube: return TEXT("textureCube");
-		case MCT_Texture2DArray: return TEXT("texture2DArray");
-		case MCT_VolumeTexture: return TEXT("volumeTexture");
-		case MCT_StaticBool: return TEXT("static bool");
-		case MCT_Bool:  return TEXT("bool");
-		case MCT_MaterialAttributes: return TEXT("FMaterialAttributes");
-		case MCT_TextureExternal: return TEXT("TextureExternal");
-		case MCT_TextureVirtual: return TEXT("TextureVirtual");
-		case MCT_VTPageTableResult: return TEXT("VTPageTableResult");
-		case MCT_ShadingModel: return TEXT("uint");
-		case MCT_UInt: return TEXT("uint");
-		case MCT_UInt1: return TEXT("uint");
-		case MCT_UInt2: return TEXT("uint2");
-		case MCT_UInt3: return TEXT("uint3");
-		case MCT_UInt4: return TEXT("uint4");
-		case MCT_Substrate: return TEXT("FSubstrateData");
-		case MCT_TextureCollection: return TEXT("FResourceCollection");
-		default: return TEXT("unknown");
-	};
-}
-
-static const TCHAR* GetShadingModelParameterName(EMaterialShadingModel InModel)
+// This struct encapsulates the full state and logic required to translate a single Material IR entry point into HLSL. It outputs
+// a the generated HLSL that evaluates a given entry point as described in the module.
+struct FEntryPointCodeGen
 {
-	switch (InModel)
+	// The resulting generated HLSL code snippets.
+	struct FResult
 	{
-		case MSM_Unlit: return TEXT("MATERIAL_SHADINGMODEL_UNLIT");
-		case MSM_DefaultLit: return TEXT("MATERIAL_SHADINGMODEL_DEFAULT_LIT");
-		case MSM_Subsurface: return TEXT("MATERIAL_SHADINGMODEL_SUBSURFACE");
-		case MSM_PreintegratedSkin: return TEXT("MATERIAL_SHADINGMODEL_PREINTEGRATED_SKIN");
-		case MSM_ClearCoat: return TEXT("MATERIAL_SHADINGMODEL_CLEAR_COAT");
-		case MSM_SubsurfaceProfile: return TEXT("MATERIAL_SHADINGMODEL_SUBSURFACE_PROFILE");
-		case MSM_TwoSidedFoliage: return TEXT("MATERIAL_SHADINGMODEL_TWOSIDED_FOLIAGE");
-		case MSM_Hair: return TEXT("MATERIAL_SHADINGMODEL_HAIR");
-		case MSM_Cloth: return TEXT("MATERIAL_SHADINGMODEL_CLOTH");
-		case MSM_Eye: return TEXT("MATERIAL_SHADINGMODEL_EYE");
-		case MSM_SingleLayerWater: return TEXT("MATERIAL_SHADINGMODEL_SINGLELAYERWATER");
-		case MSM_ThinTranslucent: return TEXT("MATERIAL_SHADINGMODEL_THIN_TRANSLUCENT");
-		default: UE_MIR_UNREACHABLE();
-	}
-}
+		// The code generated to evalute normal attribute needs to be placed in a separate slot than the other attributes (before).
+		// If the Normal SetMaterialOutput instruction is encountered, the code generated so far will be saved here and Printer.Buffer
+		// will contain the remaining code of this entry point.
+		FString NormalAttributeHLSL;
 
-static bool InstructionUsesPhiValue(const MIR::FInstruction* Instr, int32 EntryPointIndex)
-{
-	if (auto Branch = Instr->As<MIR::FBranch>())
-	{
-		return Branch->TrueBlock[EntryPointIndex].Instructions || Branch->FalseBlock[EntryPointIndex].Instructions;
-	}
-	return false;
-}
-
-static bool IsFoldable(const MIR::FInstruction* Instr, int32 EntryPointIndex)
-{
-	// Instructions that use Phi values cannot be folded/inlined.
-	if (InstructionUsesPhiValue(Instr, EntryPointIndex))
-	{
-		return false;
-	}
-
-	// Don't fold instructions that become very long to improve readability
-	switch (Instr->Kind)
-	{
-		case MIR::VK_SetMaterialOutput:
-		case MIR::VK_VTPageTableRead:
-		case MIR::VK_Call:
-		case MIR::VK_SubstrateHorizontalMixing:
-		case MIR::VK_SubstrateVerticalLayering:
-		case MIR::VK_SubstrateCoverageWeight:
-		case MIR::VK_SubstrateAdd:
-		case MIR::VK_SubstrateSelect:
-		case MIR::VK_SubstratePromoteToOperator:
-		case MIR::VK_SubstrateSlab:
-		case MIR::VK_SubstrateShadingModels:
-			return false;
-		default:
-			return true;
-	}
-}
+		// The code generated to evaluate all other non-Normal attribute outputs.
+		FString NonNormalHLSL;
+	};
 
-static const FStringView GVector4SwizzleSubset[4] = { TEXTVIEW(".x"), TEXTVIEW(".xy"), TEXTVIEW(".xyz"), {} };
+	// The Material IR module, which contains the intermediate representation of the material graph.
+	const FMaterialIRModule* Module{};
 
-struct FPrivate : FMaterialIRToHLSLTranslation
-{
 	// Utility for emitting HLSL code
-	FHLSLPrinter Printer{}; 
+	FHLSLPrinter Printer{};
+
+	// Used internally. See FResult::OutputNormalAttributeCode
+	FString OutputNormalAttributeCode;
 
 	// Number of local variables generated during translation.
 	int32 NumLocals{};
@@ -344,71 +276,26 @@ struct FPrivate : FMaterialIRToHLSLTranslation
 	// Used to increase generated HLSL readability, so that we add a newline before the first instruction that follows a SetMateiralOutput.
 	bool bLastInstructionWasSetMaterialOutput{};
 
-	// Assignments to the HLSL parameters in the material template shader.
-	TMap<FString, FString>& OutputShaderTemplateParams;
-	
-	// Constructor
-	FPrivate(FMaterialIRToHLSLTranslation& Base, TMap<FString, FString>& OutputShaderTemplateParams)
-	: FMaterialIRToHLSLTranslation{ Base }
-	, OutputShaderTemplateParams{ OutputShaderTemplateParams }
+	FEntryPointCodeGen(const FMaterialIRModule* Module)
+	: Module { Module }
 	{}
 
 	// Generates the full HLSL of the specified entry point.
-	const FMaterialIRModule::FEntryPoint& GenerateEntryPoint(int32 EntryPointIndex)
+	FResult GenerateEntryPoint(int32 ModuleEntryPointIndex, bool bPreviousFrame = false)
 	{
 		Printer = { 1 };
 
-		const FMaterialIRModule::FEntryPoint& EntryPoint = Module->GetEntryPoint(EntryPointIndex);
+		const FMaterialIRModule::FEntryPoint& EntryPoint = Module->GetEntryPoint(ModuleEntryPointIndex);
 
-		CurrentEntryPointIndex = EntryPointIndex;
+		CurrentEntryPointIndex = ModuleEntryPointIndex;
 		CurrentStage = EntryPoint.Stage;
+		bCompilingPreviousFrame = bPreviousFrame;
 
 		LowerBlock(EntryPoint.RootBlock);
 
-		return EntryPoint;
+		return { MoveTemp(OutputNormalAttributeCode), MoveTemp(Printer.Buffer) };
 	}
-	
-	// Generates the HLSL that sets the outputs of the WPO vertex stage.	
-	void GenerateVertexStageHLSL()
-	{
-		// Generate current frame vertex stage HLSL
-		GenerateEntryPoint(MIR::Stage_Vertex);
-		OutputShaderTemplateParams.Add(TEXT("evaluate_vertex_material_attributes"), MoveTemp(Printer.Buffer));
-
-		// Generate previous frame vertex stage HLSL
-		bCompilingPreviousFrame = true;
-		GenerateEntryPoint(MIR::Stage_Vertex);
-		bCompilingPreviousFrame = false;
-		OutputShaderTemplateParams.Add(TEXT("evaluate_vertex_material_attributes_prev"), MoveTemp(Printer.Buffer));
-	}
-
-	// Generates the HLSL that sets the outputs of a non-vertex stage.
-	void GenerateNonVertexStageHLSL(MIR::EStage Stage)
-	{
-		GenerateEntryPoint(Stage);
-
-		if (Stage == MIR::Stage_Pixel)
-		{
-			Printer << TEXTVIEW("PixelMaterialInputs.Subsurface = 0") << EndOfStatement;
-		}
 
-		FString ParamName = FString((Stage == MIR::Stage_Pixel) ? TEXT("calc_pixel_material_inputs_other_inputs") : TEXT("calc_pixel_material_inputs_analytic_derivatives_other_inputs"));
-		OutputShaderTemplateParams.Add(ParamName, MoveTemp(Printer.Buffer));
-
-		// Substrate extra code. For instance the shared local bases generation.
-		const FMaterialIRModule::FSubstrateData& SubstrateData = Module->GetSubstrateData();
-
-		FString* PixelMaterialOtherInputs = OutputShaderTemplateParams.Find(ParamName);
-		if (PixelMaterialOtherInputs)
-		{
-			PixelMaterialOtherInputs->Append(SubstrateData.PixelInputInitializerValuesHLSL);
-		}
-		else
-		{
-			OutputShaderTemplateParams.Add(ParamName, SubstrateData.PixelInputInitializerValuesHLSL);
-		}
-	}
-	
 	ENoOp LowerBlock(const MIR::FBlock& Block)
 	{
 		int32 OldNumLocals = NumLocals;
@@ -417,10 +304,10 @@ struct FPrivate : FMaterialIRToHLSLTranslation
 			// If this is a inlinable instruction type and it has only one user, instead of allocating a local
 			// and evaluating the instruction into it, simply skip this instruction for now.
 			// It will be generated inlined by its only user later.
-            if (Instr->GetNumUsers(CurrentEntryPointIndex) == 1 && IsFoldable(Instr, CurrentEntryPointIndex))
+			if (Instr->GetNumUsers(CurrentEntryPointIndex) == 1 && IsFoldable(Instr, CurrentEntryPointIndex))
 			{
-                continue;
-            }
+				continue;
+			}
 
 			// If this isn't a set material output instruction but last one was, add a newline for better readability.
 			const MIR::FSetMaterialOutput* SetMaterialOutput = MIR::As<MIR::FSetMaterialOutput>(Instr);
@@ -439,18 +326,18 @@ struct FPrivate : FMaterialIRToHLSLTranslation
 				uint32 LocalIndex = NumLocals++;
 				
 				// Remember the mapping between this instruction and its local index
-                InstrToLocalIndex.Add(Instr, LocalIndex);
+				InstrToLocalIndex.Add(Instr, LocalIndex);
 
 				// Print the local declaration "<Type> _<LocalIndex>", e.g. "float4 _3"
-                Printer << Instr->Type << TEXTVIEW(" _") << LocalIndex;
+				Printer << Instr->Type << TEXTVIEW(" _") << LocalIndex;
 				
 				// If this instruction doesn't use a phi value, we'll immediately assign its local to its result so output the "="
 				// E.g., if Instr is a branch, we will set its value inside the generate "if" {} scopes, so no need for a "=" now.
-                if (!InstructionUsesPhiValue(Instr, CurrentEntryPointIndex))
+				if (!InstructionUsesPhiValue(Instr, CurrentEntryPointIndex))
 				{
-                    Printer << TEXTVIEW(" = ");
-                }
-            }
+					Printer << TEXTVIEW(" = ");
+				}
+			}
 
 			LowerInstruction(Instr);
 			
@@ -458,20 +345,17 @@ struct FPrivate : FMaterialIRToHLSLTranslation
 
 			// Store the code needed to evaluate the normal in a separate chunk than the other material attributes
 			// since this needs to be emitted before the others in the material template.
-			if (SetMaterialOutput && SetMaterialOutput->Property == MP_Normal) 
+			if (SetMaterialOutput && SetMaterialOutput->Property == MP_Normal)
 			{
 				// "Normal" is treated in a special way because the rest of the attributes may lead back to reading it.
 				// Therefore, in the way MaterialTemplate.ush is structured, it needs to be evaluated before other attributes.
-				FString ParamName = (CurrentStage == MIR::Stage_Pixel)
-					? FString{ TEXTVIEW("calc_pixel_material_inputs_normal") }
-					: FString{ TEXTVIEW("calc_pixel_material_inputs_analytic_derivatives_normal") };
-				OutputShaderTemplateParams.Add(MoveTemp(ParamName), MoveTemp(Printer.Buffer));
+				OutputNormalAttributeCode = MoveTemp(Printer.Buffer);
 				Printer = { 1 };
 				bLastInstructionWasSetMaterialOutput = false;
 			}
-        }
+		}
 
-        NumLocals = OldNumLocals;
+		NumLocals = OldNumLocals;
 
 		return NoOp;
 	}
@@ -480,7 +364,7 @@ struct FPrivate : FMaterialIRToHLSLTranslation
 	{
 		// Instruction results may be shared among other dependent instructions.
 		if (const MIR::FInstruction* Instr = MIR::AsInstruction(Value))
-		{ 
+		{
 			// If this instruction has only one user (dependant) and is foldable, then generate the full HLSL for the instruction in place.
 			if (Instr->GetNumUsers(CurrentEntryPointIndex) <= 1 && IsFoldable(Instr, CurrentEntryPointIndex))
 			{
@@ -1166,7 +1050,7 @@ struct FPrivate : FMaterialIRToHLSLTranslation
 		Printer << TEXTVIEW("((") << Nop->Type << TEXTVIEW(")0)");
 	}
 
-	void LowerComposite(const MIR::FComposite* Composite) 
+	void LowerComposite(const MIR::FComposite* Composite)
 	{
 		TConstArrayView<MIR::FValue*> Components = Composite->GetComponents();
 		MIR::FPrimitive PrimitiveType = Composite->Type.GetPrimitive();
@@ -2329,457 +2213,539 @@ struct FPrivate : FMaterialIRToHLSLTranslation
 		}
 	}
 
-	/* Finalization */
+	/* helpers */
+	
+	static bool InstructionUsesPhiValue(const MIR::FInstruction* Instr, int32 EntryPointIndex)
+	{
+		if (auto Branch = Instr->As<MIR::FBranch>())
+		{
+			return Branch->TrueBlock[EntryPointIndex].Instructions || Branch->FalseBlock[EntryPointIndex].Instructions;
+		}
+		return false;
+	}
 
-	void GenerateCustomOutputDeclarations()
+	static bool IsFoldable(const MIR::FInstruction* Instr, int32 EntryPointIndex)
 	{
-		// Generate the custom output fields of the struct FMaterialPixelParameters. These fields will be computed and 
-		// set by the function that calculates all material outputs (e.g. BaseColor).
-		FHLSLPrinter VertexPrinter{ 1 };
-		FHLSLPrinter PixelPrinter{ 1 };
+		// Instructions that use Phi values cannot be folded/inlined.
+		if (InstructionUsesPhiValue(Instr, EntryPointIndex))
+		{
+			return false;
+		}
 
-		for (const MIR::FSetMaterialOutput* Output : Module->GetCustomOutputs())
+		// Don't fold instructions that become very long to improve readability
+		switch (Instr->Kind)
 		{
-			FHLSLPrinter& ThePrinter = Output->Stage == MIR::MaterialOutputStage_Vertex ? VertexPrinter : PixelPrinter;
-			ThePrinter << Output->Arg->Type << TEXT(' ') << *Output->Name << EndOfStatement;
+			case MIR::VK_SetMaterialOutput:
+			case MIR::VK_VTPageTableRead:
+			case MIR::VK_Call:
+			case MIR::VK_SubstrateHorizontalMixing:
+			case MIR::VK_SubstrateVerticalLayering:
+			case MIR::VK_SubstrateCoverageWeight:
+			case MIR::VK_SubstrateAdd:
+			case MIR::VK_SubstrateSelect:
+			case MIR::VK_SubstratePromoteToOperator:
+			case MIR::VK_SubstrateSlab:
+			case MIR::VK_SubstrateShadingModels:
+				return false;
+			default:
+				return true;
 		}
+	}
+};
+
+/////////////////////////////////////////////////////////////////////////////////////////////////////////////////////////////
+
+static void GenerateStageEntryPoints(const FMaterialIRModule* Module, TMap<FString, FString>& OutParameters)
+{
+	FEntryPointCodeGen CodeGen { Module };
+	FEntryPointCodeGen::FResult Result;
+
+	// Generate current frame vertex stage HLSL
+	Result = CodeGen.GenerateEntryPoint(MIR::Stage_Vertex, false);
+	OutParameters.Add(TEXT("evaluate_vertex_material_attributes"), MoveTemp(Result.NonNormalHLSL));
+
+	// Generate previous frame vertex stage HLSL
+	Result = CodeGen.GenerateEntryPoint(MIR::Stage_Vertex, true);
+	OutParameters.Add(TEXT("evaluate_vertex_material_attributes_prev"), MoveTemp(Result.NonNormalHLSL));
+
+	// Generate pixel stage HLSL
+	Result = CodeGen.GenerateEntryPoint(MIR::Stage_Pixel);
+	Result.NonNormalHLSL += TEXTVIEW(TAB "PixelMaterialInputs.Subsurface = 0;\n");
+	Result.NonNormalHLSL += CodeGen.Module->GetSubstrateData().PixelInputInitializerValuesHLSL;
+	OutParameters.Add(TEXT("calc_pixel_material_inputs_normal"), MoveTemp(Result.NormalAttributeHLSL));
+	OutParameters.Add(TEXT("calc_pixel_material_inputs_other_inputs"), MoveTemp(Result.NonNormalHLSL));
+
+	// Generate compute stage HLSL
+	Result = CodeGen.GenerateEntryPoint(MIR::Stage_Compute);
+	Result.NonNormalHLSL += CodeGen.Module->GetSubstrateData().PixelInputInitializerValuesHLSL;
+	OutParameters.Add(TEXT("calc_pixel_material_inputs_analytic_derivatives_normal"), MoveTemp(Result.NormalAttributeHLSL));
+	OutParameters.Add(TEXT("calc_pixel_material_inputs_analytic_derivatives_other_inputs"), MoveTemp(Result.NonNormalHLSL));
+}
+	
+static void GenerateCustomOutputDeclarations(const FMaterialIRModule* Module, TMap<FString, FString>& OutParameters)
+{
+	// Generate the custom output fields of the struct FMaterialPixelParameters. These fields will be computed and 
+	// set by the function that calculates all material outputs (e.g. BaseColor).
+	FHLSLPrinter VertexPrinter{ 1 };
+	FHLSLPrinter PixelPrinter{ 1 };
+
+	for (const MIR::FSetMaterialOutput* Output : Module->GetCustomOutputs())
+	{
+		FHLSLPrinter& ThePrinter = Output->Stage == MIR::MaterialOutputStage_Vertex ? VertexPrinter : PixelPrinter;
+		ThePrinter << Output->Arg->Type << TEXT(' ') << *Output->Name << EndOfStatement;
+	}
 		
-		OutputShaderTemplateParams.Add(TEXT("material_vertex_parameter_decls"), MoveTemp(VertexPrinter.Buffer));
-		OutputShaderTemplateParams.Add(TEXT("material_pixel_parameter_decls"), MoveTemp(PixelPrinter.Buffer));
+	OutParameters.Add(TEXT("material_vertex_parameter_decls"), MoveTemp(VertexPrinter.Buffer));
+	OutParameters.Add(TEXT("material_pixel_parameter_decls"), MoveTemp(PixelPrinter.Buffer));
 
-		// The material shader expects a getter function for each custom output. Generate these getter functions
-		// for each FMaterialPixelParameters field, generated above, matching each custom output in each group.
-		// Development Note: this is really just cruft necessary for backwards compatibility with the way the
-		// 	  old translator supported custom outputs and could be removed entirely with a little refactoring
-		//    of MaterialTemplate.ush.
-		Printer = {};
-		for (const FMaterialIRModule::FCustomOutputGroup& Group : Module->GetCustomOutputGroups())
-		{
-			Printer << TEXTVIEW("#define NUM_MATERIAL_OUTPUTS_GET") << *FString(Group.Name).ToUpper() << TEXT(' ') << Group.NumOutputs << NewLine;
+	// The material shader expects a getter function for each custom output. Generate these getter functions
+	// for each FMaterialPixelParameters field, generated above, matching each custom output in each group.
+	// Development Note: this is really just cruft necessary for backwards compatibility with the way the
+	// 	  old translator supported custom outputs and could be removed entirely with a little refactoring
+	//    of MaterialTemplate.ush.
+	FHLSLPrinter Printer = {};
+	for (const FMaterialIRModule::FCustomOutputGroup& Group : Module->GetCustomOutputGroups())
+	{
+		Printer << TEXTVIEW("#define NUM_MATERIAL_OUTPUTS_GET") << *FString(Group.Name).ToUpper() << TEXT(' ') << Group.NumOutputs << NewLine;
 
-			for (int i = 0; i < Group.NumOutputs; ++i)
-			{
-				const MIR::FSetMaterialOutput* Instr = Module->GetCustomOutputs()[i + Group.CustomOutputInstructionsBegin];
-				Printer << TEXTVIEW("#define HAVE_Get") << *Instr->Name << TEXTVIEW(" 1") << NewLine << NewLine;
+		for (int i = 0; i < Group.NumOutputs; ++i)
+		{
+			const MIR::FSetMaterialOutput* Instr = Module->GetCustomOutputs()[i + Group.CustomOutputInstructionsBegin];
+			Printer << TEXTVIEW("#define HAVE_Get") << *Instr->Name << TEXTVIEW(" 1") << NewLine << NewLine;
 				
-				// Select the right XXXParameters HLSL struct used by this output based on its 
-				FStringView Signature = Instr->Stage == MIR::MaterialOutputStage_Vertex
-					? TEXTVIEW("(inout FMaterialVertexParameters Parameters)")
-					: TEXTVIEW("(inout FMaterialPixelParameters Parameters)");
+			// Select the right XXXParameters HLSL struct used by this output based on its 
+			FStringView Signature = Instr->Stage == MIR::MaterialOutputStage_Vertex
+				? TEXTVIEW("(inout FMaterialVertexParameters Parameters)")
+				: TEXTVIEW("(inout FMaterialPixelParameters Parameters)");
 
-				// Print the non-LWC get function.
-				Printer << Instr->Arg->Type << TEXTVIEW(" Get") << Instr->Name << Signature << OpenBlock << TEXTVIEW("return ");
+			// Print the non-LWC get function.
+			Printer << Instr->Arg->Type << TEXTVIEW(" Get") << Instr->Name << Signature << OpenBlock << TEXTVIEW("return ");
 
-				Instr->Arg->Type.IsDouble()
-					? Printer << TEXTVIEW("WSDemote(Parameters.") << Instr->Name << TEXT(')')
-					: Printer << TEXTVIEW("Parameters.") << Instr->Name;
+			Instr->Arg->Type.IsDouble()
+				? Printer << TEXTVIEW("WSDemote(Parameters.") << Instr->Name << TEXT(')')
+				: Printer << TEXTVIEW("Parameters.") << Instr->Name;
 
-				Printer << EndOfStatement << CloseBlock << NewLine;
+			Printer << EndOfStatement << CloseBlock << NewLine;
 				
-				// Print the LWC get function.
-				if (!Instr->Arg->Type.IsSubstrateData())
-				{
-					MIR::FType ReturnType = Instr->Arg->Type.GetPrimitive().ToScalarKind(MIR::EScalarKind::Double);
-					Printer << ReturnType << TEXTVIEW(" Get") << Instr->Name << TEXTVIEW("_LWC") << Signature << OpenBlock << TEXTVIEW("return ");
+			// Print the LWC get function.
+			if (!Instr->Arg->Type.IsSubstrateData())
+			{
+				MIR::FType ReturnType = Instr->Arg->Type.GetPrimitive().ToScalarKind(MIR::EScalarKind::Double);
+				Printer << ReturnType << TEXTVIEW(" Get") << Instr->Name << TEXTVIEW("_LWC") << Signature << OpenBlock << TEXTVIEW("return ");
 				
-					Instr->Arg->Type.IsDouble()
-						? Printer << TEXTVIEW("Parameters.") << Instr->Name
-						: Printer << TEXTVIEW("WSPromote(Parameters.") << Instr->Name << TEXT(')');
+				Instr->Arg->Type.IsDouble()
+					? Printer << TEXTVIEW("Parameters.") << Instr->Name
+					: Printer << TEXTVIEW("WSPromote(Parameters.") << Instr->Name << TEXT(')');
 					
-					Printer << EndOfStatement << CloseBlock << NewLine;
-				}
+				Printer << EndOfStatement << CloseBlock << NewLine;
 			}
 		}
-
-		OutputShaderTemplateParams.Add(TEXT("custom_outputs"), MoveTemp(Printer.Buffer));
 	}
 
-	// Generates the definitions of all the custom HLSL functions in the module and puts the resulting string
-	// into the "custom_functions" source template parameter.
-	void GenerateCustomFunctionsHLSL()
-	{
-		Printer = { 0 };
+	OutParameters.Add(TEXT("custom_outputs"), MoveTemp(Printer.Buffer));
+}
 
-		for (const MIR::FFunctionHLSL* Function : Module->GetFunctionHLSLs())
+// Generates the definitions of all the custom HLSL functions in the module and puts the resulting string
+// into the "custom_functions" source template parameter.
+static void GenerateCustomFunctionsHLSL(const FMaterialIRModule* Module, TMap<FString, FString>& OutParameters)
+{
+	FHLSLPrinter Printer = { 0 };
+
+	for (const MIR::FFunctionHLSL* Function : Module->GetFunctionHLSLs())
+	{
+		// Print the user-specified defines
+		for (const MIR::FFunctionHLSLDefine& Define : Function->Defines)
 		{
-			// Print the user-specified defines
-			for (const MIR::FFunctionHLSLDefine& Define : Function->Defines)
-			{
-				Printer << TEXTVIEW("#ifndef ") << Define.Name << NewLine;
-				Printer << TEXTVIEW("\t#define ") << Define.Name << TEXT(' ') << Define.Value << NewLine;
-				Printer << TEXTVIEW("#endif") << NewLine;
-			}
+			Printer << TEXTVIEW("#ifndef ") << Define.Name << NewLine;
+			Printer << TEXTVIEW("\t#define ") << Define.Name << TEXT(' ') << Define.Value << NewLine;
+			Printer << TEXTVIEW("#endif") << NewLine;
+		}
 
-			// Print the user-specified include directives
-			for (FStringView Include : Function->Includes)
-			{
-				Printer << TEXTVIEW("#include \"") << Include << TEXT('\"') << NewLine;
-			}
+		// Print the user-specified include directives
+		for (FStringView Include : Function->Includes)
+		{
+			Printer << TEXTVIEW("#include \"") << Include << TEXT('\"') << NewLine;
+		}
 
-			// Write the custom function signature, for example "C5_MyCustomNode".
-			// - 'C' is only a "namespace" for custom functions
-			// - '5' is a unique id used to disambiguate distinct custom functions with the same name.
-			Printer << Function->ReturnType << TEXTVIEW(" C") << Function->UniqueId << TEXT('_') << Function->Name << BeginArgs << TEXT("FMaterialPixelParameters Parameters");
+		// Write the custom function signature, for example "C5_MyCustomNode".
+		// - 'C' is only a "namespace" for custom functions
+		// - '5' is a unique id used to disambiguate distinct custom functions with the same name.
+		Printer << Function->ReturnType << TEXTVIEW(" C") << Function->UniqueId << TEXT('_') << Function->Name << BeginArgs << TEXT("FMaterialPixelParameters Parameters");
 			
-			// Write the parameter declarations
-			for (uint16 i = 0; i < Function->NumParameters; ++i)
-			{
-				Printer << TEXTVIEW(", ");
+		// Write the parameter declarations
+		for (uint16 i = 0; i < Function->NumParameters; ++i)
+		{
+			Printer << TEXTVIEW(", ");
 
-				// Print the io keyword
-				Printer << (i < Function->NumInputOnlyParams ? TEXTVIEW("")
-				: i < Function->NumInputAndOutputParams ? TEXTVIEW("inout ")
-				: TEXTVIEW("out "));
+			// Print the io keyword
+			Printer << (i < Function->NumInputOnlyParams ? TEXTVIEW("")
+			: i < Function->NumInputAndOutputParams ? TEXTVIEW("inout ")
+			: TEXTVIEW("out "));
 
-				// Type and name
-				Printer << Function->Parameters[i].Type << TEXT(' ') << *Function->Parameters[i].Name.ToString();
-			}
+			// Type and name
+			Printer << Function->Parameters[i].Type << TEXT(' ') << *Function->Parameters[i].Name.ToString();
+		}
 
-			Printer << EndArgs << OpenBlock;
+		Printer << EndArgs << OpenBlock;
 			
-			// If the function does not contain a "return" keyword, add one.
-			bool bContainsReturn = Function->Code.Contains(TEXT("return"));
-			if (!bContainsReturn)
-			{
-				Printer << TEXTVIEW("return") << NewLine;
-			}
+		// If the function does not contain a "return" keyword, add one.
+		bool bContainsReturn = Function->Code.Contains(TEXT("return"));
+		if (!bContainsReturn)
+		{
+			Printer << TEXTVIEW("return") << NewLine;
+		}
 			
-			// Write the function code
-			Printer << Function->Code << NewLine;
-
-			if (!bContainsReturn)
-			{
-				Printer << TEXT(';');
-			}
+		// Write the function code
+		Printer << Function->Code << NewLine;
 
-			Printer << NewLine << CloseBlock << NewLine << NewLine;
+		if (!bContainsReturn)
+		{
+			Printer << TEXT(';');
 		}
 
-		OutputShaderTemplateParams.Add(TEXT("custom_functions"), MoveTemp(Printer.Buffer));
+		Printer << NewLine << CloseBlock << NewLine << NewLine;
 	}
+
+	OutParameters.Add(TEXT("custom_functions"), MoveTemp(Printer.Buffer));
+}
 	
-	void GenerateOtherTemplateStringParameters()
-	{
-		const FMaterialIRModule::FStatistics ModuleStatistics = Module->GetStatistics();
-		const FMaterialIRModule::FSubstrateData& SubstrateData = Module->GetSubstrateData();
+static void GenerateOtherTemplateStringParameters(const FMaterial* Material, const FMaterialIRModule* Module, TMap<FString, FString>& OutParameters)
+{
+	const FMaterialIRModule::FStatistics ModuleStatistics = Module->GetStatistics();
+	const FMaterialIRModule::FSubstrateData& SubstrateData = Module->GetSubstrateData();
 
-		// material_declarations
+	// material_declarations
+	{
+		FHLSLPrinter Printer = {};
+		Printer << TEXT("struct FMaterialAttributes") << OpenBlock;
+		for (const FGuid& AttributeID : FMaterialAttributeDefinitionMap::GetOrderedVisibleAttributeList())
 		{
-			Printer = {};
-			Printer << TEXT("struct FMaterialAttributes") << OpenBlock;
-			for (const FGuid& AttributeID : FMaterialAttributeDefinitionMap::GetOrderedVisibleAttributeList())
-			{
-				const EMaterialValueType PropertyType = FMaterialAttributeDefinitionMap::GetValueType(AttributeID);
-				Printer << MIR::FType::FromMaterialValueType(PropertyType) << TEXT(' ') << *FMaterialAttributeDefinitionMap::GetAttributeName(AttributeID) << EndOfStatement;
-			}
-			Printer << CloseBlock << TEXT(";\n");
-
-			OutputShaderTemplateParams.Add(TEXT("material_declarations"), MoveTemp(Printer.Buffer));
+			const EMaterialValueType PropertyType = FMaterialAttributeDefinitionMap::GetValueType(AttributeID);
+			Printer << MIR::FType::FromMaterialValueType(PropertyType) << TEXT(' ') << *FMaterialAttributeDefinitionMap::GetAttributeName(AttributeID) << EndOfStatement;
 		}
+		Printer << CloseBlock << TEXT(";\n");
+
+		OutParameters.Add(TEXT("material_declarations"), MoveTemp(Printer.Buffer));
+	}
 
-		// pixel_material_inputs (field declarations of the FPixelMaterialInputs struct)
+	// pixel_material_inputs (field declarations of the FPixelMaterialInputs struct)
+	{
+		FHLSLPrinter Printer = { 1 };
+		for (EMaterialProperty Property : UMaterialAggregate::GetMaterialAttributesProperties())
 		{
-			Printer = { 1 };
-			for (EMaterialProperty Property : UMaterialAggregate::GetMaterialAttributesProperties())
+			if (Property == MP_WorldPositionOffset) // todo: fixme
 			{
-				if (Property == MP_WorldPositionOffset) // todo: fixme
-				{
-					continue;
-				}
-				if (const FMaterialAggregateAttribute* PropertyAggregate = UMaterialAggregate::GetMaterialAttribute(Property))
-				{
-					Printer << MIR::FType::FromMaterialValueType(PropertyAggregate->ToMaterialValueType()) << TEXT(' ') << *PropertyAggregate->Name.ToString() << EndOfStatement;
-				}
+				continue;
+			}
+			if (const FMaterialAggregateAttribute* PropertyAggregate = UMaterialAggregate::GetMaterialAttribute(Property))
+			{
+				Printer << MIR::FType::FromMaterialValueType(PropertyAggregate->ToMaterialValueType()) << TEXT(' ') << *PropertyAggregate->Name.ToString() << EndOfStatement;
 			}
+		}
 
-			// Add Substrate pixel members declaration code
-			Printer.Buffer.Append(SubstrateData.PixelMembersDeclarationHLSL);
+		// Add Substrate pixel members declaration code
+		Printer.Buffer.Append(SubstrateData.PixelMembersDeclarationHLSL);
 
-			OutputShaderTemplateParams.Add(TEXT("pixel_material_inputs"), MoveTemp(Printer.Buffer));
-		}
+		OutParameters.Add(TEXT("pixel_material_inputs"), MoveTemp(Printer.Buffer));
+	}
 
-		const int32 NumUserVertexInterpolators = (ModuleStatistics.NumUsedVertexInterpolatorFloatRegisters + 1) / 2; // packed into float2s
-		const int32 ShaderTexCoordsArrayLength = ModuleStatistics.NumTexCoords + NumUserVertexInterpolators;
+	const int32 NumUserVertexInterpolators = (ModuleStatistics.NumUsedVertexInterpolatorFloatRegisters + 1) / 2; // packed into float2s
+	const int32 ShaderTexCoordsArrayLength = ModuleStatistics.NumTexCoords + NumUserVertexInterpolators;
 
-		// get_material_customized_u_vs
+	// get_material_customized_u_vs
+	{
+		FHLSLPrinter Printer = { 1 };
+		for (int32 i = 0; i < ShaderTexCoordsArrayLength; i++)
 		{
-			Printer = { 1 };
-			for (int32 i = 0; i < ShaderTexCoordsArrayLength; i++)
-			{
-				Printer << TEXTVIEW("OutTexCoords[") << i << TEXTVIEW("] = Parameters.TexCoords[") << i << TEXT(']') << EndOfStatement;
-			}
-			OutputShaderTemplateParams.Add(TEXT("get_material_customized_u_vs"), MoveTemp(Printer.Buffer));
+			Printer << TEXTVIEW("OutTexCoords[") << i << TEXTVIEW("] = Parameters.TexCoords[") << i << TEXT(']') << EndOfStatement;
 		}
+		OutParameters.Add(TEXT("get_material_customized_u_vs"), MoveTemp(Printer.Buffer));
+	}
 		
-		// texcoords parameters
-		{
-			OutputShaderTemplateParams.Add(TEXT("num_material_texcoords"), FString::FromInt(ModuleStatistics.NumTexCoords));
-			OutputShaderTemplateParams.Add(TEXT("num_material_texcoords_vertex"), FString::FromInt(ModuleStatistics.NumTexCoords));
-			OutputShaderTemplateParams.Add(TEXT("num_custom_vertex_interpolators"), FString::FromInt(NumUserVertexInterpolators));
-			OutputShaderTemplateParams.Add(TEXT("num_tex_coord_interpolators"), FString::FromInt(ShaderTexCoordsArrayLength));
-		}
-
-		// material properties
-		{
-			#define SET_PARAM_RETURN_FLOAT(ParamName, Value) OutputShaderTemplateParams.Add(TEXT(ParamName), FString::Printf(TEXT(TAB "return %.5f"), Value))
-			SET_PARAM_RETURN_FLOAT("get_material_emissive_for_cs", 0.f);
-			SET_PARAM_RETURN_FLOAT("get_material_translucency_directional_lighting_intensity", Material->GetTranslucencyDirectionalLightingIntensity());
-			SET_PARAM_RETURN_FLOAT("get_material_translucent_shadow_density_scale", Material->GetTranslucentShadowDensityScale());
-			SET_PARAM_RETURN_FLOAT("get_material_translucent_self_shadow_density_scale", Material->GetTranslucentSelfShadowDensityScale());
-			SET_PARAM_RETURN_FLOAT("get_material_translucent_self_shadow_second_density_scale", Material->GetTranslucentSelfShadowSecondDensityScale());
-			SET_PARAM_RETURN_FLOAT("get_material_translucent_self_shadow_second_opacity", Material->GetTranslucentSelfShadowSecondOpacity());
-			SET_PARAM_RETURN_FLOAT("get_material_translucent_backscattering_exponent", Material->GetTranslucentBackscatteringExponent());
-			SET_PARAM_RETURN_FLOAT("get_material_opacity_mask_clip_value", Material->GetOpacityMaskClipValue());
+	// texcoords parameters
+	{
+		OutParameters.Add(TEXT("num_material_texcoords"), FString::FromInt(ModuleStatistics.NumTexCoords));
+		OutParameters.Add(TEXT("num_material_texcoords_vertex"), FString::FromInt(ModuleStatistics.NumTexCoords));
+		OutParameters.Add(TEXT("num_custom_vertex_interpolators"), FString::FromInt(NumUserVertexInterpolators));
+		OutParameters.Add(TEXT("num_tex_coord_interpolators"), FString::FromInt(ShaderTexCoordsArrayLength));
+	}
 
-			FLinearColor Extinction = Material->GetTranslucentMultipleScatteringExtinction();
-			OutputShaderTemplateParams.Add(TEXT("get_material_translucent_multiple_scattering_extinction"), FString::Printf(TEXT(TAB "return MaterialFloat3(%.5f, %.5f, %.5f)"), Extinction.R, Extinction.G, Extinction.B));
-		}
+	// material properties
+	{
+		#define SET_PARAM_RETURN_FLOAT(ParamName, Value) OutParameters.Add(TEXT(ParamName), FString::Printf(TEXT(TAB "return %.5f"), Value))
+		SET_PARAM_RETURN_FLOAT("get_material_emissive_for_cs", 0.f);
+		SET_PARAM_RETURN_FLOAT("get_material_translucency_directional_lighting_intensity", Material->GetTranslucencyDirectionalLightingIntensity());
+		SET_PARAM_RETURN_FLOAT("get_material_translucent_shadow_density_scale", Material->GetTranslucentShadowDensityScale());
+		SET_PARAM_RETURN_FLOAT("get_material_translucent_self_shadow_density_scale", Material->GetTranslucentSelfShadowDensityScale());
+		SET_PARAM_RETURN_FLOAT("get_material_translucent_self_shadow_second_density_scale", Material->GetTranslucentSelfShadowSecondDensityScale());
+		SET_PARAM_RETURN_FLOAT("get_material_translucent_self_shadow_second_opacity", Material->GetTranslucentSelfShadowSecondOpacity());
+		SET_PARAM_RETURN_FLOAT("get_material_translucent_backscattering_exponent", Material->GetTranslucentBackscatteringExponent());
+		SET_PARAM_RETURN_FLOAT("get_material_opacity_mask_clip_value", Material->GetOpacityMaskClipValue());
 
-		// user_scene_texture_remap
-		OutputShaderTemplateParams.Add(TEXT("user_scene_texture_remap"), UE::MaterialTranslatorUtils::GenerateUserSceneTextureRemapHLSLDefines(Module->GetCompilationOutput()));
-	
-		if (Substrate::IsSubstrateEnabled())
-		{
-			// No one else write to uniform_material_expressions so no need to search before append.
-			OutputShaderTemplateParams.Add(TEXT("uniform_material_expressions"), SubstrateData.TreeHLSL);
-			// SUBSTRATE_TODO also add in there DescriptionStringCommentForDebug
-		}  
+		FLinearColor Extinction = Material->GetTranslucentMultipleScatteringExtinction();
+		OutParameters.Add(TEXT("get_material_translucent_multiple_scattering_extinction"), FString::Printf(TEXT(TAB "return MaterialFloat3(%.5f, %.5f, %.5f)"), Extinction.R, Extinction.G, Extinction.B));
 	}
-	
-	void GetShaderCompilerEnvironment(FShaderCompilerEnvironment& OutEnvironment)
-	{
-		const FMaterialCompilationOutput& CompilationOutput = Module->GetCompilationOutput();
-		EShaderPlatform ShaderPlatform = Module->GetShaderPlatform();
-
-		OutEnvironment.TargetPlatform = TargetPlatform;
-		OutEnvironment.SetDefine(TEXT("ENABLE_NEW_HLSL_GENERATOR"), 1);
-		OutEnvironment.SetDefine(TEXT("MATERIAL_ATMOSPHERIC_FOG"), false);
-		OutEnvironment.SetDefine(TEXT("MATERIAL_SKY_ATMOSPHERE"), false);
-		OutEnvironment.SetDefine(TEXT("INTERPOLATE_VERTEX_COLOR"), false);
-		OutEnvironment.SetDefine(TEXT("NEEDS_PARTICLE_COLOR"), false);
-		OutEnvironment.SetDefine(TEXT("NEEDS_PARTICLE_LOCAL_TO_WORLD"), false);
-		OutEnvironment.SetDefine(TEXT("NEEDS_PARTICLE_WORLD_TO_LOCAL"), false);
-		OutEnvironment.SetDefine(TEXT("NEEDS_PER_INSTANCE_RANDOM_PS"), false);
-		OutEnvironment.SetDefine(TEXT("USES_EYE_ADAPTATION"), false);
-		OutEnvironment.SetDefine(TEXT("USES_PER_INSTANCE_CUSTOM_DATA"), false);
-		OutEnvironment.SetDefine(TEXT("USES_PER_INSTANCE_FADE_AMOUNT"), false);
-		OutEnvironment.SetDefine(TEXT("USES_TRANSFORM_VECTOR"), false);
-		OutEnvironment.SetDefine(TEXT("WANT_PIXEL_DEPTH_OFFSET"), CompilationOutput.bUsesPixelDepthOffset);
-		OutEnvironment.SetDefineAndCompileArgument(TEXT("USES_WORLD_POSITION_OFFSET"), (bool)CompilationOutput.bUsesWorldPositionOffset);
-		OutEnvironment.SetDefineAndCompileArgument(TEXT("USES_DISPLACEMENT"), false);
-		OutEnvironment.SetDefine(TEXT("USES_EMISSIVE_COLOR"), false);
-		OutEnvironment.SetDefine(TEXT("USES_DISTORTION"), Material->IsDistorted());
-		OutEnvironment.SetDefine(TEXT("MATERIAL_ENABLE_TRANSLUCENCY_FOGGING"), Material->ShouldApplyFogging());
-		OutEnvironment.SetDefine(TEXT("MATERIAL_ENABLE_TRANSLUCENCY_CLOUD_FOGGING"), Material->ShouldApplyCloudFogging());
-		OutEnvironment.SetDefine(TEXT("MATERIAL_IS_SKY"), Material->IsSky());
-		OutEnvironment.SetDefine(TEXT("MATERIAL_COMPUTE_FOG_PER_PIXEL"), Material->ComputeFogPerPixel());
-		OutEnvironment.SetDefine(TEXT("MATERIAL_FULLY_ROUGH"), false);
-		OutEnvironment.SetDefine(TEXT("MATERIAL_USES_ANISOTROPY"), CompilationOutput.bUsesAnisotropy);
-		OutEnvironment.SetDefine(TEXT("MATERIAL_NEURAL_POST_PROCESS"), (CompilationOutput.bUsedWithNeuralNetworks || Material->IsUsedWithNeuralNetworks()) && Material->IsPostProcessMaterial());
-		OutEnvironment.SetDefine(TEXT("NUM_VIRTUALTEXTURE_SAMPLES"), CompilationOutput.UniformExpressionSet.GetVTStacks().Num());
-		OutEnvironment.SetDefine(TEXT("NUM_VIRTUALTEXTURE_FEEDBACK_REQUESTS"), CompilationOutput.NumVirtualTextureFeedbackRequests);
-		OutEnvironment.SetDefine(TEXT("MATERIAL_VIRTUALTEXTURE_FEEDBACK"), CompilationOutput.NumVirtualTextureFeedbackRequests > 0);
-		OutEnvironment.SetDefine(TEXT("IS_MATERIAL_SHADER"), true);
-		OutEnvironment.SetDefine(TEXT("VIRTUAL_TEXTURE_OUTPUT"), CompilationOutput.bHasRuntimeVirtualTextureOutputNode != 0);
-
-		// @massimo.tristano refactor this
-		uint32 DynamicParametersMask = Module->GetStatistics().DynamicParticleParameterMask;
-		if (DynamicParametersMask)
-		{
-			OutEnvironment.SetDefine(TEXT("USE_DYNAMIC_PARAMETERS"), 1);
-			OutEnvironment.SetDefine(TEXT("DYNAMIC_PARAMETERS_MASK"), DynamicParametersMask);
-		}
-
-		// Set all defines that are defined by the module.
-		// Any conditional exemption via material properties, such as 'Material->IsUsedWithInstancedStaticMeshes()', are handled during the material IR analysis.
-		// @massimo.tristano refactor this
-		for (const FName& EnvironmentDefine : Module->GetEnvironmentDefines())
-		{
-			OutEnvironment.SetDefine(EnvironmentDefine, true);
-		}
-
-		// Temporary solution. Define the integer environment defines
-		// @massimo.tristano refactor this
-		for (const auto& Pair : Module->GetIntegerEnvironmentDefines())
-		{
-			OutEnvironment.SetDefine(Pair.Key, Pair.Value);
-		}
-
-		FMaterialShadingModelField ShadingModels = Module->GetCompiledShadingModels();
-		ensure(ShadingModels.IsValid());
-
-		// Logic from FHLSLMaterialTranslator::TranslateMaterial
-		bool bOpacityPropertyIsUsed = Module->IsMaterialPropertyUsed(MP_Opacity);
-
-		bool bUsesCurvature = Module->GetFeatureLevel() == ERHIFeatureLevel::ES3_1 &&
-			((ShadingModels.HasShadingModel(MSM_SubsurfaceProfile) && Module->IsMaterialPropertyUsed(MP_CustomData0))
-			|| (ShadingModels.HasShadingModel(MSM_Eye) && bOpacityPropertyIsUsed));
-
-		int32 NumActiveShadingModels = 0;
-		if (ShadingModels.IsLit())
-		{
-			// This is to have platforms use the simple single layer water shading similar to mobile: no dynamic lights, only sun and sky, no distortion, no colored transmittance on background, no custom depth read.
-			const bool bSingleLayerWaterUsesSimpleShading = FDataDrivenShaderPlatformInfo::GetWaterUsesSimpleForwardShading(ShaderPlatform) && IsForwardShadingEnabled(ShaderPlatform);
-
-			for (int32 i = 0; i < MSM_NUM; ++i)
-			{
-				// SUBSTRATE_TODO handle that using EnumHasAnyFlags(SubstrateTranslatorData.SubstrateMaterialBsdfFeatures, ESubstrateBsdfFeature::SSS)
-				EMaterialShadingModel Model = (EMaterialShadingModel)i;
-				if (Model == MSM_Strata || !ShadingModels.HasShadingModel(Model))
-				{
-					continue;
-				}
-
-				if (Model == MSM_SingleLayerWater && !FDataDrivenShaderPlatformInfo::GetRequiresDisableForwardLocalLights(ShaderPlatform))
-				{
-					continue;
-				}
 
-				if (Model == MSM_SingleLayerWater && bSingleLayerWaterUsesSimpleShading)
-				{
-					// Value must match SINGLE_LAYER_WATER_SHADING_QUALITY_MOBILE_WITH_DEPTH_TEXTURE in SingleLayerWaterCommon.ush!
-					OutEnvironment.SetDefine(TEXT("SINGLE_LAYER_WATER_SHADING_QUALITY"), true);
-				}
+	// user_scene_texture_remap
+	OutParameters.Add(TEXT("user_scene_texture_remap"), UE::MaterialTranslatorUtils::GenerateUserSceneTextureRemapHLSLDefines(Module->GetCompilationOutput()));
+	
+	if (Substrate::IsSubstrateEnabled())
+	{
+		// No one else write to uniform_material_expressions so no need to search before append.
+		OutParameters.Add(TEXT("uniform_material_expressions"), SubstrateData.TreeHLSL);
+		// SUBSTRATE_TODO also add in there DescriptionStringCommentForDebug
+	}  
+}
 
-				OutEnvironment.SetDefine(GetShadingModelParameterName(Model), true);
-				NumActiveShadingModels += 1;
-			}
+static const TCHAR* GetShadingModelParameterName(EMaterialShadingModel InModel)
+{
+	switch (InModel)
+	{
+		case MSM_Unlit: return TEXT("MATERIAL_SHADINGMODEL_UNLIT");
+		case MSM_DefaultLit: return TEXT("MATERIAL_SHADINGMODEL_DEFAULT_LIT");
+		case MSM_Subsurface: return TEXT("MATERIAL_SHADINGMODEL_SUBSURFACE");
+		case MSM_PreintegratedSkin: return TEXT("MATERIAL_SHADINGMODEL_PREINTEGRATED_SKIN");
+		case MSM_ClearCoat: return TEXT("MATERIAL_SHADINGMODEL_CLEAR_COAT");
+		case MSM_SubsurfaceProfile: return TEXT("MATERIAL_SHADINGMODEL_SUBSURFACE_PROFILE");
+		case MSM_TwoSidedFoliage: return TEXT("MATERIAL_SHADINGMODEL_TWOSIDED_FOLIAGE");
+		case MSM_Hair: return TEXT("MATERIAL_SHADINGMODEL_HAIR");
+		case MSM_Cloth: return TEXT("MATERIAL_SHADINGMODEL_CLOTH");
+		case MSM_Eye: return TEXT("MATERIAL_SHADINGMODEL_EYE");
+		case MSM_SingleLayerWater: return TEXT("MATERIAL_SHADINGMODEL_SINGLELAYERWATER");
+		case MSM_ThinTranslucent: return TEXT("MATERIAL_SHADINGMODEL_THIN_TRANSLUCENT");
+		default: UE_MIR_UNREACHABLE();
+	}
+}
 
-			if (ShadingModels.HasShadingModel(MSM_SubsurfaceProfile) && bUsesCurvature)
+static void GetShaderCompilerEnvironment(const FMaterial* Material, const FMaterialIRModule* Module, const ITargetPlatform* TargetPlatform, FShaderCompilerEnvironment& OutEnvironment)
+{
+	const FMaterialCompilationOutput& CompilationOutput = Module->GetCompilationOutput();
+	EShaderPlatform ShaderPlatform = Module->GetShaderPlatform();
+
+	OutEnvironment.TargetPlatform = TargetPlatform;
+	OutEnvironment.SetDefine(TEXT("ENABLE_NEW_HLSL_GENERATOR"), 1);
+	OutEnvironment.SetDefine(TEXT("MATERIAL_ATMOSPHERIC_FOG"), false);
+	OutEnvironment.SetDefine(TEXT("MATERIAL_SKY_ATMOSPHERE"), false);
+	OutEnvironment.SetDefine(TEXT("INTERPOLATE_VERTEX_COLOR"), false);
+	OutEnvironment.SetDefine(TEXT("NEEDS_PARTICLE_COLOR"), false);
+	OutEnvironment.SetDefine(TEXT("NEEDS_PARTICLE_LOCAL_TO_WORLD"), false);
+	OutEnvironment.SetDefine(TEXT("NEEDS_PARTICLE_WORLD_TO_LOCAL"), false);
+	OutEnvironment.SetDefine(TEXT("NEEDS_PER_INSTANCE_RANDOM_PS"), false);
+	OutEnvironment.SetDefine(TEXT("USES_EYE_ADAPTATION"), false);
+	OutEnvironment.SetDefine(TEXT("USES_PER_INSTANCE_CUSTOM_DATA"), false);
+	OutEnvironment.SetDefine(TEXT("USES_PER_INSTANCE_FADE_AMOUNT"), false);
+	OutEnvironment.SetDefine(TEXT("USES_TRANSFORM_VECTOR"), false);
+	OutEnvironment.SetDefine(TEXT("WANT_PIXEL_DEPTH_OFFSET"), CompilationOutput.bUsesPixelDepthOffset);
+	OutEnvironment.SetDefineAndCompileArgument(TEXT("USES_WORLD_POSITION_OFFSET"), (bool)CompilationOutput.bUsesWorldPositionOffset);
+	OutEnvironment.SetDefineAndCompileArgument(TEXT("USES_DISPLACEMENT"), false);
+	OutEnvironment.SetDefine(TEXT("USES_EMISSIVE_COLOR"), false);
+	OutEnvironment.SetDefine(TEXT("USES_DISTORTION"), Material->IsDistorted());
+	OutEnvironment.SetDefine(TEXT("MATERIAL_ENABLE_TRANSLUCENCY_FOGGING"), Material->ShouldApplyFogging());
+	OutEnvironment.SetDefine(TEXT("MATERIAL_ENABLE_TRANSLUCENCY_CLOUD_FOGGING"), Material->ShouldApplyCloudFogging());
+	OutEnvironment.SetDefine(TEXT("MATERIAL_IS_SKY"), Material->IsSky());
+	OutEnvironment.SetDefine(TEXT("MATERIAL_COMPUTE_FOG_PER_PIXEL"), Material->ComputeFogPerPixel());
+	OutEnvironment.SetDefine(TEXT("MATERIAL_FULLY_ROUGH"), false);
+	OutEnvironment.SetDefine(TEXT("MATERIAL_USES_ANISOTROPY"), CompilationOutput.bUsesAnisotropy);
+	OutEnvironment.SetDefine(TEXT("MATERIAL_NEURAL_POST_PROCESS"), (CompilationOutput.bUsedWithNeuralNetworks || Material->IsUsedWithNeuralNetworks()) && Material->IsPostProcessMaterial());
+	OutEnvironment.SetDefine(TEXT("NUM_VIRTUALTEXTURE_SAMPLES"), CompilationOutput.UniformExpressionSet.GetVTStacks().Num());
+	OutEnvironment.SetDefine(TEXT("NUM_VIRTUALTEXTURE_FEEDBACK_REQUESTS"), CompilationOutput.NumVirtualTextureFeedbackRequests);
+	OutEnvironment.SetDefine(TEXT("MATERIAL_VIRTUALTEXTURE_FEEDBACK"), CompilationOutput.NumVirtualTextureFeedbackRequests > 0);
+	OutEnvironment.SetDefine(TEXT("IS_MATERIAL_SHADER"), true);
+	OutEnvironment.SetDefine(TEXT("VIRTUAL_TEXTURE_OUTPUT"), CompilationOutput.bHasRuntimeVirtualTextureOutputNode != 0);
+
+	// @massimo.tristano refactor this
+	uint32 DynamicParametersMask = Module->GetStatistics().DynamicParticleParameterMask;
+	if (DynamicParametersMask)
+	{
+		OutEnvironment.SetDefine(TEXT("USE_DYNAMIC_PARAMETERS"), 1);
+		OutEnvironment.SetDefine(TEXT("DYNAMIC_PARAMETERS_MASK"), DynamicParametersMask);
+	}
+
+	// Set all defines that are defined by the module.
+	// Any conditional exemption via material properties, such as 'Material->IsUsedWithInstancedStaticMeshes()', are handled during the material IR analysis.
+	// @massimo.tristano refactor this
+	for (const FName& EnvironmentDefine : Module->GetEnvironmentDefines())
+	{
+		OutEnvironment.SetDefine(EnvironmentDefine, true);
+	}
+
+	// Temporary solution. Define the integer environment defines
+	// @massimo.tristano refactor this
+	for (const auto& Pair : Module->GetIntegerEnvironmentDefines())
+	{
+		OutEnvironment.SetDefine(Pair.Key, Pair.Value);
+	}
+
+	FMaterialShadingModelField ShadingModels = Module->GetCompiledShadingModels();
+	ensure(ShadingModels.IsValid());
+
+	// Logic from FHLSLMaterialTranslator::TranslateMaterial
+	bool bOpacityPropertyIsUsed = Module->IsMaterialPropertyUsed(MP_Opacity);
+
+	bool bUsesCurvature = Module->GetFeatureLevel() == ERHIFeatureLevel::ES3_1 &&
+		((ShadingModels.HasShadingModel(MSM_SubsurfaceProfile) && Module->IsMaterialPropertyUsed(MP_CustomData0))
+		|| (ShadingModels.HasShadingModel(MSM_Eye) && bOpacityPropertyIsUsed));
+
+	int32 NumActiveShadingModels = 0;
+	if (ShadingModels.IsLit())
+	{
+		// This is to have platforms use the simple single layer water shading similar to mobile: no dynamic lights, only sun and sky, no distortion, no colored transmittance on background, no custom depth read.
+		const bool bSingleLayerWaterUsesSimpleShading = FDataDrivenShaderPlatformInfo::GetWaterUsesSimpleForwardShading(ShaderPlatform) && IsForwardShadingEnabled(ShaderPlatform);
+
+		for (int32 i = 0; i < MSM_NUM; ++i)
+		{
+			// SUBSTRATE_TODO handle that using EnumHasAnyFlags(SubstrateTranslatorData.SubstrateMaterialBsdfFeatures, ESubstrateBsdfFeature::SSS)
+			EMaterialShadingModel Model = (EMaterialShadingModel)i;
+			if (Model == MSM_Strata || !ShadingModels.HasShadingModel(Model))
 			{
-				OutEnvironment.SetDefine(TEXT("MATERIAL_SUBSURFACE_PROFILE_USE_CURVATURE"), true);
+				continue;
 			}
 
-			if (ShadingModels.HasShadingModel(MSM_Eye) && bUsesCurvature)
+			if (Model == MSM_SingleLayerWater && !FDataDrivenShaderPlatformInfo::GetRequiresDisableForwardLocalLights(ShaderPlatform))
 			{
-				OutEnvironment.SetDefine(TEXT("MATERIAL_SHADINGMODEL_EYE_USE_CURVATURE"), true);
+				continue;
 			}
 
-			if (ShadingModels.HasShadingModel(MSM_SingleLayerWater) && FDataDrivenShaderPlatformInfo::GetRequiresDisableForwardLocalLights(ShaderPlatform))
+			if (Model == MSM_SingleLayerWater && bSingleLayerWaterUsesSimpleShading)
 			{
-				OutEnvironment.SetDefine(TEXT("DISABLE_FORWARD_LOCAL_LIGHTS"), true);
+				// Value must match SINGLE_LAYER_WATER_SHADING_QUALITY_MOBILE_WITH_DEPTH_TEXTURE in SingleLayerWaterCommon.ush!
+				OutEnvironment.SetDefine(TEXT("SINGLE_LAYER_WATER_SHADING_QUALITY"), true);
 			}
 
-			const bool bIsWaterDistanceFieldShadowEnabled = IsWaterDistanceFieldShadowEnabled(ShaderPlatform);
-			const bool bIsWaterVSMFilteringEnabled = IsWaterVirtualShadowMapFilteringEnabled(ShaderPlatform);
-			if (ShadingModels.HasShadingModel(MSM_SingleLayerWater) && (bIsWaterDistanceFieldShadowEnabled || bIsWaterVSMFilteringEnabled))
-			{
-				OutEnvironment.SetDefine(TEXT("SINGLE_LAYER_WATER_SEPARATED_MAIN_LIGHT"), TEXT("1"));
-			}
+			OutEnvironment.SetDefine(GetShadingModelParameterName(Model), true);
+			NumActiveShadingModels += 1;
 		}
-		else
+
+		if (ShadingModels.HasShadingModel(MSM_SubsurfaceProfile) && bUsesCurvature)
 		{
-			// Unlit shading model can only exist by itself
-			OutEnvironment.SetDefine(GetShadingModelParameterName(MSM_Unlit), true);
-			NumActiveShadingModels += 1;
+			OutEnvironment.SetDefine(TEXT("MATERIAL_SUBSURFACE_PROFILE_USE_CURVATURE"), true);
 		}
 
-		if (NumActiveShadingModels == 1)
+		if (ShadingModels.HasShadingModel(MSM_Eye) && bUsesCurvature)
 		{
-			OutEnvironment.SetDefine(TEXT("MATERIAL_SINGLE_SHADINGMODEL"), true);
+			OutEnvironment.SetDefine(TEXT("MATERIAL_SHADINGMODEL_EYE_USE_CURVATURE"), true);
 		}
-		else if (!ensure(NumActiveShadingModels > 0))
+
+		if (ShadingModels.HasShadingModel(MSM_SingleLayerWater) && FDataDrivenShaderPlatformInfo::GetRequiresDisableForwardLocalLights(ShaderPlatform))
 		{
-			UE_LOG(LogMaterial, Warning, TEXT("Unknown material shading model(s). Setting to MSM_DefaultLit"));
-			OutEnvironment.SetDefine(GetShadingModelParameterName(MSM_DefaultLit), true);
+			OutEnvironment.SetDefine(TEXT("DISABLE_FORWARD_LOCAL_LIGHTS"), true);
 		}
 
-		OutEnvironment.SetDefine(TEXT("MATERIAL_LWC_ENABLED"), UE::MaterialTranslatorUtils::IsLWCEnabled() ? 1 : 0);
-		OutEnvironment.SetDefine(TEXT("WSVECTOR_IS_TILEOFFSET"), true);
-		OutEnvironment.SetDefine(TEXT("WSVECTOR_IS_DOUBLEFLOAT"), false);
+		const bool bIsWaterDistanceFieldShadowEnabled = IsWaterDistanceFieldShadowEnabled(ShaderPlatform);
+		const bool bIsWaterVSMFilteringEnabled = IsWaterVirtualShadowMapFilteringEnabled(ShaderPlatform);
+		if (ShadingModels.HasShadingModel(MSM_SingleLayerWater) && (bIsWaterDistanceFieldShadowEnabled || bIsWaterVSMFilteringEnabled))
+		{
+			OutEnvironment.SetDefine(TEXT("SINGLE_LAYER_WATER_SEPARATED_MAIN_LIGHT"), TEXT("1"));
+		}
+	}
+	else
+	{
+		// Unlit shading model can only exist by itself
+		OutEnvironment.SetDefine(GetShadingModelParameterName(MSM_Unlit), true);
+		NumActiveShadingModels += 1;
+	}
+
+	if (NumActiveShadingModels == 1)
+	{
+		OutEnvironment.SetDefine(TEXT("MATERIAL_SINGLE_SHADINGMODEL"), true);
+	}
+	else if (!ensure(NumActiveShadingModels > 0))
+	{
+		UE_LOG(LogMaterial, Warning, TEXT("Unknown material shading model(s). Setting to MSM_DefaultLit"));
+		OutEnvironment.SetDefine(GetShadingModelParameterName(MSM_DefaultLit), true);
+	}
+
+	OutEnvironment.SetDefine(TEXT("MATERIAL_LWC_ENABLED"), UE::MaterialTranslatorUtils::IsLWCEnabled() ? 1 : 0);
+	OutEnvironment.SetDefine(TEXT("WSVECTOR_IS_TILEOFFSET"), true);
+	OutEnvironment.SetDefine(TEXT("WSVECTOR_IS_DOUBLEFLOAT"), false);
 
-		if (Material->GetMaterialDomain() == MD_Volume)
+	if (Material->GetMaterialDomain() == MD_Volume)
+	{
+		TArray<const UMaterialExpressionVolumetricAdvancedMaterialOutput*> VolumetricAdvancedExpressions;
+		Material->GetMaterialInterface()->GetMaterial()->GetAllExpressionsOfType(VolumetricAdvancedExpressions);
+		if (VolumetricAdvancedExpressions.Num() > 0)
 		{
-			TArray<const UMaterialExpressionVolumetricAdvancedMaterialOutput*> VolumetricAdvancedExpressions;
-			Material->GetMaterialInterface()->GetMaterial()->GetAllExpressionsOfType(VolumetricAdvancedExpressions);
-			if (VolumetricAdvancedExpressions.Num() > 0)
+			if (VolumetricAdvancedExpressions.Num() > 1)
 			{
-				if (VolumetricAdvancedExpressions.Num() > 1)
-				{
-					UE_LOG(LogMaterial, Fatal, TEXT("Only a single UMaterialExpressionVolumetricAdvancedMaterialOutput node is supported."));
-				}
+				UE_LOG(LogMaterial, Fatal, TEXT("Only a single UMaterialExpressionVolumetricAdvancedMaterialOutput node is supported."));
+			}
 
-				const UMaterialExpressionVolumetricAdvancedMaterialOutput* VolumetricAdvancedNode = VolumetricAdvancedExpressions[0];
-				const TCHAR* Param = VolumetricAdvancedNode->GetEvaluatePhaseOncePerSample() ? TEXT("MATERIAL_VOLUMETRIC_ADVANCED_PHASE_PERSAMPLE") : TEXT("MATERIAL_VOLUMETRIC_ADVANCED_PHASE_PERPIXEL");
-				OutEnvironment.SetDefine(Param, true);
+			const UMaterialExpressionVolumetricAdvancedMaterialOutput* VolumetricAdvancedNode = VolumetricAdvancedExpressions[0];
+			const TCHAR* Param = VolumetricAdvancedNode->GetEvaluatePhaseOncePerSample() ? TEXT("MATERIAL_VOLUMETRIC_ADVANCED_PHASE_PERSAMPLE") : TEXT("MATERIAL_VOLUMETRIC_ADVANCED_PHASE_PERPIXEL");
+			OutEnvironment.SetDefine(Param, true);
 
-				OutEnvironment.SetDefine(TEXT("MATERIAL_VOLUMETRIC_ADVANCED"), true);
-				OutEnvironment.SetDefine(TEXT("MATERIAL_VOLUMETRIC_ADVANCED_GRAYSCALE_MATERIAL"), VolumetricAdvancedNode->bGrayScaleMaterial);
-				OutEnvironment.SetDefine(TEXT("MATERIAL_VOLUMETRIC_ADVANCED_RAYMARCH_VOLUME_SHADOW"), VolumetricAdvancedNode->bRayMarchVolumeShadow);
-				OutEnvironment.SetDefine(TEXT("MATERIAL_VOLUMETRIC_ADVANCED_CLAMP_MULTISCATTERING_CONTRIBUTION"), VolumetricAdvancedNode->bClampMultiScatteringContribution);
-				OutEnvironment.SetDefine(TEXT("MATERIAL_VOLUMETRIC_ADVANCED_MULTISCATTERING_OCTAVE_COUNT"), VolumetricAdvancedNode->GetMultiScatteringApproximationOctaveCount());
-				OutEnvironment.SetDefine(TEXT("MATERIAL_VOLUMETRIC_ADVANCED_CONSERVATIVE_DENSITY"), VolumetricAdvancedNode->ConservativeDensity.IsConnected());
-				OutEnvironment.SetDefine(TEXT("MATERIAL_VOLUMETRIC_ADVANCED_OVERRIDE_AMBIENT_OCCLUSION"), Material->HasAmbientOcclusionConnected());
-				OutEnvironment.SetDefine(TEXT("MATERIAL_VOLUMETRIC_ADVANCED_GROUND_CONTRIBUTION"), VolumetricAdvancedNode->bGroundContribution);
-			}
+			OutEnvironment.SetDefine(TEXT("MATERIAL_VOLUMETRIC_ADVANCED"), true);
+			OutEnvironment.SetDefine(TEXT("MATERIAL_VOLUMETRIC_ADVANCED_GRAYSCALE_MATERIAL"), VolumetricAdvancedNode->bGrayScaleMaterial);
+			OutEnvironment.SetDefine(TEXT("MATERIAL_VOLUMETRIC_ADVANCED_RAYMARCH_VOLUME_SHADOW"), VolumetricAdvancedNode->bRayMarchVolumeShadow);
+			OutEnvironment.SetDefine(TEXT("MATERIAL_VOLUMETRIC_ADVANCED_CLAMP_MULTISCATTERING_CONTRIBUTION"), VolumetricAdvancedNode->bClampMultiScatteringContribution);
+			OutEnvironment.SetDefine(TEXT("MATERIAL_VOLUMETRIC_ADVANCED_MULTISCATTERING_OCTAVE_COUNT"), VolumetricAdvancedNode->GetMultiScatteringApproximationOctaveCount());
+			OutEnvironment.SetDefine(TEXT("MATERIAL_VOLUMETRIC_ADVANCED_CONSERVATIVE_DENSITY"), VolumetricAdvancedNode->ConservativeDensity.IsConnected());
+			OutEnvironment.SetDefine(TEXT("MATERIAL_VOLUMETRIC_ADVANCED_OVERRIDE_AMBIENT_OCCLUSION"), Material->HasAmbientOcclusionConnected());
+			OutEnvironment.SetDefine(TEXT("MATERIAL_VOLUMETRIC_ADVANCED_GROUND_CONTRIBUTION"), VolumetricAdvancedNode->bGroundContribution);
 		}
+	}
 
-		OutEnvironment.SetDefine(TEXT("DUAL_SOURCE_COLOR_BLENDING_ENABLED"), false);
-		OutEnvironment.SetDefine(TEXT("TEXTURE_SAMPLE_DEBUG"), false);
+	OutEnvironment.SetDefine(TEXT("DUAL_SOURCE_COLOR_BLENDING_ENABLED"), false);
+	OutEnvironment.SetDefine(TEXT("TEXTURE_SAMPLE_DEBUG"), false);
 
-		for (int32 VTStackIndex = 0; VTStackIndex < CompilationOutput.UniformExpressionSet.GetVTStacks().Num(); ++VTStackIndex)
-		{
-			const FMaterialVirtualTextureStack& VTStack = CompilationOutput.UniformExpressionSet.GetVTStack(VTStackIndex);
+	for (int32 VTStackIndex = 0; VTStackIndex < CompilationOutput.UniformExpressionSet.GetVTStacks().Num(); ++VTStackIndex)
+	{
+		const FMaterialVirtualTextureStack& VTStack = CompilationOutput.UniformExpressionSet.GetVTStack(VTStackIndex);
 
-			// Setup page table defines to map each VT stack to either 1 or 2 page table textures, depending on how many layers it uses
-			FString PageTableValue = FString::Printf(TEXT("Material.VirtualTexturePageTable0_%d"), VTStackIndex);
-			OutEnvironment.SetDefine(*FString::Printf(TEXT("VIRTUALTEXTURE_PAGETABLE_%d"), VTStackIndex), *PageTableValue);
+		// Setup page table defines to map each VT stack to either 1 or 2 page table textures, depending on how many layers it uses
+		FString PageTableValue = FString::Printf(TEXT("Material.VirtualTexturePageTable0_%d"), VTStackIndex);
+		OutEnvironment.SetDefine(*FString::Printf(TEXT("VIRTUALTEXTURE_PAGETABLE_%d"), VTStackIndex), *PageTableValue);
 
-			// Setup page table uniform defines.
-			FString PageTableUniformValue = FString::Printf(TEXT("Material.VTPackedPageTableUniform[%d*2], Material.VTPackedPageTableUniform[%d*2+1]"), VTStackIndex, VTStackIndex);
-			OutEnvironment.SetDefine(*FString::Printf(TEXT("VIRTUALTEXTURE_PAGETABLE_UNIFORM_%d"), VTStackIndex), *PageTableUniformValue);
-		}
+		// Setup page table uniform defines.
+		FString PageTableUniformValue = FString::Printf(TEXT("Material.VTPackedPageTableUniform[%d*2], Material.VTPackedPageTableUniform[%d*2+1]"), VTStackIndex, VTStackIndex);
+		OutEnvironment.SetDefine(*FString::Printf(TEXT("VIRTUALTEXTURE_PAGETABLE_UNIFORM_%d"), VTStackIndex), *PageTableUniformValue);
+	}
 
-		TConstArrayView<UMaterialParameterCollection*> ParameterCollections = Module->GetParameterCollections();
-		for (int32 CollectionIndex = 0; CollectionIndex < ParameterCollections.Num(); CollectionIndex++)
+	TConstArrayView<UMaterialParameterCollection*> ParameterCollections = Module->GetParameterCollections();
+	for (int32 CollectionIndex = 0; CollectionIndex < ParameterCollections.Num(); CollectionIndex++)
+	{
+		// Add uniform buffer declarations for any parameter collections referenced
+		static_assert(MaxNumParameterCollectionsPerMaterial == 2);
+		static const TCHAR* CollectionNames[MaxNumParameterCollectionsPerMaterial] =
 		{
-			// Add uniform buffer declarations for any parameter collections referenced
-			static_assert(MaxNumParameterCollectionsPerMaterial == 2);
-			static const TCHAR* CollectionNames[MaxNumParameterCollectionsPerMaterial] =
-			{
-				TEXT("MaterialCollection0"),
-				TEXT("MaterialCollection1"),
-			};
+			TEXT("MaterialCollection0"),
+			TEXT("MaterialCollection1"),
+		};
 
-			// Check that the parameter collection loaded successfully.
-			UMaterialParameterCollection* ParameterCollection = ParameterCollections[CollectionIndex];
-			if (!ParameterCollection)
-			{
-				UE_LOG(LogMaterial, Warning, TEXT("Null parameter collection found in environment defines while translating material."));
-				continue;
-			}
+		// Check that the parameter collection loaded successfully.
+		UMaterialParameterCollection* ParameterCollection = ParameterCollections[CollectionIndex];
+		if (!ParameterCollection)
+		{
+			UE_LOG(LogMaterial, Warning, TEXT("Null parameter collection found in environment defines while translating material."));
+			continue;
+		}
 
-			// Ensure PostLoad is called so the uniform buffers are created in case the parameter collection was loaded async 
-			ParameterCollection->ConditionalPostLoad();
+		// Ensure PostLoad is called so the uniform buffers are created in case the parameter collection was loaded async 
+		ParameterCollection->ConditionalPostLoad();
 
-			// Check that the parameter collection uniform buffer structure is valid
-			if (!ParameterCollection->HasValidUniformBufferStruct())
-			{
-				UE_LOG(LogMaterial, Warning, TEXT("Invalid parameter collection uniform buffer struct found in environment defines while translating material."));
-				continue;
-			}
-
-			// This can potentially become an issue for MaterialCollection Uniform Buffers if they ever get non-numeric resources (eg Textures), as
-			// OutEnvironment.ResourceTableMap has a map by name, and the N ParameterCollection Uniform Buffers ALL are names "MaterialCollection"
-			// (and the hlsl cbuffers are named MaterialCollection0, etc, so the names don't match the layout)
-			FShaderUniformBufferParameter::ModifyCompilationEnvironment(CollectionNames[CollectionIndex], ParameterCollection->GetUniformBufferStruct(), Module->GetShaderPlatform(), OutEnvironment);
+		// Check that the parameter collection uniform buffer structure is valid
+		if (!ParameterCollection->HasValidUniformBufferStruct())
+		{
+			UE_LOG(LogMaterial, Warning, TEXT("Invalid parameter collection uniform buffer struct found in environment defines while translating material."));
+			continue;
 		}
+
+		// This can potentially become an issue for MaterialCollection Uniform Buffers if they ever get non-numeric resources (eg Textures), as
+		// OutEnvironment.ResourceTableMap has a map by name, and the N ParameterCollection Uniform Buffers ALL are names "MaterialCollection"
+		// (and the hlsl cbuffers are named MaterialCollection0, etc, so the names don't match the layout)
+		FShaderUniformBufferParameter::ModifyCompilationEnvironment(CollectionNames[CollectionIndex], ParameterCollection->GetUniformBufferStruct(), Module->GetShaderPlatform(), OutEnvironment);
 	}
-};
+}
 
-void FMaterialIRToHLSLTranslation::Run(TMap<FString, FString>& OutParameters, FShaderCompilerEnvironment& OutEnvironment)
+void FMaterialIRToHLSLTranslation::Run(TMap<FString, FString>& OutShaderParameters, FShaderCompilerEnvironment& OutEnvironment)
 {
-	OutParameters.Empty();
-
-	FPrivate Private{ *this, OutParameters };
+	OutShaderParameters.Empty();
 
-	Private.GenerateVertexStageHLSL();
-	Private.GenerateNonVertexStageHLSL(MIR::Stage_Pixel);
-	Private.GenerateNonVertexStageHLSL(MIR::Stage_Compute);
-	Private.GenerateCustomOutputDeclarations();
-	Private.GenerateCustomFunctionsHLSL();
-	Private.GenerateOtherTemplateStringParameters();
-	Private.GetShaderCompilerEnvironment(OutEnvironment);
+	GenerateStageEntryPoints(Module, OutShaderParameters);
+	GenerateCustomOutputDeclarations(Module, OutShaderParameters);
+	GenerateCustomFunctionsHLSL(Module, OutShaderParameters);
+	GenerateOtherTemplateStringParameters(Material, Module, OutShaderParameters);
+	GetShaderCompilerEnvironment(Material, Module, TargetPlatform, OutEnvironment);
 }
 
 #undef TAB
diff --git a/Engine/Source/Runtime/Engine/Public/Materials/MaterialIRToHLSLTranslator.h b/Engine/Source/Runtime/Engine/Public/Materials/MaterialIRToHLSLTranslator.h
--- a/Engine/Source/Runtime/Engine/Public/Materials/MaterialIRToHLSLTranslator.h
+++ b/Engine/Source/Runtime/Engine/Public/Materials/MaterialIRToHLSLTranslator.h
@@ -21,9 +21,9 @@ struct FMaterialIRToHLSLTranslation
 	// The translation target platform.
 	const ITargetPlatform* TargetPlatform{};
 	
-	// Executes the translation process, filling the material shader template parameters in OutParameters and configuring OutEnvironment for shader compilation.
+	// Executes the translation process, filling the material shader template parameters in OutShaderParameters and configuring OutEnvironment for shader compilation.
 	// After this call, interpolate the MaterialTemplate shader with the list of parameters to get the final HLSL shader.
-	void Run(TMap<FString, FString>& OutParameters, FShaderCompilerEnvironment& OutEnvironment);
+	void Run(TMap<FString, FString>& OutShaderParameters, FShaderCompilerEnvironment& OutEnvironment);
 };
 
 #endif // #if WITH_EDITOR
