Add virtual texture update cvars for editor
r.VT.MaxUploadsPerFrameInEditor.Streaming and r.VT.MaxTilesProducedPerFrameInEditor
These are the editor versions of r.VT.MaxUploadsPerFrame.Streaming and r.VT.MaxTilesProducedPerFrame
Previously editor used r.VT.MaxUploadsPerFrame which is intended for GPU page updates.
For streaming updates we can have higher page update bandwidth.
#Synced-CL 50294033

diff --git a/Engine/Source/Runtime/Engine/Private/VT/VirtualTextureScalability.cpp b/Engine/Source/Runtime/Engine/Private/VT/VirtualTextureScalability.cpp
--- a/Engine/Source/Runtime/Engine/Private/VT/VirtualTextureScalability.cpp
+++ b/Engine/Source/Runtime/Engine/Private/VT/VirtualTextureScalability.cpp
@@ -8,15 +8,6 @@
 
 namespace VirtualTextureScalability
 {
-#if WITH_EDITOR
-	static TAutoConsoleVariable<int32> CVarVTMaxUploadsPerFrameInEditor(
-		TEXT("r.VT.MaxUploadsPerFrameInEditor"),
-		32,
-		TEXT("Max number of page uploads per frame when in editor"),
-		ECVF_RenderThreadSafe
-	);
-#endif 
-
 	static TAutoConsoleVariable<int32> CVarVTMaxUploadsPerFrame(
 		TEXT("r.VT.MaxUploadsPerFrame"),
 		2,
@@ -40,7 +31,36 @@ namespace VirtualTextureScalability
 		ECVF_RenderThreadSafe | ECVF_Scalability
 	);
 
+	static TAutoConsoleVariable<int32> CVarVTMaxContinuousUpdatesPerFrame(
+		TEXT("r.VT.MaxContinuousUpdatesPerFrame"),
+		1,
+		TEXT("Max number of page uploads for pages that are already mapped."),
+		ECVF_RenderThreadSafe | ECVF_Scalability
+	);
+
 #if WITH_EDITOR
+	static TAutoConsoleVariable<int32> CVarVTMaxUploadsPerFrameInEditor(
+		TEXT("r.VT.MaxUploadsPerFrameInEditor"),
+		32,
+		TEXT("Max number of page uploads per frame when in editor"),
+		ECVF_RenderThreadSafe
+	);
+
+	static TAutoConsoleVariable<int32> CVarVTMaxUploadsPerFrameStreamingInEditor(
+		TEXT("r.VT.MaxUploadsPerFrameInEditor.Streaming"),
+		500,
+		TEXT("Max number of page uploads per frame in editor for streaming VT.")
+		TEXT("Zero means no separate budget. Negative means no limit.\n"),
+		ECVF_RenderThreadSafe | ECVF_Scalability
+	);
+
+	static TAutoConsoleVariable<int32> CVarMaxPagesProducedPerFrameInEditor(
+		TEXT("r.VT.MaxTilesProducedPerFrameInEditor"),
+		500,
+		TEXT("Max number of pages that can be produced per frame when in editor"),
+		ECVF_RenderThreadSafe | ECVF_Scalability
+	);
+
 	static TAutoConsoleVariable<int32> CVarVTMaxContinuousUpdatesPerFrameInEditor(
 		TEXT("r.VT.MaxContinuousUpdatesPerFrameInEditor"),
 		8,
@@ -49,13 +69,6 @@ namespace VirtualTextureScalability
 	);
 #endif
 
-	static TAutoConsoleVariable<int32> CVarVTMaxContinuousUpdatesPerFrame(
-		TEXT("r.VT.MaxContinuousUpdatesPerFrame"),
-		1,
-		TEXT("Max number of page uploads for pages that are already mapped."),
-		ECVF_RenderThreadSafe | ECVF_Scalability
-	);
-
 	static TAutoConsoleVariable<int32> CVarVTMaxReleasedPerFrame(
 		TEXT("r.VT.MaxReleasedPerFrame"),
 		5,
@@ -172,7 +185,7 @@ namespace VirtualTextureScalability
 	int32 GetMaxUploadsPerFrame()
 	{
 #if WITH_EDITOR
-		// Don't want this scalability setting to affect editor because we rely on reactive updates while editing.
+		// Vary setting in editor because we rely on reactive updates while editing.
 		return GIsEditor ? CVarVTMaxUploadsPerFrameInEditor.GetValueOnAnyThread() : CVarVTMaxUploadsPerFrame.GetValueOnAnyThread();
 #else
 		return CVarVTMaxUploadsPerFrame.GetValueOnAnyThread();
@@ -182,28 +195,27 @@ namespace VirtualTextureScalability
 	int32 GetMaxUploadsPerFrameForStreamingVT()
 	{
 		int32 Budget = CVarVTMaxUploadsPerFrameStreaming.GetValueOnAnyThread();
-
-		if (Budget < 0)
-		{
-			Budget = MAX_int32;
-		}
 #if WITH_EDITOR
-		// Don't want this scalability setting to affect editor because we rely on reactive updates while editing.
-		return GIsEditor ? CVarVTMaxUploadsPerFrameInEditor.GetValueOnAnyThread() : Budget;
-#else
-		return Budget;
+		// Vary setting in editor because we rely on reactive updates while editing.
+		Budget = GIsEditor ? CVarVTMaxUploadsPerFrameStreamingInEditor.GetValueOnAnyThread() : Budget;
 #endif
+		return (Budget < 0) ? MAX_int32 : Budget;
 	}
 
 	int32 GetMaxPagesProducedPerFrame()
 	{
+#if WITH_EDITOR
+		// Vary setting in editor because we rely on reactive updates while editing.
+		return GIsEditor ? CVarMaxPagesProducedPerFrameInEditor.GetValueOnAnyThread() : CVarMaxPagesProducedPerFrame.GetValueOnAnyThread();
+#else
 		return CVarMaxPagesProducedPerFrame.GetValueOnAnyThread();
+#endif
 	}
 
 	int32 GetMaxContinuousUpdatesPerFrame()
 	{
 #if WITH_EDITOR
-		// Don't want this scalability setting to affect editor because we rely on reactive updates while editing, like GPULightmass.
+		// Vary setting in editor because we rely on reactive updates while editing. eg GPULightmass.
 		return GIsEditor ? CVarVTMaxContinuousUpdatesPerFrameInEditor.GetValueOnAnyThread() : CVarVTMaxContinuousUpdatesPerFrame.GetValueOnAnyThread();
 #else
 		return CVarVTMaxContinuousUpdatesPerFrame.GetValueOnAnyThread();
