UPrimitiveComponent now implements IPhysicsBodyInstanceOwner
- IPhysicsBodyInstanceOwner GetPhysicsBodyInstanceOwnerFromHitResult & GetPhysicsBodyInstanceOwnerFromOverlapResult can now return the Component's IPhysicsBodyInstanceOwner when the component is valid
- Added IPhysicsBodyInstanceOwner::GetBodyInstance
- Deprecated IPhysicsBodyInstanceOwner GetSimplePhysicalMaterial/GetComplexPhysicalMaterials
- Refactored/unified FBodyInstance::GetSimplePhysicalMaterial/GetComplexPhysicalMaterials to work with IPhysicsBodyInstanceOwner
- Refactored FOverlapResult::GetItemIndex to work with IPhysicsBodyInstanceOwner

#Synced-CL 50298075

diff --git a/Engine/Plugins/Experimental/FastGeoStreaming/Source/FastGeoStreaming/Internal/FastGeoInstancedStaticMeshComponent.h b/Engine/Plugins/Experimental/FastGeoStreaming/Source/FastGeoStreaming/Internal/FastGeoInstancedStaticMeshComponent.h
--- a/Engine/Plugins/Experimental/FastGeoStreaming/Source/FastGeoStreaming/Internal/FastGeoInstancedStaticMeshComponent.h
+++ b/Engine/Plugins/Experimental/FastGeoStreaming/Source/FastGeoStreaming/Internal/FastGeoInstancedStaticMeshComponent.h
@@ -35,6 +35,7 @@ public:
 	virtual bool ShouldSkipNavigationDirtyAreaOnAddOrRemove() override { return true; }
 	virtual Chaos::FPhysicsObject* GetPhysicsObjectById(Chaos::FPhysicsObjectId Id) const override;
 	virtual TArray<Chaos::FPhysicsObject*> GetAllPhysicsObjects() const override;
+	virtual FBodyInstance* GetBodyInstance(FName BoneName = NAME_None, bool bGetWelded = true, int32 Index = INDEX_NONE) const override;
 	//~ End FFastGeoPrimitiveComponent interface
 
 	void ForEachInstanceMatrix(TFunctionRef<void(const FMatrix&)> Func) const;
diff --git a/Engine/Plugins/Experimental/FastGeoStreaming/Source/FastGeoStreaming/Internal/FastGeoPhysicsBodyInstanceOwner.h b/Engine/Plugins/Experimental/FastGeoStreaming/Source/FastGeoStreaming/Internal/FastGeoPhysicsBodyInstanceOwner.h
--- a/Engine/Plugins/Experimental/FastGeoStreaming/Source/FastGeoStreaming/Internal/FastGeoPhysicsBodyInstanceOwner.h
+++ b/Engine/Plugins/Experimental/FastGeoStreaming/Source/FastGeoStreaming/Internal/FastGeoPhysicsBodyInstanceOwner.h
@@ -22,16 +22,21 @@ public:
 	//~ End FChaosUserDefinedEntity interface
 
 	//~ Begin IPhysicsBodyInstanceOwner interface
+public:
 	virtual bool IsStaticPhysics() const override;
+	virtual bool IsMultiBodyOverlap() const override;
 	virtual UObject* GetSourceObject() const override;
-	virtual FTransform GetTransform() const override;
+	virtual FTransform GetPhysicsOwnerTransform() const override;
 	virtual ECollisionChannel GetCollisionObjectType() const override;
 	virtual ECollisionEnabled::Type GetCollisionEnabled() const override;
 	virtual ECollisionResponse GetCollisionResponseToChannel(ECollisionChannel Channel) const override;
-	virtual UPhysicalMaterial* GetSimplePhysicalMaterial(const TWeakObjectPtr<UBodySetup>& InBodySetup) const override;
-	virtual void GetComplexPhysicalMaterials(TArray<UPhysicalMaterial*>& OutPhysMaterials, TArray<FPhysicalMaterialMaskParams>* OutPhysMaterialMasks) const override;
 	virtual Chaos::FPhysicsObject* GetPhysicsObjectById(Chaos::FPhysicsObjectId Id) const override;
 	virtual TArray<Chaos::FPhysicsObject*> GetAllPhysicsObjects() const override;
+	virtual FBodyInstance* GetBodyInstance(FName BoneName = NAME_None, bool bGetWelded = true, int32 Index = INDEX_NONE) const override;
+	virtual UPhysicalMaterial* GetPhysicsMaterialOverride() const override;
+	virtual UMaterialInterface* GetPhysicsMaterialBase() const override;
+	virtual int32 GetNumMaterials() const override;
+	virtual UMaterialInterface* GetMaterial(int32 Index) const override;
 	//~ End IPhysicsBodyInstanceOwner interface
 
 private:
diff --git a/Engine/Plugins/Experimental/FastGeoStreaming/Source/FastGeoStreaming/Internal/FastGeoPrimitiveComponent.h b/Engine/Plugins/Experimental/FastGeoStreaming/Source/FastGeoStreaming/Internal/FastGeoPrimitiveComponent.h
--- a/Engine/Plugins/Experimental/FastGeoStreaming/Source/FastGeoStreaming/Internal/FastGeoPrimitiveComponent.h
+++ b/Engine/Plugins/Experimental/FastGeoStreaming/Source/FastGeoStreaming/Internal/FastGeoPrimitiveComponent.h
@@ -149,16 +149,18 @@ protected:
 
 	//~ Begin Physics
 	virtual bool IsStaticPhysics() const;
+	virtual bool IsMultiBodyOverlap() const;
 	virtual UObject* GetSourceObject() const;
 	virtual ECollisionChannel GetCollisionObjectType() const;
 	virtual ECollisionEnabled::Type GetCollisionEnabled() const;
-	UE_DEPRECATED(5.8, "Use GetSimplePhysicalMaterial instead.")
-	virtual UPhysicalMaterial* GetPhysicalMaterial() const { return nullptr; }
-	virtual UPhysicalMaterial* GetSimplePhysicalMaterial(const TWeakObjectPtr<UBodySetup>& InBodySetup) const { return nullptr; }
-	virtual void GetComplexPhysicalMaterials(TArray<UPhysicalMaterial*>& OutPhysMaterials, TArray<FPhysicalMaterialMaskParams>* OutPhysMaterialMasks) const {}
+	virtual UPhysicalMaterial* GetPhysicsMaterialOverride() const;
 	virtual ECollisionResponse GetCollisionResponseToChannel(ECollisionChannel Channel) const;
 	virtual Chaos::FPhysicsObject* GetPhysicsObjectById(Chaos::FPhysicsObjectId Id) const;
 	virtual TArray<Chaos::FPhysicsObject*> GetAllPhysicsObjects() const;
+	virtual FBodyInstance* GetBodyInstance(FName BoneName = NAME_None, bool bGetWelded = true, int32 Index = INDEX_NONE) const;
+
+	UE_DEPRECATED(5.8, "GetPhysicalMaterial was removed.")
+	virtual UPhysicalMaterial* GetPhysicalMaterial() const { return nullptr; }
 	//~ End Physics
 
 protected:
@@ -170,6 +172,7 @@ protected:
 	uint8 bFillCollisionUnderneathForNavmesh : 1 = false;
 	uint8 bRasterizeAsFilledConvexVolume : 1 = false;
 	uint8 bCanEverAffectNavigation : 1 = false;
+	uint8 bMultiBodyOverlap : 1 = false;
 	FCustomPrimitiveData CustomPrimitiveData;
 	TEnumAsByte<EHasCustomNavigableGeometry::Type> bHasCustomNavigableGeometry = EHasCustomNavigableGeometry::No;
 	FBodyInstance BodyInstance;
diff --git a/Engine/Plugins/Experimental/FastGeoStreaming/Source/FastGeoStreaming/Internal/FastGeoStaticMeshComponent.h b/Engine/Plugins/Experimental/FastGeoStreaming/Source/FastGeoStreaming/Internal/FastGeoStaticMeshComponent.h
--- a/Engine/Plugins/Experimental/FastGeoStreaming/Source/FastGeoStreaming/Internal/FastGeoStaticMeshComponent.h
+++ b/Engine/Plugins/Experimental/FastGeoStreaming/Source/FastGeoStreaming/Internal/FastGeoStaticMeshComponent.h
@@ -93,8 +93,6 @@ protected:
 	virtual void InitializeSceneProxyDescDynamicProperties() override;
 	virtual void ApplyWorldTransform(const FTransform& InTransform) override;
 	virtual FPrimitiveSceneProxy* CreateSceneProxy(ESceneProxyCreationError* OutError = nullptr) override;
-	virtual UPhysicalMaterial* GetSimplePhysicalMaterial(const TWeakObjectPtr<UBodySetup>& InBodySetup) const override;
-	virtual void GetComplexPhysicalMaterials(TArray<UPhysicalMaterial*>& OutPhysMaterials, TArray<FPhysicalMaterialMaskParams>* OutPhysMaterialMasks) const override;
 	//~ End FFastGeoPrimitiveComponent interface
 
 	//~ Begin FFastGeoMeshComponent interface
diff --git a/Engine/Plugins/Experimental/FastGeoStreaming/Source/FastGeoStreaming/Private/FastGeoInstancedStaticMeshComponent.cpp b/Engine/Plugins/Experimental/FastGeoStreaming/Source/FastGeoStreaming/Private/FastGeoInstancedStaticMeshComponent.cpp
--- a/Engine/Plugins/Experimental/FastGeoStreaming/Source/FastGeoStreaming/Private/FastGeoInstancedStaticMeshComponent.cpp
+++ b/Engine/Plugins/Experimental/FastGeoStreaming/Source/FastGeoStreaming/Private/FastGeoInstancedStaticMeshComponent.cpp
@@ -347,6 +347,15 @@ TArray<Chaos::FPhysicsObject*> FFastGeoInstancedStaticMeshComponent::GetAllPhysi
 	return Objects;
 }
 
+FBodyInstance* FFastGeoInstancedStaticMeshComponent::GetBodyInstance(FName BoneName, bool bGetWelded, int32 Index) const
+{
+	if (Index != INDEX_NONE && InstanceBodies.IsValidIndex(Index))
+	{
+		return const_cast<FBodyInstance*>(InstanceBodies[Index]);
+	}
+	return const_cast<FBodyInstance*>(&BodyInstance);
+}
+
 void FFastGeoInstancedStaticMeshComponent::CreateAllInstanceBodies()
 {
 	TRACE_CPUPROFILER_EVENT_SCOPE(FFastGeoInstancedStaticMeshComponent::CreateAllInstanceBodies);
diff --git a/Engine/Plugins/Experimental/FastGeoStreaming/Source/FastGeoStreaming/Private/FastGeoPhysicsBodyInstanceOwner.cpp b/Engine/Plugins/Experimental/FastGeoStreaming/Source/FastGeoStreaming/Private/FastGeoPhysicsBodyInstanceOwner.cpp
--- a/Engine/Plugins/Experimental/FastGeoStreaming/Source/FastGeoStreaming/Private/FastGeoPhysicsBodyInstanceOwner.cpp
+++ b/Engine/Plugins/Experimental/FastGeoStreaming/Source/FastGeoStreaming/Private/FastGeoPhysicsBodyInstanceOwner.cpp
@@ -48,13 +48,19 @@ bool FFastGeoPhysicsBodyInstanceOwner::IsStaticPhysics() const
 	return OwnerComponent->IsStaticPhysics();
 }
 
+bool FFastGeoPhysicsBodyInstanceOwner::IsMultiBodyOverlap() const
+{
+	check(OwnerContainer.IsValid());
+	return OwnerComponent->IsMultiBodyOverlap();
+}
+
 UObject* FFastGeoPhysicsBodyInstanceOwner::GetSourceObject() const
 {
 	check(OwnerContainer.IsValid());
 	return OwnerComponent->GetSourceObject();
 }
 
-FTransform FFastGeoPhysicsBodyInstanceOwner::GetTransform() const
+FTransform FFastGeoPhysicsBodyInstanceOwner::GetPhysicsOwnerTransform() const
 {
 	check(OwnerContainer.IsValid());
 	return OwnerComponent->GetComponentTransform();
@@ -78,26 +84,48 @@ ECollisionResponse FFastGeoPhysicsBodyInstanceOwner::GetCollisionResponseToChann
 	return OwnerComponent->GetCollisionResponseToChannel(Channel);
 }
 
-UPhysicalMaterial* FFastGeoPhysicsBodyInstanceOwner::GetSimplePhysicalMaterial(const TWeakObjectPtr<UBodySetup>& InBodySetup) const
+Chaos::FPhysicsObject* FFastGeoPhysicsBodyInstanceOwner::GetPhysicsObjectById(Chaos::FPhysicsObjectId Id) const
 {
 	check(OwnerContainer.IsValid());
-	return OwnerComponent->GetSimplePhysicalMaterial(InBodySetup);
+	return OwnerComponent->GetPhysicsObjectById(Id);
 }
 
-void FFastGeoPhysicsBodyInstanceOwner::GetComplexPhysicalMaterials(TArray<UPhysicalMaterial*>& OutPhysMaterials, TArray<FPhysicalMaterialMaskParams>* OutPhysMaterialMasks) const
+TArray<Chaos::FPhysicsObject*> FFastGeoPhysicsBodyInstanceOwner::GetAllPhysicsObjects() const
 {
 	check(OwnerContainer.IsValid());
-	return OwnerComponent->GetComplexPhysicalMaterials(OutPhysMaterials, OutPhysMaterialMasks);
+	return OwnerComponent->GetAllPhysicsObjects();
 }
 
-Chaos::FPhysicsObject* FFastGeoPhysicsBodyInstanceOwner::GetPhysicsObjectById(Chaos::FPhysicsObjectId Id) const
+FBodyInstance* FFastGeoPhysicsBodyInstanceOwner::GetBodyInstance(FName BoneName, bool bGetWelded, int32 Index) const
 {
 	check(OwnerContainer.IsValid());
-	return OwnerComponent->GetPhysicsObjectById(Id);
+	return OwnerComponent->GetBodyInstance(BoneName, bGetWelded, Index);
 }
 
-TArray<Chaos::FPhysicsObject*> FFastGeoPhysicsBodyInstanceOwner::GetAllPhysicsObjects() const
+UPhysicalMaterial* FFastGeoPhysicsBodyInstanceOwner::GetPhysicsMaterialOverride() const
 {
 	check(OwnerContainer.IsValid());
-	return OwnerComponent->GetAllPhysicsObjects();
+	return OwnerComponent->GetPhysicsMaterialOverride();
+}
+
+UMaterialInterface* FFastGeoPhysicsBodyInstanceOwner::GetPhysicsMaterialBase() const
+{
+	check(OwnerContainer.IsValid());
+	if (FFastGeoMeshComponent* MeshComponent = OwnerComponent->CastTo<FFastGeoMeshComponent>())
+	{
+		return MeshComponent->GetMaterial(0);
+	}
+	return nullptr;
+}
+
+int32 FFastGeoPhysicsBodyInstanceOwner::GetNumMaterials() const
+{
+	check(OwnerContainer.IsValid());
+	return OwnerComponent->GetNumMaterials();
+}
+
+UMaterialInterface* FFastGeoPhysicsBodyInstanceOwner::GetMaterial(int32 Index) const
+{
+	check(OwnerContainer.IsValid());
+	return OwnerComponent->GetMaterial(Index);
 }
\ No newline at end of file
diff --git a/Engine/Plugins/Experimental/FastGeoStreaming/Source/FastGeoStreaming/Private/FastGeoPrimitiveComponent.cpp b/Engine/Plugins/Experimental/FastGeoStreaming/Source/FastGeoStreaming/Private/FastGeoPrimitiveComponent.cpp
--- a/Engine/Plugins/Experimental/FastGeoStreaming/Source/FastGeoStreaming/Private/FastGeoPrimitiveComponent.cpp
+++ b/Engine/Plugins/Experimental/FastGeoStreaming/Source/FastGeoStreaming/Private/FastGeoPrimitiveComponent.cpp
@@ -47,6 +47,7 @@ FFastGeoPrimitiveComponent::FFastGeoPrimitiveComponent(const FFastGeoPrimitiveCo
 	, bFillCollisionUnderneathForNavmesh(Other.bFillCollisionUnderneathForNavmesh)
 	, bRasterizeAsFilledConvexVolume(Other.bRasterizeAsFilledConvexVolume)
 	, bCanEverAffectNavigation(Other.bCanEverAffectNavigation)
+	, bMultiBodyOverlap(Other.bMultiBodyOverlap)
 	, CustomPrimitiveData(Other.CustomPrimitiveData)
 	, bHasCustomNavigableGeometry(Other.bHasCustomNavigableGeometry)
 	, BodyInstance(Other.BodyInstance)
@@ -69,6 +70,7 @@ void FFastGeoPrimitiveComponent::Serialize(FArchive& Ar)
 	FArchive_Serialize_BitfieldBool(Ar, bFillCollisionUnderneathForNavmesh);
 	FArchive_Serialize_BitfieldBool(Ar, bRasterizeAsFilledConvexVolume);
 	FArchive_Serialize_BitfieldBool(Ar, bCanEverAffectNavigation);
+	FArchive_Serialize_BitfieldBool(Ar, bMultiBodyOverlap);
 	Ar << CustomPrimitiveData.Data;
 	Ar << bHasCustomNavigableGeometry;
 	Ar << RuntimeVirtualTextures;
@@ -161,6 +163,7 @@ void FFastGeoPrimitiveComponent::InitializeFromComponent(UActorComponent* Compon
 	bFillCollisionUnderneathForNavmesh = PrimitiveComponent->bFillCollisionUnderneathForNavmesh;
 	bRasterizeAsFilledConvexVolume = PrimitiveComponent->bRasterizeAsFilledConvexVolume;
 	bCanEverAffectNavigation = PrimitiveComponent->CanEverAffectNavigation();
+	bMultiBodyOverlap = PrimitiveComponent->bMultiBodyOverlap;
 	CustomPrimitiveData = PrimitiveComponent->GetCustomPrimitiveData();
 	bHasCustomNavigableGeometry = PrimitiveComponent->bHasCustomNavigableGeometry;
 	BodyInstance.CopyBodyInstancePropertiesFrom(&PrimitiveComponent->BodyInstance);
@@ -818,6 +821,11 @@ bool FFastGeoPrimitiveComponent::IsStaticPhysics() const
 	return GetSceneProxyDesc().Mobility != EComponentMobility::Movable && bStaticWhenNotMoveable;
 }
 
+bool FFastGeoPrimitiveComponent::IsMultiBodyOverlap() const
+{
+	return bMultiBodyOverlap;
+}
+
 UObject* FFastGeoPrimitiveComponent::GetSourceObject() const
 {
 	return GetOwnerContainer();
@@ -872,4 +880,14 @@ TArray<Chaos::FPhysicsObject*> FFastGeoPrimitiveComponent::GetAllPhysicsObjects(
 		Bodies.Add(PhysicsObject);
 	}
 	return Bodies;
+}
+
+UPhysicalMaterial* FFastGeoPrimitiveComponent::GetPhysicsMaterialOverride() const
+{
+	return BodyInstance.GetPhysMaterialOverride();
+}
+
+FBodyInstance* FFastGeoPrimitiveComponent::GetBodyInstance(FName BoneName, bool bGetWelded, int32 Index) const
+{
+	return const_cast<FBodyInstance*>(&BodyInstance);
 }
\ No newline at end of file
diff --git a/Engine/Plugins/Experimental/FastGeoStreaming/Source/FastGeoStreaming/Private/FastGeoStaticMeshComponent.cpp b/Engine/Plugins/Experimental/FastGeoStreaming/Source/FastGeoStreaming/Private/FastGeoStaticMeshComponent.cpp
--- a/Engine/Plugins/Experimental/FastGeoStreaming/Source/FastGeoStreaming/Private/FastGeoStaticMeshComponent.cpp
+++ b/Engine/Plugins/Experimental/FastGeoStreaming/Source/FastGeoStreaming/Private/FastGeoStaticMeshComponent.cpp
@@ -464,102 +464,6 @@ bool FFastGeoStaticMeshComponentBase::IsDisplayNaniteFallbackMesh() const
 }
 #endif
 
-//~ Begin IPhysicsBodyInstanceOwner interface
-UPhysicalMaterial* FFastGeoStaticMeshComponentBase::GetSimplePhysicalMaterial(const TWeakObjectPtr<UBodySetup>& InBodySetup) const
-{
-	if (!GEngine || !GEngine->DefaultPhysMaterial)
-	{
-		UE_LOG(LogPhysics, Error, TEXT("FFastGeoStaticMeshComponentBase::GetSimplePhysicalMaterial : GEngine not initialized!"));
-		return nullptr;
-	}
-
-	UPhysicalMaterial* ReturnPhysMaterial = nullptr;
-	if (BodyInstance.GetPhysMaterialOverride())
-	{
-		ReturnPhysMaterial = BodyInstance.GetPhysMaterialOverride();
-	}
-	else
-	{
-		// Use InBodySetup if possible
-		UBodySetup* BodySetup = InBodySetup.Get();
-		if (BodySetup && BodySetup->PhysMaterial)
-		{
-			ReturnPhysMaterial = BodySetup->PhysMaterial;
-		}
-		else
-		{
-			// See if the Material has a PhysicalMaterial
-			UMaterialInterface* Material = GetMaterial(0);
-			UPhysicalMaterial* PhysMatFromMaterial = Material ? Material->GetPhysicalMaterial() : nullptr;
-			if (PhysMatFromMaterial)
-			{
-				ReturnPhysMaterial = PhysMatFromMaterial;
-			}
-			// Fallback is default physical material
-			else
-			{
-				ReturnPhysMaterial = GEngine->DefaultPhysMaterial;
-			}
-		}
-	}
-
-	check(!ReturnPhysMaterial || ReturnPhysMaterial->IsValidLowLevel());
-	return ReturnPhysMaterial;
-}
-
-void FFastGeoStaticMeshComponentBase::GetComplexPhysicalMaterials(TArray<UPhysicalMaterial*>& OutPhysMaterials, TArray<FPhysicalMaterialMaskParams>* OutPhysMaterialMasks) const
-{
-	if (BodyInstance.GetPhysMaterialOverride() != nullptr)
-	{
-		OutPhysMaterials.SetNum(1);
-		OutPhysMaterials[0] = BodyInstance.GetPhysMaterialOverride();
-		check(!OutPhysMaterials[0] || OutPhysMaterials[0]->IsValidLowLevel());
-	}
-	else
-	{
-		// See if the Material has a PhysicalMaterial
-		const int32 NumMaterials = GetNumMaterials();
-		OutPhysMaterials.SetNum(NumMaterials);
-
-		if (OutPhysMaterialMasks)
-		{
-			OutPhysMaterialMasks->SetNum(NumMaterials);
-		}
-
-		for (int32 MatIdx = 0; MatIdx < NumMaterials; MatIdx++)
-		{
-			UPhysicalMaterial* PhysMat = GEngine->DefaultPhysMaterial;
-			UMaterialInterface* Material = GetMaterial(MatIdx);
-			if (Material)
-			{
-				PhysMat = Material->GetPhysicalMaterial();
-			}
-
-			OutPhysMaterials[MatIdx] = PhysMat;
-
-			if (OutPhysMaterialMasks)
-			{
-				UPhysicalMaterialMask* PhysMatMask = nullptr;
-				UMaterialInterface* PhysMatMap = nullptr;
-
-				if (Material)
-				{
-					PhysMatMask = Material->GetPhysicalMaterialMask();
-					if (PhysMatMask)
-					{
-						PhysMatMap = Material;
-					}
-				}
-
-				(*OutPhysMaterialMasks)[MatIdx].PhysicalMaterialMask = PhysMatMask;
-				(*OutPhysMaterialMasks)[MatIdx].PhysicalMaterialMap = PhysMatMap;
-			}
-		}
-	}
-}
-//~ End IPhysicsBodyInstanceOwner interface
-
-
 #if !WITH_EDITOR && UE_WITH_PSO_PRECACHING
 void FFastGeoStaticMeshComponentBase::CollectPSOPrecacheData(const FPSOPrecacheParams& BasePrecachePSOParams, FMaterialInterfacePSOPrecacheParamsList& OutParams)
 {
diff --git a/Engine/Plugins/Experimental/Mover/Source/Mover/Private/DefaultMovementSet/Modes/AsyncNavWalkingMode.cpp b/Engine/Plugins/Experimental/Mover/Source/Mover/Private/DefaultMovementSet/Modes/AsyncNavWalkingMode.cpp
--- a/Engine/Plugins/Experimental/Mover/Source/Mover/Private/DefaultMovementSet/Modes/AsyncNavWalkingMode.cpp
+++ b/Engine/Plugins/Experimental/Mover/Source/Mover/Private/DefaultMovementSet/Modes/AsyncNavWalkingMode.cpp
@@ -470,7 +470,7 @@ void UAsyncNavWalkingMode::FindBestNavMeshLocation(const FVector& TraceStart, co
 		{
 			UPrimitiveComponent* PrimComp = TestHit.GetComponent();
 			// Prefer using primitive component if valid
-			if (IPhysicsBodyInstanceOwner* BodyInstanceOwner = !PrimComp ? IPhysicsBodyInstanceOwner::GetPhysicsBodyInstanceOwnerFromHitResult(TestHit) : nullptr)
+			if (const IPhysicsBodyInstanceOwner* BodyInstanceOwner = !PrimComp ? IPhysicsBodyInstanceOwner::GetPhysicsBodyInstanceOwnerFromHitResult(TestHit) : nullptr)
 			{
 				const bool bBlockOnWorldStatic = (BodyInstanceOwner->GetCollisionResponseToChannel(ECC_WorldStatic) == ECR_Block);
 				const bool bBlockOnWorldDynamic = (BodyInstanceOwner->GetCollisionResponseToChannel(ECC_WorldDynamic) == ECR_Block);
diff --git a/Engine/Source/Runtime/Engine/Classes/Components/PrimitiveComponent.h b/Engine/Source/Runtime/Engine/Classes/Components/PrimitiveComponent.h
--- a/Engine/Source/Runtime/Engine/Classes/Components/PrimitiveComponent.h
+++ b/Engine/Source/Runtime/Engine/Classes/Components/PrimitiveComponent.h
@@ -37,6 +37,7 @@
 #include "PrimitiveSceneInfoData.h"
 #include "Streaming/RenderAssetOwnerStreamingState.h"
 #include "Materials/MaterialParameters.h"
+#include "PhysicsEngine/PhysicsBodyInstanceOwnerInterface.h"
 #include "PrimitiveComponent.generated.h"
 
 DECLARE_CYCLE_STAT_EXTERN(TEXT("BeginComponentOverlap"), STAT_BeginComponentOverlap, STATGROUP_Game, ENGINE_API);
@@ -45,6 +46,7 @@ DECLARE_CYCLE_STAT_EXTERN(TEXT("MoveComponent FastOverlap"), STAT_MoveComponent_
 class AController;
 class FPrimitiveSceneProxy;
 class UMaterialInterface;
+class UPhysicalMaterial;
 class UPrimitiveComponent;
 class UTexture;
 class URuntimeVirtualTexture;
@@ -271,7 +273,7 @@ DECLARE_DYNAMIC_MULTICAST_SPARSE_DELEGATE_TwoParams( FComponentEndTouchOverSigna
  * ShapeComponents generate geometry that is used for collision detection but are not rendered, while StaticMeshComponents and SkeletalMeshComponents contain pre-built geometry that is rendered, but can also be used for collision detection.
  */
 UCLASS(abstract, HideCategories=(Mobility, VirtualTexture), ShowCategories=(PhysicsVolume), MinimalAPI)
-class UPrimitiveComponent : public USceneComponent, public INavRelevantInterface, public IInterface_AsyncCompilation, public IPhysicsComponent, public FRenderAssetOwnerStreamingState
+class UPrimitiveComponent : public USceneComponent, public INavRelevantInterface, public IInterface_AsyncCompilation, public IPhysicsComponent, public FRenderAssetOwnerStreamingState, public IPhysicsBodyInstanceOwner
 {
 	GENERATED_BODY()
 	UE_DECLARE_COMPONENT_ACTOR_INTERFACE(PrimitiveComponent)
@@ -1502,7 +1504,7 @@ public:
 	 * @return the material used by the indexed element of this mesh.
 	 */
 	UFUNCTION(BlueprintPure, Category="Rendering|Material")
-	ENGINE_API virtual class UMaterialInterface* GetMaterial(int32 ElementIndex) const;
+	ENGINE_API virtual class UMaterialInterface* GetMaterial(int32 ElementIndex) const override;
 
 #if WITH_EDITOR
 	/** Returns the material to show in the editor details panel as being used. Skips Nanite Override materials. */
@@ -2262,7 +2264,7 @@ public:
 
 	/** Return number of material elements in this primitive */
 	UFUNCTION(BlueprintPure, Category="Rendering|Material")
-	ENGINE_API virtual int32 GetNumMaterials() const;
+	ENGINE_API virtual int32 GetNumMaterials() const override;
 	
 	/**
 	 * Returns BodyInstance of the component.
@@ -2273,7 +2275,7 @@ public:
 	*
 	* @return		Returns the BodyInstance based on various states (does component have multiple bodies? Is the body welded to another body?)
 	*/
-	ENGINE_API virtual FBodyInstance* GetBodyInstance(FName BoneName = NAME_None, bool bGetWelded = true, int32 Index = INDEX_NONE) const;
+	ENGINE_API virtual FBodyInstance* GetBodyInstance(FName BoneName = NAME_None, bool bGetWelded = true, int32 Index = INDEX_NONE) const override;
 
 	/**
 	 * Returns BodyInstanceAsyncPhysicsTickHandle of the component. For use in the Async Physics Tick event
@@ -3159,8 +3161,31 @@ public:
 	ENGINE_API virtual Chaos::FPhysicsObjectId GetIdFromGTParticle(Chaos::FGeometryParticle* Particle) const override;
 	//~ End IPhysicsComponent Interface.
 
+	//~ Begin IPhysicsBodyInstanceOwner Interface.
+
 	/** Whether the physics is static for this component. */
-	inline bool IsStaticPhysics() const { return Mobility != EComponentMobility::Movable && bStaticWhenNotMoveable; }
+	virtual bool IsStaticPhysics() const override { return Mobility != EComponentMobility::Movable && bStaticWhenNotMoveable; }
+
+	/**
+	 * If true, this body instance owner will generate individual overlaps for each overlapping physics body if it is a multi-body object. When false, it will
+	 * generate only one overlap, regardless of how many physics bodies it has and how many of them are overlapping another body. This flag has no
+	 * influence on single body objects.
+	 */
+	virtual bool IsMultiBodyOverlap() const override { return bMultiBodyOverlap; }
+
+	/** Returns the source UObject of this body instance owner. */
+	virtual UObject* GetSourceObject() const override { return const_cast<UPrimitiveComponent*>(this); }
+
+	/** Returns the transform of this body instance owner. */
+	virtual FTransform GetPhysicsOwnerTransform() const override { return GetComponentTransform(); }
+
+	/** Returns the body instance's physical material override. */
+	ENGINE_API virtual UPhysicalMaterial* GetPhysicsMaterialOverride() const override;
+
+	/** Returns the body instance owner mesh material. */
+	ENGINE_API virtual UMaterialInterface* GetPhysicsMaterialBase() const override;
+
+	//~ End IPhysicsBodyInstanceOwner Interface.
 
 	/** Wheter the component is relevant for first person */
 	inline bool IsFirstPersonRelevant() const { return FirstPersonPrimitiveType != EFirstPersonPrimitiveType::None; }
diff --git a/Engine/Source/Runtime/Engine/Classes/PhysicsEngine/PhysicsBodyInstanceOwnerInterface.h b/Engine/Source/Runtime/Engine/Classes/PhysicsEngine/PhysicsBodyInstanceOwnerInterface.h
--- a/Engine/Source/Runtime/Engine/Classes/PhysicsEngine/PhysicsBodyInstanceOwnerInterface.h
+++ b/Engine/Source/Runtime/Engine/Classes/PhysicsEngine/PhysicsBodyInstanceOwnerInterface.h
@@ -10,6 +10,7 @@
 
 class UObject;
 class UBodySetup;
+class UMaterialInterface;
 class UPhysicalMaterial;
 class IPhysicsBodyInstanceOwner;
 struct FBodyInstance;
@@ -33,7 +34,7 @@ public:
 	virtual IPhysicsBodyInstanceOwner* ResolvePhysicsBodyInstanceOwner(Chaos::FConstPhysicsObjectHandle PhysicsObject) = 0;
 };
 
-/** Interface representing the owner of a FBodyInstance (used when the owner is not an UPrimitiveComponent). */
+/** Interface representing the owner of a FBodyInstance. */
 class IPhysicsBodyInstanceOwner
 {
 public:
@@ -42,29 +43,36 @@ public:
 	UE_DEPRECATED(5.8, "Use GetPhysicsBodyInstanceOwnerFromHitResult instead.")
 	static IPhysicsBodyInstanceOwner* GetPhysicsBodyInstandeOwnerFromHitResult(const FHitResult& Result)
 	{
-		return GetPhysicsBodyInstanceOwnerFromHitResult(Result);
+		return const_cast<IPhysicsBodyInstanceOwner*>(GetPhysicsBodyInstanceOwnerFromHitResult(Result));
 	}
 
 	UE_DEPRECATED(5.8, "Use GetPhysicsBodyInstanceOwnerFromOverlapResult instead.")
 	static IPhysicsBodyInstanceOwner* GetPhysicsBodyInstandeOwnerFromOverlapResult(const FOverlapResult& OverlapResult)
 	{
-		return GetPhysicsBodyInstanceOwnerFromOverlapResult(OverlapResult);
+		return const_cast<IPhysicsBodyInstanceOwner*>(GetPhysicsBodyInstanceOwnerFromOverlapResult(OverlapResult));
 	}
 
 	/** Returns the IPhysicsBodyInstanceOwner based on a given hit result. */
-	static ENGINE_API IPhysicsBodyInstanceOwner* GetPhysicsBodyInstanceOwnerFromHitResult(const FHitResult& Result);
+	static ENGINE_API const IPhysicsBodyInstanceOwner* GetPhysicsBodyInstanceOwnerFromHitResult(const FHitResult& Result);
 
 	/** Returns the IPhysicsBodyInstanceOwner based on a given overlap result. */
-	static ENGINE_API IPhysicsBodyInstanceOwner* GetPhysicsBodyInstanceOwnerFromOverlapResult(const FOverlapResult& OverlapResult);
+	static ENGINE_API const IPhysicsBodyInstanceOwner* GetPhysicsBodyInstanceOwnerFromOverlapResult(const FOverlapResult& OverlapResult);
 
 	/** Whether the physics is static. */
 	virtual bool IsStaticPhysics() const = 0;
 
-	/** Source object for this body. */
+	/**
+	 * If true, this body instance owner will generate individual overlaps for each overlapping physics body if it is a multi-body object. When false, it will
+	 * generate only one overlap, regardless of how many physics bodies it has and how many of them are overlapping another body. This flag has no
+	 * influence on single body objects.
+	 */
+	virtual bool IsMultiBodyOverlap() const = 0;
+
+	/** Returns the source UObject of this body instance owner. */
 	virtual UObject* GetSourceObject() const = 0;
 
 	/** Returns the transform of this body instance owner. */
-	virtual FTransform GetTransform() const = 0;
+	virtual FTransform GetPhysicsOwnerTransform() const = 0;
 
 	/** Returns the channel that the body instance owner belongs to. */
 	virtual ECollisionChannel GetCollisionObjectType() const = 0;
@@ -72,14 +80,8 @@ public:
 	/** Returns the form of collision for this body instance owner. */
 	virtual ECollisionEnabled::Type GetCollisionEnabled() const = 0;
 
-	UE_DEPRECATED(5.8, "Use GetSimplePhysicalMaterial instead.")
-	virtual UPhysicalMaterial* GetPhysicalMaterial() const { return GetSimplePhysicalMaterial(nullptr); };
-
-	/** Find the correct PhysicalMaterial for simple geometry on this body (used by FBodyInstance::GetSimplePhysicalMaterial). */
-	virtual UPhysicalMaterial* GetSimplePhysicalMaterial(const TWeakObjectPtr<UBodySetup>& InBodySetup) const = 0;
-
-	/** Gets the complex PhysicalMaterials array for this body (used by FBodyInstance::GetComplexPhysicalMaterials). */
-	virtual void GetComplexPhysicalMaterials(TArray<UPhysicalMaterial*>& OutPhysMaterials, TArray<FPhysicalMaterialMaskParams>* OutPhysMaterialMasks) const = 0;
+	UE_DEPRECATED(5.8, "GetComplexPhysicalMaterials was removed.")
+	virtual void GetComplexPhysicalMaterials(TArray<UPhysicalMaterial*>& OutPhysMaterials, TArray<FPhysicalMaterialMaskParams>* OutPhysMaterialMasks) const {}
 
 	/** Gets the response type given a specific channel. */
 	virtual ECollisionResponse GetCollisionResponseToChannel(ECollisionChannel Channel) const = 0;
@@ -89,4 +91,29 @@ public:
 
 	/** Returns all physics objects associated to the physics body instance owner. */
 	virtual TArray<Chaos::FPhysicsObject*> GetAllPhysicsObjects() const = 0;
+
+	/**
+	 * Returns the body instance of the body instance owner.
+	 *
+	 * @param BoneName		Used to get body associated with specific bone. NAME_None automatically gets the root most body
+	 * @param bGetWelded	If the component has been welded to another component and bGetWelded is true we return the single welded BodyInstance that is used in the simulation
+	 * @param Index			Index used in Components with multiple body instances
+	 *
+	 * @return Returns the BodyInstance based on various states (does component have multiple bodies? Is the body welded to another body?)
+	 */
+	virtual FBodyInstance* GetBodyInstance(FName BoneName = NAME_None, bool bGetWelded = true, int32 Index = INDEX_NONE) const = 0;
+
+	/** Returns the body instance's physics material override (used by FBodyInstance::GetSimplePhysicalMaterial / FBodyInstance::GetComplexPhysicalMaterials). */
+	virtual UPhysicalMaterial* GetPhysicsMaterialOverride() const = 0;
+
+	/** Returns the body instance owner mesh material (used by FBodyInstance::GetSimplePhysicalMaterial). */
+	virtual UMaterialInterface* GetPhysicsMaterialBase() const = 0;
+
+	/** Returns the body instance owner number of materials (used by FBodyInstance::GetComplexPhysicalMaterials). */
+	virtual int32 GetNumMaterials() const = 0;
+
+	/** Returns the body instance owner material for the provided material index. (used by FBodyInstance::GetComplexPhysicalMaterials). */
+	virtual UMaterialInterface* GetMaterial(int32 Index) const = 0;
+
+	friend struct FBodyInstance;
 };
\ No newline at end of file
diff --git a/Engine/Source/Runtime/Engine/Private/Collision/CollisionConversions.cpp b/Engine/Source/Runtime/Engine/Private/Collision/CollisionConversions.cpp
--- a/Engine/Source/Runtime/Engine/Private/Collision/CollisionConversions.cpp
+++ b/Engine/Source/Runtime/Engine/Private/Collision/CollisionConversions.cpp
@@ -699,17 +699,6 @@ void ConvertQueryOverlap(const FPhysicsShape& Shape, const FPhysicsActor& Actor,
 			// since resolving FActorInstanceHandle may require to access related UObjects (e.g. owner actor), which is not safe outside game thread.
 			OutOverlap.OverlapObjectHandle = FActorInstanceHandle::MakeActorHandleToResolve(BodyInst->OwnerComponent, BodyInst->InstanceBodyIndex);
 			OutOverlap.Component = BodyInst->OwnerComponent; // Copying weak pointer is faster than assigning raw pointer.
-			bool bUseBodyIndex = false;
-#if UE_WITH_REMOTE_OBJECT_HANDLE
-			// To avoid migrations, we always set the body index and we defer bMultiBodyOverlap to the getter
-			bUseBodyIndex = true;
-#else
-			if (const UPrimitiveComponent* OwnerComponent = BodyInst->OwnerComponent.Get())
-			{
-				bUseBodyIndex = OwnerComponent->bMultiBodyOverlap;
-			}
-#endif
-			OutOverlap.SetItemIndex(bUseBodyIndex ? BodyInst->InstanceBodyIndex : INDEX_NONE);
 		}
 		OutOverlap.PhysicsObject = BodyInst->GetPhysicsActor() ? BodyInst->GetPhysicsActor()->GetPhysicsObject() : nullptr;
 
@@ -721,6 +710,19 @@ void ConvertQueryOverlap(const FPhysicsShape& Shape, const FPhysicsActor& Actor,
 				OutOverlap.PhysicsObjectOwner = ChaosUserEntityAppend->GetOwnerObject();
 			}
 		}
+
+		// Set ItemIndex
+		if (const IPhysicsBodyInstanceOwner* BodyInstanceOwner = IPhysicsBodyInstanceOwner::GetPhysicsBodyInstanceOwnerFromOverlapResult(OutOverlap))
+		{
+			bool bUseBodyIndex = false;
+#if UE_WITH_REMOTE_OBJECT_HANDLE
+			// To avoid migrations, we always set the body index and we defer bMultiBodyOverlap to the getter
+			bUseBodyIndex = true;
+#else
+			bUseBodyIndex = BodyInstanceOwner->IsMultiBodyOverlap();
+#endif
+			OutOverlap.SetItemIndex(bUseBodyIndex ? BodyInst->InstanceBodyIndex : INDEX_NONE);
+		}
 	}
 	else
 	{
diff --git a/Engine/Source/Runtime/Engine/Private/Engine/OverlapResult.cpp b/Engine/Source/Runtime/Engine/Private/Engine/OverlapResult.cpp
--- a/Engine/Source/Runtime/Engine/Private/Engine/OverlapResult.cpp
+++ b/Engine/Source/Runtime/Engine/Private/Engine/OverlapResult.cpp
@@ -13,9 +13,9 @@ int32 FOverlapResult::GetItemIndexInternal() const
 
 int32 FOverlapResult::GetItemIndex() const
 {
-	if (const UPrimitiveComponent* OwnerComponent = GetComponent())
+	if (const IPhysicsBodyInstanceOwner* BodyInstanceOwner = IPhysicsBodyInstanceOwner::GetPhysicsBodyInstanceOwnerFromOverlapResult(*this))
 	{
-		return OwnerComponent->bMultiBodyOverlap ? ItemIndex : INDEX_NONE;
+		return BodyInstanceOwner->IsMultiBodyOverlap() ? ItemIndex : INDEX_NONE;
 	}
 	return INDEX_NONE;
 }
diff --git a/Engine/Source/Runtime/Engine/Private/PhysicsEngine/BodyInstance.cpp b/Engine/Source/Runtime/Engine/Private/PhysicsEngine/BodyInstance.cpp
--- a/Engine/Source/Runtime/Engine/Private/PhysicsEngine/BodyInstance.cpp
+++ b/Engine/Source/Runtime/Engine/Private/PhysicsEngine/BodyInstance.cpp
@@ -3143,6 +3143,24 @@ UPhysicalMaterial* FBodyInstance::GetSimplePhysicalMaterial() const
 	return GetSimplePhysicalMaterial(this, OwnerComponent, GetBodySetup(), BodyInstanceOwner);
 }
 
+namespace UE::Private
+{
+	const IPhysicsBodyInstanceOwner* GetEffectiveBodyInstanceOwner(const TWeakObjectPtr<UPrimitiveComponent>& InOwnerComp, IPhysicsBodyInstanceOwner* InBodyInstanceOwner)
+	{
+		// Use provided body instance owner if valid, else use component
+		const IPhysicsBodyInstanceOwner* EffectiveBodyInstanceOwner = InBodyInstanceOwner;
+		if (!EffectiveBodyInstanceOwner)
+		{
+			if (UPrimitiveComponent* OwnerPrimComponent = InOwnerComp.Get())
+			{
+				EffectiveBodyInstanceOwner = OwnerPrimComponent;
+				check(EffectiveBodyInstanceOwner);
+			}
+		}
+		return EffectiveBodyInstanceOwner;
+	}
+}
+
 UPhysicalMaterial* FBodyInstance::GetSimplePhysicalMaterial(const FBodyInstance* BodyInstance, TWeakObjectPtr<UPrimitiveComponent> OwnerComp, TWeakObjectPtr<UBodySetup> BodySetupPtr, IPhysicsBodyInstanceOwner* BodyInstanceOwner)
 {
 	if(!GEngine || !GEngine->DefaultPhysMaterial)
@@ -3153,7 +3171,7 @@ UPhysicalMaterial* FBodyInstance::GetSimplePhysicalMaterial(const FBodyInstance*
 	}
 
 	// Find the PhysicalMaterial we need to apply to the physics bodies.
-	// (LOW priority) Engine Mat, Material PhysMat, BodySetup Mat, Component Override, Body Override (HIGH priority)
+	// (LOW priority) Engine Mat, Material PhysMat, BodySetup Mat, Component/BodyInstanceOwner Override, Body Override (HIGH priority)
 	
 	UPhysicalMaterial* ReturnPhysMaterial = NULL;
 
@@ -3163,18 +3181,14 @@ UPhysicalMaterial* FBodyInstance::GetSimplePhysicalMaterial(const FBodyInstance*
 		ReturnPhysMaterial = BodyInstance->PhysMaterialOverride;
 		check(!ReturnPhysMaterial || ReturnPhysMaterial->IsValidLowLevel());
 	}
-	else if (BodyInstanceOwner)
-	{
-		ReturnPhysMaterial = BodyInstanceOwner->GetSimplePhysicalMaterial(BodySetupPtr);
-		check(!ReturnPhysMaterial || ReturnPhysMaterial->IsValidLowLevel());
-	}
 	else
 	{
-		// Component override
-		UPrimitiveComponent* OwnerPrimComponent = OwnerComp.Get();
-		if (OwnerPrimComponent && OwnerPrimComponent->BodyInstance.PhysMaterialOverride != NULL)
+		// Component/BodyInstanceOwner override
+		const IPhysicsBodyInstanceOwner* EffectiveBodyInstanceOwner = UE::Private::GetEffectiveBodyInstanceOwner(OwnerComp, BodyInstanceOwner);
+		UPhysicalMaterial* BodyInstancePhysicalMaterialOverride = EffectiveBodyInstanceOwner ? EffectiveBodyInstanceOwner->GetPhysicsMaterialOverride() : nullptr;
+		if (BodyInstancePhysicalMaterialOverride)
 		{
-			ReturnPhysMaterial = OwnerComp->BodyInstance.PhysMaterialOverride;
+			ReturnPhysMaterial = BodyInstancePhysicalMaterialOverride;
 			check(!ReturnPhysMaterial || ReturnPhysMaterial->IsValidLowLevel());
 		}
 		else
@@ -3189,15 +3203,10 @@ UPhysicalMaterial* FBodyInstance::GetSimplePhysicalMaterial(const FBodyInstance*
 			else
 			{
 				// See if the Material has a PhysicalMaterial
-				UMeshComponent* MeshComp = Cast<UMeshComponent>(OwnerPrimComponent);
 				UPhysicalMaterial* PhysMatFromMaterial = NULL;
-				if (MeshComp != NULL)
+				if (UMaterialInterface* Material = EffectiveBodyInstanceOwner ? EffectiveBodyInstanceOwner->GetPhysicsMaterialBase() : nullptr)
 				{
-					UMaterialInterface* Material = MeshComp->GetMaterial(0);
-					if (Material != NULL)
-					{
-						PhysMatFromMaterial = Material->GetPhysicalMaterial();
-					}
+					PhysMatFromMaterial = Material->GetPhysicalMaterial();
 				}
 
 				if (PhysMatFromMaterial != NULL)
@@ -3247,12 +3256,11 @@ void FBodyInstance::GetComplexPhysicalMaterials(const FBodyInstance* BodyInstanc
 	if(!GEngine || !GEngine->DefaultPhysMaterial)
 	{
 		UE_LOG(LogPhysics, Error, TEXT("FBodyInstance::GetComplexPhysicalMaterials : GEngine not initialized! Cannot call this during native CDO construction, wrap with if(!HasAnyFlags(RF_ClassDefaultObject)) or move out of constructor, material parameters will not be correct."));
-
 		return;
 	}
 
 	// Find the PhysicalMaterial we need to apply to the physics bodies.
-	// (LOW priority) Engine Mat, Material PhysMat, Component Override, Body Override (HIGH priority)
+	// (LOW priority) Engine Mat, Material PhysMat, Component/BodyInstanceOwner Override, Body Override (HIGH priority)
 	
 	// BodyInstance override
 	if (BodyInstance && BodyInstance->PhysMaterialOverride != nullptr)
@@ -3263,24 +3271,19 @@ void FBodyInstance::GetComplexPhysicalMaterials(const FBodyInstance* BodyInstanc
 	}
 	else
 	{
-		// Use BodyInstanceOwner if valid
-		if (InBodyInstanceOwner)
-		{
-			InBodyInstanceOwner->GetComplexPhysicalMaterials(OutPhysMaterials, OutPhysMaterialMasks);
-		}
-		// Component override
-		else if (UPrimitiveComponent* OwnerPrimComponent = OwnerComp.Get())
+		// Component/BodyInstanceOwner override
+		if (const IPhysicsBodyInstanceOwner* EffectiveBodyInstanceOwner = UE::Private::GetEffectiveBodyInstanceOwner(OwnerComp, InBodyInstanceOwner))
 		{
-			if (OwnerPrimComponent->BodyInstance.PhysMaterialOverride != nullptr)
+			if (UPhysicalMaterial* BodyInstancePhysicalMaterialOverride = EffectiveBodyInstanceOwner->GetPhysicsMaterialOverride())
 			{
 				OutPhysMaterials.SetNum(1);
-				OutPhysMaterials[0] = OwnerComp->BodyInstance.PhysMaterialOverride;
+				OutPhysMaterials[0] = BodyInstancePhysicalMaterialOverride;
 				check(!OutPhysMaterials[0] || OutPhysMaterials[0]->IsValidLowLevel());
 			}
 			else
 			{
 				// See if the Material has a PhysicalMaterial
-				const int32 NumMaterials = OwnerPrimComponent->GetNumMaterials();
+				const int32 NumMaterials = EffectiveBodyInstanceOwner->GetNumMaterials();
 				OutPhysMaterials.SetNum(NumMaterials);
 
 				if (OutPhysMaterialMasks)
@@ -3291,7 +3294,7 @@ void FBodyInstance::GetComplexPhysicalMaterials(const FBodyInstance* BodyInstanc
 				for (int32 MatIdx = 0; MatIdx < NumMaterials; MatIdx++)
 				{
 					UPhysicalMaterial* PhysMat = GEngine->DefaultPhysMaterial;
-					UMaterialInterface* Material = OwnerPrimComponent->GetMaterial(MatIdx);
+					UMaterialInterface* Material = EffectiveBodyInstanceOwner->GetMaterial(MatIdx);
 					if (Material)
 					{
 						PhysMat = Material->GetPhysicalMaterial();
diff --git a/Engine/Source/Runtime/Engine/Private/PhysicsEngine/PhysicsBodyInstanceOwnerInterface.cpp b/Engine/Source/Runtime/Engine/Private/PhysicsEngine/PhysicsBodyInstanceOwnerInterface.cpp
--- a/Engine/Source/Runtime/Engine/Private/PhysicsEngine/PhysicsBodyInstanceOwnerInterface.cpp
+++ b/Engine/Source/Runtime/Engine/Private/PhysicsEngine/PhysicsBodyInstanceOwnerInterface.cpp
@@ -3,23 +3,38 @@
 #include "PhysicsEngine/PhysicsBodyInstanceOwnerInterface.h"
 #include "Engine/HitResult.h"
 #include "Engine/OverlapResult.h"
+#include "Components/PrimitiveComponent.h"
 
 #include UE_INLINE_GENERATED_CPP_BY_NAME(PhysicsBodyInstanceOwnerInterface)
 
-IPhysicsBodyInstanceOwner* IPhysicsBodyInstanceOwner::GetPhysicsBodyInstanceOwnerFromHitResult(const FHitResult& Result)
+const IPhysicsBodyInstanceOwner* IPhysicsBodyInstanceOwner::GetPhysicsBodyInstanceOwnerFromHitResult(const FHitResult& Result)
 {
 	if (IPhysicsBodyInstanceOwnerResolver* PhysicsObjectOwnerResolver = Result.PhysicsObject ? Cast<IPhysicsBodyInstanceOwnerResolver>(Result.PhysicsObjectOwner.Get()) : nullptr)
 	{
-		return PhysicsObjectOwnerResolver->ResolvePhysicsBodyInstanceOwner(Result.PhysicsObject);
+		if (IPhysicsBodyInstanceOwner* BodyInstanceOwner = PhysicsObjectOwnerResolver->ResolvePhysicsBodyInstanceOwner(Result.PhysicsObject))
+		{
+			return BodyInstanceOwner;
+		}
+	}
+	if (const UPrimitiveComponent* Component = Result.GetComponent())
+	{
+		return Component;
 	}
 	return nullptr;
 }
 
-IPhysicsBodyInstanceOwner* IPhysicsBodyInstanceOwner::GetPhysicsBodyInstanceOwnerFromOverlapResult(const FOverlapResult& Result)
+const IPhysicsBodyInstanceOwner* IPhysicsBodyInstanceOwner::GetPhysicsBodyInstanceOwnerFromOverlapResult(const FOverlapResult& Result)
 {
 	if (IPhysicsBodyInstanceOwnerResolver* PhysicsObjectOwnerResolver = Result.PhysicsObject ? Cast<IPhysicsBodyInstanceOwnerResolver>(Result.PhysicsObjectOwner.Get()) : nullptr)
 	{
-		return PhysicsObjectOwnerResolver->ResolvePhysicsBodyInstanceOwner(Result.PhysicsObject);
+		if (IPhysicsBodyInstanceOwner* BodyInstanceOwner = PhysicsObjectOwnerResolver->ResolvePhysicsBodyInstanceOwner(Result.PhysicsObject))
+		{
+			return BodyInstanceOwner;
+		}
+	}
+	if (const UPrimitiveComponent* Component = Result.GetComponent())
+	{
+		return Component;
 	}
 	return nullptr;
 }
diff --git a/Engine/Source/Runtime/Engine/Private/PrimitiveComponentPhysics.cpp b/Engine/Source/Runtime/Engine/Private/PrimitiveComponentPhysics.cpp
--- a/Engine/Source/Runtime/Engine/Private/PrimitiveComponentPhysics.cpp
+++ b/Engine/Source/Runtime/Engine/Private/PrimitiveComponentPhysics.cpp
@@ -13,7 +13,9 @@
 #include "UObject/UObjectThreadContext.h"
 #include "PhysicsEngine/PhysicsObjectExternalInterface.h"
 #include "PhysicsProxy/SingleParticlePhysicsProxy.h"
+#include "PhysicalMaterials/PhysicalMaterial.h"
 #include "Components/SkeletalMeshComponent.h"
+#include "Materials/MaterialInterface.h"
 #include "Engine/SkeletalMeshSocket.h"
 
 //////////////// PRIMITIVECOMPONENT ///////////////
@@ -1616,4 +1618,14 @@ Chaos::FPhysicsObjectId UPrimitiveComponent::GetIdFromGTParticle(Chaos::FGeometr
 	return 0;
 }
 
+UPhysicalMaterial* UPrimitiveComponent::GetPhysicsMaterialOverride() const
+{
+	return BodyInstance.GetPhysMaterialOverride();
+}
+
+UMaterialInterface* UPrimitiveComponent::GetPhysicsMaterialBase() const
+{
+	return IsA<UMeshComponent>() ? GetMaterial(0) : nullptr;
+}
+
 #undef LOCTEXT_NAMESPACE
