Fix spline metadata channel load issue and tangent issues

#Synced-CL 50274144

diff --git a/Engine/Source/Runtime/Engine/Classes/Curves/Splines/TangentSpline.h b/Engine/Source/Runtime/Engine/Classes/Curves/Splines/TangentSpline.h
--- a/Engine/Source/Runtime/Engine/Classes/Curves/Splines/TangentSpline.h
+++ b/Engine/Source/Runtime/Engine/Classes/Curves/Splines/TangentSpline.h
@@ -280,9 +280,13 @@ int32 FTangentSpline::AddAttributeValue(FName Name, const AttrType& Value, float
 	using AttrSplineType = UE::Geometry::Spline::TTangentBezierSpline<AttrType>;
 	using ControlPointType = UE::Geometry::Spline::TTangentBezierControlPoint<AttrType>;
 	using UE::Geometry::FInterval1f;
-
+	
 	if (AttrSplineType* Channel = GetTypedAttributeChannel<AttrSplineType>(Name))
 	{
+		ON_SCOPE_EXIT
+		{
+			Channel->UpdateTangents();
+		};
 		ControlPointType ControlPoint = ControlPointType(Value);
 		int32 NumPoints = Channel->GetNumPoints();
 
@@ -355,6 +359,7 @@ void FTangentSpline::SetAttributeValue(FName Name, const AttrType& Value, int32
 	if (AttrSplineType* Channel = GetTypedAttributeChannel<AttrSplineType>(Name))
 	{
 		Channel->SetValue(Index, Value);
+		Channel->UpdateTangents();
 	}
 }
 
@@ -420,6 +425,10 @@ int32 FTangentSpline::SetAttributeParameter(const FName& Name, int32 Index, floa
 
 	if (AttrSplineType* Channel = GetTypedAttributeChannel<AttrSplineType>(Name))
 	{
+		ON_SCOPE_EXIT
+		{
+			Channel->UpdateTangents();
+		};
 		float ParentSpaceParameter = Parameter;
 
 		// Prevent collapse of the channel space and keep it starting at 0.
diff --git a/Engine/Source/Runtime/GeometryCore/Public/Splines/PolyBezierSpline.h b/Engine/Source/Runtime/GeometryCore/Public/Splines/PolyBezierSpline.h
--- a/Engine/Source/Runtime/GeometryCore/Public/Splines/PolyBezierSpline.h
+++ b/Engine/Source/Runtime/GeometryCore/Public/Splines/PolyBezierSpline.h
@@ -314,6 +314,12 @@ public:
 			Ar << LoopKnotDelta;
 		}
 
+		// backup fix in case there are assets with the anchor layout missing pair knots
+		if (GetNumAnchorPoints() > 0 && Base::PairKnots.IsEmpty())
+		{
+			Reparameterize();
+		}
+
 		return true;
 	}
 
@@ -635,7 +641,7 @@ public:
 			return Base::PairKnots.Last().Value + LoopKnotDelta;
 		}
 
-		if (Index < 0 || Index >= NumAnchors)
+		if (Index < 0 || Index >= NumAnchors || !Base::PairKnots.IsValidIndex(Index))
 		{
 			// todo: fail
 			return 0.f;
@@ -1498,7 +1504,8 @@ private:
 				ValueType Out(ForceInitToZero);
 
 				Base::Values = { In, A, Out };
-				Base::PairKnots.Reset();
+				Base::PairKnots = { FKnot( 0, GetExpectedNumKnots())};
+								
 				Base::MarkFlatKnotsCacheDirty();
 
 				return;
