New material translator:
- Hoisted MaterialAttributes related material aggregate APIs into its own namespace as free functions, outside of the FMaterialAttributes struct for clarity.
- Added support for CustomUVs in the new translator.

#rnx

#Synced-CL 50155647

diff --git a/Engine/Source/Runtime/Engine/Private/Materials/MaterialAggregate.cpp b/Engine/Source/Runtime/Engine/Private/Materials/MaterialAggregate.cpp
--- a/Engine/Source/Runtime/Engine/Private/Materials/MaterialAggregate.cpp
+++ b/Engine/Source/Runtime/Engine/Private/Materials/MaterialAggregate.cpp
@@ -70,7 +70,7 @@ EMaterialValueType FMaterialAggregateAttribute::ToMaterialValueType() const
 		case EMaterialAggregateAttributeType::SubstrateData: return MCT_Substrate;
 
 		case EMaterialAggregateAttributeType::Aggregate:
-			return Aggregate.Get() == UMaterialAggregate::GetMaterialAttributes() ? MCT_MaterialAttributes : MCT_Unknown;
+			return Aggregate.Get() == MaterialAttributesAggregate::Get() ? MCT_MaterialAttributes : MCT_Unknown;
 		
 		default:
 			checkNoEntry();
@@ -135,7 +135,55 @@ static void ClearCyclicDependencies(UMaterialAggregate* Aggregate, bool bDisplay
 	}
 }
 
-static FMaterialAggregateAttribute MakeAttributeFromMaterialProperty(EMaterialProperty Property)
+int32 UMaterialAggregate::FindAttributeIndexByName(FName Name) const
+{
+	return Attributes.IndexOfByPredicate([Name](const FMaterialAggregateAttribute& Attribute)
+	{
+		return Attribute.Name == Name;
+	});
+}
+
+void UMaterialAggregate::PostLoad()
+{
+	ClearCyclicDependencies(this, false);
+
+	Super::PostLoad();
+}
+
+void UMaterialAggregate::PostEditChangeProperty(FPropertyChangedEvent& PropertyChangedEvent)
+{
+	Super::PostEditChangeProperty(PropertyChangedEvent);
+	
+	if (!(PropertyChangedEvent.ChangeType & EPropertyChangeType::ValueSet))
+	{
+		return;
+	}
+
+	const FName PropertyName = PropertyChangedEvent.Property ? PropertyChangedEvent.Property->GetFName() : NAME_None;
+	
+	if (PropertyName == GET_MEMBER_NAME_CHECKED(FMaterialAggregateAttribute, Type))
+	{
+		for (FMaterialAggregateAttribute& Attr : Attributes)
+		{
+			if (Attr.Type != EMaterialAggregateAttributeType::Aggregate)
+			{
+				Attr.Aggregate = nullptr;
+			}
+		}
+	}
+	else if (PropertyName == GET_MEMBER_NAME_CHECKED(FMaterialAggregateAttribute, Aggregate))
+	{
+		ClearCyclicDependencies(this, true);
+	}
+}
+
+#endif // #if WITH_EDITOR
+
+#if WITH_EDITOR
+namespace MaterialAttributesAggregate
+{
+
+static FMaterialAggregateAttribute MakeAttributeFromProperty(EMaterialProperty Property)
 { 
 	FMaterialAggregateAttribute Attribute;
 	Attribute.Name = (Property == MP_SubsurfaceColor) ? TEXT("Subsurface") : *FMaterialAttributeDefinitionMap::GetAttributeName(Property);
@@ -147,23 +195,28 @@ static FMaterialAggregateAttribute MakeAttributeFromMaterialProperty(EMaterialPr
 	return Attribute;
 }
 
-struct FMaterialAttributePropertyIndexMap
+struct FAttributePropertyIndexMap
 {
+	static const uint32 NumMaxCustomizedUVs = 8;
+
 	// Constant time lookups for Property/Index mapping
 	EMaterialProperty IndexToProperty[MP_MAX];
+
+	//
 	int PropertyToIndex[MP_MAX];
 
+	//
 	int NumAttributes = 0;
 
-	static const FMaterialAttributePropertyIndexMap& Get()
+	static const FAttributePropertyIndexMap& Get()
 	{
-		static const FMaterialAttributePropertyIndexMap MaterialAttributeProperties;
-		return MaterialAttributeProperties;
+		static const FAttributePropertyIndexMap Instance;
+		return Instance;
 	}
 
 	// This map is the ground truth for which material properties we compile in the new compiler.
 	// Dynamically extendable whilst retaining constant time referencing.
-	FMaterialAttributePropertyIndexMap()
+	FAttributePropertyIndexMap()
 	{
 		//Initialise all empty values to defaults so we aren't referencing incorrect data after creating the maps
 		for (int32 ArrayIndex = 0; ArrayIndex < MP_MAX; ArrayIndex++)
@@ -192,6 +245,14 @@ struct FMaterialAttributePropertyIndexMap
 		PushAttribute(MP_CustomData1);
 		PushAttribute(MP_AmbientOcclusion);
 		PushAttribute(MP_Refraction);
+		PushAttribute(MP_CustomizedUVs0);
+		PushAttribute(MP_CustomizedUVs1);
+		PushAttribute(MP_CustomizedUVs2);
+		PushAttribute(MP_CustomizedUVs3);
+		PushAttribute(MP_CustomizedUVs4);
+		PushAttribute(MP_CustomizedUVs5);
+		PushAttribute(MP_CustomizedUVs6);
+		PushAttribute(MP_CustomizedUVs7);
 		PushAttribute(MP_PixelDepthOffset);
 		PushAttribute(MP_ShadingModel);
 		PushAttribute(MP_Displacement);
@@ -216,89 +277,49 @@ struct FMaterialAttributePropertyIndexMap
 	}
 };
 
-const TConstArrayView<EMaterialProperty> UMaterialAggregate::GetMaterialAttributesProperties()
+const TConstArrayView<EMaterialProperty> GetMaterialProperties()
 {
-	return { FMaterialAttributePropertyIndexMap::Get().IndexToProperty, FMaterialAttributePropertyIndexMap::Get().NumAttributes };
+	const auto& Properties = FAttributePropertyIndexMap::Get();
+	return { Properties.IndexToProperty, Properties.NumAttributes };
 }
 
-const UMaterialAggregate* UMaterialAggregate::GetMaterialAttributes()
+const UMaterialAggregate* Get()
 {
-	static TObjectPtr<UMaterialAggregate> Instance = []
-	{
+	static const TObjectPtr<UMaterialAggregate> Instance = [] () {
 		TObjectPtr<UMaterialAggregate> MA = NewObject<UMaterialAggregate>(GetTransientPackage());
 		MA->Rename(TEXT("MaterialAttributes"));
 		MA->AddToRoot();
-		for (EMaterialProperty Property : GetMaterialAttributesProperties())
+		for (EMaterialProperty Property : GetMaterialProperties())
 		{
-			MA->Attributes.Add(MakeAttributeFromMaterialProperty(Property));
+			MA->Attributes.Add(MakeAttributeFromProperty(Property));
 		}
 		return MA;
 	}();
 	return Instance.Get();
 }
 
-const FMaterialAggregateAttribute* UMaterialAggregate::GetMaterialAttribute(EMaterialProperty Property)
+const FMaterialAggregateAttribute* GetAttribute(EMaterialProperty Property)
 {
-	int32 PropertyIndex = MaterialPropertyToAttributeIndex(Property);
-	if (FMaterialAttributePropertyIndexMap::Get().IsValidIndex(PropertyIndex))
+	int32 PropertyIndex = MaterialPropertyToAttributeIndex( Property);
+	if (FAttributePropertyIndexMap::Get().IsValidIndex(PropertyIndex))
 	{
-		return &GetMaterialAttributes()->Attributes[PropertyIndex];
+		return &Get()->Attributes[PropertyIndex];
 	}
 	return nullptr;
 }
 
-int32 UMaterialAggregate::MaterialPropertyToAttributeIndex(EMaterialProperty Property)
+int32 MaterialPropertyToAttributeIndex(EMaterialProperty Property)
 {
 	check(Property < MP_MAX);
-	return (Property < MP_MAX) ? FMaterialAttributePropertyIndexMap::Get().PropertyToIndex[Property] : INDEX_NONE;
-}
-
-EMaterialProperty UMaterialAggregate::AttributeIndexToMaterialProperty(int32 MaterialAttributeIndex)
-{
-	ensure(FMaterialAttributePropertyIndexMap::Get().IsValidIndex(MaterialAttributeIndex));
-	return FMaterialAttributePropertyIndexMap::Get().IndexToProperty[MaterialAttributeIndex];
+	return (Property < MP_MAX) ? FAttributePropertyIndexMap::Get().PropertyToIndex[Property] : INDEX_NONE;
 }
 
-int32 UMaterialAggregate::FindAttributeIndexByName(FName Name) const
+EMaterialProperty AttributeIndexToMaterialProperty(int32 MaterialAttributeIndex)
 {
-	return Attributes.IndexOfByPredicate([Name](const FMaterialAggregateAttribute& Attribute)
-	{
-		return Attribute.Name == Name;
-	});
+	ensure(FAttributePropertyIndexMap::Get().IsValidIndex(MaterialAttributeIndex));
+	return FAttributePropertyIndexMap::Get().IndexToProperty[MaterialAttributeIndex];
 }
 
-void UMaterialAggregate::PostLoad()
-{
-	ClearCyclicDependencies(this, false);
-
-	Super::PostLoad();
-}
-
-void UMaterialAggregate::PostEditChangeProperty(FPropertyChangedEvent& PropertyChangedEvent)
-{
-	Super::PostEditChangeProperty(PropertyChangedEvent);
-	
-	if (!(PropertyChangedEvent.ChangeType & EPropertyChangeType::ValueSet))
-	{
-		return;
-	}
-
-	const FName PropertyName = PropertyChangedEvent.Property ? PropertyChangedEvent.Property->GetFName() : NAME_None;
-	
-	if (PropertyName == GET_MEMBER_NAME_CHECKED(FMaterialAggregateAttribute, Type))
-	{
-		for (FMaterialAggregateAttribute& Attr : Attributes)
-		{
-			if (Attr.Type != EMaterialAggregateAttributeType::Aggregate)
-			{
-				Attr.Aggregate = nullptr;
-			}
-		}
-	}
-	else if (PropertyName == GET_MEMBER_NAME_CHECKED(FMaterialAggregateAttribute, Aggregate))
-	{
-		ClearCyclicDependencies(this, true);
-	}
-}
+} // MaterialAttributesAggregate
+#endif // #if WITH_EDITOR
 
-#endif // WITH_EDITOR
diff --git a/Engine/Source/Runtime/Engine/Private/Materials/MaterialExpressions.cpp b/Engine/Source/Runtime/Engine/Private/Materials/MaterialExpressions.cpp
--- a/Engine/Source/Runtime/Engine/Private/Materials/MaterialExpressions.cpp
+++ b/Engine/Source/Runtime/Engine/Private/Materials/MaterialExpressions.cpp
@@ -23254,7 +23254,7 @@ const UMaterialAggregate* UMaterialExpressionAggregate::GetAggregate() const
 {
 	switch (Kind)
 	{
-		case EMaterialExpressionMakeAggregateKind::MaterialAttributes: return UMaterialAggregate::GetMaterialAttributes();
+		case EMaterialExpressionMakeAggregateKind::MaterialAttributes: return MaterialAttributesAggregate::Get();
 		case EMaterialExpressionMakeAggregateKind::UserDefined: return UserAggregate;
 		default: checkNoEntry(); return nullptr;
 	}
@@ -23273,6 +23273,23 @@ const FMaterialAggregateAttribute* UMaterialExpressionAggregate::GetAggregateAtt
 	return nullptr;
 }
 
+static bool IsAggregateAttributeEnabled(const UMaterial* Material, const UMaterialAggregate* Aggregate, int AttributeIndex, const TArray<FName>& UsedAttributeNames)
+{
+	if (Aggregate == MaterialAttributesAggregate::Get())
+	{
+		EMaterialProperty Property = MaterialAttributesAggregate::AttributeIndexToMaterialProperty(AttributeIndex);
+		if (Property <= MP_CustomizedUVs7 && Property >= MP_CustomizedUVs0 + Material->NumCustomizedUVs)
+		{
+			return false;
+		}
+	}
+	if (UsedAttributeNames.Contains(Aggregate->Attributes[AttributeIndex].Name))
+	{
+		return false;
+	}
+	return true;
+}
+
 TArray<FString> UMaterialExpressionAggregate::GetPossibleAttributeNames() const
 {
 	TArray<FString> Values;
@@ -23280,7 +23297,7 @@ TArray<FString> UMaterialExpressionAggregate::GetPossibleAttributeNames() const
 	{
 		for (int32 i = 0; i < Aggregate->Attributes.Num(); ++i)
 		{
-			if (!AttributeNames.Contains(Aggregate->Attributes[i].Name))
+			if (IsAggregateAttributeEnabled(Material, Aggregate, i, AttributeNames))
 			{
 				Values.Add(Aggregate->Attributes[i].Name.ToString());
 			}
@@ -23335,7 +23352,7 @@ static EMaterialValueType GetPinValueType(const UMaterialExpressionAggregate* Ex
 	check(PinIndex >= 0);
 	if (PinIndex == 0)
 	{
-		return Expr->GetAggregate() == UMaterialAggregate::GetMaterialAttributes() ? MCT_MaterialAttributes : MCT_Unknown;
+		return Expr->GetAggregate() == MaterialAttributesAggregate::Get() ? MCT_MaterialAttributes : MCT_Unknown;
 	}
 	if (const FMaterialAggregateAttribute* Attribute = Expr->GetAggregateAttribute(PinIndex - 1))
 	{
@@ -23366,13 +23383,13 @@ void UMaterialExpressionAggregate::PreEditChange(FProperty* PropertyAboutToChang
 	PreEditData.Entries = Entries;
 }
 
-static FName PickAvailableAttribute(const UMaterialAggregate* Aggregate, const TArray<FName>& AttributeNames)
+static FName PickAvailableAttribute(const UMaterial* Material, const UMaterialAggregate* Aggregate, const TArray<FName>& AttributeNames)
 {
 	if (Aggregate)
 	{
 		for (int i = 0; i < Aggregate->Attributes.Num(); ++i)
 		{
-			if (!AttributeNames.Contains(Aggregate->Attributes[i].Name))
+			if (IsAggregateAttributeEnabled(Material, Aggregate, i, AttributeNames))
 			{
 				return Aggregate->Attributes[i].Name;
 			}
@@ -23404,7 +23421,7 @@ void UMaterialExpressionAggregate::PostEditChangeProperty(FPropertyChangedEvent&
         {
             if (AttributeNames[i] == NAME_None)
             {
-                AttributeNames[i] = PickAvailableAttribute(Aggregate, AttributeNames);
+                AttributeNames[i] = PickAvailableAttribute(Material, Aggregate, AttributeNames);
                 if (AttributeNames[i] == NAME_None)
                 {
                     AttributeNames.RemoveAt(i);
diff --git a/Engine/Source/Runtime/Engine/Private/Materials/MaterialExpressionsIR.cpp b/Engine/Source/Runtime/Engine/Private/Materials/MaterialExpressionsIR.cpp
--- a/Engine/Source/Runtime/Engine/Private/Materials/MaterialExpressionsIR.cpp
+++ b/Engine/Source/Runtime/Engine/Private/Materials/MaterialExpressionsIR.cpp
@@ -2985,12 +2985,12 @@ void UMaterialExpressionMakeMaterialAttributes::Build(MIR::FEmitter& Em)
 		PushAttributeAssignment(EMaterialProperty(MP_CustomizedUVs0 + i), &CustomizedUVs[i]);
 	}
 
-	Em.Output(0, Em.Aggregate(UMaterialAggregate::GetMaterialAttributes(), {}, Assignments.Left(NumAssignments)));
+	Em.Output(0, Em.Aggregate(MaterialAttributesAggregate::Get(), {}, Assignments.Left(NumAssignments)));
 }
 
 void UMaterialExpressionBreakMaterialAttributes::Build(MIR::FEmitter& Em)
 {
-	FValueRef Prototype = Em.CheckIsAggregate(Em.Input(&MaterialAttributes), UMaterialAggregate::GetMaterialAttributes());
+	FValueRef Prototype = Em.CheckIsAggregate(Em.Input(&MaterialAttributes), MaterialAttributesAggregate::Get());
 	UE_MIR_CHECKPOINT(Em);
 
 	static const EMaterialProperty Properties[] = {
@@ -3017,7 +3017,7 @@ void UMaterialExpressionBreakMaterialAttributes::Build(MIR::FEmitter& Em)
 
 	for (int32 Index = 0; Index < UE_ARRAY_COUNT(Properties); ++Index)
 	{
-		Em.Output(Index, Em.Subscript(Prototype, UMaterialAggregate::MaterialPropertyToAttributeIndex(Properties[Index])));
+		Em.Output(Index, Em.Subscript(Prototype, MaterialAttributesAggregate::MaterialPropertyToAttributeIndex(Properties[Index])));
 	}
 }
 
@@ -3047,13 +3047,13 @@ void UMaterialExpressionGetMaterialAttributes::Build(MIR::FEmitter& Em)
 {
 	CheckMaterialAttributesExpression(Em, AttributeGetTypes);
 
-	FValueRef Prototype = Em.CheckIsAggregate(Em.TryInput(&MaterialAttributes), UMaterialAggregate::GetMaterialAttributes());
+	FValueRef Prototype = Em.CheckIsAggregate(Em.TryInput(&MaterialAttributes), MaterialAttributesAggregate::Get());
 	
 	UE_MIR_CHECKPOINT(Em);
 	
 	Em.Output(0, Prototype);
 
-	const UMaterialAggregate* MaterialAttributesAggregate = UMaterialAggregate::GetMaterialAttributes();
+	const UMaterialAggregate* MaterialAttributesAggregate = MaterialAttributesAggregate::Get();
 	for (int32 i = 0; i < AttributeGetTypes.Num(); ++i)
 	{
 		EMaterialProperty Property = FMaterialAttributeDefinitionMap::GetProperty(AttributeGetTypes[i]);
@@ -3069,11 +3069,11 @@ void UMaterialExpressionSetMaterialAttributes::Build(MIR::FEmitter& Em)
 {
 	CheckMaterialAttributesExpression(Em, AttributeSetTypes);
 	
-	FValueRef Prototype = Em.CheckIsAggregate(Em.TryInput(&Inputs[0]), UMaterialAggregate::GetMaterialAttributes());
+	FValueRef Prototype = Em.CheckIsAggregate(Em.TryInput(&Inputs[0]), MaterialAttributesAggregate::Get());
 	
 	UE_MIR_CHECKPOINT(Em);
 
-	const UMaterialAggregate* MaterialAttributesAggregate = UMaterialAggregate::GetMaterialAttributes();
+	const UMaterialAggregate* MaterialAttributesAggregate = MaterialAttributesAggregate::Get();
 	MIR::TTemporaryArray<MIR::FAttributeAssignment> Assignments{ AttributeSetTypes.Num() };
 	int32 NumAssignments = 0;
 
@@ -3088,7 +3088,7 @@ void UMaterialExpressionSetMaterialAttributes::Build(MIR::FEmitter& Em)
 		}
 	}
 
-	Prototype = Em.Aggregate(UMaterialAggregate::GetMaterialAttributes(), Prototype, Assignments.Left(NumAssignments));
+	Prototype = Em.Aggregate(MaterialAttributesAggregate::Get(), Prototype, Assignments.Left(NumAssignments));
 
 	Em.Output(0, Prototype);
 }
@@ -3096,7 +3096,7 @@ void UMaterialExpressionSetMaterialAttributes::Build(MIR::FEmitter& Em)
 // Utility to input a MaterialAttributes value, or return the default instance (with each attribute set to zero). 
 static FValueRef InputDefaultMaterialAttributes(MIR::FEmitter& Em, FExpressionInput* Input)
 {
-	const UMaterialAggregate* MaterialAttributes = UMaterialAggregate::GetMaterialAttributes();
+	const UMaterialAggregate* MaterialAttributes = MaterialAttributesAggregate::Get();
 	FValueRef Value = Em.CheckIsAggregate(Em.TryInput(Input), MaterialAttributes);
 	return Value ? Value : Em.Aggregate(MaterialAttributes);
 }
@@ -3130,9 +3130,9 @@ static FValueRef BlendAggregate(MIR::FEmitter& Em,
 	for (int32 i = 0; i < MaterialAggregate->Attributes.Num(); ++i)
 	{
 		EMaterialExpressionBlendMode BlendMode = PixelAttributesBlendMode;
-		if (MaterialAggregate == UMaterialAggregate::GetMaterialAttributes())
+		if (MaterialAggregate == MaterialAttributesAggregate::Get())
 		{
-			EMaterialProperty Property = UMaterialAggregate::AttributeIndexToMaterialProperty(i);
+			EMaterialProperty Property = MaterialAttributesAggregate::AttributeIndexToMaterialProperty(i);
 			BlendMode = (Property == MP_WorldPositionOffset) ? VertexAttributesBlendMode : PixelAttributesBlendMode;
 		}
 
@@ -3195,7 +3195,7 @@ static FValueRef Blend(MIR::FEmitter& Em, EMaterialExpressionBlendMode PixelAttr
 
 void UMaterialExpressionBlendMaterialAttributes::Build(MIR::FEmitter& Em)
 {
-	const UMaterialAggregate* MaterialAttributes = UMaterialAggregate::GetMaterialAttributes();
+	const UMaterialAggregate* MaterialAttributes = MaterialAttributesAggregate::Get();
 	
 	FValueRef AValue = InputDefaultMaterialAttributes(Em, &A);
 	FValueRef BValue = InputDefaultMaterialAttributes(Em, &B);
@@ -3789,7 +3789,7 @@ static MIR::FType CustomMaterialOutputTypeToMIR(ECustomMaterialOutputType Type)
 		case CMOT_Float2: return MIR::FType::MakeFloatVector(2);
 		case CMOT_Float3: return MIR::FType::MakeFloatVector(3);
 		case CMOT_Float4: return MIR::FType::MakeFloatVector(4);
-		case CMOT_MaterialAttributes: return MIR::FType::MakeAggregate(UMaterialAggregate::GetMaterialAttributes());
+		case CMOT_MaterialAttributes: return MIR::FType::MakeAggregate(MaterialAttributesAggregate::Get());
 		default: UE_MIR_UNREACHABLE();
 	}
 }
@@ -4240,31 +4240,31 @@ void UMaterialExpressionSwitch::Build(MIR::FEmitter& Em)
 void UMaterialExpressionTangentOutput::Build(MIR::FEmitter& Em)
 {
 	FValueRef InputValue = Em.CastToFloat(Em.Input(&Input), 1);
-	Em.SetCustomOutputs(TEXTVIEW("TangentOutput"), { &InputValue, 1 }, MIR::MaterialOutputStage_Pixel);
+	Em.SetCustomOutputs(TEXTVIEW("TangentOutput"), { &InputValue, 1 }, MIR::EMaterialOutputFrequency::PerPixel);
 }
 
 void UMaterialExpressionBentNormalCustomOutput::Build(MIR::FEmitter& Em)
 {
 	FValueRef InputValue = Em.CastToFloat(Em.Input(&Input), 3);
-	Em.SetCustomOutputs(TEXTVIEW("BentNormal"), { &InputValue, 1 }, MIR::MaterialOutputStage_Pixel);
+	Em.SetCustomOutputs(TEXTVIEW("BentNormal"), { &InputValue, 1 }, MIR::EMaterialOutputFrequency::PerPixel);
 }
 
 void UMaterialExpressionClearCoatNormalCustomOutput::Build(MIR::FEmitter& Em)
 {
 	FValueRef InputValue = Em.CastToFloat(Em.Input(&Input), 3);
-	Em.SetCustomOutputs(TEXTVIEW("ClearCoatBottomNormal"), { &InputValue, 1 }, MIR::MaterialOutputStage_Pixel);
+	Em.SetCustomOutputs(TEXTVIEW("ClearCoatBottomNormal"), { &InputValue, 1 }, MIR::EMaterialOutputFrequency::PerPixel);
 }
 
 void UMaterialExpressionFirstPersonOutput::Build(MIR::FEmitter& Em)
 {
 	FValueRef InputValue = Em.CastToFloat(Em.InputDefaultFloat(&FirstPersonInterpolationAlpha, ConstFirstPersonInterpolationAlpha), 1);
-	Em.SetCustomOutputs(TEXTVIEW("FirstPersonOutput"), { &InputValue, 1 }, MIR::MaterialOutputStage_Vertex);
+	Em.SetCustomOutputs(TEXTVIEW("FirstPersonOutput"), { &InputValue, 1 }, MIR::EMaterialOutputFrequency::PerVertex);
 }
 
 void UMaterialExpressionMotionVectorWorldOffsetOutput::Build(MIR::FEmitter& Em)
 {
 	FValueRef InputValue = Em.CastToFloat(Em.Input(&Input), 3);
-	Em.SetCustomOutputs(TEXTVIEW("MotionVectorWorldOffsetOutput"), { &InputValue, 1 }, MIR::MaterialOutputStage_Pixel);
+	Em.SetCustomOutputs(TEXTVIEW("MotionVectorWorldOffsetOutput"), { &InputValue, 1 }, MIR::EMaterialOutputFrequency::PerPixel);
 }
 
 void UMaterialExpressionSingleLayerWaterMaterialOutput::Build(MIR::FEmitter& Em)
@@ -4287,8 +4287,7 @@ void UMaterialExpressionSingleLayerWaterMaterialOutput::Build(MIR::FEmitter& Em)
 		PhaseGValue ? Em.CastToFloat(PhaseGValue, 1) : Em.ConstantFloat(0.f),
 		Em.CastToFloat(Em.InputDefaultFloat(&ColorScaleBehindWater, 0.f), 1),
 	};
-	  
-	Em.SetCustomOutputs(TEXTVIEW("SingleLayerWaterMaterialOutput"), OutputValues, MIR::MaterialOutputStage_Pixel);
+	Em.SetCustomOutputs(TEXTVIEW("SingleLayerWaterMaterialOutput"), OutputValues, MIR::EMaterialOutputFrequency::PerPixel);
 }
 
 void UMaterialExpressionSubsurfaceMediumMaterialOutput::Build(MIR::FEmitter& Em)
@@ -4298,7 +4297,7 @@ void UMaterialExpressionSubsurfaceMediumMaterialOutput::Build(MIR::FEmitter& Em)
 		Em.CastToFloat(Em.Input(&MeanFreePath), 3),
 		Em.CastToFloat(Em.Input(&ScatteringDistribution), 1),
 	};
-	Em.SetCustomOutputs(TEXTVIEW("SubsurfaceMediumMaterialOutput"), OutputValues, MIR::MaterialOutputStage_Pixel);
+	Em.SetCustomOutputs(TEXTVIEW("SubsurfaceMediumMaterialOutput"), OutputValues, MIR::EMaterialOutputFrequency::PerPixel);
 }
 
 void UMaterialExpressionThinTranslucentMaterialOutput::Build(MIR::FEmitter& Em)
@@ -4308,7 +4307,7 @@ void UMaterialExpressionThinTranslucentMaterialOutput::Build(MIR::FEmitter& Em)
 		Em.CastToFloat(Em.InputDefaultFloat3(&TransmittanceColor, {0.5f, 0.5f, 0.5f}), 3),
 		Em.CastToFloat(Em.InputDefaultFloat(&SurfaceCoverage, 1.0f), 1),
 	};
-	Em.SetCustomOutputs(TEXTVIEW("ThinTranslucentMaterialOutput"), OutputValues, MIR::MaterialOutputStage_Pixel);
+	Em.SetCustomOutputs(TEXTVIEW("ThinTranslucentMaterialOutput"), OutputValues, MIR::EMaterialOutputFrequency::PerPixel);
 }
 
 void UMaterialExpressionVolumetricAdvancedMaterialInput::Build(MIR::FEmitter& Em)
@@ -4330,7 +4329,7 @@ void UMaterialExpressionVolumetricAdvancedMaterialOutput::Build(MIR::FEmitter& E
 		Em.CastToFloat(Em.InputDefaultFloat(&MultiScatteringEccentricity, ConstMultiScatteringEccentricity), 1),
 		Em.CastToFloat(Em.InputDefaultFloat4(&ConservativeDensity, FVector4f{ 1.0f, 1.0f, 1.0f, 1.0f }), 4),
 	};
-	Em.SetCustomOutputs(TEXTVIEW("VolumetricAdvancedMaterialOutput"), OutputValues, MIR::MaterialOutputStage_Pixel);
+	Em.SetCustomOutputs(TEXTVIEW("VolumetricAdvancedMaterialOutput"), OutputValues, MIR::EMaterialOutputFrequency::PerPixel);
 }
 
 void UMaterialExpressionAbsorptionMediumMaterialOutput::Build(MIR::FEmitter& Em)
@@ -4339,7 +4338,7 @@ void UMaterialExpressionAbsorptionMediumMaterialOutput::Build(MIR::FEmitter& Em)
 	{
 		Em.CastToFloat(Em.InputDefaultFloat3(&TransmittanceColor, { 1.0f, 1.0f, 1.0f }), 3),
 	};
-	Em.SetCustomOutputs(TEXTVIEW("AbsorptionMediumMaterialOutput"), OutputValues, MIR::MaterialOutputStage_Pixel);
+	Em.SetCustomOutputs(TEXTVIEW("AbsorptionMediumMaterialOutput"), OutputValues, MIR::EMaterialOutputFrequency::PerPixel);
 }
 
 void UMaterialExpressionDecalMipmapLevel::Build(MIR::FEmitter& Em)
@@ -4703,7 +4702,7 @@ void UMaterialExpressionSpeedTree::Build(MIR::FEmitter& Em)
 		[](FMaterialIRModule* Module, uintptr_t)
 		{
 			Module->AddEnvironmentDefine(TEXT("USES_SPEEDTREE"));
-			Module->GetStatistics().NumTexCoords = FMath::Max(Module->GetStatistics().NumTexCoords, 8);
+			Module->GetStatistics().NumInterpolatedTexCoords = FMath::Max(Module->GetStatistics().NumInterpolatedTexCoords, 8);
 		});
 
 	TStaticArray<FValueRef, MIR::EStage::NumStages> ValuePerStage = {};
@@ -4789,7 +4788,7 @@ void UMaterialExpressionNeuralNetworkInput::Build(MIR::FEmitter& Em)
 		Em.CastToFloat(Em.InputDefaultFloat(&Mask, 1.0f), 1)
 	};
 
-	Em.SetCustomOutputs(TEXTVIEW("NeuralInput"), OutputValues, MIR::MaterialOutputStage_Pixel);
+	Em.SetCustomOutputs(TEXTVIEW("NeuralInput"), OutputValues, MIR::EMaterialOutputFrequency::PerPixel);
 }
 
 void UMaterialExpressionNeuralNetworkOutput::Build(MIR::FEmitter& Em)
@@ -4888,14 +4887,14 @@ FSubstrateRegisteredSharedLocalBasis MIRSubstrateCompilationInfoCreateSharedLoca
 
 	SubstrateSharedNormalParamName = FString(TEXT("SubstrateNormal"));
 	SubstrateSharedNormalParamName.AppendInt(FinalSharedLocalBasisIndex);
-	Em->SetCustomOutputs(SubstrateSharedNormalParamName, { NormalArg, 1 }, MIR::MaterialOutputStage_Pixel);
+	Em->SetCustomOutputs(SubstrateSharedNormalParamName, { NormalArg, 1 }, MIR::EMaterialOutputFrequency::PerPixel);
 	SubstrateSharedNormalParamCode = FString::Printf(TEXT("Parameters.%s0"), *SubstrateSharedNormalParamName);
 
 	// Always write tangent even if not used later for now.
 	SubstrateSharedTangentParamName = FString(TEXT("SubstrateTangent"));
 	SubstrateSharedTangentParamName.AppendInt(FinalSharedLocalBasisIndex);
 	FValueRef DefaultTangent = EmitVertexTangent(*Em);
-	Em->SetCustomOutputs(SubstrateSharedTangentParamName, { TangentArg ? TangentArg : &DefaultTangent, 1}, MIR::MaterialOutputStage_Pixel);
+	Em->SetCustomOutputs(SubstrateSharedTangentParamName, { TangentArg ? TangentArg : &DefaultTangent, 1}, MIR::EMaterialOutputFrequency::PerPixel);
 	SubstrateSharedTangentParamCode = FString::Printf(TEXT("Parameters.%s0"), *SubstrateSharedTangentParamName);
 
 	FSubstrateCompilationContext& SubstrateCtx = SubstrateTranslatorData->SubstrateCompilationContext[SubstrateTranslatorData->CurrentSubstrateCompilationContext];
diff --git a/Engine/Source/Runtime/Engine/Private/Materials/MaterialIREmitter.cpp b/Engine/Source/Runtime/Engine/Private/Materials/MaterialIREmitter.cpp
--- a/Engine/Source/Runtime/Engine/Private/Materials/MaterialIREmitter.cpp
+++ b/Engine/Source/Runtime/Engine/Private/Materials/MaterialIREmitter.cpp
@@ -504,7 +504,7 @@ FType FEmitter::GetMaterialAggregateAttributeType(const UMaterialAggregate* Aggr
 		case EMaterialAggregateAttributeType::Float2: return FType::MakeVector(EScalarKind::Float, 2);
 		case EMaterialAggregateAttributeType::Float3: return FType::MakeVector(EScalarKind::Float, 3);
 		case EMaterialAggregateAttributeType::Float4: return FType::MakeVector(EScalarKind::Float, 4);
-		case EMaterialAggregateAttributeType::MaterialAttributes: return FType::MakeAggregate(UMaterialAggregate::GetMaterialAttributes());
+		case EMaterialAggregateAttributeType::MaterialAttributes: return FType::MakeAggregate(MaterialAttributesAggregate::Get());
 		case EMaterialAggregateAttributeType::SubstrateData: return FType::MakeSubstrateData();
 		case EMaterialAggregateAttributeType::Aggregate: return FType::MakeAggregate(Aggregate->Attributes[AttributeIndex].Aggregate.Get());
 		default: UE_MIR_UNREACHABLE();
@@ -1397,7 +1397,7 @@ FValueRef FEmitter::ScreenTexture(EScreenTexture TextureKind)
 
 // Helper function to create a single SetMaterialOutput instruction. It is used to create
 // the final instruction for both base material attributes and custom outputs.
-static FSetMaterialOutput* CreateSetMaterialOutput(FEmitter& Em, FMaterialIRModule* Module, FStringView Name, EMaterialProperty Property, EMaterialOutputStage Stage, FValueRef Arg)
+static FSetMaterialOutput* CreateSetMaterialOutput(FEmitter& Em, FMaterialIRModule* Module, FStringView Name, EMaterialProperty Property, EMaterialOutputFrequency Frequency, FValueRef Arg)
 {
 	// Material attributes is not a real material attribute, the translator should never create a SetMaterialOutput on this property.
 	check(Property != MP_MaterialAttributes); 
@@ -1405,7 +1405,7 @@ static FSetMaterialOutput* CreateSetMaterialOutput(FEmitter& Em, FMaterialIRModu
 	// Create the set material output instruction.
 	FSetMaterialOutput* Instr = EmitNew<FSetMaterialOutput>(Em, FType::MakeVoid());
 	Instr->Property           = Property;
-	Instr->Stage	          = Stage;
+	Instr->Frequency	      = Frequency;
 	Instr->Name               = Module->InternString(Name);
 	Instr->Arg                = Arg;
 
@@ -1415,7 +1415,7 @@ static FSetMaterialOutput* CreateSetMaterialOutput(FEmitter& Em, FMaterialIRModu
                     : 1; 								// All other attributes (BaseColor, Emissive, etc)
 
 	// Add the instruction to list of outputs of the stages it is evaluated in.
-	if (Stage == MaterialOutputStage_Vertex)
+	if (Frequency == EMaterialOutputFrequency::PerVertex)
 	{
 		Module->GetEntryPoint(MIR::Stage_Vertex).Outputs.Push(Instr);
 	}
@@ -1428,17 +1428,23 @@ static FSetMaterialOutput* CreateSetMaterialOutput(FEmitter& Em, FMaterialIRModu
 	return Instr;
 }
 
+static bool IsVertexStageMaterialProperty(EMaterialProperty Property)
+{
+	return Property == MP_WorldPositionOffset
+		|| (Property >= MP_CustomizedUVs0 && Property <= MP_CustomizedUVs7);
+}
+
 void FEmitter::Private_SetMaterialOutput(EMaterialProperty Property, FValueRef Arg)
 {
 	checkf(Property != MP_CustomOutput, TEXT("This API does not support CustomOutputs, use SetCustomOutput() instead."));
 
-	FString Name         = UMaterialAggregate::GetMaterialAttribute(Property)->Name.ToString();
-	EMaterialOutputStage Stage = (Property == MP_WorldPositionOffset) ? MIR::MaterialOutputStage_Vertex : MIR::MaterialOutputStage_Pixel;
+	const FString Name          	 	= MaterialAttributesAggregate::GetAttribute(Property)->Name.ToString();
+	const EMaterialOutputFrequency Freq = IsVertexStageMaterialProperty(Property) ? MIR::EMaterialOutputFrequency::PerVertex : MIR::EMaterialOutputFrequency::PerPixel;
 
-	CreateSetMaterialOutput(*this, Module, Name, Property, Stage, Arg);
+	CreateSetMaterialOutput(*this, Module, Name, Property, Freq, Arg);
 }
 
-void FEmitter::SetCustomOutputs(FStringView Name, TConstArrayView<FValueRef> Args, EMaterialOutputStage Stage)
+void FEmitter::SetCustomOutputs(FStringView Name, TConstArrayView<FValueRef> Args, EMaterialOutputFrequency Frequency)
 {
 	// If this is a custom output, check that its name is unique.
 	for (const FMaterialIRModule::FCustomOutputGroup& Group : Module->GetCustomOutputGroups())
@@ -1472,7 +1478,7 @@ void FEmitter::SetCustomOutputs(FStringView Name, TConstArrayView<FValueRef> Arg
 		int32 NameLen = Name.Len() + FCString::Snprintf(TempNameBuffer.GetData() + Name.Len(), TempNameBuffer.Num() - Name.Len(), TEXT("%d"), i);
 
 		// Create the SetMaterialOutput instruction
-		FSetMaterialOutput* Instr = CreateSetMaterialOutput(*this, Module, { TempNameBuffer.GetData(), NameLen }, MP_CustomOutput, Stage, Args[i]);
+		FSetMaterialOutput* Instr = CreateSetMaterialOutput(*this, Module, { TempNameBuffer.GetData(), NameLen }, MP_CustomOutput, Frequency, Args[i]);
 	
 		// Push the instruction to the custom outputs in the module
 		Module->AddCustomOutput(Instr);
@@ -1644,7 +1650,7 @@ FValueRef FEmitter::Aggregate(const UMaterialAggregate* InAggregate)
 	return Aggregate(InAggregate, {}, TConstArrayView<FValueRef>{});
 }
 
-static FValueRef EmitAttributeDefaultValue(FEmitter& Emitter, const UMaterialInterface* Material, const FMaterialAggregateAttribute& Attribute)
+static FValueRef EmitAttributeDefaultValue(FEmitter& Emitter, UMaterialInterface* Material, const FMaterialAggregateAttribute& Attribute)
 {
 	switch (Attribute.Type)
 	{
@@ -1661,7 +1667,7 @@ static FValueRef EmitAttributeDefaultValue(FEmitter& Emitter, const UMaterialInt
 		case EMaterialAggregateAttributeType::Float3: return Emitter.ConstantFloat3({ (TFloat)Attribute.DefaultValue.X, (TFloat)Attribute.DefaultValue.Y, (TFloat)Attribute.DefaultValue.Z });
 		case EMaterialAggregateAttributeType::Float4: return Emitter.ConstantFloat4(UE::Math::TVector4<TFloat>{ Attribute.DefaultValue.X, Attribute.DefaultValue.Y, Attribute.DefaultValue.Z, Attribute.DefaultValue.W });
 		case EMaterialAggregateAttributeType::ShadingModel: return Emitter.ConstantInt((TInteger)Material->GetShadingModels().GetFirstShadingModel());
-		case EMaterialAggregateAttributeType::MaterialAttributes: return Emitter.Aggregate(UMaterialAggregate::GetMaterialAttributes());
+		case EMaterialAggregateAttributeType::MaterialAttributes: return Emitter.Aggregate(MaterialAttributesAggregate::Get());
 		case EMaterialAggregateAttributeType::SubstrateData: return Emitter.SubstrateDefaultSlab();
 		case EMaterialAggregateAttributeType::Aggregate: return Emitter.Aggregate(Attribute.Aggregate);
 		default: UE_MIR_UNREACHABLE();
@@ -4231,7 +4237,7 @@ FValueRef FEmitter::InlineHLSL(const FMaterialExternalCodeDeclaration* InExterna
 	}
 
 	check(InExternalCodeDeclaration != nullptr);
-	FType ReturnType = FType::FromMaterialValueType(InExternalCodeDeclaration->GetReturnTypeValue());
+	FType ReturnType = FType::FromMaterialValueType( InExternalCodeDeclaration->GetReturnTypeValue());
 	return EmitInlineHLSL(*this, ReturnType, InExternalCodeDeclaration, {}, InputValues, ValueFlags, UsedGraphProperties);
 }
 
diff --git a/Engine/Source/Runtime/Engine/Private/Materials/MaterialIRModule.cpp b/Engine/Source/Runtime/Engine/Private/Materials/MaterialIRModule.cpp
--- a/Engine/Source/Runtime/Engine/Private/Materials/MaterialIRModule.cpp
+++ b/Engine/Source/Runtime/Engine/Private/Materials/MaterialIRModule.cpp
@@ -39,12 +39,6 @@ void FMaterialIRModule::Empty()
 	CustomOutputGroups.Empty();
 	SubstrateData = {};
 	Allocator.Flush();
-
-	// Initialize statistics used external input usage mask.
-	for (int i = 0; i < MIR::NumStages; ++i)
-	{
-		Statistics.ExternalInputUsedMask[i].Init(false, (int)MIR::EExternalInput::Count);
-	}
 }
 
 int32 FMaterialIRModule::AddEntryPoint(FStringView Name, MIR::EStage Stage, int32 NumOutputsHint)
diff --git a/Engine/Source/Runtime/Engine/Private/Materials/MaterialIRModuleBuilder.cpp b/Engine/Source/Runtime/Engine/Private/Materials/MaterialIRModuleBuilder.cpp
--- a/Engine/Source/Runtime/Engine/Private/Materials/MaterialIRModuleBuilder.cpp
+++ b/Engine/Source/Runtime/Engine/Private/Materials/MaterialIRModuleBuilder.cpp
@@ -109,7 +109,7 @@ struct FMaterialIRModuleBuilderImpl
 		Module->AddEntryPoint(TEXTVIEW("VertexStage"), MIR::Stage_Vertex, 1);
 
 		// Declare the entry points to evaluate both the pixel and compute stages.
-		int32 Num = UMaterialAggregate::GetMaterialAttributesProperties().Num();
+		int32 Num = MaterialAttributesAggregate::GetMaterialProperties().Num();
 		Module->AddEntryPoint(TEXTVIEW("PixelStage"), MIR::Stage_Pixel, Num);
 		Module->AddEntryPoint(TEXTVIEW("ComputeStage"), MIR::Stage_Compute, Num);
 
@@ -123,7 +123,7 @@ struct FMaterialIRModuleBuilderImpl
 		Emitter.Initialize();
 
 		// Create an IR value to hold the material attributes aggregate default
-		DefaultMaterialAggregate = Emitter.Aggregate(UMaterialAggregate::GetMaterialAttributes());
+		DefaultMaterialAggregate = Emitter.Aggregate(MaterialAttributesAggregate::Get());
 
 		// Push the root analysis context
 		AnalysisContextStack.Emplace();
@@ -131,6 +131,9 @@ struct FMaterialIRModuleBuilderImpl
 		// Set the preview input expression
 		PreviewInput.Expression = Builder->PreviewExpression;
 
+
+		// SUBSTRATE INITIALIZATION
+
 		// Setup Substrate translator data to process the Substrate tree and simplification.
 		FSubstrateCompilationConfig* SubstrateCompilationConfig = nullptr;	// SUBSTRATE_TODO: forward SubstrateCompilationConfig from compilation step to make substrate tab simplification preview work.
 		if (SubstrateCompilationConfig)
@@ -160,9 +163,20 @@ struct FMaterialIRModuleBuilderImpl
 		SubstrateTranslatorData.ProcessSubstrateTreeAndTopology(MaterialShadingModels, MaterialGraphNodePreviewExpression, bDoClearFunctionStack);
 	}
 
+	// Some properties may be deactivated in the material, such as CustomizedUVx properties. This function returns whether given
+	// property is active in the translating material.
+	bool PropertyIsActiveInMaterial(const EMaterialProperty Property)
+	{
+		if (MIR::IsMaterialPropertyCustomizedUV(Property) && Property >= MP_CustomizedUVs0 + Emitter.BaseMaterial->NumCustomizedUVs)
+		{
+			return false;
+		}
+		return true;
+	}
+
 	void Step_PushRootExpressionDependencies()
 	{
-		//If we are processing from the Material Attributes output node, push the dependency here.
+		//If we are processing from the Material Attributes output node, push that dependency here.
 		if (Builder->MaterialInterface->GetMaterial()->bUseMaterialAttributes)
 		{
 			FMaterialInputDescription Input;
@@ -171,8 +185,15 @@ struct FMaterialIRModuleBuilderImpl
 		}
 		else
 		{
-			for (EMaterialProperty Property : UMaterialAggregate::GetMaterialAttributesProperties())
+			// Otherwise iterate over each individual material attribute pin and push its connected expression as a dependency.
+			for (EMaterialProperty Property : MaterialAttributesAggregate::GetMaterialProperties())
 			{
+				// Ignore deactivated properties.
+				if (!PropertyIsActiveInMaterial(Property))
+				{
+					continue;
+				}
+
 				// Read the material input associated to this property
 				FMaterialInputDescription Input;
 				GetExpressionInputDescription(Property, Input);
@@ -233,7 +254,7 @@ struct FMaterialIRModuleBuilderImpl
 		{
 			// Save the the FrontMaterial to a custom output.
 			MIR::FValueRef SubstrateFrontMaterial = MIR::FValueRef(AnalysisContextStack.Last().GetOutputValue(SubstrateTranslatorData.FrontMaterialExpr->GetOutput(0))); // SUBSTRATE_TODO always assuming output 0 for now but this can be wrong for some material function
-			Emitter.SetCustomOutputs(TEXTVIEW("SubstrateFrontMaterial"), { &SubstrateFrontMaterial, 1 }, MIR::EMaterialOutputStage::MaterialOutputStage_Pixel);
+			Emitter.SetCustomOutputs(TEXTVIEW("SubstrateFrontMaterial"), { &SubstrateFrontMaterial, 1 }, MIR::EMaterialOutputFrequency::PerPixel);
 		}
 
 		// Remove all substrate nodes and function calls from all the contexts.
@@ -299,7 +320,7 @@ struct FMaterialIRModuleBuilderImpl
 			{
 				// Save the output to the FrontMaterial, now fully simplified, to the custom output representing fully simplified.
 				MIR::FValueRef SubstrateFrontMaterial = MIR::FValueRef(AnalysisContextStack.Last().GetOutputValue(SubstrateTranslatorData.FrontMaterialExpr->GetOutput(0)));
-				Emitter.SetCustomOutputs(TEXTVIEW("SubstrateFrontMaterialFullySimplified"), { &SubstrateFrontMaterial, 1 }, MIR::EMaterialOutputStage::MaterialOutputStage_Pixel);
+				Emitter.SetCustomOutputs(TEXTVIEW("SubstrateFrontMaterialFullySimplified"), { &SubstrateFrontMaterial, 1 }, MIR::EMaterialOutputFrequency::PerPixel);
 
 				SubstrateTranslatorData.CurrentSubstrateCompilationContext = Substrate::ESubstrateCompilationContext::SCC_Default;
 			}
@@ -653,7 +674,7 @@ struct FMaterialIRModuleBuilderImpl
 		// First, if the material is flagged to use the material attributes aggregate, read its value now
 		// so that we can extract its individual attributes later.
 		MIR::FValue* MaterialAttributesValue = nullptr;
-		if (Builder->MaterialInterface->GetMaterial()->bUseMaterialAttributes)
+		if (Emitter.BaseMaterial->bUseMaterialAttributes)
 		{
 			FMaterialInputDescription InputDesc;
 			GetExpressionInputDescription(MP_MaterialAttributes, InputDesc);
@@ -667,11 +688,17 @@ struct FMaterialIRModuleBuilderImpl
 			}
 
 			// Make sure a valid value is present and it is of the correct type.
-			check(!MaterialAttributesValue->IsPoison() && MaterialAttributesValue->Type.AsAggregate() == UMaterialAggregate::GetMaterialAttributes());
+			check(!MaterialAttributesValue->IsPoison() && MaterialAttributesValue->Type.AsAggregate() == MaterialAttributesAggregate::Get());
 		}
 
-		for (EMaterialProperty Property : UMaterialAggregate::GetMaterialAttributesProperties())
+		for (EMaterialProperty Property : MaterialAttributesAggregate::GetMaterialProperties())
 		{
+			// Ignore CustomUVs properties that haven't been enabled in the base material.
+			if (!PropertyIsActiveInMaterial(Property))
+			{
+				continue;
+			}
+
 			// Get the input description of this material property (input, type, default value, etc).
 			FMaterialInputDescription InputDesc;
 			GetExpressionInputDescription(Property, InputDesc);
@@ -683,7 +710,7 @@ struct FMaterialIRModuleBuilderImpl
 			// and manually flow it into the this property material expression input pin.
 			if (MaterialAttributesValue)
 			{
-				PropertyValue = Emitter.Subscript(MaterialAttributesValue, UMaterialAggregate::MaterialPropertyToAttributeIndex(Property));
+				PropertyValue = Emitter.Subscript(MaterialAttributesValue, MaterialAttributesAggregate::MaterialPropertyToAttributeIndex(Property));
 			}
 			else
 			{
@@ -708,7 +735,7 @@ struct FMaterialIRModuleBuilderImpl
 						PropertyValue = Emitter.Vector3(PropertyValue, Zero, Zero);
 					}
 				}
-				else if (Property == MP_FrontMaterial )
+				else if (Property == MP_FrontMaterial)
 				{
 					// Just check that we get a struct for now. 
 					// SUBSTRATE_TODO We should probably add a proper SubstrateData struct enum?
@@ -724,6 +751,13 @@ struct FMaterialIRModuleBuilderImpl
 					PropertyValue = Emitter.Cast(PropertyValue, OutputArgType);
 				}
 			}
+			else if (MIR::IsMaterialPropertyCustomizedUV(Property))
+			{
+				// To keep backwards compatibility with the old translator, a CustomUV attribute pin with no given value
+				// does not default to (0,0) (as it should) but to the uncustomized texture coordinate. If we find a customized UV
+				// attribute with no value, simply don't emit it entirely. This will default to the default texture coordinates.
+				continue;
+			}
 			else if (Property == MP_FrontMaterial && PropertyValue == nullptr && SubstrateTranslatorData.bMaterialUsesRootNodeToSubstrateHiddenConversion)
 			{
 				PropertyValue = BuildSubstrateFrontMaterialIRFromRootNodeHiddenConversion();
@@ -736,7 +770,7 @@ struct FMaterialIRModuleBuilderImpl
 			else
 			{
 				// Otherwise, fallback to assigning this material output to its default value.
-				PropertyValue = Emitter.Subscript(DefaultMaterialAggregate, UMaterialAggregate::MaterialPropertyToAttributeIndex(Property));
+				PropertyValue = Emitter.Subscript(DefaultMaterialAggregate, MaterialAttributesAggregate::MaterialPropertyToAttributeIndex(Property));
 			}
 			
 			// Quit if some error occurred in the operations above.
@@ -758,9 +792,6 @@ struct FMaterialIRModuleBuilderImpl
 				PropertyValue = Emitter.Lerp(PropertyValue, Emitter.Swizzle(SelectionColor, MIR::FSwizzleMask::XYZ()), Emitter.Subscript(SelectionColor, 3));
 			}
 
-			// Set this property value onto the module
-			Module->SetPropertyValue(Property, PropertyValue);
-
 			// Emit the SetMaterialProperty instruction
 			Emitter.Private_SetMaterialOutput(Property, PropertyValue);
 
@@ -850,8 +881,8 @@ struct FMaterialIRModuleBuilderImpl
 			// When we are running hidden conversion (bMaterialUsesRootNodeToSubstrateHiddenConversion),
 			// We know that we will only ever have a single BSDF and shared local bases. So simply use the same non-fully-simplified as the fully-simplified version of the material too.
 			MIR::FValueRef PropertyValueRef = MIR::FValueRef(FrontMaterialPropertyValue);
-			Emitter.SetCustomOutputs(TEXTVIEW("SubstrateFrontMaterialFullySimplified"), { &PropertyValueRef, 1 }, MIR::EMaterialOutputStage::MaterialOutputStage_Pixel);
-			Emitter.SetCustomOutputs(TEXTVIEW("SubstrateFrontMaterial"), { &PropertyValueRef, 1 }, MIR::EMaterialOutputStage::MaterialOutputStage_Pixel);
+			Emitter.SetCustomOutputs(TEXTVIEW("SubstrateFrontMaterialFullySimplified"), { &PropertyValueRef, 1 }, MIR::EMaterialOutputFrequency::PerPixel);
+			Emitter.SetCustomOutputs(TEXTVIEW("SubstrateFrontMaterial"), { &PropertyValueRef, 1 }, MIR::EMaterialOutputFrequency::PerPixel);
 			// We also copy the copy the shared local bases structure from within GeneratePixelInputInitializerHLSLCode.
 		}
 		else
@@ -872,12 +903,12 @@ struct FMaterialIRModuleBuilderImpl
 			const bool bValid = SubstrateTranslatorData.GenerateSubstrateTreeHLSLCode(bSubstrateFrontMaterialProvided);
 			Module->SubstrateData.TreeHLSL = SubstrateTranslatorData.SubstrateTreeHLSL;
 
-			// TOOD add that back
-			//			if (!bValid)
-			//			{
-			//				TranslationResult = EHLSLMaterialTranslatorResult::Failure;
-			//				return;
-			//			}
+			// TODO add that back
+			// if (!bValid)
+			// {
+			// 	TranslationResult = EHLSLMaterialTranslatorResult::Failure;
+			// 	return;
+			// }
 		}
 
 		// SUBSTRATE_TODO export context
@@ -945,20 +976,33 @@ struct FMaterialIRModuleBuilderImpl
 		TArray<MIR::FValue*> ValueStack{};
 		TSet<MIR::FValue*> VisitedValues{};
 
-		// Analyze the nodes in each entry point.
-		// Note: process the vertex stage last so that vertex interpolator getters in the pixel stage are analyzed first,
-		//       adding the setters in the vertex stage entry point outputs. This is done so that only actually used vertex interpolators
-		//       are compiled in.
-		check(Module->GetEntryPoint(0).Stage == MIR::EStage::Stage_Vertex);
-		for (int i = 1; i <= Module->GetNumEntryPoints(); ++i)
+		// The analysis of several vertex values (e.g. interpolators, customized UVs) depend on other non-vertex values
+		// values analysis results. For this reason, analyze stages in reverse order of execution (compute/pixel stages first,
+		// vertex stage last). For example, the entire subgraph evaluating a CustomizedUV can be optimized out if it is found
+		// that during pixel/compute stages that relative texture coordinate isn't actually used.
+		MIR::TTemporaryArray<int32> SortedEntryPointIndices { Module->GetNumEntryPoints() };
+		for (int32 i = 0; i < Module->GetNumEntryPoints(); ++i)
+		{
+			SortedEntryPointIndices[i] = i;
+		}
+		SortedEntryPointIndices.StableSort([this](int32 A, int32 B)
+		{
+			return Module->GetEntryPoint(A).Stage > Module->GetEntryPoint(B).Stage;
+		});
+
+		// Analyze all the nodes in each entry point starting from their root outputs and walking backwards along the Use edges.
+		for (int32 EntryPointIndex : SortedEntryPointIndices)
 		{
 			// Reset bookkeeping to process new output subgraph
 			ValueStack.Empty(ValueStack.Max());
 			VisitedValues.Empty(Module->Values.Num());
 
-			const int EntryPointIndex = i % Module->GetNumEntryPoints();
 			FMaterialIRModule::FEntryPoint& EntryPoint = Module->GetEntryPoint(EntryPointIndex);
 
+			// If the analysis of previous stages resulted in some outputs in this stage to be actually unnecessary, remove
+			// these optimized out root outputs now.
+			RemoveDisabledEntryPointsOutputs(EntryPoint);
+
 			// Stable sort the outputs in this entry point by their indicated priority, then name.
 			// Note: this is needed to support situations where a final entry-point instruction needs to execute
 			//    before another, for instance we want to set the Normal attribute before all others, because the
@@ -1022,10 +1066,8 @@ struct FMaterialIRModuleBuilderImpl
 
 				// If this is the first time this value is analyzed in any entry point, let the analyzer process it.
 				// Note that individual value processing is independent from the stage it runs on so we can perform it only once.
-				if (!Value->HasFlags(MIR::EValueFlags::AnalyzedInAnyStage))
+				if ((Value->Flags & MIR::EValueFlags::AnalyzedInAllStagesMask) == MIR::EValueFlags::None)
 				{
-					Value->SetFlags(MIR::EValueFlags::AnalyzedInAnyStage);
-
 					// Flow the graph properties downstream from the value's uses into this value.
 					for (MIR::FValue* Use : Value->GetUses())
 					{
@@ -1062,6 +1104,21 @@ struct FMaterialIRModuleBuilderImpl
 		return Module->IsValid();
 	}
 
+	// Remove instructions from the given entry point outputs that, due to semantic analysis of previous other, are found not to contribute to the end result.
+	void RemoveDisabledEntryPointsOutputs(FMaterialIRModule::FEntryPoint& EntryPoint)
+	{
+		if (EntryPoint.Stage == MIR::Stage_Vertex)
+		{
+			EntryPoint.Outputs.RemoveAllSwap([this](const MIR::FInstruction* Instr)
+			{
+				const MIR::FSetMaterialOutput* SetMaterialOutput = MIR::As<MIR::FSetMaterialOutput>(Instr);
+				return SetMaterialOutput
+					&& MIR::IsMaterialPropertyCustomizedUV(SetMaterialOutput->Property)
+					&& (Module->GetStatistics().InterpolatedTexCoordsMask & (1 << (SetMaterialOutput->Property - MP_CustomizedUVs0))) == 0;
+			});
+		}
+	}
+
 	bool Step_PerformModuleWideAnalysis()
 	{
 		// If the material doesn't use expression shading models, or they aren't valid, initialize the shading models to the ones from the material.
@@ -1147,7 +1204,7 @@ struct FMaterialIRModuleBuilderImpl
 		CompilationOutput.UniformExpressionSet.SetParameterCollections(Module->ParameterCollections);
 
 		// NumUsedUVScalars
-		CompilationOutput.NumUsedUVScalars = Module->Statistics.NumTexCoords * 2;
+		CompilationOutput.NumUsedUVScalars = Module->Statistics.NumInterpolatedTexCoords * 2;
 
 		// NumVirtualTextureFeedbackRequests
 		int32 NumVirtualTextureFeedbackRequests = 0;
@@ -1205,16 +1262,19 @@ struct FMaterialIRModuleBuilderImpl
 			// This function walks the instruction graph and puts each instruction into the inner most possible block.
 			InstructionStack.Empty(InstructionStack.Max());
 
-			FMaterialIRModule::FEntryPoint& Procedure = Module->GetEntryPoint(EntryPointIndex);
+			FMaterialIRModule::FEntryPoint& EntryPoint = Module->GetEntryPoint(EntryPointIndex);
 
 			// Push all entry point final outputs onto the instruction stack to begin.
 			// Note: the first output on the stack will be the first to be evaluated in the entry point root block.
-			for (MIR::FValue* Output : Procedure.Outputs)
+			for (MIR::FValue* Output : EntryPoint.Outputs)
 			{
-				if (MIR::FInstruction* Instr = MIR::AsInstruction(Output))
+				if (Output->Kind != MIR::EValueKind::VK_Nop)
 				{
-					Instr->Linkage[EntryPointIndex].Block = &Procedure.RootBlock;
-					InstructionStack.Push(Instr);
+					if (MIR::FInstruction* Instr = MIR::AsInstruction(Output))
+					{
+						Instr->Linkage[EntryPointIndex].Block = &EntryPoint.RootBlock;
+						InstructionStack.Push(Instr);
+					}
 				}
 			}
 
@@ -1227,17 +1287,16 @@ struct FMaterialIRModuleBuilderImpl
 				Instr->Linkage[EntryPointIndex].Next = InstrBlock->Instructions;
 				InstrBlock->Instructions = Instr;
 
-				//
+				// NOP instructions uses are only analyzed but not evaluated. Ignore them.
 				if (Instr->As<MIR::FNop>())
 				{
 					continue;
 				}
 
-				TConstArrayView<MIR::FValue*> Uses = Instr->GetUsesForStage(Procedure.Stage);
+				TConstArrayView<MIR::FValue*> Uses = Instr->GetUsesForStage(EntryPoint.Stage);
 				for (int32 UseIndex = 0; UseIndex < Uses.Num(); ++UseIndex)
 				{
 					MIR::FInstruction* UseInstr = MIR::AsInstruction(Uses[UseIndex]);
-
 					if (!UseInstr)
 					{
 						continue;
@@ -1261,7 +1320,6 @@ struct FMaterialIRModuleBuilderImpl
 						: TargetBlock;
 
 					// Increase the number of times this dependency instruction has been considered.
-					// When all of its users have processed, we can carry on visiting this instruction.
 					++UseLinkage.NumProcessedUsers;
 					check(UseLinkage.NumProcessedUsers <= UseLinkage.NumUsers);
 
@@ -1394,7 +1452,7 @@ bool FMaterialIRModuleBuilder::Build(FMaterialIRModule* TargetModule)
 	// visited expression in order to emit the IR values that implement that expression
 	// semantics. At the end of this step the IR values graph has been built, but is still
 	// missing the root SetMaterialProperty instructions.
-	if(!Impl.Step_BuildMaterialExpressionsToIRGraph())
+	if (!Impl.Step_BuildMaterialExpressionsToIRGraph())
 	{
 		return false;
 	}
diff --git a/Engine/Source/Runtime/Engine/Private/Materials/MaterialIRToHLSLTranslator.cpp b/Engine/Source/Runtime/Engine/Private/Materials/MaterialIRToHLSLTranslator.cpp
--- a/Engine/Source/Runtime/Engine/Private/Materials/MaterialIRToHLSLTranslator.cpp
+++ b/Engine/Source/Runtime/Engine/Private/Materials/MaterialIRToHLSLTranslator.cpp
@@ -1176,8 +1176,16 @@ struct FEntryPointCodeGen
 
 	void LowerSetMaterialOutput(const MIR::FSetMaterialOutput* SetMaterialOutput)
 	{
-		Printer << (CurrentStage != MIR::Stage_Vertex && SetMaterialOutput->Property != MP_CustomOutput ? TEXTVIEW("PixelMaterialInputs.") : TEXTVIEW("Parameters."))
-				<< SetMaterialOutput->Name.ToView() << TEXTVIEW(" = ") << LowerValue(SetMaterialOutput->Arg);
+		if (SetMaterialOutput->Property >= MP_CustomizedUVs0 && SetMaterialOutput->Property <= MP_CustomizedUVs7)
+		{
+			Printer << TEXTVIEW("Parameters.CustomizedUVs[") << (int)(SetMaterialOutput->Property - MP_CustomizedUVs0) << TEXT(']');
+		}
+		else
+		{
+			Printer << (CurrentStage != MIR::Stage_Vertex && SetMaterialOutput->Property != MP_CustomOutput ? TEXTVIEW("PixelMaterialInputs.") : TEXTVIEW("Parameters.")) << SetMaterialOutput->Name.ToView();
+		}
+
+		Printer << TEXTVIEW(" = ") << LowerValue(SetMaterialOutput->Arg);
 		
 		bLastInstructionWasSetMaterialOutput = true;
 	}
@@ -2189,7 +2197,7 @@ struct FEntryPointCodeGen
 
 		while (SlotBegin < SlotEnd)
 		{
-			const int RegisterIndex = Module->GetStatistics().NumTexCoords + SlotBegin / 2;
+			const int RegisterIndex = Module->GetStatistics().NumInterpolatedTexCoords + SlotBegin / 2;
 			const int RegisterSlotOffset = SlotBegin % 2;
 
 			// How many components fit in the current float2 from this offset. The minimum between the remaining
@@ -2270,7 +2278,7 @@ static void GenerateStageEntryPoints(const FMaterialIRModule* Module, TMap<FStri
 
 	// Generate pixel stage HLSL
 	Result = CodeGen.GenerateEntryPoint(MIR::Stage_Pixel);
-	Result.NonNormalHLSL += TEXTVIEW(TAB "PixelMaterialInputs.Subsurface = 0;\n");
+	Result.NonNormalHLSL += TEXTVIEW("PixelMaterialInputs.Subsurface = 0;\n");
 	Result.NonNormalHLSL += CodeGen.Module->GetSubstrateData().PixelInputInitializerValuesHLSL;
 	OutParameters.Add(TEXT("calc_pixel_material_inputs_normal"), MoveTemp(Result.NormalAttributeHLSL));
 	OutParameters.Add(TEXT("calc_pixel_material_inputs_other_inputs"), MoveTemp(Result.NonNormalHLSL));
@@ -2281,17 +2289,23 @@ static void GenerateStageEntryPoints(const FMaterialIRModule* Module, TMap<FStri
 	OutParameters.Add(TEXT("calc_pixel_material_inputs_analytic_derivatives_normal"), MoveTemp(Result.NormalAttributeHLSL));
 	OutParameters.Add(TEXT("calc_pixel_material_inputs_analytic_derivatives_other_inputs"), MoveTemp(Result.NonNormalHLSL));
 }
-	
-static void GenerateCustomOutputDeclarations(const FMaterialIRModule* Module, TMap<FString, FString>& OutParameters)
+
+static void GenerateCustomOutputDeclarations(const FMaterial* Material, const FMaterialIRModule* Module, TMap<FString, FString>& OutParameters)
 {
 	// Generate the custom output fields of the struct FMaterialPixelParameters. These fields will be computed and 
 	// set by the function that calculates all material outputs (e.g. BaseColor).
 	FHLSLPrinter VertexPrinter{ 1 };
 	FHLSLPrinter PixelPrinter{ 1 };
 
+	const int NumCustomizedUVs = Material->GetNumCustomizedUVs();
+	if (NumCustomizedUVs > 0)
+	{
+		VertexPrinter << TEXTVIEW("MaterialFloat2 CustomizedUVs[") << NumCustomizedUVs << TEXT(']') << EndOfStatement;
+	}
+
 	for (const MIR::FSetMaterialOutput* Output : Module->GetCustomOutputs())
 	{
-		FHLSLPrinter& ThePrinter = Output->Stage == MIR::MaterialOutputStage_Vertex ? VertexPrinter : PixelPrinter;
+		FHLSLPrinter& ThePrinter = Output->Frequency == MIR::EMaterialOutputFrequency::PerVertex ? VertexPrinter : PixelPrinter;
 		ThePrinter << Output->Arg->Type << TEXT(' ') << *Output->Name << EndOfStatement;
 	}
 		
@@ -2314,7 +2328,7 @@ static void GenerateCustomOutputDeclarations(const FMaterialIRModule* Module, TM
 			Printer << TEXTVIEW("#define HAVE_Get") << *Instr->Name << TEXTVIEW(" 1") << NewLine << NewLine;
 				
 			// Select the right XXXParameters HLSL struct used by this output based on its 
-			FStringView Signature = Instr->Stage == MIR::MaterialOutputStage_Vertex
+			FStringView Signature = Instr->Frequency == MIR::EMaterialOutputFrequency::PerVertex
 				? TEXTVIEW("(inout FMaterialVertexParameters Parameters)")
 				: TEXTVIEW("(inout FMaterialPixelParameters Parameters)");
 
@@ -2431,13 +2445,13 @@ static void GenerateOtherTemplateStringParameters(const FMaterial* Material, con
 	// pixel_material_inputs (field declarations of the FPixelMaterialInputs struct)
 	{
 		FHLSLPrinter Printer = { 1 };
-		for (EMaterialProperty Property : UMaterialAggregate::GetMaterialAttributesProperties())
+		for (EMaterialProperty Property : MaterialAttributesAggregate::GetMaterialProperties())
 		{
-			if (Property == MP_WorldPositionOffset) // todo: fixme
+			if (Property == MP_WorldPositionOffset || (Property >= MP_CustomizedUVs0 && Property <= MP_CustomizedUVs7)) // todo: fixme
 			{
 				continue;
 			}
-			if (const FMaterialAggregateAttribute* PropertyAggregate = UMaterialAggregate::GetMaterialAttribute(Property))
+			if (const FMaterialAggregateAttribute* PropertyAggregate = MaterialAttributesAggregate::GetAttribute(Property))
 			{
 				Printer << MIR::FType::FromMaterialValueType(PropertyAggregate->ToMaterialValueType()) << TEXT(' ') << *PropertyAggregate->Name.ToString() << EndOfStatement;
 			}
@@ -2450,22 +2464,25 @@ static void GenerateOtherTemplateStringParameters(const FMaterial* Material, con
 	}
 
 	const int32 NumUserVertexInterpolators = (ModuleStatistics.NumUsedVertexInterpolatorFloatRegisters + 1) / 2; // packed into float2s
-	const int32 ShaderTexCoordsArrayLength = ModuleStatistics.NumTexCoords + NumUserVertexInterpolators;
+	const int32 ShaderTexCoordsArrayLength = ModuleStatistics.NumInterpolatedTexCoords + NumUserVertexInterpolators;
 
 	// get_material_customized_u_vs
 	{
 		FHLSLPrinter Printer = { 1 };
 		for (int32 i = 0; i < ShaderTexCoordsArrayLength; i++)
 		{
-			Printer << TEXTVIEW("OutTexCoords[") << i << TEXTVIEW("] = Parameters.TexCoords[") << i << TEXT(']') << EndOfStatement;
+			const bool bIsCustomizedUV = i < MIR::MaxNumCustomizedUVs && Module->IsMaterialPropertyUsed((EMaterialProperty)(MP_CustomizedUVs0 + i));
+			Printer << TEXTVIEW("OutTexCoords[") << i << TEXTVIEW("] = Parameters.");
+			Printer << (bIsCustomizedUV ? TEXTVIEW("CustomizedUVs[") : TEXTVIEW("TexCoords["));
+			Printer << i << TEXT(']') << EndOfStatement;
 		}
 		OutParameters.Add(TEXT("get_material_customized_u_vs"), MoveTemp(Printer.Buffer));
 	}
 		
 	// texcoords parameters
 	{
-		OutParameters.Add(TEXT("num_material_texcoords"), FString::FromInt(ModuleStatistics.NumTexCoords));
-		OutParameters.Add(TEXT("num_material_texcoords_vertex"), FString::FromInt(ModuleStatistics.NumTexCoords));
+		OutParameters.Add(TEXT("num_material_texcoords"), FString::FromInt(ModuleStatistics.NumInterpolatedTexCoords));
+		OutParameters.Add(TEXT("num_material_texcoords_vertex"), FString::FromInt(ModuleStatistics.NumInterpolatedTexCoords));
 		OutParameters.Add(TEXT("num_custom_vertex_interpolators"), FString::FromInt(NumUserVertexInterpolators));
 		OutParameters.Add(TEXT("num_tex_coord_interpolators"), FString::FromInt(ShaderTexCoordsArrayLength));
 	}
@@ -2742,7 +2759,7 @@ void FMaterialIRToHLSLTranslation::Run(TMap<FString, FString>& OutShaderParamete
 	OutShaderParameters.Empty();
 
 	GenerateStageEntryPoints(Module, OutShaderParameters);
-	GenerateCustomOutputDeclarations(Module, OutShaderParameters);
+	GenerateCustomOutputDeclarations(Material, Module, OutShaderParameters);
 	GenerateCustomFunctionsHLSL(Module, OutShaderParameters);
 	GenerateOtherTemplateStringParameters(Material, Module, OutShaderParameters);
 	GetShaderCompilerEnvironment(Material, Module, TargetPlatform, OutEnvironment);
diff --git a/Engine/Source/Runtime/Engine/Private/Materials/MaterialIRTypes.cpp b/Engine/Source/Runtime/Engine/Private/Materials/MaterialIRTypes.cpp
--- a/Engine/Source/Runtime/Engine/Private/Materials/MaterialIRTypes.cpp
+++ b/Engine/Source/Runtime/Engine/Private/Materials/MaterialIRTypes.cpp
@@ -115,7 +115,7 @@ FType FType::FromMaterialValueType(EMaterialValueType Type)
 		case MCT_Float3x3: return FType::MakeFloat(3, 3);
 		case MCT_Float4x4: return FType::MakeFloat(4, 4);
 		case MCT_LWCMatrix: return FType::MakeDouble(4, 4);
-		case MCT_MaterialAttributes: return FType::MakeAggregate(UMaterialAggregate::GetMaterialAttributes());
+		case MCT_MaterialAttributes: return FType::MakeAggregate(MaterialAttributesAggregate::Get());
 		case MCT_Substrate: return FType::MakeSubstrateData();
 		default: UE_MIR_UNREACHABLE();
 	}
diff --git a/Engine/Source/Runtime/Engine/Private/Materials/MaterialIRValueAnalyzer.cpp b/Engine/Source/Runtime/Engine/Private/Materials/MaterialIRValueAnalyzer.cpp
--- a/Engine/Source/Runtime/Engine/Private/Materials/MaterialIRValueAnalyzer.cpp
+++ b/Engine/Source/Runtime/Engine/Private/Materials/MaterialIRValueAnalyzer.cpp
@@ -45,15 +45,8 @@ static void ErrorUnlessFeatureLevelSupported(FMaterialIRModule* Module, ERHIFeat
 static void AnalyzeExternalInput(FMaterialIRValueAnalyzer& Analyzer, MIR::FExternalInput* ExternalInput)
 {
 	FMaterialIRModule::FStatistics& Statistics = Analyzer.Module->GetStatistics();
-	Statistics.ExternalInputUsedMask[MIR::Stage_Vertex][(int)ExternalInput->Id] = true;
-	Statistics.ExternalInputUsedMask[MIR::Stage_Pixel][(int)ExternalInput->Id] = true;
 
-	if (MIR::IsExternalInputTexCoordOrPartialDerivative(ExternalInput->Id))
-	{
-		int32 TexCoordIndex = MIR::ExternalInputToTexCoordIndex(ExternalInput->Id);
-		Statistics.NumTexCoords = FMath::Max(Statistics.NumTexCoords, TexCoordIndex + 1);
-	}
-	else if (Analyzer.MaterialInterface->IsPostProcessMaterial() && MIR::IsExternalInputWorldPosition(ExternalInput->Id))
+	if (Analyzer.MaterialInterface->IsPostProcessMaterial() && MIR::IsExternalInputWorldPosition(ExternalInput->Id))
 	{
 		Analyzer.CompilationOutput->SetIsSceneTextureUsed(PPI_SceneDepth);
 	}
@@ -77,6 +70,20 @@ static void AnalyzeExternalInput(FMaterialIRValueAnalyzer& Analyzer, MIR::FExter
 	}
 }
 
+static void AnalyzeInStageExternalInput(FMaterialIRValueAnalyzer& Analyzer, MIR::FExternalInput* ExternalInput, MIR::EStage Stage)
+{
+	FMaterialIRModule::FStatistics& Statistics = Analyzer.Module->GetStatistics();
+
+	if (MIR::IsExternalInputTexCoordOrPartialDerivative(ExternalInput->Id) && Stage != MIR::EStage::Stage_Vertex)
+	{
+		// @massimo.tristano this follows the logic in the old translator. A better way would instead be to **allocate** an index
+		// to a texture coordinate based on usage and map that index here. 
+		int32 TexCoordIndex = MIR::ExternalInputToTexCoordIndex(ExternalInput->Id);
+		Statistics.NumInterpolatedTexCoords = FMath::Max(Statistics.NumInterpolatedTexCoords, TexCoordIndex + 1);
+		Statistics.InterpolatedTexCoordsMask |= 1 << TexCoordIndex;
+	}
+}
+
 static void AnalyzeInStageHardwarePartialDerivative(FMaterialIRValueAnalyzer& Analyzer, MIR::FHardwarePartialDerivative* PartialDerivative, MIR::EStage Stage)
 {
 	if (Stage == MIR::Stage_Vertex)
@@ -785,7 +792,7 @@ static void AnalyzeSetMaterialOutput(FMaterialIRValueAnalyzer& Analyzer, MIR::FS
 {
 	// If this output is evaluated in the vertex stage, ensure its input is also valid there.
 	// Inputs marked "PixelStageOnly" cannot be used when generating vertex-stage code
-	if (SetMaterialOutput->Stage == MIR::MaterialOutputStage_Vertex && SetMaterialOutput->Arg->HasSubgraphProperties(MIR::EGraphProperties::PixelStageOnly))
+	if (SetMaterialOutput->Frequency == MIR::EMaterialOutputFrequency::PerVertex && SetMaterialOutput->Arg->HasSubgraphProperties(MIR::EGraphProperties::PixelStageOnly))
 	{
 		Analyzer.Module->AddError({}, FString::Printf(TEXT("Input to output '%s' is not available during the vertex stage."), *SetMaterialOutput->Name));
 		return;
@@ -794,9 +801,12 @@ static void AnalyzeSetMaterialOutput(FMaterialIRValueAnalyzer& Analyzer, MIR::FS
 	// Normal output cannot read the pixel normal
 	if (SetMaterialOutput->Property == EMaterialProperty::MP_Normal && SetMaterialOutput->HasSubgraphProperties(MIR::EGraphProperties::ReadsPixelNormal))
 	{
-		Analyzer.Module->AddError(nullptr, TEXT("Cannot set material attribute Normal to a value that depends on reading the pixel normal, as that would create a circular dependency."));
+		Analyzer.Module->AddError(nullptr, TEXT("Cannot set material attribute %s to a value that depends on reading the pixel normal, as that would create a circular dependency."));
 		return;
 	}
+
+	// Register the material output argument into the module. This "declares" the final output value of a material property.
+	Analyzer.Module->SetPropertyValue(SetMaterialOutput->Property, SetMaterialOutput->Arg);
 }
 
 static void AnalyzeMaterialParameterCollection(FMaterialIRValueAnalyzer& Analyzer, MIR::FMaterialParameterCollection* MaterialParameterCollection)
@@ -859,10 +869,10 @@ void FMaterialIRValueAnalyzer::Setup(UMaterialInterface* InMaterialInterface, FM
 
 void FMaterialIRValueAnalyzer::Analyze(MIR::FValue* Value)
 {
-	#define EXPAND_CASE(ValueType, ...) \
+	#define EXPAND_CASE(ValueType) \
 		case MIR::VK_##ValueType: \
 		{ \
-			Analyze##ValueType(*this, static_cast<MIR::F##ValueType*>(Value) , ## __VA_ARGS__); \
+			Analyze##ValueType(*this, static_cast<MIR::F##ValueType*>(Value)); \
 			break; \
 		}
 
@@ -900,6 +910,7 @@ void FMaterialIRValueAnalyzer::AnalyzeInStage(MIR::FValue* Value, MIR::EStage St
 	// IMPORTANT: Before adding a case here, read the FMaterialIRValueAnalyzer documentation.
 	switch (Value->Kind)
 	{
+		EXPAND_CASE(ExternalInput);
 		EXPAND_CASE(InlineHLSL);
 		EXPAND_CASE(VTPageTableRead);
 		EXPAND_CASE(ScreenTexture);
diff --git a/Engine/Source/Runtime/Engine/Private/Materials/SubstrateTranslatorCommon.cpp b/Engine/Source/Runtime/Engine/Private/Materials/SubstrateTranslatorCommon.cpp
--- a/Engine/Source/Runtime/Engine/Private/Materials/SubstrateTranslatorCommon.cpp
+++ b/Engine/Source/Runtime/Engine/Private/Materials/SubstrateTranslatorCommon.cpp
@@ -180,7 +180,7 @@ void FSubstrateCompilationContext::SubstrateEvaluateSharedLocalBases(
 
 			SubstrateSharedNormalParamName += ParameterSharedLocalBasesName;
 			SubstrateSharedNormalParamName.AppendInt(FinalSharedLocalBasisIndex);
-			Em->SetCustomOutputs(SubstrateSharedNormalParamName, { &SubstrateSharedLocalBasesInfo.NormalArg, 1 }, MIR::MaterialOutputStage_Pixel);
+			Em->SetCustomOutputs(SubstrateSharedNormalParamName, { &SubstrateSharedLocalBasesInfo.NormalArg, 1 }, MIR::EMaterialOutputFrequency::PerPixel);
 			SubstrateSharedNormalParamCode = FString::Printf(TEXT("Parameters.%s0"), *SubstrateSharedNormalParamName);
 
 			if (SubstrateSharedLocalBasesInfo.TangentArg.IsValid())
@@ -189,7 +189,7 @@ void FSubstrateCompilationContext::SubstrateEvaluateSharedLocalBases(
 
 				SubstrateSharedTangentParamName += ParameterSharedLocalBasesName;
 				SubstrateSharedTangentParamName.AppendInt(FinalSharedLocalBasisIndex);
-				Em->SetCustomOutputs(SubstrateSharedTangentParamName, { &SubstrateSharedLocalBasesInfo.TangentArg, 1 }, MIR::MaterialOutputStage_Pixel);
+				Em->SetCustomOutputs(SubstrateSharedTangentParamName, { &SubstrateSharedLocalBasesInfo.TangentArg, 1 }, MIR::EMaterialOutputFrequency::PerPixel);
 				SubstrateSharedTangentParamCode = FString::Printf(TEXT("Parameters.%s0"), *SubstrateSharedTangentParamName);
 			}
 
diff --git a/Engine/Source/Runtime/Engine/Public/Materials/MaterialAggregate.h b/Engine/Source/Runtime/Engine/Public/Materials/MaterialAggregate.h
--- a/Engine/Source/Runtime/Engine/Public/Materials/MaterialAggregate.h
+++ b/Engine/Source/Runtime/Engine/Public/Materials/MaterialAggregate.h
@@ -92,31 +92,41 @@ public:
 
 public:
 #if WITH_EDITOR
+	// Finds the index of an attribute in the aggregate by name.
+	int32 FindAttributeIndexByName(FName Name) const;
+
+	virtual void PostLoad() override;
+	virtual void PostEditChangeProperty(FPropertyChangedEvent& PropertyChangedEvent) override;
+#endif // #if WITH_EDITOR
+};
+
 
-	static const TConstArrayView<EMaterialProperty> GetMaterialAttributesProperties();
+#if WITH_EDITOR
 
-	// Returns a reference to the global material attributes aggregate used by the engine.
-	static const UMaterialAggregate* GetMaterialAttributes();
+// Aggregate APIs specific for the MaterialAttributes data type.
+namespace MaterialAttributesAggregate
+{
 
-	// Returns a single cached aggregate for the specified material property.
-	static const FMaterialAggregateAttribute* GetMaterialAttribute(EMaterialProperty Property);
+// Returns an array of all the MaterialAttributes properties.
+const TConstArrayView<EMaterialProperty> GetMaterialProperties();
 
-	// Maps a material property to the corresponding attribute index in the aggregate.
-	static int32 MaterialPropertyToAttributeIndex(EMaterialProperty Property);
+// Returns a reference to the global material attributes aggregate used by the engine.
+const UMaterialAggregate* Get();
 
-	// Maps an attribute index back to the corresponding material property.
-	static EMaterialProperty AttributeIndexToMaterialProperty(int32 MaterialAttributeIndex);
+// Returns a single cached aggregate for the specified material property.
+const FMaterialAggregateAttribute* GetAttribute(EMaterialProperty Property);
 
-	static inline bool IsMaterialPropertyEnabled(EMaterialProperty Property)
-	{
-		return MaterialPropertyToAttributeIndex(Property) != INDEX_NONE;
-	}
+// Maps a material property to the corresponding attribute index in the aggregate.
+int32 MaterialPropertyToAttributeIndex(EMaterialProperty Property);
 
-	// Finds the index of an attribute in the aggregate by name.
-	int32 FindAttributeIndexByName(FName Name) const;
+// Maps an attribute index back to the corresponding material property.
+EMaterialProperty AttributeIndexToMaterialProperty(int32 MaterialAttributeIndex);
 
-	virtual void PostLoad() override;
-	virtual void PostEditChangeProperty(FPropertyChangedEvent& PropertyChangedEvent) override;
+// Returns whether specified material Property is enabled in a MaterialAttributes aggregate with given number of custom UVs enabled.
+inline bool IsPropertyEnabled(EMaterialProperty Property)
+{
+	return MaterialPropertyToAttributeIndex(Property) != INDEX_NONE;
+}
 
+} // MaterialAttributesAggregate
 #endif // #if WITH_EDITOR
-};
diff --git a/Engine/Source/Runtime/Engine/Public/Materials/MaterialIR.h b/Engine/Source/Runtime/Engine/Public/Materials/MaterialIR.h
--- a/Engine/Source/Runtime/Engine/Public/Materials/MaterialIR.h
+++ b/Engine/Source/Runtime/Engine/Public/Materials/MaterialIR.h
@@ -59,9 +59,9 @@ enum class EValueFlags : uint8
 	// Value has been analyzed for the compute stage.
 	AnalyzedForStageCompute = 1 << 2,
 
-	// Value has been analyzed in some stage.
-	AnalyzedInAnyStage = 1 << 3,
-
+	// Mask that includes all the analyzed flag in all stages.
+	AnalyzedInAllStagesMask = AnalyzedForStageVertex | AnalyzedForStagePixel | AnalyzedForStageCompute,
+	
 	// Whether to substitute abstract tokens like "<PREV>" in code expressions of InlineHLSL instructions.
 	SubstituteTags = 1 << 4,
 
@@ -743,10 +743,10 @@ struct TComposite : FComposite
 };
 
 // Whether a material output should be computed per-vertex or per-pixel.
-enum EMaterialOutputStage
+enum class EMaterialOutputFrequency
 {
-	MaterialOutputStage_Vertex,
-	MaterialOutputStage_Pixel,
+	PerVertex,
+	PerPixel,
 };
 
 // Instruction that sets material attribute (i.e., "BaseColor") to its value.
@@ -756,7 +756,7 @@ struct FSetMaterialOutput : TInstruction<VK_SetMaterialOutput, 1>
 	EMaterialProperty Property : 16;
 	
 	// Whether this output is evaluated per-vertex or per-pixel.
-	EMaterialOutputStage Stage : 16;
+	EMaterialOutputFrequency Frequency : 16;
 
 	// The name of this material output (e.g. "BaseColor")
 	FSimpleStringView Name;
diff --git a/Engine/Source/Runtime/Engine/Public/Materials/MaterialIRCommon.h b/Engine/Source/Runtime/Engine/Public/Materials/MaterialIRCommon.h
--- a/Engine/Source/Runtime/Engine/Public/Materials/MaterialIRCommon.h
+++ b/Engine/Source/Runtime/Engine/Public/Materials/MaterialIRCommon.h
@@ -123,6 +123,14 @@ bool Find(const TMap<TKey, TValue>& Map, const TKey& Key, TValue& OutValue)
 	return false;
 }
 
+// Returns whether given property is a customized UV.
+static inline bool IsMaterialPropertyCustomizedUV(EMaterialProperty Property)
+{
+	return Property >= MP_CustomizedUVs0 && Property <= MP_CustomizedUVs7;
+}
+
+static constexpr int MaxNumCustomizedUVs = 8;
+
 } // namespace MIR
 
 #define UE_MIR_UNREACHABLE() { check(!"Unreachable"); UE_ASSUME(false); }
diff --git a/Engine/Source/Runtime/Engine/Public/Materials/MaterialIREmitter.h b/Engine/Source/Runtime/Engine/Public/Materials/MaterialIREmitter.h
--- a/Engine/Source/Runtime/Engine/Public/Materials/MaterialIREmitter.h
+++ b/Engine/Source/Runtime/Engine/Public/Materials/MaterialIREmitter.h
@@ -545,7 +545,7 @@ public:
 	// the output index (e.g. "TangentOutput0" for a group named "TangentOutput").
 	// The type of each custom output is inferred by the type of the argument value it is set to, therefore make sure
 	// you cast each output argument to the desired type before calling this function.
-	void SetCustomOutputs(FStringView Name, TConstArrayView<FValueRef> Args, EMaterialOutputStage Stage);
+	void SetCustomOutputs(FStringView Name, TConstArrayView<FValueRef> Args, EMaterialOutputFrequency Frequency);
 	
 	// Emits a 2D column vector from X and Y scalars.
 	FValueRef Vector2(FValueRef X, FValueRef Y);
diff --git a/Engine/Source/Runtime/Engine/Public/Materials/MaterialIRModule.h b/Engine/Source/Runtime/Engine/Public/Materials/MaterialIRModule.h
--- a/Engine/Source/Runtime/Engine/Public/Materials/MaterialIRModule.h
+++ b/Engine/Source/Runtime/Engine/Public/Materials/MaterialIRModule.h
@@ -33,17 +33,17 @@ public:
 	// Stores information about the resources used by the translated material.
 	struct FStatistics
 	{
-		// Tracks external inputs used per frequency.
-		TBitArray<> ExternalInputUsedMask[MIR::NumStages]; 
+		// Mask to identify which TexCoord indices are actually used in the pixel/compute stages.
+		int32 InterpolatedTexCoordsMask{};
 
-		// Number of texture coordinates used by the material.
-		int32 NumTexCoords;
+		// Number of texture coordinates used by the material in the pixel/compute stages.
+		int32 NumInterpolatedTexCoords{};
 		
 		// Total number of user vertex-interpolator float registers used by the material.
-		int32 NumUsedVertexInterpolatorFloatRegisters;
+		int32 NumUsedVertexInterpolatorFloatRegisters{};
 
 		// Dynamic particle parameter indices used by this material.
-		uint32 DynamicParticleParameterMask;
+		uint32 DynamicParticleParameterMask{};
 	};
 
 	// An entry point is a sequence of instructions that evaluates a set of "output" instructions.
