Refactored AnimSequenceTransformProvider to use timestamps instead of incrementing from the render thread. This allows synchronized access from the GameThread as well.

#Synced-CL 50172231
#Synced-CL 50172231

#Synced-CL 50172231

diff --git a/Engine/Source/Runtime/Core/Public/UObject/FortniteMainBranchObjectVersions.inl b/Engine/Source/Runtime/Core/Public/UObject/FortniteMainBranchObjectVersions.inl
--- a/Engine/Source/Runtime/Core/Public/UObject/FortniteMainBranchObjectVersions.inl
+++ b/Engine/Source/Runtime/Core/Public/UObject/FortniteMainBranchObjectVersions.inl
@@ -760,5 +760,8 @@ FORTNITE_MAIN_VERSION(DeprecateUAFExternalPackages, 248)
 // RigVM object pins hold an object ptr
 FORTNITE_MAIN_VERSION(RigVMPinDefaultValueObjectPtr, 249)
 
+// AnimSequenceTracks now have timestamps
+FORTNITE_MAIN_VERSION(AnimSequenceTrackTimestamps, 250)
+
 // -----<new versions can be added above this line>-------------------------------------------------
 #undef FORTNITE_MAIN_VERSION
diff --git a/Engine/Source/Runtime/Engine/Classes/Animation/AnimSequenceTransformProviderData.h b/Engine/Source/Runtime/Engine/Classes/Animation/AnimSequenceTransformProviderData.h
--- a/Engine/Source/Runtime/Engine/Classes/Animation/AnimSequenceTransformProviderData.h
+++ b/Engine/Source/Runtime/Engine/Classes/Animation/AnimSequenceTransformProviderData.h
@@ -46,13 +46,13 @@ struct FAnimSequenceTrackAutoPlayData
 	UPROPERTY(EditAnywhere, BlueprintReadWrite, Category = Animation)
 	float PlayRate = 1.0f;
 
-	// The mode to control looping behavior.
-	UPROPERTY(EditAnywhere, BlueprintReadWrite, Category = Animation)
-	EAnimSequenceTrackLoopMode LoopMode = EAnimSequenceTrackLoopMode::Loop;
-
 	// The duration of the blend transition in seconds. Use 0 for an instance switch.
 	UPROPERTY(EditAnywhere, BlueprintReadWrite, Category = Animation, meta = (ClampMin = "0.0"))
 	float BlendTime = 0.0f;
+
+	// The mode to control looping behavior.
+	UPROPERTY(EditAnywhere, BlueprintReadWrite, Category = Animation)
+	EAnimSequenceTrackLoopMode LoopMode = EAnimSequenceTrackLoopMode::Loop;
 };
 
 USTRUCT(BlueprintType)
@@ -75,6 +75,10 @@ struct FAnimSequenceTrackManualData
 	// A time position to update previous frame transforms.
 	UPROPERTY(EditAnywhere, BlueprintReadWrite, Category = Animation, meta = (EditCondition = "bHasPreviousPosition", EditConditionHides))
 	float PreviousPosition = 0.0f;
+
+	// The mode to control looping behavior for when the position is not within the sequence bounds.
+	UPROPERTY(EditAnywhere, BlueprintReadWrite, Category = Animation)
+	EAnimSequenceTrackLoopMode LoopMode = EAnimSequenceTrackLoopMode::Loop;
 };
 
 USTRUCT(BlueprintType)
@@ -85,13 +89,13 @@ struct FAnimSequenceTransformProviderSequence
 	UPROPERTY(EditAnywhere, BlueprintReadWrite, Category = Animation)
 	TObjectPtr<class UAnimSequence> Sequence;
 
-	// The rate multiplier to play the animation.
+	// The default rate multiplier to play the animation.
 	UPROPERTY(EditAnywhere, BlueprintReadWrite, Category = Animation)
 	float PlayRate = 1.0f;
 
-	// The mode to control looping behavior.
+	// The default position to start playing the animation.
 	UPROPERTY(EditAnywhere, BlueprintReadWrite, Category = Animation)
-	EAnimSequenceTrackLoopMode LoopMode = EAnimSequenceTrackLoopMode::Loop;
+	float Position = 0.0f;
 };
 
 class FAnimSequenceTransformProviderProxy;
@@ -111,18 +115,19 @@ public:
 	/////////////////////////////////////////////////////////////////////////////////////////////////
 	// AutoPlay
 
-	static FAnimSequenceTrackPackedData Pack(FAnimSequenceTrackAutoPlayData AutoPlayData)
+	static FAnimSequenceTrackPackedData Pack(FAnimSequenceTrackAutoPlayData AutoPlayData, double WorldTime)
 	{
 		FAnimSequenceTrackPackedData Data;
 		Data.SequenceIndex = AutoPlayData.SequenceIndex;
 		Data.bAutoPlay = 1;
 		Data.LoopMode = static_cast<uint32>(AutoPlayData.LoopMode);
 		Data.PatchFlags = 0;
-		Data.bHasPreviousPosition = true;
-		Data.Position = AutoPlayData.Position;
-		Data.PreviousPosition = AutoPlayData.Position;
-		Data.PlayRate = AutoPlayData.PlayRate;
-		Data.BlendTime = AutoPlayData.BlendTime;
+		Data.bHasPreviousPosition = false;
+		Data.Auto.ReferenceTimestamp = GetTimestamp(WorldTime, AutoPlayData.Position, AutoPlayData.PlayRate);
+		Data.Auto.Position = AutoPlayData.Position;
+		Data.Auto.PreviousPosition = AutoPlayData.Position;
+		Data.Auto.PlayRate = AutoPlayData.PlayRate;
+		Data.Auto.BlendTime = AutoPlayData.BlendTime;
 		return Data;
 	}
 
@@ -131,34 +136,19 @@ public:
 		return bAutoPlay;
 	}
 
-	FAnimSequenceTrackAutoPlayData UnpackAutoPlay() const
-	{
-		check(IsAutoPlay());
-
-		FAnimSequenceTrackAutoPlayData Data;
-		Data.SequenceIndex = SequenceIndex;
-		Data.Position = Position;
-		Data.PlayRate = PlayRate;
-		Data.BlendTime = BlendTime;
-		Data.LoopMode = static_cast<EAnimSequenceTrackLoopMode>(LoopMode);
-		return Data;
-	}
-
 	/////////////////////////////////////////////////////////////////////////////////////////////////
 	// Manual
 
-	static FAnimSequenceTrackPackedData Pack(FAnimSequenceTrackManualData ManualData)
+	static FAnimSequenceTrackPackedData Pack(FAnimSequenceTrackManualData ManualData, float SequenceLength)
 	{
 		FAnimSequenceTrackPackedData Data;
 		Data.SequenceIndex = ManualData.SequenceIndex;
 		Data.bAutoPlay = 0;
-		Data.LoopMode = 0;
+		Data.LoopMode = static_cast<uint32>(ManualData.LoopMode);
 		Data.PatchFlags = 0;
 		Data.bHasPreviousPosition = ManualData.bHasPreviousPosition;
-		Data.Position = ManualData.Position;
-		Data.PreviousPosition = ManualData.bHasPreviousPosition ? ManualData.PreviousPosition : ManualData.Position;
-		Data.PlayRate = 0.0f;
-		Data.BlendTime = 0.0f;
+		Data.Manual.Position = GetPosition(ManualData.LoopMode, ManualData.Position, SequenceLength);
+		Data.Manual.PreviousPosition = ManualData.bHasPreviousPosition ? GetPosition(ManualData.LoopMode, ManualData.PreviousPosition, SequenceLength) : Data.Manual.Position;
 		return Data;
 	}
 
@@ -167,18 +157,6 @@ public:
 		return !bAutoPlay;
 	}
 
-	FAnimSequenceTrackManualData UnpackManual() const
-	{
-		check(IsManual());
-
-		FAnimSequenceTrackManualData Data;
-		Data.SequenceIndex = SequenceIndex;
-		Data.Position = Position;
-		Data.PreviousPosition = PreviousPosition;
-		Data.bHasPreviousPosition = bHasPreviousPosition;
-		return Data;
-	}
-
 	/////////////////////////////////////////////////////////////////////////////////////////////////
 
 	int32 GetSequenceIndex() const
@@ -186,14 +164,19 @@ public:
 		return SequenceIndex;
 	}
 
-	float GetPosition() const
+	float GetPosition(double WorldTime, double SequenceLength) const
 	{
-		return Position;
+		if (bAutoPlay)
+		{
+			const double Position = (WorldTime - Auto.ReferenceTimestamp) * Auto.PlayRate;
+			return GetPosition(GetLoopMode(), Position, SequenceLength);
+		}
+		return Manual.Position;
 	}
 
 	float GetPreviousPosition() const
 	{
-		return PreviousPosition;
+		return bAutoPlay ? Auto.PreviousPosition : Manual.PreviousPosition;
 	}
 
 	bool HasPreviousPosition() const
@@ -214,7 +197,7 @@ public:
 
 	float GetBlendTime() const
 	{
-		return BlendTime;
+		return IsAutoPlay() ? Auto.BlendTime : 0.0f;
 	}
 
 	void ResetPatchFlags()
@@ -227,34 +210,95 @@ public:
 		return static_cast<EAnimSequenceTrackPatchFlags>(PatchFlags);
 	}
 
-	void SetPlayRate(float NewPlayRate)
+	void SetPlayRate(float NewPlayRate, double WorldTime)
 	{
-		PlayRate = NewPlayRate;
+		check(IsAutoPlay());
+		const double Position = (WorldTime - Auto.ReferenceTimestamp) * Auto.PlayRate;
+		Auto.ReferenceTimestamp = GetTimestamp(WorldTime, Position, NewPlayRate);
+		Auto.PlayRate = NewPlayRate;
 		PatchFlags |= (uint32)EAnimSequenceTrackPatchFlags::PlayRate;
 	}
 
 	float GetPlayRate() const
 	{
-		return PlayRate;
+		return IsAutoPlay() ? Auto.PlayRate : 0.0f;
 	}
 
-	ENGINE_API void Tick(float SequenceLength, float DeltaTime);
+	ENGINE_API void Serialize(FArchive& Ar);
 
-	ENGINE_API void Wrap(float SequenceLength);
+	void PostLoad(double WorldTime)
+	{
+		if (IsAutoPlay())
+		{
+			Auto.ReferenceTimestamp = GetTimestamp(WorldTime, Auto.Position, Auto.PlayRate);
+		}
+	}
 
-	ENGINE_API void Serialize(FArchive& Ar);
+	void Tick(float Position)
+	{
+		if (IsAutoPlay())
+		{
+			Auto.PreviousPosition = Position;
+			bHasPreviousPosition = true;
+		}
+	}
 
 private:
+	static float GetPosition(EAnimSequenceTrackLoopMode LoopMode, float Position, float SequenceLength)
+	{
+		if (LoopMode == EAnimSequenceTrackLoopMode::Loop)
+		{
+			if (SequenceLength >= 0.0)
+			{
+				Position = SequenceLength == 0.0 ? 0.0 : FMath::Fmod(Position, SequenceLength);
+
+				// Returns a negative range if input is negative.
+				if (Position < 0.0)
+				{
+					Position += SequenceLength;
+				}
+			}
+		}
+		else
+		{
+			Position = FMath::Clamp(Position, 0.0, SequenceLength);
+		}
+		return Position;
+	}
+
+	static double GetTimestamp(double WorldTime, float Position, float PlayRate)
+	{
+		if (FMath::Abs(PlayRate) > UE_SMALL_NUMBER)
+		{
+			return WorldTime - static_cast<double>(Position / PlayRate);
+		}
+		return WorldTime;
+	}
+
 	uint32 SequenceIndex : 20;
 	uint32 LoopMode : 1;
 	uint32 bAutoPlay : 1;
 	uint32 bHasPreviousPosition : 1;
 	uint32 PatchFlags : 8;
 	uint32 Unused : 1 = 0;
-	float Position;
-	float PreviousPosition;
-	float PlayRate;
-	float BlendTime;
+
+	union
+	{
+		struct
+		{
+			double ReferenceTimestamp;
+			float Position;
+			float PreviousPosition;
+			float PlayRate;
+			float BlendTime;
+		} Auto;
+
+		struct
+		{
+			float Position;
+			float PreviousPosition;
+		} Manual;
+	};
 
 	friend class FAnimSequenceTransformProvider;
 };
@@ -297,10 +341,10 @@ public:
 
 	const FAnimSequenceTrackPackedData& GetSource() const { return Source; }
 	const FAnimSequenceTrackPackedData& GetTarget() const { return Target; }
-	
-	ENGINE_API void Patch(const FAnimSequenceTrackPackedData& Request, bool bWasActive);
 
-	ENGINE_API void Tick(float SourceSequenceLength, float TargetSequenceLength, float DeltaTime);
+	ENGINE_API void Patch(const FAnimSequenceTrackPackedData& Request, double WorldTime, bool bWasActive);
+
+	ENGINE_API void Tick(float DeltaTime, float SourcePosition, float TargetPosition);
 
 private:
 	FAnimSequenceTrackPackedData Source;
@@ -316,7 +360,7 @@ public:
 	ENGINE_API ~FAnimSequenceTransformProviderRenderData();
 
 private:
-	void Patch(FAnimSequenceTrackPool::FPatch&& Patch);
+	void Patch(FAnimSequenceTrackPool::FPatch&& Patch, double WorldTime);
 	void InsertProxy(FInstancedSkinningSceneExtensionProxy* Proxy);
 	void RemoveProxy(FInstancedSkinningSceneExtensionProxy* Proxy);
 
@@ -422,6 +466,8 @@ private:
 	UPROPERTY(EditAnywhere, BlueprintReadOnly, Category = TransformProvider, meta = (AllowPrivateAccess = "true", ShowOnlyInnerProperties))
 	TArray<struct FAnimSequenceTransformProviderSequence> Sequences;
 
+	double GetWorldTime() const;
+
 	FRenderCommandFence DestroyFence;
 
 	friend class UAnimSequenceTransformProviderDataInstance;
@@ -439,6 +485,7 @@ public:
 	virtual FTransformProviderRenderProxy* CreateRenderProxy(FInstancedSkinningSceneExtensionProxy* ExtensionProxy) const override;
 	virtual void BeginDestroy() override;
 	virtual void Serialize(FArchive& Ar) override;
+	virtual void SubmitChanges() override;
 	virtual UTransformProviderData* GetRootTransformProvider() override;
 
 public:
@@ -466,13 +513,9 @@ public:
 	UFUNCTION(BlueprintCallable, Category = "TransformProviders|AnimSequence")
 	ENGINE_API bool DeallocateTrack(int32 TrackIndex);
 
-	// Submits all pending changes to rendering.
-	UFUNCTION(BlueprintCallable, Category = "TransformProviders|AnimSequence")
-	ENGINE_API virtual void SubmitChanges() override;
-
 	// Get the sequence index currently playing on a track
 	UFUNCTION(BlueprintCallable, Category = "TransformProviders|AnimSequence")
-	ENGINE_API int32 GetTrackSequenceIndex(int32 TrackIndex) const;
+	ENGINE_API int32 GetSequenceIndex(int32 TrackIndex) const;
 
 	// Assigns a new play rate to an existing track.
 	UFUNCTION(BlueprintCallable, Category = "TransformProviders|AnimSequence")
@@ -490,6 +533,10 @@ public:
 	UFUNCTION(BlueprintCallable, Category = "TransformProviders|AnimSequence")
 	ENGINE_API EAnimSequenceTrackLoopMode GetLoopMode(int32 TrackIndex) const;
 
+	// Get current position of a track.
+	UFUNCTION(BlueprintCallable, Category = "TransformProviders|AnimSequence")
+	ENGINE_API float GetPosition(int32 TrackIndex) const;
+
 private:
 	void Initialize(UAnimSequenceTransformProviderData* ProviderData);
 	void MarkOwnerRenderStateDirty();
@@ -504,6 +551,9 @@ private:
 	// Prevents redundant MarkRenderDynamicDataDirty calls within a frame.
 	bool bPendingRenderStateDirty = false;
 
+	// Defers initialization of timestamps after serialization until the world is available.
+	bool bNeedsPostLoadInit = false;
+
 	FAnimSequenceTrackPool Tracks;
 	TUniquePtr<FAnimSequenceTransformProviderRenderData> RenderData{ new FAnimSequenceTransformProviderRenderData };
 	FRenderCommandFence InstanceDestroyFence;
diff --git a/Engine/Source/Runtime/Engine/Private/Animation/AnimSequenceTransformProviderData.cpp b/Engine/Source/Runtime/Engine/Private/Animation/AnimSequenceTransformProviderData.cpp
--- a/Engine/Source/Runtime/Engine/Private/Animation/AnimSequenceTransformProviderData.cpp
+++ b/Engine/Source/Runtime/Engine/Private/Animation/AnimSequenceTransformProviderData.cpp
@@ -8,8 +8,8 @@
 #include "Animation/AnimSequence.h"
 #include "Animation/Skeleton.h"
 #include "SceneInterface.h"
-#include "UObject/AnimObjectVersion.h"
 #include "UObject/UObjectThreadContext.h"
+#include "UObject/FortniteMainBranchObjectVersion.h"
 #include "ObjectCacheContext.h"
 
 #if WITH_EDITOR
@@ -29,48 +29,6 @@ static TAutoConsoleVariable<bool> CVarAnimSequenceTransformProviderBlending(
 
 ///////////////////////////////////////////////////////////////////////////////////////////////////
 
-static float GetAnimSequenceTrackLoopedPosition(float Position, float TrackLength)
-{
-	if (TrackLength >= 0.0f)
-	{
-		Position = TrackLength == 0.0f ? 0.0f : FMath::Fmod(Position, TrackLength);
-
-		// Returns a negative range if input is negative.
-		if (Position < 0.0f)
-		{
-			Position += TrackLength;
-		}
-	}
-	return Position;
-}
-
-void FAnimSequenceTrackPackedData::Tick(float SequenceLength, float DeltaTime)
-{
-	if (bAutoPlay)
-	{
-		PreviousPosition = Position;
-		Position += PlayRate * DeltaTime;
-		bHasPreviousPosition = false;
-		Wrap(SequenceLength);
-	}
-}
-
-void FAnimSequenceTrackPackedData::Wrap(float SequenceLength)
-{
-	const EAnimSequenceTrackLoopMode LoopModeEnum = GetLoopMode();
-
-	if (LoopModeEnum == EAnimSequenceTrackLoopMode::Loop)
-	{
-		Position = GetAnimSequenceTrackLoopedPosition(Position, SequenceLength);
-		PreviousPosition = GetAnimSequenceTrackLoopedPosition(PreviousPosition, SequenceLength);
-	}
-	else
-	{
-		Position = FMath::Clamp(Position, 0.0f, SequenceLength);
-		PreviousPosition = FMath::Clamp(PreviousPosition, 0.0f, SequenceLength);
-	}
-}
-
 void FAnimSequenceTrackPackedData::Serialize(FArchive& Ar)
 {
 	constexpr auto MakeBitMask = [](uint32 Bits) -> uint32 { return (Bits >= 32) ? 0xffffffffu : ((1u << Bits) - 1u); };
@@ -93,22 +51,25 @@ void FAnimSequenceTrackPackedData::Serialize(FArchive& Ar)
 	uint32 PackedBits;
 	if (Ar.IsSaving())
 	{
+		const uint8 bManual = !bAutoPlay;
+
 		PackedBits =
 			(SequenceIndex & SequenceIndexMask)
 			| ((LoopMode & LoopModeMask) << LoopModeShift)
 			| ((bAutoPlay & AutoPlayMask) << AutoPlayShift)
-			| ((bHasPreviousPosition & HasPreviousPositionMask) << HasPreviousPositionShift);
+			| ((bManual & bHasPreviousPosition & HasPreviousPositionMask) << HasPreviousPositionShift);
 		Ar << PackedBits;
-		Ar << Position;
 
 		if (bAutoPlay)
 		{
-			Ar << PlayRate;
-			Ar << BlendTime;
+			Ar << Auto.Position;
+			Ar << Auto.PlayRate;
+			Ar << Auto.BlendTime;
 		}
 		else
 		{
-			Ar << PreviousPosition;
+			Ar << Manual.Position;
+			Ar << Manual.PreviousPosition;
 		}
 	}
 	else
@@ -119,32 +80,54 @@ void FAnimSequenceTrackPackedData::Serialize(FArchive& Ar)
 		bAutoPlay = (PackedBits >> AutoPlayShift) & AutoPlayMask;
 		bHasPreviousPosition = (PackedBits >> HasPreviousPositionShift) & HasPreviousPositionMask;
 		PatchFlags = 0;
-		Ar << Position;
 
-		if (bAutoPlay)
+		if (Ar.CustomVer(FFortniteMainBranchObjectVersion::GUID) >= FFortniteMainBranchObjectVersion::AnimSequenceTrackTimestamps)
 		{
-			Ar << PlayRate;
-			Ar << BlendTime;
-			PreviousPosition = 0.0f;
+			if (bAutoPlay)
+			{
+				Ar << Auto.Position;
+				Ar << Auto.PlayRate;
+				Ar << Auto.BlendTime;
+				Auto.PreviousPosition = Auto.Position;
+			}
+			else
+			{
+				Ar << Manual.Position;
+				Ar << Manual.PreviousPosition;
+			}
 		}
 		else
 		{
-			Ar << PreviousPosition;
+			float Position;
+			Ar << Position;
+
+			if (bAutoPlay)
+			{
+				Auto.Position = Position;
+				Auto.PreviousPosition = Auto.Position;
+				Ar << Auto.PlayRate;
+				Ar << Auto.BlendTime;
+			}
+			else
+			{
+				Manual.Position = Position;
+				Ar << Manual.PreviousPosition;
+			}
 		}
 	}
 }
 
 ///////////////////////////////////////////////////////////////////////////////////////////////////
 
-void FAnimSequenceTrackRenderData::Patch(const FAnimSequenceTrackPackedData& Data, bool bWasActive)
+void FAnimSequenceTrackRenderData::Patch(const FAnimSequenceTrackPackedData& Data, double WorldTime, bool bWasActive)
 {
 	// Only patching specific flags.
 	if (Data.GetPatchFlags() != EAnimSequenceTrackPatchFlags::None)
 	{
 		if (EnumHasAnyFlags(Data.GetPatchFlags(), EAnimSequenceTrackPatchFlags::PlayRate))
 		{
-			Source.SetPlayRate(Data.GetPlayRate());
-			Target.SetPlayRate(Data.GetPlayRate());
+			Source.SetPlayRate(Data.GetPlayRate(), WorldTime);
+			Target.SetPlayRate(Data.GetPlayRate(), WorldTime);
 		}
 
 		if (EnumHasAnyFlags(Data.GetPatchFlags(), EAnimSequenceTrackPatchFlags::LoopMode))
@@ -159,7 +142,7 @@ void FAnimSequenceTrackRenderData::Patch(const FAnimSequenceTrackPackedData& Dat
 	// Patching the entire struct. Starts a new animation.
 	else
 	{
-		if (Data.GetBlendTime() <= 0.0f || !bWasActive)
+		if (Data.GetBlendTime() <= 0.0f || !bWasActive || !CVarAnimSequenceTransformProviderBlending.GetValueOnRenderThread())
 		{
 			Source = Data;
 		}
@@ -173,14 +156,8 @@ void FAnimSequenceTrackRenderData::Patch(const FAnimSequenceTrackPackedData& Dat
 	}
 }
 
-void FAnimSequenceTrackRenderData::Tick(float SourceSequenceLength, float TargetSequenceLength, float DeltaTime)
+void FAnimSequenceTrackRenderData::Tick(float DeltaTime, float SourcePosition, float TargetPosition)
 {
-	Source.Tick(SourceSequenceLength, DeltaTime);
-	Source.Wrap(SourceSequenceLength);
-
-	Target.Tick(TargetSequenceLength, DeltaTime);
-	Target.Wrap(TargetSequenceLength);
-
 	const float BlendTime = Target.GetBlendTime();
 
 	if (BlendTime > 0.0f && BlendWeight < 1.0f)
@@ -193,6 +170,9 @@ void FAnimSequenceTrackRenderData::Tick(float SourceSequenceLength, float Target
 			Source = Target;
 		}
 	}
+
+	Source.Tick(SourcePosition);
+	Target.Tick(TargetPosition);
 }
 
 ///////////////////////////////////////////////////////////////////////////////////////////////////
@@ -212,7 +192,7 @@ void FAnimSequenceTransformProviderRenderData::RemoveProxy(FInstancedSkinningSce
 	Proxies.Remove(Proxy);
 }
 
-void FAnimSequenceTransformProviderRenderData::Patch(FAnimSequenceTrackPool::FPatch&& Patch)
+void FAnimSequenceTransformProviderRenderData::Patch(FAnimSequenceTrackPool::FPatch&& Patch, double WorldTime)
 {
 	if (Patch.GetNumTracks() != Tracks.GetNumTracks())
 	{
@@ -222,9 +202,9 @@ void FAnimSequenceTransformProviderRenderData::Patch(FAnimSequenceTrackPool::FPa
 		}
 	}
 
-	Tracks.Patch<FAnimSequenceTrackPackedData>(MoveTemp(Patch), [](FAnimSequenceTrackRenderData& Dst, const FAnimSequenceTrackPackedData& Src, bool bWasActive)
+	Tracks.Patch<FAnimSequenceTrackPackedData>(MoveTemp(Patch), [WorldTime](FAnimSequenceTrackRenderData& Dst, const FAnimSequenceTrackPackedData& Src, bool bWasActive)
 	{
-		Dst.Patch(Src, bWasActive);
+		Dst.Patch(Src, WorldTime, bWasActive);
 	});
 }
 
@@ -244,10 +224,10 @@ FAnimSequenceTransformProviderProxy::FAnimSequenceTransformProviderProxy(TConstA
 
 		FAnimSequenceTrackAutoPlayData Data;
 		Data.SequenceIndex = SequenceIndex;
-		Data.LoopMode = Sequence.LoopMode;
+		Data.Position = Sequence.Position;
 		Data.PlayRate = Sequence.PlayRate;
 
-		TrackAllocator->AllocateTrack(FAnimSequenceTrackPackedData::Pack(Data));
+		TrackAllocator->AllocateTrack(FAnimSequenceTrackPackedData::Pack(Data, 0.0));
 	}
 }
 
@@ -452,6 +432,45 @@ const TArray<FAnimSequenceTransformProviderSequence>& UAnimSequenceTransformProv
 	return Sequences;
 }
 
+namespace AnimSequenceTrackTime
+{
+	float GetTimeScale()
+	{
+		static const auto* CVar = IConsoleManager::Get().FindTConsoleVariableDataFloat(TEXT("r.AnimSequenceTransformProvider.TimeScale"));
+		return CVar ? CVar->GetValueOnAnyThread() : 1.0f;
+	}
+
+	thread_local struct FTime
+	{
+		UWorld* World = nullptr;
+		uint64 Frame = INDEX_NONE;
+		double WorldTime = 0.0;
+
+	} GTimeCache;
+
+	double GetTime(UWorld* World)
+	{
+		if (World)
+		{
+			if (GTimeCache.Frame != GFrameCounter || GTimeCache.World != World)
+			{
+				GTimeCache = { World, GFrameCounter, World->GetTimeSeconds() * GetTimeScale() };
+			}
+			return GTimeCache.WorldTime;
+		}
+		return 0.0;
+	}
+}
+
+double UAnimSequenceTransformProviderData::GetWorldTime() const
+{
+#if WITH_ENGINE
+	return AnimSequenceTrackTime::GetTime(GetWorld());
+#else
+	return 0.0;
+#endif
+}
+
 #if WITH_EDITOR
 
 void UAnimSequenceTransformProviderData::PostEditChangeProperty(FPropertyChangedEvent& PropertyChangedEvent)
@@ -506,10 +525,7 @@ UAnimSequenceTransformProviderData* UAnimSequenceTransformProviderData::CreateFr
 		return nullptr;
 	}
 
-	UAnimSequenceTransformProviderData* NewProviderData = NewObject<UAnimSequenceTransformProviderData>(
-		NewPackage,
-		*NewAssetName,
-		RF_Public | RF_Standalone);
+	UAnimSequenceTransformProviderData* NewProviderData = NewObject<UAnimSequenceTransformProviderData>(NewPackage, *NewAssetName, RF_Public | RF_Standalone);
 
 	if (!NewProviderData)
 	{
@@ -521,13 +537,18 @@ UAnimSequenceTransformProviderData* UAnimSequenceTransformProviderData::CreateFr
 		FAnimSequenceTransformProviderSequence Sequence;
 		if (BankItem.BankAsset)
 		{
+			if (!NewProviderData->SkinnedAsset)
+			{
+				NewProviderData->SkinnedAsset = BankItem.BankAsset->Asset;
+			}
+
 			const TArray<FAnimBankSequence>& BankSequences = BankItem.BankAsset->Sequences;
 			if (BankSequences.IsValidIndex(BankItem.SequenceIndex))
 			{
 				const FAnimBankSequence& BankSequence = BankSequences[BankItem.SequenceIndex];
 				Sequence.Sequence = BankSequence.Sequence;
-				Sequence.LoopMode = BankSequence.bLooping ? EAnimSequenceTrackLoopMode::Loop : EAnimSequenceTrackLoopMode::Clamp;
 				Sequence.PlayRate = BankSequence.PlayRate;
+				Sequence.Position = BankSequence.Position;
 			}
 		}
 		NewProviderData->Sequences.Add(Sequence);
@@ -561,10 +582,7 @@ UAnimSequenceTransformProviderData* UAnimSequenceTransformProviderData::CreateFr
 		return nullptr;
 	}
 
-	UAnimSequenceTransformProviderData* NewProviderData = NewObject<UAnimSequenceTransformProviderData>(
-		NewPackage,
-		*NewAssetName,
-		RF_Public | RF_Standalone);
+	UAnimSequenceTransformProviderData* NewProviderData = NewObject<UAnimSequenceTransformProviderData>(NewPackage, *NewAssetName, RF_Public | RF_Standalone);
 
 	if (!NewProviderData)
 	{
@@ -575,11 +593,12 @@ UAnimSequenceTransformProviderData* UAnimSequenceTransformProviderData::CreateFr
 	{
 		FAnimSequenceTransformProviderSequence Sequence;
 		Sequence.Sequence = BankSequence.Sequence;
-		Sequence.LoopMode = BankSequence.bLooping ? EAnimSequenceTrackLoopMode::Loop : EAnimSequenceTrackLoopMode::Clamp;
 		Sequence.PlayRate = BankSequence.PlayRate;
+		Sequence.Position = BankSequence.Position;
 		NewProviderData->Sequences.Add(Sequence);
 	}
 
+	NewProviderData->SkinnedAsset = InAnimBank->Asset;
 	NewProviderData->MarkPackageDirty();
 	FAssetRegistryModule::AssetCreated(NewProviderData);
 	return NewProviderData;
@@ -672,6 +691,19 @@ void UAnimSequenceTransformProviderDataInstance::SubmitChanges()
 	// SubmitChanges is called during the ISKM instance data update. Guard against multiple concurrent updates.
 	UE::TScopeLock Lock(SubmitChangesMutex);
 
+	if (bNeedsPostLoadInit)
+	{
+		const double WorldTime = GetWorldTime();
+
+		Tracks.EnumerateActiveTracks([&](int32 TrackIndex)
+		{
+			Tracks.GetData(TrackIndex).PostLoad(WorldTime);
+			Tracks.SetDirty(TrackIndex);
+		});
+
+		bNeedsPostLoadInit = false;
+	}
+
 	bPendingRenderStateDirty = false;
 
 	if (Tracks.GetNumDirtyTracks() > 0)
@@ -681,9 +713,9 @@ void UAnimSequenceTransformProviderDataInstance::SubmitChanges()
 			Data.ResetPatchFlags();
 		});
 
-		ENQUEUE_RENDER_COMMAND(AnimSequenceDataInstanceSubmit)([&RenderData = *RenderData, Patch = MoveTemp(Patch)](FRHICommandList&) mutable
+		ENQUEUE_RENDER_COMMAND(AnimSequenceDataInstanceSubmit)([&RenderData = *RenderData, Patch = MoveTemp(Patch), WorldTime = GetWorldTime()](FRHICommandList&) mutable
 		{
-			RenderData.Patch(MoveTemp(Patch));
+			RenderData.Patch(MoveTemp(Patch), WorldTime);
 		});
 	}
 }
@@ -698,29 +730,26 @@ void UAnimSequenceTransformProviderDataInstance::Initialize(UAnimSequenceTransfo
 
 void UAnimSequenceTransformProviderDataInstance::Serialize(FArchive& Ar)
 {
-	Ar.UsingCustomVersion(FAnimObjectVersion::GUID);
+	Ar.UsingCustomVersion(FFortniteMainBranchObjectVersion::GUID);
 	Super::Serialize(Ar);
 	Tracks.Serialize(Ar);
 
 	if (Ar.IsLoading())
 	{
-		check(RenderData);
-		RenderData->Tracks = Tracks;
 		if (Parent)
 		{
 			Initialize(Parent);
 		}
+
+		bNeedsPostLoadInit = true;
 	}
 }
 
 int32 UAnimSequenceTransformProviderDataInstance::AllocateAutoPlayTrack(const FAnimSequenceTrackAutoPlayData& Data)
 {
-	FAnimSequenceTrackPackedData PackedData = FAnimSequenceTrackPackedData::Pack(Data);
-
 	if (Sequences.IsValidIndex(Data.SequenceIndex))
 	{
-		PackedData.Wrap(GetSequencePlayLength(Data.SequenceIndex));
-		int32 TrackIndex = Tracks.AllocateTrack(PackedData);
+		int32 TrackIndex = Tracks.AllocateTrack(FAnimSequenceTrackPackedData::Pack(Data, GetWorldTime()));
 		if (TrackIndex != INDEX_NONE)
 		{
 			MarkOwnerRenderStateDirty();
@@ -733,12 +762,9 @@ int32 UAnimSequenceTransformProviderDataInstance::AllocateAutoPlayTrack(const FA
 
 int32 UAnimSequenceTransformProviderDataInstance::AllocateManualTrack(const FAnimSequenceTrackManualData& Data)
 {
-	FAnimSequenceTrackPackedData PackedData = FAnimSequenceTrackPackedData::Pack(Data);
-
 	if (Sequences.IsValidIndex(Data.SequenceIndex))
 	{
-		PackedData.Wrap(GetSequencePlayLength(Data.SequenceIndex));
-		int32 TrackIndex = Tracks.AllocateTrack(PackedData);
+		int32 TrackIndex = Tracks.AllocateTrack(FAnimSequenceTrackPackedData::Pack(Data, GetSequencePlayLength(Data.SequenceIndex)));
 		if (TrackIndex != INDEX_NONE)
 		{
 			MarkOwnerRenderStateDirty();
@@ -751,12 +777,9 @@ int32 UAnimSequenceTransformProviderDataInstance::AllocateManualTrack(const FAni
 
 bool UAnimSequenceTransformProviderDataInstance::UpdateAutoPlayTrack(int32 TrackIndex, const FAnimSequenceTrackAutoPlayData& Data)
 {
-	FAnimSequenceTrackPackedData PackedData = FAnimSequenceTrackPackedData::Pack(Data);
-
 	if (Sequences.IsValidIndex(Data.SequenceIndex))
 	{
-		PackedData.Wrap(GetSequencePlayLength(Data.SequenceIndex));
-		if (Tracks.UpdateTrack(TrackIndex, PackedData))
+		if (Tracks.UpdateTrack(TrackIndex, FAnimSequenceTrackPackedData::Pack(Data, GetWorldTime())))
 		{
 			MarkOwnerRenderStateDirty();
 			return true;
@@ -768,12 +791,9 @@ bool UAnimSequenceTransformProviderDataInstance::UpdateAutoPlayTrack(int32 Track
 
 bool UAnimSequenceTransformProviderDataInstance::UpdateManualTrack(int32 TrackIndex, const FAnimSequenceTrackManualData& Data)
 {
-	FAnimSequenceTrackPackedData PackedData = FAnimSequenceTrackPackedData::Pack(Data);
-
 	if (Sequences.IsValidIndex(Data.SequenceIndex))
 	{
-		PackedData.Wrap(GetSequencePlayLength(Data.SequenceIndex));
-		if (Tracks.UpdateTrack(TrackIndex, PackedData))
+		if (Tracks.UpdateTrack(TrackIndex, FAnimSequenceTrackPackedData::Pack(Data, GetSequencePlayLength(Data.SequenceIndex))))
 		{
 			MarkOwnerRenderStateDirty();
 			return true;
@@ -810,10 +830,14 @@ bool UAnimSequenceTransformProviderDataInstance::SetPlayRate(int32 TrackIndex, f
 {
 	if (Tracks.IsActiveIndex(TrackIndex))
 	{
-		Tracks.GetData(TrackIndex).SetPlayRate(PlayRate);
-		Tracks.SetDirty(TrackIndex);
-		MarkOwnerRenderStateDirty();
-		return true;
+		FAnimSequenceTrackPackedData& Data = Tracks.GetData(TrackIndex);
+		if (Data.IsAutoPlay())
+		{
+			Data.SetPlayRate(PlayRate, GetWorldTime());
+			Tracks.SetDirty(TrackIndex);
+			MarkOwnerRenderStateDirty();
+			return true;
+		}
 	}
 
 	return false;
@@ -825,6 +849,7 @@ float UAnimSequenceTransformProviderDataInstance::GetPlayRate(int32 TrackIndex)
 	{
 		return Tracks.GetData(TrackIndex).GetPlayRate();
 	}
+
 	return 1.0f;
 }
 
@@ -834,14 +859,28 @@ EAnimSequenceTrackLoopMode UAnimSequenceTransformProviderDataInstance::GetLoopMo
 	{
 		return Tracks.GetData(TrackIndex).GetLoopMode();
 	}
+
 	return EAnimSequenceTrackLoopMode::Loop;
 }
 
-int32 UAnimSequenceTransformProviderDataInstance::GetTrackSequenceIndex(int32 TrackIndex) const
+int32 UAnimSequenceTransformProviderDataInstance::GetSequenceIndex(int32 TrackIndex) const
 {
 	if (Tracks.IsActiveIndex(TrackIndex))
 	{
 		return Tracks.GetData(TrackIndex).GetSequenceIndex();
 	}
+
 	return INDEX_NONE;
 }
+
+float UAnimSequenceTransformProviderDataInstance::GetPosition(int32 TrackIndex) const
+{
+	if (Tracks.IsActiveIndex(TrackIndex))
+	{
+		const FAnimSequenceTrackPackedData& PackedData = Tracks.GetData(TrackIndex);
+		const float SequenceLength = GetSequencePlayLength(PackedData.GetSequenceIndex());
+		return PackedData.GetPosition(GetWorldTime(), SequenceLength);
+	}
+
+	return 0.0f;
+}
\ No newline at end of file
diff --git a/Engine/Source/Runtime/Renderer/Private/Skinning/AnimSequenceTransformProvider.cpp b/Engine/Source/Runtime/Renderer/Private/Skinning/AnimSequenceTransformProvider.cpp
--- a/Engine/Source/Runtime/Renderer/Private/Skinning/AnimSequenceTransformProvider.cpp
+++ b/Engine/Source/Runtime/Renderer/Private/Skinning/AnimSequenceTransformProvider.cpp
@@ -58,7 +58,7 @@ static TAutoConsoleVariable<float> CVarAnimSequenceTransformProviderTimeScale(
 
 static TAutoConsoleVariable<bool> CVarAnimSequenceTransformProviderGPU(
 	TEXT("r.AnimSequenceTransformProvider.GPU"),
-	false,
+	true,
 	TEXT(""),
 	FConsoleVariableDelegate::CreateLambda([](IConsoleVariable* InVariable)
 	{
@@ -877,16 +877,16 @@ void FAnimSequenceTransformProvider::ProvideGPUTransforms(FSkinningTransformProv
 	RetargetingDataAllocator.Consolidate();
 	LeafBonePathAllocator.Consolidate();
 
-	uint32 MaxTotalSubTracks = 0;
+	uint32 MaxTotalActiveTracks = 0;
 
 	for (const FSkinningTransformProvider::FProviderIndirection Indirection : Context.Indirections)
 	{
 		FInstancedSkinningSceneExtensionProxy* ExtensionProxy = static_cast<FInstancedSkinningSceneExtensionProxy*>(Context.Proxies[Indirection.Index]);
 		FProxyUserData& ProxyUserData = static_cast<FProxyUserData&>(ExtensionProxy->GetUserDataChecked());
-		MaxTotalSubTracks += 2 * ProxyUserData.Tracks.GetNumActiveTracks();
+		MaxTotalActiveTracks += ProxyUserData.Tracks.GetNumActiveTracks();
 	}
 
-	if (!MaxTotalSubTracks)
+	if (!MaxTotalActiveTracks)
 	{
 		return;
 	}
@@ -907,9 +907,14 @@ void FAnimSequenceTransformProvider::ProvideGPUTransforms(FSkinningTransformProv
 		}
 	};
 
+	const float GlobalTimeScale = CVarAnimSequenceTransformProviderTimeScale.GetValueOnRenderThread();
+	const double WorldTime = Context.ViewFamily.Time.GetWorldTimeSeconds() * GlobalTimeScale;
+	const float DeltaTime = Context.ViewFamily.Time.GetDeltaWorldTimeSeconds() * GlobalTimeScale;
+
 	TArray<FTrackUpdateRequest, SceneRenderingAllocator> TrackUpdateRequests;
-	TrackUpdateRequests.Reserve(MaxTotalSubTracks);
+	TrackUpdateRequests.Reserve(MaxTotalActiveTracks);
 
+	uint32 NumTotalSubTracks = 0;
 	uint32 NumTotalIntermediateTransforms = 0;
 	uint32 NumTotalEvaluateBlockHeaders = 0;
 	uint32 NumTotalEvaluateBlendBlockHeaders = 0;
@@ -928,7 +933,7 @@ void FAnimSequenceTransformProvider::ProvideGPUTransforms(FSkinningTransformProv
 		const uint32 NumEvaluateBlockHeaders    = FMath::DivideAndRoundUp(LOD.BoneTables.NumCompactBones, BonesPerGroup);
 		const uint32 NumScatterBlockHeaders     = FMath::DivideAndRoundUp(LOD.BoneTables.NumTransforms, BonesPerGroup);
 		const uint32 NumConcatenateBlockHeaders = FMath::DivideAndRoundUp(LOD.LeafBonePathData.NumLeafBones, BonesPerGroup);
-		const uint32 CurrentTransformSlot = Indirection.CurrentTransformOffset < Indirection.PreviousTransformOffset ? 0 : 1;
+		const uint32 CurrentTransformSlot       = Indirection.CurrentTransformOffset < Indirection.PreviousTransformOffset ? 0 : 1;
 	
 		int32 NumSubTracks = 0;
 		int32 NumEvaluateSubTracks = 0;
@@ -936,18 +941,14 @@ void FAnimSequenceTransformProvider::ProvideGPUTransforms(FSkinningTransformProv
 
 		Tracks.EnumerateActiveTracks([&] (int32 TrackIndex)
 		{
-			const FAnimSequenceTrackRenderData& Data   = Tracks.GetData(TrackIndex);
-			const FAnimSequenceTrackPackedData& Source = Data.GetSource();
-			const FAnimSequenceTrackPackedData& Target = Data.GetTarget();
-			const float SourcePosition                 = Source.GetPosition();
-			const float SourcePreviousPosition         = Source.GetPreviousPosition();
-			const float TargetPosition                 = Target.GetPosition();
-			const float TargetPreviousPosition         = Target.GetPreviousPosition();
-			const float BlendWeight                    = Data.GetBlendWeight();
+			const FAnimSequenceTrackRenderData& Data = Tracks.GetData(TrackIndex);
+			const FTrackData SourceTrack             = ProxyUserData.GetTrackData(Data.GetSource(), WorldTime, LODIndex);
+			const FTrackData TargetTrack             = ProxyUserData.GetTrackData(Data.GetTarget(), WorldTime, LODIndex);
+			const float BlendWeight                  = Data.GetBlendWeight();
 
 			// Dirty previous transforms means the whole allocation is dirty. Rebuild everything.
 			const bool bForceUpdateAll = EnumHasAnyFlags(Indirection.DirtyBoneTransforms, EDirtyBoneTransforms::Previous);
-			FTrackHistory::FUpdateResult UpdateResult = ProxyUserData.TrackHistories[TrackIndex].Update(Data, CurrentTransformSlot, bForceUpdateAll);
+			FTrackHistory::FUpdateResult UpdateResult = ProxyUserData.TrackHistories[TrackIndex].Update(SourceTrack, TargetTrack, CurrentTransformSlot, bForceUpdateAll);
 
 			FTrackUpdateRequest UpdateRequest;
 			UpdateRequest.SetUpdateMode(EPreviousBoneTransformUpdateMode::None);
@@ -955,7 +956,7 @@ void FAnimSequenceTransformProvider::ProvideGPUTransforms(FSkinningTransformProv
 
 			if (UpdateResult.bNeedsPrevious)
 			{
-				UpdateRequest.UpdateMode = uint8(UpdateResult.bNeedsCurrent && UpdateResult.bNeedsPrevious && TargetPosition == TargetPreviousPosition && SourcePosition == SourcePreviousPosition
+				UpdateRequest.UpdateMode = uint8(UpdateResult.bNeedsCurrent && UpdateResult.bNeedsPrevious && TargetTrack.Position == TargetTrack.PreviousPosition && SourceTrack.Position == SourceTrack.PreviousPosition
 					? EPreviousBoneTransformUpdateMode::DuplicateCurrentToPrevious
 					: EPreviousBoneTransformUpdateMode::UpdatePrevious);
 			}
@@ -969,6 +970,7 @@ void FAnimSequenceTransformProvider::ProvideGPUTransforms(FSkinningTransformProv
 			TrackUpdateRequests.Emplace(UpdateRequest);
 		});
 
+		NumTotalSubTracks                 += NumSubTracks;
 		NumTotalIntermediateTransforms    += NumSubTracks * LOD.BoneTables.NumCompactBones;
 		NumTotalEvaluateBlockHeaders      += NumEvaluateSubTracks * NumEvaluateBlockHeaders;
 		NumTotalEvaluateBlendBlockHeaders += NumEvaluateBlendSubTracks * NumEvaluateBlockHeaders;
@@ -976,6 +978,11 @@ void FAnimSequenceTransformProvider::ProvideGPUTransforms(FSkinningTransformProv
 		NumTotalConcatenateBlockHeaders   += NumSubTracks * NumConcatenateBlockHeaders;
 	}
 
+	if (!NumTotalSubTracks)
+	{
+		return;
+	}
+
 	uint32 NumTotalSequenceTransforms = 0;
 	uint32 NumTotalRetargetingDatas = 0;
 	uint32 NumTotalLeafBonePathElements = 0;
@@ -1056,7 +1063,8 @@ void FAnimSequenceTransformProvider::ProvideGPUTransforms(FSkinningTransformProv
 	UE::Tasks::FTask ProcessTask = GraphBuilder.AddSetupTask(
 		[
 			  this
-			, &ViewFamily = Context.ViewFamily
+			, WorldTime
+			, DeltaTime
 			, Indirections = Context.Indirections
 			, Proxies = Context.Proxies
 			, &NumTotalFramesRequested
@@ -1215,8 +1223,6 @@ void FAnimSequenceTransformProvider::ProvideGPUTransforms(FSkinningTransformProv
 			IntermediateTransformOffset += SourceMapping.NumCompactBones * sizeof(FBoneTransformWithScale);
 		};
 
-		const float GlobalTimeScale = CVarAnimSequenceTransformProviderTimeScale.GetValueOnRenderThread();
-		const float DeltaTime = ViewFamily.Time.GetDeltaWorldTimeSeconds() * GlobalTimeScale;
 		uint32 UpdateResultIndex = 0;
 
 		for (const FSkinningTransformProvider::FProviderIndirection Indirection : Indirections)
@@ -1234,23 +1240,11 @@ void FAnimSequenceTransformProvider::ProvideGPUTransforms(FSkinningTransformProv
 			Tracks.EnumerateActiveTracks([&] (int32 TrackIndex)
 			{
 				FAnimSequenceTrackRenderData& Data = Tracks.GetData(TrackIndex);
-
-				const FAnimSequenceTrackPackedData& Source = Data.GetSource();
-				const FAnimSequenceTrackPackedData& Target = Data.GetTarget();
-
-				const FMapping& SourceMapping = ProxyUserData.GetMapping(Source, LODIndex);
-				const FMapping& TargetMapping = ProxyUserData.GetMapping(Target, LODIndex);
-				FSequence& SourceSequence = ProxyUserData.GetSequence(Source);
-				FSequence& TargetSequence = ProxyUserData.GetSequence(Target);
-
-				const FTrackUpdateRequest UpdateRequest = TrackUpdateRequests[UpdateResultIndex++];
-	
-				const float SourcePosition         = Source.GetPosition();
-				const float SourcePreviousPosition = Source.GetPreviousPosition();
-				const float TargetPosition         = Target.GetPosition();
-				const float TargetPreviousPosition = Target.GetPreviousPosition();
+				const FTrackData SourceTrack       = ProxyUserData.GetTrackData(Data.GetSource(), WorldTime, LODIndex);
+				const FTrackData TargetTrack       = ProxyUserData.GetTrackData(Data.GetTarget(), WorldTime, LODIndex);
 				const float BlendWeight            = Data.GetBlendWeight();
 
+				const FTrackUpdateRequest UpdateRequest   = TrackUpdateRequests[UpdateResultIndex++];
 				const uint32 TrackTransformOffset         = TrackIndex * LOD.BoneTables.NumTransforms * 2 * sizeof(FCompressedBoneTransform);
 				const uint32 TrackCurrentTransformOffset  = TrackTransformOffset + CurrentTransformOffset;
 				const uint32 TrackPreviousTransformOffset = TrackTransformOffset + PreviousTransformOffset;
@@ -1263,9 +1257,9 @@ void FAnimSequenceTransformProvider::ProvideGPUTransforms(FSkinningTransformProv
 					EvaluateAnimSequence(
 						Indirection,
 						LOD,
-						SourceMapping, TargetMapping,
-						SourceSequence, TargetSequence,
-						SourcePosition, TargetPosition,
+						SourceTrack.Mapping, TargetTrack.Mapping,
+						SourceTrack.Sequence, TargetTrack.Sequence,
+						SourceTrack.Position, TargetTrack.Position,
 						BlendWeight,
 						LOD.BoneTables.NumTransforms,
 						OutputTransformOffset1, OutputTransformOffset2);
@@ -1278,9 +1272,9 @@ void FAnimSequenceTransformProvider::ProvideGPUTransforms(FSkinningTransformProv
 					EvaluateAnimSequence(
 						Indirection,
 						LOD,
-						SourceMapping, TargetMapping,
-						SourceSequence, TargetSequence,
-						SourcePreviousPosition, TargetPreviousPosition,
+						SourceTrack.Mapping, TargetTrack.Mapping,
+						SourceTrack.Sequence, TargetTrack.Sequence,
+						SourceTrack.PreviousPosition, TargetTrack.PreviousPosition,
 						BlendWeight,
 						LOD.BoneTables.NumTransforms,
 						TrackPreviousTransformOffset, TrackPreviousTransformOffset);
@@ -1288,7 +1282,7 @@ void FAnimSequenceTransformProvider::ProvideGPUTransforms(FSkinningTransformProv
 					NumPreviousTransformUpdates++;
 				}
 
-				Data.Tick(SourceSequence.PlayLength, TargetSequence.PlayLength, DeltaTime);
+				Data.Tick(DeltaTime, SourceTrack.Position, TargetTrack.Position);
 				NumActiveTracks++;
 			});
 		}
@@ -1928,6 +1922,7 @@ void FAnimSequenceTransformProvider::ProvideCPUTransforms(FSkinningTransformProv
 		};
 
 		const float GlobalTimeScale = CVarAnimSequenceTransformProviderTimeScale.GetValueOnRenderThread();
+		const double WorldTime = ViewFamily.Time.GetWorldTimeSeconds() * GlobalTimeScale;
 		const float DeltaTime = ViewFamily.Time.GetDeltaWorldTimeSeconds() * GlobalTimeScale;
 
 		uint32 NumActiveTracks = 0;
@@ -1959,30 +1954,21 @@ void FAnimSequenceTransformProvider::ProvideCPUTransforms(FSkinningTransformProv
 			{
 				FAnimSequenceTrackRenderData& Data = Tracks.GetData(TrackIndex);
 
-				const FAnimSequenceTrackPackedData& Source = Data.GetSource();
-				const FAnimSequenceTrackPackedData& Target = Data.GetTarget();
+				const FTrackData SourceTrack = ProxyUserData.GetTrackData(Data.GetSource(), WorldTime, LODIndex);
+				const FTrackData TargetTrack = ProxyUserData.GetTrackData(Data.GetTarget(), WorldTime, LODIndex);
 
-				const FMapping& SourceMapping = ProxyUserData.GetMapping(Source, LODIndex);
-				const FMapping& TargetMapping = ProxyUserData.GetMapping(Target, LODIndex);
-				FSequence& SourceSequence = ProxyUserData.GetSequence(Source);
-				FSequence& TargetSequence = ProxyUserData.GetSequence(Target);
-				
 				// Dirty previous transforms means the whole allocation is dirty. Rebuild everything.
 				const bool bDirtyTransformAllocation = EnumHasAnyFlags(Indirection.DirtyBoneTransforms, EDirtyBoneTransforms::Previous);
-				const FTrackHistory::FUpdateResult UpdateResult = ProxyUserData.TrackHistories[TrackIndex].Update(Data, CurrentTransformSlot, bDirtyTransformAllocation);
+				const FTrackHistory::FUpdateResult UpdateResult = ProxyUserData.TrackHistories[TrackIndex].Update(SourceTrack, TargetTrack, CurrentTransformSlot, bDirtyTransformAllocation);
 	
-				const float SourcePosition         = Source.GetPosition();
-				const float SourcePreviousPosition = Source.GetPreviousPosition();
-				const float TargetPosition         = Target.GetPosition();
-				const float TargetPreviousPosition = Target.GetPreviousPosition();
-				const float BlendWeight            = Data.GetBlendWeight();
+				const float BlendWeight = Data.GetBlendWeight();
 
 				EPreviousBoneTransformUpdateMode UpdateMode = EPreviousBoneTransformUpdateMode::None;
 				TArrayView<FCompressedBoneTransform> DstTransformsPrevious;
 
 				if (UpdateResult.bNeedsPrevious)
 				{
-					UpdateMode = UpdateResult.bNeedsCurrent && UpdateResult.bNeedsPrevious && TargetPosition == TargetPreviousPosition && SourcePosition == SourcePreviousPosition
+					UpdateMode = UpdateResult.bNeedsCurrent && UpdateResult.bNeedsPrevious && TargetTrack.Position == TargetTrack.PreviousPosition && SourceTrack.Position == SourceTrack.PreviousPosition
 						? EPreviousBoneTransformUpdateMode::DuplicateCurrentToPrevious
 						: EPreviousBoneTransformUpdateMode::UpdatePrevious;
 
@@ -1997,9 +1983,9 @@ void FAnimSequenceTransformProvider::ProvideCPUTransforms(FSkinningTransformProv
 						ExtensionProxy,
 						DstTransforms,
 						UpdateMode == EPreviousBoneTransformUpdateMode::DuplicateCurrentToPrevious ? DstTransformsPrevious : TArrayView<FCompressedBoneTransform>{},
-						SourceSequence, TargetSequence,
-						SourceMapping, TargetMapping,
-						SourcePosition, TargetPosition,
+						SourceTrack.Sequence, TargetTrack.Sequence,
+						SourceTrack.Mapping, TargetTrack.Mapping,
+						SourceTrack.Position, TargetTrack.Position,
 						BlendWeight);
 
 					NumCurrentTransformUpdates++;
@@ -2012,15 +1998,15 @@ void FAnimSequenceTransformProvider::ProvideCPUTransforms(FSkinningTransformProv
 						ExtensionProxy,
 						DstTransformsPrevious,
 						{},
-						SourceSequence, TargetSequence,
-						SourceMapping, TargetMapping,
-						SourcePreviousPosition, TargetPreviousPosition,
+						SourceTrack.Sequence, TargetTrack.Sequence,
+						SourceTrack.Mapping, TargetTrack.Mapping,
+						SourceTrack.PreviousPosition, TargetTrack.PreviousPosition,
 						BlendWeight);
 
 					NumPreviousTransformUpdates++;
 				}
 
-				Data.Tick(SourceSequence.PlayLength, TargetSequence.PlayLength, DeltaTime);
+				Data.Tick(DeltaTime, SourceTrack.Position, TargetTrack.Position);
 				NumActiveTracks++;
 			});
 		}
diff --git a/Engine/Source/Runtime/Renderer/Private/Skinning/AnimSequenceTransformProvider.h b/Engine/Source/Runtime/Renderer/Private/Skinning/AnimSequenceTransformProvider.h
--- a/Engine/Source/Runtime/Renderer/Private/Skinning/AnimSequenceTransformProvider.h
+++ b/Engine/Source/Runtime/Renderer/Private/Skinning/AnimSequenceTransformProvider.h
@@ -203,6 +203,16 @@ private:
 		uint32 RefCount = 0;
 	};
 
+	struct FTrackData
+	{
+		const FMapping& Mapping;
+		FSequence& Sequence;
+		int32 SequenceIndex;
+		float Position;
+		float PreviousPosition;
+		bool bHasPreviousPosition;
+	};
+
 	class FTrackHistory
 	{
 	public:
@@ -214,32 +224,32 @@ private:
 			uint8 bNeedsPrevious : 1 = 0;
 		};
 
-		FUpdateResult Update(const FAnimSequenceTrackRenderData& RenderData, int32 SlotIndex, bool bForceUpdate)
+		FUpdateResult Update(const FTrackData& SourceTrack, const FTrackData& TargetTrack, int32 SlotIndex, bool bForceUpdate)
 		{
 			FUpdateResult Result;
-			Result.bNeedsCurrent  = bForceUpdate;
+			Result.bNeedsCurrent = bForceUpdate;
 			Result.bNeedsPrevious = bForceUpdate;
-			UpdateSlot(SourceSequenceIndex, SourceSlots, RenderData.GetSource(), SlotIndex, Result);
-			UpdateSlot(TargetSequenceIndex, TargetSlots, RenderData.GetTarget(), SlotIndex, Result);
+			UpdateSlot(SourceSequenceIndex, SourceSlots, SourceTrack, SlotIndex, Result);
+			UpdateSlot(TargetSequenceIndex, TargetSlots, TargetTrack, SlotIndex, Result);
 			return Result;
 		}
 
 	private:
-		void UpdateSlot(int32& SequenceIndex, FVector2f& Slots, const FAnimSequenceTrackPackedData& Data, int32 InCurrentSlot, FUpdateResult& Result)
+		void UpdateSlot(int32& SequenceIndex, FVector2f& Slots, const FTrackData& Track, int32 InCurrentSlot, FUpdateResult& Result)
 		{
-			if (SequenceIndex != Data.GetSequenceIndex())
+			if (SequenceIndex != Track.SequenceIndex)
 			{
-				SequenceIndex = Data.GetSequenceIndex();
+				SequenceIndex = Track.SequenceIndex;
 				Slots = FVector2f{ Invalid, Invalid };
 			}
 
-			Result.bNeedsCurrent  |= Slots[ InCurrentSlot] != Data.GetPosition();
-			Result.bNeedsPrevious |= Slots[!InCurrentSlot] != Data.GetPreviousPosition() && Data.HasPreviousPosition();
-			Slots[InCurrentSlot] = Data.GetPosition();
+			Result.bNeedsCurrent  |= Slots[ InCurrentSlot] != Track.Position;
+			Result.bNeedsPrevious |= Slots[!InCurrentSlot] != Track.PreviousPosition && Track.bHasPreviousPosition;
+			Slots[InCurrentSlot] = Track.Position;
 
-			if (Data.HasPreviousPosition())
+			if (Track.bHasPreviousPosition)
 			{
-				Slots[!InCurrentSlot] = Data.GetPreviousPosition();
+				Slots[!InCurrentSlot] = Track.PreviousPosition;
 			}
 		}
 
@@ -261,10 +271,28 @@ private:
 
 		const FMapping& GetMapping(const FAnimSequenceTrackPackedData& Data, uint32 LODIndex) const
 		{
-			const int32 SequenceIndex = SequenceToMappingTable.IsValidIndex(Data.GetSequenceIndex()) ? Data.GetSequenceIndex() : 0;
+			check(LODIndex >= MinLODLevel && LODIndex < (uint32)LODs.Num());
+			const int32 SequenceIndex = Sequences.IsValidIndex(Data.GetSequenceIndex()) ? Data.GetSequenceIndex() : 0;
 			return *SequenceToMappingTable[SequenceIndex * LODs.Num() + LODIndex];
 		}
 
+		FTrackData GetTrackData(const FAnimSequenceTrackPackedData& Data, double WorldTime, uint32 LODIndex) const
+		{
+			FSequence& Sequence = GetSequence(Data);
+			const float Position = Data.GetPosition(WorldTime, Sequence.PlayLength);
+			const float PreviousPosition = Data.HasPreviousPosition() ? Data.GetPreviousPosition() : Position;
+
+			return FTrackData
+			{
+				  .Mapping  = GetMapping(Data, LODIndex)
+				, .Sequence = Sequence
+				, .SequenceIndex = Data.GetSequenceIndex()
+				, .Position = Position
+				, .PreviousPosition = PreviousPosition
+				, .bHasPreviousPosition = Data.HasPreviousPosition()
+			};
+		}
+
 		struct FLOD
 		{
 			FBoneTables BoneTables;
