[Backout] - CL50309928
#rnx
Original CL Desc
-----------------------------------------------------------------
- Add collision_profile to has_collision.
- Fix per-shape collision data for bodyinstances which are not in a skeletal mesh

#Synced-CL 50223803

diff --git a/Engine/Source/Runtime/Engine/Private/PhysicsEngine/BodyInstance.cpp b/Engine/Source/Runtime/Engine/Private/PhysicsEngine/BodyInstance.cpp
--- a/Engine/Source/Runtime/Engine/Private/PhysicsEngine/BodyInstance.cpp
+++ b/Engine/Source/Runtime/Engine/Private/PhysicsEngine/BodyInstance.cpp
@@ -4728,15 +4728,15 @@ void FBodyInstance::BuildBodyFilterData(FBodyCollisionFilterData& OutFilterData,
 			UseResponse = FCollisionResponseContainer::CreateMinContainer(UseResponse, SkelMeshComp->BodyInstance.CollisionResponses.GetResponseContainer());
 		}
 
-		bUseNotifyRBCollision = bUseNotifyRBCollision && SkelMeshComp->BodyInstance.bNotifyRigidBodyCollision;
-	}
+		if (ShapeIndex != INDEX_NONE)
+		{
+			// If we're looking at a specific shape, take the intersection of its collision responses
+			// with the base primitive component's collision responses.
+			FCollisionResponseContainer ShapeResponse = GetShapeResponseToChannels(ShapeIndex, UseResponse);
+			UseResponse = FCollisionResponseContainer::CreateMinContainer(UseResponse, ShapeResponse);
+		}
 
-	if (ShapeIndex != INDEX_NONE)
-	{
-		// If we're looking at a specific shape, take the intersection of its collision responses
-		// with the base primitive component's collision responses.
-		FCollisionResponseContainer ShapeResponse = GetShapeResponseToChannels(ShapeIndex, UseResponse);
-		UseResponse = FCollisionResponseContainer::CreateMinContainer(UseResponse, ShapeResponse);
+		bUseNotifyRBCollision = bUseNotifyRBCollision && SkelMeshComp->BodyInstance.bNotifyRigidBodyCollision;
 	}
 
 #if WITH_EDITOR
