[Shaders] Dump shader debug info when internal compiler errors were triggered, despite a successful compile job.

This is to better report and diagnose issues with FXC crashes that are circumvented by the cross-compilation pipeline.
On FXC crashes, HLSL is simplified with the SPIR-V based cross-compiler, and if successful, no error is reported to prevent build health issues.
This becomes a problem when the compilation is successful but the output becomes incompatible with other shader stages at PSO creation time.
To improve debuggability while also preserving build health, this change always dumps debug info for shader compile jobs when an internal compiler error is triggered.
This also reports those errors as structured log entries with Display verbosity.

Also fixes the following shader errors:
/Engine/Private/DeferredShadingCommon.ush(806,113-120): error X3508: 'MobileFetchGBuffer': output parameter 'GBufferD' not completely initialized
/Engine/Private/SingleLayerWaterComposite.usf(622,7-42): error X3020: operator cannot be used with a bool lvalue
/Engine/Private/HZBTracing.ush(176,4-63): error X3020: operator cannot be used with a bool lvalue
/Plugin/TmvMedia/Private/TmvTextureBufferSwizzler.usf(36,33): error X3000: syntax error: unexpected token '>'
/Plugin/TmvMedia/Private/TmvStructuredBufferSwizzler.usf(163,61-78): error X3508: 'GetBufferPositions_Planar': output parameter 'OutBufferPositions' not completely initialized
/Plugin/FX/Niagara/Shared/NiagaraRibbonShared.h(24,2-5): error X3000: syntax error: unexpected token 'enum'

#rnx

#Synced-CL 50315807

diff --git a/Engine/Plugins/FX/Niagara/Source/NiagaraShader/Private/NiagaraRibbonCompute.cpp b/Engine/Plugins/FX/Niagara/Source/NiagaraShader/Private/NiagaraRibbonCompute.cpp
--- a/Engine/Plugins/FX/Niagara/Source/NiagaraShader/Private/NiagaraRibbonCompute.cpp
+++ b/Engine/Plugins/FX/Niagara/Source/NiagaraShader/Private/NiagaraRibbonCompute.cpp
@@ -2,6 +2,7 @@
 
 
 #include "NiagaraRibbonCompute.h"
+#include "ShaderCompilerCore.h"
 
 constexpr uint32 FNiagaraRibbonComputeCommon::VertexGenReductionInitializationThreadSize;
 constexpr uint32 FNiagaraRibbonComputeCommon::VertexGenReductionPropagationThreadSize;
@@ -19,6 +20,9 @@ void FNiagaraRibbonComputeCommon::ModifyCompilationEnvironment(const FGlobalShad
 	{
 		OutEnvironment.SetDefine(TEXT("THREADGROUP_SIZE"), ThreadGroupSize);
 	}
+
+	// Precompile all NiagaraRibbon shaders with DXC, because FXC does not support the enum type in NiagaraRibbonShared.h
+	OutEnvironment.CompilerFlags.Add(CFLAG_PrecompileWithDXC);
 }
 
 void FNiagaraRibbonSortPhase1CS::ModifyCompilationEnvironment(const FGlobalShaderPermutationParameters& Parameters, FShaderCompilerEnvironment& OutEnvironment)
diff --git a/Engine/Plugins/Media/TmvMedia/Shaders/Private/TmvStructuredBufferSwizzler.usf b/Engine/Plugins/Media/TmvMedia/Shaders/Private/TmvStructuredBufferSwizzler.usf
--- a/Engine/Plugins/Media/TmvMedia/Shaders/Private/TmvStructuredBufferSwizzler.usf
+++ b/Engine/Plugins/Media/TmvMedia/Shaders/Private/TmvStructuredBufferSwizzler.usf
@@ -120,6 +120,9 @@ void GetBufferPositions_PartialTiles(in int2 InPixelPos, out int4 OutBufferPosit
 
 	int xTileCoord = (InPixelPos.x) % TileSize.x;
 	int yTileCoord = (InPixelPos.y) % TileSize.y;
+	
+	// Always default initialize `out` parameters to avoid error X3508 in FXC
+	OutBufferPositions = 0;
 
 	// Luminance is always full size.
 	OutBufferPositions[0] = TileDesc.Offsets[0] + yTileCoord * TileDesc.Strides[0] + xTileCoord;
@@ -146,6 +149,9 @@ void GetBufferPositions_RenderTiles(in int2 InPixelPos, out int4 OutBufferPositi
 
 	int TileBufferOffset = TileY * NumTiles.x + TileX * TileBufferFullStride;
 
+	// Always default initialize `out` parameters to avoid error X3508 in FXC
+	OutBufferPositions = 0;
+	
 	// Luminance is always full size.
 	OutBufferPositions[0] = TileBufferOffset + TileBufferOffsets[0] + yTileCoord * TileBufferStrides[0] + xTileCoord;
 
@@ -162,6 +168,9 @@ void GetBufferPositions_RenderTiles(in int2 InPixelPos, out int4 OutBufferPositi
 /** Calculate the buffer "position" of the given pixel for all planes (channels). */
 void GetBufferPositions_Planar(in int2 InPixelPos, out int4 OutBufferPositions)
 {
+	// Always default initialize `out` parameters to avoid error X3508 in FXC
+	OutBufferPositions = 0;
+	
 	// Luminance is always full size.
 	OutBufferPositions[0] = PlaneBufferOffsets[0] + InPixelPos.y * PlaneBufferStrides[0] + InPixelPos.x;
 	
diff --git a/Engine/Plugins/Media/TmvMedia/Shaders/Private/TmvTextureBufferSwizzler.usf b/Engine/Plugins/Media/TmvMedia/Shaders/Private/TmvTextureBufferSwizzler.usf
--- a/Engine/Plugins/Media/TmvMedia/Shaders/Private/TmvTextureBufferSwizzler.usf
+++ b/Engine/Plugins/Media/TmvMedia/Shaders/Private/TmvTextureBufferSwizzler.usf
@@ -33,7 +33,7 @@ float AlphaScale;
 
 float4 MipTint;
 
-float SampleComponent(Texture2D<> InTexture, float2 UV)
+float SampleComponent(Texture2D InTexture, float2 UV)
 {
 	return Texture2DSample(InTexture, InPlaneTextureSampler, UV).r;
 }
diff --git a/Engine/Shaders/Private/DeferredShadingCommon.ush b/Engine/Shaders/Private/DeferredShadingCommon.ush
--- a/Engine/Shaders/Private/DeferredShadingCommon.ush
+++ b/Engine/Shaders/Private/DeferredShadingCommon.ush
@@ -811,6 +811,8 @@ void MobileFetchGBuffer(in float2 UV, out half4 GBufferA, out half4 GBufferB, ou
 	GBufferC = VulkanSubpassFetch3();
 	#if MOBILE_EXTENDED_GBUFFER
 		GBufferD = VulkanSubpassFetch4();
+	#else
+		GBufferD = 0;
 	#endif
 #elif (METAL_ES3_1_PROFILE && !MAC && FRAMEBUFFER_FETCH_ACCESS)
 	GBufferA = SubpassFetchRGBA_1(); 
@@ -818,6 +820,8 @@ void MobileFetchGBuffer(in float2 UV, out half4 GBufferA, out half4 GBufferB, ou
 	GBufferC = SubpassFetchRGBA_3(); 
 	#if MOBILE_EXTENDED_GBUFFER
 		GBufferD = SubpassFetchRGBA_4();
+	#else
+		GBufferD = 0;
 	#endif
 #elif (USE_GLES_FBF_DEFERRED && FRAMEBUFFER_FETCH_ACCESS)
 	GBufferA = GLSubpassFetch1(); 
@@ -830,6 +834,8 @@ void MobileFetchGBuffer(in float2 UV, out half4 GBufferA, out half4 GBufferB, ou
 	GBufferC = Texture2DSampleLevel(MobileSceneTextures.GBufferCTexture, MobileSceneTextures.GBufferCTextureSampler, UV, 0);
 	#if MOBILE_EXTENDED_GBUFFER
 		GBufferD = Texture2DSampleLevel(MobileSceneTextures.GBufferDTexture, MobileSceneTextures.GBufferDTextureSampler, UV, 0);
+	#else
+		GBufferD = 0;
 	#endif
 #endif
 
diff --git a/Engine/Shaders/Private/HZBTracing.ush b/Engine/Shaders/Private/HZBTracing.ush
--- a/Engine/Shaders/Private/HZBTracing.ush
+++ b/Engine/Shaders/Private/HZBTracing.ush
@@ -173,7 +173,7 @@ void TraceHZB(
 #endif
 		{
 			IntersectionTime = min(IntersectionTime, PlaneIntersectionTimes.z);
-			bSkippedTile &= IntersectionTime != PlaneIntersectionTimes.z;
+			bSkippedTile = bSkippedTile && IntersectionTime != PlaneIntersectionTimes.z;
 		}
 
 		if (bSkippedTile)
diff --git a/Engine/Shaders/Private/SingleLayerWaterComposite.usf b/Engine/Shaders/Private/SingleLayerWaterComposite.usf
--- a/Engine/Shaders/Private/SingleLayerWaterComposite.usf
+++ b/Engine/Shaders/Private/SingleLayerWaterComposite.usf
@@ -619,7 +619,7 @@ void WaterTileClassificationBuildListsCS(
 					#if DOWNSAMPLE_FACTOR_X == 2
 						Mask |= select(FullTileCoord.x + 1 < FullTiledViewRes.x, Mask << 1u, Mask);
 					#endif
-						bTileUsed |= (MaskDword & Mask) != 0;
+						bTileUsed = bTileUsed || (MaskDword & Mask) != 0;
 
 						if (bTileUsed)
 						{
diff --git a/Engine/Shaders/Public/ShaderVersion.ush b/Engine/Shaders/Public/ShaderVersion.ush
--- a/Engine/Shaders/Public/ShaderVersion.ush
+++ b/Engine/Shaders/Public/ShaderVersion.ush
@@ -3,4 +3,4 @@
 // Each time the console command is executed it generates a new GUID. As this file is included
 // in Platform.ush (which should be included in any shader) it allows to invalidate the shader DDC.
 // 
-#pragma message("UESHADERMETADATA_VERSION 70B48F1C-5F2A-4874-9ED9-C5242BFF44EF")
\ No newline at end of file
+#pragma message("UESHADERMETADATA_VERSION 92B142B4-CAEB-47DF-A89D-83910D57E933")
\ No newline at end of file
diff --git a/Engine/Source/Developer/Windows/ShaderFormatD3D/Private/D3DShaderCompiler.cpp b/Engine/Source/Developer/Windows/ShaderFormatD3D/Private/D3DShaderCompiler.cpp
--- a/Engine/Source/Developer/Windows/ShaderFormatD3D/Private/D3DShaderCompiler.cpp
+++ b/Engine/Source/Developer/Windows/ShaderFormatD3D/Private/D3DShaderCompiler.cpp
@@ -837,8 +837,17 @@ static bool CompileAndProcessD3DShaderFXCExt(
 		}
 
 		// Some materials give FXC a hard time to optimize and the compiler fails with an internal error.
-		if (bPrecompileWithDXC || Result == HRESULT_FROM_WIN32(ERROR_ARITHMETIC_OVERFLOW) || Result == E_OUTOFMEMORY || Result == E_FAIL || (Result != S_OK && CompileErrorsContainInternalError(Errors.GetReference())))
+		const bool bFxcCompilerFailed = Result == HRESULT_FROM_WIN32(ERROR_ARITHMETIC_OVERFLOW) || Result == E_OUTOFMEMORY || Result == E_FAIL || (Result != S_OK && CompileErrorsContainInternalError(Errors.GetReference()));
+		if (bPrecompileWithDXC || bFxcCompilerFailed)
 		{
+			// If we're cross-compiling HLSL to simplified HLSL due to an FXC error, force diagnostics output and dump debug information.
+			// This is meant to help track down potential issues that only occur at PSO creation time.
+			if (bFxcCompilerFailed)
+			{
+				Output.bForceDumpDebugInfo = true;
+				Output.bForceLogDiagnosticMessages = true;
+			}
+
 			// If we ran out of memory, it's likely the next attempt will crash, too.
 			// Report the error now in case CompileHlslToSpirv throws an exception.
 			if (Result == E_OUTOFMEMORY)
@@ -973,7 +982,7 @@ static bool CompileAndProcessD3DShaderFXCExt(
 				}
 			}
 
-			if (!bPrecompileWithDXC && SUCCEEDED(Result))
+			if (bFxcCompilerFailed && SUCCEEDED(Result))
 			{
 				// Reset our previously set error code
 				FSCWErrorCode::Reset();
diff --git a/Engine/Source/Runtime/Engine/Private/ShaderCompiler/ShaderCompiler.cpp b/Engine/Source/Runtime/Engine/Private/ShaderCompiler/ShaderCompiler.cpp
--- a/Engine/Source/Runtime/Engine/Private/ShaderCompiler/ShaderCompiler.cpp
+++ b/Engine/Source/Runtime/Engine/Private/ShaderCompiler/ShaderCompiler.cpp
@@ -1489,7 +1489,7 @@ bool FShaderCompilingManager::ShouldRecompileToDumpShaderDebugInfo(const FShader
 		{
 			return bShouldDump && (bErrors);
 		}
-		else if (DumpShaderDebugInfo == EDumpShaderDebugInfo::OnErrorOrWarning)
+		else if (DumpShaderDebugInfo == EDumpShaderDebugInfo::OnErrorOrWarning || Output.bForceDumpDebugInfo)
 		{
 			return bShouldDump && (bErrors || bWarnings);
 		}
@@ -2243,12 +2243,29 @@ void FShaderCompilingManager::ProcessCompiledShaderMaps(
 
 					if (ErrorInfo.UniqueWarnings.Num() > 0)
 					{
-						UE_LOG(LogShaderCompilers, Warning, TEXT("Warnings while compiling Material %s for platform %s:"),
-							*Material->GetDebugName(),
-							*LegacyShaderPlatformToShaderFormat(CompilingShaderMap->GetShaderPlatform()).ToString());
+						const FName ShaderPlatformName = LegacyShaderPlatformToShaderFormat(CompilingShaderMap->GetShaderPlatform());
+						if (ShouldShowShaderWarnings())
+						{
+							UE_LOG(LogShaderCompilers, Warning, TEXT("Warnings while compiling Material %s for platform %s:"),
+								*Material->GetDebugName(), *ShaderPlatformName.ToString());
+						}
+						else
+						{
+							UE_LOG(LogShaderCompilers, Display, TEXT("Diagnostic messages while compiling Material %s for platform %s:"),
+								*Material->GetDebugName(), *ShaderPlatformName.ToString());
+						}
+
 						for (const FString& UniqueWarning : ErrorInfo.UniqueWarnings)
 						{
-							UE_LOG(LogShaders, Warning, TEXT("  %s"), *UniqueWarning);
+							// Use structured logs to avoid build machine string matching heuristic that will cause a cook to fail
+							if (ShouldShowShaderWarnings())
+							{
+								UE_LOGFMT(LogShaderCompilers, Warning, "  {DiagnosticMessage}", UniqueWarning);
+							}
+							else
+							{
+								UE_LOGFMT(LogShaderCompilers, Display, "  {DiagnosticMessage}", UniqueWarning);
+							}
 						}
 					}
 				}
@@ -2286,13 +2303,11 @@ void FShaderCompilingManager::ProcessCompiledShaderMaps(
 
 					for (const FString& UniqueError : ErrorInfo.UniqueErrors)
 					{
-						FString ErrorMessage = UniqueError;
-						// Work around build machine string matching heuristics that will cause a cook to fail
-						ErrorMessage.ReplaceInline(TEXT("error "), TEXT("err0r "), ESearchCase::CaseSensitive);
-						ErrorString += FString::Printf(TEXT("%s\n"), *ErrorMessage);
+						ErrorString += FString::Printf(TEXT("%s\n"), *UniqueError);
 					}
 
-					UE_LOG(LogShaderCompilers, Warning, TEXT("%s"), *ErrorString);
+					// Use structured logs to avoid build machine string matching heuristic that will cause a cook to fail
+					UE_LOGFMT(LogShaderCompilers, Warning, "{DiagnosticMessage}", ErrorString);
 				}
 
 				if (bReleaseCompilingId)
diff --git a/Engine/Source/Runtime/Engine/Private/ShaderCompiler/ShaderCompilerEditor.cpp b/Engine/Source/Runtime/Engine/Private/ShaderCompiler/ShaderCompilerEditor.cpp
--- a/Engine/Source/Runtime/Engine/Private/ShaderCompiler/ShaderCompilerEditor.cpp
+++ b/Engine/Source/Runtime/Engine/Private/ShaderCompiler/ShaderCompilerEditor.cpp
@@ -619,7 +619,14 @@ void ProcessCompiledGlobalShaders(const TArray<FShaderCommonCompileJobPtr>& Comp
 
 	for (const FString& WarningString : ShaderDiagInfo.UniqueWarnings)
 	{
-		UE_LOGFMT_NSLOC(LogShaders, Warning, "Shaders", "GlobalShaderCompileWarning", "{WarningMessage}", ("WarningMessage", WarningString));
+		if (ShouldShowShaderWarnings())
+		{
+			UE_LOGFMT_NSLOC(LogShaders, Warning, "Shaders", "GlobalShaderCompileWarning", "{WarningMessage}", ("WarningMessage", WarningString));
+		}
+		else
+		{
+			UE_LOGFMT_NSLOC(LogShaders, Display, "Shaders", "GlobalShaderCompileWarning", "{WarningMessage}", ("WarningMessage", WarningString));
+		}
 	}
 
 	TArray<EShaderPlatform> ShaderPlatformsProcessed;
diff --git a/Engine/Source/Runtime/RenderCore/Internal/ShaderDiagnostics.h b/Engine/Source/Runtime/RenderCore/Internal/ShaderDiagnostics.h
--- a/Engine/Source/Runtime/RenderCore/Internal/ShaderDiagnostics.h
+++ b/Engine/Source/Runtime/RenderCore/Internal/ShaderDiagnostics.h
@@ -26,4 +26,5 @@ private:
 
 RENDERCORE_API FString GetSingleJobCompilationDump(const FShaderCompileJob* SingleJob);
 RENDERCORE_API bool IsShaderDevelopmentModeEnabled();
+RENDERCORE_API bool ShouldShowShaderWarnings();
 
diff --git a/Engine/Source/Runtime/RenderCore/Private/ShaderCompilerCore.cpp b/Engine/Source/Runtime/RenderCore/Private/ShaderCompilerCore.cpp
--- a/Engine/Source/Runtime/RenderCore/Private/ShaderCompilerCore.cpp
+++ b/Engine/Source/Runtime/RenderCore/Private/ShaderCompilerCore.cpp
@@ -2,4 +2,4 @@
 
 #include "ShaderCompilerCore.h"
 
-const FGuid UE_SHADER_CACHE_VERSION{"24752250-148B-42A5-AE80-B8C513ED34B6"};
+const FGuid UE_SHADER_CACHE_VERSION{"B90637EA-0482-48F0-8D5D-58D3C7F71810"};
diff --git a/Engine/Source/Runtime/RenderCore/Private/ShaderDiagnostics.cpp b/Engine/Source/Runtime/RenderCore/Private/ShaderDiagnostics.cpp
--- a/Engine/Source/Runtime/RenderCore/Private/ShaderDiagnostics.cpp
+++ b/Engine/Source/Runtime/RenderCore/Private/ShaderDiagnostics.cpp
@@ -157,7 +157,7 @@ bool IsShaderDevelopmentModeEnabled()
 	return (CVarShaderDevelopmentMode.GetValueOnAnyThread() != 0);
 }
 
-bool ShouldShowWarnings()
+bool ShouldShowShaderWarnings()
 {
 	// show warnings if explicitly requested via the r.ShowShaderCompilerWarnings cvar, or if shader dev mode is enabled.
 	// in either case additional optional filtering happens at the job level via r.ShaderCompilerWarningsFilter (this can
@@ -165,6 +165,19 @@ bool ShouldShowWarnings()
 	return (CVarShowShaderWarnings.GetValueOnAnyThread() != 0) || IsShaderDevelopmentModeEnabled();
 }
 
+// Returns true if the shader compile job or pipeline compile job requested to always log its diagnostic messages.
+static bool AreDiagnosticMessagesForced(const FShaderCommonCompileJob& InJob)
+{
+	bool bOutAreDiagnosticMessagesForced = false;
+	InJob.ForEachSingleShaderJob([&bOutAreDiagnosticMessagesForced](const FShaderCompileJob& SingleJob) {
+		if (SingleJob.Output.bForceLogDiagnosticMessages)
+		{
+			bOutAreDiagnosticMessagesForced = true;
+		}
+	});
+	return bOutAreDiagnosticMessagesForced;
+}
+
 FShaderDiagnosticInfo::FShaderDiagnosticInfo(const TArray<FShaderCommonCompileJobPtr>& Jobs)
 {
 	// Gather unique errors
@@ -174,7 +187,7 @@ FShaderDiagnosticInfo::FShaderDiagnosticInfo(const TArray<FShaderCommonCompileJo
 		{
 			AddAndProcessErrorsForJob(*Job);
 		}
-		else if (ShouldShowWarnings())
+		else if (ShouldShowShaderWarnings() || AreDiagnosticMessagesForced(*Job))
 		{
 			AddWarningsForJob(*Job);
 		}
diff --git a/Engine/Source/Runtime/RenderCore/Public/ShaderCompilerCore.h b/Engine/Source/Runtime/RenderCore/Public/ShaderCompilerCore.h
--- a/Engine/Source/Runtime/RenderCore/Public/ShaderCompilerCore.h
+++ b/Engine/Source/Runtime/RenderCore/Public/ShaderCompilerCore.h
@@ -496,6 +496,8 @@ struct FShaderCompilerOutput
 	,	bSupportsQueryingUsedAttributes(false)
 	,	bSerializingForCache(false)
 	,	bInterpolatorRemovalFailed(false)
+	,	bForceDumpDebugInfo(false)
+	,	bForceLogDiagnosticMessages(false)
 	{
 	}
 
@@ -512,7 +514,25 @@ struct FShaderCompilerOutput
 	bool bSucceeded;
 	bool bSupportsQueryingUsedAttributes;
 	bool bSerializingForCache;
+
+	/** Specifies that the interpolator removal optimization failed for this job.
+	 * If the job still succeeded, then all other shaders for the associated pipeline job
+	 * must also skip interpolator removal to avoid mismatch between interpolators of shader stages.
+	 */
 	bool bInterpolatorRemovalFailed;
+
+	/** Specifies to always dump debug information for this job output.
+	 * If the compile job was successful, this allows to dump debug information to diagnose failed intermediate results, regardless of r.DumpShaderDebugInfo configuration.
+	 * This is primarily used when the cross-compilation pipeline is triggered due to internal compiler errors (usually with FXC).
+	 */
+	bool bForceDumpDebugInfo;
+
+	/** Specifies to always log diagnostic messages.
+	 * If the compile job was successful, this will be printed with Display verbosity.
+	 * If r.ShaderDevelopmentMode is enabled, this will be printed with Warning verbosity.
+	 */
+	bool bForceLogDiagnosticMessages;
+
 	TArray<FString> UsedAttributes;
 
 	TArray<FShaderCodeValidationStride> ParametersStrideToValidate;
@@ -608,6 +628,8 @@ struct FShaderCompilerOutput
 			Ar << Output.CompileTime;
 		}
 		Ar << Output.bInterpolatorRemovalFailed;
+		Ar << Output.bForceDumpDebugInfo;
+		Ar << Output.bForceLogDiagnosticMessages;
 		Ar << Output.OutputHash;
 		Ar << Output.NumInstructions;
 		Ar << Output.NumTextureSamplers;
diff --git a/Engine/Source/Runtime/Renderer/Private/HeterogeneousVolumes/HeterogeneousVolumesPreshadingPipeline.cpp b/Engine/Source/Runtime/Renderer/Private/HeterogeneousVolumes/HeterogeneousVolumesPreshadingPipeline.cpp
--- a/Engine/Source/Runtime/Renderer/Private/HeterogeneousVolumes/HeterogeneousVolumesPreshadingPipeline.cpp
+++ b/Engine/Source/Runtime/Renderer/Private/HeterogeneousVolumes/HeterogeneousVolumesPreshadingPipeline.cpp
@@ -290,8 +290,10 @@ class FRenderSingleScatteringWithPreshadingCS : public FGlobalShader
 			FVirtualShadowMapArray::SetShaderDefines(OutEnvironment);
 		}
 
-		// This shader takes a very long time to compile with FXC, so we pre-compile it with DXC first and then forward the optimized HLSL to FXC.
-		//OutEnvironment.CompilerFlags.Add(CFLAG_PrecompileWithDXC);
+		// This shader does not compile with FXC, due to the following error,
+		// so we pre-compile it with DXC first and then forward the optimized HLSL to FXC.
+		//   /Engine/Private/Common.ush(307,9-83): error X3538: Sampler parameter must come from a literal expression.
+		OutEnvironment.CompilerFlags.Add(CFLAG_PrecompileWithDXC);
 		OutEnvironment.CompilerFlags.Add(CFLAG_AllowTypedUAVLoads);
 
 		OutEnvironment.SetDefine(TEXT("GET_PRIMITIVE_DATA_OVERRIDE"), 1);
