Fixed input pin display name when using Mip Value Mode "Mip Bias" on a virtual texture sample material node

#Synced-CL 50259661

diff --git a/Engine/Source/Runtime/Engine/Private/Materials/MaterialExpressions.cpp b/Engine/Source/Runtime/Engine/Private/Materials/MaterialExpressions.cpp
--- a/Engine/Source/Runtime/Engine/Private/Materials/MaterialExpressions.cpp
+++ b/Engine/Source/Runtime/Engine/Private/Materials/MaterialExpressions.cpp
@@ -3310,7 +3310,7 @@ FName UMaterialExpressionRuntimeVirtualTextureSample::GetInputName(int32 InputIn
 		}
 		if (MipValueMode == RVTMVM_MipBias)
 		{
-			return TEXT("Mip Level");
+			return TEXT("Mip Bias");
 		}
 		else if (MipValueMode == RVTMVM_DerivativeUV)
 		{
