FTangentSpline: Fix some version setting/updating logic and fix constness of GetSplineLength

#Synced-CL 50305360

diff --git a/Engine/Source/Runtime/Engine/Classes/Curves/Splines/TangentSpline.h b/Engine/Source/Runtime/Engine/Classes/Curves/Splines/TangentSpline.h
--- a/Engine/Source/Runtime/Engine/Classes/Curves/Splines/TangentSpline.h
+++ b/Engine/Source/Runtime/Engine/Classes/Curves/Splines/TangentSpline.h
@@ -2,10 +2,10 @@
 
 #pragma once
 
+#include "Async/TransactionallySafeMutex.h"
 #include "BoxTypes.h"
 #include "CoreMinimal.h"
 #include "Math/InterpCurve.h"
-#include "Async/TransactionallySafeMutex.h"
 #include "Misc/TransactionallySafeRWLock.h"
 
 PRAGMA_DISABLE_EXPERIMENTAL_WARNINGS
@@ -54,7 +54,7 @@ public:
 	}
 
 	UE_API float GetSegmentLength(const int32 Index, const float Param, const FVector& Scale3D = FVector(1.0f)) const;
-	UE_API float GetSplineLength();
+	UE_API float GetSplineLength() const;
 
 	UE_API virtual void SetClosedLoop(bool bInClosedLoop) override;
 	UE_API void SetClosedLoop(bool bInClosedLoop, bool bUpdateSpline);
diff --git a/Engine/Source/Runtime/Engine/Private/Curves/Splines/TangentSpline.cpp b/Engine/Source/Runtime/Engine/Private/Curves/Splines/TangentSpline.cpp
--- a/Engine/Source/Runtime/Engine/Private/Curves/Splines/TangentSpline.cpp
+++ b/Engine/Source/Runtime/Engine/Private/Curves/Splines/TangentSpline.cpp
@@ -270,6 +270,8 @@ FTangentSpline::FTangentSpline(const FSplineCurves& Other)
 	}
 	
 	UpdateSpline(UpdateParams);
+
+	Version = Other.Version;
 }
 
 /** Copy assignment */
@@ -282,6 +284,7 @@ FTangentSpline& FTangentSpline::operator=(const FTangentSpline& Other)
 		RotationCurve = Other.RotationCurve;
 		ScaleCurve = Other.ScaleCurve;
 		ReparamStepsPerSegment = Other.ReparamStepsPerSegment;
+		Version = Other.Version;
 	}
 
 	return *this;
@@ -490,7 +493,7 @@ float FTangentSpline::GetSegmentLength(const int32 Index, const float Param, con
 	return Length;
 }
 
-float FTangentSpline::GetSplineLength()
+float FTangentSpline::GetSplineLength() const
 {
 	// Evaluate the ParamToLength channel at the very end of the spline
 	const float MaxParameter = GetParameterSpace().Max;
@@ -1353,6 +1356,8 @@ void FTangentSpline::UpdateSpline()
 
 	// Updates legacy curves based on points, never evaluates internal spline.
 	RebuildLegacyCurves();
+
+	++Version;
 }
 
 void FTangentSpline::UpdateSpline(const FUpdateSplineParams& Params)
