Fix for LOD explosion when copying pose from a parent at a higher LOD

#Synced-CL 50280823

diff --git a/Engine/Plugins/ChaosClothAsset/Source/ChaosClothAssetEngine/Private/ChaosClothAsset/ClothComponent.cpp b/Engine/Plugins/ChaosClothAsset/Source/ChaosClothAssetEngine/Private/ChaosClothAsset/ClothComponent.cpp
--- a/Engine/Plugins/ChaosClothAsset/Source/ChaosClothAssetEngine/Private/ChaosClothAsset/ClothComponent.cpp
+++ b/Engine/Plugins/ChaosClothAsset/Source/ChaosClothAssetEngine/Private/ChaosClothAsset/ClothComponent.cpp
@@ -575,9 +575,9 @@ void UChaosClothComponent::GetAdditionalRequiredBonesForLeader(int32 LeaderLODIn
 
 			for (const FBoneIndexType RequiredBone : SkeletalMeshRenderData->LODRenderData[LODIndex].RequiredBones)
 			{
-				if (LeaderBoneMap.IsValidIndex(RequiredBone))
+				if (GetLeaderBoneMap().IsValidIndex(RequiredBone))
 				{
-					const int32 FollowerRequiredLeaderBone = LeaderBoneMap[RequiredBone];
+					const int32 FollowerRequiredLeaderBone = GetLeaderBoneMap()[RequiredBone];
 					if (FollowerRequiredLeaderBone != INDEX_NONE)
 					{
 						RequiredBones.Add((FBoneIndexType)FollowerRequiredLeaderBone);
diff --git a/Engine/Plugins/Experimental/UAF/UAF/Source/UAF/Private/Component/SkinnedMeshComponentExtensions.cpp b/Engine/Plugins/Experimental/UAF/UAF/Source/UAF/Private/Component/SkinnedMeshComponentExtensions.cpp
--- a/Engine/Plugins/Experimental/UAF/UAF/Source/UAF/Private/Component/SkinnedMeshComponentExtensions.cpp
+++ b/Engine/Plugins/Experimental/UAF/UAF/Source/UAF/Private/Component/SkinnedMeshComponentExtensions.cpp
@@ -21,7 +21,8 @@ namespace UE::Anim
 void ConvertLocalSpaceToComponentSpace(
 	const USkinnedMeshComponent* InComponent,
 	const UE::UAF::FLODPoseHeap& InLODPose,
-	TArrayView<FTransform> OutComponentSpaceTransforms)
+	TArrayView<FTransform> OutComponentSpaceTransforms,
+	TBitArray<>& OutValidMeshPoseTransforms)
 {
 	SCOPE_CYCLE_COUNTER(STAT_AnimNext_ConvertLocalSpaceToComponentSpace);
 
@@ -33,6 +34,8 @@ void ConvertLocalSpaceToComponentSpace(
 	checkf(LODBoneIndexToMeshBoneIndexMap.Num() == InLODPose.GetNumBones(), TEXT("Buffer mismatch: %d:%d"), LODBoneIndexToMeshBoneIndexMap.Num(), InLODPose.GetNumBones());
 	checkf(MeshBoneIndexToParentMeshBoneIndexMap.Num() == OutComponentSpaceTransforms.Num(), TEXT("Buffer mismatch: %d:%d"), MeshBoneIndexToParentMeshBoneIndexMap.Num(), OutComponentSpaceTransforms.Num());
 
+	OutValidMeshPoseTransforms.Init(false, OutComponentSpaceTransforms.Num());
+
 	const int32 NumPoseLODBones = InLODPose.GetNumBones();
 	const FBoneIndexType* RESTRICT MeshBoneIndexToParentMeshBoneIndexMapPtr = MeshBoneIndexToParentMeshBoneIndexMap.GetData();
 	const FBoneIndexType* RESTRICT LODBoneIndexToMeshBoneIndexMapPtr = LODBoneIndexToMeshBoneIndexMap.GetData();
@@ -42,6 +45,7 @@ void ConvertLocalSpaceToComponentSpace(
 	{
 		check(LODBoneIndexToMeshBoneIndexMap[0] == 0);
 		OutComponentSpaceTransforms[0] = InLODPose.LocalTransformsView[0];
+		OutValidMeshPoseTransforms[0] = true;
 	}
 
 	for (int32 LODBoneIndex = 1; LODBoneIndex < NumPoseLODBones; ++LODBoneIndex)
@@ -57,6 +61,8 @@ void ConvertLocalSpaceToComponentSpace(
 
 		FTransform::Multiply(ComponentSpaceTransform, &LocalSpaceTransform, ParentComponentSpaceTransform);
 
+		OutValidMeshPoseTransforms[MeshBoneIndex] = true;
+
 		ComponentSpaceTransform->NormalizeRotation();
 
 		checkSlow(ComponentSpaceTransform->IsRotationNormalized());
@@ -90,6 +96,8 @@ void ConvertLocalSpaceToComponentSpace(
 
 			ComponentSpaceTransform->NormalizeRotation();
 
+			OutValidMeshPoseTransforms[MeshBoneIndex] = true;
+
 			checkSlow(ComponentSpaceTransform->IsRotationNormalized());
 			checkSlow(!ComponentSpaceTransform->ContainsNaN());
 		}
@@ -107,7 +115,8 @@ void FSkinnedMeshComponentExtensions::CompleteAndDispatch(
 	TArrayView<FTransform> ComponentSpaceTransforms = InComponent->GetEditableComponentSpaceTransforms();
 	if (ComponentSpaceTransforms.Num() > 0)
 	{
-		ConvertLocalSpaceToComponentSpace(InComponent, InLODPose, ComponentSpaceTransforms);
+		ConvertLocalSpaceToComponentSpace(InComponent, InLODPose, ComponentSpaceTransforms, InComponent->GetMutableValidMeshPoseTransforms());
+		InComponent->bValidMeshPoseTransformsUpToDate = true;
 
 		// Flag buffer for flip
 		InComponent->bNeedToFlipSpaceBaseBuffers = true;
diff --git a/Engine/Plugins/Experimental/UAF/UAF/Source/UAF/Private/Graph/RigUnit_UAFWriteInstancedSkinnedMeshPose.cpp b/Engine/Plugins/Experimental/UAF/UAF/Source/UAF/Private/Graph/RigUnit_UAFWriteInstancedSkinnedMeshPose.cpp
--- a/Engine/Plugins/Experimental/UAF/UAF/Source/UAF/Private/Graph/RigUnit_UAFWriteInstancedSkinnedMeshPose.cpp
+++ b/Engine/Plugins/Experimental/UAF/UAF/Source/UAF/Private/Graph/RigUnit_UAFWriteInstancedSkinnedMeshPose.cpp
@@ -21,7 +21,8 @@ namespace UE::Anim
 extern void ConvertLocalSpaceToComponentSpace(
 	const USkinnedMeshComponent* InComponent,
 	const UE::UAF::FLODPoseHeap& InLODPose,
-	TArrayView<FTransform> OutComponentSpaceTransforms);
+	TArrayView<FTransform> OutComponentSpaceTransforms,
+	TBitArray<>& OutValidMeshPoseTransforms);
 
 }
 
@@ -58,7 +59,9 @@ FRigUnit_UAFWriteInstancedSkinnedMeshPose_Execute()
 	TArray<FCompressedBoneTransform> ReferenceToLocalTransforms;
 	ComponentSpaceTransforms.SetNumUninitialized(RefBasesInvMatrix.Num());
 	ReferenceToLocalTransforms.SetNumUninitialized(RefBasesInvMatrix.Num());
-	UE::Anim::ConvertLocalSpaceToComponentSpace(SkinnedMeshComponent, Pose.LODPose, ComponentSpaceTransforms);
+
+	TBitArray<> DummyValidMeshPoseTransforms;
+	UE::Anim::ConvertLocalSpaceToComponentSpace(SkinnedMeshComponent, Pose.LODPose, ComponentSpaceTransforms, DummyValidMeshPoseTransforms);
 
 	for (int32 Index = 0; Index < ComponentSpaceTransforms.Num(); ++Index)
 	{
diff --git a/Engine/Source/Runtime/AnimGraphRuntime/Private/AnimNodes/AnimNode_CopyPoseFromMesh.cpp b/Engine/Source/Runtime/AnimGraphRuntime/Private/AnimNodes/AnimNode_CopyPoseFromMesh.cpp
--- a/Engine/Source/Runtime/AnimGraphRuntime/Private/AnimNodes/AnimNode_CopyPoseFromMesh.cpp
+++ b/Engine/Source/Runtime/AnimGraphRuntime/Private/AnimNodes/AnimNode_CopyPoseFromMesh.cpp
@@ -162,6 +162,11 @@ void FAnimNode_CopyPoseFromMesh::PreUpdate(const UAnimInstance* InAnimInstance)
 			SourceMeshTransformArray.Reset();
 			SourceMeshTransformArray.Append((bUROInSync || bUsingExternalInterpolation) && bArraySizesMatch ? CachedComponentSpaceTransforms : CurrentMeshComponent->GetComponentSpaceTransforms());
 
+			if (UE::Anim::bUseSafeMeshPoseIndices)
+			{
+				SourceMeshValidMeshPoseTransforms = CurrentMeshComponent->GetValidMeshPoseTransforms();
+			}
+
 			// Ref skeleton is need for parent index lookups later, so store it now
 			CurrentlyUsedMesh = CurrentMeshComponent->GetSkeletalMeshAsset();
 
@@ -232,8 +237,11 @@ void FAnimNode_CopyPoseFromMesh::Evaluate_AnyThread(FPoseContext& Output)
 				const int32* Value = BoneMapToSource.Find(MeshBoneIndex.GetInt());
  				if (Value && SourceMeshTransformArray.IsValidIndex(*Value))
 				{
-					const int32 SourceBoneIndex = *Value;
-					MeshPoses.SetComponentSpaceTransform(PoseBoneIndex, SourceMeshTransformArray[SourceBoneIndex]);
+					if ((!UE::Anim::bUseSafeMeshPoseIndices || (SourceMeshValidMeshPoseTransforms.IsValidIndex(*Value) && SourceMeshValidMeshPoseTransforms[*Value])))
+					{
+						const int32 SourceBoneIndex = *Value;
+						MeshPoses.SetComponentSpaceTransform(PoseBoneIndex, SourceMeshTransformArray[SourceBoneIndex]);
+					}
 				}
 			}
 
@@ -247,19 +255,22 @@ void FAnimNode_CopyPoseFromMesh::Evaluate_AnyThread(FPoseContext& Output)
 				const int32* Value = BoneMapToSource.Find(MeshBoneIndex.GetInt());
 				if (Value && SourceMeshTransformArray.IsValidIndex(*Value))
 				{
-					const int32 SourceBoneIndex = *Value;
-					const int32 ParentIndex = CurrentMesh->GetRefSkeleton().GetParentIndex(SourceBoneIndex);
-					const FCompactPoseBoneIndex MyParentIndex = RequiredBones.GetParentBoneIndex(PoseBoneIndex);
-					// only apply if I also have parent, otherwise, it should apply the space bases
-					if (SourceMeshTransformArray.IsValidIndex(ParentIndex) && MyParentIndex != INDEX_NONE)
-					{
-						const FTransform& ParentTransform = SourceMeshTransformArray[ParentIndex];
-						const FTransform& ChildTransform = SourceMeshTransformArray[SourceBoneIndex];
-						OutPose[PoseBoneIndex] = ChildTransform.GetRelativeTransform(ParentTransform);
-					}
-					else
+					if ((!UE::Anim::bUseSafeMeshPoseIndices || (SourceMeshValidMeshPoseTransforms.IsValidIndex(*Value) && SourceMeshValidMeshPoseTransforms[*Value])))
 					{
-						OutPose[PoseBoneIndex] = SourceMeshTransformArray[SourceBoneIndex];
+						const int32 SourceBoneIndex = *Value;
+						const int32 ParentIndex = CurrentMesh->GetRefSkeleton().GetParentIndex(SourceBoneIndex);
+						const FCompactPoseBoneIndex MyParentIndex = RequiredBones.GetParentBoneIndex(PoseBoneIndex);
+						// only apply if I also have parent, otherwise, it should apply the space bases
+						if (SourceMeshTransformArray.IsValidIndex(ParentIndex) && MyParentIndex != INDEX_NONE)
+						{
+							const FTransform& ParentTransform = SourceMeshTransformArray[ParentIndex];
+							const FTransform& ChildTransform = SourceMeshTransformArray[SourceBoneIndex];
+							OutPose[PoseBoneIndex] = ChildTransform.GetRelativeTransform(ParentTransform);
+						}
+						else
+						{
+							OutPose[PoseBoneIndex] = SourceMeshTransformArray[SourceBoneIndex];
+						}
 					}
 				}
 			}
diff --git a/Engine/Source/Runtime/AnimGraphRuntime/Public/AnimNodes/AnimNode_CopyPoseFromMesh.h b/Engine/Source/Runtime/AnimGraphRuntime/Public/AnimNodes/AnimNode_CopyPoseFromMesh.h
--- a/Engine/Source/Runtime/AnimGraphRuntime/Public/AnimNodes/AnimNode_CopyPoseFromMesh.h
+++ b/Engine/Source/Runtime/AnimGraphRuntime/Public/AnimNodes/AnimNode_CopyPoseFromMesh.h
@@ -70,6 +70,9 @@ private:
 	// Cached transforms, copied on the game thread
 	TArray<FTransform> SourceMeshTransformArray;
 
+	// Cached transforms, copied on the game thread
+	TBitArray<> SourceMeshValidMeshPoseTransforms;
+
 	// Cached curves, copied on the game thread
 	FBlendedHeapCurve SourceCurves;
 
diff --git a/Engine/Source/Runtime/Engine/Classes/Components/SkeletalMeshComponent.h b/Engine/Source/Runtime/Engine/Classes/Components/SkeletalMeshComponent.h
--- a/Engine/Source/Runtime/Engine/Classes/Components/SkeletalMeshComponent.h
+++ b/Engine/Source/Runtime/Engine/Classes/Components/SkeletalMeshComponent.h
@@ -1783,6 +1783,12 @@ public:
 	  * @param	LODIndex	Index of LOD [0-(MaxLOD-1)]
 	*/
 	ENGINE_API void ComputeRequiredBones(TArray<FBoneIndexType>& OutRequiredBones, TArray<FBoneIndexType>& OutFillComponentSpaceTransformsRequiredBones, int32 LODIndex, bool bIgnorePhysicsAsset) const;
+	
+	/** Computes a dense bit array of valid bone transforms, based on a list of required bones.
+	 * The bit array is sized to the full number of bones of our RefSkeleton.
+	 * This can be used to check our component space bone transforms for validity, without iterating over the full required bones list.
+	 */
+	ENGINE_API void ComputeValidMeshPoseTransforms(TArray<FBoneIndexType>& InRequiredBones, TBitArray<>& OutValidTransforms) const;
 
 	// Get the required virtual bones from the SkeletalMesh Reference Skeleton
 	static ENGINE_API void GetRequiredVirtualBones(const USkeletalMesh* SkeletalMesh, TArray<FBoneIndexType>& OutRequiredBones);
diff --git a/Engine/Source/Runtime/Engine/Classes/Components/SkinnedMeshComponent.h b/Engine/Source/Runtime/Engine/Classes/Components/SkinnedMeshComponent.h
--- a/Engine/Source/Runtime/Engine/Classes/Components/SkinnedMeshComponent.h
+++ b/Engine/Source/Runtime/Engine/Classes/Components/SkinnedMeshComponent.h
@@ -55,6 +55,7 @@ namespace Nanite
 namespace UE::Anim
 {
 	struct FSkinnedMeshComponentExtensions;
+	ENGINE_API extern bool bUseSafeMeshPoseIndices;
 }
 
 DECLARE_DELEGATE_OneParam(FOnAnimUpdateRateParamsCreated, FAnimUpdateRateParameters*)
@@ -436,6 +437,12 @@ private:
 	/** Temporary array of of component-space bone matrices, update each frame and used for rendering the mesh. */
 	TArray<FTransform> ComponentSpaceTransformsArray[2];
 
+	/**
+	 * Mesh Pose (sized to our skinned asset's ref skeleton) bit array of valid bone indices. 
+	 * This should be populated by derived classes to avoid children copying invalid transforms through i.e. leader pose/copy pose.
+	 */
+	TBitArray<> ValidMeshPoseTransforms;
+
 #if WITH_STATE_STREAM_ACTOR
 	FSkinnedMeshHandle MeshHandle;
 #endif
@@ -537,6 +544,7 @@ protected:
 
 public:
 	const TArray<int32>& GetLeaderBoneMap() const { return LeaderBoneMap; }
+	TArray<int32>& GetMutableLeaderBoneMap() { return LeaderBoneMap; }
 
 	UE_DEPRECATED(5.1, "This method has been deprecated. Please use the GetLeaderBoneMap.")
 	const TArray<int32>& GetMasterBoneMap() const { return GetLeaderBoneMap(); }
@@ -768,6 +776,9 @@ protected:
 	/** used to cache previous bone transform or not */
 	uint8 bHasValidBoneTransform:1;
 
+	/** used to cache our valid mesh pose transform list for the current animation update. */
+	uint8 bValidMeshPoseTransformsUpToDate : 1;
+
 	/** Whether or not a Skin Weight profile is currently set for this component */
 	uint8 bSkinWeightProfileSet:1;
 
@@ -1633,6 +1644,8 @@ protected:
 		PRAGMA_ENABLE_DEPRECATION_WARNINGS
 	}
 
+	inline TBitArray<>& GetMutableValidMeshPoseTransforms() { return ValidMeshPoseTransforms; }
+
 public:
 	virtual void UpdateVisualizeLODString(FString& DebugString) {}
 
@@ -1683,6 +1696,8 @@ public:
 		return GetComponentSpaceTransforms().Num(); 
 	}
 
+	inline const TBitArray<>& GetValidMeshPoseTransforms() const { return ValidMeshPoseTransforms; }
+
 	/** Access BoneVisibilityStates for reading */
 	const TArray<uint8>& GetBoneVisibilityStates() const 
 	{
diff --git a/Engine/Source/Runtime/Engine/Private/Components/SkeletalMeshComponent.cpp b/Engine/Source/Runtime/Engine/Private/Components/SkeletalMeshComponent.cpp
--- a/Engine/Source/Runtime/Engine/Private/Components/SkeletalMeshComponent.cpp
+++ b/Engine/Source/Runtime/Engine/Private/Components/SkeletalMeshComponent.cpp
@@ -499,6 +499,7 @@ USkeletalMeshComponent::USkeletalMeshComponent(const FObjectInitializer& ObjectI
 
 	bAllowClothActors = true;
 	bPostEvaluatingAnimation = false;
+	bValidMeshPoseTransformsUpToDate = false;
 	bDisableAnimDynamicsAnimNode = false;
 	bAllowAnimCurveEvaluation = true;
 	bDisablePostProcessBlueprint = false;
@@ -1769,6 +1770,7 @@ bool USkeletalMeshComponent::UpdateLODStatus()
 	if (Super::UpdateLODStatus())
 	{
 		bRequiredBonesUpToDate = false;
+		bValidMeshPoseTransformsUpToDate = false;
 		return true;
 	}
 
@@ -2227,6 +2229,15 @@ void USkeletalMeshComponent::ComputeRequiredBones(TArray<FBoneIndexType>& OutReq
 	FAnimationRuntime::EnsureParentsPresent(OutFillComponentSpaceTransformsRequiredBones, SkelMesh->GetRefSkeleton());
 }
 
+void USkeletalMeshComponent::ComputeValidMeshPoseTransforms(TArray<FBoneIndexType>& InRequiredBones, TBitArray<>& OutValidTransforms) const
+{
+	OutValidTransforms.Init(false, GetNumComponentSpaceTransforms());
+	for (int32 i = 0; i < InRequiredBones.Num(); ++i)
+	{
+		OutValidTransforms[InRequiredBones[i]] = true; 
+	}
+}
+
 /*static*/ void USkeletalMeshComponent::GetRequiredVirtualBones(const USkeletalMesh* SkeletalMesh, TArray<FBoneIndexType>& OutRequiredBones)
 {
 	check (SkeletalMesh != nullptr);
@@ -2705,7 +2716,7 @@ void USkeletalMeshComponent::WritebackClothingSimulationData()
 			OverrideComponent = LeaderPoseComponent.Get();
 
 			// Check if our bone map is actually valid, if not there is no clothing data to build
-			if (LeaderBoneMap.Num() == 0)
+			if (GetLeaderBoneMap().Num() == 0)
 			{
 				CurrentSimulationData.Reset();
 				return;
@@ -3522,12 +3533,16 @@ bool USkeletalMeshComponent::AllocateTransformData()
 		}
 		PRAGMA_ENABLE_DEPRECATION_WARNINGS
 
+		GetMutableValidMeshPoseTransforms().Init(false, GetNumComponentSpaceTransforms());
+		bValidMeshPoseTransformsUpToDate = false;
+
 		return true;
 	}
 
 	PRAGMA_DISABLE_DEPRECATION_WARNINGS
 	BoneSpaceTransforms.Empty();
 	PRAGMA_ENABLE_DEPRECATION_WARNINGS
+	GetMutableValidMeshPoseTransforms().Empty();
 	return false;
 }
 
@@ -5055,6 +5070,12 @@ void USkeletalMeshComponent::FinalizeBoneTransform()
 {
 	Super::FinalizeBoneTransform();
 
+	if (UE::Anim::bUseSafeMeshPoseIndices && bValidMeshPoseTransformsUpToDate == false)
+	{
+		ComputeValidMeshPoseTransforms(FillComponentSpaceTransformsRequiredBones, GetMutableValidMeshPoseTransforms());
+		bValidMeshPoseTransformsUpToDate = true;
+	}
+
 	// After pose has been finalized, dispatch AnimNotifyEvents in case they want to use up to date pose.
 	// (For example attaching particle systems to up to date sockets).
 
diff --git a/Engine/Source/Runtime/Engine/Private/Components/SkinnedMeshComponent.cpp b/Engine/Source/Runtime/Engine/Private/Components/SkinnedMeshComponent.cpp
--- a/Engine/Source/Runtime/Engine/Private/Components/SkinnedMeshComponent.cpp
+++ b/Engine/Source/Runtime/Engine/Private/Components/SkinnedMeshComponent.cpp
@@ -62,6 +62,14 @@
 DEFINE_LOG_CATEGORY_STATIC(LogSkinnedMeshComp, Log, All);
 #define LOCTEXT_NAMESPACE "SkinnedMeshComponent"
 
+namespace UE::Anim
+{
+	bool bUseSafeMeshPoseIndices = false;
+	FAutoConsoleVariableRef CVarUseSafeMeshPoseIndices(TEXT("a.UseSafeMeshPoseIndices"), bUseSafeMeshPoseIndices
+		, TEXT("Whether to compute and use valid mesh pose indices on skeletal mesh components. This allows us to safely reset missing bones in leader/follower setups. This has a small additional cost.")
+		, ECVF_Default);
+}
+
 #if !UE_BUILD_SHIPPING
 int32 GSkeletalMeshForceLOD = -1;
 FAutoConsoleVariableRef CVarSkeletalMeshForceLOD(
@@ -482,6 +490,7 @@ USkinnedMeshComponent::USkinnedMeshComponent(const FObjectInitializer& ObjectIni
 	CurrentEditableComponentTransforms = 0;
 	CurrentReadComponentTransforms = 1;
 	bNeedToFlipSpaceBaseBuffers = false;
+	bValidMeshPoseTransformsUpToDate = false;
 	bBoneVisibilityDirty = false;
 
 	bCanEverAffectNavigation = false;
@@ -2229,9 +2238,9 @@ FMatrix USkinnedMeshComponent::GetBoneMatrix(int32 BoneIdx) const
 	const USkinnedMeshComponent* const LeaderPoseComponentInst = LeaderPoseComponent.Get();
 	if(LeaderPoseComponentInst)
 	{
-		if(BoneIdx < LeaderBoneMap.Num())
+		if(BoneIdx < GetLeaderBoneMap().Num())
 		{
-			int32 ParentBoneIndex = LeaderBoneMap[BoneIdx];
+			int32 ParentBoneIndex = GetLeaderBoneMap()[BoneIdx];
 
 			// If ParentBoneIndex is valid, grab matrix from LeaderPoseComponent.
 			if(	ParentBoneIndex != INDEX_NONE && 
@@ -2290,9 +2299,9 @@ FTransform USkinnedMeshComponent::GetBoneTransform(int32 BoneIdx, const FTransfo
 			// isn't valid so for now return identity
 			return FTransform::Identity;
 		}
-		if(BoneIdx >= 0 && BoneIdx < LeaderBoneMap.Num())
+		if(BoneIdx >= 0 && BoneIdx < GetLeaderBoneMap().Num())
 		{
-			const int32 LeaderBoneIndex = LeaderBoneMap[BoneIdx];
+			const int32 LeaderBoneIndex = GetLeaderBoneMap()[BoneIdx];
 			const int32 NumLeaderTransforms = LeaderPoseComponentInst->GetNumComponentSpaceTransforms();
 
 			// If LeaderBoneIndex is valid, grab matrix from LeaderPoseComponent.
@@ -3314,7 +3323,7 @@ bool USkinnedMeshComponent::GetMissingLeaderBoneRelativeTransform(int32 InBoneIn
 		CommonAncestorBoneIndex = FollowerRefSkeleton.GetParentIndex(CommonAncestorBoneIndex);
 		if(CommonAncestorBoneIndex != INDEX_NONE)
 		{
-			OutInfo.CommonAncestorBoneIndex = LeaderBoneMap[CommonAncestorBoneIndex];
+			OutInfo.CommonAncestorBoneIndex = GetLeaderBoneMap()[CommonAncestorBoneIndex];
 			if(OutInfo.CommonAncestorBoneIndex != INDEX_NONE)
 			{
 				OutInfo.RelativeTransform = RelativeTransform;
@@ -3330,7 +3339,7 @@ bool USkinnedMeshComponent::GetMissingLeaderBoneRelativeTransform(int32 InBoneIn
 
 void USkinnedMeshComponent::UpdateLeaderBoneMap()
 {
-	LeaderBoneMap.Reset();
+	GetMutableLeaderBoneMap().Reset();
 	MissingLeaderBoneMap.Reset();
 
 	if (GetSkinnedAsset())
@@ -3616,9 +3625,9 @@ FQuat USkinnedMeshComponent::GetBoneQuaternion(FName BoneName, EBoneSpaces::Type
 		const USkinnedMeshComponent* const LeaderPoseComponentInst = LeaderPoseComponent.Get();
 		if(LeaderPoseComponentInst)
 		{
-			if(BoneIndex < LeaderBoneMap.Num())
+			if(BoneIndex < GetLeaderBoneMap().Num())
 			{
-				int32 ParentBoneIndex = LeaderBoneMap[BoneIndex];
+				int32 ParentBoneIndex = GetLeaderBoneMap()[BoneIndex];
 				// If ParentBoneIndex is valid, grab matrix from LeaderPoseComponent.
 				if(	ParentBoneIndex != INDEX_NONE && 
 					ParentBoneIndex < LeaderPoseComponentInst->GetNumComponentSpaceTransforms())
@@ -3666,9 +3675,9 @@ FVector USkinnedMeshComponent::GetBoneLocation(FName BoneName, EBoneSpaces::Type
 		const USkinnedMeshComponent* const LeaderPoseComponentInst = LeaderPoseComponent.Get();
 		if(LeaderPoseComponentInst)
 		{
-			if(BoneIndex < LeaderBoneMap.Num())
+			if(BoneIndex < GetLeaderBoneMap().Num())
 			{
-				int32 ParentBoneIndex = LeaderBoneMap[BoneIndex];
+				int32 ParentBoneIndex = GetLeaderBoneMap()[BoneIndex];
 				// If ParentBoneIndex is valid, grab transform from LeaderPoseComponent.
 				if(	ParentBoneIndex != INDEX_NONE && 
 					ParentBoneIndex < LeaderPoseComponentInst->GetNumComponentSpaceTransforms())
diff --git a/Engine/Source/Runtime/Engine/Private/SkeletalRender.cpp b/Engine/Source/Runtime/Engine/Private/SkeletalRender.cpp
--- a/Engine/Source/Runtime/Engine/Private/SkeletalRender.cpp
+++ b/Engine/Source/Runtime/Engine/Private/SkeletalRender.cpp
@@ -428,7 +428,8 @@ void UpdateRefToLocalMatricesInner(
 	const TArray<FBoneIndexType>* RequiredBoneSets[3] = { &LOD.ActiveBoneIndices, ExtraRequiredBoneIndices, NULL };
 
 	const bool bBoneVisibilityStatesValid = BoneVisibilityStates.Num() == ComponentTransform.Num();
-	
+	const bool bCheckForValidMeshPoseOnLeader = UE::Anim::bUseSafeMeshPoseIndices && bIsLeaderCompValid && !DynamicData.GetLeaderValidMeshPoseTransforms()->IsEmpty();
+
 	for (int32 RequiredBoneSetIndex = 0; RequiredBoneSets[RequiredBoneSetIndex] != NULL; RequiredBoneSetIndex++)
 	{
 		const TArray<FBoneIndexType>& RequiredBoneIndices = *RequiredBoneSets[RequiredBoneSetIndex];
@@ -449,7 +450,9 @@ void UpdateRefToLocalMatricesInner(
 			{
 				// If valid, use matrix from parent component.
 				const int32 LeaderBoneIndex = (*LeaderBoneMap)[ThisBoneIndex];
-				if (ComponentTransform.IsValidIndex(LeaderBoneIndex))
+				const TBitArray<>& LeaderValidMeshPoseTransforms = *(DynamicData.GetLeaderValidMeshPoseTransforms());
+
+				if (ComponentTransform.IsValidIndex(LeaderBoneIndex) && (!bCheckForValidMeshPoseOnLeader || LeaderValidMeshPoseTransforms[LeaderBoneIndex]))
 				{
 					const int32 ParentIndex = RefSkeleton.GetParentIndex(ThisBoneIndex);
 					bool bNeedToHideBone = BoneVisibilityStates[LeaderBoneIndex] != BVS_Visible;
diff --git a/Engine/Source/Runtime/Engine/Private/SkinnedMeshSceneProxyDesc.cpp b/Engine/Source/Runtime/Engine/Private/SkinnedMeshSceneProxyDesc.cpp
--- a/Engine/Source/Runtime/Engine/Private/SkinnedMeshSceneProxyDesc.cpp
+++ b/Engine/Source/Runtime/Engine/Private/SkinnedMeshSceneProxyDesc.cpp
@@ -454,6 +454,7 @@ FSkinnedMeshSceneProxyDynamicData::FSkinnedMeshSceneProxyDynamicData(const USkin
 	, BoneVisibilityStates(InLeaderPoseComponent != nullptr ? InLeaderPoseComponent->GetBoneVisibilityStates() : InSkinnedMeshComponent->GetBoneVisibilityStates())
 	, PreviousBoneVisibilityStates(InLeaderPoseComponent != nullptr ? InLeaderPoseComponent->GetPreviousBoneVisibilityStates() : InSkinnedMeshComponent->GetPreviousBoneVisibilityStates()) 
 	, LeaderBoneMap(InSkinnedMeshComponent->GetLeaderBoneMap())
+	, LeaderValidMeshPoseTransforms(InLeaderPoseComponent != nullptr ? &InLeaderPoseComponent->GetValidMeshPoseTransforms() : nullptr)
 	, SkinCacheUsage(InSkinnedMeshComponent->SkinCacheUsage)
 	, ComponentWorldTransform(InSkinnedMeshComponent->GetComponentTransform())
 	, CurrentBoneTransformRevisionNumber(InSkinnedMeshComponent->GetBoneTransformRevisionNumber())
diff --git a/Engine/Source/Runtime/Engine/Public/SkinnedMeshSceneProxyDesc.h b/Engine/Source/Runtime/Engine/Public/SkinnedMeshSceneProxyDesc.h
--- a/Engine/Source/Runtime/Engine/Public/SkinnedMeshSceneProxyDesc.h
+++ b/Engine/Source/Runtime/Engine/Public/SkinnedMeshSceneProxyDesc.h
@@ -144,6 +144,7 @@ struct FSkinnedMeshSceneProxyDynamicData
 	FTransform GetComponentTransform() const { return ComponentWorldTransform; }
 	const TSharedPtr<FSkelMeshRefPoseOverride>& GetRefPoseOverride() const { return RefPoseOverride; }
 	const TArrayView<const int32>& GetLeaderBoneMap() const { return LeaderBoneMap; }
+	const TBitArray<>* GetLeaderValidMeshPoseTransforms() const { return LeaderValidMeshPoseTransforms; }
 
 	uint32 GetBoneTransformRevisionNumber() const { return CurrentBoneTransformRevisionNumber; }
 	uint32 GetPreviousBoneTransformRevisionNumber() const { return PreviousBoneTransformRevisionNumber; }
@@ -164,6 +165,7 @@ struct FSkinnedMeshSceneProxyDynamicData
 	TArrayView<const uint8> PreviousBoneVisibilityStates{};
 
 	TArrayView<const int32> LeaderBoneMap{};
+	const TBitArray<>* LeaderValidMeshPoseTransforms = nullptr;
 	TArrayView<const ESkinCacheUsage> SkinCacheUsage{};
 
 	TArrayView<const FMatrix44f> ReferenceToLocalMatrices;
