Main - UE.EditorAutomation(RunTest=Audio) Linux_ASan vulkan - InitBadChunkStereoLoudnessAnalysis failed AddressSanitizer: heap-buffer-overflow on address: UE-362816

#Synced-CL 50260267

#Synced-CL 50260267

diff --git a/Engine/Source/Runtime/Engine/Private/Audio.cpp b/Engine/Source/Runtime/Engine/Private/Audio.cpp
--- a/Engine/Source/Runtime/Engine/Private/Audio.cpp
+++ b/Engine/Source/Runtime/Engine/Private/Audio.cpp
@@ -1909,7 +1909,6 @@ bool FWaveModInfo::ReadWaveInfo( const uint8* WaveData, int32 WaveDataSize, FStr
 
 				SampleDataEnd = SampleDataEnd - ExtraBytes;
 				SampleDataSize = SampleDataEnd - SampleDataStart;
-				RiffChunk->ChunkLen = INTEL_ORDER32(SampleDataSize);
 			}
 		}
 
