Update NavAgentSelector and NavAgentSelectorCustomization to support 31 Agents.
#tests PIE

#Synced-CL 50348711

diff --git a/Engine/Source/Editor/DetailCustomizations/Private/NavAgentSelectorCustomization.cpp b/Engine/Source/Editor/DetailCustomizations/Private/NavAgentSelectorCustomization.cpp
--- a/Engine/Source/Editor/DetailCustomizations/Private/NavAgentSelectorCustomization.cpp
+++ b/Engine/Source/Editor/DetailCustomizations/Private/NavAgentSelectorCustomization.cpp
@@ -93,7 +93,7 @@ void FNavAgentSelectorCustomization::CustomizeChildren(TSharedRef<class IPropert
 		return;
 	}
 
-	const int32 NumAgents = FMath::Min(NavSysCDO->GetSupportedAgents().Num(), 16);
+	const int32 NumAgents = FMath::Min(NavSysCDO->GetSupportedAgents().Num(), FNavAgentSelector::AgentsMaxCount);
 
 	for (uint32 Idx = 0; Idx < NumChildren; Idx++)
 	{
@@ -231,7 +231,7 @@ bool FNavAgentSelectorCustomization::ComputeSupportedAgentCount(const UNavigatio
 		return false;
 	}
 
-	const int32 NumAgents = FMath::Min(NavSysCDO->GetSupportedAgents().Num(), 16);
+	const int32 NumAgents = FMath::Min(NavSysCDO->GetSupportedAgents().Num(), FNavAgentSelector::AgentsMaxCount);
 
 	uint32 NumChildren = 0;
 	StructHandle->GetNumChildren(NumChildren);
diff --git a/Engine/Source/Runtime/Engine/Classes/AI/Navigation/NavAgentSelector.h b/Engine/Source/Runtime/Engine/Classes/AI/Navigation/NavAgentSelector.h
--- a/Engine/Source/Runtime/Engine/Classes/AI/Navigation/NavAgentSelector.h
+++ b/Engine/Source/Runtime/Engine/Classes/AI/Navigation/NavAgentSelector.h
@@ -11,6 +11,7 @@ struct FNavAgentSelector
 
 	static const uint32 InitializedBit = 0x80000000;
 	static const uint32 AllAgentsMask = 0x7fffffff;
+	static const int32 AgentsMaxCount = 31; // Must be less than or equal to the maximum number of agents defined by AllAgentsMask.
 
 #if CPP
 	union
@@ -50,6 +51,36 @@ struct FNavAgentSelector
 			uint32 bSupportsAgent14 : 1;
 			UPROPERTY(EditAnywhere, Category = Default)
 			uint32 bSupportsAgent15 : 1;
+			UPROPERTY(EditAnywhere, Category = Default)
+			uint32 bSupportsAgent16 : 1;
+			UPROPERTY(EditAnywhere, Category = Default)
+			uint32 bSupportsAgent17 : 1;
+			UPROPERTY(EditAnywhere, Category = Default)
+			uint32 bSupportsAgent18 : 1;
+			UPROPERTY(EditAnywhere, Category = Default)
+			uint32 bSupportsAgent19 : 1;
+			UPROPERTY(EditAnywhere, Category = Default)
+			uint32 bSupportsAgent20 : 1;
+			UPROPERTY(EditAnywhere, Category = Default)
+			uint32 bSupportsAgent21 : 1;
+			UPROPERTY(EditAnywhere, Category = Default)
+			uint32 bSupportsAgent22 : 1;
+			UPROPERTY(EditAnywhere, Category = Default)
+			uint32 bSupportsAgent23 : 1;
+			UPROPERTY(EditAnywhere, Category = Default)
+			uint32 bSupportsAgent24 : 1;
+			UPROPERTY(EditAnywhere, Category = Default)
+			uint32 bSupportsAgent25 : 1;
+			UPROPERTY(EditAnywhere, Category = Default)
+			uint32 bSupportsAgent26 : 1;
+			UPROPERTY(EditAnywhere, Category = Default)
+			uint32 bSupportsAgent27 : 1;
+			UPROPERTY(EditAnywhere, Category = Default)
+			uint32 bSupportsAgent28 : 1;
+			UPROPERTY(EditAnywhere, Category = Default)
+			uint32 bSupportsAgent29 : 1;
+			UPROPERTY(EditAnywhere, Category = Default)
+			uint32 bSupportsAgent30 : 1;
 #if CPP
 		};
 		uint32 PackedBits;
@@ -60,7 +91,7 @@ struct FNavAgentSelector
 
 	inline bool Contains(int32 AgentIndex) const
 	{
-		return (AgentIndex >= 0 && AgentIndex < 16) ? !!(PackedBits & (1 << AgentIndex)) : false;
+		return (AgentIndex >= 0 && AgentIndex < AgentsMaxCount) ? !!(PackedBits & (1 << AgentIndex)) : false;
 	}
 
 	inline bool ContainsAnyAgent() const
@@ -71,7 +102,7 @@ struct FNavAgentSelector
 
 	inline void Set(int32 AgentIndex)
 	{
-		if (AgentIndex >= 0 && AgentIndex < 16)
+		if (AgentIndex >= 0 && AgentIndex < AgentsMaxCount)
 		{
 			PackedBits |= (1 << AgentIndex);
 		}
