Add operator== to spatial hash

Will be used for change detection/dirtying when spatial hashes applied to components.

#rnx

#Synced-CL 50297198

diff --git a/Engine/Source/Runtime/Engine/Public/InstanceDataSceneProxy.h b/Engine/Source/Runtime/Engine/Public/InstanceDataSceneProxy.h
--- a/Engine/Source/Runtime/Engine/Public/InstanceDataSceneProxy.h
+++ b/Engine/Source/Runtime/Engine/Public/InstanceDataSceneProxy.h
@@ -213,6 +213,8 @@ public:
 
 	struct FCompressedSpatialHashItem
 	{
+		bool operator==(const FCompressedSpatialHashItem& Other) const = default;
+
 		// TODO: compress the location - store delta to primitive loc for example
 		RenderingSpatialHash::FLocation64 Location;
 		int32 NumInstances;
