Fix deadlock when modifying an AnimSequence which has itself set as the RefPoseSeq

#Synced-CL 50348084

diff --git a/Engine/Source/Runtime/Engine/Private/Animation/AnimSequence.cpp b/Engine/Source/Runtime/Engine/Private/Animation/AnimSequence.cpp
--- a/Engine/Source/Runtime/Engine/Private/Animation/AnimSequence.cpp
+++ b/Engine/Source/Runtime/Engine/Private/Animation/AnimSequence.cpp
@@ -2538,7 +2538,8 @@ void UAnimSequence::FlagDependentAnimationsAsRawDataOnly() const
 	for (TObjectIterator<UAnimSequence> Iter; Iter; ++Iter)
 	{
 		UAnimSequence* Seq = *Iter;
-		if (Seq->HasAllFlags(RF_LoadCompleted) && Seq->RefPoseSeq == this)
+		/** Ignore self, as data will already be recompressed at higher-level callsite */
+		if (Seq->HasAllFlags(RF_LoadCompleted) && Seq != this && Seq->RefPoseSeq == this)
 		{
 			Seq->ClearAllCompressionData();
 			Seq->BeginCacheDerivedDataForCurrentPlatform();
