Ensure volumetric fog and volume rendering do not have shader VS export to GS compilation flag on platforms that do not support GS
Enable OneColorShader and CopyTextureShaders for ES31 renderer
Changed Point argument name to StepIndex in PostProcessMobile::Circle (Point is a reserved shader keyword on some platforms)

#rnx

#Synced-CL 50152864

diff --git a/Engine/Shaders/Private/OneColorShader.usf b/Engine/Shaders/Private/OneColorShader.usf
--- a/Engine/Shaders/Private/OneColorShader.usf
+++ b/Engine/Shaders/Private/OneColorShader.usf
@@ -140,8 +140,6 @@ void MainLongGPUTask(
 	OutColor *= .000001f;
 }
 
-#if FEATURE_LEVEL >= FEATURE_LEVEL_SM5
-
 float4 FillValue;
 RWTexture2D<float4> FillTexture;
 float4 Params0; // Texture Width,Height (.xy); Use Exclude Rect 1 : 0 (.z)
@@ -188,5 +186,3 @@ void MainFillTextureCS(uint2 XY : SV_DispatchThreadID)
 		}
 	}
 }
-
-#endif
diff --git a/Engine/Shaders/Private/PostProcessMobile.usf b/Engine/Shaders/Private/PostProcessMobile.usf
--- a/Engine/Shaders/Private/PostProcessMobile.usf
+++ b/Engine/Shaders/Private/PostProcessMobile.usf
@@ -67,9 +67,9 @@ float4 BufferBSizeAndInvSize;
 float4 BufferViewRect;
 
 // Point on circle.
-float2 Circle(float Start, float Points, float Point) 
+float2 Circle(float Start, float Points, float StepIndex) 
 {
-	float Rad = (3.141592 * 2.0 * (1.0 / Points)) * (Point + Start);
+	float Rad = (3.141592 * 2.0 * (1.0 / Points)) * (StepIndex + Start);
 	return float2(sin(Rad), cos(Rad));
 }
 
diff --git a/Engine/Source/Runtime/Engine/Private/VolumeRendering.cpp b/Engine/Source/Runtime/Engine/Private/VolumeRendering.cpp
--- a/Engine/Source/Runtime/Engine/Private/VolumeRendering.cpp
+++ b/Engine/Source/Runtime/Engine/Private/VolumeRendering.cpp
@@ -26,7 +26,10 @@ bool FWriteToSliceVS::ShouldCompilePermutation(const FGlobalShaderPermutationPar
 void FWriteToSliceVS::ModifyCompilationEnvironment(const FGlobalShaderPermutationParameters& Parameters, FShaderCompilerEnvironment& OutEnvironment)
 {
 	FGlobalShader::ModifyCompilationEnvironment(Parameters, OutEnvironment);
-	OutEnvironment.CompilerFlags.Add(CFLAG_VertexToGeometryShader);
+	if (RHISupportsGeometryShaders(Parameters.Platform))
+	{
+		OutEnvironment.CompilerFlags.Add(CFLAG_VertexToGeometryShader);
+	}
 }
 
 IMPLEMENT_SHADER_TYPE(,FWriteToSliceVS,TEXT("/Engine/Private/TranslucentLightingShaders.usf"),TEXT("WriteToSliceMainVS"),SF_Vertex);
diff --git a/Engine/Source/Runtime/RenderCore/Private/OneColorShader.cpp b/Engine/Source/Runtime/RenderCore/Private/OneColorShader.cpp
--- a/Engine/Source/Runtime/RenderCore/Private/OneColorShader.cpp
+++ b/Engine/Source/Runtime/RenderCore/Private/OneColorShader.cpp
@@ -94,7 +94,7 @@ FFillTextureCS::FFillTextureCS(const ShaderMetaType::CompiledShaderInitializerTy
 
 bool FFillTextureCS::ShouldCompilePermutation(const FGlobalShaderPermutationParameters& Parameters)
 {
-	return IsFeatureLevelSupported(Parameters.Platform, ERHIFeatureLevel::SM5);
+	return true;
 }
 
 IMPLEMENT_SHADER_TYPE(,FFillTextureCS,TEXT("/Engine/Private/OneColorShader.usf"),TEXT("MainFillTextureCS"),SF_Compute);
diff --git a/Engine/Source/Runtime/RenderCore/Public/CopyTextureShaders.h b/Engine/Source/Runtime/RenderCore/Public/CopyTextureShaders.h
--- a/Engine/Source/Runtime/RenderCore/Public/CopyTextureShaders.h
+++ b/Engine/Source/Runtime/RenderCore/Public/CopyTextureShaders.h
@@ -146,7 +146,7 @@ public:
 
 	static bool ShouldCompilePermutation(const FGlobalShaderPermutationParameters& Parameters)
 	{
-		return IsFeatureLevelSupported(Parameters.Platform, ERHIFeatureLevel::SM5);
+		return true;
 	}
 
 	static inline void ModifyCompilationEnvironment(const FGlobalShaderPermutationParameters& Parameters, FShaderCompilerEnvironment& OutEnvironment)
diff --git a/Engine/Source/Runtime/Renderer/Private/VolumetricFog.cpp b/Engine/Source/Runtime/Renderer/Private/VolumetricFog.cpp
--- a/Engine/Source/Runtime/Renderer/Private/VolumetricFog.cpp
+++ b/Engine/Source/Runtime/Renderer/Private/VolumetricFog.cpp
@@ -320,7 +320,10 @@ class FWriteToBoundingSphereVS : public FGlobalShader
 	static void ModifyCompilationEnvironment(const FGlobalShaderPermutationParameters& Parameters, FShaderCompilerEnvironment& OutEnvironment)
 	{
 		FGlobalShader::ModifyCompilationEnvironment(Parameters, OutEnvironment);
-		OutEnvironment.CompilerFlags.Add(CFLAG_VertexToGeometryShader);
+		if (RHISupportsGeometryShaders(Parameters.Platform))
+		{
+			OutEnvironment.CompilerFlags.Add(CFLAG_VertexToGeometryShader);
+		}
 	}
 };
 
