Add OnHDRScreenshotCaptured delegate to UGameViewportClient, fired when an HDR screenshot is captured.
OnScreenshotCaptured does not fire when HDR is enabled. A separate delegate is needed to allow returning the HDR bitmap as TArray<FLinearColor> rather than TArray<FColor>.

#Synced-CL 50352938

diff --git a/Engine/Source/Runtime/Engine/Classes/Engine/GameViewportClient.h b/Engine/Source/Runtime/Engine/Classes/Engine/GameViewportClient.h
--- a/Engine/Source/Runtime/Engine/Classes/Engine/GameViewportClient.h
+++ b/Engine/Source/Runtime/Engine/Classes/Engine/GameViewportClient.h
@@ -523,11 +523,17 @@ public:
 	/** Make sure all navigation objects have appropriate path rendering components set.  Called when EngineShowFlags.Navigation is set. */
 	ENGINE_API virtual void VerifyPathRenderingComponents();
 
-	/** Accessor for delegate called when a screenshot is captured */
+	/** Accessor for delegate called when a non-HDR screenshot is captured */
 	static FOnScreenshotCaptured& OnScreenshotCaptured()
 	{
 		return ScreenshotCapturedDelegate;
 	}
+
+	/** Accessor for delegate called when an HDR screenshot is captured */
+	static FOnHDRScreenshotCaptured& OnHDRScreenshotCaptured()
+	{
+		return HDRScreenshotCapturedDelegate;
+	}
 	
 	/** Accessor for delegate called when a viewport is rendered */
 	static FOnViewportRendered& OnViewportRendered()
@@ -1036,8 +1042,11 @@ private:
 	FOnGameViewportInputKey GameViewportInputKeyDelegate;
 #endif
 
-	/** Delegate called at the end of the frame when a screenshot is captured */
+	/** Delegate called at the end of the frame when a non-HDR screenshot is captured */
 	static ENGINE_API FOnScreenshotCaptured ScreenshotCapturedDelegate;
+
+	/** Delegate called at the end of the frame when an HDR screenshot is captured */
+	static ENGINE_API FOnHDRScreenshotCaptured HDRScreenshotCapturedDelegate;
 	
 	/** Delegate called right after the viewport is rendered */
 	static ENGINE_API FOnViewportRendered ViewportRenderedDelegate;
diff --git a/Engine/Source/Runtime/Engine/Classes/Engine/GameViewportDelegates.h b/Engine/Source/Runtime/Engine/Classes/Engine/GameViewportDelegates.h
--- a/Engine/Source/Runtime/Engine/Classes/Engine/GameViewportDelegates.h
+++ b/Engine/Source/Runtime/Engine/Classes/Engine/GameViewportDelegates.h
@@ -10,15 +10,25 @@
 class FViewport;
 
 /**
- * Delegate type used by UGameViewportClient when a screenshot has been captured
+ * Delegate type used by UGameViewportClient when a non-HDR screenshot has been captured
  *
  * The first parameter is the width.
  * The second parameter is the height.
- * The third parameter is the array of bitmap data.
+ * The third parameter is the array of bitmap data (FColor).
  * @see UGameViewportClient
  */
 DECLARE_MULTICAST_DELEGATE_ThreeParams(FOnScreenshotCaptured, int32 /*Width*/, int32 /*Height*/, const TArray<FColor>& /*Colors*/);
 
+/**
+ * Delegate type used by UGameViewportClient when an HDR screenshot has been captured
+ *
+ * The first parameter is the width.
+ * The second parameter is the height.
+ * The third parameter is the array of bitmap data (FLinearColor).
+ * @see UGameViewportClient
+ */
+DECLARE_MULTICAST_DELEGATE_ThreeParams(FOnHDRScreenshotCaptured, int32 /*Width*/, int32 /*Height*/, const TArray<FLinearColor>& /*Colors*/);
+
 /**
  * Delegate type used by UGameViewportClient when a viewport is rendered
  *
diff --git a/Engine/Source/Runtime/Engine/Private/GameViewportClient.cpp b/Engine/Source/Runtime/Engine/Private/GameViewportClient.cpp
--- a/Engine/Source/Runtime/Engine/Private/GameViewportClient.cpp
+++ b/Engine/Source/Runtime/Engine/Private/GameViewportClient.cpp
@@ -78,9 +78,12 @@
 /** This variable allows forcing full screen of the first player controller viewport, even if there are multiple controllers plugged in and no cinematic playing. */
 bool GForceFullscreen = false;
 
-/** Delegate called at the end of the frame when a screenshot is captured */
+/** Delegate called at the end of the frame when a non-HDR screenshot is captured */
 FOnScreenshotCaptured UGameViewportClient::ScreenshotCapturedDelegate;
 
+/** Delegate called at the end of the frame when an HDR screenshot is captured */
+FOnHDRScreenshotCaptured UGameViewportClient::HDRScreenshotCapturedDelegate;
+
 /** Delegate called right after the viewport is rendered */
 FOnViewportRendered UGameViewportClient::ViewportRenderedDelegate;
 
@@ -2398,7 +2401,7 @@ bool UGameViewportClient::ProcessScreenShots(FViewport* InViewport)
 				if (ScreenshotCapturedDelegate.IsBound() && CVarScreenshotDelegate.GetValueOnGameThread())
 				{
 					// Ensure that all pixels' alpha is set to 255
-					for (auto& Color : Bitmap)
+					for (FColor& Color : Bitmap)
 					{
 						Color.A = 255;
 					}
@@ -2413,7 +2416,21 @@ bool UGameViewportClient::ProcessScreenShots(FViewport* InViewport)
 			}
 			else
 			{
-				bIsScreenshotSaved = ProcessScreenshotData(BitmapHDR, Size, 1.0f, bHdrEnabled, bIsUI, TEXT(".exr"));
+				if (HDRScreenshotCapturedDelegate.IsBound() && CVarScreenshotDelegate.GetValueOnGameThread())
+				{
+					// Ensure that all pixels' alpha is set to 1.0f
+					for (FLinearColor& LinearColor : BitmapHDR)
+					{
+						LinearColor.A = 1.0f;
+					}
+
+					// If delegate subscribed, fire it instead of writing out a file to disk
+					HDRScreenshotCapturedDelegate.Broadcast(Size.X, Size.Y, BitmapHDR);
+				}
+				else
+				{
+					bIsScreenshotSaved = ProcessScreenshotData(BitmapHDR, Size, 1.0f, bHdrEnabled, bIsUI, TEXT(".exr"));
+				}
 			}
 		}
 
diff --git a/Engine/Source/Runtime/Engine/Private/UnrealClient.cpp b/Engine/Source/Runtime/Engine/Private/UnrealClient.cpp
--- a/Engine/Source/Runtime/Engine/Private/UnrealClient.cpp
+++ b/Engine/Source/Runtime/Engine/Private/UnrealClient.cpp
@@ -339,6 +339,7 @@ FString GetMemoryString( const double Value, const bool bAutoType )
 
 FOnScreenshotRequestProcessed FScreenshotRequest::ScreenshotProcessedDelegate;
 FOnScreenshotCaptured FScreenshotRequest::ScreenshotCapturedDelegate;
+FOnHDRScreenshotCaptured FScreenshotRequest::HDRScreenshotCapturedDelegate;
 bool FScreenshotRequest::bIsScreenshotRequested = false;
 FString FScreenshotRequest::Filename;
 FString FScreenshotRequest::NextScreenshotName;
diff --git a/Engine/Source/Runtime/Engine/Public/UnrealClient.h b/Engine/Source/Runtime/Engine/Public/UnrealClient.h
--- a/Engine/Source/Runtime/Engine/Public/UnrealClient.h
+++ b/Engine/Source/Runtime/Engine/Public/UnrealClient.h
@@ -197,6 +197,7 @@ public:
 
 DECLARE_MULTICAST_DELEGATE(FOnScreenshotRequestProcessed);
 DECLARE_MULTICAST_DELEGATE_ThreeParams(FOnScreenshotCaptured, int32 /*Width*/, int32 /*Height*/, const TArray<FColor>& /*Colors*/);
+DECLARE_MULTICAST_DELEGATE_ThreeParams(FOnHDRScreenshotCaptured, int32 /*Width*/, int32 /*Height*/, const TArray<FLinearColor>& /*Colors*/);
 
 
 struct FScreenshotRequest
@@ -270,9 +271,15 @@ struct FScreenshotRequest
 		return ScreenshotCapturedDelegate;
 	}
 
+	static FOnHDRScreenshotCaptured& OnHDRScreenshotCaptured()
+	{
+		return HDRScreenshotCapturedDelegate;
+	}
+
 private:
 	static ENGINE_API FOnScreenshotRequestProcessed ScreenshotProcessedDelegate;
 	static ENGINE_API FOnScreenshotCaptured ScreenshotCapturedDelegate;
+	static ENGINE_API FOnHDRScreenshotCaptured HDRScreenshotCapturedDelegate;
 	static ENGINE_API bool bIsScreenshotRequested;
 	static ENGINE_API FString NextScreenshotName;
 	static ENGINE_API FString Filename;
