[StreamableManager/ALT]
- Make sure to call RegisterGCObject/UnregisterGCObject in the most derived class to avoid pure virtual function call from other threads

#Synced-CL 50254652

diff --git a/Engine/Source/Runtime/Engine/Private/StreamableManager.cpp b/Engine/Source/Runtime/Engine/Private/StreamableManager.cpp
--- a/Engine/Source/Runtime/Engine/Private/StreamableManager.cpp
+++ b/Engine/Source/Runtime/Engine/Private/StreamableManager.cpp
@@ -1600,7 +1600,7 @@ void RemoveActiveHandle(FStreamable& Streamable, FStreamableHandle& Handle)
 	
 
 FStreamableManager::FStreamableManager()
-	: FGCObject(EFlags::AddStableNativeReferencesOnly)
+	: FGCObject(EFlags::AddStableNativeReferencesOnly | EFlags::RegisterLater)
 {
 #if UE_STREAMABLE_MANAGER_PROFILING_ENABLED
 	FStreamableManagerProfiling::TryInitialize();
@@ -1608,10 +1608,18 @@ FStreamableManager::FStreamableManager()
 	FCoreUObjectDelegates::GetPreGarbageCollectDelegate().AddRaw(this, &FStreamableManager::OnPreGarbageCollect);
 	bForceSynchronousLoads = false;
 	ManagerName = TEXT("StreamableManager");
+
+	// Needs to be called in the most derived constructor to avoid still
+	// having pure virtual calls in the virtual table at registration time.
+	RegisterGCObject();
 }
 
 FStreamableManager::~FStreamableManager()
 {
+	// Needs to be called in the most derived destructor to avoid leaving
+	// a registered object with half its virtual table destroyed.
+	UnregisterGCObject();
+
 	FCoreUObjectDelegates::GetPreGarbageCollectDelegate().RemoveAll(this);
 
 	for (const TPair<FSoftObjectPath, FStreamable*>& Pair : StreamableItems)
