World Partition - HLOD: Added option to force inclusion in the ray tracing far field
* Added bForceRayTracingFarField to UHLODLayer. When enabled, static mesh components generated for this HLOD layer will be included in the ray tracing far field, even if the HLOD is spatially loaded. This is useful for cases where the last HLOD level is spatially loaded but should still contribute to ray traced reflections and global illumination.

#Synced-CL 50353017

diff --git a/Engine/Plugins/Editor/WorldPartitionHLODUtilities/Source/Private/WorldPartition/HLOD/Utilities/WorldPartitionHLODUtilities.cpp b/Engine/Plugins/Editor/WorldPartitionHLODUtilities/Source/Private/WorldPartition/HLOD/Utilities/WorldPartitionHLODUtilities.cpp
--- a/Engine/Plugins/Editor/WorldPartitionHLODUtilities/Source/Private/WorldPartition/HLOD/Utilities/WorldPartitionHLODUtilities.cpp
+++ b/Engine/Plugins/Editor/WorldPartitionHLODUtilities/Source/Private/WorldPartition/HLOD/Utilities/WorldPartitionHLODUtilities.cpp
@@ -1006,11 +1006,12 @@ PRAGMA_ENABLE_DEPRECATION_WARNINGS
 				{
 					if (UStaticMesh* StaticMesh = StaticMeshComponent->GetStaticMesh())
 					{
-						// If the HLOD process did create this static mesh
-						if (StaticMesh->GetPackage() == HLODBuildContext.AssetsOuter)
+						// Set up ray tracing far fields for always loaded HLODs (only for HLOD-created meshes)
+						// or if requested by the HLOD layer (for any mesh that supports ray tracing)
+						if (StaticMesh->bSupportRayTracing)
 						{
-							// Set up ray tracing far fields for always loaded HLODs
-							if (!HLODActor->GetIsSpatiallyLoaded() && StaticMesh->bSupportRayTracing)
+							const bool bIsHLODCreatedMesh = StaticMesh->GetPackage() == HLODBuildContext.AssetsOuter;
+							if ((bIsHLODCreatedMesh && !HLODActor->GetIsSpatiallyLoaded()) || HLODLayer->GetForceRayTracingFarField())
 							{
 								StaticMeshComponent->bRayTracingFarField = true;
 							}
diff --git a/Engine/Source/Runtime/Engine/Private/WorldPartition/HLOD/HLODLayer.cpp b/Engine/Source/Runtime/Engine/Private/WorldPartition/HLOD/HLODLayer.cpp
--- a/Engine/Source/Runtime/Engine/Private/WorldPartition/HLOD/HLODLayer.cpp
+++ b/Engine/Source/Runtime/Engine/Private/WorldPartition/HLOD/HLODLayer.cpp
@@ -262,6 +262,11 @@ void UHLODLayer::ComputeHLODHash(FHLODHashBuilder& InHashBuilder) const
 	{
 		InHashBuilder.HashField(HLODModifierClass->GetPathName(), GET_MEMBER_NAME_CHECKED(UHLODLayer, HLODModifierClass));
 	}
+
+	if (bForceRayTracingFarField)
+	{
+		InHashBuilder.HashField(bForceRayTracingFarField, GET_MEMBER_NAME_CHECKED(UHLODLayer, bForceRayTracingFarField));
+	}
 }
 
 bool UHLODLayer::MustBeLoadedOnServer() const
diff --git a/Engine/Source/Runtime/Engine/Public/WorldPartition/HLOD/HLODLayer.h b/Engine/Source/Runtime/Engine/Public/WorldPartition/HLOD/HLODLayer.h
--- a/Engine/Source/Runtime/Engine/Public/WorldPartition/HLOD/HLODLayer.h
+++ b/Engine/Source/Runtime/Engine/Public/WorldPartition/HLOD/HLODLayer.h
@@ -57,6 +57,7 @@ public:
 	const UHLODBuilderSettings* GetHLODBuilderSettings() const { return HLODBuilderSettings; }
 	const TSubclassOf<AWorldPartitionHLOD> GetHLODActorClass() const { return HLODActorClass; }
 	const TSubclassOf<UWorldPartitionHLODModifier> GetHLODModifierClass() const { return HLODModifierClass; }
+	bool GetForceRayTracingFarField() const { return bForceRayTracingFarField; }
 	UHLODLayer* GetParentLayer() const { return ParentLayer; }
 	void SetParentLayer(UHLODLayer* InParentLayer) { ParentLayer = InParentLayer; }
 	UHLODLayer* GetLinkedLayer() const { return LinkedLayer; }
@@ -133,6 +134,14 @@ private:
 	UPROPERTY(EditAnywhere, Category = General, AdvancedDisplay, meta = (DisplayName = "HLOD Modifier Class"))
 	TSubclassOf<UWorldPartitionHLODModifier> HLODModifierClass;
 
+	/**
+	 * When enabled, static mesh components generated for this HLOD layer will be included in the ray tracing far field,
+	 * even if the HLOD is spatially loaded. This is useful for cases where the last HLOD level is spatially
+	 * loaded but should still contribute to ray traced reflections and global illumination.
+	 */
+	UPROPERTY(EditAnywhere, Category = General, AdvancedDisplay)
+	uint32 bForceRayTracingFarField : 1;
+
 	friend class FWorldPartitionHLODLayerDetailsCustomization;
 
 private:
