- Remove unused TempBuffer

#preflght
#rnx

#Synced-CL 50255171

diff --git a/Engine/Source/Runtime/Engine/Private/VT/VirtualTextureTranscodeCache.cpp b/Engine/Source/Runtime/Engine/Private/VT/VirtualTextureTranscodeCache.cpp
--- a/Engine/Source/Runtime/Engine/Private/VT/VirtualTextureTranscodeCache.cpp
+++ b/Engine/Source/Runtime/Engine/Private/VT/VirtualTextureTranscodeCache.cpp
@@ -424,8 +424,6 @@ FVTTranscodeTileHandle FVirtualTextureTranscodeCache::SubmitTask(
 	FVTUploadTileBuffer StagingBuffer[VIRTUALTEXTURE_SPACE_MAXLAYERS];
 	bool bNeedToLaunchTask = false;
 
-	TArray<uint8, TInlineAllocator<16u * 1024>> TempBuffer;
-
 	uint32 TileBaseOffset = ~0u;
 	for (uint32 LayerIndex = 0u; LayerIndex < InParams.VTData->GetNumLayers(); ++LayerIndex)
 	{
