Subtitles' bCanEditDuration should now respect DurationType on load.

#Synced-CL 50321610

diff --git a/Engine/Source/Runtime/Engine/Private/Subtitles/SubtitlesAndClosedCaptionsTypes.cpp b/Engine/Source/Runtime/Engine/Private/Subtitles/SubtitlesAndClosedCaptionsTypes.cpp
--- a/Engine/Source/Runtime/Engine/Private/Subtitles/SubtitlesAndClosedCaptionsTypes.cpp
+++ b/Engine/Source/Runtime/Engine/Private/Subtitles/SubtitlesAndClosedCaptionsTypes.cpp
@@ -11,5 +11,19 @@ void USubtitleAssetUserData::PostEditChangeProperty(FPropertyChangedEvent& Prope
 	{
 		Subtitle.bCanEditDuration = (Subtitle.SubtitleDurationType == ESubtitleDurationType::UseDurationProperty);
 	}
+
+	// Calling Super:: afterwards here so that bCanEditDuration is part of the change that gets put in the Transaction snapshot.
+	Super::PostEditChangeProperty(PropertyChangedEvent);
+}
+
+void USubtitleAssetUserData::PostInitProperties()
+{
+	// Calling Super:: first: ensures that everything is fully initialized and settled before fiddling with Transient values.
+	Super::PostInitProperties();
+
+	for (FSubtitleAssetData& Subtitle : Subtitles)
+	{
+		Subtitle.bCanEditDuration = (Subtitle.SubtitleDurationType == ESubtitleDurationType::UseDurationProperty);
+	}
 }
 #endif // WITH_EDITOR
\ No newline at end of file
diff --git a/Engine/Source/Runtime/Engine/Public/Subtitles/SubtitlesAndClosedCaptionsTypes.h b/Engine/Source/Runtime/Engine/Public/Subtitles/SubtitlesAndClosedCaptionsTypes.h
--- a/Engine/Source/Runtime/Engine/Public/Subtitles/SubtitlesAndClosedCaptionsTypes.h
+++ b/Engine/Source/Runtime/Engine/Public/Subtitles/SubtitlesAndClosedCaptionsTypes.h
@@ -128,8 +128,10 @@ public:
 		return Duration;
 	};
 
+	// These functions ensure Duration values are only editable when DurationType == UseDurationProperty:
 #if WITH_EDITOR
 	virtual void PostEditChangeProperty(FPropertyChangedEvent& PropertyChangedEvent) override;
+	virtual void PostInitProperties() override;
 #endif // WITH_EDITOR
 };
 
