Kismet: Also save as hdr/exr when type is f32 (currently only when f16).

#rnx

#Synced-CL 49977182

diff --git a/Engine/Source/Runtime/Engine/Private/KismetRenderingLibrary.cpp b/Engine/Source/Runtime/Engine/Private/KismetRenderingLibrary.cpp
--- a/Engine/Source/Runtime/Engine/Private/KismetRenderingLibrary.cpp
+++ b/Engine/Source/Runtime/Engine/Private/KismetRenderingLibrary.cpp
@@ -260,7 +260,7 @@ void UKismetRenderingLibrary::ExportRenderTarget(UObject* WorldContextObject, UT
 			FBufferArchive Buffer;
 
 			bool bSuccess = false;
-			if (TextureRenderTarget->RenderTargetFormat == RTF_RGBA16f)
+			if ((TextureRenderTarget->RenderTargetFormat == RTF_RGBA16f) || (TextureRenderTarget->RenderTargetFormat == RTF_RGBA32f))
 			{
 				// Note == is case insensitive
 				if (FPaths::GetExtension(TotalFileName) == TEXT("HDR"))
