Removed DIRECT_MODE bone storage for GPU skin. It's too slow. Instead, GPU instance LOD activates when section bone maps are unified across LODs. This is currently just by chance, but a follow-up CL will add an opt-in to build unified bone maps for skeletal meshes.

#Synced-CL 50310951

diff --git a/Engine/Shaders/Private/GpuSkinCommon.ush b/Engine/Shaders/Private/GpuSkinCommon.ush
--- a/Engine/Shaders/Private/GpuSkinCommon.ush
+++ b/Engine/Shaders/Private/GpuSkinCommon.ush
@@ -222,18 +222,7 @@ uint GetSkinningBoneTransformOffset(uint TransformBufferOffset, uint InstanceOff
 
 FBoneMatrix LoadSkinningBoneTransform(FComputeBoneMatrixInputs Inputs, uint TransformIndex)
 {
-	uint TransformBufferIndex = Inputs.BoneTransformOffset;
-
-#if GPUSKIN_TRANSFORM_STORAGE_DIRECT_ENABLED
-	if (Inputs.TransformStorageMode == GPUSKIN_TRANSFORM_STORAGE_DIRECT)
-	{
-		TransformBufferIndex += Inputs.SectionBoneMap[Inputs.SectionBoneOffset + TransformIndex];
-	}
-	else
-#endif
-	{
-		TransformBufferIndex += Inputs.SectionBoneOffset + TransformIndex;
-	}
+	uint TransformBufferIndex = Inputs.BoneTransformOffset + Inputs.SectionBoneOffset + TransformIndex;
 
 	return LoadCompressedBoneTransformTransposed(Inputs.BoneMatrices, 0, TransformBufferIndex);
 }
diff --git a/Engine/Shaders/Shared/SkinningDefinitions.h b/Engine/Shaders/Shared/SkinningDefinitions.h
--- a/Engine/Shaders/Shared/SkinningDefinitions.h
+++ b/Engine/Shaders/Shared/SkinningDefinitions.h
@@ -7,9 +7,6 @@
 #define GPUSKIN_TRANSFORM_STORAGE_BONE_MAP 0
 #define GPUSKIN_TRANSFORM_STORAGE_DIRECT   1
 
-// This path is experimental and currently adds cost to the shader.
-#define GPUSKIN_TRANSFORM_STORAGE_DIRECT_ENABLED 0
-
 #ifdef __cplusplus
 #define SKINNING_STATIC_ASSERT(bCondition, Message) static_assert(bCondition, Message)
 #else
diff --git a/Engine/Source/Runtime/Engine/Private/InstancedSkeletalMeshObject.cpp b/Engine/Source/Runtime/Engine/Private/InstancedSkeletalMeshObject.cpp
--- a/Engine/Source/Runtime/Engine/Private/InstancedSkeletalMeshObject.cpp
+++ b/Engine/Source/Runtime/Engine/Private/InstancedSkeletalMeshObject.cpp
@@ -6,8 +6,6 @@
 #include "Rendering/SkeletalMeshRenderData.h"
 #include "Rendering/RenderCommandPipes.h"
 
-extern EBoneTransformStorageMode GetBoneTransformStorageMode();
-
 extern void InitGPUSkinVertexFactoryComponents(FGPUSkinDataType* VertexFactoryData, const FSkeletalMeshObjectGPUSkin::FVertexFactoryBuffers& VertexBuffers, FGPUBaseSkinVertexFactory* VertexFactory);
 
 FInstancedSkeletalMeshObject::FInstancedSkeletalMeshObject(
@@ -20,14 +18,15 @@ FInstancedSkeletalMeshObject::FInstancedSkeletalMeshObject(
 	, Type(InType)
 {
 	const int32 NumLODs = InRenderData->LODRenderData.Num();
+	const bool bIsGPUSkin = Type == EType::GPUSkin;
 
 	for (int32 LODIndex = 0; LODIndex < NumLODs; ++LODIndex)
 	{
 		LODs.Emplace(InRenderData, LODIndex, InFeatureLevel);
 	}
 
-	bUseGpuLodSelection = NumLODs > 1 && InMeshDesc.bUseGpuLodSelection && (GPUSKIN_TRANSFORM_STORAGE_DIRECT_ENABLED || InRenderData->bSharedLODBoneMap);
-	BoneTransformStorageMode = Type == EType::GPUSkin && !bUseGpuLodSelection ? ::GetBoneTransformStorageMode() : EBoneTransformStorageMode::Direct;
+	bUseGpuLodSelection = NumLODs > 1 && InMeshDesc.bUseGpuLodSelection && bIsGPUSkin && InRenderData->HasUnifiedBoneMap();
+	BoneTransformStorageMode = bIsGPUSkin ? EBoneTransformStorageMode::BoneMap : EBoneTransformStorageMode::Direct;
 
 	InitResources(InMeshDesc);
 
@@ -191,9 +190,6 @@ void FInstancedSkeletalMeshObject::FLOD::InitResources(const FSkelMeshComponentL
 	FSkeletalMeshLODRenderData& LODData = RenderData->LODRenderData[LODIndex];
 
 	FSkeletalMeshObjectGPUSkin::FVertexFactoryBuffers VertexBuffers;
-#if GPUSKIN_TRANSFORM_STORAGE_DIRECT_ENABLED
-	VertexBuffers.SectionBoneMapBuffer = InBoneTransformStorageMode == EBoneTransformStorageMode::Direct ? &LODData.SectionBoneMapBuffer : nullptr;
-#endif
 	VertexBuffers.StaticVertexBuffers = &LODData.StaticVertexBuffers;
 	VertexBuffers.ColorVertexBuffer = FSkeletalMeshObject::GetColorVertexBuffer(LODData, InLODInfo);
 	VertexBuffers.SkinWeightVertexBuffer = FSkeletalMeshObject::GetSkinWeightVertexBuffer(LODData, InLODInfo);
diff --git a/Engine/Source/Runtime/Engine/Private/InstancedSkinnedMeshSceneProxy.cpp b/Engine/Source/Runtime/Engine/Private/InstancedSkinnedMeshSceneProxy.cpp
--- a/Engine/Source/Runtime/Engine/Private/InstancedSkinnedMeshSceneProxy.cpp
+++ b/Engine/Source/Runtime/Engine/Private/InstancedSkinnedMeshSceneProxy.cpp
@@ -73,7 +73,7 @@ FInstancedSkinnedMeshSceneProxy::FInstancedSkinnedMeshSceneProxy(const FInstance
 
 	bAlwaysHasVelocity = true;
 	bInstancedSkinnedMesh = true;
-	bUseGpuLodSelection = NumLODs > 1 && InMeshDesc.bUseGpuLodSelection && (GPUSKIN_TRANSFORM_STORAGE_DIRECT_ENABLED && InRenderData->bSharedLODBoneMap);
+	bUseGpuLodSelection = NumLODs > 1 && InMeshDesc.bUseGpuLodSelection && InRenderData->HasUnifiedBoneMap();
 
 	SetupInstanceSceneDataBuffers(Data.InstanceDataSceneProxy->GeInstanceSceneDataBuffers());
 }
\ No newline at end of file
diff --git a/Engine/Source/Runtime/Engine/Private/SkeletalMeshLODRenderData.cpp b/Engine/Source/Runtime/Engine/Private/SkeletalMeshLODRenderData.cpp
--- a/Engine/Source/Runtime/Engine/Private/SkeletalMeshLODRenderData.cpp
+++ b/Engine/Source/Runtime/Engine/Private/SkeletalMeshLODRenderData.cpp
@@ -191,40 +191,6 @@ FArchive& operator<<(FArchive& Ar, FSkelMeshRenderSection& S)
 	return Ar;
 }
 
-void FSectionBoneMapBuffer::Init(TConstArrayView<FSkelMeshRenderSection> Sections)
-{
-	for (const FSkelMeshRenderSection& Section : Sections)
-	{
-		if (Section.IsValid())
-		{
-			Indices.Append(Section.BoneMap);
-		}
-	}
-}
-
-void FSectionBoneMapBuffer::InitRHI(FRHICommandListBase& RHICmdList)
-{
-	if (Indices.IsEmpty())
-	{
-		return;
-	}
-
-	const FRHIBufferCreateDesc CreateDesc =
-		FRHIBufferCreateDesc::CreateVertex(TEXT("SectionBoneMapBuffer"), sizeof(uint16))
-		.AddUsage(BUF_Static | BUF_ShaderResource)
-		.DetermineInitialState();
-
-	VertexBufferRHI = UE::RHIResourceUtils::CreateBufferFromArray<FBoneIndexType>(RHICmdList, TEXT("SectionBoneMapBuffer"), BUF_Static | BUF_ShaderResource, Indices);
-
-	ShaderResourceViewRHI = RHICmdList.CreateShaderResourceView(
-		VertexBufferRHI,
-		FRHIViewDesc::CreateBufferSRV()
-		.SetType(FRHIViewDesc::EBufferType::Typed)
-		.SetFormat(PF_R16_UINT));
-
-	Indices.Empty();
-}
-
 void FSkeletalMeshLODRenderData::InitMorphResources()
 {
 	if (!MorphTargetVertexInfoBuffers.IsRHIInitialized() && MorphTargetVertexInfoBuffers.IsMorphCPUDataValid() && MorphTargetVertexInfoBuffers.NumTotalBatches > 0)
@@ -248,11 +214,6 @@ void FSkeletalMeshLODRenderData::InitResources(bool bNeedsVertexColors, int32 LO
 		IncrementMemoryStats(bNeedsVertexColors);
 	}
 
-#if GPUSKIN_TRANSFORM_STORAGE_DIRECT_ENABLED
-	SectionBoneMapBuffer.Init(RenderSections);
-	BeginInitResource(OwnerName, &SectionBoneMapBuffer, &UE::RenderCommandPipe::SkeletalMesh);
-#endif
-
 	MultiSizeIndexContainer.SetOwnerName(OwnerName);
 	MultiSizeIndexContainer.InitResources();
 
@@ -350,9 +311,6 @@ void FSkeletalMeshLODRenderData::ReleaseResources()
 
 	MultiSizeIndexContainer.ReleaseResources();
 	
-#if GPUSKIN_TRANSFORM_STORAGE_DIRECT_ENABLED
-	BeginReleaseResource(&SectionBoneMapBuffer, &UE::RenderCommandPipe::SkeletalMesh);
-#endif
 	BeginReleaseResource(&StaticVertexBuffers.PositionVertexBuffer, &UE::RenderCommandPipe::SkeletalMesh);
 	BeginReleaseResource(&StaticVertexBuffers.StaticMeshVertexBuffer, &UE::RenderCommandPipe::SkeletalMesh);
 	SkinWeightVertexBuffer.BeginReleaseResources();
@@ -383,9 +341,6 @@ void FSkeletalMeshLODRenderData::ReleaseResources()
 void FSkeletalMeshLODRenderData::IncrementMemoryStats(bool bNeedsVertexColors)
 {
 	INC_DWORD_STAT_BY(STAT_SkeletalMeshIndexMemory, MultiSizeIndexContainer.IsIndexBufferValid() ? (MultiSizeIndexContainer.GetIndexBuffer()->Num() * MultiSizeIndexContainer.GetDataTypeSize()) : 0);
-#if GPUSKIN_TRANSFORM_STORAGE_DIRECT_ENABLED
-	INC_DWORD_STAT_BY(STAT_SkeletalMeshVertexMemory, SectionBoneMapBuffer.GetSize());
-#endif
 	INC_DWORD_STAT_BY(STAT_SkeletalMeshVertexMemory, StaticVertexBuffers.PositionVertexBuffer.GetStride() * StaticVertexBuffers.PositionVertexBuffer.GetNumVertices());
 	INC_DWORD_STAT_BY(STAT_SkeletalMeshVertexMemory, StaticVertexBuffers.StaticMeshVertexBuffer.GetResourceSize());
 	INC_DWORD_STAT_BY(STAT_SkeletalMeshVertexMemory, SkinWeightVertexBuffer.GetVertexDataSize());
@@ -404,9 +359,6 @@ void FSkeletalMeshLODRenderData::IncrementMemoryStats(bool bNeedsVertexColors)
 void FSkeletalMeshLODRenderData::DecrementMemoryStats()
 {
 	DEC_DWORD_STAT_BY(STAT_SkeletalMeshIndexMemory, MultiSizeIndexContainer.IsIndexBufferValid() ? (MultiSizeIndexContainer.GetIndexBuffer()->Num() * MultiSizeIndexContainer.GetDataTypeSize()) : 0);
-#if GPUSKIN_TRANSFORM_STORAGE_DIRECT_ENABLED
-	DEC_DWORD_STAT_BY(STAT_SkeletalMeshVertexMemory, SectionBoneMapBuffer.GetSize());
-#endif
 	DEC_DWORD_STAT_BY(STAT_SkeletalMeshVertexMemory, StaticVertexBuffers.PositionVertexBuffer.GetStride() * StaticVertexBuffers.PositionVertexBuffer.GetNumVertices());
 	DEC_DWORD_STAT_BY(STAT_SkeletalMeshVertexMemory, StaticVertexBuffers.StaticMeshVertexBuffer.GetResourceSize());
 
diff --git a/Engine/Source/Runtime/Engine/Private/SkeletalMeshRenderData.cpp b/Engine/Source/Runtime/Engine/Private/SkeletalMeshRenderData.cpp
--- a/Engine/Source/Runtime/Engine/Private/SkeletalMeshRenderData.cpp
+++ b/Engine/Source/Runtime/Engine/Private/SkeletalMeshRenderData.cpp
@@ -468,7 +468,7 @@ FSkeletalMeshRenderData::FSkeletalMeshRenderData()
 	, PendingFirstLODIdx(0)
 	, LODBiasModifier(0)
 	, bSupportRayTracing(true)
-	, bSharedLODBoneMap(false)
+	, bUnifiedBoneMap(false)
 	, bInitialized(false)
 {
 	ClearNaniteResources(NaniteResourcesPtr);
@@ -627,7 +627,7 @@ PRAGMA_DISABLE_DEPRECATION_WARNINGS
 PRAGMA_ENABLE_DEPRECATION_WARNINGS
 }
 
-static bool ComputeLODsShareBoneMaps(const TIndirectArray<FSkeletalMeshLODRenderData>& LODRenderData)
+static bool HasUnifiedBoneMap(const TIndirectArray<FSkeletalMeshLODRenderData>& LODRenderData)
 {
 	const int32 NumLODs = LODRenderData.Num();
 
@@ -685,7 +685,7 @@ void FSkeletalMeshRenderData::InitResources(bool bNeedsVertexColors, USkinnedAss
 			}
 		}
 
-		bSharedLODBoneMap = ComputeLODsShareBoneMaps(LODRenderData);
+		bUnifiedBoneMap = ::HasUnifiedBoneMap(LODRenderData);
 
 		check(NaniteResourcesPtr.IsValid());
 		NaniteResourcesPtr->InitResources(Owner);
diff --git a/Engine/Source/Runtime/Engine/Private/SkeletalRenderGPUSkin.cpp b/Engine/Source/Runtime/Engine/Private/SkeletalRenderGPUSkin.cpp
--- a/Engine/Source/Runtime/Engine/Private/SkeletalRenderGPUSkin.cpp
+++ b/Engine/Source/Runtime/Engine/Private/SkeletalRenderGPUSkin.cpp
@@ -105,31 +105,6 @@ EBoneTransformAllocationMode GetBoneTransformAllocationMode()
 		: EBoneTransformAllocationMode::Scene;
 }
 
-static int32 GGPUSkinBoneTransformStorageMode = 1;
-static FAutoConsoleVariableRef CVarGPUSkinBoneStorageMode(
-	TEXT("r.GPUSkin.BoneTransformStorageMode"),
-	GGPUSkinBoneTransformStorageMode,
-	TEXT("Controls how bone transforms are stored on the GPU when the scene extension allocation mode is used.")
-	TEXT(" 0 - Direct:            Store bone transforms agnostic of LOD and use a section bone map buffer to fetch the correct bone index on the GPU.")
-	TEXT(" 1 - Indirect Bone Map: Store bone transforms unrolled by LOD section."),
-	FConsoleVariableDelegate::CreateLambda([](IConsoleVariable* InVariable)
-	{
-		FGlobalComponentRecreateRenderStateContext Context;
-	}),
-	ECVF_RenderThreadSafe
-);
-
-EBoneTransformStorageMode GetBoneTransformStorageMode()
-{
-#if GPUSKIN_TRANSFORM_STORAGE_DIRECT_ENABLED
-	return GGPUSkinBoneTransformStorageMode == 0
-		? EBoneTransformStorageMode::Direct
-		: EBoneTransformStorageMode::BoneMap;
-#else
-	return EBoneTransformStorageMode::BoneMap;
-#endif
-}
-
 bool GGPUSkinSceneExtensionEnabled = true;
 static FAutoConsoleVariableRef CVarGPUSkinUseSceneExtension(
 	TEXT("r.GPUSkin.UseSceneExtension"),
@@ -599,7 +574,7 @@ FSkeletalMeshObjectGPUSkin::FSkeletalMeshObjectGPUSkin(const USkinnedMeshCompone
 FSkeletalMeshObjectGPUSkin::FSkeletalMeshObjectGPUSkin(const FSkinnedMeshSceneProxyDesc& InMeshDesc, FSkeletalMeshRenderData* InSkelMeshRenderData, ERHIFeatureLevel::Type InFeatureLevel)
 	: FSkeletalMeshObject(InMeshDesc, InSkelMeshRenderData, InFeatureLevel)
 {
-	BoneTransformStorageMode = ::GetBoneTransformStorageMode();
+	BoneTransformStorageMode = EBoneTransformStorageMode::BoneMap;
 
 	FSkeletalMeshObjectGPUSkin* PreviousMeshObject = nullptr;
 	if (InMeshDesc.PreviousMeshObject && InMeshDesc.PreviousMeshObject->IsGPUSkinMesh())
@@ -2120,10 +2095,6 @@ void InitGPUSkinVertexFactoryComponents(FGPUSkinDataType* VertexFactoryData, con
 	VertexFactoryData->NumBoneInfluences = WeightBuffer->GetMaxBoneInfluences();
 
 	FGPUBaseSkinVertexFactory::FShaderDataType* ShaderData = VertexFactory ? &VertexFactory->GetShaderData() : nullptr;
-	if (ShaderData && VertexBuffers.SectionBoneMapBuffer)
-	{
-		ShaderData->SectionBoneMapBuffer = VertexBuffers.SectionBoneMapBuffer->ShaderResourceViewRHI;
-	}
 
 	GPUSkinBoneInfluenceType BoneInfluenceType = WeightBuffer->GetBoneInfluenceType();
 	if (BoneInfluenceType == GPUSkinBoneInfluenceType::UnlimitedBoneInfluence)
@@ -2472,9 +2443,6 @@ void FSkeletalMeshObjectGPUSkin::GetUsedVertexFactoryData(
  */
 void FSkeletalMeshObjectGPUSkin::FSkeletalMeshObjectLOD::GetVertexBuffers(FVertexFactoryBuffers& OutVertexBuffers, FSkeletalMeshLODRenderData& LODData, EBoneTransformStorageMode StorageMode)
 {
-#if GPUSKIN_TRANSFORM_STORAGE_DIRECT_ENABLED
-	OutVertexBuffers.SectionBoneMapBuffer = StorageMode == EBoneTransformStorageMode::Direct ? &LODData.SectionBoneMapBuffer : nullptr;
-#endif
 	OutVertexBuffers.StaticVertexBuffers = &LODData.StaticVertexBuffers;
 	OutVertexBuffers.ColorVertexBuffer = MeshObjectColorBuffer;
 	OutVertexBuffers.SkinWeightVertexBuffer = MeshObjectWeightBuffer;
diff --git a/Engine/Source/Runtime/Engine/Private/SkeletalRenderGPUSkin.h b/Engine/Source/Runtime/Engine/Private/SkeletalRenderGPUSkin.h
--- a/Engine/Source/Runtime/Engine/Private/SkeletalRenderGPUSkin.h
+++ b/Engine/Source/Runtime/Engine/Private/SkeletalRenderGPUSkin.h
@@ -480,7 +480,6 @@ public:
 	 */
 	struct FVertexFactoryBuffers
 	{
-		FSectionBoneMapBuffer* SectionBoneMapBuffer = nullptr;
 		FStaticMeshVertexBuffers* StaticVertexBuffers = nullptr;
 		const FSkinWeightVertexBuffer* SkinWeightVertexBuffer = nullptr;
 		FColorVertexBuffer*	ColorVertexBuffer = nullptr;
diff --git a/Engine/Source/Runtime/Engine/Public/Rendering/SkeletalMeshLODRenderData.h b/Engine/Source/Runtime/Engine/Public/Rendering/SkeletalMeshLODRenderData.h
--- a/Engine/Source/Runtime/Engine/Public/Rendering/SkeletalMeshLODRenderData.h
+++ b/Engine/Source/Runtime/Engine/Public/Rendering/SkeletalMeshLODRenderData.h
@@ -124,24 +124,6 @@ struct FSkelMeshRenderSection
 	ENGINE_API friend FArchive& operator<<(FArchive& Ar, FSkelMeshRenderSection& S);
 };
 
-class FSectionBoneMapBuffer : public FVertexBufferWithSRV
-{
-public:
-	FSectionBoneMapBuffer() = default;
-
-	ENGINE_API void Init(TConstArrayView<FSkelMeshRenderSection> Sections);
-
-	ENGINE_API void InitRHI(FRHICommandListBase& RHICmdList) override;
-
-	uint32 GetSize() const
-	{
-		return VertexBufferRHI ? VertexBufferRHI->GetSize() : 0;
-	}
-
-private:
-	TArray<FBoneIndexType> Indices;
-};
-
 class FSkeletalMeshLODRenderData : public UE::Private::FQueryableRefCountedObject
 {
 public:
@@ -172,11 +154,6 @@ public:
 	/** GPU buffer for half edge data of an LOD, useful for mesh deformers */
 	FSkeletalMeshHalfEdgeBuffer HalfEdgeBuffer;
 
-#if GPUSKIN_TRANSFORM_STORAGE_DIRECT_ENABLED
-	/** Buffer containing bone indices laid out by section. */
-	FSectionBoneMapBuffer SectionBoneMapBuffer;
-#endif
-
 	TArray<FBoneIndexType> ActiveBoneIndices;
 
 	TArray<FBoneIndexType> RequiredBones;
diff --git a/Engine/Source/Runtime/Engine/Public/Rendering/SkeletalMeshRenderData.h b/Engine/Source/Runtime/Engine/Public/Rendering/SkeletalMeshRenderData.h
--- a/Engine/Source/Runtime/Engine/Public/Rendering/SkeletalMeshRenderData.h
+++ b/Engine/Source/Runtime/Engine/Public/Rendering/SkeletalMeshRenderData.h
@@ -44,7 +44,7 @@ public:
 	bool bSupportRayTracing;
 
 	/** True if all LODs share the same section bone map (enables instanced skinned GPU LOD without indirection) */
-	bool bSharedLODBoneMap;
+	bool bUnifiedBoneMap;
 
 	/** Optional, name of the owning skeletal mesh */
 	FName OwnerName = NAME_None;
@@ -145,6 +145,11 @@ public:
 		return bInitialized;
 	}
 
+	bool HasUnifiedBoneMap() const
+	{
+		return bUnifiedBoneMap;
+	}
+
 	ENGINE_API bool HasValidNaniteData() const;
 
 private:
