SoundWaveProxyPlayer

Remove `check. We can safely handle the Resampler consuming fewer samples than expected. But there appear to be situations where the resampler will consume 0 samples.

We should at least assert that the number of samples is divisible by the number of channels - at least the Resampler should be consuming samples based on the input number of channels

#Synced-CL 50324507

diff --git a/Engine/Source/Runtime/Engine/Private/SoundWaveProxyPlayer.cpp b/Engine/Source/Runtime/Engine/Private/SoundWaveProxyPlayer.cpp
--- a/Engine/Source/Runtime/Engine/Private/SoundWaveProxyPlayer.cpp
+++ b/Engine/Source/Runtime/Engine/Private/SoundWaveProxyPlayer.cpp
@@ -372,8 +372,8 @@ void FSoundWaveProxyPlayer::RenderMultiChannelAudio(Audio::FMultichannelBufferVi
 		int32 NumSamplesRemaining = SourceAudio.Num() - NumSamplesConsumed;
 		if (NumSamplesConsumed < SourceAudio.Num())
 		{
-			// move the last frame of samples to the start of the SourceBuffer for future processing
-			check(NumSamplesRemaining == NumChannels);
+			// move remaining samples to the start of the SourceBuffer for future processing
+			check(NumSamplesRemaining % NumChannels == 0);
 			FMemory::Memmove(SourceAudio.GetData(), SourceAudio.GetData() + NumSamplesConsumed, NumSamplesRemaining * sizeof(float));
 		}
 		SourceAudio.SetNum(NumSamplesRemaining, EAllowShrinking::No);
