Batch external actor loads

#rnx
#Synced-CL 50152864
#Synced-CL 50152864

diff --git a/Engine/Source/Runtime/Engine/Private/WorldPartition/WorldPartitionRuntimeLevelStreamingCell.cpp b/Engine/Source/Runtime/Engine/Private/WorldPartition/WorldPartitionRuntimeLevelStreamingCell.cpp
--- a/Engine/Source/Runtime/Engine/Private/WorldPartition/WorldPartitionRuntimeLevelStreamingCell.cpp
+++ b/Engine/Source/Runtime/Engine/Private/WorldPartition/WorldPartitionRuntimeLevelStreamingCell.cpp
@@ -438,16 +438,18 @@ bool UWorldPartitionRuntimeLevelStreamingCell::OnPrepareGeneratorPackageForCook(
 		UWorldPartition* WorldPartition = OuterWorld->GetWorldPartition();
 
 		FWorldPartitionLevelHelper::FPackageReferencer PackageReferencer;
+		TArray<int32> AsyncRequestIDs;
 		FWorldPartitionLevelHelper::FLoadActorsParams Params = FWorldPartitionLevelHelper::FLoadActorsParams()
 			.SetOuterWorld(OuterWorld)
 			.SetDestLevel(nullptr)
 			.SetActorPackages(Packages)
 			.SetPackageReferencer(&PackageReferencer)
 			.SetCompletionCallback([](bool) {})
-			.SetLoadAsync(false)
+			.SetLoadAsync(true, &AsyncRequestIDs)
 			.SetInstancingContext(FLinkerInstancingContext(false)); // Don't do SoftObjectPath remapping for PersistentLevel actors because references can end up in different cells
 
-		verify(FWorldPartitionLevelHelper::LoadActors(MoveTemp(Params)));
+		FWorldPartitionLevelHelper::LoadActors(MoveTemp(Params));
+		FlushAsyncLoading(AsyncRequestIDs);
 
 		FWorldPartitionLevelHelper::MoveExternalActorsToLevel(Packages, OuterWorld->PersistentLevel, OutModifiedPackages);
 
@@ -585,16 +587,18 @@ bool UWorldPartitionRuntimeLevelStreamingCell::OnPopulateGeneratedPackageForCook
 
 		// Load cell Actors
 		FWorldPartitionLevelHelper::FPackageReferencer PackageReferencer;
+		TArray<int32> AsyncRequestIDs;
 		FWorldPartitionLevelHelper::FLoadActorsParams Params = FWorldPartitionLevelHelper::FLoadActorsParams()
 			.SetOuterWorld(OuterWorld)
 			.SetDestLevel(nullptr)
 			.SetActorPackages(Packages)
 			.SetPackageReferencer(&PackageReferencer)
 			.SetCompletionCallback([](bool) {})
-			.SetLoadAsync(false)
+			.SetLoadAsync(true, &AsyncRequestIDs)
 			.SetInstancingContext(FLinkerInstancingContext(false)); // Don't do SoftObjectPath remapping for PersistentLevel actors because references can end up in different cells
 
-		verify(FWorldPartitionLevelHelper::LoadActors(MoveTemp(Params)));
+		FWorldPartitionLevelHelper::LoadActors(MoveTemp(Params));
+		FlushAsyncLoading(AsyncRequestIDs);
 
 		// Create a level and move these actors in it
 		ULevel* NewLevel = FWorldPartitionLevelHelper::CreateEmptyLevelForRuntimeCell(this, OuterWorld, LevelStreaming->GetWorldAsset().ToString(), InPackage);
