Fixed regression crash in new material translator when connecting a Substrate slab to front mateiral.

#rnx

#Synced-CL 50394150

diff --git a/Engine/Source/Runtime/Engine/Private/Materials/MaterialIRValueAnalyzer.cpp b/Engine/Source/Runtime/Engine/Private/Materials/MaterialIRValueAnalyzer.cpp
--- a/Engine/Source/Runtime/Engine/Private/Materials/MaterialIRValueAnalyzer.cpp
+++ b/Engine/Source/Runtime/Engine/Private/Materials/MaterialIRValueAnalyzer.cpp
@@ -806,7 +806,10 @@ static void AnalyzeSetMaterialOutput(FMaterialIRValueAnalyzer& Analyzer, MIR::FS
 	}
 
 	// Register the material output argument into the module. This "declares" the final output value of a material property.
-	Analyzer.Module->SetPropertyValue(SetMaterialOutput->Property, SetMaterialOutput->Arg);
+	if (SetMaterialOutput->Property < MP_MaterialAttributes)
+	{
+		Analyzer.Module->SetPropertyValue(SetMaterialOutput->Property, SetMaterialOutput->Arg);
+	}
 }
 
 static void AnalyzeMaterialParameterCollection(FMaterialIRValueAnalyzer& Analyzer, MIR::FMaterialParameterCollection* MaterialParameterCollection)
diff --git a/Engine/Source/Runtime/Engine/Public/Materials/MaterialIRModule.h b/Engine/Source/Runtime/Engine/Public/Materials/MaterialIRModule.h
--- a/Engine/Source/Runtime/Engine/Public/Materials/MaterialIRModule.h
+++ b/Engine/Source/Runtime/Engine/Public/Materials/MaterialIRModule.h
@@ -128,7 +128,7 @@ public:
 	void AddValue(MIR::FValue* Value) { Values.Add(Value); }
 
 	// Sets the value of a material property for future reference.
-	void SetPropertyValue(EMaterialProperty Property, MIR::FValue* Value) { PropertyValues[Property] = Value; }
+	void SetPropertyValue(EMaterialProperty Property, MIR::FValue* Value) { check(Property < MP_MaterialAttributes); PropertyValues[Property] = Value; }
 	
 	// Records a custom output instruction.
 	void AddCustomOutput(const MIR::FSetMaterialOutput* Instr) { CustomOutputs.Push(Instr); }
