Substrate - Toon - added basic exemple of shadow masking for legacy gbuffer path. That will likely evolve and change later.

#rnx
#tests pc editor

#Synced-CL 50393747

diff --git a/Engine/Shaders/Private/ShadowProjectionPixelShader.usf b/Engine/Shaders/Private/ShadowProjectionPixelShader.usf
--- a/Engine/Shaders/Private/ShadowProjectionPixelShader.usf
+++ b/Engine/Shaders/Private/ShadowProjectionPixelShader.usf
@@ -300,10 +300,16 @@ void Main(
 			#if SUBSTRATE_GBUFFER_FORMAT==1
 			const FSubstrateSubsurfaceHeader SSSHeader = SubstrateLoadSubsurfaceHeader(Substrate.MaterialTextureArray, Substrate.FirstSliceStoringSubstrateSSSData, SVPos.xy);
 			const uint SSSType = SubstrateSubSurfaceHeaderGetSSSType(SSSHeader);
+			#else
+			#if SUBSTRATE_ENABLED
+				FToonProfileShadowSettings ToonProfileShadowSettings = (FToonProfileShadowSettings)0;
+			#endif
 			#endif
 
 			bool bApplyPCFOverBlurCorrection = true; // Set to false when the transmittance from extinction is correctly respected. Toon & TransmitttanceShadow disables that to have a physically based non squared transmittance curve.
 			float Density = 1;
+			float ToonProfileId = 0.0f;
+			float2 PatternUVs = ScreenUV;
 			if (!bIsEyeOrHair)
 			{
 			#if SUBSTRATE_GBUFFER_FORMAT==1
@@ -325,12 +331,11 @@ void Main(
 				{
 					bApplyPCFOverBlurCorrection = false;
 
-					float2 PatternUVs = 0.0f;
-					float ToonProfileId = 0.0f;
 					UnpackToonCustomData(PackRGBA8(GBufferData.CustomData), PatternUVs, ToonProfileId);
 
 					// Density is extinction as 1/cm (not 1/m) to match with the transmittance shadow evaluation from ManualPCF below.
-					Density = GetToonProfileShadowSettings(ToonProfileId).ShadowExtinctionCoefficientInInvCentimeter;
+					ToonProfileShadowSettings = GetToonProfileShadowSettings(ToonProfileId);
+					Density = ToonProfileShadowSettings.ShadowExtinctionCoefficientInInvCentimeter;
 				}
 				else
 				#endif
@@ -389,6 +394,14 @@ void Main(
 				// This branch is taken by Subsurface shading model (Wrap for Substrate) or the toon shading model. 
 				// => SSSTransmittanceOrOpticalDepth stores the result of ManualPCF() evaluating Transmittance.
 				SSSTransmittanceOrOpticalDepth = ManualPCF(ShadowPosition.xy, Settings);
+
+			#if SUBSTRATE_GBUFFER_FORMAT==1
+				// SUBSTRATE_TODO implmenent that path
+			#else
+				#if SUBSTRATE_ENABLED
+				SSSTransmittanceOrOpticalDepth *= EvaluateToonProfileShadowHatching(ToonProfileId, 0*ScreenUV + 1*PatternUVs, ToonProfileShadowSettings, SSSTransmittanceOrOpticalDepth);
+				#endif
+			#endif
 			}
 
 			// Legacy code maintained: here we have a square applied over the transmittance curve or encoded optical depth... 
diff --git a/Engine/Shaders/Private/ToonProfile.usf b/Engine/Shaders/Private/ToonProfile.usf
--- a/Engine/Shaders/Private/ToonProfile.usf
+++ b/Engine/Shaders/Private/ToonProfile.usf
@@ -17,6 +17,9 @@ float DiffuseDitherPatternSize;
 float SpecularDitherPatternStrength;
 float SpecularDitherPatternSize;
 
+float ShadowHatchingPatternSize;
+float ShadowHatchingPatternStrength;
+
 RWTexture2DArray<float4>	ToonProfileTexture;
 
 Texture2D<float4>			DiffuseDitherPattern;
@@ -25,8 +28,12 @@ SamplerState				DiffuseDitherPatternSampler;
 Texture2D<float4>			SpecularDitherPattern;
 SamplerState				SpecularDitherPatternSampler;
 
+Texture2D<float4>			ShadowHatchingPattern;
+SamplerState				ShadowHatchingPatternSampler;
+
 Buffer<float4>				DiffuseRampBufferSRV;
 Buffer<float4>				SpecularRampBufferSRV;
+Buffer<float4>				ShadowRampBufferSRV;
 
 [numthreads(8, 8, 1)]
 void ToonProfileUpdateAtlasTileCS(uint3 DispatchThreadId : SV_DispatchThreadID)
@@ -47,6 +54,10 @@ void ToonProfileUpdateAtlasTileCS(uint3 DispatchThreadId : SV_DispatchThreadID)
 		{
 			ToonProfileTexture[uint3(PixelPos, ToonProfileIndex)] = SpecularRampBufferSRV[PixelPos.x];
 		}
+		else if (PixelPos.y == TOON_PROFILE_SHADOW_RAMP_Y)
+		{
+			ToonProfileTexture[uint3(PixelPos, ToonProfileIndex)] = ShadowRampBufferSRV[PixelPos.x];
+		}
 		else if (PixelPos.y == TOON_PROFILE_PARAMETERS_Y)
 		{
 			if (PixelPos.x == 0)
@@ -55,14 +66,14 @@ void ToonProfileUpdateAtlasTileCS(uint3 DispatchThreadId : SV_DispatchThreadID)
 			}
 			else if (PixelPos.x == 1)
 			{
-				ToonProfileTexture[uint3(PixelPos, ToonProfileIndex)] = float4(ToonProfileShadowExtinctionCoefficientInInvCentimeter, 0.0f, 0.0f, 0.0f);
+				ToonProfileTexture[uint3(PixelPos, ToonProfileIndex)] = float4(ToonProfileShadowExtinctionCoefficientInInvCentimeter, ShadowHatchingPatternSize, ShadowHatchingPatternStrength, 0.0f);
 			}
 			else
 			{
 				ToonProfileTexture[uint3(PixelPos, ToonProfileIndex)] = DefaultColor;
 			}
 		}
-		else if (PixelPos.y >= TOON_PROFILE_DITHERPATTERN_Y && PixelPos.y <= uint(TOON_PROFILE_DITHERPATTERN_Y + ProfileResolution))
+		else if (PixelPos.y >= uint(TOON_PROFILE_TEXTURES_Y + TOON_PROFILE_TEXTURE_DITHERPATTERN_IDX * ProfileResolution) && PixelPos.y < uint(TOON_PROFILE_TEXTURES_Y + (TOON_PROFILE_TEXTURE_DITHERPATTERN_IDX+1) * ProfileResolution))
 		{
 		/*
 			For a Toon profile of full resolution N, we render an edge of 1 texel all around to make sure the texture is repeatable, marked `e` below.
@@ -87,7 +98,7 @@ void ToonProfileUpdateAtlasTileCS(uint3 DispatchThreadId : SV_DispatchThreadID)
 			| e | e | e | e | e | e | e | e |
 		0,1 --------------------------------- 1,1
 		*/
-			const float2 FullOrigin = float2(0.0f, TOON_PROFILE_DITHERPATTERN_Y);
+			const float2 FullOrigin = float2(0.0f, TOON_PROFILE_TEXTURES_Y + TOON_PROFILE_TEXTURE_DITHERPATTERN_IDX * ProfileResolution);
 			const float2 BakeOrigin = FullOrigin + 1.0f;
 
 			const float  BakeResolution = ProfileResolution - 2.0f;
@@ -100,6 +111,21 @@ void ToonProfileUpdateAtlasTileCS(uint3 DispatchThreadId : SV_DispatchThreadID)
 			const float SpecularDitherPatternValue = SpecularDitherPattern.SampleLevel(SpecularDitherPatternSampler, FullUVs, 0.0f).y;
 			ToonProfileTexture[uint3(PixelPos, ToonProfileIndex)] = float4(DiffuseDitherPatternValue, SpecularDitherPatternValue, 0.0f, 0.0f);
 		}
+		else if (PixelPos.y >= uint(TOON_PROFILE_TEXTURES_Y + TOON_PROFILE_TEXTURE_SHADOWHATCHING_IDX * ProfileResolution) && PixelPos.y < uint(TOON_PROFILE_TEXTURES_Y + (TOON_PROFILE_TEXTURE_SHADOWHATCHING_IDX + 1) * ProfileResolution))
+		{
+			// This texture is stored with an edge for filtering similarly to TOON_PROFILE_TEXTURE_DITHERPATTERN_IDX
+			const float2 FullOrigin = float2(0.0f, TOON_PROFILE_TEXTURES_Y + TOON_PROFILE_TEXTURE_SHADOWHATCHING_IDX * ProfileResolution);
+			const float2 BakeOrigin = FullOrigin + 1.0f;
+
+			const float  BakeResolution = ProfileResolution - 2.0f;
+			const float  BakeResolutionInv = 1.0f / BakeResolution;
+
+			const float2 FullCoord  = float2(PixelPos.xy) + 0.5f.xx;
+			const float2 FullUVs    = (FullCoord - BakeOrigin) * BakeResolutionInv;
+
+			const float4 ShadowHatchingPatternValues = ShadowHatchingPattern.SampleLevel(ShadowHatchingPatternSampler, FullUVs, 0.0f).rgba;
+			ToonProfileTexture[uint3(PixelPos, ToonProfileIndex)] = ShadowHatchingPatternValues;
+		}
 		else
 		{
 			ToonProfileTexture[uint3(PixelPos, ToonProfileIndex)] = DefaultColor;
diff --git a/Engine/Shaders/Private/ToonProfileCommon.ush b/Engine/Shaders/Private/ToonProfileCommon.ush
--- a/Engine/Shaders/Private/ToonProfileCommon.ush
+++ b/Engine/Shaders/Private/ToonProfileCommon.ush
@@ -24,6 +24,10 @@ float3 EvaluateToonProfileSpecularRamp(in float ToonProfileId, in float Specular
 	return EvaluateToonProfileRamp(ToonProfileId, Specular, TOON_PROFILE_SPECULAR_RAMP_Y);
 }
 
+float EvaluateToonProfileShadowRamp(in float ToonProfileId, in float Transmittance)
+{
+	return EvaluateToonProfileRamp(ToonProfileId, Transmittance, TOON_PROFILE_SHADOW_RAMP_Y).r;
+}
 
 struct FToonProfileEvaluateSettings
 {
@@ -49,20 +53,18 @@ FToonProfileEvaluateSettings GetToonProfileEvaluateSettings(in float ToonProfile
 	return Ret;
 }
 
-float EvaluateToonProfileDiffuseDither(in float ToonProfileId, in float2 ScreenUV, in FToonProfileEvaluateSettings ToonProfileEvaluateSettings)
+float EvaluateToonProfileDiffuseDither(in float ToonProfileId, in float2 PatternUVs, in FToonProfileEvaluateSettings ToonProfileEvaluateSettings)
 {
 	if (ToonProfileEvaluateSettings.DiffuseDitherPatternStrength > 0.0f)
 	{
-		float2 UVs = ScreenUV * (1.0f / ToonProfileEvaluateSettings.DiffuseDitherPatternSize);	// Project the pattern on screen as a percentage of it.
-		UVs.y *= View.ViewSizeAndInvSize.y * View.ViewSizeAndInvSize.z;							// Make sure the pattern square shape is enforced on screen.
-		UVs = frac(UVs);																		// Pattern wrap repetition.
-
 		// See ToonProfile.usf for the baking/sampling description.
 
-		const float2 FullOrigin = float2(0.0f, TOON_PROFILE_DITHERPATTERN_Y);
+		PatternUVs = frac(PatternUVs * ToonProfileEvaluateSettings.DiffuseDitherPatternSize);
+
+		const float2 FullOrigin = float2(0.0f, TOON_PROFILE_TEXTURES_Y + TOON_PROFILE_TEXTURE_DITHERPATTERN_IDX * View.ToonProfileDitherPatternSizeAndInvSize.x);
 		const float2 BakeOrigin = FullOrigin + 1.0f;
 		
-		const float2 BakeCoords        = BakeOrigin + UVs * View.ToonProfileDitherPatternSizeAndInvSize.x;
+		const float2 BakeCoords        = BakeOrigin + PatternUVs * View.ToonProfileDitherPatternSizeAndInvSize.x;
 
 		const float3 ToonProfileCoord = float3(
 			BakeCoords.xy * View.ToonProfileTextureSizeAndInvSize.zw,
@@ -74,20 +76,18 @@ float EvaluateToonProfileDiffuseDither(in float ToonProfileId, in float2 ScreenU
 	return 0.0f;
 }
 
-float EvaluateToonProfileSpecularDither(in float ToonProfileId, in float2 ScreenUV, in FToonProfileEvaluateSettings ToonProfileEvaluateSettings)
+float EvaluateToonProfileSpecularDither(in float ToonProfileId, in float2 PatternUVs, in FToonProfileEvaluateSettings ToonProfileEvaluateSettings)
 {
 	if (ToonProfileEvaluateSettings.SpecularDitherPatternStrength > 0.0f)
 	{
-		float2 UVs = ScreenUV * (1.0f / ToonProfileEvaluateSettings.SpecularDitherPatternSize);	// Project the pattern on screen as a percentage of it.
-		UVs.y *= View.ViewSizeAndInvSize.y * View.ViewSizeAndInvSize.z;							// Make sure the pattern square shape is enforced on screen.
-		UVs = frac(UVs);																		// Pattern wrap repetition.
-
 		// See ToonProfile.usf for the baking/sampling description.
 
-		const float2 FullOrigin = float2(0.0f, TOON_PROFILE_DITHERPATTERN_Y);
+		PatternUVs = frac(PatternUVs * ToonProfileEvaluateSettings.SpecularDitherPatternSize);
+
+		const float2 FullOrigin = float2(0.0f, TOON_PROFILE_TEXTURES_Y + TOON_PROFILE_TEXTURE_DITHERPATTERN_IDX * View.ToonProfileTextureSizeAndInvSize.x);
 		const float2 BakeOrigin = FullOrigin + 1.0f;
 		
-		const float2 BakeCoords = BakeOrigin + UVs * View.ToonProfileDitherPatternSizeAndInvSize.x;
+		const float2 BakeCoords = BakeOrigin + PatternUVs * View.ToonProfileDitherPatternSizeAndInvSize.x;
 
 		const float3 ToonProfileCoord = float3(
 			BakeCoords.xy * View.ToonProfileTextureSizeAndInvSize.zw,
@@ -103,6 +103,8 @@ struct FToonProfileShadowSettings
 {
 	float ShadowExtinctionCoefficientInInvCentimeter;	// Needed because transmittance evaluation code from ShadowProjectionPixelShader.usf uses centimeters.
 	float ShadowMFPInMeters;							// Needed for storage in substrate Subsurface Header.
+	float ShadowHatchingPatternSize;
+	float ShadowHatchingPatternStrength;
 };
 
 FToonProfileShadowSettings GetToonProfileShadowSettings(in float ToonProfileId)
@@ -119,6 +121,47 @@ FToonProfileShadowSettings GetToonProfileShadowSettings(in float ToonProfileId)
 	const float ShadowExtinctionCoefficientInInvMeter = Ret.ShadowExtinctionCoefficientInInvCentimeter * (1.0f / CENTIMETER_TO_METER); // 1/cm => 1/m
 	Ret.ShadowMFPInMeters = 1.0f / max(0.000001f, ShadowExtinctionCoefficientInInvMeter);
 
+	Ret.ShadowHatchingPatternSize = Data.y;
+
+	Ret.ShadowHatchingPatternStrength = Data.z;
+
 	return Ret;
 }
 
+// SUBSTRATE_TODO instead of shadow/transmittance, it could also account for NoL or AO?
+float EvaluateToonProfileShadowHatching(in float ToonProfileId, in float2 PatternUVs, in FToonProfileShadowSettings ToonProfileShadowSettings, in float Transmittance)
+{
+	if (ToonProfileShadowSettings.ShadowHatchingPatternStrength > 0.0f)
+	{
+		// See ToonProfile.usf for the baking/sampling description.
+
+		PatternUVs = frac(PatternUVs * ToonProfileShadowSettings.ShadowHatchingPatternSize);
+
+		const float2 FullOrigin = float2(0.0f, TOON_PROFILE_TEXTURES_Y + TOON_PROFILE_TEXTURE_SHADOWHATCHING_IDX * View.ToonProfileTextureSizeAndInvSize.x);
+		const float2 BakeOrigin = FullOrigin + 1.0f;
+
+		const float2 BakeCoords = BakeOrigin + PatternUVs * View.ToonProfileDitherPatternSizeAndInvSize.x;
+
+		const float3 ToonProfileCoord = float3(
+			BakeCoords.xy * View.ToonProfileTextureSizeAndInvSize.zw,
+			ToonProfileId
+		);
+
+		const float4 ShadowHatching = View.ToonProfileTexture.SampleLevel(View.ToonProfileSampler, ToonProfileCoord, 0).rgba;
+
+		const float TransmittanceCurve = EvaluateToonProfileShadowRamp(ToonProfileId, Transmittance).x;
+		const float Step = 0.25f;
+		const float StepInv = 1.0f / Step;
+		const float4 Step4 = float4(1.0f, 0.75f, 0.50f, 0.25f);
+	
+		const float curveInputX = TransmittanceCurve;
+		float4 LerpFactors = 1.0f-saturate(abs(curveInputX - Step4) * StepInv);
+	
+		// Fade to white from 0.95 to 1.0. SUBSTRATE_TODOO make that better parameter.
+		const float FadeToWhite = saturate(1.0f - (curveInputX - 0.95f) / 0.05);
+	
+		return lerp(1.0f, dot(ShadowHatching, LerpFactors), ToonProfileShadowSettings.ShadowHatchingPatternStrength * FadeToWhite);
+	}
+	return 1.0f;
+}
+
diff --git a/Engine/Shaders/Shared/ToonProfileDefinitions.h b/Engine/Shaders/Shared/ToonProfileDefinitions.h
--- a/Engine/Shaders/Shared/ToonProfileDefinitions.h
+++ b/Engine/Shaders/Shared/ToonProfileDefinitions.h
@@ -4,5 +4,10 @@
 
 #define TOON_PROFILE_DIFFUSE_RAMP_Y		0
 #define TOON_PROFILE_SPECULAR_RAMP_Y	1
-#define TOON_PROFILE_PARAMETERS_Y		2
-#define TOON_PROFILE_DITHERPATTERN_Y	3
+#define TOON_PROFILE_SHADOW_RAMP_Y		2
+#define TOON_PROFILE_PARAMETERS_Y		3
+#define TOON_PROFILE_TEXTURES_Y			4
+
+#define TOON_PROFILE_TEXTURE_DITHERPATTERN_IDX		0
+#define TOON_PROFILE_TEXTURE_SHADOWHATCHING_IDX		1
+#define TOON_PROFILE_TEXTURE_COUNT					2
diff --git a/Engine/Source/Runtime/Engine/Classes/Engine/ToonProfile.h b/Engine/Source/Runtime/Engine/Classes/Engine/ToonProfile.h
--- a/Engine/Source/Runtime/Engine/Classes/Engine/ToonProfile.h
+++ b/Engine/Source/Runtime/Engine/Classes/Engine/ToonProfile.h
@@ -79,6 +79,31 @@ struct FToonProfileStruct
 	UPROPERTY(Category = "Shadow", EditAnywhere, meta = (DisplayName = "Shadow Extinction", UIMin = 0.0, UIMax = 10.0, ClampMin = 0.0, ClampMax = 100.0, SliderExponent = 2.0))
 	float ShadowExtinctionCoefficient;
 
+	/**
+	* Defines the shadow ramp used to control how the four channels for the hatching pattern are distributed along the transmittance curve.
+	*/
+	UPROPERTY(Category = "Shadow", EditAnywhere, meta = (AllowZoomOutput = "false", ShowZoomButtons = "false", ViewMinInput = "0", ViewMaxInput = "1", ViewMinOutput = "0", ViewMaxOutput = "1", TimelineLength = "1", ShowInputGridNumbers = "false", ShowOutputGridNumbers = "false"))
+	FRuntimeCurveLinearColor ShadowRamp;
+
+	/**
+	 * The shadow texture defining the hatching pattern applied to the grey scale shadow. R applied first close to the lit front surface and progressively G, B and A channels lerped when away from the lit front surface.
+	 * Typically R, G, B and A would have light to bold hatching pattern.
+	 */
+	UPROPERTY(Category = "Shadow", EditAnywhere)
+	TObjectPtr<UTexture> ShadowHatchingPattern;
+
+	/**
+	 * The size on screen of the shadow hatching pattern.
+	 */
+	UPROPERTY(Category = "Shadow", EditAnywhere, meta = (UIMin = 0.0, UIMax = 1.0, ClampMin = 0.0))
+	float ShadowHatchingPatternSize;
+
+	/**
+	 * The uniform strength of the shadow hatching pattern.
+	 */
+	UPROPERTY(Category = "Shadow", EditAnywhere, meta = (UIMin = 0.0, UIMax = 1.0, ClampMin = 0.0))
+	float ShadowHatchingPatternStrength;
+
 	FToonProfileStruct();
 
 	void Invalidate()
diff --git a/Engine/Source/Runtime/Engine/Private/Rendering/ToonProfile.cpp b/Engine/Source/Runtime/Engine/Private/Rendering/ToonProfile.cpp
--- a/Engine/Source/Runtime/Engine/Private/Rendering/ToonProfile.cpp
+++ b/Engine/Source/Runtime/Engine/Private/Rendering/ToonProfile.cpp
@@ -65,6 +65,8 @@ public:
 		SHADER_PARAMETER(float, DiffuseDitherPatternSize)
 		SHADER_PARAMETER(float, SpecularDitherPatternStrength)
 		SHADER_PARAMETER(float, SpecularDitherPatternSize)
+		SHADER_PARAMETER(float, ShadowHatchingPatternSize)
+		SHADER_PARAMETER(float, ShadowHatchingPatternStrength)
 
 		SHADER_PARAMETER(float, ToonProfileShadowExtinctionCoefficientInInvCentimeter)
 
@@ -76,8 +78,12 @@ public:
 		SHADER_PARAMETER_RDG_TEXTURE(Texture2D, SpecularDitherPattern)
 		SHADER_PARAMETER_SAMPLER(SamplerState, SpecularDitherPatternSampler)
 
+		SHADER_PARAMETER_RDG_TEXTURE(Texture2D, ShadowHatchingPattern)
+		SHADER_PARAMETER_SAMPLER(SamplerState, ShadowHatchingPatternSampler)
+
 		SHADER_PARAMETER_RDG_BUFFER_SRV(Buffer, DiffuseRampBufferSRV)
 		SHADER_PARAMETER_RDG_BUFFER_SRV(Buffer, SpecularRampBufferSRV)
+		SHADER_PARAMETER_RDG_BUFFER_SRV(Buffer, ShadowRampBufferSRV)
 	END_SHADER_PARAMETER_STRUCT()
 
 	static int32 GetThreadGroupSize()
@@ -315,7 +321,7 @@ IPooledRenderTarget* FToonProfileTextureManager::GetAtlasTexture(FRDGBuilder& Gr
 
 		// See ToonProfileCommon.ush
 		const uint32 ResolutionX = ToonProfileResolution;
-		const uint32 ResolutionY = TOON_PROFILE_DITHERPATTERN_Y + ToonProfileResolution; // row for ramps, parameters and 2d dithering patterns. // SUBSTRATE_TODO shared offset header
+		const uint32 ResolutionY = TOON_PROFILE_TEXTURES_Y + TOON_PROFILE_TEXTURE_COUNT * ToonProfileResolution; // row for ramps, parameters and 2d dithering patterns.
 
 		// 1. Create atlas texture
 		const FIntPoint ToonProfileResolution2D = FIntPoint(ToonProfileResolution, ResolutionY);
@@ -339,23 +345,30 @@ IPooledRenderTarget* FToonProfileTextureManager::GetAtlasTexture(FRDGBuilder& Gr
 
 			FTextureResource* DiffuseDitherPatternTextureResource = Data.DiffuseDitherPattern ? Data.DiffuseDitherPattern->GetResource() : nullptr;
 			FRHITexture* DiffuseDitherPatternTexture2DRHI = DiffuseDitherPatternTextureResource ? DiffuseDitherPatternTextureResource->GetTexture2DRHI() : nullptr;
-			bool bHasDiffuseDitherPaternTexture = DiffuseDitherPatternTexture2DRHI != nullptr;
+			bool bHasDiffuseDitherPatternTexture = DiffuseDitherPatternTexture2DRHI != nullptr;
 			DiffuseDitherPatternTexture2DRHI = DiffuseDitherPatternTexture2DRHI ? DiffuseDitherPatternTexture2DRHI : GBlackTexture->TextureRHI->GetTexture2D();
 
 			FTextureResource* SpecularDitherPatternTextureResource = Data.SpecularDitherPattern ? Data.SpecularDitherPattern->GetResource() : nullptr;
 			FRHITexture* SpecularDitherPatternTexture2DRHI = SpecularDitherPatternTextureResource ? SpecularDitherPatternTextureResource->GetTexture2DRHI() : nullptr;
-			bool bHasSpecularDitherPaternTexture = SpecularDitherPatternTexture2DRHI != nullptr;
+			bool bHasSpecularDitherPatternTexture = SpecularDitherPatternTexture2DRHI != nullptr;
 			SpecularDitherPatternTexture2DRHI = SpecularDitherPatternTexture2DRHI ? SpecularDitherPatternTexture2DRHI : GBlackTexture->TextureRHI->GetTexture2D();
+
+			FTextureResource* ShadowHatchingPatternTextureResource = Data.ShadowHatchingPattern ? Data.ShadowHatchingPattern->GetResource() : nullptr;
+			FRHITexture* ShadowHatchingPatternTexture2DRHI = ShadowHatchingPatternTextureResource ? ShadowHatchingPatternTextureResource->GetTexture2DRHI() : nullptr;
+			bool bHasShadowHatchingPatternTexture = ShadowHatchingPatternTexture2DRHI != nullptr;
+			ShadowHatchingPatternTexture2DRHI = ShadowHatchingPatternTexture2DRHI ? ShadowHatchingPatternTexture2DRHI : GBlackTexture->TextureRHI->GetTexture2D();
 			
 			FToonProfileUpdateAtlasTileCS::FParameters* PassParameters = GraphBuilder.AllocParameters<FToonProfileUpdateAtlasTileCS::FParameters>();
 			PassParameters->ToonProfileTexture = ToonProfileUAVSkipBarrier;
 			PassParameters->ToonProfileTextureResolution = ToonProfileTexture->Desc.Extent;
 			PassParameters->ToonProfileIndex = LayerIt;
-			PassParameters->DiffuseDitherPatternStrength = bHasDiffuseDitherPaternTexture ? Data.DiffuseDitherPatternStrength : 0.0f;
+			PassParameters->DiffuseDitherPatternStrength = bHasDiffuseDitherPatternTexture ? Data.DiffuseDitherPatternStrength : 0.0f;
 			PassParameters->DiffuseDitherPatternSize = Data.DiffuseDitherPatternSize;
-			PassParameters->SpecularDitherPatternStrength = bHasSpecularDitherPaternTexture ? Data.SpecularDitherPatternStrength : 0.0f;
+			PassParameters->SpecularDitherPatternStrength = bHasSpecularDitherPatternTexture ? Data.SpecularDitherPatternStrength : 0.0f;
 			PassParameters->SpecularDitherPatternSize = Data.SpecularDitherPatternSize;
 			PassParameters->ToonProfileShadowExtinctionCoefficientInInvCentimeter = Data.GetShadowExtinctionCoefficientInInvCentimeter();
+			PassParameters->ShadowHatchingPatternSize = Data.ShadowHatchingPatternSize;
+			PassParameters->ShadowHatchingPatternStrength = bHasShadowHatchingPatternTexture ? Data.ShadowHatchingPatternStrength : 0.0f;
 
 			auto CreateBufferFRomCurve = [&](FRuntimeCurveLinearColor ColorCurve)
 			{
@@ -386,12 +399,18 @@ IPooledRenderTarget* FToonProfileTextureManager::GetAtlasTexture(FRDGBuilder& Gr
 			FRDGBufferRef SpecularRampBuffer = CreateBufferFRomCurve(Data.SpecularRamp);
 			PassParameters->SpecularRampBufferSRV = GraphBuilder.CreateSRV(SpecularRampBuffer, Format);
 
+			FRDGBufferRef ShadowRampBuffer = CreateBufferFRomCurve(Data.ShadowRamp);
+			PassParameters->ShadowRampBufferSRV = GraphBuilder.CreateSRV(ShadowRampBuffer, Format);
+
 			PassParameters->DiffuseDitherPattern = GraphBuilder.RegisterExternalTexture(CreateRenderTarget(DiffuseDitherPatternTexture2DRHI, TEXT("DiffuseDitherPattern")));
 			PassParameters->DiffuseDitherPatternSampler = TStaticSamplerState<SF_Trilinear, AM_Wrap, AM_Wrap>::GetRHI();
 
 			PassParameters->SpecularDitherPattern = GraphBuilder.RegisterExternalTexture(CreateRenderTarget(SpecularDitherPatternTexture2DRHI, TEXT("SpecularDitherPattern")));
 			PassParameters->SpecularDitherPatternSampler = TStaticSamplerState<SF_Trilinear, AM_Wrap, AM_Wrap>::GetRHI();
 
+			PassParameters->ShadowHatchingPattern = GraphBuilder.RegisterExternalTexture(CreateRenderTarget(ShadowHatchingPatternTexture2DRHI, TEXT("ShadowHatchingPattern")));
+			PassParameters->ShadowHatchingPatternSampler = TStaticSamplerState<SF_Trilinear, AM_Wrap, AM_Wrap>::GetRHI();
+
 			FToonProfileUpdateAtlasTileCS::FPermutationDomain PermutationVector;
 			TShaderMapRef<FToonProfileUpdateAtlasTileCS> Shader(GlobalShaderMap, PermutationVector);
 			FComputeShaderUtils::AddPass(GraphBuilder, 
@@ -464,13 +483,21 @@ FToonProfileStruct::FToonProfileStruct()
 
 	DiffuseDitherPattern = nullptr;
 	DiffuseDitherPatternStrength = 0.0f;	// By default, skip sampling dither pattern.
-	DiffuseDitherPatternSize = 0.1;
+	DiffuseDitherPatternSize = 0.1f;
 
 	SpecularDitherPattern = nullptr;
 	SpecularDitherPatternStrength = 0.0f;	// By default, skip sampling dither pattern.
-	SpecularDitherPatternSize = 0.1;
+	SpecularDitherPatternSize = 0.1f;
 
 	ShadowExtinctionCoefficient = 10.0f;	// Default extinction coefficients that work well.
+
+	ShadowHatchingPattern = nullptr;
+	ShadowHatchingPatternSize = 0.1f;
+	ShadowHatchingPatternStrength = 1.0f;
+
+	AddPointToCurve(ShadowRamp, 0.00, 0.0);
+	AddPointToCurve(ShadowRamp, 1.00, 1.0);
+	
 }
 
 ///////////////////////////////////////////////////////////////////////////////////////////////////
