Corrected inconsistent initial conditionals in void AActor::GetReplicatedCustomConditionState() reversed compared to what is later set in AActor::PreReplication

#Synced-CL 50186587

diff --git a/Engine/Source/Runtime/Engine/Private/ActorReplication.cpp b/Engine/Source/Runtime/Engine/Private/ActorReplication.cpp
--- a/Engine/Source/Runtime/Engine/Private/ActorReplication.cpp
+++ b/Engine/Source/Runtime/Engine/Private/ActorReplication.cpp
@@ -585,8 +585,8 @@ void AActor::GetReplicatedCustomConditionState(FCustomPropertyConditionState& Ou
 {
 	Super::GetReplicatedCustomConditionState(OutActiveState);
 
-	DOREPCUSTOMCONDITION_ACTIVE_FAST(AActor, AttachmentReplication, IsReplicatingMovement());
-	DOREPCUSTOMCONDITION_ACTIVE_FAST(AActor, ReplicatedMovement, RootComponent && !RootComponent->GetIsReplicated());
+	DOREPCUSTOMCONDITION_ACTIVE_FAST(AActor, AttachmentReplication, RootComponent && !RootComponent->GetIsReplicated());
+	DOREPCUSTOMCONDITION_ACTIVE_FAST(AActor, ReplicatedMovement, IsReplicatingMovement());
 }
 
 bool AActor::ReplicateSubobjects(UActorChannel *Channel, FOutBunch *Bunch, FReplicationFlags *RepFlags)
