[SimpleStreamableAssetManager] Fixed TSimpleSparseArray::Add capacity check bug.

#Synced-CL 50205858

diff --git a/Engine/Source/Runtime/Engine/Public/Streaming/SimpleStreamableAssetManager.h b/Engine/Source/Runtime/Engine/Public/Streaming/SimpleStreamableAssetManager.h
--- a/Engine/Source/Runtime/Engine/Public/Streaming/SimpleStreamableAssetManager.h
+++ b/Engine/Source/Runtime/Engine/Public/Streaming/SimpleStreamableAssetManager.h
@@ -225,7 +225,7 @@ private:
 	{
 		int32 Add(T&& InElement)
 		{
-			if (UsedElementsCount == FreeElementIndexHint)
+			if (UsedElementsCount == UsedElements.Num())
 			{
 				UsedElements.Add(false, GSimpleStreamableAssetManagerSparseArrayGrowSize);
 				Elements.AddDefaulted(GSimpleStreamableAssetManagerSparseArrayGrowSize);
@@ -457,4 +457,4 @@ FSimpleStreamableAssetManager::FRegister::FRegister(const TObject* InObject, con
 	}
 }
 
-inline int32 GetTypeHash(const FSimpleStreamableAssetManager::FAssetRecord& Object) { return Object.AssetRegistrationIndex; }
\ No newline at end of file
+inline int32 GetTypeHash(const FSimpleStreamableAssetManager::FAssetRecord& Object) { return Object.AssetRegistrationIndex; }
