In Synch Body lock only when accessing the external acceleration structure

#Synced-CL 50192114

diff --git a/Engine/Source/Runtime/Engine/Private/PhysicsEngine/Experimental/PhysScene_Chaos.cpp b/Engine/Source/Runtime/Engine/Private/PhysicsEngine/Experimental/PhysScene_Chaos.cpp
--- a/Engine/Source/Runtime/Engine/Private/PhysicsEngine/Experimental/PhysScene_Chaos.cpp
+++ b/Engine/Source/Runtime/Engine/Private/PhysicsEngine/Experimental/PhysScene_Chaos.cpp
@@ -80,6 +80,10 @@ FAutoConsoleVariableRef CVar_KinematicDeferralLogInvalidBodies(TEXT("p.Kinematic
 bool bGClusterUnionSyncBodiesMoveNewComponents = true;
 FAutoConsoleVariableRef CVar_GClusterUnionSyncBodiesCheckDirtyFlag(TEXT("p.ClusterUnion.SyncBodiesMoveNewComponents"), bGClusterUnionSyncBodiesMoveNewComponents, TEXT("Enable a fix to ensure new components in a cluster union are moved once on add (even if the cluster is not moving)."));
 
+bool bLockSceneOnlyToUpdateAcc = false;
+FAutoConsoleVariableRef CVar_LockSceneOnlyToUpdateAcc(TEXT("p.LockSceneOnlyToUpdateAcc"), bLockSceneOnlyToUpdateAcc, TEXT("Lock scene only to update the external acceleration structure. The scene won't be lock the entire Synch Body time."));
+
+
 DECLARE_CYCLE_STAT(TEXT("Update Kinematics On Deferred SkelMeshes"), STAT_UpdateKinematicsOnDeferredSkelMeshesChaos, STATGROUP_Physics);
 
 struct FPendingAsyncPhysicsCommand
@@ -2423,7 +2427,16 @@ void FPhysScene_Chaos::OnSyncBodies(Chaos::FPhysicsSolverBase* Solver)
 
 						}
 
-						PhysScene->UpdateActorInAccelerationStructure(BodyInstance->GetPhysicsActor());
+						if (bLockSceneOnlyToUpdateAcc)
+						{
+							FScopedSceneLock_Chaos SceneLock(PhysScene, EPhysicsInterfaceScopedLockType::Write);
+							PhysScene->UpdateActorInAccelerationStructure(BodyInstance->GetPhysicsActor());
+						}
+						else
+						{
+							PhysScene->UpdateActorInAccelerationStructure(BodyInstance->GetPhysicsActor());
+						}
+						
 					}
 
 					if (Proxy->GetWakeEvent() != Chaos::EWakeEventEntry::None && !bPendingMove)
@@ -2478,7 +2491,15 @@ void FPhysScene_Chaos::OnSyncBodies(Chaos::FPhysicsSolverBase* Solver)
 		{
 			if (Proxy && Outer->GetSpacialAcceleration())
 			{
-				UpdateAccelerationStructureFromGeometryCollectionProxy(*Proxy, *Outer->GetSpacialAcceleration());
+				if (bLockSceneOnlyToUpdateAcc)
+				{
+					FScopedSceneLock_Chaos SceneLock(PhysScene, EPhysicsInterfaceScopedLockType::Write);
+					UpdateAccelerationStructureFromGeometryCollectionProxy(*Proxy, *Outer->GetSpacialAcceleration());
+				}
+				else
+				{
+					UpdateAccelerationStructureFromGeometryCollectionProxy(*Proxy, *Outer->GetSpacialAcceleration());
+				}
 			}
 		}
 
@@ -2546,13 +2567,28 @@ void FPhysScene_Chaos::OnSyncBodies(Chaos::FPhysicsSolverBase* Solver)
 					}
 				}
 
-				if (bShouldBeInSQ)
+				if (bLockSceneOnlyToUpdateAcc)
 				{
-					Interface->AddToSpatialAcceleration({ &Handle, 1 }, Outer->GetSpacialAcceleration());
+					FScopedSceneLock_Chaos SceneLock(PhysScene, EPhysicsInterfaceScopedLockType::Write);
+					if (bShouldBeInSQ)
+					{
+						Interface->AddToSpatialAcceleration({ &Handle, 1 }, Outer->GetSpacialAcceleration());
+					}
+					else
+					{
+						Interface->RemoveFromSpatialAcceleration({ &Handle, 1 }, Outer->GetSpacialAcceleration());
+					}
 				}
 				else
 				{
-					Interface->RemoveFromSpatialAcceleration({ &Handle, 1 }, Outer->GetSpacialAcceleration());
+					if (bShouldBeInSQ)
+					{
+						Interface->AddToSpatialAcceleration({ &Handle, 1 }, Outer->GetSpacialAcceleration());
+					}
+					else
+					{
+						Interface->RemoveFromSpatialAcceleration({ &Handle, 1 }, Outer->GetSpacialAcceleration());
+					}
 				}
 
 				DirtyParticle->ClearEvents();
@@ -2563,13 +2599,23 @@ void FPhysScene_Chaos::OnSyncBodies(Chaos::FPhysicsSolverBase* Solver)
 	FDispatcher Dispatcher;
 	auto ExecuteCallback = [this, Solver, &Dispatcher]()
 	{
-		FPhysicsCommand::ExecuteWrite(this, [this, Solver, &Dispatcher](FPhysScene* PhysScene)
+		if (bLockSceneOnlyToUpdateAcc)
 		{
 			Dispatcher.Outer = this;
-			Dispatcher.PhysScene = PhysScene;
+			Dispatcher.PhysScene = this;
 			Dispatcher.Solver = Solver;
 			Solver->PullPhysicsStateForEachDirtyProxy_External(Dispatcher);
-		});
+		}
+		else
+		{
+			FPhysicsCommand::ExecuteWrite(this, [this, Solver, &Dispatcher](FPhysScene* PhysScene)
+			{
+				Dispatcher.Outer = this;
+				Dispatcher.PhysScene = PhysScene;
+				Dispatcher.Solver = Solver;
+				Solver->PullPhysicsStateForEachDirtyProxy_External(Dispatcher);
+			});
+		}
 	};
 
 #if UE_WITH_REMOTE_OBJECT_HANDLE && UE_AUTORTFM
