Default transform providers to enabled.

#Synced-CL 50310951

diff --git a/Engine/Source/Runtime/Engine/Public/Animation/TransformProviderData.h b/Engine/Source/Runtime/Engine/Public/Animation/TransformProviderData.h
--- a/Engine/Source/Runtime/Engine/Public/Animation/TransformProviderData.h
+++ b/Engine/Source/Runtime/Engine/Public/Animation/TransformProviderData.h
@@ -99,5 +99,5 @@ public:
 
 public:
 	UPROPERTY(EditAnywhere, Category = TransformProvider)
-	bool bEnabled;
+	bool bEnabled = true;
 };
