Allow static mesh to enable streaming without assigning an LOD Group

#Synced-CL 50311255

diff --git a/Engine/Source/Runtime/Engine/Classes/Engine/StaticMesh.h b/Engine/Source/Runtime/Engine/Classes/Engine/StaticMesh.h
--- a/Engine/Source/Runtime/Engine/Classes/Engine/StaticMesh.h
+++ b/Engine/Source/Runtime/Engine/Classes/Engine/StaticMesh.h
@@ -579,6 +579,16 @@ enum class EStaticMeshPaintSupport : uint8
 	Disabled,
 };
 
+UENUM()
+enum class EStaticMeshLODStreaming : uint8
+{
+	/** Use the default behavior (LOD Group controls LOD streaming). */
+	Default,
+	/** Enable LOD streaming. */
+	Enabled,
+	/** Disable LOD streaming. */
+	Disabled,
+};
 
 /**
  * A StaticMesh is a piece of geometry that consists of a static set of polygons.
@@ -741,6 +751,10 @@ private:
 	UE_DEPRECATED(5.0, "This must be protected for async build, always use the accessors even internally.")
 		UPROPERTY()
 	int32 LightmapUVVersion;
+
+	UPROPERTY(EditAnywhere, Category = LodSettings)
+	EStaticMeshLODStreaming LODStreaming;
+
 public:
 
 	int32 GetLightmapUVVersion() const
@@ -763,6 +777,11 @@ public:
 		LODGroup = InLODGroup;
 	}
 
+	void SetLODStreaming(EStaticMeshLODStreaming InLODStreaming)
+	{
+		LODStreaming = InLODStreaming;
+	}
+
 	static FName GetLODGroupMemberName()
 	{
 		return GET_MEMBER_NAME_CHECKED(UStaticMesh, LODGroup);
diff --git a/Engine/Source/Runtime/Engine/Private/StaticMesh.cpp b/Engine/Source/Runtime/Engine/Private/StaticMesh.cpp
--- a/Engine/Source/Runtime/Engine/Private/StaticMesh.cpp
+++ b/Engine/Source/Runtime/Engine/Private/StaticMesh.cpp
@@ -3557,8 +3557,19 @@ bool UStaticMesh::GetEnableLODStreaming(const ITargetPlatform* TargetPlatform) c
 		return false;
 	}
 
-	const FStaticMeshLODGroup& LODGroupSettings = TargetPlatform->GetStaticMeshLODSettings().GetLODGroup(GetLODGroup());
-	return LODGroupSettings.IsLODStreamingSupported();
+	switch (LODStreaming)
+	{
+		case EStaticMeshLODStreaming::Enabled:
+			return true;
+		case EStaticMeshLODStreaming::Disabled:
+			return false;
+		case EStaticMeshLODStreaming::Default:
+		default:
+		{
+			const FStaticMeshLODGroup& LODGroupSettings = TargetPlatform->GetStaticMeshLODSettings().GetLODGroup(GetLODGroup());
+			return LODGroupSettings.IsLODStreamingSupported();
+		}
+	}
 }
 
 void UStaticMesh::PostDuplicate(bool bDuplicateForPIE)
