Added tick unregistration logic to ULevel::ClearLevelComponents to make sure that when incremental unregister is not permitted, tick functions are still unregistered from actors.

diff --git a/Engine/Source/Runtime/Engine/Private/Level.cpp b/Engine/Source/Runtime/Engine/Private/Level.cpp
--- a/Engine/Source/Runtime/Engine/Private/Level.cpp
+++ b/Engine/Source/Runtime/Engine/Private/Level.cpp
@@ -141,6 +141,12 @@ static FAutoConsoleVariableRef CVarAllowIncrementalPreUnregisterComponents(
 	GAllowIncrementalPreUnregisterComponents,
 	TEXT("Whether to allow level streaming incremental pre-unregistration of actor components."));
 
+bool GUnregisterTickFunctionsOnClearLevelComponents = true;
+static FAutoConsoleVariableRef CVarUnregisterTickFunctionsOnClearLevelComponents(
+	TEXT("s.UnregisterTickFunctionsOnClearLevelComponents"),
+	GUnregisterTickFunctionsOnClearLevelComponents,
+	TEXT("Whether to also unregister actor tick functions when clearing level components"));
+
 #if WITH_EDITOR
 
 TArray<ULevel::FLevelExternalActorsPathsProviderDelegate> ULevel::LevelExternalActorsPathsProviders;
@@ -1590,6 +1596,13 @@ void ULevel::ClearLevelComponents()
 		if (Actor)
 		{
 			Actor->UnregisterAllComponents();
+
+			if(GUnregisterTickFunctionsOnClearLevelComponents)
+			{
+				// Match tick function unregistration in IncrementalUnregisterComponents to ensure the same behaviour
+				// when incremental unregister is not permitted.
+				Actor->RegisterAllActorTickFunctions(false, false);
+			}
 		}
 	}
 	bIncrementalUnregisterComponentsCompleted = false;
