[Subtitles] - Duration Property parameters greyed out after reopening the Editor
UE-362719
Initialize editor transients via PostLoad, after the underlying subtitle asset data has been loaded.

#Synced-CL 50388552

diff --git a/Engine/Source/Runtime/Engine/Private/Subtitles/SubtitlesAndClosedCaptionsTypes.cpp b/Engine/Source/Runtime/Engine/Private/Subtitles/SubtitlesAndClosedCaptionsTypes.cpp
--- a/Engine/Source/Runtime/Engine/Private/Subtitles/SubtitlesAndClosedCaptionsTypes.cpp
+++ b/Engine/Source/Runtime/Engine/Private/Subtitles/SubtitlesAndClosedCaptionsTypes.cpp
@@ -26,4 +26,15 @@ void USubtitleAssetUserData::PostInitProperties()
 		Subtitle.bCanEditDuration = (Subtitle.SubtitleDurationType == ESubtitleDurationType::UseDurationProperty);
 	}
 }
+
+void USubtitleAssetUserData::PostLoad()
+{
+	// Calling Super:: first: ensures that everything is fully initialized and settled before fiddling with Transient values.
+	Super::PostLoad();
+
+	for (FSubtitleAssetData& Subtitle : Subtitles)
+	{
+		Subtitle.bCanEditDuration = (Subtitle.SubtitleDurationType == ESubtitleDurationType::UseDurationProperty);
+	}
+}
 #endif // WITH_EDITOR
\ No newline at end of file
diff --git a/Engine/Source/Runtime/Engine/Public/Subtitles/SubtitlesAndClosedCaptionsTypes.h b/Engine/Source/Runtime/Engine/Public/Subtitles/SubtitlesAndClosedCaptionsTypes.h
--- a/Engine/Source/Runtime/Engine/Public/Subtitles/SubtitlesAndClosedCaptionsTypes.h
+++ b/Engine/Source/Runtime/Engine/Public/Subtitles/SubtitlesAndClosedCaptionsTypes.h
@@ -132,6 +132,7 @@ public:
 #if WITH_EDITOR
 	virtual void PostEditChangeProperty(FPropertyChangedEvent& PropertyChangedEvent) override;
 	virtual void PostInitProperties() override;
+	virtual void PostLoad() override;
 #endif // WITH_EDITOR
 };
 
