Substrate - NMT
- Fixed debug view mode pixel shader.
- Fixed compilation of some material. Using MaterialInterface as main source of data instead of MaterialResource not always present for some ShaderPlatform.
- Emitter can now error without any expression slected without crashing

#tests PC editor dx12
#rnx

#Synced-CL 50398316

diff --git a/Engine/Shaders/Private/DebugViewModePixelShader.usf b/Engine/Shaders/Private/DebugViewModePixelShader.usf
--- a/Engine/Shaders/Private/DebugViewModePixelShader.usf
+++ b/Engine/Shaders/Private/DebugViewModePixelShader.usf
@@ -135,6 +135,14 @@ MaterialFloat StoreTexSample(in out FTexCoordScalesParams Params, float4 C, int
 	return 1.f;
 }
 
+#define SUBSTRATE_INLINE_SHADING 1
+#include "Substrate/Substrate.ush"
+#if SUBSTRATE_MATERIAL_EXPORT_EXECUTED || SUBSTRATE_GBUFFER_FORMAT==0
+#define MATERIAL_SUBSTRATE_OPAQUE_PRECOMPUTED_LIGHTING 0
+#define SUBSTRATE_MATERIAL_EXPORT_REQUESTED 1
+#include "/Engine/Private/Substrate/SubstrateExport.ush"
+#endif
+
 #include "/Engine/Generated/Material.ush"
 #include "BasePassCommonFunctions.ush"
 #include "DebugViewModeCommon.ush"
diff --git a/Engine/Source/Runtime/Engine/Private/Materials/HLSLMaterialTranslator.cpp b/Engine/Source/Runtime/Engine/Private/Materials/HLSLMaterialTranslator.cpp
--- a/Engine/Source/Runtime/Engine/Private/Materials/HLSLMaterialTranslator.cpp
+++ b/Engine/Source/Runtime/Engine/Private/Materials/HLSLMaterialTranslator.cpp
@@ -688,7 +688,7 @@ FHLSLMaterialTranslator::FHLSLMaterialTranslator(FMaterial* InMaterial,
 	{
 		SubstrateTranslatorData.SubstrateCompilationConfig = *InSubstrateCompilationConfig;
 	}
-	SubstrateTranslatorData.Initialize(Material, MaterialCompilationOutput, this, nullptr);
+	SubstrateTranslatorData.Initialize(Material->GetMaterialInterface(), MaterialCompilationOutput, this, nullptr);
 
 	EnvironmentDefines.Reset(new FEnvironmentDefines);
 
diff --git a/Engine/Source/Runtime/Engine/Private/Materials/MaterialIREmitter.cpp b/Engine/Source/Runtime/Engine/Private/Materials/MaterialIREmitter.cpp
--- a/Engine/Source/Runtime/Engine/Private/Materials/MaterialIREmitter.cpp
+++ b/Engine/Source/Runtime/Engine/Private/Materials/MaterialIREmitter.cpp
@@ -346,10 +346,18 @@ void FEmitter::Error(FStringView Message)
 	Error.Expression = Expression;
 
 	// Add the node type to the error message
-	const int32 ChopCount = FCString::Strlen(TEXT("MaterialExpression"));
-	const FString& ErrorClassName = Expression->GetClass()->GetName();
-	
-	Error.Message.Appendf(TEXT("(Node %s) %s"), *ErrorClassName + ChopCount, Message.GetData());
+	if (Expression)
+	{
+		const int32 ChopCount = FCString::Strlen(TEXT("MaterialExpression"));
+		const FString& ErrorClassName = Expression->GetClass()->GetName();
+
+		Error.Message.Appendf(TEXT("(Node %s) %s"), *ErrorClassName + ChopCount, Message.GetData());
+	}
+	else
+	{
+		// Error without any associated expression being built.
+		Error.Message.Appendf(TEXT("%s"), Message.GetData());
+	}
 	
 	Module->Errors.Push(Error);
 	bCurrentExpressionHasErrors = true;
diff --git a/Engine/Source/Runtime/Engine/Private/Materials/MaterialIRModuleBuilder.cpp b/Engine/Source/Runtime/Engine/Private/Materials/MaterialIRModuleBuilder.cpp
--- a/Engine/Source/Runtime/Engine/Private/Materials/MaterialIRModuleBuilder.cpp
+++ b/Engine/Source/Runtime/Engine/Private/Materials/MaterialIRModuleBuilder.cpp
@@ -140,7 +140,7 @@ struct FMaterialIRModuleBuilderImpl
 		{
 			SubstrateTranslatorData.SubstrateCompilationConfig = *SubstrateCompilationConfig;
 		}
-		SubstrateTranslatorData.Initialize(Builder->MaterialInterface->GetMaterialResource(Module->ShaderPlatform), Module->CompilationOutput, nullptr, &Emitter);
+		SubstrateTranslatorData.Initialize(Builder->MaterialInterface, Module->CompilationOutput, nullptr, &Emitter);
 
 		FMaterialShadingModelField MaterialShadingModels;
 		// If the material gets its shading model from material expressions and we have compiled one or more shading model expressions already, 
@@ -806,10 +806,8 @@ struct FMaterialIRModuleBuilderImpl
 	{
 		MIR::FValue* FrontMaterialPropertyValue = Emitter.SubstrateDefaultSlab();
 
-		FMaterial* CompilerMaterial = SubstrateTranslatorData.GetMaterialResource();
-		if (CompilerMaterial)
 		{
-			UMaterialInterface* MaterialInterface = CompilerMaterial ? CompilerMaterial->GetMaterialInterface() : nullptr;
+			UMaterialInterface* MaterialInterface = SubstrateTranslatorData.GetMaterialInterface();
 			UMaterial* BaseMaterial = MaterialInterface ? MaterialInterface->GetMaterial() : nullptr;
 			UMaterialEditorOnlyData* EditorOnlyData = BaseMaterial ? BaseMaterial->GetEditorOnlyData() : nullptr;
 
@@ -827,7 +825,7 @@ struct FMaterialIRModuleBuilderImpl
 				//		bHasSSS, SSSProfile,
 				//		*SubstrateTranslatorData.GetClearCoatNormal(), *SubstrateTranslatorData.GetCustomTangent());
 			}
-			else if ((CompilerMaterial->GetMaterialDomain() == MD_Surface || CompilerMaterial->GetMaterialDomain() == MD_RuntimeVirtualTexture || CompilerMaterial->GetMaterialDomain() == MD_DeferredDecal) && EditorOnlyData)
+			else if (BaseMaterial && (BaseMaterial->MaterialDomain == MD_Surface || BaseMaterial->MaterialDomain == MD_RuntimeVirtualTexture || BaseMaterial->MaterialDomain == MD_DeferredDecal) && EditorOnlyData)
 			{
 				USubsurfaceProfile* SSSProfile = MaterialInterface ? MaterialInterface->GetSubsurfaceProfileRoot_Internal() : nullptr;
 				const bool bHasSSS = SSSProfile != nullptr;
@@ -845,7 +843,7 @@ struct FMaterialIRModuleBuilderImpl
 					bHasSSS, SSSProfile,
 					EditorOnlyData);
 			}
-			else if (CompilerMaterial->GetMaterialDomain() == MD_Volume && EditorOnlyData)
+			else if (BaseMaterial && BaseMaterial->MaterialDomain == MD_Volume && EditorOnlyData)
 			{
 				check(false); // SUBSTRATE_TODO
 				return FrontMaterialPropertyValue;
@@ -858,24 +856,29 @@ struct FMaterialIRModuleBuilderImpl
 				//		MaterialInterface->GetShadingModels().HasOnlyShadingModel(MSM_Unlit),
 				//		EditorOnlyData);
 			}
-			else if (CompilerMaterial->GetMaterialDomain() == MD_LightFunction && EditorOnlyData)
+			else if (BaseMaterial && BaseMaterial->MaterialDomain == MD_LightFunction && EditorOnlyData)
 			{
 				check(false); // SUBSTRATE_TODO
 				return FrontMaterialPropertyValue;
 				//	return UMaterialExpressionSubstrateLightFunction::CompileCommon(this, EditorOnlyData->EmissiveColor, EditorOnlyData);
 			}
-			else if (CompilerMaterial->GetMaterialDomain() == MD_PostProcess && EditorOnlyData)
+			else if (BaseMaterial && BaseMaterial->MaterialDomain == MD_PostProcess && EditorOnlyData)
 			{
 				check(false); // SUBSTRATE_TODO
 				return FrontMaterialPropertyValue;
 				//	return UMaterialExpressionSubstratePostProcess::CompileCommon(this, EditorOnlyData->EmissiveColor, EditorOnlyData->Opacity, EditorOnlyData);
 			}
-			else if (CompilerMaterial->GetMaterialDomain() == MD_UI && EditorOnlyData)
+			else if (BaseMaterial && BaseMaterial->MaterialDomain == MD_UI && EditorOnlyData)
 			{
 				check(false); // SUBSTRATE_TODO
 				return FrontMaterialPropertyValue;
 				//	return UMaterialExpressionSubstrateUI::CompileCommon(this, EditorOnlyData->EmissiveColor, EditorOnlyData->Opacity, EditorOnlyData);
 			}
+			else
+			{
+				Emitter.Error(TEXT("Substrate could not compile the material."));
+				return FrontMaterialPropertyValue;
+			}
 
 
 			// When we are running hidden conversion (bMaterialUsesRootNodeToSubstrateHiddenConversion),
@@ -885,11 +888,6 @@ struct FMaterialIRModuleBuilderImpl
 			Emitter.SetCustomOutputs(TEXTVIEW("SubstrateFrontMaterial"), { &PropertyValueRef, 1 }, MIR::EMaterialOutputFrequency::PerPixel);
 			// We also copy the copy the shared local bases structure from within GeneratePixelInputInitializerHLSLCode.
 		}
-		else
-		{
-			check(false); // SUBSTRATE_TODO should not happen
-			return FrontMaterialPropertyValue;
-		}
 
 		return FrontMaterialPropertyValue;
 	}
diff --git a/Engine/Source/Runtime/Engine/Private/Materials/SubstrateTranslatorCommon.cpp b/Engine/Source/Runtime/Engine/Private/Materials/SubstrateTranslatorCommon.cpp
--- a/Engine/Source/Runtime/Engine/Private/Materials/SubstrateTranslatorCommon.cpp
+++ b/Engine/Source/Runtime/Engine/Private/Materials/SubstrateTranslatorCommon.cpp
@@ -285,11 +285,9 @@ void FSubstrateCompilationContext::Initialise()
 
 bool FSubstrateCompilationContext::SubstrateGenerateDerivedMaterialOperatorData(FSubstrateTranslatorData& SubstrateTranslatorData)
 {
-	FMaterial* CompilerMaterial = SubstrateTranslatorData.GetMaterialResource();
-
 	if (SubstrateMaterialExpressionRegisteredOperators.IsEmpty())
 	{
-		SubstrateTranslatorData.Errorf(FString::Printf(TEXT("Could not find any Substrate operators or BSDFs in Material %s (asset: %s).\r\n"), *CompilerMaterial->GetAssetName(), *CompilerMaterial->GetAssetPath().ToString()));
+		SubstrateTranslatorData.Errorf(FString::Printf(TEXT("Could not find any Substrate operators or BSDFs in Material %s (asset: %s).\r\n"), *SubstrateTranslatorData.GetAssetName(), *SubstrateTranslatorData.GetAssetPath()));
 		return false;
 	}
 
@@ -349,19 +347,22 @@ bool FSubstrateCompilationContext::SubstrateGenerateDerivedMaterialOperatorData(
 
 		if (bMustHaveLeftChild && It.LeftIndex == INDEX_NONE)
 		{
-			SubstrateTranslatorData.Errorf(FString::Printf(TEXT("A Substrate Operator %s node is missing its first input from material %s (asset: %s).\r\n"), GetSubstrateOperatorStr(It.OperatorType), *CompilerMaterial->GetAssetName(), *CompilerMaterial->GetAssetPath().ToString()));
+			SubstrateTranslatorData.Errorf(FString::Printf(TEXT("A Substrate Operator %s node is missing its first input from material %s (asset: %s).\r\n"), GetSubstrateOperatorStr(It.OperatorType), *SubstrateTranslatorData.GetAssetName(), *SubstrateTranslatorData.GetAssetPath()));
 			return false;
 		}
 		if (bMustHaveRightChild && It.RightIndex == INDEX_NONE)
 		{
-			SubstrateTranslatorData.Errorf(FString::Printf(TEXT("A Substrate Operator %s node is missing its second input from material %s (asset: %s).\r\n"), GetSubstrateOperatorStr(It.OperatorType), *CompilerMaterial->GetAssetName(), *CompilerMaterial->GetAssetPath().ToString()));
+			SubstrateTranslatorData.Errorf(FString::Printf(TEXT("A Substrate Operator %s node is missing its second input from material %s (asset: %s).\r\n"), GetSubstrateOperatorStr(It.OperatorType), *SubstrateTranslatorData.GetAssetName(), *SubstrateTranslatorData.GetAssetPath()));
 			return false;
 		}
 	}
-	SubstrateMaterialRootOperator = &SubstrateMaterialExpressionRegisteredOperators[RootIndex];
+	if (RootIndex != INDEX_NONE)
+	{
+		SubstrateMaterialRootOperator = &SubstrateMaterialExpressionRegisteredOperators[RootIndex];
+	}
 	if (!SubstrateMaterialRootOperator)
 	{
-		SubstrateTranslatorData.Errorf(FString::Printf(TEXT("Cannot find the root of the Substrate Tree for Material %s (asset: %s).\r\n"), *CompilerMaterial->GetAssetName(), *CompilerMaterial->GetAssetPath().ToString()));
+		SubstrateTranslatorData.Errorf(FString::Printf(TEXT("Cannot find the root of the Substrate Tree for Material %s (asset: %s).\r\n"), *SubstrateTranslatorData.GetAssetName(), *SubstrateTranslatorData.GetAssetPath()));
 		return false;
 	}
 
@@ -454,7 +455,7 @@ bool FSubstrateCompilationContext::SubstrateGenerateDerivedMaterialOperatorData(
 			}
 			else
 			{
-				SubstrateTranslatorData.Errorf(FString::Printf(TEXT("Substrate error: unhandled operator type from ValidateAllowedList. (asset: %s).\r\n"), *CompilerMaterial->GetAssetName()));
+				SubstrateTranslatorData.Errorf(FString::Printf(TEXT("Substrate error: unhandled operator type from ValidateAllowedList. (asset: %s).\r\n"), *SubstrateTranslatorData.GetAssetName()));
 				return false;
 			}
 
@@ -510,7 +511,7 @@ bool FSubstrateCompilationContext::SubstrateGenerateDerivedMaterialOperatorData(
 				}
 				NotAllowedTypeNames += "}";
 
-				SubstrateTranslatorData.Errorf(FString::Printf(TEXT("Substrate error: Operator %s cannot be used with the following data flowing through %s. (asset: %s).\r\n"), *OperatorName, *NotAllowedTypeNames, *CompilerMaterial->GetAssetName()));
+				SubstrateTranslatorData.Errorf(FString::Printf(TEXT("Substrate error: Operator %s cannot be used with the following data flowing through %s. (asset: %s).\r\n"), *OperatorName, *NotAllowedTypeNames, *SubstrateTranslatorData.GetAssetName()));
 				return false;
 			}
 			return true;
@@ -540,7 +541,7 @@ bool FSubstrateCompilationContext::SubstrateGenerateDerivedMaterialOperatorData(
 				}
 				else
 				{
-					SubstrateTranslatorData.Errorf(FString::Printf(TEXT("Substrate error: Select node can only be parameter blended as of today. (asset: %s).\r\n"), *CompilerMaterial->GetAssetName()));
+					SubstrateTranslatorData.Errorf(FString::Printf(TEXT("Substrate error: Select node can only be parameter blended as of today. (asset: %s).\r\n"), *SubstrateTranslatorData.GetAssetName()));
 				}
 				break;
 			}
@@ -571,7 +572,7 @@ bool FSubstrateCompilationContext::SubstrateGenerateDerivedMaterialOperatorData(
 					}
 					else
 					{
-						SubstrateTranslatorData.Errorf(FString::Printf(TEXT("Substrate error: unhandled Unlit BSDF sub type from WalkParameterBlendedSubTree. (asset: %s).\r\n"), *CompilerMaterial->GetAssetName()));
+						SubstrateTranslatorData.Errorf(FString::Printf(TEXT("Substrate error: unhandled Unlit BSDF sub type from WalkParameterBlendedSubTree. (asset: %s).\r\n"), *SubstrateTranslatorData.GetAssetName()));
 					}
 				}
 				else if (CurrentOperator.BSDFType == SUBSTRATE_BSDF_TYPE_SLAB)
@@ -600,7 +601,7 @@ bool FSubstrateCompilationContext::SubstrateGenerateDerivedMaterialOperatorData(
 				}
 				else
 				{
-					SubstrateTranslatorData.Errorf(FString::Printf(TEXT("Substrate error: unhandled BSDF type from WalkParameterBlendedSubTree. (asset: %s).\r\n"), *CompilerMaterial->GetAssetName()));
+					SubstrateTranslatorData.Errorf(FString::Printf(TEXT("Substrate error: unhandled BSDF type from WalkParameterBlendedSubTree. (asset: %s).\r\n"), *SubstrateTranslatorData.GetAssetName()));
 				}
 				break;
 			}
@@ -612,7 +613,7 @@ bool FSubstrateCompilationContext::SubstrateGenerateDerivedMaterialOperatorData(
 		SubstrateTreeValidation = WalkOperators(*SubstrateMaterialRootOperator);
 		if (!SubstrateTreeValidation.bValid)
 		{
-			SubstrateTranslatorData.Errorf(FString::Printf(TEXT("Substrate error: invalid Substrate node graph. (asset: %s).\r\n"), *CompilerMaterial->GetAssetName()));
+			SubstrateTranslatorData.Errorf(FString::Printf(TEXT("Substrate error: invalid Substrate node graph. (asset: %s).\r\n"), *SubstrateTranslatorData.GetAssetName()));
 			return false;
 		}
 	}
@@ -715,7 +716,7 @@ bool FSubstrateCompilationContext::SubstrateGenerateDerivedMaterialOperatorData(
 		}
 		else
 		{
-			SubstrateTranslatorData.Errorf(FString::Printf(TEXT("Unkown Substrate material simplification status error for %s (asset: %s).\r\n"), *CompilerMaterial->GetAssetName(), *CompilerMaterial->GetAssetPath().ToString()));
+			SubstrateTranslatorData.Errorf(FString::Printf(TEXT("Unkown Substrate material simplification status error for %s (asset: %s).\r\n"), *SubstrateTranslatorData.GetAssetName(), *SubstrateTranslatorData.GetAssetPath()));
 			return false;
 		}
 
@@ -816,7 +817,7 @@ bool FSubstrateCompilationContext::SubstrateGenerateDerivedMaterialOperatorData(
 			if (CompilationContextIndex == ESubstrateCompilationContext::SCC_Default)
 			{
 				// Only write those data for the default material
-				const bool bIsOpaqueOrMasked = IsOpaqueOrMaskedBlendMode(*CompilerMaterial);
+				const bool bIsOpaqueOrMasked = IsOpaqueOrMaskedBlendMode(*SubstrateTranslatorData.GetMaterialInterface());
 				SubstrateTranslatorData.bSubstrateMaterialOutputOpaqueRoughRefractions = bSubstrateUsesVerticalLayering && bIsOpaqueOrMasked;
 			}
 			bSubstrateMaterialIsUnlitNode = bHasUnlit;
@@ -1271,7 +1272,7 @@ bool FSubstrateCompilationContext::SubstrateGenerateDerivedMaterialOperatorData(
 					}
 					case SUBSTRATE_BSDF_TYPE_SINGLELAYERWATER:
 					{
-						SubstrateTranslatorData.Errorf(FString::Printf(TEXT("Substrate error: single layer water should go through through a dedicated fast path in %s (asset: %s).\r\n"), *CompilerMaterial->GetAssetName(), *CompilerMaterial->GetAssetPath().ToString()));
+						SubstrateTranslatorData.Errorf(FString::Printf(TEXT("Substrate error: single layer water should go through through a dedicated fast path in %s (asset: %s).\r\n"), *SubstrateTranslatorData.GetAssetName(), *SubstrateTranslatorData.GetAssetPath()));
 						return false;
 					}
 					case SUBSTRATE_BSDF_TYPE_TOON:
@@ -1289,7 +1290,7 @@ bool FSubstrateCompilationContext::SubstrateGenerateDerivedMaterialOperatorData(
 					}
 					default:
 					{
-						SubstrateTranslatorData.Errorf(FString::Printf(TEXT("Unkownd BSDF type encountered in %s (asset: %s).\r\n"), *CompilerMaterial->GetAssetName(), *CompilerMaterial->GetAssetPath().ToString()));
+						SubstrateTranslatorData.Errorf(FString::Printf(TEXT("Unkownd BSDF type encountered in %s (asset: %s).\r\n"), *SubstrateTranslatorData.GetAssetName(), *SubstrateTranslatorData.GetAssetPath()));
 						return false;
 					}
 					}
@@ -1301,7 +1302,7 @@ bool FSubstrateCompilationContext::SubstrateGenerateDerivedMaterialOperatorData(
 			if (OperatorCount > SUBSTRATE_MAX_OPERATOR_COUNT)
 			{
 				// Why do we have an operator limit: due to the size of the array of Operator in FSubstrateTre and the way the Substrate tree is exported for advanced debug purpose. Parameter blending can help working around that. 
-				SubstrateTranslatorData.Errorf(FString::Printf(TEXT("Material %s have too many Substrate Operators (asset: %s): %d / %d. Please note that BSDFs also count as an operator. Use parameter blending to workaround that limitation.\r\n"), *CompilerMaterial->GetAssetName(), *CompilerMaterial->GetAssetPath().ToString(), OperatorCount, SUBSTRATE_MAX_OPERATOR_COUNT));
+				SubstrateTranslatorData.Errorf(FString::Printf(TEXT("Material %s have too many Substrate Operators (asset: %s): %d / %d. Please note that BSDFs also count as an operator. Use parameter blending to workaround that limitation.\r\n"), *SubstrateTranslatorData.GetAssetName(), *SubstrateTranslatorData.GetAssetPath(), OperatorCount, SUBSTRATE_MAX_OPERATOR_COUNT));
 				return false;
 			}
 
@@ -1309,7 +1310,7 @@ bool FSubstrateCompilationContext::SubstrateGenerateDerivedMaterialOperatorData(
 			if (!SubstrateSimplificationStatus.bMaterialFitsInMemoryBudget && SubstrateSimplificationStatus.bFullSimplificationStepHasBeenRun && SubstrateSimplificationStatus.bSlabSimplificationStepHasBeenRun)
 			{
 				// If we have already run the full simplification but the material still does not fit in memory, we must fail the material compilation.
-				SubstrateTranslatorData.Errorf(FString::Printf(TEXT("Material %s could not be simplified to fit in Substrate per pixel (asset: %s).\r\n"), *CompilerMaterial->GetAssetName(), *CompilerMaterial->GetAssetPath().ToString()));
+				SubstrateTranslatorData.Errorf(FString::Printf(TEXT("Material %s could not be simplified to fit in Substrate per pixel (asset: %s).\r\n"), *SubstrateTranslatorData.GetAssetName(), *SubstrateTranslatorData.GetAssetPath()));
 				return false;
 			}
 			if (!bRequestMaterialDetailsSet)
@@ -1340,24 +1341,25 @@ bool FSubstrateCompilationContext::SubstrateGenerateDerivedMaterialOperatorData(
 				SubstrateTranslatorData.GetMaterialCompilationOutput().SubstrateMaterialCompilationOutput.RequestedClosurePerPixel = SubstrateSimplificationStatus.OriginalRequestedClosureCount;
 				SubstrateTranslatorData.GetMaterialCompilationOutput().SubstrateMaterialCompilationOutput.PlatformClosurePixel = SubstrateClosurePerPixel;
 
-				SubstrateTranslatorData.GetMaterialCompilationOutput().SubstrateMaterialCompilationOutput.bIsThin = CompilerMaterial->IsThinSurface() ? 1 : 0;
+				UMaterialInterface* MaterialInterface = SubstrateTranslatorData.GetMaterialInterface();
+				SubstrateTranslatorData.GetMaterialCompilationOutput().SubstrateMaterialCompilationOutput.bIsThin = MaterialInterface && MaterialInterface->IsThinSurface() ? 1 : 0;
 
 				uint32 BSDFTypeCount = FMath::CountBits(SubstrateTreeValidation.MaterialType & MTV_ValidMultipleBSDFs);
 
 				// The order of ifs here is important.
-				if (CompilerMaterial->IsLightFunction())
+				if (MaterialInterface && MaterialInterface->IsLightFunctionMaterial())
 				{
 					SubstrateTranslatorData.GetMaterialCompilationOutput().SubstrateMaterialCompilationOutput.MaterialType = SUBSTRATE_MATERIAL_TYPE_LIGHTFUNCTION;
 				}
-				else if (CompilerMaterial->IsPostProcessMaterial())
+				else if (MaterialInterface && MaterialInterface->IsPostProcessMaterial())
 				{
 					SubstrateTranslatorData.GetMaterialCompilationOutput().SubstrateMaterialCompilationOutput.MaterialType = SUBSTRATE_MATERIAL_TYPE_POSTPROCESS;
 				}
-				else if (CompilerMaterial->IsUIMaterial())
+				else if (MaterialInterface && MaterialInterface->IsUIMaterial())
 				{
 					SubstrateTranslatorData.GetMaterialCompilationOutput().SubstrateMaterialCompilationOutput.MaterialType = SUBSTRATE_MATERIAL_TYPE_UI;
 				}
-				else if (CompilerMaterial->IsDeferredDecal())
+				else if (MaterialInterface && MaterialInterface->IsDeferredDecal())
 				{
 					SubstrateTranslatorData.GetMaterialCompilationOutput().SubstrateMaterialCompilationOutput.MaterialType = SUBSTRATE_MATERIAL_TYPE_DECAL;
 				}
@@ -1457,7 +1459,6 @@ FSubstrateTranslatorData::FSubstrateTranslatorData()
 , bSubstrateUsesConversionFromLegacy(false)
 , bSubstrateMaterialOutputOpaqueRoughRefractions(false)
 , SubstrateCompilationConfig()
-, MaterialResource(nullptr)
 , MaterialCompilationOutput(nullptr)
 , CompilerLegacy(nullptr)
 , CompilerNew(nullptr)
@@ -1470,15 +1471,23 @@ FSubstrateTranslatorData::~FSubstrateTranslatorData()
 }
 
 
-void FSubstrateTranslatorData::Initialize(FMaterial* InCompilerMaterial, FMaterialCompilationOutput& InMaterialCompilationOutput, FMaterialCompiler* InCompilerLegacy, MIR::FEmitter* InCompilerNew)
+void FSubstrateTranslatorData::Initialize(UMaterialInterface* InMaterialInterface, FMaterialCompilationOutput& InMaterialCompilationOutput, FMaterialCompiler* InCompilerLegacy, MIR::FEmitter* InCompilerNew)
 {
 	check(!!InCompilerLegacy != !!InCompilerNew); // XOR
 
 	FrontMaterialExpr = nullptr;
-	MaterialResource = InCompilerMaterial;	// We handle the case where MaterialResource is nullptr down the line.
+
 	MaterialCompilationOutput = &InMaterialCompilationOutput;
 	CompilerLegacy = InCompilerLegacy;
 	CompilerNew = InCompilerNew;
+
+	MaterialInterface = InMaterialInterface;
+
+	if (MaterialInterface)
+	{
+		AssetName = MaterialInterface->GetName();
+		AssetPath = MaterialInterface->GetPathName();
+	}
 }
 
 void FSubstrateTranslatorData::Errorf(const FString& Format)
@@ -1512,10 +1521,6 @@ void FSubstrateTranslatorData::ProcessSubstrateTreeAndTopology(
 	UMaterialExpression* ExpressionToPreview,
 	bool& bDoClearFunctionStack)
 {
-	if (!MaterialResource)
-		return;
-
-
 	EShaderPlatform Platform = GetShaderPlatform();
 	bSubstrateEnabled = Substrate::IsSubstrateEnabled();
 	bProjectSubstrateHiddenMaterialAssetConversionEnabled = Substrate::IsHiddenMaterialAssetConversionEnabled();
@@ -1543,7 +1548,7 @@ void FSubstrateTranslatorData::ProcessSubstrateTreeAndTopology(
 	}
 	else
 	{
-		FSubstrateMaterialInput* FrontMaterialInput = MaterialResource->GetMaterialInterface() ? &MaterialResource->GetMaterialInterface()->GetMaterial()->GetEditorOnlyData()->FrontMaterial : nullptr;
+		FSubstrateMaterialInput* FrontMaterialInput = MaterialInterface ? &MaterialInterface->GetMaterial()->GetEditorOnlyData()->FrontMaterial : nullptr;
 		FrontMaterialExpr = FrontMaterialInput ? FrontMaterialInput->GetTracedInput().Expression : nullptr;
 		FrontMaterialOutputIndex = FrontMaterialInput ? FrontMaterialInput->OutputIndex : INDEX_NONE;
 	}
@@ -1599,7 +1604,6 @@ void FSubstrateTranslatorData::ProcessSubstrateTreeAndTopology(
 
 	if (bMaterialUsesRootNodeToSubstrateHiddenConversion)
 	{
-		UMaterialInterface* MaterialInterface = MaterialResource->GetMaterialInterface();
 		UMaterial* BaseMaterial = MaterialInterface ? MaterialInterface->GetMaterial() : nullptr;
 		UMaterialEditorOnlyData* EditorOnlyData = BaseMaterial ? BaseMaterial->GetEditorOnlyData() : nullptr;
 
@@ -1610,8 +1614,11 @@ void FSubstrateTranslatorData::ProcessSubstrateTreeAndTopology(
 			UMaterialExpressionClearCoatNormalCustomOutput* ClearCoatBottomNormalOutput = nullptr;
 			UMaterialExpressionTangentOutput* TangentOutput = nullptr;
 			TArray<class UMaterialExpressionCustomOutput*> CustomOutputExpressions;
-			if(BaseMaterial)
+			if (BaseMaterial)
+			{
 				BaseMaterial->GetAllCustomOutputExpressions(CustomOutputExpressions);
+			}
+
 			for (UMaterialExpressionCustomOutput* Expression : CustomOutputExpressions)
 			{
 				// Gather custom output for thin translucency
@@ -1685,7 +1692,7 @@ void FSubstrateTranslatorData::ProcessSubstrateTreeAndTopology(
 						EditorOnlyData->ShadingModelFromMaterialExpression.IsConnected(),
 						EditorOnlyData->MaterialAttributes.IsConnected(MP_EmissiveColor));
 				}
-				else if ((MaterialResource->GetMaterialDomain() == MD_Surface || MaterialResource->GetMaterialDomain() == MD_RuntimeVirtualTexture || MaterialResource->GetMaterialDomain() == MD_DeferredDecal) && EditorOnlyData)
+				else if (BaseMaterial &&  (BaseMaterial->MaterialDomain == MD_Surface || BaseMaterial->MaterialDomain == MD_RuntimeVirtualTexture || BaseMaterial->MaterialDomain == MD_DeferredDecal) && EditorOnlyData)
 				{
 					FSubstrateOperator* Op = UMaterialExpressionSubstrateShadingModels::SubstrateGenerateMaterialTopologyTreeCommon(
 						*this, RootNodeConversionGUID, nullptr /*Parent*/, 0,
@@ -1695,7 +1702,7 @@ void FSubstrateTranslatorData::ProcessSubstrateTreeAndTopology(
 						*CustomTangent,
 						EditorOnlyData->ShadingModelFromMaterialExpression);
 
-					if (MaterialResource->GetMaterialDomain() == MD_DeferredDecal)
+					if (BaseMaterial && BaseMaterial->MaterialDomain == MD_DeferredDecal)
 					{
 						// We notify that this is for decal using this usage flag.
 						// Otherwise, we do not need to do anything else, the ConvertToDecal node is only here to enforce parameter blending.
@@ -1703,24 +1710,24 @@ void FSubstrateTranslatorData::ProcessSubstrateTreeAndTopology(
 						Op->SubUsage |= SUBSTRATE_OPERATOR_SUBUSAGE_DECAL;
 					}
 				}
-				else if (MaterialResource->GetMaterialDomain() == MD_Volume && EditorOnlyData)
+				else if (BaseMaterial && BaseMaterial->MaterialDomain == MD_Volume && EditorOnlyData)
 				{
 					UMaterialExpressionSubstrateVolumetricFogCloudBSDF::SubstrateGenerateMaterialTopologyTreeCommon(
 						*this, RootNodeConversionGUID, nullptr /*Parent*/, 0,
 						EditorOnlyData->EmissiveColor,
 						EditorOnlyData->AmbientOcclusion);
 				}
-				else if (MaterialResource->GetMaterialDomain() == MD_LightFunction && EditorOnlyData)
+				else if (BaseMaterial && BaseMaterial->MaterialDomain == MD_LightFunction && EditorOnlyData)
 				{
 					UMaterialExpressionSubstrateLightFunction::SubstrateGenerateMaterialTopologyTreeCommon(
 						*this, RootNodeConversionGUID, nullptr /*Parent*/, 0);
 				}
-				else if (MaterialResource->GetMaterialDomain() == MD_PostProcess && EditorOnlyData)
+				else if (BaseMaterial && BaseMaterial->MaterialDomain == MD_PostProcess && EditorOnlyData)
 				{
 					UMaterialExpressionSubstratePostProcess::SubstrateGenerateMaterialTopologyTreeCommon(
 						*this, RootNodeConversionGUID, nullptr /*Parent*/, 0);
 				}
-				else if (MaterialResource->GetMaterialDomain() == MD_UI && EditorOnlyData)
+				else if (BaseMaterial && BaseMaterial->MaterialDomain == MD_UI && EditorOnlyData)
 				{
 					UMaterialExpressionSubstrateUI::SubstrateGenerateMaterialTopologyTreeCommon(
 						*this, RootNodeConversionGUID, nullptr /*Parent*/, 0);
@@ -1732,7 +1739,7 @@ void FSubstrateTranslatorData::ProcessSubstrateTreeAndTopology(
 
 				if (SubstrateCtx.bSubstrateTreeOutOfStackDepthOccurred)
 				{
-					Errorf(FString::Printf(TEXT(" %s [%s]: Substrate - Cyclic graph detected when we only support acyclic graph."), *MaterialResource->GetAssetName(), *MaterialResource->GetAssetPath().ToString()));
+					Errorf(FString::Printf(TEXT(" %s [%s]: Substrate - Cyclic graph detected when we only support acyclic graph."), *GetAssetName(), *GetAssetPath()));
 					return;
 				}
 				if (!SubstrateCtx.SubstrateGenerateDerivedMaterialOperatorData(*this))
@@ -1755,7 +1762,8 @@ void FSubstrateTranslatorData::ProcessSubstrateTreeAndTopology(
 			FSubstrateCompilationContext& SubstrateCtx = SubstrateCompilationContext[SubstrateCompilationContextIndex];
 			SubstrateCtx = FSubstrateCompilationContext(ESubstrateCompilationContext(SubstrateCompilationContextIndex));
 
-			FScalarMaterialInput* SurfaceThickness = MaterialResource->IsThinSurface() && MaterialResource->GetMaterialInterface() ? &MaterialResource->GetMaterialInterface()->GetMaterial()->GetEditorOnlyData()->SurfaceThickness : nullptr;
+			UMaterial* Material = MaterialInterface ? MaterialInterface->GetMaterial() : nullptr;
+			FScalarMaterialInput* SurfaceThickness = Material && Material->IsThinSurface() ? &Material->GetEditorOnlyData()->SurfaceThickness : nullptr;
 
 			SubstrateThicknessStackPush(nullptr, SurfaceThickness);
 
@@ -1766,7 +1774,7 @@ void FSubstrateTranslatorData::ProcessSubstrateTreeAndTopology(
 
 			if (SubstrateCtx.bSubstrateTreeOutOfStackDepthOccurred)
 			{
-				Errorf(FString::Printf(TEXT(" %s [%s]: Substrate - Cyclic graph detected when we only support acyclic graph."), *MaterialResource->GetAssetName(), *MaterialResource->GetAssetPath().ToString()));
+				Errorf(FString::Printf(TEXT(" %s [%s]: Substrate - Cyclic graph detected when we only support acyclic graph."), *GetAssetName(), *GetAssetPath()));
 				return;
 			}
 			if (!SubstrateCtx.SubstrateGenerateDerivedMaterialOperatorData(*this))
@@ -1811,11 +1819,10 @@ bool FSubstrateTranslatorData::GenerateSubstrateTreeHLSLCode(bool bSubstrateFron
 				// The material can be null when some entries are automatically generated, for instance in the material layer blending system
 				bSubstrateFrontMaterialIsValid &= SubstrateCtx.SubstrateMaterialEffectiveClosureCount > 0;
 
-				UMaterialInterface* MaterialInterface = MaterialResource->GetMaterialInterface();
-				UMaterial* BaseMaterial = MaterialInterface ? MaterialInterface->GetMaterial() : nullptr;
-				if (!bSubstrateFrontMaterialIsValid && BaseMaterial && BaseMaterial->bUseMaterialAttributes)
+				UMaterial* Material = MaterialInterface ? MaterialInterface->GetMaterial() : nullptr;
+				if (!bSubstrateFrontMaterialIsValid && Material && Material->bUseMaterialAttributes)
 				{
-					Errorf(FString::Printf(TEXT(" %s [%s]: Substrate - Material has no BSDF: this can happen when the SubstrateConvertMaterialAttributes node is not used and when material attributes are enabled."), *MaterialResource->GetAssetName(), *MaterialResource->GetAssetPath().ToString()));
+					Errorf(FString::Printf(TEXT(" %s [%s]: Substrate - Material has no BSDF: this can happen when the SubstrateConvertMaterialAttributes node is not used and when material attributes are enabled."), *GetAssetName(), *GetAssetPath()));
 					return false;
 				}
 			}
@@ -2050,7 +2057,7 @@ bool FSubstrateTranslatorData::GenerateSubstrateTreeHLSLCode(bool bSubstrateFron
 					SubstrateCtx.SubstrateEvaluateSharedLocalBases(RequestedSharedLocalBasesCount, nullptr, nullptr);// CompilerNew);
 					if (RequestedSharedLocalBasesCount > SUBSTRATE_MAX_SHAREDLOCALBASES_REGISTERS)
 					{
-						Errorf(FString::Printf(TEXT(" %s [%s]: Substrate - Material has more unique normal/tangent basis than the allowed limit %d/%d."), *MaterialResource->GetAssetName(), *MaterialResource->GetAssetPath().ToString(), RequestedSharedLocalBasesCount, SUBSTRATE_MAX_SHAREDLOCALBASES_REGISTERS));
+						Errorf(FString::Printf(TEXT(" %s [%s]: Substrate - Material has more unique normal/tangent basis than the allowed limit %d/%d."), *GetAssetName(), *GetAssetPath(), RequestedSharedLocalBasesCount, SUBSTRATE_MAX_SHAREDLOCALBASES_REGISTERS));
 					}
 				}
 			}
@@ -2148,7 +2155,8 @@ void FSubstrateTranslatorData::GenerateMaterialCompilationOutputAndDefines(
 {
 	bSubstratePremultipliedAlphaOpacityOverridden = bMaterialIsSubstrate && bOpacityPropertyIsUsed;
 
-	bDistortionAccountForCoverage = bSubstrateEnabled && !bMaterialUsesRootNodeToSubstrateHiddenConversion && MaterialResource->GetRefractionCoverageMode() == RCM_CoverageAccountedFor;
+	UMaterial* Material = MaterialInterface ? MaterialInterface->GetMaterial() : nullptr;
+	bDistortionAccountForCoverage = bSubstrateEnabled && !bMaterialUsesRootNodeToSubstrateHiddenConversion && Material && Material->RefractionCoverageMode == RCM_CoverageAccountedFor;
 
 	if (bMaterialIsSubstrate)
 	{
@@ -2159,7 +2167,7 @@ void FSubstrateTranslatorData::GenerateMaterialCompilationOutputAndDefines(
 		// Unlit cannot be combined with other BSDF so we can simply pick the default strata context
 		bSubstrateOptimizedUnlit = SubstrateCompilationContext[ESubstrateCompilationContext::SCC_Default].bSubstrateMaterialIsUnlitNode;
 
-		bSubstrateRoughnessTracking = MaterialResource->HasSubstrateRoughnessTracking();
+		bSubstrateRoughnessTracking = Material && Material->HasSubstrateRoughnessTracking();
 
 		{
 			// For now, the fully simplified mode is used for Lumen or anything else supported inlined evaluation. The export is only valid for the default case.
@@ -2384,21 +2392,21 @@ FSubstrateOperator& FSubstrateTranslatorData::SubstrateCompilationRegisterOperat
 	if (SubstrateCtx.SubstrateMaterialExpressionToOperatorIndex.Find(SubstrateExpressionGuid))
 	{
 		// It is not possible to register/use a Substrate BSDF multiple times with this same exact graph path. (that would break the Substrate tree code generation)
-		Errorf(FString::Printf(TEXT("Material %s: It is not possible to uses a Substrate BSDF (or any ouput of type SubstrateData) multiple times within a Substrate material topology with the same graph path GUID (asset: %s).\r\n"), *MaterialResource->GetAssetName(), *MaterialResource->GetAssetPath().ToString()));
+		Errorf(FString::Printf(TEXT("Material %s: It is not possible to uses a Substrate BSDF (or any ouput of type SubstrateData) multiple times within a Substrate material topology with the same graph path GUID (asset: %s).\r\n"), *GetAssetName(), *GetAssetPath()));
 		return DefaultOperatorOnError;
 	}
 
 	const uint32 NewOperatorIndex = SubstrateCtx.SubstrateMaterialExpressionToOperatorIndex.Num();
 	if (NewOperatorIndex >= SUBSTRATE_MAX_COMPILER_REGISTERED_OPERATOR_COUNT)
 	{
-		Errorf(FString::Printf(TEXT("Material %s have too many Substrate Operators: the compiler is failing (asset: %s).\r\n"), *MaterialResource->GetAssetName(), *MaterialResource->GetAssetPath().ToString()));
+		Errorf(FString::Printf(TEXT("Material %s have too many Substrate Operators: the compiler is failing (asset: %s).\r\n"), *GetAssetName(), *GetAssetPath()));
 		return DefaultOperatorOnError;
 	}
 
 	int32* ParentOperatorIndex = SubstrateCtx.SubstrateMaterialExpressionToOperatorIndex.Find(SubstrateParentExpressionGuid);
 	if (Parent != nullptr && ParentOperatorIndex == nullptr)
 	{
-		Errorf(FString::Printf(TEXT("Material %s tries to register unknown operator parents (asset: %s).\r\n"), *MaterialResource->GetAssetName(), *MaterialResource->GetAssetPath().ToString()));
+		Errorf(FString::Printf(TEXT("Material %s tries to register unknown operator parents (asset: %s).\r\n"), *GetAssetName(), *GetAssetPath()));
 		return DefaultOperatorOnError;
 	}
 
@@ -2464,7 +2472,7 @@ FSubstrateOperator* FSubstrateTranslatorData::SubstrateCompilationGetOperatorFro
 	int32 OperatorCount = SubstrateCtx.SubstrateMaterialExpressionToOperatorIndex.Num();
 	if (OperatorIndex < 0 || OperatorIndex >= OperatorCount)
 	{
-		Errorf(FString::Printf(TEXT("SubstrateCompilationGetOperatorFromIndex - OperatorIndex out of range %s (asset: %s).\r\n"), *MaterialResource->GetAssetName(), *MaterialResource->GetAssetPath().ToString()));
+		Errorf(FString::Printf(TEXT("SubstrateCompilationGetOperatorFromIndex - OperatorIndex out of range %s (asset: %s).\r\n"), *GetAssetName(), *GetAssetPath()));
 		return nullptr;
 	};
 	return &SubstrateCtx.SubstrateMaterialExpressionRegisteredOperators[OperatorIndex];
diff --git a/Engine/Source/Runtime/Engine/Private/Materials/SubstrateTranslatorCommon.h b/Engine/Source/Runtime/Engine/Private/Materials/SubstrateTranslatorCommon.h
--- a/Engine/Source/Runtime/Engine/Private/Materials/SubstrateTranslatorCommon.h
+++ b/Engine/Source/Runtime/Engine/Private/Materials/SubstrateTranslatorCommon.h
@@ -243,11 +243,11 @@ struct FSubstrateTranslatorData : public FSubstrateTranslatorDataInterface
 	FSubstrateTranslatorData();
 	virtual ~FSubstrateTranslatorData();
 
-	void Initialize(FMaterial* CompilerMaterial, FMaterialCompilationOutput& FMaterialCompilationOutput, FMaterialCompiler* CompilerLegacy = nullptr, MIR::FEmitter* CompilerNew = nullptr);
+	void Initialize(UMaterialInterface* MaterialInterface, FMaterialCompilationOutput& MaterialCompilationOutput, FMaterialCompiler* CompilerLegacy = nullptr, MIR::FEmitter* CompilerNew = nullptr);
 
 	EShaderPlatform GetShaderPlatform();
 
-	FMaterial* GetMaterialResource() { return MaterialResource; }
+	UMaterialInterface* GetMaterialInterface() { return MaterialInterface; }
 	FMaterialCompilationOutput& GetMaterialCompilationOutput() { return *MaterialCompilationOutput; }
 
 	void ProcessSubstrateTreeAndTopology(
@@ -308,13 +308,19 @@ struct FSubstrateTranslatorData : public FSubstrateTranslatorDataInterface
 	virtual void PushFunction(class FMaterialFunctionCompileState* FunctionState);
 	virtual void PopFunction();
 
+	const FString& GetAssetName() { return AssetName; }
+	const FString& GetAssetPath() { return AssetPath; }
+
 private:
 
-	FMaterial*						MaterialResource;
+	UMaterialInterface*				MaterialInterface;
 	FMaterialCompilationOutput*		MaterialCompilationOutput;
 	FMaterialCompiler*				CompilerLegacy;
 	MIR::FEmitter*					CompilerNew;
 
+	FString							AssetName;
+	FString							AssetPath;
+
 	FMaterialShadingModelField		MaterialShadingModels;
 
 	FExpressionInput				NullInput;
