Subsutrate - NMT - added support for simple clear coat node.

#tests pc editor dx12
#rnx

#Synced-CL 50398316

diff --git a/Engine/Source/Runtime/Engine/Private/Materials/MaterialExpressionsIR.cpp b/Engine/Source/Runtime/Engine/Private/Materials/MaterialExpressionsIR.cpp
--- a/Engine/Source/Runtime/Engine/Private/Materials/MaterialExpressionsIR.cpp
+++ b/Engine/Source/Runtime/Engine/Private/Materials/MaterialExpressionsIR.cpp
@@ -5531,6 +5531,126 @@ void UMaterialExpressionSubstrateShadingModels::Build(MIR::FEmitter& Em)
 		bHasSSS, SubsurfaceProfile/*, EditorOnlyData = nullptr*/);
 }
 
+void UMaterialExpressionSubstrateSimpleClearCoatBSDF::Build(MIR::FEmitter& Em)
+{
+	FSubstrateTranslatorData* SubstrateTranslatorData = Em.GetSubstrateTranslatorData();
+
+	FGuid PathUniqueId = SubstrateTranslatorData->SubstrateTreeStackGetPathUniqueId();
+	FSubstrateOperator& SubstrateOperator = SubstrateTranslatorData->SubstrateCompilationGetOperator(PathUniqueId);
+
+	const float DefaultSpecular = 0.5f;
+	const float DefaultF0 = DielectricSpecularToF0(DefaultSpecular);
+	
+	FValueRef EmissiveColorArg = CompileWithDefaultFloat3(Em, EmissiveColor, 0.0f, 0.0f, 0.0f);
+	FValueRef DiffuseAlbedoArg = CompileWithDefaultFloat3(Em, DiffuseAlbedo, 0.18f, 0.18f, 0.18f);
+	FValueRef F0Arg = CompileWithDefaultFloat3(Em, F0, DefaultF0, DefaultF0, DefaultF0);
+	FValueRef RoughnessArg = CompileWithDefaultFloat1(Em, Roughness, 0.5f);
+	FValueRef ClearCoatArg = CompileWithDefaultFloat1(Em, ClearCoatCoverage, 1.0f);
+	FValueRef ClearCoatRoughnessArg = CompileWithDefaultFloat1(Em, ClearCoatRoughness, 0.1f);
+	FValueRef ThicknessArg = Em.ConstantFloat(SUBSTRATE_LAYER_DEFAULT_THICKNESS_CM);
+
+	FValueRef AnisotropyArg = Em.ConstantFloat(0.0f);
+	FValueRef F90Arg = Em.ConstantFloat3(FVector3f(1.0f, 1.0f, 1.0f));
+	FValueRef SubsurfaceProfileArg = Em.ConstantFloat(0.0f);
+	FValueRef SSSMFPArg = Em.ConstantFloat3(FVector3f(0.0f, 0.0f, 0.0f));
+	FValueRef SSSMFPScaleArg = Em.ConstantFloat(1.0f);
+	FValueRef SSSPhaseAnisotropyArg = Em.ConstantFloat(0.0f);
+	FValueRef FuzzAmountArg = Em.ConstantFloat(0.0f);
+	FValueRef FuzzColorArg = Em.ConstantFloat3(FVector3f(0.0f, 0.0f, 0.0f));
+	FValueRef FuzzRoughnessArg = Em.ConstantFloat(0.5f);
+	FValueRef GlintValueArg = Em.ConstantFloat(1.0f);
+	FValueRef GlintUVArg = Em.ConstantFloat2(FVector2f(0.0f, 0.0f));
+	FValueRef SpecularProfileArg = Em.ConstantFloat(float(INDEX_NONE));
+	FValueRef SubSurfaceTypeArg = Em.ConstantFloat(float(MSS_None));
+
+	const bool bIsAtTheBottomOfTopology = SubstrateOperator.bIsBottom > 0;
+	const bool bIsThinSurface = Material ? Material->IsThinSurface() : false;
+	float ClearCoatUseSecondNormal = 0.0f; // Constant(ClearCoatBottomNormal != Normal ? 1.0f : 0.0f);SUBSTRATE_TODO
+
+	FValueRef NormalArg = CompileWithDefaultNormalWS(Em, Material, Normal);
+	const FSubstrateRegisteredSharedLocalBasis NewRegisteredSharedLocalBasis = MIRSubstrateCompilationInfoCreateSharedLocalBasis2(SubstrateTranslatorData, &NormalArg, nullptr);
+	SubstrateOperator.BSDFRegisteredSharedLocalBasis = NewRegisteredSharedLocalBasis;
+	FValueRef SharedLocalBasisIndexMacroArg = Em.InlineHLSL(MIR::FType::MakeIntScalar(), *GetSubstrateSharedLocalBasisIndexMacroInner(NewRegisteredSharedLocalBasis, SubstrateTranslatorData->CurrentSubstrateCompilationContext));
+
+	FSubstrateOperator* PromoteToOperator = !SubstrateOperator.bUseParameterBlending || (SubstrateOperator.bUseParameterBlending && SubstrateOperator.bRootOfParameterBlendingSubTree) ? &SubstrateOperator : nullptr;
+	if (PromoteToOperator)
+	{
+		MIR::FSubstrateSlabDesc SubstrateSlabDesc = {
+				NormalArg,
+				DiffuseAlbedoArg,
+				F0Arg,
+				F90Arg,
+				RoughnessArg,
+				AnisotropyArg,
+				SubsurfaceProfileArg,
+				SSSMFPArg,
+				SSSMFPScaleArg,
+				SSSPhaseAnisotropyArg,
+				SubSurfaceTypeArg,
+				EmissiveColorArg,
+				ClearCoatRoughnessArg,												//		Second roughness
+				ClearCoatArg,														//		Second roughness weight
+				Em.ConstantFloat(1.0f),												//		SecondRoughnessAsSimpleClearCoat = true
+				Em.ConstantFloat(ClearCoatUseSecondNormal),							//		*SubstrateGetCastParameterCode(ClearCoatUseSecondNormal, MCT_Float),SUBSTRATE_TODO
+				Em.ConstantFloat3(FVector3f::Zero()),								//		*SubstrateGetCastParameterCode(ClearCoatBottomNormal, MCT_Float3),SUBSTRATE_TODO
+				FuzzAmountArg,
+				FuzzColorArg,
+				FuzzRoughnessArg,
+				GlintValueArg,
+				GlintUVArg,
+				SpecularProfileArg,
+				ThicknessArg,
+				bIsThinSurface ? Em.ConstantTrue() : Em.ConstantFalse(),
+				bIsAtTheBottomOfTopology ? Em.ConstantTrue() : Em.ConstantFalse(),
+				SharedLocalBasisIndexMacroArg };
+
+		FValueRef Slab = Em.SubstrateSlab(SubstrateSlabDesc);
+
+		FValueRef PromototedSlab = Em.SubstratePromoteToOperator(
+			Slab,
+			Em.ConstantInt(PromoteToOperator->Index),
+			Em.ConstantInt(PromoteToOperator->BSDFIndex),
+			Em.ConstantInt(PromoteToOperator->LayerDepth),
+			Em.ConstantInt(PromoteToOperator->bIsBottom ? 1 : 0));
+
+		Em.Output(0, PromototedSlab);
+	}
+	else
+	{
+		MIR::FSubstrateSlabDesc SubstrateSlabDesc = {
+			NormalArg,
+			DiffuseAlbedoArg,
+			F0Arg,
+			F90Arg,
+			RoughnessArg,
+			AnisotropyArg,
+			SubsurfaceProfileArg,
+			SSSMFPArg,
+			SSSMFPScaleArg,
+			SSSPhaseAnisotropyArg,
+			SubSurfaceTypeArg,
+			EmissiveColorArg,
+			ClearCoatRoughnessArg,												//		Second roughness
+			ClearCoatArg,														//		Second roughness weight
+			Em.ConstantFloat(1.0f),												//		SecondRoughnessAsSimpleClearCoat = true
+			Em.ConstantFloat(ClearCoatUseSecondNormal),							//		*SubstrateGetCastParameterCode(ClearCoatUseSecondNormal, MCT_Float),SUBSTRATE_TODO
+			Em.ConstantFloat3(FVector3f::Zero()),								//		*SubstrateGetCastParameterCode(ClearCoatBottomNormal, MCT_Float3),SUBSTRATE_TODO
+			FuzzAmountArg,
+			FuzzColorArg,
+			FuzzRoughnessArg,
+			GlintValueArg,
+			GlintUVArg,
+			SpecularProfileArg,
+			ThicknessArg,
+			bIsThinSurface ? Em.ConstantTrue() : Em.ConstantFalse(),
+			bIsAtTheBottomOfTopology ? Em.ConstantTrue() : Em.ConstantFalse(),
+			SharedLocalBasisIndexMacroArg };
+		FValueRef Slab = Em.SubstrateSlab(SubstrateSlabDesc);
+
+		Em.Output(0, Slab);
+	}
+}
+
 static FValueRef SubstrateBlendNormal(MIR::FEmitter& Em, FValueRef NormalCodeChunk0, FValueRef NormalCodeChunk1, FValueRef MixCodeChunk)
 {
 	FValueRef SafeMixCodeChunk = Em.Saturate(MixCodeChunk);
diff --git a/Engine/Source/Runtime/Engine/Public/Materials/MaterialExpressionSubstrate.h b/Engine/Source/Runtime/Engine/Public/Materials/MaterialExpressionSubstrate.h
--- a/Engine/Source/Runtime/Engine/Public/Materials/MaterialExpressionSubstrate.h
+++ b/Engine/Source/Runtime/Engine/Public/Materials/MaterialExpressionSubstrate.h
@@ -499,7 +499,7 @@ class UMaterialExpressionSubstrateSimpleClearCoatBSDF : public UMaterialExpressi
 	//~ Begin UMaterialExpression Interface
 #if WITH_EDITOR
 	virtual int32 Compile(class FMaterialCompiler* Compiler, int32 OutputIndex) override;
-	// SUBSTRATE_TODO virtual void Build(MIR::FEmitter& Emitter) override;
+	virtual void Build(MIR::FEmitter& Emitter) override;
 	virtual void GetCaption(TArray<FString>& OutCaptions) const override;
 	virtual EMaterialValueType GetOutputValueType(int32 OutputIndex) override;
 	virtual EMaterialValueType GetInputValueType(int32 InputIndex) override;
