Substrate - Toon - Updated profile default values and ui min/max.

#rnx
#tests pc editor
#Synced-CL 50393747

diff --git a/Engine/Source/Runtime/Engine/Classes/Engine/ToonProfile.h b/Engine/Source/Runtime/Engine/Classes/Engine/ToonProfile.h
--- a/Engine/Source/Runtime/Engine/Classes/Engine/ToonProfile.h
+++ b/Engine/Source/Runtime/Engine/Classes/Engine/ToonProfile.h
@@ -45,7 +45,7 @@ struct FToonProfileStruct
 	/**
 	 * The size on screen of the dithering pattern.
 	 */
-	UPROPERTY(Category = "Diffuse", EditAnywhere, meta = (UIMin = 0.0, UIMax = 1.0, ClampMin = 0.0))
+	UPROPERTY(Category = "Diffuse", EditAnywhere, meta = (UIMin = 1.0, UIMax = 10.0, ClampMin = 1.0))
 	float DiffuseDitherPatternSize;
 
 	/**
@@ -69,7 +69,7 @@ struct FToonProfileStruct
 	/**
 	 * The size on screen of the dithering pattern.
 	 */
-	UPROPERTY(Category = "Specular", EditAnywhere, meta = (UIMin = 0.0, UIMax = 1.0, ClampMin = 0.0))
+	UPROPERTY(Category = "Specular", EditAnywhere, meta = (UIMin = 1.0, UIMax = 10.0, ClampMin = 1.0))
 	float SpecularDitherPatternSize;
 
 	/**
@@ -95,7 +95,7 @@ struct FToonProfileStruct
 	/**
 	 * The size on screen of the shadow hatching pattern.
 	 */
-	UPROPERTY(Category = "Shadow", EditAnywhere, meta = (UIMin = 0.0, UIMax = 1.0, ClampMin = 0.0))
+	UPROPERTY(Category = "Shadow", EditAnywhere, meta = (UIMin = 1.0, UIMax = 10.0, ClampMin = 1.0))
 	float ShadowHatchingPatternSize;
 
 	/**
diff --git a/Engine/Source/Runtime/Engine/Private/Rendering/ToonProfile.cpp b/Engine/Source/Runtime/Engine/Private/Rendering/ToonProfile.cpp
--- a/Engine/Source/Runtime/Engine/Private/Rendering/ToonProfile.cpp
+++ b/Engine/Source/Runtime/Engine/Private/Rendering/ToonProfile.cpp
@@ -483,16 +483,16 @@ FToonProfileStruct::FToonProfileStruct()
 
 	DiffuseDitherPattern = nullptr;
 	DiffuseDitherPatternStrength = 0.0f;	// By default, skip sampling dither pattern.
-	DiffuseDitherPatternSize = 0.1f;
+	DiffuseDitherPatternSize = 1.0f;
 
 	SpecularDitherPattern = nullptr;
 	SpecularDitherPatternStrength = 0.0f;	// By default, skip sampling dither pattern.
-	SpecularDitherPatternSize = 0.1f;
+	SpecularDitherPatternSize = 1.0f;
 
 	ShadowExtinctionCoefficient = 10.0f;	// Default extinction coefficients that work well.
 
 	ShadowHatchingPattern = nullptr;
-	ShadowHatchingPatternSize = 0.1f;
+	ShadowHatchingPatternSize = 1.0f;
 	ShadowHatchingPatternStrength = 1.0f;
 
 	AddPointToCurve(ShadowRamp, 0.00, 0.0);
