Fixed an issue in ULevelStreaming::ServerUpdateLevelVisibility where an old PlayerController could send the RPC with an invalid NetConnection during seamless map travel.

#Synced-CL 50205858

diff --git a/Engine/Source/Runtime/Engine/Private/LevelStreaming.cpp b/Engine/Source/Runtime/Engine/Private/LevelStreaming.cpp
--- a/Engine/Source/Runtime/Engine/Private/LevelStreaming.cpp
+++ b/Engine/Source/Runtime/Engine/Private/LevelStreaming.cpp
@@ -797,7 +797,9 @@ void ULevelStreaming::ServerUpdateLevelVisibility(bool bIsVisible, bool bTryMake
 		UWorld* World = GetWorld();
 		for (FConstPlayerControllerIterator Iterator = World->GetPlayerControllerIterator(); Iterator; ++Iterator)
 		{
-			if (APlayerController* PlayerController = Iterator->Get())
+			APlayerController* PlayerController = Iterator->Get(); 
+			UNetConnection* NetConnection = IsValid(PlayerController) && !PlayerController->IsPendingKillPending() ? PlayerController->GetNetConnection() : nullptr;
+			if (NetConnection && NetConnection->GetUChildConnection() == nullptr && NetConnection->GetConnectionState() == USOCK_Open)
 			{
 				FUpdateLevelVisibilityLevelInfo LevelVisibility(LoadedLevel, bIsVisible, bTryMakeVisible);
 				LevelVisibility.PackageName = PlayerController->NetworkRemapPath(LevelVisibility.PackageName, false);
