Virtual texture visualization. Add API for setting mode.
Fix scaling calculation in shader.
#Synced-CL 50433575

diff --git a/Engine/Shaders/Private/VisualizeVirtualTexture.usf b/Engine/Shaders/Private/VisualizeVirtualTexture.usf
--- a/Engine/Shaders/Private/VisualizeVirtualTexture.usf
+++ b/Engine/Shaders/Private/VisualizeVirtualTexture.usf
@@ -30,7 +30,7 @@ void Main(
 	OutColor.xyz = SceneColorLuminance;
 	OutColor.w = 1;
 
-	uint2 PixelPos = (uint2)(SvPosition.xy * (float2)ViewSize / Input_ViewportSize);
+	uint2 PixelPos = (uint2)(saturate(UV) * (float2)(ViewSize - 1));
 	uint ReadIndex = PixelPos.y * ViewSize.x + PixelPos.x;
 	uint PackedDebugValue = DebugBuffer[ReadIndex];
 	uint DebugValue = 0;
diff --git a/Engine/Source/Runtime/Engine/Internal/VT/VirtualTextureVisualizationData.h b/Engine/Source/Runtime/Engine/Internal/VT/VirtualTextureVisualizationData.h
--- a/Engine/Source/Runtime/Engine/Internal/VT/VirtualTextureVisualizationData.h
+++ b/Engine/Source/Runtime/Engine/Internal/VT/VirtualTextureVisualizationData.h
@@ -3,8 +3,10 @@
 #pragma once
 
 #include "CoreMinimal.h"
+#include "VirtualTextureVisualizationData.generated.h"
 
 /** The virtual texture visualization modes. */
+UENUM()
 enum class EVirtualTextureVisualizationMode : uint8
 {
 	None,
@@ -55,6 +57,9 @@ public:
 	/** Get the display description from the mode name. **/
 	ENGINE_API FText GetModeDisplayDesc(FName const& InModeName) const;
 
+	/** Set the current mode through the console and return the previously set mode. */
+	ENGINE_API EVirtualTextureVisualizationMode SetModeUsingConsoleCommand(EVirtualTextureVisualizationMode InMode);
+
 	/** Return the console command name for enabling single mode visualization. */
 	static const TCHAR* GetVisualizeConsoleCommandName();
 
diff --git a/Engine/Source/Runtime/Engine/Private/VT/VirtualTextureVisualizationData.cpp b/Engine/Source/Runtime/Engine/Private/VT/VirtualTextureVisualizationData.cpp
--- a/Engine/Source/Runtime/Engine/Private/VT/VirtualTextureVisualizationData.cpp
+++ b/Engine/Source/Runtime/Engine/Private/VT/VirtualTextureVisualizationData.cpp
@@ -139,6 +139,27 @@ FText FVirtualTextureVisualizationData::GetModeDisplayDesc(FName const& InModeNa
 	}
 }
 
+EVirtualTextureVisualizationMode FVirtualTextureVisualizationData::SetModeUsingConsoleCommand(EVirtualTextureVisualizationMode InMode)
+{
+	EVirtualTextureVisualizationMode PreviousMode = EVirtualTextureVisualizationMode::None;
+	if (IConsoleVariable* ConsoleVariable = IConsoleManager::Get().FindConsoleVariable(GetVisualizeConsoleCommandName()))
+	{
+		FString PreviousModeString = ConsoleVariable->GetString();
+		if (const FModeRecord* Record = ModeArray.FindByPredicate([PreviousModeString](const FModeRecord& Record) { return Record.ModeString == PreviousModeString; }))
+		{
+			PreviousMode = Record->ModeID;
+		}
+		
+		FString ModeString; // Empty string will implement EVirtualTextureVisualizationMode::None
+		if (const FModeRecord* Record = ModeArray.FindByPredicate([InMode](const FModeRecord& Record) { return Record.ModeID == InMode; }))
+		{
+			ModeString = Record->ModeString;
+		}
+		ConsoleVariable->AsVariable()->SetWithCurrentPriority(*ModeString);
+	}
+	return PreviousMode;
+}
+
 FVirtualTextureVisualizationData& GetVirtualTextureVisualizationData()
 {
 	if (!GVirtualTextureVisualizationData.IsInitialized())
