Substrate - Toon
- Fixed toon profile update when changing resolution cvras.
- Fixed toon profile update when textures are streamed in and were previously not ready.
- clean up shader code.
#tests PC editor dx12
#rnx

#Synced-CL 50434260

diff --git a/Engine/Shaders/Private/HeightFogCommon.ush b/Engine/Shaders/Private/HeightFogCommon.ush
--- a/Engine/Shaders/Private/HeightFogCommon.ush
+++ b/Engine/Shaders/Private/HeightFogCommon.ush
@@ -220,21 +220,6 @@ float PreComputeFogOriginFactor(float OriginHeight, float FogHeight, float FogFa
 	return FogDensity * pow(2.0f, CollapsedFogParameterPower);
 }
 
-#define REFLECTION_CAPTURE_ROUGHEST_MIP 1
-#define REFLECTION_CAPTURE_ROUGHNESS_MIP_SCALE 1.2
-
-////////////////////**
-/////////////////// * Compute absolute mip for a reflection capture cubemap given a roughness.
-/////////////////// */
-///////////////////half ComputeReflectionCaptureMipFromRoughness222(half Roughness, half CubemapMaxMip)
-///////////////////{
-///////////////////	// Heuristic that maps roughness to mip level
-///////////////////	// This is done in a way such that a certain mip level will always have the same roughness, regardless of how many mips are in the texture
-///////////////////	// Using more mips in the cubemap just allows sharper reflections to be supported
-///////////////////	half LevelFrom1x1 = REFLECTION_CAPTURE_ROUGHEST_MIP - REFLECTION_CAPTURE_ROUGHNESS_MIP_SCALE * log2(max(Roughness, 0.001));
-///////////////////	return CubemapMaxMip - 1 - LevelFrom1x1;
-///////////////////}
-
 // @param WorldPositionRelativeToCamera = WorldPosition - InCameraPosition. It is the vector from the camera to the world position to shade (independent from TranslatedPosition).
 half4 GetExponentialHeightFog(float3 WorldPositionRelativeToCamera, float ExcludeDistance, uint EyeIndex, ViewState InView, bool bOverrideOrigin = false, float3 OverrideRayStartRelativeToCamera=0.0f.xxx, float3 OverrideRayDir = 0.0f.xxx, float OverrideRayLength = 0.0f)
 {
diff --git a/Engine/Source/Runtime/Engine/Private/Materials/MaterialInterface.cpp b/Engine/Source/Runtime/Engine/Private/Materials/MaterialInterface.cpp
--- a/Engine/Source/Runtime/Engine/Private/Materials/MaterialInterface.cpp
+++ b/Engine/Source/Runtime/Engine/Private/Materials/MaterialInterface.cpp
@@ -2049,7 +2049,6 @@ void UMaterialInterface::UpdateMaterialRenderProxy(FMaterialRenderProxy& Proxy)
 			{
 				UToonProfile* Profile = nullptr;
 				FToonProfileStruct Settings;
-				const FTextureReference* Texture = nullptr;
 				FGuid Guid;
 			};
 			TArray<FEntry> Entries;
diff --git a/Engine/Source/Runtime/Engine/Private/Rendering/ToonProfile.cpp b/Engine/Source/Runtime/Engine/Private/Rendering/ToonProfile.cpp
--- a/Engine/Source/Runtime/Engine/Private/Rendering/ToonProfile.cpp
+++ b/Engine/Source/Runtime/Engine/Private/Rendering/ToonProfile.cpp
@@ -39,8 +39,21 @@ static TAutoConsoleVariable<int32> CVarToonProfileForceUpdate(
 	ECVF_Cheat | ECVF_RenderThreadSafe);
 #endif
 
+static uint32 GetToonProfileResolution()
+{
+	return FMath::Max(64, CVarToonProfileResolution.GetValueOnRenderThread());
+}
+
 static bool ForceUpdateToonProfile()
 {
+	static int32 LastToonProfileResolution = 0;
+	const int32 NewToonProfileResolution = GetToonProfileResolution();
+	if (LastToonProfileResolution != NewToonProfileResolution)
+	{
+		LastToonProfileResolution = NewToonProfileResolution;
+		return true;
+	}
+
 #if !(UE_BUILD_SHIPPING || UE_BUILD_TEST)
 	return (CVarToonProfileForceUpdate.GetValueOnAnyThread() == 1);
 #else
@@ -161,13 +174,17 @@ public:
 
 private:
 
-	struct FProfileEntry
+	struct FToonProfileEntry
 	{
 		FToonProfileStruct Settings;
 		const UToonProfile* Profile = nullptr; // Game thread pointer! Do not dereference, only for comparison.
 		FName ParameterName;
+		FIntPoint DiffuseDitherPatternCachedResolution = FIntPoint::ZeroValue;
+		FIntPoint SpecularDitherPatternCachedResolution = FIntPoint::ZeroValue;
+		FIntPoint ShadowHatchingPatternCachedResolution = FIntPoint::ZeroValue;
 	};
-	TArray<FProfileEntry> ToonProfileEntries;
+	TArray<FToonProfileEntry> ToonProfileEntries;
+	const FTextureReference* Texture = nullptr;
 };
 
 // Global resources - lives on the render thread
@@ -181,7 +198,7 @@ FToonProfileTextureManager::FToonProfileTextureManager()
 	check(IsInGameThread());
 
 	// add element 0, it is used as default profile
-	FProfileEntry& Entry = ToonProfileEntries.AddDefaulted_GetRef();
+	FToonProfileEntry& Entry = ToonProfileEntries.AddDefaulted_GetRef();
 	Entry.Settings = FToonProfileStruct();
 	Entry.Profile = nullptr;
 }
@@ -219,7 +236,7 @@ int32 FToonProfileTextureManager::AddProfile(const UToonProfile* InProfile, cons
 				ToonProfileIdQuery.bFound = true;
 				ToonProfileIdQuery.AllocationId = i;
 
-				FProfileEntry& Entry= ToonProfileEntries[ToonProfileIdQuery.AllocationId];
+				FToonProfileEntry& Entry= ToonProfileEntries[ToonProfileIdQuery.AllocationId];
 				Entry.Profile 		= InProfile;
 				Entry.ParameterName = FName(TEXT("__ToonProfile") + Guid.ToString());
 				break;
@@ -231,7 +248,7 @@ int32 FToonProfileTextureManager::AddProfile(const UToonProfile* InProfile, cons
 			ToonProfileIdQuery.bFound = true;
 			ToonProfileIdQuery.AllocationId = ToonProfileEntries.Num();
 
-			FProfileEntry& Entry= ToonProfileEntries.AddDefaulted_GetRef();
+			FToonProfileEntry& Entry= ToonProfileEntries.AddDefaulted_GetRef();
 			Entry.Profile 		= InProfile;
 			Entry.ParameterName = FName(TEXT("__ToonProfile") + Guid.ToString());
 		}
@@ -302,9 +319,31 @@ IPooledRenderTarget* FToonProfileTextureManager::GetAtlasTexture(FRDGBuilder& Gr
 	const uint32 ToonProfileCount = ToonProfileEntries.Num();
 	check(ToonProfileCount);
 
-	// Since reference texture can be streamed/loaded progressively, we track if textures have been updated to update the LUT
+	// Force update if cvars have changed.
 	bool bForceUpdate = ForceUpdateToonProfile();
 
+	// Force update if a texture has been streamed in and thus has changed. That is done by tracking change in resolution.
+	for (uint32 LayerIt = 0; LayerIt < ToonProfileCount; ++LayerIt)
+	{
+		FToonProfileEntry& Entry = ToonProfileEntries[LayerIt];
+		const FToonProfileStruct& Data = Entry.Settings;
+
+		auto CheckIfTextureHasStreamedIn = [&](FTextureResource* TextureResource, FIntPoint& CachedResolution)
+		{
+			const uint32 TextureSizeX = TextureResource ? TextureResource->GetSizeX() : 0u;
+			const uint32 TextureSizeY = TextureResource ? TextureResource->GetSizeY() : 0u;
+			if (TextureSizeX != CachedResolution.X || TextureSizeY != CachedResolution.Y)
+			{
+				CachedResolution = FIntPoint(TextureSizeX, TextureSizeY);
+				bForceUpdate = true;
+			}
+		};
+
+		CheckIfTextureHasStreamedIn(Data.DiffuseDitherPattern  ? Data.DiffuseDitherPattern->GetResource()  : nullptr, Entry.DiffuseDitherPatternCachedResolution);
+		CheckIfTextureHasStreamedIn(Data.SpecularDitherPattern ? Data.SpecularDitherPattern->GetResource() : nullptr, Entry.SpecularDitherPatternCachedResolution);
+		CheckIfTextureHasStreamedIn(Data.ShadowHatchingPattern ? Data.ShadowHatchingPattern->GetResource() : nullptr, Entry.ShadowHatchingPatternCachedResolution);
+	}
+
 	if (bForceUpdate)
 	{
 		GToonProfileTextureAtlas.SafeRelease();
@@ -315,7 +354,7 @@ IPooledRenderTarget* FToonProfileTextureManager::GetAtlasTexture(FRDGBuilder& Gr
 	{
 		FRHICommandListImmediate& RHICmdList = GraphBuilder.RHICmdList;
 
-		const uint32 ToonProfileResolution = CVarToonProfileResolution.GetValueOnRenderThread();
+		const uint32 ToonProfileResolution = GetToonProfileResolution();
 
 		RDG_EVENT_SCOPE(GraphBuilder, "ToonProfile");
 
