Virtual Texturing: r.VT.MaxUploadMemory and r.VT.MaxUploadRequests can now be set to 0.
A value of 0 disables their upload throttling.
r.VT.MaxUploadRequests defaults to 0 so that by default we only throttle VT uploads by memory budget.

#Synced-CL 50433575

diff --git a/Engine/Source/Runtime/Engine/Private/VT/VirtualTextureUploadCache.cpp b/Engine/Source/Runtime/Engine/Private/VT/VirtualTextureUploadCache.cpp
--- a/Engine/Source/Runtime/Engine/Private/VT/VirtualTextureUploadCache.cpp
+++ b/Engine/Source/Runtime/Engine/Private/VT/VirtualTextureUploadCache.cpp
@@ -28,14 +28,16 @@ static TAutoConsoleVariable<int32> CVarVTMaxUploadMemory(
 	TEXT("r.VT.MaxUploadMemory"),
 	64,
 	TEXT("Maximum amount of upload memory to allocate in MB before throttling virtual texture streaming requests.\n")
-	TEXT("We never throttle high priority requests so allocation can peak above this value."),
+	TEXT("We never throttle high priority requests so allocation can peak above this value.\n")
+	TEXT("0 or negative values disable the restriction."),
 	ECVF_RenderThreadSafe
 );
 
 static TAutoConsoleVariable<int32> CVarMaxUploadRequests(
 	TEXT("r.VT.MaxUploadRequests"),
-	6000,
-	TEXT("Maximum number of virtual texture tile upload requests that can be in flight."),
+	0,
+	TEXT("Maximum number of virtual texture tile upload requests that can be in flight.\n")
+	TEXT("0 (default) or negative values disable the restriction."),
 	ECVF_RenderThreadSafe);
 
 
@@ -609,7 +611,16 @@ void FVirtualTextureUploadCache::UpdateFreeList(FRHICommandListBase& RHICmdList,
 
 uint32 FVirtualTextureUploadCache::IsInMemoryBudget() const
 {
-	return 
-		PendingUpload.Num() + PendingRelease.Num() <= CVarMaxUploadRequests.GetValueOnRenderThread() &&
-		TileAllocator.TotalAllocatedBytes() <= (uint64)CVarVTMaxUploadMemory.GetValueOnRenderThread() * 1024u * 1024u;
+	const int32 MaxUploadRequests = CVarMaxUploadRequests.GetValueOnRenderThread();
+	const uint64 MaxUploadMemoryBytes = (uint64)FMath::Max(CVarVTMaxUploadMemory.GetValueOnRenderThread(), 0) * 1024u * 1024u;
+
+	if (MaxUploadRequests > 0 && PendingUpload.Num() + PendingRelease.Num() > MaxUploadRequests)
+	{
+		return false;
+	}
+	if (MaxUploadMemoryBytes > 0 && TileAllocator.TotalAllocatedBytes() > MaxUploadMemoryBytes)
+	{
+		return false;
+	}
+	return true;
 }
