[WorldPartition] Fix incorrect positioning of actors with absolute transforms attached to editor-only actors
Actors using absolute location/rotation/scale that are attached to editor-only parent actors were being incorrectly offset at runtime. The editor-only parent transform was being applied unconditionally without respecting the child's absolute transform flags.

#Synced-CL 50425068

diff --git a/Engine/Source/Runtime/Engine/Private/WorldPartition/WorldPartitionLevelHelper.cpp b/Engine/Source/Runtime/Engine/Private/WorldPartition/WorldPartitionLevelHelper.cpp
--- a/Engine/Source/Runtime/Engine/Private/WorldPartition/WorldPartitionLevelHelper.cpp
+++ b/Engine/Source/Runtime/Engine/Private/WorldPartition/WorldPartitionLevelHelper.cpp
@@ -751,6 +751,36 @@ bool FWorldPartitionLevelHelper::LoadActorsInternal(FLoadActorsParams&& InParams
 		return PackageNameBuilder.ToString();
 	};
 
+	auto GetEffectiveEditorOnlyParentTransform = [](const FTransform& InEditorOnlyParentTransform, const AActor* InActor) -> FTransform
+	{
+		if (USceneComponent* RootComponent = InActor->GetRootComponent())
+		{
+			if (RootComponent->IsUsingAbsoluteLocation() || RootComponent->IsUsingAbsoluteRotation() || RootComponent->IsUsingAbsoluteScale())
+			{
+				FTransform EffectiveTransform = FTransform::Identity;
+
+				if (!RootComponent->IsUsingAbsoluteLocation())
+				{
+					EffectiveTransform.CopyTranslation(InEditorOnlyParentTransform);
+				}
+
+				if (!RootComponent->IsUsingAbsoluteRotation())
+				{
+					EffectiveTransform.CopyRotation(InEditorOnlyParentTransform);
+				}
+
+				if (!RootComponent->IsUsingAbsoluteScale())
+				{
+					EffectiveTransform.CopyScale3D(InEditorOnlyParentTransform);
+				}
+
+				return EffectiveTransform;
+			}
+		}
+		
+		return InEditorOnlyParentTransform;
+	};
+
 	if (!InParams.ActorPackages.IsEmpty())
 	{
 		ActorPackagesToLoad.Reserve(InParams.ActorPackages.Num());
@@ -836,7 +866,7 @@ bool FWorldPartitionLevelHelper::LoadActorsInternal(FLoadActorsParams&& InParams
 			}
 		}
 
-		FLoadPackageAsyncDelegate CompletionCallback = FLoadPackageAsyncDelegate::CreateLambda([LoadProgress, PackageObjectMapping, LoadedOverrides = InLoadedPropertyOverrides, PackageReferencer = InParams.PackageReferencer, OuterWorld = InParams.OuterWorld, DestLevel = InParams.DestLevel, CompletionCallback = InParams.CompletionCallback, bSilenceLoadFailures = InParams.bSilenceLoadFailures](const FName& LoadedPackageName, UPackage* LoadedPackage, EAsyncLoadingResult::Type Result)
+		FLoadPackageAsyncDelegate CompletionCallback = FLoadPackageAsyncDelegate::CreateLambda([LoadProgress, PackageObjectMapping, LoadedOverrides = InLoadedPropertyOverrides, PackageReferencer = InParams.PackageReferencer, OuterWorld = InParams.OuterWorld, DestLevel = InParams.DestLevel, CompletionCallback = InParams.CompletionCallback, bSilenceLoadFailures = InParams.bSilenceLoadFailures, GetEffectiveEditorOnlyParentTransform](const FName& LoadedPackageName, UPackage* LoadedPackage, EAsyncLoadingResult::Type Result)
 		{
 			const FName ActorName = *FPaths::GetExtension(PackageObjectMapping->Path.ToString());
 			check(LoadProgress->NumPendingLoadRequests);
@@ -932,7 +962,7 @@ bool FWorldPartitionLevelHelper::LoadActorsInternal(FLoadActorsParams&& InParams
 						ActorPropertyOverridesAnnotation.AddAnnotation(Actor, FWorldPartitionLevelHelper::FActorPropertyOverridesAnnotation(MoveTemp(ActorPropertyOverrides), PackageObjectMapping->ContainerTransform));
 					}
 
-					const FTransform TransformToApply = PackageObjectMapping->ContainerTransform * PackageObjectMapping->EditorOnlyParentTransform;
+					const FTransform TransformToApply = PackageObjectMapping->ContainerTransform * GetEffectiveEditorOnlyParentTransform(PackageObjectMapping->EditorOnlyParentTransform, Actor);
 					FLevelUtils::FApplyLevelTransformParams TransformParams(nullptr, TransformToApply);
 					TransformParams.Actor = Actor;
 					TransformParams.bDoPostEditMove = false;
@@ -971,7 +1001,7 @@ bool FWorldPartitionLevelHelper::LoadActorsInternal(FLoadActorsParams&& InParams
 				}
 				else if (!PackageObjectMapping->EditorOnlyParentTransform.Equals(FTransform::Identity))
 				{
-					FLevelUtils::FApplyLevelTransformParams TransformParams(nullptr, PackageObjectMapping->EditorOnlyParentTransform);
+					FLevelUtils::FApplyLevelTransformParams TransformParams(nullptr, GetEffectiveEditorOnlyParentTransform(PackageObjectMapping->EditorOnlyParentTransform, Actor));
 					TransformParams.Actor = Actor;
 					TransformParams.bDoPostEditMove = false;
 					FLevelUtils::ApplyLevelTransform(TransformParams);
