Conform archetype subobject flags

[RN] Component subobjects tagged as RF_ArchetypeObject now have their subobject flags conformed, forcing RF_Public and consequently preventing 'linked to private object' error
#Synced-CL 50069009

diff --git a/Engine/Source/Runtime/Engine/Private/SimpleConstructionScript.cpp b/Engine/Source/Runtime/Engine/Private/SimpleConstructionScript.cpp
--- a/Engine/Source/Runtime/Engine/Private/SimpleConstructionScript.cpp
+++ b/Engine/Source/Runtime/Engine/Private/SimpleConstructionScript.cpp
@@ -13,6 +13,7 @@
 #include "Kismet2/BlueprintEditorUtils.h"
 #include "Kismet2/ComponentEditorUtils.h"
 #include "Kismet2/Kismet2NameValidators.h"
+#include "Logging/StructuredLog.h"
 #else
 #include "UObject/LinkerLoad.h"
 #endif
@@ -23,6 +24,38 @@
 // We append this suffix to template object names because the FObjectProperty we create at compile time will also be outered to the generated Blueprint class, and because we need cooking to be deterministic with respect to template object names.
 const FString USimpleConstructionScript::ComponentTemplateNameSuffix(TEXT("_GEN_VARIABLE"));
 
+#if WITH_EDITOR
+namespace UE::Private
+{
+	void ConformArchetypeVisibility(UObject* Template)
+	{
+		TArray<UObject*> SubObjects;
+		GetObjectsWithOuter(
+			Template, 
+			SubObjects, 
+			/*bIncludeNestedObjects*/ true, 
+			RF_Transient, 
+			EInternalObjectFlags::Garbage);
+
+		const EObjectFlags Flags = Template->GetMaskedFlags(RF_PropagateToSubObjects);
+
+		for (UObject* SubObject : SubObjects)
+		{
+			// In theory it might be good to propagate to all subobjects of
+			// the template object, but I'm being conservative for now:
+			if(SubObject->HasAnyFlags(RF_ArchetypeObject))
+			{
+				if(!SubObject->HasAllFlags(Flags))
+				{
+					UE_LOGFMT(LogBlueprint, Warning, "Conforming subobject flags: {subobject}", SubObject->GetPathName());
+				}
+				SubObject->SetFlags(Flags);
+			}
+		}
+	}
+}
+#endif
+
 USimpleConstructionScript::USimpleConstructionScript(const FObjectInitializer& ObjectInitializer)
 	: Super(ObjectInitializer)
 {
@@ -1577,6 +1610,10 @@ void USimpleConstructionScript::ValidateNodeTemplates(FCompilerResultsLog& Messa
 					"SimpleConstructionScript", "CorruptComponentFixed", 
 					"SCS Node component template {0} has been recreated because it was unexpectedly not owned by the class"), FText::FromString(VariableName)).ToString());
 			}
+			
+#if WITH_EDITOR
+			UE::Private::ConformArchetypeVisibility(Node->ComponentTemplate);
+#endif//WITH_EDITOR
 		}
 	}
 }
