Adding simple counter to ensure if too many Blueprint async actions are created.
The number of allowed actions can be set using the console variable `bp.MaxAsyncActionCount`. Set the number to <= 0 to disable the ensure.

#Synced-CL 50443327

diff --git a/Engine/Source/Runtime/Engine/Classes/Kismet/BlueprintAsyncActionBase.h b/Engine/Source/Runtime/Engine/Classes/Kismet/BlueprintAsyncActionBase.h
--- a/Engine/Source/Runtime/Engine/Classes/Kismet/BlueprintAsyncActionBase.h
+++ b/Engine/Source/Runtime/Engine/Classes/Kismet/BlueprintAsyncActionBase.h
@@ -23,6 +23,9 @@ public:
 	/** Default UObject constructor */
 	ENGINE_API UBlueprintAsyncActionBase(const FObjectInitializer& ObjectInitializer = FObjectInitializer::Get());
 
+	/** Called when the object is destroyed */
+	ENGINE_API virtual void BeginDestroy() override;
+
 	/** Called to trigger the action once the delegates have been bound */
 	UFUNCTION(BlueprintCallable, meta=(BlueprintInternalUseOnly="true"))
 	ENGINE_API virtual void Activate();
diff --git a/Engine/Source/Runtime/Engine/Private/Kismet/BlueprintAsyncActionBase.cpp b/Engine/Source/Runtime/Engine/Private/Kismet/BlueprintAsyncActionBase.cpp
--- a/Engine/Source/Runtime/Engine/Private/Kismet/BlueprintAsyncActionBase.cpp
+++ b/Engine/Source/Runtime/Engine/Private/Kismet/BlueprintAsyncActionBase.cpp
@@ -4,9 +4,45 @@
 #include "Engine/Engine.h"
 #include "Engine/GameInstance.h"
 #include "Engine/World.h"
+#include "Misc/NotNull.h"
 
 #include UE_INLINE_GENERATED_CPP_BY_NAME(BlueprintAsyncActionBase)
 
+#if DO_ENSURE
+
+#include <atomic>
+
+namespace UE::BlueprintAsyncAction::Private
+{
+	static TAutoConsoleVariable<int32> CVarMaxActionCount(
+		TEXT("bp.MaxAsyncActionCount"),
+		10000,
+		TEXT("The maximum number of async actions of a that can be active at a time before triggering an ensure. Set to <= 0 to disable")
+	);
+
+	/** Total number of async actions that are active */
+	static std::atomic<int32> GlobalAsyncActionCount = 0;
+
+	static void DebugRecordAsyncActionCreated(TNotNull<const UBlueprintAsyncActionBase*> Action)
+	{
+		const int32 CurrentAsyncActionCount = GlobalAsyncActionCount.fetch_add(1, std::memory_order_relaxed) + 1;
+		const int32 MaxActionCount = CVarMaxActionCount.GetValueOnAnyThread();
+		if (MaxActionCount > 0)
+		{
+			const UClass* ActionClass = Action->GetClass();
+			ensureMsgf(CurrentAsyncActionCount <= MaxActionCount,
+				TEXT("Exceeded the maximum number of allowed async actions. Last created action class was %s"),
+				*GetNameSafe(ActionClass));
+		}
+	}
+
+	static void DebugRecordAsyncActionDestroyed()
+	{
+		GlobalAsyncActionCount.fetch_sub(1, std::memory_order_relaxed);
+	}
+}
+#endif // DO_ENSURE
+
 //////////////////////////////////////////////////////////////////////////
 // UBlueprintAsyncActionBase
 
@@ -16,7 +52,23 @@ UBlueprintAsyncActionBase::UBlueprintAsyncActionBase(const FObjectInitializer& O
 	if (!HasAnyFlags(RF_ClassDefaultObject))
 	{
 		SetFlags(RF_StrongRefOnFrame);
+
+#if DO_ENSURE
+		UE::BlueprintAsyncAction::Private::DebugRecordAsyncActionCreated(this);
+#endif // DO_ENSURE
+	}
+}
+
+void UBlueprintAsyncActionBase::BeginDestroy()
+{
+	Super::BeginDestroy();
+
+#if DO_ENSURE
+	if (!HasAnyFlags(RF_ClassDefaultObject))
+	{
+		UE::BlueprintAsyncAction::Private::DebugRecordAsyncActionDestroyed();
 	}
+#endif // DO_ENSURE
 }
 
 void UBlueprintAsyncActionBase::Activate()
