Use weak ptrs in PhysicsProxyToComponentMap

#Synced-CL 50451748

#Synced-CL 50451748

diff --git a/Engine/Source/Runtime/Engine/Private/PhysicsEngine/Experimental/PhysScene_Chaos.cpp b/Engine/Source/Runtime/Engine/Private/PhysicsEngine/Experimental/PhysScene_Chaos.cpp
--- a/Engine/Source/Runtime/Engine/Private/PhysicsEngine/Experimental/PhysScene_Chaos.cpp
+++ b/Engine/Source/Runtime/Engine/Private/PhysicsEngine/Experimental/PhysScene_Chaos.cpp
@@ -950,9 +950,13 @@ void FPhysScene_Chaos::AddReferencedObjects(FReferenceCollector& Collector)
 	Collector.AddReferencedObject(ChaosEventRelay);
 #if WITH_EDITOR
 	UE_CHAOS_ASYNC_INITBODY_READSCOPELOCK(PhysicsProxyComponentMapsLock);
-	for (auto& Pair : PhysicsProxyToComponentMap)
+	
+	for (TPair<IPhysicsProxyBase*, TWeakObjectPtr<UPrimitiveComponent>>& Pair : PhysicsProxyToComponentMap)
 	{
-		Collector.AddReferencedObject(Pair.Get<1>());
+		if (TObjectPtr<UPrimitiveComponent> Component = Pair.Get<1>().Get())
+		{
+			Collector.AddReferencedObject(Component);
+		}
 	}
 #endif
 }
@@ -961,9 +965,13 @@ template<>
 UPrimitiveComponent* FPhysScene_Chaos::GetOwningComponent(const IPhysicsProxyBase* PhysicsProxy) const
 {
 	UE_CHAOS_ASYNC_INITBODY_READSCOPELOCK(PhysicsProxyComponentMapsLock);
-	if (const TObjectPtr<UPrimitiveComponent>* FoundComp = PhysicsProxyToComponentMap.Find(PhysicsProxy))
+	if (const TWeakObjectPtr<UPrimitiveComponent>* ComponentWeak = PhysicsProxyToComponentMap.Find(PhysicsProxy))
 	{
-		return *FoundComp;
+		const TStrongObjectPtr<UPrimitiveComponent> StrongFoundComp = ComponentWeak->Pin();
+		if (ensureAlwaysMsgf(StrongFoundComp.IsValid(), TEXT("GetOwningComponent found primitive component in map which has been garbage collected")))
+		{
+			return StrongFoundComp.Get();
+		}
 	}
 	return nullptr;
 }
@@ -1479,18 +1487,21 @@ void FPhysScene_Chaos::CreatePhysicsReplicationCache()
 void FPhysScene_Chaos::RemoveFromComponentMaps(IPhysicsProxyBase* InObject)
 {
 	UE_CHAOS_ASYNC_INITBODY_WRITESCOPELOCK(PhysicsProxyComponentMapsLock);
-	TObjectPtr<UPrimitiveComponent>* const Component = PhysicsProxyToComponentMap.Find(InObject);
-	if (Component)
+	const TWeakObjectPtr<UPrimitiveComponent>* ComponentWeak = PhysicsProxyToComponentMap.Find(InObject);
+	if (ComponentWeak)
 	{
-		PendingMigratedPhysicsStateByComponent.Remove(*Component);
-
-		TArray<IPhysicsProxyBase*>* ProxyArray = ComponentToPhysicsProxyMap.Find(*Component);
-		if (ProxyArray)
+		if (TObjectPtr<UPrimitiveComponent> Component = ComponentWeak->Get(true))
 		{
-			ProxyArray->RemoveSingleSwap(InObject, EAllowShrinking::No);
-			if (ProxyArray->Num() == 0)
+			PendingMigratedPhysicsStateByComponent.Remove(Component.Get());
+
+			TArray<IPhysicsProxyBase*>* ProxyArray = ComponentToPhysicsProxyMap.Find(Component.Get());
+			if (ProxyArray)
 			{
-				ComponentToPhysicsProxyMap.Remove(*Component);
+				ProxyArray->RemoveSingleSwap(InObject, EAllowShrinking::No);
+				if (ProxyArray->Num() == 0)
+				{
+					ComponentToPhysicsProxyMap.Remove(Component.Get());
+				}
 			}
 		}
 	}
diff --git a/Engine/Source/Runtime/Engine/Public/Physics/Experimental/PhysScene_Chaos.h b/Engine/Source/Runtime/Engine/Public/Physics/Experimental/PhysScene_Chaos.h
--- a/Engine/Source/Runtime/Engine/Public/Physics/Experimental/PhysScene_Chaos.h
+++ b/Engine/Source/Runtime/Engine/Public/Physics/Experimental/PhysScene_Chaos.h
@@ -193,8 +193,8 @@ public:
 	OwnerType* GetOwningComponent(const IPhysicsProxyBase* PhysicsProxy) const
 	{ 
 		UE_CHAOS_ASYNC_INITBODY_READSCOPELOCK(PhysicsProxyComponentMapsLock);
-		auto* CompPtr = PhysicsProxyToComponentMap.Find(PhysicsProxy);
-		return CompPtr ? Cast<OwnerType>(*CompPtr) : nullptr;
+		const TWeakObjectPtr<UPrimitiveComponent>* ComponentWeak = PhysicsProxyToComponentMap.Find(PhysicsProxy);
+		return (ComponentWeak != nullptr && ComponentWeak->IsValid()) ? Cast<OwnerType>(ComponentWeak->Get()) : nullptr;
 	}
 
 	template<>
@@ -497,8 +497,9 @@ private:
 	// Temp Interface
 	TArray<FCollisionNotifyInfo> MNotifies;
 
-	// Maps PhysicsProxy to Component that created the PhysicsProxy
-	TMap<IPhysicsProxyBase*, TObjectPtr<UPrimitiveComponent>> PhysicsProxyToComponentMap;
+	// Maps PhysicsProxy to Component that created the PhysicsProxy.
+	// Keeps a weak ptr so that we can track validity of the primitive.
+	TMap<IPhysicsProxyBase*, TWeakObjectPtr<UPrimitiveComponent>> PhysicsProxyToComponentMap;
 
 	// Maps Component to PhysicsProxy that is created
 	TMap<UPrimitiveComponent*, TArray<IPhysicsProxyBase*>> ComponentToPhysicsProxyMap;
