Configurable dynamic loading range scaling leading to quicker PIE startup on large maps

[RN] minor WorldPartition/Streaming
#tests Verified PIE start improvements of around 5 seconds in CitySample

#Synced-CL 50456686

diff --git a/Engine/Source/Editor/WorldPartitionEditor/Private/WorldPartition/Customizations/WorldPartitionDetailsCustomization.cpp b/Engine/Source/Editor/WorldPartitionEditor/Private/WorldPartition/Customizations/WorldPartitionDetailsCustomization.cpp
--- a/Engine/Source/Editor/WorldPartitionEditor/Private/WorldPartition/Customizations/WorldPartitionDetailsCustomization.cpp
+++ b/Engine/Source/Editor/WorldPartitionEditor/Private/WorldPartition/Customizations/WorldPartitionDetailsCustomization.cpp
@@ -3,6 +3,7 @@
 #include "WorldPartitionDetailsCustomization.h"
 #include "WorldPartition/WorldPartition.h"
 #include "WorldPartition/WorldPartitionRuntimeHash.h"
+#include "WorldPartition/WorldPartitionEditorPerProjectUserSettings.h"
 #include "WorldPartition/WorldPartitionEditorSettings.h"
 #include "PropertyCustomizationHelpers.h"
 #include "ScopedTransaction.h"
@@ -10,13 +11,15 @@
 #include "DetailCategoryBuilder.h"
 #include "DetailWidgetRow.h"
 #include "Engine/World.h"
+#include "GameFramework/WorldSettings.h"
 #include "Misc/MessageDialog.h"
-#include "Widgets/Layout/SBorder.h"
-#include "Widgets/Text/STextBlock.h"
-#include "Widgets/Layout/SBox.h"
+#include "SEnumCombo.h"
 #include "Widgets/Input/SButton.h"
 #include "Widgets/Input/SCheckBox.h"
 #include "Widgets/Input/SNumericEntryBox.h"
+#include "Widgets/Layout/SBorder.h"
+#include "Widgets/Layout/SBox.h"
+#include "Widgets/Text/STextBlock.h"
 #include "IDocumentation.h"
 
 #define LOCTEXT_NAMESPACE "FWorldPartitionDetails"
@@ -140,6 +143,73 @@ void FWorldPartitionDetails::CustomizeDetails(IDetailLayoutBuilder& InDetailBuil
 			InDetailBuilder.HideProperty(InDetailBuilder.GetProperty(GET_MEMBER_NAME_CHECKED(UWorldPartition, DefaultHLODLayer), UWorldPartition::StaticClass()));
 		}
 	}
+
+	{
+		// Custom editor for UWorldPartitionEditorPerProjectUserSettings::bForceDisableDynamicLoadingRangeScaling, so it is easier for users to find than going into editor settings
+		WorldPartitionCategory.AddCustomRow(LOCTEXT("LocallyDisableDynamicLoadingRangeScaling", "Locally Disable Dynamic Loading Range Scaling"), true)
+		.NameContent()
+		[
+			SNew(STextBlock)
+			.Text(LOCTEXT("LocallyDisableDynamicLoadingRangeScaling", "Locally Disable Dynamic Loading Range Scaling"))
+			.ToolTipText(LOCTEXT("LocallyDisableDynamicLoadingRangeScaling_ToolTip", "Disable Dynamic Loading Range Scaling in local user settings, without modifying the map and affecting other users."))
+			.Font(IDetailLayoutBuilder::GetDetailFont())
+		]
+		.ValueContent()
+		[
+			SNew(SCheckBox)
+			.IsChecked(MakeAttributeLambda([this]()
+			{
+				if (WorldPartition.IsValid())
+				{
+					return GetDefault<UWorldPartitionEditorPerProjectUserSettings>()->GetForceDisableDynamicLoadingRangeScaling(WorldPartition->GetWorld()) ? ECheckBoxState::Checked : ECheckBoxState::Unchecked;
+				}
+				return ECheckBoxState::Unchecked;
+			}))
+			.OnCheckStateChanged_Lambda([this](ECheckBoxState InCheckBoxState)
+			{
+				if (WorldPartition.IsValid())
+				{
+					GetMutableDefault<UWorldPartitionEditorPerProjectUserSettings>()->SetForceDisableDynamicLoadingRangeScaling(WorldPartition->GetWorld(), InCheckBoxState == ECheckBoxState::Checked);
+				}
+			})
+		]
+		.Visibility(TAttribute<EVisibility>::CreateLambda([this]() { return WorldPartition.IsValid() && WorldPartition->SupportsStreaming() ? EVisibility::Visible : EVisibility::Hidden; }));
+
+		// Custom editor for AWorldSettings::DynamicLoadingRangeScaling, to allow disabling it on a complex condition
+		UWorld* World = WorldPartition->GetWorld();
+		AWorldSettings* WorldSettings = (World != nullptr) ? World->GetWorldSettings() : nullptr;
+		if (WorldSettings != nullptr)
+		{
+			const TSharedPtr<FStructOnScope> DataView = MakeShared<FStructOnScope>(FDynamicLoadingRangeScaling::StaticStruct(), reinterpret_cast<uint8*>(&WorldSettings->DynamicLoadingRangeScaling));
+			FAddPropertyParams AddParams;
+			AddParams.CreateCategoryNodes(true);
+			if (IDetailPropertyRow* PropertyRow = WorldPartitionCategory.AddExternalStructureProperty(DataView, NAME_None, EPropertyLocation::Advanced, AddParams))
+			{
+				PropertyRow->IsEnabled(TAttribute<bool>::CreateLambda([this]()
+				{
+					if (WorldPartition.IsValid())
+					{
+						return !GetDefault<UWorldPartitionEditorPerProjectUserSettings>()->GetForceDisableDynamicLoadingRangeScaling(WorldPartition->GetWorld());
+					}
+					return false;
+				}));
+				PropertyRow->Visibility(TAttribute<EVisibility>::CreateLambda([this]() { return WorldPartition.IsValid() && WorldPartition->SupportsStreaming() ? EVisibility::Visible : EVisibility::Hidden; }));
+				PropertyRow->ShouldAutoExpand(false);
+				PropertyRow->DisplayName(LOCTEXT("DynamicLoadingRangeScaling", "Dynamic Loading Range Scaling"));
+				PropertyRow->ToolTip(LOCTEXT("DynamicLoadingRangeScaling_ToolTip", "Locally force disable the dynamic loading range scaling - stored in local per-user settings."));
+
+				TSharedPtr<IPropertyHandle> PropertyHandle = PropertyRow->GetPropertyHandle();
+				PropertyHandle->SetOnPropertyValuePreChange(FSimpleDelegate::CreateLambda([this]()
+				{
+					if (WorldPartition.IsValid() && WorldPartition->GetWorld())
+					{
+						// Handle undo and dirty
+						WorldPartition->GetWorld()->GetWorldSettings()->Modify();
+					}
+				}));
+			}
+		}
+	}
 }
 
 void FWorldPartitionDetails::HandleWorldPartitionEnableStreamingChanged(ECheckBoxState InCheckState)
diff --git a/Engine/Source/Runtime/Engine/Classes/GameFramework/WorldSettings.h b/Engine/Source/Runtime/Engine/Classes/GameFramework/WorldSettings.h
--- a/Engine/Source/Runtime/Engine/Classes/GameFramework/WorldSettings.h
+++ b/Engine/Source/Runtime/Engine/Classes/GameFramework/WorldSettings.h
@@ -333,6 +333,29 @@ struct FBroadphaseSettings
 	uint32 MBPNumSubdivs;
 };
 
+/** Settings for dynamically increasing the loading range in PIE after BeginPlay - makes PIE startup quicker on large maps */
+USTRUCT()
+struct FDynamicLoadingRangeScaling
+{
+	GENERATED_BODY()
+
+	/** Enable for this map, unless user force disabled using bForceDisableDynamicLoadingRangeScaling. */
+	UPROPERTY(EditAnywhere, Category = WorldPartitionSetup, AdvancedDisplay)
+	bool bEnabled = false;
+
+	/** Start scaling at which ratio, from 0.f to 100.f. 100.f is equivalent to disabling the feature. 0.f is not recommended if floor or terrain tiles are streamed, as the pawn will have nothing to stand on initially. */
+	UPROPERTY(EditAnywhere, Category = WorldPartitionSetup, AdvancedDisplay, meta=(EditCondition = "bEnabled", Units="%", ClampMin = "0.0", ClampMax = "100.0", UIMin = "0.0", UIMax = "100.0"))
+	float InitialDynamicLoadingRangeScaling = 100.f;
+
+	/** Delay before scaling starts, to prevent additional world streaming while initial resources are also streaming in mips, LODs, etc. */
+	UPROPERTY(EditAnywhere, Category = WorldPartitionSetup, AdvancedDisplay, meta=(EditCondition = "bEnabled", Units="s", ClampMin = "0.0", UIMin = "0.0"))
+	float DynamicLoadingRangeScalingDelay = 0.f;
+
+	/** How long to scale for after the delay, in seconds. */
+	UPROPERTY(EditAnywhere, Category = WorldPartitionSetup, AdvancedDisplay, meta=(EditCondition = "bEnabled", Units="s", ClampMin = "0.0", UIMin = "0.0"))
+	float DynamicLoadingRangeScalingDuration = 10.f;
+};
+
 /**
  * Actor containing all script accessible world properties.
  */
@@ -343,6 +366,10 @@ class AWorldSettings : public AInfo, public IInterface_AssetUserData
 
 	ENGINE_API virtual void GetLifetimeReplicatedProps(TArray< FLifetimeProperty > & OutLifetimeProps) const override;
 
+#if WITH_EDITOR
+	bool GetDynamicLoadingRangeScalingEnable() const;
+#endif // WITH_EDITOR
+
 	/** PRECOMPUTED VISIBILITY SETTINGS **/
 
 	/** 
@@ -472,6 +499,10 @@ public:
 	uint8 bAllowTimeDilation : 1;
 
 #if WITH_EDITORONLY_DATA
+	// This property is customized in FWorldPartitionDetails
+	UPROPERTY()
+	FDynamicLoadingRangeScaling DynamicLoadingRangeScaling;
+
 	/** if set to true, this hide the streaming disabled warning available in the viewport */
 	UPROPERTY(EditAnywhere, Category = WorldPartitionSetup, AdvancedDisplay, meta = (EditCondition = "WorldPartition != nullptr"))
 	uint8 bHideEnableStreamingWarning : 1;
diff --git a/Engine/Source/Runtime/Engine/Private/WorldPartition/RuntimeHashSet/RuntimePartitionLevelStreaming.cpp b/Engine/Source/Runtime/Engine/Private/WorldPartition/RuntimeHashSet/RuntimePartitionLevelStreaming.cpp
--- a/Engine/Source/Runtime/Engine/Private/WorldPartition/RuntimeHashSet/RuntimePartitionLevelStreaming.cpp
+++ b/Engine/Source/Runtime/Engine/Private/WorldPartition/RuntimeHashSet/RuntimePartitionLevelStreaming.cpp
@@ -22,8 +22,6 @@ bool URuntimePartitionLevelStreaming::GenerateStreaming(const FGenerateStreaming
 	TMap<FName, TArray<const IStreamingGenerationContext::FActorSetInstance*>> CellsActorSetInstances;
 	for (const IStreamingGenerationContext::FActorSetInstance* ActorSetInstance : *InParams.ActorSetInstances)
 	{
-		TArray<FName> ActorSetGridNameList;
-
 		FName LevelName = NAME_Default;
 		if (!ActorSetInstance->RuntimeGrid.IsNone())
 		{
@@ -37,7 +35,6 @@ bool URuntimePartitionLevelStreaming::GenerateStreaming(const FGenerateStreaming
 				}
 			}
 		}
-		ActorSetGridNameList.Add(LevelName);
 
 		FName CellName;
 		if (ActorSetInstance->bIsSpatiallyLoaded)
diff --git a/Engine/Source/Runtime/Engine/Private/WorldPartition/RuntimeHashSet/WorldPartitionRuntimeHashSet.cpp b/Engine/Source/Runtime/Engine/Private/WorldPartition/RuntimeHashSet/WorldPartitionRuntimeHashSet.cpp
--- a/Engine/Source/Runtime/Engine/Private/WorldPartition/RuntimeHashSet/WorldPartitionRuntimeHashSet.cpp
+++ b/Engine/Source/Runtime/Engine/Private/WorldPartition/RuntimeHashSet/WorldPartitionRuntimeHashSet.cpp
@@ -641,6 +641,17 @@ void UWorldPartitionRuntimeHashSet::ForEachStreamingCellsSources(const TArray<FW
 	const FWorldPartitionStreamingContext& Context = InContext.IsValid() ? InContext : StackContext;
 	check(Context.IsValid());
 
+	float DynamicLoadingRangeScale2D = 1.f;
+	float DynamicLoadingRangeScale3D = 1.f;
+#if WITH_EDITOR
+	if (UWorldPartitionSubsystem* WorldPartitionSubsystem = UWorld::GetSubsystem<UWorldPartitionSubsystem>(GetWorld()))
+	{
+		const FDynamicLoadingRangeScales& Scales = WorldPartitionSubsystem->GetDynamicLoadingRangeScales();
+		DynamicLoadingRangeScale2D = Scales.Scale2D;
+		DynamicLoadingRangeScale3D = Scales.Scale3D;
+	}
+#endif // WITH_EDITOR
+
 	// Non-spatially loaded cells
 	for (const FRuntimePartitionStreamingData* StreamingData : RuntimeNonSpatiallyLoadedDataGridList)
 	{
@@ -697,7 +708,7 @@ void UWorldPartitionRuntimeHashSet::ForEachStreamingCellsSources(const TArray<FW
 
 				for (const FRuntimePartitionStreamingData* StreamingData : *StreamingDataList)
 				{
-					Source.ForEachShape(StreamingData->GetLoadingRange(), false, [this, &Source, StreamingData, &Context, &Func](const FSphericalSector& Shape)
+					Source.ForEachShape(StreamingData->GetLoadingRange(), false, [this, &Source, StreamingData, &Context, &Func, DynamicLoadingRangeScale2D, DynamicLoadingRangeScale3D](const FSphericalSector& Shape)
 					{
 						auto ForEachIntersectingElementFunc = [this, &Source, &Shape, &Context, &Func](UWorldPartitionRuntimeCell* Cell)
 						{
@@ -717,34 +728,38 @@ void UWorldPartitionRuntimeHashSet::ForEachStreamingCellsSources(const TArray<FW
 							}
 						};
 
-						auto ForEachIntersectingElement = [&StreamingData, &Source, &ForEachIntersectingElementFunc]<typename ShapeType>(const ShapeType& Shape)
+						auto ForEachIntersectingElement = [&StreamingData, &Source, &ForEachIntersectingElementFunc, DynamicLoadingRangeScale2D, DynamicLoadingRangeScale3D]<typename ShapeType>(ShapeType& Shape)
 						{
+							float ShapeRadius = Shape.Radius;
 							if (StreamingData->SpatialIndex.IsValid())
 							{
 								if (Source.bForce2D)
 								{
+									Shape.Radius = ShapeRadius * DynamicLoadingRangeScale2D;
 									StreamingData->SpatialIndexForce2D->ForEachIntersectingElement(Shape, ForEachIntersectingElementFunc);
 								}
 								else
 								{
+									Shape.Radius = ShapeRadius * DynamicLoadingRangeScale3D;
 									StreamingData->SpatialIndex->ForEachIntersectingElement(Shape, ForEachIntersectingElementFunc);
 								}
 							}
 
 							if (StreamingData->SpatialIndex2D.IsValid())
 							{
+								Shape.Radius = ShapeRadius * DynamicLoadingRangeScale2D;
 								StreamingData->SpatialIndex2D->ForEachIntersectingElement(Shape, ForEachIntersectingElementFunc);
 							}
 						};
 
 						if (Shape.IsSphere())
 						{
-							const FStaticSpatialIndex::FSphere ShapeSphere(Shape.GetCenter(), Shape.GetRadius());
+							FStaticSpatialIndex::FSphere ShapeSphere(Shape.GetCenter(), Shape.GetRadius());
 							ForEachIntersectingElement(ShapeSphere);
 						}
 						else
 						{
-							const FStaticSpatialIndex::FCone ShapeCone(Shape.GetCenter(), Shape.GetAxis(), Shape.GetRadius(), Shape.GetAngle());
+							FStaticSpatialIndex::FCone ShapeCone(Shape.GetCenter(), Shape.GetAxis(), Shape.GetRadius(), Shape.GetAngle());
 							ForEachIntersectingElement(ShapeCone);
 						}
 					});
diff --git a/Engine/Source/Runtime/Engine/Private/WorldPartition/RuntimeHashSet/WorldPartitionRuntimeHashSetDebugDraw.cpp b/Engine/Source/Runtime/Engine/Private/WorldPartition/RuntimeHashSet/WorldPartitionRuntimeHashSetDebugDraw.cpp
--- a/Engine/Source/Runtime/Engine/Private/WorldPartition/RuntimeHashSet/WorldPartitionRuntimeHashSetDebugDraw.cpp
+++ b/Engine/Source/Runtime/Engine/Private/WorldPartition/RuntimeHashSet/WorldPartitionRuntimeHashSetDebugDraw.cpp
@@ -8,6 +8,7 @@
 #include "WorldPartition/DataLayer/DataLayerManager.h"
 #include "WorldPartition/WorldPartitionDraw2DContext.h"
 #include "WorldPartition/WorldPartitionDebugHelper.h"
+#include "WorldPartition/WorldPartitionSubsystem.h"
 #include "DrawDebugHelpers.h"
 #include "Math/TransformCalculus2D.h"
 #include "Misc/Paths.h"
@@ -112,7 +113,18 @@ bool UWorldPartitionRuntimeHashSet::Draw2D(FWorldPartitionDraw2DContext& DrawCon
 	const FVector2D GridScreenHalfExtent = 0.5f * GridScreenExtent;
 	const FVector2D GridScreenInitialOffset = CanvasRegion.Min;
 
-	auto DrawStreamingData = [this, &WorldPartitionTransform](const FRuntimePartitionStreamingData* StreamingData, const FBox2D& Region2D, const FBox2D& GridScreenBounds, TFunctionRef<FVector2D(const FVector2D&, bool)> InWorldToScreen, FWorldPartitionDraw2DContext& DrawContext)
+	float DynamicLoadingRangeScale2D = 1.f;
+	float DynamicLoadingRangeScale3D = 1.f;
+#if WITH_EDITOR
+	if (UWorldPartitionSubsystem* WorldPartitionSubsystem = UWorld::GetSubsystem<UWorldPartitionSubsystem>(GetWorld()))
+	{
+		const FDynamicLoadingRangeScales& Scales = WorldPartitionSubsystem->GetDynamicLoadingRangeScales();
+		DynamicLoadingRangeScale2D = Scales.Scale2D;
+		DynamicLoadingRangeScale3D = Scales.Scale3D;
+	}
+#endif // WITH_EDITOR
+
+	auto DrawStreamingData = [this, &WorldPartitionTransform, DynamicLoadingRangeScale2D, DynamicLoadingRangeScale3D](const FRuntimePartitionStreamingData* StreamingData, const FBox2D& Region2D, const FBox2D& GridScreenBounds, TFunctionRef<FVector2D(const FVector2D&, bool)> InWorldToScreen, FWorldPartitionDraw2DContext& DrawContext)
 	{
 		const UWorldPartition* WorldPartition = GetOuterUWorldPartition();
 		const TArray<FWorldPartitionStreamingSource>& Sources = WorldPartition->GetStreamingSources();
@@ -235,17 +247,19 @@ bool UWorldPartitionRuntimeHashSet::Draw2D(FWorldPartitionDraw2DContext& DrawCon
 		for (const FWorldPartitionStreamingSource& Source : Sources)
 		{
 			const FColor Color = Source.GetDebugColor();
-			Source.ForEachShape(StreamingData->GetLoadingRange(), StreamingData->Name, true, [&Color, &WorldToScreen, &GridScreenBounds, &DrawContext, this](const FSphericalSector& Shape)
+			Source.ForEachShape(StreamingData->GetLoadingRange(), StreamingData->Name, true, [&Color, &WorldToScreen, &GridScreenBounds, &DrawContext, &Source, DynamicLoadingRangeScale2D, DynamicLoadingRangeScale3D, this](const FSphericalSector& Shape)
 			{
 				check(!Shape.IsNearlyZero())
 
+				float DynamicLoadingRangeScale = (Source.bForce2D) ? DynamicLoadingRangeScale2D : DynamicLoadingRangeScale3D;
+
 				// Spherical Sector
 				const FVector2D Center2D(FVector2D(Shape.GetCenter()));
 				const FSphericalSector::FReal Angle = Shape.GetAngle();
 				const int32 MaxSegments = FMath::Max(4, FMath::CeilToInt(64 * Angle / 360.f));
 				const float AngleIncrement = Angle / MaxSegments;
 				const FVector2D Axis = FVector2D(Shape.GetAxis());
-				const FVector Startup = FRotator(0, -0.5f * Angle, 0).RotateVector(Shape.GetScaledAxis());				
+				const FVector Startup = FRotator(0, -0.5f * Angle, 0).RotateVector(Shape.GetScaledAxis() * DynamicLoadingRangeScale);
 				FVector2D LineStart = FVector2D(Startup);
 
 				if (!Shape.IsSphere())
@@ -266,7 +280,7 @@ bool UWorldPartitionRuntimeHashSet::Draw2D(FWorldPartitionDraw2DContext& DrawCon
 				DrawContext.PushDrawSegment(GridScreenBounds, WorldToScreen(Center2D + LineStart), WorldToScreen(Center2D + (Shape.IsSphere() ? FVector2D(Startup) : FVector2D::ZeroVector)), Color, 2);
 
 				// Draw direction vector
-				DrawContext.PushDrawSegment(GridScreenBounds, WorldToScreen(Center2D), WorldToScreen(Center2D + Axis * Shape.GetRadius()), Color, 4);
+				DrawContext.PushDrawSegment(GridScreenBounds, WorldToScreen(Center2D), WorldToScreen(Center2D + Axis * Shape.GetRadius() * DynamicLoadingRangeScale), Color, 4);
 			});
 
 			// Draw velocity vector			
@@ -365,6 +379,17 @@ void UWorldPartitionRuntimeHashSet::Draw3D(const TArray<FWorldPartitionStreaming
 	const EWorldPartitionRuntimeCellVisualizeMode VisualizeMode = GetStreamingCellVisualizeMode();
 	const TMap<FName, FColor> DataLayerDebugColors = GetDataLayerDebugColors(WorldPartition);
 
+	float DynamicLoadingRangeScale2D = 1.f;
+	float DynamicLoadingRangeScale3D = 1.f;
+#if WITH_EDITOR
+	if (UWorldPartitionSubsystem* WorldPartitionSubsystem = UWorld::GetSubsystem<UWorldPartitionSubsystem>(GetWorld()))
+	{
+		const FDynamicLoadingRangeScales& Scales = WorldPartitionSubsystem->GetDynamicLoadingRangeScales();
+		DynamicLoadingRangeScale2D = Scales.Scale2D;
+		DynamicLoadingRangeScale3D = Scales.Scale3D;
+	}
+#endif // WITH_EDITOR
+
 	auto DebugDrawCell3D = [WorldPartitionTransform, OwningWorld, ContentBundleManager, VisualizeMode, &DataLayerDebugColors](UWorldPartitionRuntimeCell* Cell)
 	{
 		if (!IsDebugCellShown(Cell))
@@ -422,13 +447,13 @@ void UWorldPartitionRuntimeHashSet::Draw3D(const TArray<FWorldPartitionStreaming
 		DrawDebugBox(OwningWorld, CellPos, Box.GetExtent(), WorldPartitionTransform.GetRotation(), BoxColor.WithAlpha(255), false, -1.f, 0, 20.f);
 	};
 
-	ForEachStreamingData([this, &Sources, VisualizeMode, &DataLayerDebugColors, ContentBundleManager, OwningWorld, &WorldPartitionTransform, &DebugDrawCell3D](const FRuntimePartitionStreamingData& StreamingData)
+	ForEachStreamingData([this, &Sources, VisualizeMode, &DataLayerDebugColors, ContentBundleManager, OwningWorld, &WorldPartitionTransform, &DebugDrawCell3D, DynamicLoadingRangeScale2D, DynamicLoadingRangeScale3D](const FRuntimePartitionStreamingData& StreamingData)
 	{
 		if (StreamingData.SpatialIndex)
 		{
 			for (const FWorldPartitionStreamingSource& Source : Sources)
 			{
-				Source.ForEachShape(StreamingData.GetLoadingRange(), StreamingData.Name, false, [this, &StreamingData, VisualizeMode, &DataLayerDebugColors, ContentBundleManager, OwningWorld, &WorldPartitionTransform, &DebugDrawCell3D](const FSphericalSector& Shape)
+				Source.ForEachShape(StreamingData.GetLoadingRange(), StreamingData.Name, false, [this, &StreamingData, VisualizeMode, &DataLayerDebugColors, ContentBundleManager, OwningWorld, &WorldPartitionTransform, &DebugDrawCell3D, &Source, DynamicLoadingRangeScale2D, DynamicLoadingRangeScale3D](const FSphericalSector& Shape)
 				{
 					auto ForEachIntersectingElement = [&]<typename ShapeType>(const ShapeType& InShape)
 					{
@@ -438,14 +463,16 @@ void UWorldPartitionRuntimeHashSet::Draw3D(const TArray<FWorldPartitionStreaming
 						});
 					};
 
+					float DynamicLoadingRangeScale = (Source.bForce2D) ? DynamicLoadingRangeScale2D : DynamicLoadingRangeScale3D;
+
 					if (Shape.IsSphere())
 					{
-						const FStaticSpatialIndex::FSphere Sphere(Shape.GetCenter(), Shape.GetRadius());
+						const FStaticSpatialIndex::FSphere Sphere(Shape.GetCenter(), Shape.GetRadius() * DynamicLoadingRangeScale);
 						ForEachIntersectingElement(Sphere);
 					}
 					else
 					{
-						const FStaticSpatialIndex::FCone Cone(Shape.GetCenter(), Shape.GetAxis(), Shape.GetRadius(), Shape.GetAngle());
+						const FStaticSpatialIndex::FCone Cone(Shape.GetCenter(), Shape.GetAxis(), Shape.GetRadius() * DynamicLoadingRangeScale, Shape.GetAngle());
 						ForEachIntersectingElement(Cone);
 					}
 				});
@@ -457,7 +484,7 @@ void UWorldPartitionRuntimeHashSet::Draw3D(const TArray<FWorldPartitionStreaming
 			for (const FWorldPartitionStreamingSource& Source : Sources)
 			{
 				const bool bCalcBoundsIn2D = true;
-				const FBox SourceBounds = Source.CalcBounds(StreamingData.GetLoadingRange(), StreamingData.Name, bCalcBoundsIn2D);
+				const FBox SourceBounds = Source.CalcBounds(StreamingData.GetLoadingRange() * DynamicLoadingRangeScale2D, StreamingData.Name, bCalcBoundsIn2D);
 				const FBox2D SourceBounds2D(FVector2D(SourceBounds.Min.X, SourceBounds.Min.Y), FVector2D(SourceBounds.Max.X, SourceBounds.Max.Y));
 				StreamingData.SpatialIndex2D->ForEachIntersectingElement(SourceBounds2D, [&DebugDrawCell3D](UWorldPartitionRuntimeCell* Cell)
 				{
diff --git a/Engine/Source/Runtime/Engine/Private/WorldPartition/WorldPartitionEditorPerProjectUserSettings.cpp b/Engine/Source/Runtime/Engine/Private/WorldPartition/WorldPartitionEditorPerProjectUserSettings.cpp
--- a/Engine/Source/Runtime/Engine/Private/WorldPartition/WorldPartitionEditorPerProjectUserSettings.cpp
+++ b/Engine/Source/Runtime/Engine/Private/WorldPartition/WorldPartitionEditorPerProjectUserSettings.cpp
@@ -140,6 +140,27 @@ const FWorldPartitionPerWorldSettings* UWorldPartitionEditorPerProjectUserSettin
 	return nullptr;
 }
 
+#if WITH_EDITORONLY_DATA
+bool UWorldPartitionEditorPerProjectUserSettings::GetForceDisableDynamicLoadingRangeScaling(UWorld* InWorld) const
+{
+	if (const FWorldPartitionPerWorldSettings* PerWorldSettings = GetWorldPartitionPerWorldSettings(InWorld))
+	{
+		return PerWorldSettings->bForceDisableDynamicLoadingRangeScaling;
+	}
+	return false;
+}
+
+void UWorldPartitionEditorPerProjectUserSettings::SetForceDisableDynamicLoadingRangeScaling(UWorld* InWorld, bool bOverride)
+{
+	if (ShouldSaveSettings(InWorld))
+	{
+		FWorldPartitionPerWorldSettings& PerWorldSettings = PerWorldEditorSettings.FindOrAdd(TSoftObjectPtr<UWorld>(InWorld));
+		PerWorldSettings.bForceDisableDynamicLoadingRangeScaling = bOverride;
+		SaveConfig();
+	}
+}
+#endif // WITH_EDITORONLY_DATA
+
 bool UWorldPartitionEditorPerProjectUserSettings::ShouldSaveSettings(const UWorld* InWorld) const
 {
 	return InWorld && !InWorld->IsGameWorld() && (InWorld->WorldType != EWorldType::Inactive) && FPackageName::DoesPackageExist(InWorld->GetPackage()->GetName());
@@ -151,3 +172,4 @@ bool UWorldPartitionEditorPerProjectUserSettings::ShouldLoadSettings(const UWorl
 }
 
 #endif
+
diff --git a/Engine/Source/Runtime/Engine/Private/WorldPartition/WorldPartitionRuntimeSpatialHash.cpp b/Engine/Source/Runtime/Engine/Private/WorldPartition/WorldPartitionRuntimeSpatialHash.cpp
--- a/Engine/Source/Runtime/Engine/Private/WorldPartition/WorldPartitionRuntimeSpatialHash.cpp
+++ b/Engine/Source/Runtime/Engine/Private/WorldPartition/WorldPartitionRuntimeSpatialHash.cpp
@@ -374,7 +374,7 @@ void FSpatialHashStreamingGrid::GetCells(const FWorldPartitionStreamingQuerySour
 	}
 }
 
-void FSpatialHashStreamingGrid::GetCells(const TArray<FWorldPartitionStreamingSource>& Sources, UWorldPartitionRuntimeHash::FStreamingSourceCells& OutActivateCells, UWorldPartitionRuntimeHash::FStreamingSourceCells& OutLoadCells, bool bEnableZCulling, const FWorldPartitionStreamingContext& Context) const
+void FSpatialHashStreamingGrid::GetCells(const TArray<FWorldPartitionStreamingSource>& Sources, UWorldPartitionRuntimeHash::FStreamingSourceCells& OutActivateCells, UWorldPartitionRuntimeHash::FStreamingSourceCells& OutLoadCells, bool bEnableZCulling, const FWorldPartitionStreamingContext& Context, float LoadingRangeScale) const
 {
 	TRACE_CPUPROFILER_EVENT_SCOPE(FSpatialHashStreamingGrid::GetCells);
 
@@ -393,13 +393,25 @@ void FSpatialHashStreamingGrid::GetCells(const TArray<FWorldPartitionStreamingSo
 
 	const float GridLoadingRange = GetLoadingRange();
 	const FSquare2DGridHelper& Helper = GetGridHelper();
+
+	FSphericalSector AdjustedShape{EForceInit::ForceInit};
+
 	for (const FWorldPartitionStreamingSource& Source : Sources)
 	{
 		Source.ForEachShape(GridLoadingRange, GridName, /*bProjectIn2D*/ true, [&](const FSphericalSector& Shape)
 		{
-			FStreamingSourceInfo Info(Source, Shape);
+			// Adjust radius for dynamic scales
+			const FSphericalSector* EffectiveShape = &Shape;
+			if (LoadingRangeScale < 1.f)
+			{
+				AdjustedShape = Shape;
+				AdjustedShape.SetRadius(AdjustedShape.GetRadius() * LoadingRangeScale);
+				EffectiveShape = &AdjustedShape;
+			}
 
-			Helper.ForEachIntersectingCells(Shape, [&](const FGridCellCoord& Coords)
+			FStreamingSourceInfo Info(Source, *EffectiveShape);
+
+			Helper.ForEachIntersectingCells(*EffectiveShape, [&](const FGridCellCoord& Coords)
 			{
 				bool bAddedActivatedCell = false;
 
@@ -414,7 +426,7 @@ void FSpatialHashStreamingGrid::GetCells(const TArray<FWorldPartitionStreamingSo
 						if (bEnableZCulling)
 						{
 							const FVector2D CellMinMaxZ(Cell->GetContentBounds().Min.Z, Cell->GetContentBounds().Max.Z);
-							bIncludeCell = TRange<double>::Inclusive(CellMinMaxZ.X, CellMinMaxZ.Y).Overlaps(TRange<double>::Inclusive(Shape.GetCenter().Z - Shape.GetRadius(), Shape.GetCenter().Z + Shape.GetRadius()));
+							bIncludeCell = TRange<double>::Inclusive(CellMinMaxZ.X, CellMinMaxZ.Y).Overlaps(TRange<double>::Inclusive(EffectiveShape->GetCenter().Z - EffectiveShape->GetRadius(), EffectiveShape->GetCenter().Z + EffectiveShape->GetRadius()));
 						}
 						
 						if (bIncludeCell)
@@ -422,16 +434,16 @@ void FSpatialHashStreamingGrid::GetCells(const TArray<FWorldPartitionStreamingSo
 							switch (Cell->GetCellEffectiveWantedState(Context))
 							{
 							case EDataLayerRuntimeState::Loaded:
-								OutLoadCells.AddCell(Cell, Source, Shape, Context);
+								OutLoadCells.AddCell(Cell, Source, *EffectiveShape, Context);
 								break;
 							case EDataLayerRuntimeState::Activated:
 								switch (Source.TargetState)
 								{
 								case EStreamingSourceTargetState::Loaded:
-									OutLoadCells.AddCell(Cell, Source, Shape, Context);
+									OutLoadCells.AddCell(Cell, Source, *EffectiveShape, Context);
 									break;
 								case EStreamingSourceTargetState::Activated:
-									OutActivateCells.AddCell(Cell, Source, Shape, Context);
+									OutActivateCells.AddCell(Cell, Source, *EffectiveShape, Context);
 									bAddedActivatedCell = !Settings.bUseAlignedGridLevels && Settings.bSnapNonAlignedGridLevelsToLowerLevels;
 									break;
 								default:
@@ -837,6 +849,14 @@ void FSpatialHashStreamingGrid::Draw3D(const UWorldPartitionRuntimeSpatialHash*
 	const UContentBundleManager* ContentBundleManager = OwningWorld->ContentBundleManager;
 	TMap<FName, FColor> DataLayerDebugColors = GetDataLayerDebugColors2(WorldPartition);
 
+	float DynamicLoadingRangeScale = 1.f;
+#if WITH_EDITOR
+	if (UWorldPartitionSubsystem* WorldPartitionSubsystem = UWorld::GetSubsystem<UWorldPartitionSubsystem>(Owner->GetWorld()))
+	{
+		const FDynamicLoadingRangeScales& Scales = WorldPartitionSubsystem->GetDynamicLoadingRangeScales();
+		DynamicLoadingRangeScale = Scales.Scale2D;
+	}
+#endif // WITH_EDITOR
 	const FSquare2DGridHelper& Helper = GetGridHelper();
 	int32 MinGridLevel = FMath::Clamp<int32>(GShowRuntimeSpatialHashGridLevel, 0, GridLevels.Num() - 1);
 	int32 MaxGridLevel = FMath::Clamp<int32>(MinGridLevel + GShowRuntimeSpatialHashGridLevelCount - 1, 0, GridLevels.Num() - 1);
@@ -849,7 +869,7 @@ void FSpatialHashStreamingGrid::Draw3D(const UWorldPartitionRuntimeSpatialHash*
 
 	for (const FWorldPartitionStreamingSource& Source : Sources)
 	{
-		FBox Region = Source.CalcBounds(GridLoadingRange, GridName);
+		FBox Region = Source.CalcBounds(GridLoadingRange * DynamicLoadingRangeScale, GridName);
 		Region += FBox(Region.GetCenter() - MinExtent, Region.GetCenter() + MinExtent);
 
 		for (int32 GridLevel = MinGridLevel; GridLevel <= MaxGridLevel; ++GridLevel)
@@ -937,9 +957,11 @@ void FSpatialHashStreamingGrid::Draw3D(const UWorldPartitionRuntimeSpatialHash*
 		}
 
 		const FColor Color = Source.GetDebugColor();
-		Source.ForEachShape(GetLoadingRange(), GridName, /*bProjectIn2D*/ true, [&Color, &SourceLocationZ, &Transform, &OwningWorld, this](const FSphericalSector& Shape)
+		Source.ForEachShape(GridLoadingRange, GridName, /*bProjectIn2D*/ true, [&Color, &SourceLocationZ, &Transform, &OwningWorld, DynamicLoadingRangeScale, this](const FSphericalSector& Shape)
 		{
 			FSphericalSector ZOffsettedShape = Shape;
+			// Adjust radius for dynamic scales
+			ZOffsettedShape.SetRadius(ZOffsettedShape.GetRadius() * DynamicLoadingRangeScale);
 			ZOffsettedShape.SetCenter(FVector(FVector2D(ZOffsettedShape.GetCenter()), SourceLocationZ));
 			DrawStreamingSource3D(OwningWorld, ZOffsettedShape, Transform, Color);
 		});
@@ -1086,12 +1108,21 @@ void FSpatialHashStreamingGrid::Draw2D(const UWorldPartitionRuntimeSpatialHash*
 
 	// Draw Streaming Sources
 	const float GridLoadingRange = GetLoadingRange();
+	float DynamicLoadingRangeScale = 1.f;
+#if WITH_EDITOR
+	if (UWorldPartitionSubsystem* WorldPartitionSubsystem = UWorld::GetSubsystem<UWorldPartitionSubsystem>(Owner->GetWorld()))
+	{
+		const FDynamicLoadingRangeScales& Scales = WorldPartitionSubsystem->GetDynamicLoadingRangeScales();
+		DynamicLoadingRangeScale = Scales.Scale2D;
+	}
+#endif // WITH_EDITOR
+
 	for (const FWorldPartitionStreamingSource& Source : Sources)
 	{
 		const FColor Color = Source.GetDebugColor();
-		Source.ForEachShape(GridLoadingRange, GridName, /*bProjectIn2D*/ true, [&Color, &WorldToScreen, &GridScreenBounds, &DrawContext, this](const FSphericalSector& Shape)
+		Source.ForEachShape(GridLoadingRange, GridName, /*bProjectIn2D*/ true, [&Color, &WorldToScreen, &GridScreenBounds, &DrawContext, DynamicLoadingRangeScale, this](const FSphericalSector& Shape)
 		{
-			DrawStreamingSource2D(GridScreenBounds, Shape, WorldToScreen, Color, DrawContext);
+			DrawStreamingSource2D(GridScreenBounds, Shape, WorldToScreen, Color, DrawContext, DynamicLoadingRangeScale);
 		});
 
 		// Draw velocity vector			
@@ -1105,9 +1136,9 @@ void FSpatialHashStreamingGrid::Draw2D(const UWorldPartitionRuntimeSpatialHash*
 	}
 }
 
-void FSpatialHashStreamingGrid::DrawStreamingSource2D(const FBox2D& GridScreenBounds, const FSphericalSector& Shape, TFunctionRef<FVector2D(const FVector2D&)> WorldToScreen, const FColor& Color, FWorldPartitionDraw2DContext& DrawContext) const
+void FSpatialHashStreamingGrid::DrawStreamingSource2D(const FBox2D& GridScreenBounds, const FSphericalSector& Shape, TFunctionRef<FVector2D(const FVector2D&)> WorldToScreen, const FColor& Color, FWorldPartitionDraw2DContext& DrawContext, float LoadingRangeScale) const
 {
-	check(!Shape.IsNearlyZero())
+	check(!Shape.IsNearlyZero());
 
 	// Spherical Sector
 	const FVector2D Center2D(FVector2D(Shape.GetCenter()));
@@ -1115,7 +1146,7 @@ void FSpatialHashStreamingGrid::DrawStreamingSource2D(const FBox2D& GridScreenBo
 	const int32 MaxSegments = FMath::Max(4, FMath::CeilToInt(64 * Angle / 360.f));
 	const float AngleIncrement = Angle / MaxSegments;
 	const FVector2D Axis = FVector2D(Shape.GetAxis());
-	const FVector Startup = FRotator(0, -0.5f * Angle, 0).RotateVector(Shape.GetScaledAxis());
+	const FVector Startup = FRotator(0, -0.5f * Angle, 0).RotateVector(Shape.GetScaledAxis() * LoadingRangeScale);
 
 	FVector2D LineStart = FVector2D(Startup);
 	if (!Shape.IsSphere())
@@ -1134,7 +1165,7 @@ void FSpatialHashStreamingGrid::DrawStreamingSource2D(const FBox2D& GridScreenBo
 	DrawContext.PushDrawSegment(GridScreenBounds, WorldToScreen(Center2D + LineStart), WorldToScreen(Center2D + (Shape.IsSphere() ? FVector2D(Startup) : FVector2D::ZeroVector)), Color, 2);
 
 	// Draw direction vector
-	DrawContext.PushDrawSegment(GridScreenBounds, WorldToScreen(Center2D), WorldToScreen(Center2D + Axis * Shape.GetRadius()), Color, 2);
+	DrawContext.PushDrawSegment(GridScreenBounds, WorldToScreen(Center2D), WorldToScreen(Center2D + Axis * Shape.GetRadius() * LoadingRangeScale), Color, 2);
 }
 
 // ------------------------------------------------------------------------------------------------
@@ -1796,6 +1827,14 @@ void UWorldPartitionRuntimeSpatialHash::ForEachStreamingCellsSources(const TArra
 
 	FStreamingSourceCells ActivateStreamingSourceCells;
 	FStreamingSourceCells LoadStreamingSourceCells;
+	float DynamicLoadingRangeScale = 1.f;
+#if WITH_EDITOR
+	if (UWorldPartitionSubsystem* WorldPartitionSubsystem = UWorld::GetSubsystem<UWorldPartitionSubsystem>(GetWorld()))
+	{
+		const FDynamicLoadingRangeScales& Scales = WorldPartitionSubsystem->GetDynamicLoadingRangeScales();
+		DynamicLoadingRangeScale = Scales.Scale2D;
+	}
+#endif // WITH_EDITOR
 
 	if (Sources.Num() == 0)
 	{
@@ -1815,7 +1854,7 @@ void UWorldPartitionRuntimeSpatialHash::ForEachStreamingCellsSources(const TArra
 		{
 			if (IsCellRelevantFor(StreamingGrid.bClientOnlyVisible))
 			{
-				StreamingGrid.GetCells(Sources, ActivateStreamingSourceCells, LoadStreamingSourceCells, GetEffectiveEnableZCulling(bEnableZCulling), Context);
+				StreamingGrid.GetCells(Sources, ActivateStreamingSourceCells, LoadStreamingSourceCells, GetEffectiveEnableZCulling(bEnableZCulling), Context, DynamicLoadingRangeScale);
 			}
 		});
 	}
@@ -2306,11 +2345,15 @@ bool UWorldPartitionRuntimeSpatialHash::Draw2D(FWorldPartitionDraw2DContext& Dra
 	FBox GridsShapeBounds(ForceInit);
 	FBox2D GridsBounds(ForceInit);
 	int32 GridIndex = 0;
+
 	for (const FSpatialHashStreamingGrid* StreamingGrid : FilteredStreamingGrids)
 	{
 		for (const FWorldPartitionStreamingSource& Source : Sources)
 		{
-			Source.ForEachShape(StreamingGrid->GetLoadingRange(), StreamingGrid->GridName, /*bProjectIn2D*/ true, [&GridsShapeBounds](const FSphericalSector& Shape) { GridsShapeBounds += Shape.CalcBounds(); });
+			Source.ForEachShape(StreamingGrid->GetLoadingRange(), StreamingGrid->GridName, /*bProjectIn2D*/ true, [&GridsShapeBounds](const FSphericalSector& Shape)
+			{
+				GridsShapeBounds += Shape.CalcBounds();
+			});
 		}
 
 		FVector2D GridReferenceWorldPos;
diff --git a/Engine/Source/Runtime/Engine/Private/WorldPartition/WorldPartitionStreamingPolicy.cpp b/Engine/Source/Runtime/Engine/Private/WorldPartition/WorldPartitionStreamingPolicy.cpp
--- a/Engine/Source/Runtime/Engine/Private/WorldPartition/WorldPartitionStreamingPolicy.cpp
+++ b/Engine/Source/Runtime/Engine/Private/WorldPartition/WorldPartitionStreamingPolicy.cpp
@@ -184,6 +184,10 @@ uint32 UWorldPartitionStreamingPolicy::ComputeUpdateStreamingHash(bool bCanOptim
 			HashBuilder << GetWorld()->GetSubsystem<UWorldPartitionSubsystem>()->GetServerClientsVisibleLevelsHash();
 		}
 
+#if WITH_EDITOR
+		HashBuilder << GetWorld()->GetSubsystem<UWorldPartitionSubsystem>()->GetDynamicLoadingRangeScales().LinearScale;
+#endif // WITH_EDITOR
+
 		return HashBuilder.GetHash();
 	}
 
diff --git a/Engine/Source/Runtime/Engine/Private/WorldPartition/WorldPartitionSubsystem.cpp b/Engine/Source/Runtime/Engine/Private/WorldPartition/WorldPartitionSubsystem.cpp
--- a/Engine/Source/Runtime/Engine/Private/WorldPartition/WorldPartitionSubsystem.cpp
+++ b/Engine/Source/Runtime/Engine/Private/WorldPartition/WorldPartitionSubsystem.cpp
@@ -11,6 +11,7 @@
 #include "WorldPartition/WorldPartitionStreamingSource.h"
 #include "WorldPartition/WorldPartitionReplay.h"
 #include "WorldPartition/WorldPartitionDraw2DContext.h"
+#include "WorldPartition/WorldPartitionSettings.h"
 #include "WorldPartition/WorldPartitionStreamingPolicy.h"
 #include "WorldPartition/ActorDescContainerSubsystem.h"
 #include "WorldPartition/DataLayer/WorldDataLayersActorDesc.h"
@@ -474,6 +475,10 @@ void UWorldPartitionSubsystem::Initialize(FSubsystemCollectionBase& Collection)
 	GetWorld()->OnWorldPartitionUninitialized().AddUObject(this, &UWorldPartitionSubsystem::OnWorldPartitionUninitialized);
 	if (GetWorld()->IsGameWorld())
 	{
+#if WITH_EDITOR
+		FWorldDelegates::OnStartGameInstance.AddUObject(this, &UWorldPartitionSubsystem::OnStartGameInstance);
+#endif // WITH_EDITOR
+
 		FLevelStreamingDelegates::OnLevelBeginMakingVisible.AddUObject(this, &UWorldPartitionSubsystem::OnLevelBeginMakingVisible);
 		FLevelStreamingDelegates::OnLevelBeginMakingInvisible.AddUObject(this, &UWorldPartitionSubsystem::OnLevelBeginMakingInvisible);
 		FLevelStreamingDelegates::OnLevelStreamingTargetStateChanged.AddUObject(this, &UWorldPartitionSubsystem::OnLevelStreamingTargetStateChanged);
@@ -503,6 +508,10 @@ void UWorldPartitionSubsystem::Deinitialize()
 		FLevelStreamingDelegates::OnLevelBeginMakingInvisible.RemoveAll(this);
 		FLevelStreamingDelegates::OnLevelStreamingTargetStateChanged.RemoveAll(this);
 		FLevelStreamingDelegates::OnLevelStreamingStateChanged.RemoveAll(this);
+
+#if WITH_EDITOR
+		FWorldDelegates::OnStartGameInstance.RemoveAll(this);
+#endif // WITH_EDITOR
 	}
 
 	// At this point World Partition should be uninitialized
@@ -558,6 +567,13 @@ uint32 UWorldPartitionSubsystem::GetOverriddenLoadingRangesEpoch()
 }
 #endif
 
+#if WITH_EDITOR
+void UWorldPartitionSubsystem::OnStartGameInstance(UGameInstance* GameInstance)
+{
+	InitDynamicLoadingRangeScale();
+}
+#endif // WITH_EDITOR
+
 void UWorldPartitionSubsystem::OnWorldPartitionInitialized(UWorldPartition* InWorldPartition)
 {
 	if (RegisteredWorldPartitions.IsEmpty())
@@ -1223,7 +1239,10 @@ void UWorldPartitionSubsystem::UpdateStreamingStateInternal(const UWorld* InWorl
 
 	ON_SCOPE_EXIT
 	{
-		WorldPartitionSubsystem->OnStreamingStateUpdated().Broadcast();
+		if (WorldPartitionSubsystem)
+		{
+			WorldPartitionSubsystem->OnStreamingStateUpdated().Broadcast();
+		}
 	};
 
 	TArray<TObjectPtr<UWorldPartition>> RegisteredWorldPartitionsCopy;
@@ -1232,7 +1251,7 @@ void UWorldPartitionSubsystem::UpdateStreamingStateInternal(const UWorld* InWorl
 		if (RegisteredWorldPartitionsCopy.IsEmpty())
 		{
 			// Make temp copy of array as UpdateStreamingState may FlushAsyncLoading, which may add a new world partition to RegisteredWorldPartitions while iterating
-			RegisteredWorldPartitionsCopy = InWorldPartition ? TArray<TObjectPtr<UWorldPartition>>({ InWorldPartition }) : WorldPartitionSubsystem->RegisteredWorldPartitions;
+			RegisteredWorldPartitionsCopy = InWorldPartition ? TArray<TObjectPtr<UWorldPartition>>({ InWorldPartition }) : (WorldPartitionSubsystem) ? WorldPartitionSubsystem->RegisteredWorldPartitions : TArray<TObjectPtr<UWorldPartition>>{};
 		}
 		return RegisteredWorldPartitionsCopy;
 	};
@@ -1247,7 +1266,7 @@ void UWorldPartitionSubsystem::UpdateStreamingStateInternal(const UWorld* InWorl
 
 	const bool bIsServer = IsServer(InWorld);
 	const bool bServerStreamingEnabled = bIsServer && WorldPartitionSubsystem && WorldPartitionSubsystem->HasAnyWorldPartitionServerStreamingEnabled();
-	const int32 WorldPartitionUpdateCount = InWorldPartition ? 1 : WorldPartitionSubsystem->RegisteredWorldPartitions.Num();
+	const int32 WorldPartitionUpdateCount = InWorldPartition ? 1 : (WorldPartitionSubsystem) ? WorldPartitionSubsystem->RegisteredWorldPartitions.Num() : 0;
 
 	const bool bIsInBlockTillLevelStreamingCompleted = InWorld->GetIsInBlockTillLevelStreamingCompleted();
 	const bool bForceDisableIncrementalUpdate = IsHighPriorityLoading(InWorld) || !InWorld->bMatchStarted || InWorld->IsInSeamlessTravel() || bIsInBlockTillLevelStreamingCompleted;
@@ -1256,11 +1275,18 @@ void UWorldPartitionSubsystem::UpdateStreamingStateInternal(const UWorld* InWorl
 									!bForceDisableIncrementalUpdate &&
 									(!bIsServer || bServerStreamingEnabled); // No increment on server except if server streaming is enabled
 
+#if WITH_EDITOR
+	if (WorldPartitionSubsystem)
+	{
+		WorldPartitionSubsystem->UpdateDynamicLoadingRangeScale();
+	}
+#endif // WITH_EDITOR
+
 	// Update streaming state of all registered world partitions
 	if (bIncrementalUpdate)
 	{
 		// Early out if incremental update doesn't complete
-		if (!WorldPartitionSubsystem->IncrementalUpdateStreamingState())
+		if (WorldPartitionSubsystem && !WorldPartitionSubsystem->IncrementalUpdateStreamingState())
 		{
 			return;
 		}
@@ -1698,6 +1724,71 @@ void UWorldPartitionSubsystem::DrawStreamingStatusLegend(class UCanvas* Canvas,
 	Offset.X += MaxTextWidth + 10;
 }
 
+#if WITH_EDITOR
+
+namespace UE
+{
+	namespace Private
+	{
+		void AssignLinearScale(FDynamicLoadingRangeScales& Scales, float Linear)
+		{
+			Scales.LinearScale = Linear;
+			Scales.Scale2D = FMath::Sqrt(Scales.LinearScale);
+			Scales.Scale3D = FMath::Pow(Scales.LinearScale, 1.f / 3.f);
+		}
+	}
+}
+
+void UWorldPartitionSubsystem::InitDynamicLoadingRangeScale()
+{
+	if (GetWorld()->GetWorldSettings()->GetDynamicLoadingRangeScalingEnable())
+	{
+		FDynamicLoadingRangeScaling& ScaleSettings = GetWorld()->GetWorldSettings()->DynamicLoadingRangeScaling;
+		UE::Private::AssignLinearScale(LoadingRangeScales, FMath::Clamp(ScaleSettings.InitialDynamicLoadingRangeScaling, 0.f, 100.f) * 0.01f);
+		DynamicLoadingRangeScaleStartTime = GetWorld()->GetTimeSeconds();
+	}
+	else
+	{
+		UE::Private::AssignLinearScale(LoadingRangeScales, 1.f);
+	}
+}
+
+void UWorldPartitionSubsystem::UpdateDynamicLoadingRangeScale()
+{
+	if (LoadingRangeScales.LinearScale < 1.f) // Can skip the update if the scale is at 1.f already
+	{
+		if (GetWorld()->GetWorldSettings()->GetDynamicLoadingRangeScalingEnable())
+		{
+			FDynamicLoadingRangeScaling& ScaleSettings = GetWorld()->GetWorldSettings()->DynamicLoadingRangeScaling;
+			// ScaleSettings.InitialDynamicLoadingRangeScaling is a percentage value
+			float InitialScale = FMath::Clamp(ScaleSettings.InitialDynamicLoadingRangeScaling * 0.01f, 0.f, 1.f);
+			if (GetWorld()->GetIsInBlockTillLevelStreamingCompleted() && (LoadingRangeScales.LinearScale == InitialScale))
+			{
+				// Do not interpolate as long as we are still loading the levels, and use the initial loading range
+				DynamicLoadingRangeScaleStartTime = GetWorld()->GetTimeSeconds(); // StartTime is the time when initial level loading is complete, so update the time
+			}
+			else
+			{
+				// Start interpolating
+				float InterpolationTime = GetWorld()->GetTimeSeconds() - DynamicLoadingRangeScaleStartTime;
+				float Delay = FMath::Max(ScaleSettings.DynamicLoadingRangeScalingDelay, 0.0f);
+				float Duration = FMath::Max(ScaleSettings.DynamicLoadingRangeScalingDuration, 0.0f);
+
+				float EffectiveTime = InterpolationTime - Delay;
+				float Progress = (EffectiveTime < 0.f) ? 0.f : (Duration == 0.f) ? 1.f : FMath::Min(EffectiveTime / Duration, 1.f);
+				UE::Private::AssignLinearScale(LoadingRangeScales, Progress * (1.f - InitialScale) + InitialScale);
+			}
+		}
+		else
+		{
+			// Reset scale to 1.f if user disables the feature during runtime
+			UE::Private::AssignLinearScale(LoadingRangeScales, 1.f);
+		}
+	}
+}
+
+#endif // WITH_EDITOR
+
 /*
  * FStreamingSourceVelocity Implementation
  */
diff --git a/Engine/Source/Runtime/Engine/Private/WorldSettings.cpp b/Engine/Source/Runtime/Engine/Private/WorldSettings.cpp
--- a/Engine/Source/Runtime/Engine/Private/WorldSettings.cpp
+++ b/Engine/Source/Runtime/Engine/Private/WorldSettings.cpp
@@ -38,7 +38,8 @@
 #include "HierarchicalLOD.h"
 #include "HLOD/HLODSetup.h"
 #include "WorldPartition/DataLayer/DataLayerManager.h"
-#endif 
+#include "WorldPartition/WorldPartitionEditorPerProjectUserSettings.h"
+#endif
 
 #include UE_INLINE_GENERATED_CPP_BY_NAME(WorldSettings)
 
@@ -77,6 +78,7 @@ namespace WorldSettingsUtils
 	}
 }
 #endif
+
 AWorldSettings::FOnNaniteSettingsChanged AWorldSettings::OnNaniteSettingsChanged;
 
 AWorldSettings::AWorldSettings(const FObjectInitializer& ObjectInitializer)
@@ -398,6 +400,17 @@ void AWorldSettings::GetLifetimeReplicatedProps( TArray< FLifetimeProperty > & O
 	DOREPLIFETIME_WITH_PARAMS_FAST( AWorldSettings, NaniteSettings, SharedParams );
 }
 
+#if WITH_EDITOR
+bool AWorldSettings::GetDynamicLoadingRangeScalingEnable() const
+{
+	// The per-project, per-user settings are stored for the Editor world, not the PIE world.
+	UWorld* EditorWorld = (GEditor) ? GEditor->GetEditorWorldContext().World() : nullptr;
+
+	// DynamicLoadingRangeScaling is enabled in WorldSettings, but can also be disabled via a per-user, per-map override
+	return DynamicLoadingRangeScaling.bEnabled && EditorWorld && !GetDefault<UWorldPartitionEditorPerProjectUserSettings>()->GetForceDisableDynamicLoadingRangeScaling(EditorWorld);
+}
+#endif // WITH_EDITOR
+
 const FGuid FWorldSettingCustomVersion::GUID(0x1ED048F4, 0x2F2E4C68, 0x89D053A4, 0xF18F102D);
 // Register the custom version with core
 FCustomVersionRegistration GRegisterWorldSettingCustomVersion(FWorldSettingCustomVersion::GUID, FWorldSettingCustomVersion::LatestVersion, TEXT("WorldSettingVer"));
diff --git a/Engine/Source/Runtime/Engine/Public/WorldPartition/WorldPartitionEditorPerProjectUserSettings.h b/Engine/Source/Runtime/Engine/Public/WorldPartition/WorldPartitionEditorPerProjectUserSettings.h
--- a/Engine/Source/Runtime/Engine/Public/WorldPartition/WorldPartitionEditorPerProjectUserSettings.h
+++ b/Engine/Source/Runtime/Engine/Public/WorldPartition/WorldPartitionEditorPerProjectUserSettings.h
@@ -26,6 +26,7 @@ struct FWorldPartitionPerWorldSettings
 		LoadedEditorLocationVolumes.Empty();
 		NotLoadedDataLayers.Empty();
 		LoadedDataLayers.Empty();
+		bForceDisableDynamicLoadingRangeScaling = false;
 	}
 #endif
 
@@ -41,6 +42,10 @@ struct FWorldPartitionPerWorldSettings
 
 	UPROPERTY()
 	TArray<TSoftObjectPtr<const UDataLayerAsset>> LoadedDataLayers;
+
+	/** Lets users disable the map's DynamicLoadingRangeScaling enable setting locally without having to checkout the map or risk submitting the setting */
+	UPROPERTY()
+	bool bForceDisableDynamicLoadingRangeScaling = false;
 #endif
 };
 
@@ -140,13 +145,16 @@ public:
 	ENGINE_API TArray<UDataLayerAsset*> GetWorldDataLayerAssetsNotLoadedInEditor(UWorld* InWorld);
 
 	ENGINE_API void SetWorldDataLayersNonDefaultEditorLoadStates(UWorld* InWorld, const TArray<const UDataLayerInstance*>& InDataLayersLoadedInEditor, const TArray<const UDataLayerInstance*>& InDataLayersNotLoadedInEditor);
-		
+
 private:
 	const FWorldPartitionPerWorldSettings* GetWorldPartitionPerWorldSettings(UWorld* InWorld) const;
 #endif
 
 #if WITH_EDITORONLY_DATA
 public:
+	ENGINE_API bool GetForceDisableDynamicLoadingRangeScaling(UWorld* InWorld) const;
+	ENGINE_API void SetForceDisableDynamicLoadingRangeScaling(UWorld* InWorld, bool bOverride);
+
 	/** True when the Data Layer Outliner is displaying Editor Data Layers */
 	UPROPERTY(config, EditAnywhere, Category = "Data Layer")
 	uint32 bHideEditorDataLayers : 1;
diff --git a/Engine/Source/Runtime/Engine/Public/WorldPartition/WorldPartitionRuntimeSpatialHash.h b/Engine/Source/Runtime/Engine/Public/WorldPartition/WorldPartitionRuntimeSpatialHash.h
--- a/Engine/Source/Runtime/Engine/Public/WorldPartition/WorldPartitionRuntimeSpatialHash.h
+++ b/Engine/Source/Runtime/Engine/Public/WorldPartition/WorldPartitionRuntimeSpatialHash.h
@@ -184,7 +184,7 @@ struct FSpatialHashStreamingGrid
 	// Used by PIE/Game
 	ENGINE_API int64 GetCellSize(int32 Level) const;
 	ENGINE_API void GetCells(const FWorldPartitionStreamingQuerySource& QuerySource, TSet<const UWorldPartitionRuntimeCell*>& OutCells, bool bEnableZCulling, FWorldPartitionQueryCache* QueryCache = nullptr) const;
-	ENGINE_API void GetCells(const TArray<FWorldPartitionStreamingSource>& Sources, UWorldPartitionRuntimeHash::FStreamingSourceCells& OutActivateCells, UWorldPartitionRuntimeHash::FStreamingSourceCells& OutLoadCells, bool bEnableZCulling, const FWorldPartitionStreamingContext& Context) const;
+	ENGINE_API void GetCells(const TArray<FWorldPartitionStreamingSource>& Sources, UWorldPartitionRuntimeHash::FStreamingSourceCells& OutActivateCells, UWorldPartitionRuntimeHash::FStreamingSourceCells& OutLoadCells, bool bEnableZCulling, const FWorldPartitionStreamingContext& Context, float LoadingRangeScale = 1.f) const;
 	ENGINE_API void GetNonSpatiallyLoadedCells(TSet<const UWorldPartitionRuntimeCell*>& OutActivateCells, TSet<const UWorldPartitionRuntimeCell*>& OutLoadCells, const FWorldPartitionStreamingContext& Context) const;
 	ENGINE_API void Draw2D(const class UWorldPartitionRuntimeSpatialHash* Owner, const FBox2D& Region2D, const FBox2D& GridScreenBounds, TFunctionRef<FVector2D(const FVector2D&, bool)> WorldToScreen, FWorldPartitionDraw2DContext& DrawContext) const;
 	ENGINE_API void Draw3D(const class UWorldPartitionRuntimeSpatialHash* Owner, const TArray<FWorldPartitionStreamingSource>& Sources, const FTransform& Transform) const;
@@ -206,7 +206,7 @@ struct FSpatialHashStreamingGrid
 private:
 	ENGINE_API void ForEachRuntimeCell(const FGridCellCoord& Coords, TFunctionRef<void(const UWorldPartitionRuntimeCell*)> Func) const;
 	ENGINE_API void ForEachLayerCell(const FGridCellCoord& Coords, TFunctionRef<void(const FSpatialHashStreamingGridLayerCell*)> Func) const;
-	ENGINE_API void DrawStreamingSource2D(const FBox2D& GridScreenBounds, const FSphericalSector& Shape, TFunctionRef<FVector2D(const FVector2D&)> WorldToScreen, const FColor& Color, FWorldPartitionDraw2DContext& DrawContext) const;
+	ENGINE_API void DrawStreamingSource2D(const FBox2D& GridScreenBounds, const FSphericalSector& Shape, TFunctionRef<FVector2D(const FVector2D&)> WorldToScreen, const FColor& Color, FWorldPartitionDraw2DContext& DrawContext, float LoadingRangeScale = 1.f) const;
 	ENGINE_API void DrawStreamingSource3D(UWorld* World, const FSphericalSector& Shape, const FTransform& Transform, const FColor& Color) const;
 	ENGINE_API void GetFilteredCellsForDebugDraw(const FSpatialHashStreamingGridLayerCell* LayerCell, TArray<const UWorldPartitionRuntimeCell*>& FilteredCells) const;
 	ENGINE_API EWorldPartitionRuntimeCellVisualizeMode GetStreamingCellVisualizeMode() const;
diff --git a/Engine/Source/Runtime/Engine/Public/WorldPartition/WorldPartitionSubsystem.h b/Engine/Source/Runtime/Engine/Public/WorldPartition/WorldPartitionSubsystem.h
--- a/Engine/Source/Runtime/Engine/Public/WorldPartition/WorldPartitionSubsystem.h
+++ b/Engine/Source/Runtime/Engine/Public/WorldPartition/WorldPartitionSubsystem.h
@@ -13,6 +13,7 @@
 #include "Misc/Guid.h"
 #include "WorldPartitionSubsystem.generated.h"
 
+class UGameInstance;
 class ULevel;
 class ULevelStreaming;
 class UWorldPartition;
@@ -43,6 +44,17 @@ private:
 	FVector AvgVelocity;
 };
 
+// Container for the different domains of loading range scales
+struct FDynamicLoadingRangeScales
+{
+	// Increases linearly with time
+	float LinearScale = 1.0f;
+	// Square root of LinearScale
+	float Scale2D = 1.0f;
+	// Cubic root of LinearScale
+	float Scale3D = 1.0f;
+};
+
 /**
  * UWorldPartitionSubsystem
  */
@@ -120,6 +132,9 @@ public:
 
 	static ENGINE_API bool IsRunningConvertWorldPartitionCommandlet();
 
+	// When enabled, returns a linear loading range scaling ramp to boost initial PIE availability, at the cost of popping. When disabled, the values are 1.f
+	const FDynamicLoadingRangeScales& GetDynamicLoadingRangeScales() const { return LoadingRangeScales; }
+
 	UE_DEPRECATED(5.4, "Use UActorDescContainerSubsystem instead")
 	UActorDescContainer* RegisterContainer(FName PackageName) { return nullptr; }
 	UE_DEPRECATED(5.4, "Use UActorDescContainerSubsystem instead")
@@ -130,6 +145,7 @@ public:
 	void UpdateContainerBounds(FName PackageName) { }
 
 	TSet<FWorldPartitionHandle> SelectedActorHandles;
+
 private:
 	ENGINE_API FWorldPartitionActorFilter GetWorldPartitionActorFilterInternal(const FString& InWorldPackage, EWorldPartitionActorFilterType InFilterTypes, TSet<FString>& InOutVisitedPackageStack) const;
 #endif
@@ -152,6 +168,9 @@ private:
 	ENGINE_API void UpdateStreamingSourcesVelocity(const UWorldPartition* InWorldPartition, TArray<FWorldPartitionStreamingSource>& InOutStreamingSources) const;
 	uint32 GetStreamingSourcesHash() const { return StreamingSourcesHash; }
 
+#if WITH_EDITOR
+	void OnStartGameInstance(UGameInstance* GameInstance);
+#endif // WITH_EDITOR
 	void OnWorldPartitionInitialized(UWorldPartition* InWorldPartition);
 	void OnWorldPartitionUninitialized(UWorldPartition* InWorldPartition);
 	void OnLevelStreamingStateChanged(UWorld* InWorld, const ULevelStreaming* InStreamingLevel, ULevel* LevelIfLoaded, ELevelStreamingState PreviousState, ELevelStreamingState NewState);
@@ -179,6 +198,15 @@ private:
 	friend class UWorldPartition;
 	friend class UWorldPartitionStreamingPolicy;
 
+#if WITH_EDITOR
+	// If enabled, compute a linear loading range scaling ramp to boost initial PIE availability, at the cost of popping.
+	void InitDynamicLoadingRangeScale();
+	void UpdateDynamicLoadingRangeScale();
+
+	float DynamicLoadingRangeScaleStartTime = -1.0f;
+	FDynamicLoadingRangeScales LoadingRangeScales{1.f};
+#endif // WITH_EDITOR
+
 	// Registered world partitions & incremental update
 	TArray<TObjectPtr<UWorldPartition>> RegisteredWorldPartitions;
 	TSet<TObjectPtr<UWorldPartition>> IncrementalUpdateWorldPartitions;
