Address review comments from Release CL.
Replace WITH_EDITOR with WITH_EDITORONLY_DATA to match the data being accessed.
Use common getter for the Default Textures, to avoid duplicating logic.

[RN] minor | Rendering Architecture

#Synced-CL 50483943

diff --git a/Engine/Source/Runtime/Engine/Private/Materials/Material.cpp b/Engine/Source/Runtime/Engine/Private/Materials/Material.cpp
--- a/Engine/Source/Runtime/Engine/Private/Materials/Material.cpp
+++ b/Engine/Source/Runtime/Engine/Private/Materials/Material.cpp
@@ -4762,24 +4762,20 @@ void UMaterial::DumpDebugInfo(FOutputDevice& OutputDevice) const
 
 TObjectPtr<UObject> UMaterial::GetDefaultTextureAtIdx(const int32 InIdx, TMicRecursionGuard RecursionGuard)
 {
-#if WITH_EDITORONLY_DATA
-	const TArray<TObjectPtr<UObject>>& UsedTextures = GetDefaultTextures();
-#else
-	const TArray<TObjectPtr<UObject>>& UsedTextures = ReferencedDefaultTextures;
-#endif
+	const TArrayView<const TObjectPtr<UObject>> UsedDefaultTextures = GetReferencedTextures();
 
-	if (!UsedTextures.IsValidIndex(InIdx))
+	if (!UsedDefaultTextures.IsValidIndex(InIdx))
 	{
 		static bool bWarnedOnce = false;
 		if (!bWarnedOnce)
 		{
-			UE_LOG(LogMaterial, Warning, TEXT("%s: Requesting an invalid TextureIndex! (%u / %u)"), *GetName(), InIdx, UsedTextures.Num());
+			UE_LOG(LogMaterial, Warning, TEXT("%s: Requesting an invalid TextureIndex! (%u / %u)"), *GetName(), InIdx, UsedDefaultTextures.Num());
 			bWarnedOnce = true;
 		}
 		return nullptr;
 	}
 
-	return UsedTextures[InIdx];
+	return UsedDefaultTextures[InIdx];
 }
 
 void UMaterial::SaveShaderStableKeys(const class ITargetPlatform* TP)
diff --git a/Engine/Source/Runtime/Engine/Private/Materials/MaterialCachedData.cpp b/Engine/Source/Runtime/Engine/Private/Materials/MaterialCachedData.cpp
--- a/Engine/Source/Runtime/Engine/Private/Materials/MaterialCachedData.cpp
+++ b/Engine/Source/Runtime/Engine/Private/Materials/MaterialCachedData.cpp
@@ -1135,7 +1135,7 @@ void FMaterialCachedExpressionData::GetAllGlobalParameterInfoOfType(EMaterialPar
 	}
 }
 
-#if WITH_EDITOR
+#if WITH_EDITORONLY_DATA
 int32 FMaterialCachedExpressionData::GetDefaultTextureIdx(UObject* InTexture) const
 {
 	return EditorOnlyData->EditorReferencedDefaultTextures.Find(TObjectPtr<UObject>(InTexture));
diff --git a/Engine/Source/Runtime/Engine/Private/Materials/MaterialInterface.cpp b/Engine/Source/Runtime/Engine/Private/Materials/MaterialInterface.cpp
--- a/Engine/Source/Runtime/Engine/Private/Materials/MaterialInterface.cpp
+++ b/Engine/Source/Runtime/Engine/Private/Materials/MaterialInterface.cpp
@@ -312,7 +312,7 @@ const FMaterialCachedExpressionData& UMaterialInterface::GetCachedExpressionData
 	return LocalData ? *LocalData : FMaterialCachedExpressionData::EmptyData;
 }
 
-#if WITH_EDITOR
+#if WITH_EDITORONLY_DATA
 int32 UMaterialInterface::GetDefaultTextureIdx(UObject* InTexture) const
 {
 	const FMaterialCachedExpressionData& UsedCachedData = GetCachedExpressionData();
@@ -329,7 +329,7 @@ const TArray<TObjectPtr<UObject>>& UMaterialInterface::GetDefaultTextures() cons
 	const FMaterialCachedExpressionData& UsedCachedData = GetCachedExpressionData();
 	return UsedCachedData.EditorOnlyData->EditorReferencedDefaultTextures;
 }
-#endif // WITH_EDITOR
+#endif // WITH_EDITORONLY_DATA
 
 bool UMaterialInterface::IsUsingNewHLSLGenerator() const
 {
diff --git a/Engine/Source/Runtime/Engine/Public/MaterialCachedData.h b/Engine/Source/Runtime/Engine/Public/MaterialCachedData.h
--- a/Engine/Source/Runtime/Engine/Public/MaterialCachedData.h
+++ b/Engine/Source/Runtime/Engine/Public/MaterialCachedData.h
@@ -272,7 +272,7 @@ struct FMaterialCachedExpressionData
 	void GetAllGlobalParametersOfType(EMaterialParameterType Type, TMap<FMaterialParameterInfo, FMaterialParameterMetadata>& OutParameters) const;
 	void GetAllGlobalParameterInfoOfType(EMaterialParameterType Type, TArray<FMaterialParameterInfo>& OutParameterInfo, TArray<FGuid>& OutParameterIds) const;
 
-#if WITH_EDITOR
+#if WITH_EDITORONLY_DATA
 	/** @return The index at which the given texture sits in the default textures array or INDEX_NONE if not present. */
 	int32 GetDefaultTextureIdx(UObject* InTexture) const;
 #endif
diff --git a/Engine/Source/Runtime/Engine/Public/Materials/MaterialInterface.h b/Engine/Source/Runtime/Engine/Public/Materials/MaterialInterface.h
--- a/Engine/Source/Runtime/Engine/Public/Materials/MaterialInterface.h
+++ b/Engine/Source/Runtime/Engine/Public/Materials/MaterialInterface.h
@@ -545,14 +545,14 @@ public:
 
 	ENGINE_API virtual const FMaterialCachedExpressionData& GetCachedExpressionData(TMicRecursionGuard RecursionGuard = TMicRecursionGuard()) const;
 
-#if WITH_EDITOR
+#if WITH_EDITORONLY_DATA
 	/** @return The Index at which the given texture is present in the editor only Default Textures array. */
 	ENGINE_API int32 GetDefaultTextureIdx(UObject* InTexture) const;
 
 	/** @return Whether the given texture is present in the editor only Default Textures. */
 	ENGINE_API bool ContainsDefaultTexture(UObject* InTexture) const;
 
-	/** @return The array of default textures present in editor. */
+	/** @return The array of default textures. */
 	ENGINE_API const TArray<TObjectPtr<UObject>>& GetDefaultTextures() const;
 #endif
 
