This change fixes a bug where focus was lost in the Details panel after editing a property of a Blueprint-added component, due to reconstruction of the SCS (Simple Construction Script).

Problem:
When a property is edited, the SCS component is destroyed and rebuilt. This destroys all associated widgets, including the one that previously had keyboard focus. As a result, pressing Enter would not refocus the same field or any appropriate fallback widget, breaking expected keyboard flow.

Fix Details:
- The focus path is now saved as an array of widget type + index pairs before reconstruction.
- After the widget tree rebuilds, the new widget tree is traversed using this index path to identify a new matching widget. If it fails to find the exact path, it will abort.
- This allows the system to restore focus to the last focusable element if the index path is exactly the same.

Validation:
- Verified consistent focus behavior across property edits in the Details panel for Blueprint-added components (SpinBoxes, CheckBoxes, TextBoxes, ComboBoxes).

#Synced-CL 50447978

diff --git a/Engine/Source/Runtime/Engine/Private/ActorConstruction.cpp b/Engine/Source/Runtime/Engine/Private/ActorConstruction.cpp
--- a/Engine/Source/Runtime/Engine/Private/ActorConstruction.cpp
+++ b/Engine/Source/Runtime/Engine/Private/ActorConstruction.cpp
@@ -20,6 +20,7 @@
 
 #if WITH_EDITOR
 #include "Editor.h"
+#include "EditorWidgetFocusUtils.h"
 #else
 #include "Engine/World.h"
 #include "UObject/Package.h"
@@ -248,6 +249,7 @@ bool AActor::HasNonTrivialUserConstructionScript() const
 	return false;
 }
 
+
 #if WITH_EDITOR
 void AActor::RerunConstructionScripts()
 {
@@ -575,6 +577,9 @@ void AActor::RerunConstructionScripts()
 			}
 		}
 
+		using namespace WidgetFocusRestorationUtils;
+		TOptional<FFocusedWidgetData> FocusedWidgetData = GetFocusedWidgetData();
+		
 		// Destroy existing components
 		DestroyConstructedComponents();
 
@@ -796,6 +801,8 @@ void AActor::RerunConstructionScripts()
 			CurrentTransactionAnnotation = nullptr;
 		}
 
+		RestoreFocusedWidget(FocusedWidgetData);
+
 		// Remove the name to SCS node maps now that we're done constructing
 		for (const UBlueprintGeneratedClass* BPClass : ParentBPClassStack)
 		{
diff --git a/Engine/Source/Runtime/Engine/Private/EditorWidgetFocusUtils.cpp b/Engine/Source/Runtime/Engine/Private/EditorWidgetFocusUtils.cpp
new file mode 100644
--- /dev/null
+++ b/Engine/Source/Runtime/Engine/Private/EditorWidgetFocusUtils.cpp
@@ -0,0 +1,152 @@
+// Copyright Epic Games, Inc. All Rights Reserved.
+
+#include "EditorWidgetFocusUtils.h"
+
+#include "Framework/Application/SlateApplication.h"
+#include "Framework/Application/SlateUser.h"
+
+namespace Private
+{
+	int32 FindChild(FChildren* Children, const TSharedRef<SWidget>& ChildToFind)
+	{
+		if (Children)
+		{
+			for (int32 ChildIdx = 0; ChildIdx < Children->Num(); ++ChildIdx)
+			{
+				if (TSharedRef<SWidget> Child = Children->GetChildAt(ChildIdx); ChildToFind == Child)
+				{
+					return ChildIdx;
+				}
+			}
+		}
+		
+		return INDEX_NONE;
+	}
+}
+using namespace WidgetFocusRestorationUtils;
+
+TOptional<FFocusedWidgetData> WidgetFocusRestorationUtils::GetFocusedWidgetData()
+{
+	if (!FSlateApplication::IsInitialized())
+	{
+		return TOptional<FFocusedWidgetData>{};
+	}
+
+	TSharedPtr<SWidget> FocusedWidget = FSlateApplication::Get().GetKeyboardFocusedWidget();
+	// Invalid focused widget, return empty
+	if (!FocusedWidget.IsValid())
+	{
+		return TOptional<FFocusedWidgetData>{};
+	}
+
+	FWidgetPath FocusPath;
+	const bool bFoundPath = FSlateApplication::Get().FindPathToWidget(FocusedWidget.ToSharedRef(), FocusPath, EVisibility::All);
+
+	// Invalid Path, return empty
+	if (!bFoundPath)
+	{
+		return TOptional<FFocusedWidgetData>{};
+	}
+
+	// Build a path of indices to find focus location after SCS rerun
+	TArray<FIndexTypePair> IndexPathToWidget;
+
+	// Number of widgets, minus the window widget
+	const int32 NumberOfChildren = FocusPath.Widgets.Num() - 1;
+	IndexPathToWidget.Reserve(NumberOfChildren);
+
+	TSharedPtr<SWindow> RootWidget = FocusPath.TopLevelWindow;
+	FChildren* WidgetChildren = RootWidget->GetChildren();
+
+	// Path of indices starts after the TopLevelWindow widget
+	constexpr int IndexPathStart = 1;
+	for (int32 WidgetIdx = IndexPathStart; WidgetIdx < FocusPath.Widgets.Num(); ++WidgetIdx)
+	{
+		TSharedRef<SWidget> Widget = FocusPath.Widgets[WidgetIdx].Widget;
+
+		if (const int32 ChildIdx = Private::FindChild(WidgetChildren, Widget); ChildIdx != INDEX_NONE)
+		{
+			IndexPathToWidget.Add(FIndexTypePair{ .Index = ChildIdx, .Type = Widget->GetType() });
+			WidgetChildren = Widget->GetChildren();
+		}
+		else
+		{
+			// Could not build an index path to the focused widget, return empty
+			return TOptional<FFocusedWidgetData>{};
+		}
+	}
+
+	check(NumberOfChildren == IndexPathToWidget.Num());
+
+	return FFocusedWidgetData{
+		.WidgetPath = IndexPathToWidget,
+		.WeakRootWidget = RootWidget
+	};
+}
+
+void WidgetFocusRestorationUtils::RestoreFocusedWidget(const TOptional<FFocusedWidgetData>& FocusedWidgetData)
+{
+	if (!FocusedWidgetData.IsSet())
+	{
+		return;
+	}
+
+	if (!FSlateApplication::IsInitialized())
+	{
+		return;
+	}
+
+	// Try restore focus to the replacement widget after widget reconstruction (next Tick)
+	FTSTicker::GetCoreTicker().AddTicker(
+		FTickerDelegate::CreateLambda([Data = FocusedWidgetData.GetValue()](float DeltaTime)
+			{
+				// Invalid Root Widget, do not attempt to restore focus
+				if (!Data.WeakRootWidget.IsValid())
+				{
+					return false;
+				}
+
+				TSharedPtr<SWidget> CurrentWidget = Data.WeakRootWidget.Pin();
+				TSharedPtr<SWidget> LastFocusableWidget = nullptr;
+
+				// Go down the widget path to find the reconstructed focused widget
+				for (FIndexTypePair ChildIndexTypePair : Data.WidgetPath)
+				{
+					FChildren* Children = CurrentWidget->GetChildren();
+
+					// Invalid Child Index, do not attempt to restore focus
+					if (!Children || ChildIndexTypePair.Index >= Children->Num())
+					{
+						return false;
+					}
+
+					const TSharedRef<SWidget> ChildWidget = Children->GetChildAt(ChildIndexTypePair.Index);
+
+					// Invalid Child Type, do not attempt to restore focus
+					if (ChildWidget->GetType() != ChildIndexTypePair.Type)
+					{
+						return false;
+					}
+
+					CurrentWidget = ChildWidget;
+
+					// Check if the current widget is visible and supports keyboard focus
+					if (CurrentWidget->SupportsKeyboardFocus() && CurrentWidget->GetVisibility().IsVisible())
+					{
+						LastFocusableWidget = CurrentWidget;
+					}
+				}
+
+				// Invalid Widget, do not attempt to restore focus
+				if (!LastFocusableWidget.IsValid())
+				{
+					return false;
+				}
+
+				FSlateApplication::Get().SetKeyboardFocus(LastFocusableWidget, EFocusCause::SetDirectly);
+
+				return false;
+			}
+		)
+	);
+}
diff --git a/Engine/Source/Runtime/Engine/Private/EditorWidgetFocusUtils.h b/Engine/Source/Runtime/Engine/Private/EditorWidgetFocusUtils.h
new file mode 100644
--- /dev/null
+++ b/Engine/Source/Runtime/Engine/Private/EditorWidgetFocusUtils.h
@@ -0,0 +1,22 @@
+// Copyright Epic Games, Inc. All Rights Reserved.
+#pragma once
+#include "Widgets/SWidget.h"
+
+namespace WidgetFocusRestorationUtils
+{
+	struct FIndexTypePair
+	{
+		int32 Index;
+		FName Type;
+	};
+
+	struct FFocusedWidgetData
+	{
+		TArray<FIndexTypePair> WidgetPath;
+		TWeakPtr<SWidget> WeakRootWidget;
+	};
+
+	TOptional<FFocusedWidgetData> GetFocusedWidgetData();
+
+	void RestoreFocusedWidget(const TOptional<FFocusedWidgetData>& FocusedWidgetData);
+}
