[SkeletalMesh] Move inlined SkeletalMeshLODRenderData to FByteBulkData.

#rnx

#Synced-CL 50437893

diff --git a/Engine/Source/Runtime/Core/Private/UObject/FortniteMainBranchObjectVersion_Restricted.inl b/Engine/Source/Runtime/Core/Private/UObject/FortniteMainBranchObjectVersion_Restricted.inl
--- a/Engine/Source/Runtime/Core/Private/UObject/FortniteMainBranchObjectVersion_Restricted.inl
+++ b/Engine/Source/Runtime/Core/Private/UObject/FortniteMainBranchObjectVersion_Restricted.inl
@@ -6,6 +6,6 @@ static void AppendRestrictedFortniteMainStreamObjectVersionGuids(TMap<FGuid, FGu
 	// patch effects.
 	const FDevSystemGuids& DevGuids = FDevSystemGuids::Get();
 	SystemGuids.Add(DevGuids.NANITE_DERIVEDDATA_VER, FGuid("F365D7A48F454525AAADC0678EA085A0"));
-	SystemGuids.Add(DevGuids.SkeletalMeshDerivedDataVersion, FGuid("C998F2F3B2884A02B7E290B8C25835F0"));
+	SystemGuids.Add(DevGuids.SkeletalMeshDerivedDataVersion, FGuid("A3F92C1D4E7B108AD61C8F4209B5E3F0"));
 	SystemGuids.Add(DevGuids.STATICMESH_DERIVEDDATA_VER, FGuid("3ABF32149F9D4D838750368AB95573BE"));
 }
diff --git a/Engine/Source/Runtime/Engine/Private/SkeletalMeshLODRenderData.cpp b/Engine/Source/Runtime/Engine/Private/SkeletalMeshLODRenderData.cpp
--- a/Engine/Source/Runtime/Engine/Private/SkeletalMeshLODRenderData.cpp
+++ b/Engine/Source/Runtime/Engine/Private/SkeletalMeshLODRenderData.cpp
@@ -1012,12 +1012,7 @@ void FSkeletalMeshLODRenderData::Serialize(FArchive& Ar, UObject* Owner, int32 I
 #endif
 		Ar << BuffersSize;
 
-		if (bInlined)
-		{
-			SerializeStreamedData(Ar, OwnerMesh, Idx, ClassDataStripFlags, bNeedsCPUAccess, bForceKeepCPUResources);
-			bIsLODOptional = false;
-		}
-		else if (Ar.IsFilterEditorOnly())
+		if (Ar.IsFilterEditorOnly())
 		{
 			bool bDiscardBulkData = false;
 
@@ -1037,22 +1032,28 @@ void FSkeletalMeshLODRenderData::Serialize(FArchive& Ar, UObject* Owner, int32 I
 					SerializeStreamedData(MemWriter, OwnerMesh, Idx, ClassDataStripFlags, bNeedsCPUAccess, bForceKeepCPUResources);
 				}
 
-				bIsLODOptional = bIsBelowMinLOD;
-				const uint32 BulkDataFlags = (bDiscardBulkData ? 0 : BULKDATA_Force_NOT_InlinePayload)
-					| (bIsLODOptional ? BULKDATA_OptionalPayload : 0);
-				const uint32 OldBulkDataFlags = BulkData.GetBulkDataFlags();
-				BulkData.ClearBulkDataFlags(0xffffffffu);
-				BulkData.SetBulkDataFlags(BulkDataFlags);
+				const uint32 OldBulkDataFlags = StreamingBulkData.GetBulkDataFlags();
+				StreamingBulkData.ClearBulkDataFlags(0xffffffffu);
 				if (TmpBuff.Num() > 0)
 				{
-					BulkData.Lock(LOCK_READ_WRITE);
-					void* BulkDataMem = BulkData.Realloc(TmpBuff.Num());
+					StreamingBulkData.Lock(LOCK_READ_WRITE);
+					void* BulkDataMem = StreamingBulkData.Realloc(TmpBuff.Num());
 					FMemory::Memcpy(BulkDataMem, TmpBuff.GetData(), TmpBuff.Num());
-					BulkData.Unlock();
+					StreamingBulkData.Unlock();
 				}
-				BulkData.Serialize(Ar, Owner, Idx);
-				BulkData.ClearBulkDataFlags(0xffffffffu);
-				BulkData.SetBulkDataFlags(OldBulkDataFlags);
+
+				bIsLODOptional = bIsBelowMinLOD;
+
+				const uint32 BulkDataInlineFlags = bInlined ? BULKDATA_ForceInlinePayload | BULKDATA_SingleUse : BULKDATA_Force_NOT_InlinePayload;
+				const uint32 BulkDataFlags = (bDiscardBulkData ? 0 : BulkDataInlineFlags)
+					| (bIsLODOptional ? BULKDATA_OptionalPayload : 0);
+
+				StreamingBulkData.SetBulkDataFlags(BulkDataFlags);
+
+				StreamingBulkData.Serialize(Ar, Owner, Idx);
+				
+				StreamingBulkData.ClearBulkDataFlags(0xffffffffu);
+				StreamingBulkData.SetBulkDataFlags(OldBulkDataFlags);
 			}
 			else
 #endif
@@ -1087,6 +1088,20 @@ void FSkeletalMeshLODRenderData::Serialize(FArchive& Ar, UObject* Owner, int32 I
 			{
 				SerializeAvailabilityInfo(Ar, Idx, StripFlags.IsClassDataStripped(CDSF_AdjacencyData_DEPRECATED), bNeedsCPUAccess);
 			}
+
+			if (bInlined && Ar.IsLoading() && StreamingBulkData.GetBulkDataSize() > 0)
+			{
+				void* Data = StreamingBulkData.Lock(LOCK_READ_ONLY);
+				FMemoryReaderView MemReader(MakeMemoryView(Data, StreamingBulkData.GetBulkDataSize()), true);
+				MemReader.SetByteSwapping(Ar.IsByteSwapping());
+				SerializeStreamedData(MemReader, OwnerMesh, Idx, ClassDataStripFlags, bNeedsCPUAccess, bForceKeepCPUResources);
+				StreamingBulkData.Unlock();
+			}
+		}
+		else if (bInlined)
+		{
+			SerializeStreamedData(Ar, OwnerMesh, Idx, ClassDataStripFlags, bNeedsCPUAccess, bForceKeepCPUResources);
+			bIsLODOptional = false;
 		}
 	}
 }
diff --git a/Engine/Source/Runtime/Engine/Public/Rendering/SkeletalMeshLODRenderData.h b/Engine/Source/Runtime/Engine/Public/Rendering/SkeletalMeshLODRenderData.h
--- a/Engine/Source/Runtime/Engine/Public/Rendering/SkeletalMeshLODRenderData.h
+++ b/Engine/Source/Runtime/Engine/Public/Rendering/SkeletalMeshLODRenderData.h
@@ -170,10 +170,6 @@ public:
 	/** Whether this LOD is below MinLod */
 	uint32 bIsLODOptional : 1;
 
-#if WITH_EDITOR
-	FByteBulkData BulkData;
-#endif
-
 #if RHI_RAYTRACING
 	/** Game thread reference counter of static skeletal mesh objects referencing this render data */
 	int32 NumReferencingStaticSkeletalMeshObjects = 0;
