[Functional Tests] Move Rescan in PostLoad to allow async loading with world composition

#Synced-CL 50456686

diff --git a/Engine/Source/Editor/UnrealEd/Private/UnrealEdMisc.cpp b/Engine/Source/Editor/UnrealEd/Private/UnrealEdMisc.cpp
--- a/Engine/Source/Editor/UnrealEd/Private/UnrealEdMisc.cpp
+++ b/Engine/Source/Editor/UnrealEd/Private/UnrealEdMisc.cpp
@@ -992,10 +992,11 @@ bool FUnrealEdMisc::EnableWorldComposition(UWorld* InWorld, bool bEnable)
 			return false;
 		}
 			
-		auto WorldCompostion = NewObject<UWorldComposition>(InWorld);
+		auto WorldComposition = NewObject<UWorldComposition>(InWorld);
 		// All map files found in the same and folder and all sub-folders will be added ass sub-levels to this map
 		// Make sure user understands this
-		int32 NumFoundSublevels = WorldCompostion->GetTilesList().Num();
+		WorldComposition->Rescan();
+		int32 NumFoundSublevels = WorldComposition->GetTilesList().Num();
 		if (NumFoundSublevels)
 		{
 			FFormatNamedArguments Arguments;
@@ -1006,13 +1007,13 @@ bool FUnrealEdMisc::EnableWorldComposition(UWorld* InWorld, bool bEnable)
 			auto AppResult = FMessageDialog::Open(EAppMsgType::OkCancel, Message);
 			if (AppResult != EAppReturnType::Ok)
 			{
-				WorldCompostion->MarkAsGarbage();
+				WorldComposition->MarkAsGarbage();
 				return false;
 			}
 		}
 			
 		// 
-		InWorld->WorldComposition = WorldCompostion;
+		InWorld->WorldComposition = WorldComposition;
 		UWorldComposition::WorldCompositionChangedEvent.Broadcast(InWorld);
 	}
 	
diff --git a/Engine/Source/Runtime/Engine/Classes/Engine/WorldComposition.h b/Engine/Source/Runtime/Engine/Classes/Engine/WorldComposition.h
--- a/Engine/Source/Runtime/Engine/Classes/Engine/WorldComposition.h
+++ b/Engine/Source/Runtime/Engine/Classes/Engine/WorldComposition.h
@@ -169,7 +169,6 @@ class UWorldComposition : public UObject
 
 private:
 	// UObject interface
-	ENGINE_API virtual void PostInitProperties() override;
 	ENGINE_API virtual void Serialize(FArchive& Ar) override;
 	ENGINE_API virtual void PostDuplicate(bool bDuplicateForPIE) override;
 	ENGINE_API virtual void PostLoad() override;
@@ -179,7 +178,7 @@ private:
 	void PopulateStreamingLevels();
 
 	/** Calculates tiles absolute positions based on relative positions */
-	void CaclulateTilesAbsolutePositions();
+	void CalculateTilesAbsolutePositions();
 
 	/** Resets world composition structures */
 	void Reset();
diff --git a/Engine/Source/Runtime/Engine/Private/WorldComposition.cpp b/Engine/Source/Runtime/Engine/Private/WorldComposition.cpp
--- a/Engine/Source/Runtime/Engine/Private/WorldComposition.cpp
+++ b/Engine/Source/Runtime/Engine/Private/WorldComposition.cpp
@@ -42,17 +42,6 @@ UWorldComposition::UWorldComposition(const FObjectInitializer& ObjectInitializer
 {
 }
 
-void UWorldComposition::PostInitProperties()
-{
-	Super::PostInitProperties();
-
-	if (!IsTemplate() && !GetOutermost()->HasAnyPackageFlags(PKG_PlayInEditor) && !FPlatformProperties::RequiresCookedData())
-	{
-		// Tiles information is not serialized to disk in the editor, and should be regenerated on world composition object construction
-		Rescan();
-	}
-}
-
 void UWorldComposition::Serialize( FArchive& Ar )
 {
 	Super::Serialize(Ar);
@@ -75,6 +64,7 @@ void UWorldComposition::PostDuplicate(bool bDuplicateForPIE)
 
 	if (bDuplicateForPIE)
 	{
+		// When duplicating for PIE, we assume the source was up-to-date and has fresh Tiles, so no need to rescan here. But Tiles need PIE package name fixup.
 		FixupForPIE(GetOutermost()->GetPIEInstanceID());	
 	}
 }
@@ -83,12 +73,18 @@ void UWorldComposition::PostLoad()
 {
 	Super::PostLoad();
 
+	if (!IsTemplate() && !GetOutermost()->HasAnyPackageFlags(PKG_PlayInEditor) && !FPlatformProperties::RequiresCookedData())
+	{
+		// Tiles information is not serialized to disk in the editor, and should be regenerated after world composition object load
+		Rescan();
+	}
+
 	if (FPlatformProperties::RequiresCookedData())
 	{
 		// Create streaming levels for each Tile
 		PopulateStreamingLevels();
 		// Calculate absolute positions since they are not serialized to disk
-		CaclulateTilesAbsolutePositions();
+		CalculateTilesAbsolutePositions();
 	}
 
 	UWorld* World = GetWorld();
@@ -325,7 +321,7 @@ void UWorldComposition::Rescan()
 	PopulateStreamingLevels();
 
 	// Calculate absolute positions since they are not serialized to disk
-	CaclulateTilesAbsolutePositions();
+	CalculateTilesAbsolutePositions();
 }
 
 void UWorldComposition::ReinitializeForPIE()
@@ -362,7 +358,7 @@ ULevelStreaming* UWorldComposition::CreateStreamingLevel(const FWorldComposition
 	return StreamingLevel;
 }
 
-void UWorldComposition::CaclulateTilesAbsolutePositions()
+void UWorldComposition::CalculateTilesAbsolutePositions()
 {
 	for (FWorldCompositionTile& Tile : Tiles)
 	{
