SpringMath: Added functions for estimating smoothing time from stopping distance and vice versa.

#Synced-CL 50483943

diff --git a/Engine/Source/Runtime/Engine/Classes/Kismet/BlueprintSpringMathLibrary.h b/Engine/Source/Runtime/Engine/Classes/Kismet/BlueprintSpringMathLibrary.h
--- a/Engine/Source/Runtime/Engine/Classes/Kismet/BlueprintSpringMathLibrary.h
+++ b/Engine/Source/Runtime/Engine/Classes/Kismet/BlueprintSpringMathLibrary.h
@@ -415,6 +415,26 @@ public:
 	UFUNCTION(BlueprintCallable, meta=(Category="Math|Spring|Experimental"))
 	static ENGINE_API void SpringCharacterUpdate(UPARAM(ref) FVector& InOutPosition, UPARAM(ref) FVector& InOutVelocity, UPARAM(ref) FVector& InOutAcceleration, const FVector& TargetVelocity, float DeltaTime, float SmoothingTime = 0.2f);
 
+	UE_EXPERIMENTAL(5.8, "SpringMath is experimental")
+	/** Estimates an approximate stopping distance from a smoothing time for a character using a simple damped spring
+	 *
+	 * @param Velocity The velocity of the character at the start of the stop
+	 * @param SmoothingTime The time over which to smooth velocity. It takes roughly the smoothing time in order for the character to reach the target velocity.
+	 * @return An estimation of the total distance traveled during the stop
+	 */
+	UFUNCTION(BlueprintPure, meta = (Category = "Math|Spring|Experimental", ReturnDisplayName = "Stopping Distance"))
+	static ENGINE_API float SpringCharacterSmoothingTimeToStoppingDistance(const float Velocity, const float SmoothingTime = 0.2f);
+
+	UE_EXPERIMENTAL(5.8, "SpringMath is experimental")
+	/** Estimates an approximate smoothing time from a stopping distance for a character using by a simple damped spring
+	 *
+	 * @param Velocity The velocity of the character at the start of the stop
+	 * @param StoppingDistance The total distance traveled during the stop
+	 * @return An estimation of the SmoothingTime of the spring
+	 */
+	UFUNCTION(BlueprintPure, meta = (Category = "Math|Spring|Experimental", ReturnDisplayName = "Smoothing Time"))
+	static ENGINE_API float SpringCharacterStoppingDistanceToSmoothingTime(const float Velocity, const float StoppingDistance = 100.0f);
+
 	UE_EXPERIMENTAL(5.7, "SpringMath is experimental")
 	/** Update a position representing a character given a target velocity using a velocity spring.
 	 * A velocity spring tracks an intermediate velocity which moves at a maximum acceleration linearly towards a target.
diff --git a/Engine/Source/Runtime/Engine/Private/BlueprintSpringMathLibrary.cpp b/Engine/Source/Runtime/Engine/Private/BlueprintSpringMathLibrary.cpp
--- a/Engine/Source/Runtime/Engine/Private/BlueprintSpringMathLibrary.cpp
+++ b/Engine/Source/Runtime/Engine/Private/BlueprintSpringMathLibrary.cpp
@@ -237,6 +237,16 @@ void UBlueprintSpringMathLibrary::SpringCharacterUpdate(FVector& InOutPosition,
 	SpringMath::SpringCharacterUpdate(InOutPosition, InOutVelocity, InOutAcceleration, TargetVelocity, SmoothingTime, DeltaTime);
 }
 
+float UBlueprintSpringMathLibrary::SpringCharacterSmoothingTimeToStoppingDistance(const float Velocity, const float SmoothingTime)
+{
+	return SpringMath::SpringCharacterSmoothingTimeToStoppingDistance(Velocity, SmoothingTime);
+}
+
+float UBlueprintSpringMathLibrary::SpringCharacterStoppingDistanceToSmoothingTime(const float Velocity, const float StoppingDistance)
+{
+	return SpringMath::SpringCharacterStoppingDistanceToSmoothingTime(Velocity, StoppingDistance);
+}
+
 void UBlueprintSpringMathLibrary::VelocitySpringCharacterUpdate(FVector& InOutPosition, FVector& InOutVelocity, FVector& InOutVelocityIntermediate,
 	FVector& InOutAcceleration, const FVector& TargetVelocity, float DeltaTime, float SmoothingTime, float MaxAcceleration)
 {
diff --git a/Engine/Source/Runtime/Engine/Public/Animation/SpringMath.h b/Engine/Source/Runtime/Engine/Public/Animation/SpringMath.h
--- a/Engine/Source/Runtime/Engine/Public/Animation/SpringMath.h
+++ b/Engine/Source/Runtime/Engine/Public/Animation/SpringMath.h
@@ -21,6 +21,11 @@ private:
 		return 4.0f / FMath::Max(SmoothingTime, UE_KINDA_SMALL_NUMBER);
 	}
 
+	static constexpr float DampingToSmoothingTime(float Damping)
+	{
+		return 4.0f / FMath::Max(Damping, UE_KINDA_SMALL_NUMBER);
+	}
+
 	/** Divides one scale by another, clamping the Rhs to be greater than some minimum */
 	static inline FVector ScaleDivMax(FVector Lhs, FVector Rhs)
 	{
@@ -614,6 +619,21 @@ public:
 		}
 	}
 
+	/** Estimates an approximate stopping distance from a smoothing time for a character using a simple damped spring */
+	UE_EXPERIMENTAL(5.8, "SpringMath is experimental")
+	static inline float SpringCharacterSmoothingTimeToStoppingDistance(const float Velocity, const float SmoothingTime)
+	{
+		const float Y = SmoothingTimeToDamping(SmoothingTime) / 2.0f;
+		return 2.0f * FMath::Abs(Velocity) / FMath::Max(Y, UE_SMALL_NUMBER);
+	}
+
+	/** Estimates an approximate spring smoothing time from a stopping distance and velocity for a character using by a simple damped spring */
+	UE_EXPERIMENTAL(5.8, "SpringMath is experimental")
+	static inline float SpringCharacterStoppingDistanceToSmoothingTime(const float Velocity, const float StoppingDistance)
+	{
+		return DampingToSmoothingTime(4.0f * FMath::Abs(Velocity) / FMath::Max(StoppingDistance, UE_SMALL_NUMBER));
+	}
+
 	/** Update a position representing a character given a target velocity using a velocity spring.
 	 * A velocity spring tracks an intermediate velocity which moves at a maximum acceleration linearly towards a target.
 	 * This means unlike the "SpringCharacterUpdate", it will take longer to reach a target velocity that is further away from the current velocity.
