Fixed auto-computed mass display for an actor to not include shapes that are flagged to not contribute to mass. Currently the mass is over-represented in the editor but correct at runtime as all shapes are used to calculate mass from density for this display regardless of their eligibility for mass calculations

#Synced-CL 50199821

diff --git a/Engine/Source/Runtime/Engine/Classes/PhysicsEngine/AggregateGeom.h b/Engine/Source/Runtime/Engine/Classes/PhysicsEngine/AggregateGeom.h
--- a/Engine/Source/Runtime/Engine/Classes/PhysicsEngine/AggregateGeom.h
+++ b/Engine/Source/Runtime/Engine/Classes/PhysicsEngine/AggregateGeom.h
@@ -156,6 +156,30 @@ struct FKAggregateGeom
 		});
 	}
 
+	template<typename CallableType>
+	void VisitSimElements(CallableType&& InCallable) const
+	{
+		for(const FKSphereElem& Elem : SphereElems)
+		{
+			InCallable(Elem);
+		}
+
+		for(const FKBoxElem& Elem : BoxElems)
+		{
+			InCallable(Elem);
+		}
+
+		for(const FKSphylElem& Elem : SphylElems)
+		{
+			InCallable(Elem);
+		}
+
+		for(const FKConvexElem& Elem : ConvexElems)
+		{
+			InCallable(Elem);
+		}
+	}
+
 	FKShapeElem* GetElement(const EAggCollisionShape::Type Type, const int32 Index)
 	{
 		switch (Type)
@@ -558,6 +582,7 @@ struct FKAggregateGeom
 
 	/** Returns the volume of this element */
 	ENGINE_API FVector::FReal GetScaledVolume(const FVector& Scale3D) const;
+	ENGINE_API FVector::FReal GetScaledVolumeForMassContribution(const FVector& Scale3D) const;
 
 	ENGINE_API FGuid MakeDDCKey() const;
 
diff --git a/Engine/Source/Runtime/Engine/Classes/PhysicsEngine/BodySetup.h b/Engine/Source/Runtime/Engine/Classes/PhysicsEngine/BodySetup.h
--- a/Engine/Source/Runtime/Engine/Classes/PhysicsEngine/BodySetup.h
+++ b/Engine/Source/Runtime/Engine/Classes/PhysicsEngine/BodySetup.h
@@ -128,6 +128,13 @@ struct FCookBodySetupInfo
 
 struct FPhysXCookHelper;
 
+/** Option for GetScaledVolume to either query everything, or only geometries that will be used to calculate mass */
+enum class EVolumeCalculationMethod
+{
+	AllGeoms,
+	OnlyMassRelevantGeoms
+};
+
 /**
  * BodySetup contains all collision information that is associated with a single asset.
  * A single BodySetup instance is shared among many BodyInstances so that geometry data is not duplicated.
@@ -404,7 +411,7 @@ public:
 	ENGINE_API virtual float GetVolume(const FVector& Scale) const;
 
 	/** Returns the volume of this element givent the scale */
-	ENGINE_API virtual FVector::FReal GetScaledVolume(const FVector& Scale) const;
+	ENGINE_API virtual FVector::FReal GetScaledVolume(const FVector& Scale, EVolumeCalculationMethod InMethod = EVolumeCalculationMethod::AllGeoms) const;
 
 	/** Release Physics meshes (ConvexMeshes, TriMesh & TriMeshNegX) */
 	ENGINE_API void ClearPhysicsMeshes();
diff --git a/Engine/Source/Runtime/Engine/Private/PhysicsEngine/BodySetup.cpp b/Engine/Source/Runtime/Engine/Private/PhysicsEngine/BodySetup.cpp
--- a/Engine/Source/Runtime/Engine/Private/PhysicsEngine/BodySetup.cpp
+++ b/Engine/Source/Runtime/Engine/Private/PhysicsEngine/BodySetup.cpp
@@ -1743,29 +1743,29 @@ FVector::FReal FKAggregateGeom::GetVolume(const FVector& Scale) const
 	return GetScaledVolume(Scale);
 }
 
-FVector::FReal FKAggregateGeom::GetScaledVolume(const FVector& Scale) const
+FVector::FReal FKAggregateGeom::GetScaledVolume(const FVector& Scale3D) const
 {
 	float Volume = 0.0f;
 
-	for (auto SphereElemIt = SphereElems.CreateConstIterator(); SphereElemIt; ++SphereElemIt )
+	VisitSimElements([&Scale3D, &Volume] <typename ElementType> (const ElementType & Elem)
 	{
-		Volume += SphereElemIt->GetScaledVolume(Scale);
-	}
+		Volume += Elem.GetScaledVolume(Scale3D);
+	});
 
-	for (auto BoxElemIt = BoxElems.CreateConstIterator(); BoxElemIt; ++BoxElemIt )
-	{
-		Volume += BoxElemIt->GetScaledVolume(Scale);
-	}
+	return Volume;
+}
 
-	for (auto SphylElemIt = SphylElems.CreateConstIterator(); SphylElemIt; ++SphylElemIt )
-	{
-		Volume += SphylElemIt->GetScaledVolume(Scale);
-	}
+FVector::FReal FKAggregateGeom::GetScaledVolumeForMassContribution(const FVector& Scale3D) const
+{
+	float Volume = 0.0f;
 
-	for (auto ConvexElemIt = ConvexElems.CreateConstIterator(); ConvexElemIt; ++ConvexElemIt )
+	VisitSimElements([&Scale3D, &Volume] <typename ElementType> (const ElementType& Elem)
 	{
-		Volume += ConvexElemIt->GetScaledVolume(Scale);
-	}
+		if(Elem.GetContributeToMass())
+		{
+			Volume += Elem.GetScaledVolume(Scale3D);
+		}
+	});
 
 	return Volume;
 }
@@ -2389,7 +2389,7 @@ float UBodySetup::CalculateMass(const UPrimitiveComponent* Component) const
 	}
 
 	// Then scale mass to avoid big differences between big and small objects.
-	const FVector::FReal BasicVolume = GetScaledVolume(ComponentScale);
+	const FVector::FReal BasicVolume = GetScaledVolume(ComponentScale, EVolumeCalculationMethod::OnlyMassRelevantGeoms);
 
 	// The below TODO is probably fixed now. GetSCaledVolume handles negative scales...
 	//@TODO: Some static meshes are triggering this - disabling until content can be analyzed - ensureMsgf(BasicVolume >= 0.0f, TEXT("UBodySetup::CalculateMass(%s) - The volume of the aggregate geometry is negative"), *Component->GetReadableName());
@@ -2407,12 +2407,15 @@ float UBodySetup::GetVolume(const FVector& Scale) const
 	return GetScaledVolume(Scale);
 }
 
-
-FVector::FReal UBodySetup::GetScaledVolume(const FVector& Scale) const
+FVector::FReal UBodySetup::GetScaledVolume(const FVector& Scale, EVolumeCalculationMethod InMethod) const
 {
-	return AggGeom.GetScaledVolume(Scale);
-}
+	if(InMethod == EVolumeCalculationMethod::AllGeoms)
+	{
+		return AggGeom.GetScaledVolume(Scale);
+	}
 
+	return AggGeom.GetScaledVolumeForMassContribution(Scale);
+}
 
 /** Helper function to safely copy instances of this shape*/
 void FKLevelSetElem::CloneElem(const FKLevelSetElem& Other)
