Added asset validation to AnimSequenceTransformProviderData and support for specifying a parent in order to override the skinned mesh.

#Synced-CL 50468060

diff --git a/Engine/Source/Runtime/Engine/Classes/Animation/AnimSequenceTransformProviderData.h b/Engine/Source/Runtime/Engine/Classes/Animation/AnimSequenceTransformProviderData.h
--- a/Engine/Source/Runtime/Engine/Classes/Animation/AnimSequenceTransformProviderData.h
+++ b/Engine/Source/Runtime/Engine/Classes/Animation/AnimSequenceTransformProviderData.h
@@ -18,6 +18,7 @@ class UAnimSequence;
 class UAnimSequenceTransformProviderData;
 class FAnimSequenceTransformProviderProxy;
 class IAnimSequenceTransformProvider;
+class FDataValidationContext;
 
 UENUM(BlueprintType)
 enum class EAnimSequenceTrackLoopMode : uint8
@@ -432,9 +433,14 @@ public:
 	virtual FTransformProviderRenderProxy* CreateRenderProxy(FInstancedSkinningSceneExtensionProxy* ExtensionProxy) const override;
 	virtual void BeginDestroy() override;
 	virtual bool IsReadyForFinishDestroy() override;
+	virtual UTransformProviderData* GetRootTransformProvider() override;
+	virtual void PreSave(FObjectPreSaveContext SaveContext) override;
+	virtual void PostSaveRoot(FObjectPostSaveRootContext SaveContext) override;
+	virtual void PostLoad() override;
 
 #if WITH_EDITOR
 	virtual void PostEditChangeProperty(FPropertyChangedEvent& PropertyChangedEvent) override;
+	virtual EDataValidationResult IsDataValid(FDataValidationContext& Context) const override;
 #endif
 
 public:
@@ -460,10 +466,15 @@ protected:
 	bool IsValidFor(FInstancedSkinningSceneExtensionProxy* ExtensionProxy) const;
 
 private:
+	virtual void SetParentProvider(UAnimSequenceTransformProviderData* Provider);
+
 	UPROPERTY(EditAnywhere, BlueprintReadWrite, Category = TransformProvider, meta = (AllowPrivateAccess = "true", ShowOnlyInnerProperties))
 	TObjectPtr<USkinnedAsset> SkinnedAsset;
 
-	UPROPERTY(EditAnywhere, BlueprintReadOnly, Category = TransformProvider, meta = (AllowPrivateAccess = "true", ShowOnlyInnerProperties))
+	UPROPERTY(EditAnywhere, BlueprintReadWrite, Category = TransformProvider, meta = (AllowPrivateAccess = "true", ShowOnlyInnerProperties))
+	TObjectPtr<UAnimSequenceTransformProviderData> ParentProvider;
+
+	UPROPERTY(EditAnywhere, BlueprintReadOnly, Category = TransformProvider, meta = (AllowPrivateAccess = "true", ShowOnlyInnerProperties, EditCondition = "ParentProvider == nullptr"))
 	TArray<struct FAnimSequenceTransformProviderSequence> Sequences;
 
 	double GetWorldTime() const;
@@ -486,7 +497,6 @@ public:
 	virtual void BeginDestroy() override;
 	virtual void Serialize(FArchive& Ar) override;
 	virtual void SubmitChanges() override;
-	virtual UTransformProviderData* GetRootTransformProvider() override;
 
 public:
 	// Copies items from this anim bank to create a new anim bank instance.
@@ -538,12 +548,8 @@ public:
 	ENGINE_API float GetPosition(int32 TrackIndex) const;
 
 private:
-	void Initialize(UAnimSequenceTransformProviderData* ProviderData);
 	void MarkOwnerRenderStateDirty();
-
-	// The parent provider data this instance was created from.
-	UPROPERTY()
-	TObjectPtr<UAnimSequenceTransformProviderData> Parent;
+	void SetParentProvider(UAnimSequenceTransformProviderData* Provider) override;
 
 	// Cached owner component for efficient dirty marking.
 	UInstancedSkinnedMeshComponent* CachedOwnerComponent = nullptr;
diff --git a/Engine/Source/Runtime/Engine/Private/Animation/AnimSequenceTransformProviderData.cpp b/Engine/Source/Runtime/Engine/Private/Animation/AnimSequenceTransformProviderData.cpp
--- a/Engine/Source/Runtime/Engine/Private/Animation/AnimSequenceTransformProviderData.cpp
+++ b/Engine/Source/Runtime/Engine/Private/Animation/AnimSequenceTransformProviderData.cpp
@@ -10,7 +10,13 @@
 #include "SceneInterface.h"
 #include "UObject/UObjectThreadContext.h"
 #include "UObject/FortniteMainBranchObjectVersion.h"
+#include "UObject/ObjectSaveContext.h"
 #include "ObjectCacheContext.h"
+#include "Framework/Notifications/NotificationManager.h"
+#include "Widgets/Notifications/SNotificationList.h"
+#include "Rendering/SkeletalMeshRenderData.h"
+#include "Misc/DataValidation.h"
+#include "Misc/MessageDialog.h"
 
 #if WITH_EDITOR
 #include "Animation/AnimBank.h"
@@ -18,6 +24,8 @@
 #include "Misc/PackageName.h"
 #endif
 
+#define LOCTEXT_NAMESPACE "AnimSequenceTranformProviderData"
+
 ///////////////////////////////////////////////////////////////////////////////////////////////////
 
 static TAutoConsoleVariable<bool> CVarAnimSequenceTransformProviderBlending(
@@ -160,7 +168,7 @@ void FAnimSequenceTrackRenderData::Tick(float DeltaTime, float SourcePosition, f
 {
 	const float BlendTime = Target.GetBlendTime();
 
-	if (BlendTime > 0.0f && BlendWeight < 1.0f)
+	if (BlendTime > 0.0f && BlendWeight < 1.0f && CVarAnimSequenceTransformProviderBlending.GetValueOnRenderThread())
 	{
 		BlendWeight += DeltaTime / BlendTime;
 
@@ -285,7 +293,16 @@ bool UAnimSequenceTransformProviderData::IsEnabled() const
 		}
 	}
 
-	return bEnabled;
+	// Any disabled parent in the chain disables the child.
+	for (const UAnimSequenceTransformProviderData* Current = this; Current; Current = Current->ParentProvider)
+	{
+		if (!Current->bEnabled)
+		{
+			return false;
+		}
+	}
+
+	return true;
 }
 
 const FGuid& UAnimSequenceTransformProviderData::GetTransformProviderID() const
@@ -471,6 +488,63 @@ double UAnimSequenceTransformProviderData::GetWorldTime() const
 #endif
 }
 
+UTransformProviderData* UAnimSequenceTransformProviderData::GetRootTransformProvider()
+{
+	constexpr int32 MaxDepth = 8;
+	int32 Depth = 0;
+	UAnimSequenceTransformProviderData* Current = this;
+	for (; Current->ParentProvider && Depth < MaxDepth; ++Depth)
+	{
+		Current = Current->ParentProvider;
+	}
+	ensureMsgf(Depth != MaxDepth, TEXT("UAnimSequenceTransformProvider '%s' either has too deep of a hierarchy or is a circular reference."), *GetPathName());
+	return Current;
+}
+
+void UAnimSequenceTransformProviderData::SetParentProvider(UAnimSequenceTransformProviderData* ProviderData)
+{
+	ParentProvider = ProviderData;
+
+	if (ParentProvider)
+	{
+		Sequences = ProviderData->Sequences;
+	}
+	else
+	{
+		Sequences.Empty();
+	}
+}
+
+void UAnimSequenceTransformProviderData::PreSave(FObjectPreSaveContext SaveContext)
+{
+	Super::PreSave(SaveContext);
+
+	if (ParentProvider)
+	{
+		Sequences.Empty();
+	}
+}
+
+void UAnimSequenceTransformProviderData::PostSaveRoot(FObjectPostSaveRootContext SaveContext)
+{
+	Super::PostSaveRoot(SaveContext);
+
+	if (ParentProvider)
+	{
+		Sequences = ParentProvider->GetSequences();
+	}
+}
+
+void UAnimSequenceTransformProviderData::PostLoad()
+{
+	Super::PostLoad();
+
+	if (ParentProvider && IsValid(ParentProvider))
+	{
+		Sequences = ParentProvider->GetSequences();
+	}
+}
+
 #if WITH_EDITOR
 
 void UAnimSequenceTransformProviderData::PostEditChangeProperty(FPropertyChangedEvent& PropertyChangedEvent)
@@ -482,21 +556,72 @@ void UAnimSequenceTransformProviderData::PostEditChangeProperty(FPropertyChanged
 		return;
 	}
 
-	// Propagate Sequences changes to all instances
-	FObjectCacheContextScope ObjectCacheScope;
-	for (UInstancedSkinnedMeshComponent* Component : ObjectCacheScope.GetContext().GetInstancedSkinnedMeshComponents(GetRootTransformProvider()))
+	if (PropertyChangedEvent.GetPropertyName() == GET_MEMBER_NAME_CHECKED(UAnimSequenceTransformProviderData, ParentProvider))
 	{
-		if (Component && IsValid(Component))
+		if (ParentProvider)
+		{
+			if (UAnimSequenceTransformProviderDataInstance* Provider = Cast<UAnimSequenceTransformProviderDataInstance>(ParentProvider))
+			{
+				FMessageDialog::Open(
+					EAppMsgType::Ok,
+					NSLOCTEXT("AnimSequenceTransformProvider", "ParentNotInstance", "AnimSequenceTransformProviderDataInstance cannot be set as a parent.")
+				);
+				ParentProvider = nullptr;
+			}
+
+			UAnimSequenceTransformProviderData* Current = ParentProvider;
+			while (Current)
+			{
+				if (Current == this)
+				{
+					FMessageDialog::Open(
+						EAppMsgType::Ok,
+						NSLOCTEXT("AnimSequenceTransformProvider", "CircularReferences", "Setting this parent provider creates a circular reference.")
+					);
+					ParentProvider = nullptr;
+					break;
+				}
+				Current = Current->ParentProvider;
+			}
+		}
+
+		SetParentProvider(ParentProvider);
+	}
+
+	// Propagate Sequences changes to all children.
+	TArray<TPair<int32, UAnimSequenceTransformProviderData*>> ChildrenByDepth;
+
+	ForEachObjectOfClass(UAnimSequenceTransformProviderData::StaticClass(), [&ChildrenByDepth](UObject* Obj)
+	{
+		if (UAnimSequenceTransformProviderData* Provider = Cast<UAnimSequenceTransformProviderData>(Obj))
 		{
-			// Sync instance data if it's a child of this provider
-			if (UAnimSequenceTransformProviderDataInstance* Instance = Cast<UAnimSequenceTransformProviderDataInstance>(Component->GetTransformProvider()))
+			if (Provider->ParentProvider)
 			{
-				if (Instance->Parent == this)
+				int32 Depth = 0;
+				UAnimSequenceTransformProviderData* Current = Provider->ParentProvider;
+				while (Current)
 				{
-					Instance->Initialize(this);
+					++Depth;
+					Current = Current->ParentProvider;
 				}
+				ChildrenByDepth.Emplace(Depth, Provider);
 			}
+		}
+	});
+
+	ChildrenByDepth.Sort([](const auto& A, const auto& B) { return A.Key < B.Key; });
+
+	for (const auto& Pair : ChildrenByDepth)
+	{
+		Pair.Value->SetParentProvider(Pair.Value->ParentProvider);
+	}
 
+	// Mark render state dirty for components which reference it.
+	FObjectCacheContextScope ObjectCacheScope;
+	for (UInstancedSkinnedMeshComponent* Component : ObjectCacheScope.GetContext().GetInstancedSkinnedMeshComponents(GetRootTransformProvider()))
+	{
+		if (Component && IsValid(Component))
+		{
 			Component->MarkRenderStateDirty();
 		}
 	}
@@ -604,6 +729,73 @@ UAnimSequenceTransformProviderData* UAnimSequenceTransformProviderData::CreateFr
 	return NewProviderData;
 }
 
+EDataValidationResult UAnimSequenceTransformProviderData::IsDataValid(FDataValidationContext& Context) const
+{
+	EDataValidationResult Result = Super::IsDataValid(Context);
+
+	for (const FAnimSequenceTransformProviderSequence& Sequence : Sequences)
+	{
+		if (Sequence.Sequence == nullptr)
+		{
+			Context.AddError(FText::Format(LOCTEXT("NullSequence",
+				"AnimSequenceTransformProvider {0} invalid null sequence."),
+				FText::FromString(GetPathName())
+			));
+			Result = EDataValidationResult::Invalid;
+			continue;
+		}
+
+		const USkeleton* SourceSkeleton = Sequence.Sequence->GetSkeleton();
+
+		if (SourceSkeleton == nullptr)
+		{
+			Context.AddError(FText::Format(LOCTEXT("NullSourceSkeleton",
+				"AnimSequenceTransformProvider {0} sequence {1} has a null skeleton."),
+				FText::FromString(GetPathName()),
+				FText::FromString(Sequence.Sequence->GetPathName())
+			));
+			Result = EDataValidationResult::Invalid;
+			continue;
+		}
+
+		if (SkinnedAsset)
+		{
+			const USkeleton* TargetSkeleton = SkinnedAsset->GetSkeleton();
+
+			if (TargetSkeleton == nullptr)
+			{
+				Context.AddError(FText::Format(LOCTEXT("NullTargetSkeleton",
+					"AnimSequenceTransformProvider {0} skinned mesh {1} has a null skeleton."),
+					FText::FromString(GetPathName()),
+					FText::FromString(SkinnedAsset->GetPathName())
+				));
+				Result = EDataValidationResult::Invalid;
+				continue;
+			}
+
+			if (SourceSkeleton != TargetSkeleton)
+			{
+				const FSkeletonRemapping& SkeletonRemapping = UE::Anim::FSkeletonRemappingRegistry::Get().GetRemapping(SourceSkeleton, TargetSkeleton);
+
+				if (!SkeletonRemapping.IsValid())
+				{
+					Context.AddError(FText::Format(LOCTEXT("IncompatibleSkeletons",
+						"AnimSequenceTransformProvider {0} Incompatible skeletons: sequence {1} uses skeleton {2}, but skinned asset {3} uses skeleton {4} with no valid remapping."),
+						FText::FromString(GetPathName()),
+						FText::FromString(Sequence.Sequence->GetPathName()),
+						FText::FromString(SourceSkeleton->GetPathName()),
+						FText::FromString(SkinnedAsset->GetPathName())
+					));
+					Result = EDataValidationResult::Invalid;
+					continue;
+				}
+			}
+		}
+	}
+
+	return Result;
+}
+
 #endif // WITH_EDITOR
 
 ///////////////////////////////////////////////////////////////////////////////////////////////////
@@ -613,13 +805,24 @@ UAnimSequenceTransformProviderDataInstance* UAnimSequenceTransformProviderDataIn
 	if (IsValid(ProviderData) && IsValid(Owner))
 	{
 		UAnimSequenceTransformProviderDataInstance* AnimSequenceInstance = NewObject<UAnimSequenceTransformProviderDataInstance>(Owner);
-		AnimSequenceInstance->Initialize(ProviderData);
+		AnimSequenceInstance->SetParentProvider(ProviderData);
 		return AnimSequenceInstance;
 	}
 
 	return nullptr;
 }
 
+void UAnimSequenceTransformProviderDataInstance::SetParentProvider(UAnimSequenceTransformProviderData* ProviderData)
+{
+	Super::SetParentProvider(ProviderData);
+
+	if (ProviderData)
+	{
+		SkinnedAsset = ProviderData->SkinnedAsset;
+		bEnabled = ProviderData->bEnabled;
+	}
+}
+
 bool UAnimSequenceTransformProviderDataInstance::IsEnabled() const
 {
 	return Tracks.GetNumTracks() != 0 && Super::IsEnabled();
@@ -641,11 +844,6 @@ uint32 UAnimSequenceTransformProviderDataInstance::GetSkinningDataOffset(int32 I
 	return 0u;
 }
 
-UTransformProviderData* UAnimSequenceTransformProviderDataInstance::GetRootTransformProvider()
-{
-	return Parent ? Parent->GetRootTransformProvider() : this;
-}
-
 FTransformProviderRenderProxy* UAnimSequenceTransformProviderDataInstance::CreateRenderProxy(FInstancedSkinningSceneExtensionProxy* SceneExtensionProxy) const
 {
 	if (IsValidFor(SceneExtensionProxy))
@@ -720,14 +918,6 @@ void UAnimSequenceTransformProviderDataInstance::SubmitChanges()
 	}
 }
 
-void UAnimSequenceTransformProviderDataInstance::Initialize(UAnimSequenceTransformProviderData* ProviderData)
-{
-	SkinnedAsset = ProviderData->SkinnedAsset;
-	Parent    = ProviderData;
-	Sequences = ProviderData->Sequences;
-	bEnabled  = ProviderData->bEnabled;
-}
-
 void UAnimSequenceTransformProviderDataInstance::Serialize(FArchive& Ar)
 {
 	Ar.UsingCustomVersion(FFortniteMainBranchObjectVersion::GUID);
@@ -736,11 +926,6 @@ void UAnimSequenceTransformProviderDataInstance::Serialize(FArchive& Ar)
 
 	if (Ar.IsLoading())
 	{
-		if (Parent)
-		{
-			Initialize(Parent);
-		}
-
 		bNeedsPostLoadInit = true;
 	}
 }
@@ -883,4 +1068,6 @@ float UAnimSequenceTransformProviderDataInstance::GetPosition(int32 TrackIndex)
 	}
 
 	return 0.0f;
-}
\ No newline at end of file
+}
+
+#undef LOCTEXT_NAMESPACE
\ No newline at end of file
