Remove GInputLatencyTimer
- Unused global object that's only ever written to. The computed results are never used.

#Synced-CL 50397505

diff --git a/Engine/Source/Runtime/D3D12RHI/Private/D3D12Viewport.cpp b/Engine/Source/Runtime/D3D12RHI/Private/D3D12Viewport.cpp
--- a/Engine/Source/Runtime/D3D12RHI/Private/D3D12Viewport.cpp
+++ b/Engine/Source/Runtime/D3D12RHI/Private/D3D12Viewport.cpp
@@ -948,16 +948,6 @@ void FD3D12DynamicRHI::RHIEndDrawingViewport(FRHICommandListImmediate& RHICmdLis
 					//No frame wait
 				}
 			}
-
-			// If the input latency timer has been triggered, block until the GPU is completely
-			// finished displaying this frame and calculate the delta time.
-			if (GInputLatencyTimer.RenderThreadTrigger)
-			{
-				Viewport->WaitForFrameEventCompletion();
-				uint32 EndTime = FPlatformTime::Cycles();
-				GInputLatencyTimer.DeltaTime = EndTime - GInputLatencyTimer.StartTime;
-				GInputLatencyTimer.RenderThreadTrigger = false;
-			}
 		});
 	}
 
diff --git a/Engine/Source/Runtime/Engine/Private/UnrealClient.cpp b/Engine/Source/Runtime/Engine/Private/UnrealClient.cpp
--- a/Engine/Source/Runtime/Engine/Private/UnrealClient.cpp
+++ b/Engine/Source/Runtime/Engine/Private/UnrealClient.cpp
@@ -1709,26 +1709,6 @@ void FViewport::HighResScreenshot()
 	HighResScreenshotBeginFrame(nullptr);
 }
 
-struct FEndDrawingCommandParams
-{
-	FViewport* Viewport;
-	uint32 bLockToVsync : 1;
-	uint32 bShouldTriggerTimerEvent : 1;
-	uint32 bShouldPresent : 1;
-};
-
-/**
- * Helper function used in ENQUEUE_RENDER_COMMAND below. Needed to be split out due to
- * use of macro and former already being one.
- *
- * @param Parameters	Parameters passed from the gamethread to the renderthread command.
- */
-static void ViewportEndDrawing(FRHICommandListImmediate& RHICmdList, FEndDrawingCommandParams Parameters)
-{	
-	GInputLatencyTimer.RenderThreadTrigger = Parameters.bShouldTriggerTimerEvent;
-	Parameters.Viewport->EndRenderFrame(RHICmdList, Parameters.bShouldPresent, Parameters.bLockToVsync);
-}
-
 /** Starts a new rendering frame. Called from the rendering thread. */
 void FViewport::BeginRenderFrame(FRHICommandListImmediate& RHICmdList)
 {
@@ -1768,23 +1748,23 @@ FRHIGPUMask FViewport::GetGPUMask(FRHICommandListImmediate& RHICmdList) const
 */
 void FViewport::EnqueueBeginRenderFrame(const bool bShouldPresent)
 {
-	FViewport* Viewport = this;
-	ENQUEUE_RENDER_COMMAND(BeginDrawingCommand, LoopTick)(
-		[Viewport](FRHICommandListImmediate& RHICmdList)
-		{
-			Viewport->BeginRenderFrame(RHICmdList);
-		});
+	ENQUEUE_RENDER_COMMAND(BeginDrawingCommand, LoopTick)([this](FRHICommandListImmediate& RHICmdList)
+	{
+		BeginRenderFrame(RHICmdList);
+	});
 }
 
 
 void FViewport::EnqueueEndRenderFrame(const bool bLockToVsync, const bool bShouldPresent)
 {
-	FEndDrawingCommandParams Params = { this, (uint32)bLockToVsync, (uint32)GInputLatencyTimer.GameThreadTrigger, (uint32)(PresentAndStopMovieDelay > 0 ? 0 : bShouldPresent) };
-	ENQUEUE_RENDER_COMMAND(EndDrawingCommand, LoopTick)(
-		[Params](FRHICommandListImmediate& RHICmdList)
-		{
-			ViewportEndDrawing(RHICmdList, Params);
-		});
+	bool bPresent = (PresentAndStopMovieDelay <= 0)
+		? bShouldPresent
+		: false;
+
+	ENQUEUE_RENDER_COMMAND(EndDrawingCommand, LoopTick)([this, bLockToVsync, bPresent](FRHICommandListImmediate& RHICmdList)
+	{
+		EndRenderFrame(RHICmdList, bPresent, bLockToVsync);
+	});
 }
 
 void FViewport::Draw( bool bShouldPresent /*= true */)
@@ -1890,8 +1870,6 @@ void FViewport::Draw( bool bShouldPresent /*= true */)
 				SetRequiresVsync(bLockToVsync);
 				SetCameraCut(bCameraCut);
 				EnqueueEndRenderFrame(bLockToVsync, bShouldPresent);
-
-				GInputLatencyTimer.GameThreadTrigger = false;
 			}
 			else if (World && World->IsGameWorld() && !GUseUnifiedTimeBudgetForStreaming)
 			{
diff --git a/Engine/Source/Runtime/OpenGLDrv/Private/OpenGLViewport.cpp b/Engine/Source/Runtime/OpenGLDrv/Private/OpenGLViewport.cpp
--- a/Engine/Source/Runtime/OpenGLDrv/Private/OpenGLViewport.cpp
+++ b/Engine/Source/Runtime/OpenGLDrv/Private/OpenGLViewport.cpp
@@ -166,16 +166,6 @@ void FOpenGLDynamicRHI::RHIEndDrawingViewport(FRHICommandListImmediate& RHICmdLi
 					//No frame wait
 				}
 			}
-		
-			// If the input latency timer has been triggered, block until the GPU is completely
-			// finished displaying this frame and calculate the delta time.
-			if ( GInputLatencyTimer.RenderThreadTrigger )
-			{
-				Viewport->WaitForFrameEventCompletion();
-				uint32 EndTime = FPlatformTime::Cycles();
-				GInputLatencyTimer.DeltaTime = EndTime - GInputLatencyTimer.StartTime;
-				GInputLatencyTimer.RenderThreadTrigger = false;
-			}
 		}
 
 		EndFrameTick();
diff --git a/Engine/Source/Runtime/RenderCore/Private/RenderCore.cpp b/Engine/Source/Runtime/RenderCore/Private/RenderCore.cpp
--- a/Engine/Source/Runtime/RenderCore/Private/RenderCore.cpp
+++ b/Engine/Source/Runtime/RenderCore/Private/RenderCore.cpp
@@ -234,34 +234,6 @@ static TAutoConsoleVariable<int32> CVarForceLODShadow(
 
 #endif // EXPOSE_FORCE_LOD
 
-FInputLatencyTimer GInputLatencyTimer( 2.0f );
-
-//
-// FInputLatencyTimer implementation.
-//
-
-/** Potentially starts the timer on the gamethread, based on the UpdateFrequency. */
-void FInputLatencyTimer::GameThreadTick()
-{
-#if STATS
-	if (FThreadStats::IsCollectingData())
-	{
-		if ( !bInitialized )
-		{
-			LastCaptureTime	= FPlatformTime::Seconds();
-			bInitialized	= true;
-		}
-		const double CurrentTimeInSeconds = FPlatformTime::Seconds();
-		if ( (CurrentTimeInSeconds - LastCaptureTime) > UpdateFrequency )
-		{
-			LastCaptureTime		= CurrentTimeInSeconds;
-			StartTime			= FPlatformTime::Cycles();
-			GameThreadTrigger	= true;
-		}
-	}
-#endif
-}
-
 FPixelRenderCounters GPixelRenderCounters;
 
 void TickPixelRenderCounters()
diff --git a/Engine/Source/Runtime/RenderCore/Public/RenderTimer.h b/Engine/Source/Runtime/RenderCore/Public/RenderTimer.h
--- a/Engine/Source/Runtime/RenderCore/Public/RenderTimer.h
+++ b/Engine/Source/Runtime/RenderCore/Public/RenderTimer.h
@@ -4,55 +4,6 @@
 
 #include "HAL/PlatformTime.h"
 
-/**
- * Encapsulates a latency timer that measures the time from when mouse input
- * is read on the gamethread until that frame is fully displayed by the GPU.
- */
-struct FInputLatencyTimer
-{
-	/**
-	 * Constructor
-	 * @param InUpdateFrequency	How often the timer should be updated (in seconds).
-	 */
-	FInputLatencyTimer(float InUpdateFrequency)
-		: UpdateFrequency(InUpdateFrequency)
-	{
-	}
-
-	/** Potentially starts the timer on the gamethread, based on the UpdateFrequency. */
-	RENDERCORE_API void GameThreadTick();
-
-	/** @return The number of seconds of input latency. */
-	inline float GetDeltaSeconds() const
-	{
-		return FPlatformTime::ToSeconds(DeltaTime);
-	}
-
-	/** Whether GInputLatencyTimer is initialized or not. */
-	bool	bInitialized = false;
-
-	/** Whether a measurement has been triggered on the gamethread. */
-	bool	GameThreadTrigger = false;
-
-	/** Whether a measurement has been triggered on the renderthread. */
-	bool	RenderThreadTrigger = false;
-
-	/** Start time (in FPlatformTime::Cycles). */
-	uint32	StartTime = 0;
-
-	/** Last delta time that was measured (in FPlatformTime::Cycles). */
-	uint32	DeltaTime = 0;
-
-	/** Last time we did a measurement (in seconds). */
-	double	LastCaptureTime = 0.0f;
-
-	/** How often we should do a measurement (in seconds). */
-	float	UpdateFrequency;
-};
-
-/** Global input latency timer. Defined in UnrealClient.cpp */
-extern RENDERCORE_API FInputLatencyTimer GInputLatencyTimer;
-
 /** How many cycles the renderthread used (excluding idle time). It's set once per frame in FViewport::Draw. */
 extern RENDERCORE_API uint32 GRenderThreadTime;
 
diff --git a/Engine/Source/Runtime/VulkanRHI/Private/VulkanViewport.cpp b/Engine/Source/Runtime/VulkanRHI/Private/VulkanViewport.cpp
--- a/Engine/Source/Runtime/VulkanRHI/Private/VulkanViewport.cpp
+++ b/Engine/Source/Runtime/VulkanRHI/Private/VulkanViewport.cpp
@@ -810,16 +810,6 @@ bool FVulkanViewport::Present(FVulkanCommandListContext& Context, const FRHIPres
 		IssueFrameEvent();
 	}
 
-	// If the input latency timer has been triggered, block until the GPU is completely
-	// finished displaying this frame and calculate the delta time.
-	//if (GInputLatencyTimer.RenderThreadTrigger)
-	//{
-	//	WaitForFrameEventCompletion();
-	//	uint32 EndTime = FPlatformTime::Cycles();
-	//	GInputLatencyTimer.DeltaTime = EndTime - GInputLatencyTimer.StartTime;
-	//	GInputLatencyTimer.RenderThreadTrigger = false;
-	//}
-
 	AcquiredImageIndex = -1;
 
 	++PresentCount;
diff --git a/Engine/Source/Runtime/Windows/D3D11RHI/Private/D3D11Viewport.cpp b/Engine/Source/Runtime/Windows/D3D11RHI/Private/D3D11Viewport.cpp
--- a/Engine/Source/Runtime/Windows/D3D11RHI/Private/D3D11Viewport.cpp
+++ b/Engine/Source/Runtime/Windows/D3D11RHI/Private/D3D11Viewport.cpp
@@ -854,16 +854,6 @@ void FD3D11DynamicRHI::RHIEndDrawingViewport(FRHICommandListImmediate& RHICmdLis
 				//No frame wait
 			}
 		}
-
-		// If the input latency timer has been triggered, block until the GPU is completely
-		// finished displaying this frame and calculate the delta time.
-		if ( GInputLatencyTimer.RenderThreadTrigger )
-		{
-			Viewport->WaitForFrameEventCompletion();
-			uint32 EndTime = FPlatformTime::Cycles();
-			GInputLatencyTimer.DeltaTime = EndTime - GInputLatencyTimer.StartTime;
-			GInputLatencyTimer.RenderThreadTrigger = false;
-		}
 	});
 }
 
