Fix to Issue 1213056: Compile errors in EditorWidgetFocusUtils.cpp

#Synced-CL 50447978

diff --git a/Engine/Source/Runtime/Engine/Private/EditorWidgetFocusUtils.cpp b/Engine/Source/Runtime/Engine/Private/EditorWidgetFocusUtils.cpp
--- a/Engine/Source/Runtime/Engine/Private/EditorWidgetFocusUtils.cpp
+++ b/Engine/Source/Runtime/Engine/Private/EditorWidgetFocusUtils.cpp
@@ -5,24 +5,28 @@
 #include "Framework/Application/SlateApplication.h"
 #include "Framework/Application/SlateUser.h"
 
-namespace Private
+namespace WidgetFocusRestorationUtils
 {
-	int32 FindChild(FChildren* Children, const TSharedRef<SWidget>& ChildToFind)
+	namespace Private
 	{
-		if (Children)
+		int32 FindChild(FChildren* Children, const TSharedRef<SWidget>& ChildToFind)
 		{
-			for (int32 ChildIdx = 0; ChildIdx < Children->Num(); ++ChildIdx)
+			if (Children)
 			{
-				if (TSharedRef<SWidget> Child = Children->GetChildAt(ChildIdx); ChildToFind == Child)
+				for (int32 ChildIdx = 0; ChildIdx < Children->Num(); ++ChildIdx)
 				{
-					return ChildIdx;
+					if (TSharedRef<SWidget> Child = Children->GetChildAt(ChildIdx); ChildToFind == Child)
+					{
+						return ChildIdx;
+					}
 				}
 			}
+
+			return INDEX_NONE;
 		}
-		
-		return INDEX_NONE;
 	}
 }
+
 using namespace WidgetFocusRestorationUtils;
 
 TOptional<FFocusedWidgetData> WidgetFocusRestorationUtils::GetFocusedWidgetData()
@@ -64,7 +68,7 @@ TOptional<FFocusedWidgetData> WidgetFocusRestorationUtils::GetFocusedWidgetData(
 	{
 		TSharedRef<SWidget> Widget = FocusPath.Widgets[WidgetIdx].Widget;
 
-		if (const int32 ChildIdx = Private::FindChild(WidgetChildren, Widget); ChildIdx != INDEX_NONE)
+		if (const int32 ChildIdx = WidgetFocusRestorationUtils::Private::FindChild(WidgetChildren, Widget); ChildIdx != INDEX_NONE)
 		{
 			IndexPathToWidget.Add(FIndexTypePair{ .Index = ChildIdx, .Type = Widget->GetType() });
 			WidgetChildren = Widget->GetChildren();
