Disable bTrackMemory flags to FAudioDevice::Precache

#tests

#Synced-CL 50465865

diff --git a/Engine/Source/Runtime/AudioMixer/Private/SoundWaveDecoder.cpp b/Engine/Source/Runtime/AudioMixer/Private/SoundWaveDecoder.cpp
--- a/Engine/Source/Runtime/AudioMixer/Private/SoundWaveDecoder.cpp
+++ b/Engine/Source/Runtime/AudioMixer/Private/SoundWaveDecoder.cpp
@@ -564,7 +564,7 @@ namespace Audio
 		{
 			if (PrecacheState == ESoundWavePrecacheState::NotStarted)
 			{
-				AudioDevice->Precache(InitData.SoundWave, true);
+				AudioDevice->Precache(InitData.SoundWave, true, false, false);
 			}
 			check(InitData.SoundWave->GetPrecacheState() == ESoundWavePrecacheState::Done);
 			return InitDecodingSourceInternal(InitData);
diff --git a/Engine/Source/Runtime/Engine/Private/AudioDevice.cpp b/Engine/Source/Runtime/Engine/Private/AudioDevice.cpp
--- a/Engine/Source/Runtime/Engine/Private/AudioDevice.cpp
+++ b/Engine/Source/Runtime/Engine/Private/AudioDevice.cpp
@@ -803,7 +803,7 @@ void FAudioDevice::PrecacheStartupSounds()
 		for (TObjectIterator<USoundWave> It; It; ++It)
 		{
 			USoundWave* SoundWave = *It;
-			Precache(SoundWave);
+			Precache(SoundWave, false, false, false);
 		}
 
 		bStartupSoundsPreCached = true;
diff --git a/Engine/Source/Runtime/Engine/Private/GameViewportClient.cpp b/Engine/Source/Runtime/Engine/Private/GameViewportClient.cpp
--- a/Engine/Source/Runtime/Engine/Private/GameViewportClient.cpp
+++ b/Engine/Source/Runtime/Engine/Private/GameViewportClient.cpp
@@ -2458,7 +2458,7 @@ void UGameViewportClient::Precache()
 			for(TObjectIterator<USoundWave> It;It;++It)
 			{
 				USoundWave* SoundWave = *It;
-				AudioDevice->Precache( SoundWave );
+				AudioDevice->Precache( SoundWave, false, false, false );
 			}
 			UE_LOG(LogPlayerManagement, Log, TEXT("Precaching sounds completed..."));
 		}
diff --git a/Engine/Source/Runtime/Engine/Private/SampleBufferIO.cpp b/Engine/Source/Runtime/Engine/Private/SampleBufferIO.cpp
--- a/Engine/Source/Runtime/Engine/Private/SampleBufferIO.cpp
+++ b/Engine/Source/Runtime/Engine/Private/SampleBufferIO.cpp
@@ -35,7 +35,7 @@ namespace Audio
 			LoadingSoundWaveInfo.Status = FLoadingSoundWaveInfo::LoadStatus::Loading;
 
 			// Kick off a decompression/precache of the sound wave
-			AudioDevice->Precache(InSoundWave, bSynchronous, true, true);
+			AudioDevice->Precache(InSoundWave, bSynchronous, false, true);
 		}
 		else
 		{
diff --git a/Engine/Source/Runtime/Engine/Private/SoundWave.cpp b/Engine/Source/Runtime/Engine/Private/SoundWave.cpp
--- a/Engine/Source/Runtime/Engine/Private/SoundWave.cpp
+++ b/Engine/Source/Runtime/Engine/Private/SoundWave.cpp
@@ -2394,7 +2394,7 @@ void USoundWave::PostLoad()
 		if (AudioDevice)
 		{
 			// Upload the data to the hardware, but only if we've precached startup sounds already
-			AudioDevice->Precache(this);
+			AudioDevice->Precache(this, false, false, false);
 		}
 	}
 
diff --git a/Engine/Source/Runtime/Engine/Public/AudioDevice.h b/Engine/Source/Runtime/Engine/Public/AudioDevice.h
--- a/Engine/Source/Runtime/Engine/Public/AudioDevice.h
+++ b/Engine/Source/Runtime/Engine/Public/AudioDevice.h
@@ -621,7 +621,7 @@ public:
 	 * @param	bTrackMemory	If true, the audio mem stats will be updated
 	 * @param   bForceFullDecompression If true, the sound wave will be fully decompressed regardless of size.
 	 */
-	ENGINE_API virtual void Precache(USoundWave* SoundWave, bool bSynchronous = false, bool bTrackMemory = true, bool bForceFullDecompression = false);
+	ENGINE_API virtual void Precache(USoundWave* SoundWave, bool bSynchronous = false, bool bTrackMemory = false, bool bForceFullDecompression = false);
 
 	UE_DEPRECATED(5.7, "With PCM audio caching removed, this hueristic is no longer used.")
 	ENGINE_API float GetCompressionDurationThreshold(const FSoundGroup &SoundGroup);
