Fix to a comment to UActorComponent::bAsyncPhysicsTickEnabled that sounds too similar to the description of PhysicsSettings::bTickPhysicsAsync, and is too vague for users to understand how it impacts them.

#Synced-CL 50489768

diff --git a/Engine/Source/Runtime/Engine/Classes/Components/ActorComponent.h b/Engine/Source/Runtime/Engine/Classes/Components/ActorComponent.h
--- a/Engine/Source/Runtime/Engine/Classes/Components/ActorComponent.h
+++ b/Engine/Source/Runtime/Engine/Classes/Components/ActorComponent.h
@@ -390,7 +390,11 @@ private:
 	/** Tracks whether the component has been added to the world's pre end of frame sync list */
 	uint8 bMarkedForPreEndOfFrameSync : 1;
 
-	/** Whether to use use the async physics tick with this component. */
+	/**
+	* Whether this component will receive a call to AsyncPhysicsTickComponent (in code) and ReceiveAsyncPhysicsTick (a Blueprint implementable event).
+    * These functions will execute on the game thread and are allowed to read and write both Physics Thread (PT) and Game Thread (GT) Data .
+	* Performance Warning: This will cause both threads to halt and force a full sync of data between PT and GT. The PT will block until the GT is done executing all AsyncPhysicsTick functions.
+	*/
 	uint8 bAsyncPhysicsTickEnabled : 1;
 
 	enum struct EPhysicsStateAsyncCreationState : uint8
@@ -1414,7 +1418,11 @@ protected:
 	 *		already been registered. Setting bForceUpdate to true overrides that check */
 	ENGINE_API void HandleCanEverAffectNavigationChange(bool bForceUpdate = false);
 
-	/** Sets bAsyncPhysicsTickEnabled which determines whether to use use the async physics tick with this component. */
+	/**
+	* Sets bAsyncPhysicsTickEnabled which determines whether this component will receive a call to AsyncPhysicsTickComponent (in code) and ReceiveAsyncPhysicsTick (a Blueprint implementable event).
+	* These functions will execute on the game thread and are allowed to read and write both Physics Thread (PT) and Game Thread (GT) Data .
+	* Performance Warning: This will cause both threads to halt and force a full sync of data between PT and GT. The PT will block until the GT is done executing all AsyncPhysicsTick functions.
+	*/
 	ENGINE_API void SetAsyncPhysicsTickEnabled(bool bEnabled);
 
 	/** Defers the call to SetAsyncPhysicsTickEnabled(false) to the next async tick. This allows you to remove the
