Adding simple counter to ensure if too many Blueprint async actions are created.
The number of allowed actions can be set using the console variable `bp.MaxAsyncActionCount`. Set the number to <= 0 to disable the ensure.

#Synced-CL 50485578

diff --git a/Engine/Source/Runtime/Engine/Classes/Kismet/BlueprintAsyncActionBase.h b/Engine/Source/Runtime/Engine/Classes/Kismet/BlueprintAsyncActionBase.h
--- a/Engine/Source/Runtime/Engine/Classes/Kismet/BlueprintAsyncActionBase.h
+++ b/Engine/Source/Runtime/Engine/Classes/Kismet/BlueprintAsyncActionBase.h
@@ -23,6 +23,9 @@ public:
 	/** Default UObject constructor */
 	ENGINE_API UBlueprintAsyncActionBase(const FObjectInitializer& ObjectInitializer = FObjectInitializer::Get());
 
+	/** Called when the object is destroyed */
+	ENGINE_API virtual void BeginDestroy() override;
+
 	/** Called to trigger the action once the delegates have been bound */
 	UFUNCTION(BlueprintCallable, meta=(BlueprintInternalUseOnly="true"))
 	ENGINE_API virtual void Activate();
diff --git a/Engine/Source/Runtime/Engine/Private/Kismet/BlueprintAsyncActionBase.cpp b/Engine/Source/Runtime/Engine/Private/Kismet/BlueprintAsyncActionBase.cpp
--- a/Engine/Source/Runtime/Engine/Private/Kismet/BlueprintAsyncActionBase.cpp
+++ b/Engine/Source/Runtime/Engine/Private/Kismet/BlueprintAsyncActionBase.cpp
@@ -7,6 +7,45 @@
 
 #include UE_INLINE_GENERATED_CPP_BY_NAME(BlueprintAsyncActionBase)
 
+#if DO_ENSURE
+
+#include "Async/TransactionallySafeMutex.h"
+#include "Misc/NotNull.h"
+
+namespace UE::BlueprintAsyncAction::Private
+{
+	static TAutoConsoleVariable<int32> CVarMaxActionCount(
+		TEXT("bp.MaxAsyncActionCount"),
+		10000,
+		TEXT("The maximum number of async actions of a that can be active at a time before triggering an ensure. Set to <= 0 to disable")
+	);
+
+	/** Total number of async actions that are active */
+	static int32 GlobalAsyncActionCount = 0;
+	static UE::FTransactionallySafeMutex GlobalActionCountMutex;
+
+	static void DebugRecordAsyncActionCreated(TNotNull<const UBlueprintAsyncActionBase*> Action)
+	{
+		UE::TScopeLock Lock(GlobalActionCountMutex);
+		GlobalAsyncActionCount++;
+
+		const int32 MaxActionCount = CVarMaxActionCount.GetValueOnAnyThread();
+		if (MaxActionCount > 0)
+		{
+			ensureMsgf(GlobalAsyncActionCount <= MaxActionCount,
+				TEXT("Exceeded the maximum number of allowed async actions. Last created action class was %s"),
+				*GetNameSafe(Action->GetClass()));
+		}
+	}
+
+	static void DebugRecordAsyncActionDestroyed()
+	{
+		UE::TScopeLock Lock(GlobalActionCountMutex);
+		--GlobalAsyncActionCount;
+	}
+}
+#endif // DO_ENSURE
+
 //////////////////////////////////////////////////////////////////////////
 // UBlueprintAsyncActionBase
 
@@ -16,7 +55,23 @@ UBlueprintAsyncActionBase::UBlueprintAsyncActionBase(const FObjectInitializer& O
 	if (!HasAnyFlags(RF_ClassDefaultObject))
 	{
 		SetFlags(RF_StrongRefOnFrame);
+
+#if DO_ENSURE
+		UE::BlueprintAsyncAction::Private::DebugRecordAsyncActionCreated(this);
+#endif // DO_ENSURE
+	}
+}
+
+void UBlueprintAsyncActionBase::BeginDestroy()
+{
+	Super::BeginDestroy();
+
+#if DO_ENSURE
+	if (!HasAnyFlags(RF_ClassDefaultObject))
+	{
+		UE::BlueprintAsyncAction::Private::DebugRecordAsyncActionDestroyed();
 	}
+#endif // DO_ENSURE
 }
 
 void UBlueprintAsyncActionBase::Activate()
