Don't use bForceInstanceCulling unless we are mobile gpu scene with the GPU skin vertex factory.

#Synced-CL 50498661

diff --git a/Engine/Source/Runtime/Engine/Private/SkeletalMeshSceneProxy.cpp b/Engine/Source/Runtime/Engine/Private/SkeletalMeshSceneProxy.cpp
--- a/Engine/Source/Runtime/Engine/Private/SkeletalMeshSceneProxy.cpp
+++ b/Engine/Source/Runtime/Engine/Private/SkeletalMeshSceneProxy.cpp
@@ -113,6 +113,7 @@ FSkeletalMeshSceneProxy::FSkeletalMeshSceneProxy(const FSkinnedMeshSceneProxyDes
 		,	bForceWireframe(InMeshDesc.bForceWireframe)
 		,	bCanHighlightSelectedSections(InMeshDesc.bCanHighlightSelectedSections)
 		,	bRenderStatic(InMeshDesc.bRenderStatic)
+		,	bMobileGPUScene(false)
 #if !(UE_BUILD_SHIPPING || UE_BUILD_TEST)
 		,	bDrawDebugSkeleton(InMeshDesc.ShouldDrawDebugSkeleton())
 #endif
@@ -343,14 +344,11 @@ FSkeletalMeshSceneProxy::FSkeletalMeshSceneProxy(const FSkinnedMeshSceneProxyDes
 	}
 
 	EnableGPUSceneSupportFlags();
-
-	if (UseGPUScene(GetScene().GetShaderPlatform(), FeatureLevel))
-	{
-		bVFRequiresPrimitiveUniformBuffer = PlatformGPUSceneUsesUniformBufferView(GetScene().GetShaderPlatform());
-	}
+	bMobileGPUScene = UseGPUScene(GetScene().GetShaderPlatform(), FeatureLevel) && PlatformGPUSceneUsesUniformBufferView(GetScene().GetShaderPlatform());
+	bVFRequiresPrimitiveUniformBuffer |= bMobileGPUScene;
 
 	if (IsAllowingApproximateOcclusionQueries())
-	{		
+	{
 		bAllowApproximateOcclusion = (bAllowApproximateOcclusion || bRenderStatic);
 	}
 
@@ -757,7 +755,7 @@ void FSkeletalMeshSceneProxy::GetStaticMeshBatch(
 	BatchElement.IndexBuffer = LODData.MultiSizeIndexContainer.GetIndexBuffer();
 	BatchElement.bFetchInstanceCountFromScene = bInstancedSkinnedMesh;
 	BatchElement.NumInstances = GetInstanceDataHeader().NumInstances;
-	BatchElement.bForceInstanceCulling = true;
+	BatchElement.bForceInstanceCulling = bMobileGPUScene && VertexFactory && !VertexFactory->GetType()->SupportsGPUSkinPassThrough();
 
 	if (bUseGpuLodSelection)
 	{
@@ -811,7 +809,7 @@ void FSkeletalMeshSceneProxy::CreateBaseMeshBatch(const FSceneView* View, const
 	BatchElement.NumInstances = GetInstanceDataHeader().NumInstances;
 	BatchElement.NumPrimitives = Section.NumTriangles;
 	BatchElement.PrimitiveUniformBuffer = GetUniformBuffer();
-	BatchElement.bForceInstanceCulling = true;
+	BatchElement.bForceInstanceCulling = bMobileGPUScene && Mesh.VertexFactory && !Mesh.VertexFactory->GetType()->SupportsGPUSkinPassThrough();
 }
 
 uint8 FSkeletalMeshSceneProxy::GetCurrentFirstLODIdx_Internal() const
diff --git a/Engine/Source/Runtime/Engine/Public/SkeletalMeshSceneProxy.h b/Engine/Source/Runtime/Engine/Public/SkeletalMeshSceneProxy.h
--- a/Engine/Source/Runtime/Engine/Public/SkeletalMeshSceneProxy.h
+++ b/Engine/Source/Runtime/Engine/Public/SkeletalMeshSceneProxy.h
@@ -180,6 +180,7 @@ protected:
 	uint8 bIsCPUSkinned : 1;
 	uint8 bCanHighlightSelectedSections : 1;
 	uint8 bRenderStatic:1;
+	uint8 bMobileGPUScene:1;
 #if RHI_RAYTRACING
 	uint8 bHasRayTracingRepresentation : 1;
 	uint8 bRayTraceStatic : 1;
