Attempted fix for potential Race in SoundWave::PostLoad

#Synced-CL 50506470

diff --git a/Engine/Source/Runtime/Engine/Private/SoundWave.cpp b/Engine/Source/Runtime/Engine/Private/SoundWave.cpp
--- a/Engine/Source/Runtime/Engine/Private/SoundWave.cpp
+++ b/Engine/Source/Runtime/Engine/Private/SoundWave.cpp
@@ -554,7 +554,7 @@ const TArrayView<uint8> FSoundWaveData::GetZerothChunkDataView() const
 
 bool FSoundWaveData::HasChunkSeekTable(int32 InChunkIndex) const
 {
-	const TIndirectArray<FStreamedAudioChunk>& Chunks = RunningPlatformData.Chunks;
+	const TIndirectArray<FStreamedAudioChunk>& Chunks = RunningPlatformData.GetChunks();
 	if (Chunks.IsValidIndex(InChunkIndex))
 	{
 		return Chunks[InChunkIndex].SeekOffsetInAudioFrames != INDEX_NONE;
@@ -564,7 +564,7 @@ bool FSoundWaveData::HasChunkSeekTable(int32 InChunkIndex) const
 
 int32 FSoundWaveData::FindChunkIndexForSeeking(uint32 InTimeInAudioFrames) const
 {	
-	const TIndirectArray<FStreamedAudioChunk>& Chunks = RunningPlatformData.Chunks;
+	const TIndirectArray<FStreamedAudioChunk>& Chunks = RunningPlatformData.GetChunks();
 	int32 NumChunks = Chunks.Num();
 
 	// Find most in range chunk.
@@ -2300,14 +2300,18 @@ void USoundWave::PostLoad()
 
 		if (ShouldUseStreamCaching())
 		{
-			const TIndirectArray<FStreamedAudioChunk>& Chunks = SoundWaveDataPtr->RunningPlatformData.Chunks;
+#if !WITH_EDITOR
+			const TIndirectArray<FStreamedAudioChunk>& Chunks = SoundWaveDataPtr->RunningPlatformData.GetChunks();
 
 			// First chunk can be optionally inlined for the sake of pre-buffering.
 			const bool bIsFirstChunkInlined = Chunks.IsValidIndex(1) ? Chunks[1].BulkData.IsInlined() : false;
+#else
+			static constexpr bool bIsFirstChunkInlined = false; 
+#endif //WITH_EDITOR			
 
 			// Allowing retaining if we're NOT in the editor AND we've NOT disabled retaining (via a cvar).
 			const bool bAllowRetaining = !GIsEditor && !DisableRetainingCVar;
-
+			
 			// Retaining the "inlined" first chunk (which can be any load mode) is enabled via a CVAR.
 			const bool bShouldRetainInlinedFirstChunked = bIsFirstChunkInlined && !!bRetainInlinedFirstChunkCVar;
 
@@ -2316,17 +2320,15 @@ void USoundWave::PostLoad()
 			const bool bShouldRetainFirstChunk = bAllowRetaining &&
 				(ActualLoadingBehavior == ESoundWaveLoadingBehavior::RetainOnLoad || bShouldRetainInlinedFirstChunked);
 
-			UE_CLOG(Audio::MatchesLogFilter(*GetName()), LogAudioTiming, Verbose, TEXT("PostLoad: Wave='%s'\tFormat='%s'\tLooping=%s\tChunk1(Retain=%s,Inlined=%s,Size=%lld),LB=%s, NumChunks=%d"), 
+			UE_CLOG(Audio::MatchesLogFilter(*GetName()), LogAudioTiming, Verbose, TEXT("PostLoad: Wave='%s'\tFormat='%s'\tLooping=%s\tChunk1(Retain=%s,Inlined=%s),LB=%s"), 
 				*GetFName().ToString(),
 				*GetRuntimeFormat().ToString(),
 				ToCStr(LexToString(bLooping)),
 				ToCStr(LexToString(bShouldRetainFirstChunk)),
 				ToCStr(LexToString(bIsFirstChunkInlined)),
-				Chunks.IsValidIndex(1) ? Chunks[1].BulkData.GetElementCount() : 0,	
-				Audio::ShortEnumString(EnumToString(ActualLoadingBehavior)),
-				Chunks.Num()
+				Audio::ShortEnumString(EnumToString(ActualLoadingBehavior))
 			);
-			
+						
 			if (bShouldRetainFirstChunk)
 			{
 				RetainCompressedAudio(false);
