Dstm: Chaos vehicle support 2.4: Reinitialize animation instance in SkeletalMeshComponent during migration and add serialization for bone transforms on SkinnedMeshComponent

#Synced-CL 50453497

diff --git a/Engine/Source/Runtime/Engine/Private/Components/SkinnedMeshComponent.cpp b/Engine/Source/Runtime/Engine/Private/Components/SkinnedMeshComponent.cpp
--- a/Engine/Source/Runtime/Engine/Private/Components/SkinnedMeshComponent.cpp
+++ b/Engine/Source/Runtime/Engine/Private/Components/SkinnedMeshComponent.cpp
@@ -664,6 +664,18 @@ void USkinnedMeshComponent::Serialize(FArchive& Ar)
 	}
 
 	NotifyIfSkinnedAssetChanged();
+
+#if UE_WITH_REMOTE_OBJECT_HANDLE
+	if (Ar.IsMigratingRemoteObjects())
+	{
+		// TODO: Optimize by re-initializing anims (AllocateTransformData + InitAnim in SkeletalMeshComponent::PostMigrate) instead of serializing.
+		// Blocked by UE-363679 (server hang in PostMigrate).
+		Ar << ComponentSpaceTransformsArray[0];
+		Ar << ComponentSpaceTransformsArray[1];
+	}
+
+#endif
+
 }
 
 void USkinnedMeshComponent::PostLoad()
