Refine flag conformance to only the public flag for archetype objects, which will generated a GLEO if it is not set, but only when used by a child type

#rnx

#Synced-CL 50069009

diff --git a/Engine/Source/Runtime/Engine/Private/SimpleConstructionScript.cpp b/Engine/Source/Runtime/Engine/Private/SimpleConstructionScript.cpp
--- a/Engine/Source/Runtime/Engine/Private/SimpleConstructionScript.cpp
+++ b/Engine/Source/Runtime/Engine/Private/SimpleConstructionScript.cpp
@@ -29,6 +29,11 @@ namespace UE::Private
 {
 	void ConformArchetypeVisibility(UObject* Template)
 	{
+		if(!Template->HasAllFlags(RF_Public))
+		{
+			return;
+		}
+
 		TArray<UObject*> SubObjects;
 		GetObjectsWithOuter(
 			Template, 
@@ -37,19 +42,17 @@ namespace UE::Private
 			RF_Transient, 
 			EInternalObjectFlags::Garbage);
 
-		const EObjectFlags Flags = Template->GetMaskedFlags(RF_PropagateToSubObjects);
-
 		for (UObject* SubObject : SubObjects)
 		{
 			// In theory it might be good to propagate to all subobjects of
 			// the template object, but I'm being conservative for now:
 			if(SubObject->HasAnyFlags(RF_ArchetypeObject))
 			{
-				if(!SubObject->HasAllFlags(Flags))
+				if(!SubObject->HasAllFlags(RF_Public))
 				{
 					UE_LOGFMT(LogBlueprint, Warning, "Conforming subobject flags: {subobject}", SubObject->GetPathName());
 				}
-				SubObject->SetFlags(Flags);
+				SubObject->SetFlags(RF_Public);
 			}
 		}
 	}
