Fixed an issue where the override 'Use Half Precision' option in material editor wasn't working when the cvar r.Mobile.FloatPrecisionMode=2 (full precision everywhere).
Created a new option "Use Half-precision including shader parameters" in material editor to allow uniform buffers with custom parameters like PreshaderBuffer to be set to half  now also uses half precision

#Synced-CL 50531545

diff --git a/Engine/Source/Runtime/Core/Private/HAL/ConsoleManager.cpp b/Engine/Source/Runtime/Core/Private/HAL/ConsoleManager.cpp
--- a/Engine/Source/Runtime/Core/Private/HAL/ConsoleManager.cpp
+++ b/Engine/Source/Runtime/Core/Private/HAL/ConsoleManager.cpp
@@ -3862,9 +3862,9 @@ static TAutoConsoleVariable<int32> CVarMobileAllowDistanceFieldShadows(
 static TAutoConsoleVariable<int32> CVarMobileFloatPrecisionMode(
 	TEXT("r.Mobile.FloatPrecisionMode"),
 	0,
-	TEXT("0: Use Half-precision (default)\n"
-		 "1: Half precision, except Full precision for material expressions\n"
-		 "2: Force use of high precision in pixel shaders.\n"),
+	TEXT("0: Use Half-precision ALU & some uniform buffers (default)\n"
+		 "1: Half precision ALU, except Full precision for material expressions\n"
+		 "2: Force use of high precision ALU & uniform buffers in pixel shaders.\n"),
 	ECVF_ReadOnly);
 
 static TAutoConsoleVariable<bool> CVarMobileAllowDitheredLODTransition(
diff --git a/Engine/Source/Runtime/Core/Public/UObject/FortniteReleaseBranchCustomObjectVersions.inl b/Engine/Source/Runtime/Core/Public/UObject/FortniteReleaseBranchCustomObjectVersions.inl
--- a/Engine/Source/Runtime/Core/Public/UObject/FortniteReleaseBranchCustomObjectVersions.inl
+++ b/Engine/Source/Runtime/Core/Public/UObject/FortniteReleaseBranchCustomObjectVersions.inl
@@ -66,5 +66,8 @@ FORTNITE_RELEASE_VERSION(ChaosClothAddNewFictitiousForcesModes, 18)
 // Updated how blend spaces calculate the effective animation length of the sample point
 FORTNITE_RELEASE_VERSION(BlendSpaceUpdatedSamplePointAnimationLengthDerivation, 19)
 
+// Fix to make the override option to "Use Half Precision" work in Material Editor when the cvar r.Mobile.FloatPrecisionMode=2
+FORTNITE_RELEASE_VERSION(MaterialHalfPrecisionFixup, 20)
+
 // -----<new versions can be added above this line>-------------------------------------------------
 #undef FORTNITE_RELEASE_VERSION
\ No newline at end of file
diff --git a/Engine/Source/Runtime/Engine/Classes/Engine/EngineTypes.h b/Engine/Source/Runtime/Engine/Classes/Engine/EngineTypes.h
--- a/Engine/Source/Runtime/Engine/Classes/Engine/EngineTypes.h
+++ b/Engine/Source/Runtime/Engine/Classes/Engine/EngineTypes.h
@@ -271,8 +271,10 @@ enum EMaterialFloatPrecisionMode : int
 	MFPM_Full_MaterialExpressionOnly UMETA(DisplayName = "Use Full-precision for MaterialExpressions only"),
 	/** All the floats are full-precision */
 	MFPM_Full UMETA(DisplayName = "Use Full-precision for every float"),
-	/** Half precision, except explict 'float' in .ush/.usf*/
+	/** Half precision, except explicit 'float' in .ush/.usf*/
 	MFPM_Half UMETA(DisplayName = "Use Half-precision"),
+	/* Same as 'Half precision' but includes custom shader parameters uniforms */
+	MFPM_Half_WithShaderParameters UMETA(DisplayName = "Use Half-precision including shader parameters"),
 	MFPM_MAX,
 };
 
diff --git a/Engine/Source/Runtime/Engine/Private/Materials/Material.cpp b/Engine/Source/Runtime/Engine/Private/Materials/Material.cpp
--- a/Engine/Source/Runtime/Engine/Private/Materials/Material.cpp
+++ b/Engine/Source/Runtime/Engine/Private/Materials/Material.cpp
@@ -108,6 +108,11 @@
 #include "ThumbnailRendering/SceneThumbnailInfoWithPrimitive.h"
 #include "Editor.h"
 #include "ScopedTransaction.h"
+#include "Engine/RendererSettings.h"
+#include "Subsystems/AssetEditorSubsystem.h"
+#include "UObject/SavePackage.h"
+#include "Misc/MapErrors.h"
+#include "UObject/FortniteReleaseBranchCustomObjectVersion.h"
 #endif
 #include "ShaderCodeLibrary.h"
 #include "String/ParseTokens.h"
@@ -4470,7 +4475,7 @@ void UMaterial::PostLoad()
 		DecalBlendMode = DBM_Translucent;
 	}
 
-	if(bUseFullPrecision_DEPRECATED && FloatPrecisionMode == EMaterialFloatPrecisionMode::MFPM_Half)
+	if(bUseFullPrecision_DEPRECATED && (FloatPrecisionMode == EMaterialFloatPrecisionMode::MFPM_Half || FloatPrecisionMode == EMaterialFloatPrecisionMode::MFPM_Half_WithShaderParameters))
 	{
 		FloatPrecisionMode = EMaterialFloatPrecisionMode::MFPM_Full;
 		bUseFullPrecision_DEPRECATED = false;
@@ -4602,6 +4607,49 @@ void UMaterial::PostLoad()
 			PropagateDataToMaterialProxy();
 		}
 	}
+	
+	// Avoid running the following checks when running cooks
+	if (!IsRunningCookCommandlet())
+	{
+		if(FloatPrecisionMode == EMaterialFloatPrecisionMode::MFPM_Half && GetLinkerCustomVersion(FFortniteReleaseBranchCustomObjectVersion::GUID) < FFortniteReleaseBranchCustomObjectVersion::MaterialHalfPrecisionFixup)
+		{
+			static const auto CVar = IConsoleManager::Get().FindConsoleVariable(TEXT("r.Mobile.FloatPrecisionMode"));
+			if (CVar)
+			{
+				int CVarMobilePrecisionMode = FMath::Clamp(CVar->GetInt(), (int32)EMobileFloatPrecisionMode::Half, (int32)EMobileFloatPrecisionMode::Full);
+				if (CVarMobilePrecisionMode == EMobileFloatPrecisionMode::Full)
+				{
+					// Choosing 'Use Half Precision' on a per material property in the Material Editor was not working as intended when the cvar r.Mobile.FloatPrecisionMode=2 (it was a no-op)
+					// With this option now fixed, we revert this material's Float Precision option from Half --> Default to maintain the generated code the same and avoid performance/visual regressions
+					FloatPrecisionMode = EMaterialFloatPrecisionMode::MFPM_Default;
+
+					const FText FloatPrecisionModeSelectedOnMaterial =
+						FText::Format(LOCTEXT("MapCheck_Message_MaterialUseHalfPrecisionCheck",
+							"Material {0} had the option \"Use Half Precision\" from Float Precision Mode checked. This option was not working as intended before so it's being reverted to Default to maintain previous behaviour."),
+							FText::FromName(this->GetFName()));
+
+					FMessageLog("MapCheck").Warning()
+						->AddToken(FTextToken::Create(FloatPrecisionModeSelectedOnMaterial))
+						->AddToken(FActionToken::Create(LOCTEXT("MapCheck_SaveMaterialWithDefaultFloatingPointPrecision", "Please save the material to keep this change"),
+							LOCTEXT("MapCheck_SaveMaterialWithDefaultFloatingPointPrecision_Desc", "Reverts this material precision from Half -> Default to preserve previous behaviour and saves it"),
+							FOnActionTokenExecuted::CreateWeakLambda(this, [this]()
+								{
+									this->Modify();
+									this->MarkPackageDirty();
+									UPackage* Package = this->GetOutermost();
+
+									FSavePackageArgs SaveArgs;
+									SaveArgs.TopLevelFlags = RF_Public | RF_Standalone;
+
+									const FString PackageFilename = FPackageName::LongPackageNameToFilename(Package->GetName(), FPackageName::GetAssetPackageExtension());
+									UPackage::SavePackage(Package, this, *PackageFilename, SaveArgs);
+
+								})));
+				}
+			}
+		}
+	}
+
 #endif // WITH_EDITOR
 
 	checkf(CachedExpressionData, TEXT("Missing cached expression data for material, should have been either serialized or created during PostLoad"));
diff --git a/Engine/Source/Runtime/Engine/Private/Materials/MaterialShared.cpp b/Engine/Source/Runtime/Engine/Private/Materials/MaterialShared.cpp
--- a/Engine/Source/Runtime/Engine/Private/Materials/MaterialShared.cpp
+++ b/Engine/Source/Runtime/Engine/Private/Materials/MaterialShared.cpp
@@ -3685,7 +3685,18 @@ bool FMaterial::BeginCompileShaderMap(
 		FShaderCompileUtilities::ApplyDerivedDefines(*MaterialEnvironment, nullptr, (EShaderPlatform)GetShaderPlatform());
 
 		{
-			TArray<FShaderParametersMetadata*> UniformBufferStructs = NewCompilationOutput.UniformExpressionSet.CreateBufferStructs();
+			// Get this material's override floating point precision setting
+			EMaterialFloatPrecisionMode MaterialFloatPrecisionMode = EMaterialFloatPrecisionMode::MFPM_Default;
+			const UMaterialInterface* MaterialInterface = GetMaterialInterface();
+			if (MaterialInterface)
+			{
+				MaterialFloatPrecisionMode = MaterialInterface->GetMaterial()->FloatPrecisionMode;
+			}
+
+			// Sets the floating point precision to be used for the PreshaderBuffer uniform. This is controlled via this material's floating point override set in the material editor
+			EShaderPrecisionModifier::Type ParameterPrecision = (IsMobilePlatform((EShaderPlatform)GetShaderPlatform()) && MaterialFloatPrecisionMode == MFPM_Half_WithShaderParameters) ? EShaderPrecisionModifier::Half : EShaderPrecisionModifier::Float;
+			TArray<FShaderParametersMetadata*> UniformBufferStructs = NewCompilationOutput.UniformExpressionSet.CreateBufferStructs(ParameterPrecision);
+
 			SetupMaterialEnvironment(UniformBufferStructs, NewCompilationOutput.UniformExpressionSet, *MaterialEnvironment);
 			for (FShaderParametersMetadata* UniformBufferStruct : UniformBufferStructs)
 			{
diff --git a/Engine/Source/Runtime/Engine/Private/Materials/MaterialUniformExpressions.cpp b/Engine/Source/Runtime/Engine/Private/Materials/MaterialUniformExpressions.cpp
--- a/Engine/Source/Runtime/Engine/Private/Materials/MaterialUniformExpressions.cpp
+++ b/Engine/Source/Runtime/Engine/Private/Materials/MaterialUniformExpressions.cpp
@@ -521,7 +521,7 @@ static bool IsNumericMember(const FShaderParametersMetadata::FMember& Member)
 
 }
 
-TArray<FShaderParametersMetadata*> FUniformExpressionSet::CreateBufferStructs()
+TArray<FShaderParametersMetadata*> FUniformExpressionSet::CreateBufferStructs(EShaderPrecisionModifier::Type ParameterPrecision)
 {
 	// Make sure FUniformExpressionSet::CreateDebugLayout() is in sync
 	TArray<FShaderParametersMetadata::FMember> Members;
@@ -634,8 +634,13 @@ TArray<FShaderParametersMetadata*> FUniformExpressionSet::CreateBufferStructs()
 	}
 
 	if (UniformPreshaderBufferSize > 0u)
-	{
-		new(Members) FShaderParametersMetadata::FMember(TEXT("PreshaderBuffer"),TEXT(""),__LINE__,NextMemberOffset,UBMT_FLOAT32,EShaderPrecisionModifier::Float,1,4, UniformPreshaderBufferSize,NULL);
+	{	
+		// ParameterPrecision below will either be Half or Float based on a user's selection on this material's override floating point option in Material Editor
+		// "PreshaderBuffer" uses half precision when the option "Use Half-precision also on shader parameters" is selected
+		// Other options will make the uniform "PreshaderBuffer" use full float precision
+		// Using the material fp override on a per material case gives users fine grained control to choose between taking the cvar r.Mobile.FloatPrecisionMode value with "Default" or by explicitly choosing the level of half-precision they want 
+		new(Members) FShaderParametersMetadata::FMember(TEXT("PreshaderBuffer"), TEXT(""), __LINE__, NextMemberOffset, UBMT_FLOAT32, ParameterPrecision, 1, 4, UniformPreshaderBufferSize, NULL);
+		
 		NextMemberOffset += UniformPreshaderBufferSize * sizeof(FVector4f);
 		if (CompactUniformsVS && CompactUniformsVS->UniformPreshaderBufferSize > 0)
 		{
diff --git a/Engine/Source/Runtime/Engine/Private/ShaderCompiler/ShaderCompiler.cpp b/Engine/Source/Runtime/Engine/Private/ShaderCompiler/ShaderCompiler.cpp
--- a/Engine/Source/Runtime/Engine/Private/ShaderCompiler/ShaderCompiler.cpp
+++ b/Engine/Source/Runtime/Engine/Private/ShaderCompiler/ShaderCompiler.cpp
@@ -3654,7 +3654,8 @@ void GlobalBeginCompileShader(
 		Input.Environment.CompilerFlags.Add(CFLAG_NoFastMath);
 	}
     
-	if (bUsingMobileRenderer)
+	// At this point only apply the cvar r.Mobile.FloatPrecisionMode value to global shaders
+	if (bUsingMobileRenderer && (ShaderType == nullptr || ShaderType->GetGlobalShaderType() != nullptr))
     {
         static IConsoleVariable* CVar = IConsoleManager::Get().FindConsoleVariable(TEXT("r.Mobile.FloatPrecisionMode"));
         Input.Environment.FullPrecisionInPS |= CVar ? (CVar->GetInt() == EMobileFloatPrecisionMode::Full) : false;
diff --git a/Engine/Source/Runtime/Engine/Public/MaterialShared.h b/Engine/Source/Runtime/Engine/Public/MaterialShared.h
--- a/Engine/Source/Runtime/Engine/Public/MaterialShared.h
+++ b/Engine/Source/Runtime/Engine/Public/MaterialShared.h
@@ -699,7 +699,7 @@ public:
 	ENGINE_API FString GetSummaryString() const;
 
 #if WITH_EDITOR
-	TArray<FShaderParametersMetadata*> CreateBufferStructs();
+	TArray<FShaderParametersMetadata*> CreateBufferStructs(EShaderPrecisionModifier::Type ParameterPrecision);
 #endif
 
 	void SetParameterCollections(TConstArrayView<const class UMaterialParameterCollection*> Collections);
diff --git a/Engine/Source/Runtime/Engine/Public/Materials/Material.h b/Engine/Source/Runtime/Engine/Public/Materials/Material.h
--- a/Engine/Source/Runtime/Engine/Public/Materials/Material.h
+++ b/Engine/Source/Runtime/Engine/Public/Materials/Material.h
@@ -915,10 +915,14 @@ public:
 	uint8 bUseFullPrecision_DEPRECATED : 1;
 
 	/**
-	 *	How to use full (highp) precision in the pixel shader.
-	 *	highp is slower than the default (mediump) but can be used to work around precision-related rendering errors.
-	 *  Use 'Full-precision for MaterialExpressions only' if you still want to keep the precision of the halfs in .ush/.usf
-	 *	This setting has no effect on older mobile devices that do not support high precision.
+	 *	Per Material override to change floating point precision in the pixel shader (ALU and uniforms precision)
+	 *	Full-precision (highp) is slower than (half-precision/mediump) but can be used to work around precision-related rendering errors.
+	 *  Use:
+	 *  'Default' to take the precision set by the project wide cvar r.Mobile.FloatPrecisionMode. It defaults to 0 (half precision)
+	 *	'Full-precision for MaterialExpressions only' if you still want to keep half precision in.ush/.usf shaders (global shaders)
+	 *	'Full-precision for every float' to force every float and uniform to full-precision
+	 *	'Half-precision' to try to use half precision ALU and uniforms where appropriate
+	 *	'Half-precision including shader parameters' same behaviour as 'Half-precision' but it also uses half precision for uniforms with custom parameters like PreshaderBuffer uniform
 	 */
 	UPROPERTY(EditAnywhere, BlueprintReadOnly, Category = Mobile)
 	TEnumAsByte<EMaterialFloatPrecisionMode> FloatPrecisionMode;
