Add nullptr check on vertex factory before checking it in check

#Synced-CL 50524817

diff --git a/Engine/Source/Runtime/Engine/Private/Materials/MaterialShader.cpp b/Engine/Source/Runtime/Engine/Private/Materials/MaterialShader.cpp
--- a/Engine/Source/Runtime/Engine/Private/Materials/MaterialShader.cpp
+++ b/Engine/Source/Runtime/Engine/Private/Materials/MaterialShader.cpp
@@ -3437,7 +3437,7 @@ FPSOPrecacheDataArray FMaterialShaderMap::CollectPSOPrecacheData(const FMaterial
 	TRACE_CPUPROFILER_EVENT_SCOPE(FMaterialShaderMap::CollectPSOPrecacheData);
 
 	// Shouldn't get here if the type doesn't support precaching
-	check(PrecacheParams.VertexFactoryData.VertexFactoryType->SupportsPSOPrecaching());
+	check(PrecacheParams.VertexFactoryData.VertexFactoryType == nullptr || PrecacheParams.VertexFactoryData.VertexFactoryType->SupportsPSOPrecaching());
 
 #if PSO_PRECACHING_VALIDATE
 	ConditionalBreakOnPSOPrecacheMaterial(*PrecacheParams.Material, INDEX_NONE);
