Remove dead code from FViewport base virtual functions
- Only two types that inherit from FViewport exist now: FSceneViewport and FDummyViewport. The former is used for all scene rendering, the latter is used for the high resolution screenshot and other off-screen rendering.
- The following functions were never called because they were always overridden in both those derived types: FViewport::BeginRenderFrame, FViewport::EndRenderFrame
- FViewport::UpdateViewportRHI was protected and never called.

#Synced-CL 50397505

diff --git a/Engine/Source/Runtime/Engine/Private/Slate/SceneViewport.cpp b/Engine/Source/Runtime/Engine/Private/Slate/SceneViewport.cpp
--- a/Engine/Source/Runtime/Engine/Private/Slate/SceneViewport.cpp
+++ b/Engine/Source/Runtime/Engine/Private/Slate/SceneViewport.cpp
@@ -1943,6 +1943,15 @@ void FSceneViewport::SetRenderTargetTextureRenderThread(FTextureRHIRef& RT)
 	}
 }
 
+void FSceneViewport::SetInitialSize(FIntPoint InitialSizeXY)
+{
+	// Initial size only works if the viewport has not yet been resized
+	if (GetSizeXY() == FIntPoint::ZeroValue)
+	{
+		UpdateViewportRHI(false, InitialSizeXY.X, InitialSizeXY.Y, EWindowMode::Windowed, PF_Unknown);
+	}
+}
+
 void FSceneViewport::UpdateViewportRHI(bool bDestroyed, uint32 NewSizeX, uint32 NewSizeY, EWindowMode::Type NewWindowMode, EPixelFormat PreferredPixelFormat)
 {
 	{
diff --git a/Engine/Source/Runtime/Engine/Private/UnrealClient.cpp b/Engine/Source/Runtime/Engine/Private/UnrealClient.cpp
--- a/Engine/Source/Runtime/Engine/Private/UnrealClient.cpp
+++ b/Engine/Source/Runtime/Engine/Private/UnrealClient.cpp
@@ -1595,10 +1595,12 @@ void FViewport::HighResScreenshot()
 	const FString CachedScreenshotName = FScreenshotRequest::GetFilename();
 
 	FDummyViewport* DummyViewport = new FDummyViewport(ViewportClient);
-	DummyViewport->SetupHDR(GetDisplayColorGamut(), GetDisplayOutputFormat(), GetSceneHDREnabled());
+	DummyViewport->SetInitialSize(FIntPoint(
+		(GScreenshotResolutionX > 0) ? GScreenshotResolutionX : SizeX,
+		(GScreenshotResolutionY > 0) ? GScreenshotResolutionY : SizeY
+	));
 
-	DummyViewport->SizeX = (GScreenshotResolutionX > 0) ? GScreenshotResolutionX : SizeX;
-	DummyViewport->SizeY = (GScreenshotResolutionY > 0) ? GScreenshotResolutionY : SizeY;
+	DummyViewport->SetupHDR(GetDisplayColorGamut(), GetDisplayOutputFormat(), GetSceneHDREnabled());
 
 	BeginInitResource(DummyViewport);
 
@@ -1709,35 +1711,6 @@ void FViewport::HighResScreenshot()
 	HighResScreenshotBeginFrame(nullptr);
 }
 
-/** Starts a new rendering frame. Called from the rendering thread. */
-void FViewport::BeginRenderFrame(FRHICommandListImmediate& RHICmdList)
-{
-	check( IsInRenderingThread() );
-	UpdateRenderTargetSurfaceRHIToCurrentBackBuffer();
-}
-
-/**
- *	Ends a rendering frame. Called from the rendering thread.
- *	@param bPresent		Whether the frame should be presented to the screen
- *	@param bLockToVsync	Whether the GPU should block until VSYNC before presenting
- */
-void FViewport::EndRenderFrame(FRHICommandListImmediate& RHICmdList, bool bPresent, bool bLockToVsync)
-{
-	check( IsInRenderingThread() );
-
-	GetViewportRHI()->SetPresentLatencyCallbacks(
-		[CurrentFrameCounter = GFrameCounterRenderThread]()
-	{
-		UEngine::SetPresentLatencyMarkerStart(CurrentFrameCounter);
-	},
-	[CurrentFrameCounter = GFrameCounterRenderThread]()
-	{
-		UEngine::SetPresentLatencyMarkerEnd(CurrentFrameCounter);
-	});
-
-	RHICmdList.EndDrawingViewport(GetViewportRHI(), FRHIPresentArgs(/*.FrameCounter*/ GFrameCounterRenderThread, /*.bPresent*/ bPresent, /*.bLockToVsync*/ bLockToVsync));
-}
-
 FRHIGPUMask FViewport::GetGPUMask(FRHICommandListImmediate& RHICmdList) const
 {
 	return FRHIGPUMask::FromIndex(RHIGetViewportNextPresentGPUIndex(GetViewportRHI()));
@@ -1754,7 +1727,6 @@ void FViewport::EnqueueBeginRenderFrame(const bool bShouldPresent)
 	});
 }
 
-
 void FViewport::EnqueueEndRenderFrame(const bool bLockToVsync, const bool bShouldPresent)
 {
 	bool bPresent = (PresentAndStopMovieDelay <= 0)
@@ -2184,66 +2156,6 @@ void FViewport::EnumerateHitProxiesInRect(FIntRect InRect, TFunctionRef<bool(HHi
 	}
 }
 
-void FViewport::UpdateViewportRHI(bool bDestroyed, uint32 NewSizeX, uint32 NewSizeY, EWindowMode::Type NewWindowMode, EPixelFormat PreferredPixelFormat)
-{
-	{
-		// Temporarily stop rendering thread.
-		FlushRenderingCommands();
-
-		// Update the viewport attributes.
-		// This is done AFTER the command flush done by UpdateViewportRHI, to avoid disrupting rendering thread accesses to the old viewport size.
-		SizeX = NewSizeX;
-		SizeY = NewSizeY;
-		WindowMode = NewWindowMode;
-
-		// Release the viewport's resources.
-		BeginReleaseResource(this);
-
-		// Don't reinitialize the viewport RHI if the viewport has been destroyed.
-		if(bDestroyed)
-		{
-			if(IsValidRef(ViewportRHI))
-			{
-				// If the viewport RHI has already been initialized, release it.
-				ViewportRHI.SafeRelease();
-			}
-		}
-		else
-		{
-			if(IsValidRef(ViewportRHI))
-			{
-				// If the viewport RHI has already been initialized, resize it.
-				RHIResizeViewport(
-					ViewportRHI,
-					SizeX,
-					SizeY,
-					IsFullscreen(),
-					PreferredPixelFormat
-					);
-			}
-			else
-			{
-				// Initialize the viewport RHI with the new viewport state.
-				ViewportRHI = RHICreateViewport(
-					GetWindow(),
-					SizeX,
-					SizeY,
-					IsFullscreen(),
-					EPixelFormat::PF_Unknown /* ie, use default format */);
-			}
-		
-			// Initialize the viewport's resources.
-			BeginInitResource(this);
-		}
-	}
-
-	if ( !bDestroyed )
-	{
-		// send a notification that the viewport has been resized
-		ViewportResizedEvent.Broadcast(this, 0);
-	}
-}
-
 FIntRect FViewport::CalculateViewExtents(float AspectRatio, float DesiredAspectRatio, const FIntRect& ViewRect, const FIntPoint& DestSize)
 {
 	FIntRect Result = ViewRect;
@@ -2447,16 +2359,6 @@ void FViewport::UpdateRenderTargetSurfaceRHIToCurrentBackBuffer()
 	}
 }
 
-void FViewport::SetInitialSize( FIntPoint InitialSizeXY )
-{
-	// Initial size only works if the viewport has not yet been resized
-	if( GetSizeXY() == FIntPoint::ZeroValue )
-	{
-		UpdateViewportRHI( false, InitialSizeXY.X, InitialSizeXY.Y, EWindowMode::Windowed, PF_Unknown );
-	}
-}
-
-
 ENGINE_API bool GetViewportScreenShot(FViewport* Viewport, TArray<FColor>& Bitmap, const FIntRect& ViewRect /*= FIntRect()*/, const FReadSurfaceDataFlags& ReadSurfaceFlags /*= FReadSurfaceDataFlags()*/)
 {
 	// Read the contents of the viewport into an array.
diff --git a/Engine/Source/Runtime/Engine/Public/DummyViewport.h b/Engine/Source/Runtime/Engine/Public/DummyViewport.h
--- a/Engine/Source/Runtime/Engine/Public/DummyViewport.h
+++ b/Engine/Source/Runtime/Engine/Public/DummyViewport.h
@@ -34,6 +34,15 @@ public:
 		bSceneHDREnabled = bInSceneHDREnabled;
 	}
 
+	virtual void SetInitialSize(FIntPoint InitialSizeXY) override
+	{
+	    if (GetSizeXY() == FIntPoint::ZeroValue)
+	    {
+			SizeX = InitialSizeXY.X;
+			SizeY = InitialSizeXY.Y;
+	    }
+	}
+
 	virtual void*	GetWindow() override { return 0; }
 	virtual void	MoveWindow(int32 NewPosX, int32 NewPosY, int32 NewSizeX, int32 NewSizeY) override {}
 	virtual void	Destroy() override {}
diff --git a/Engine/Source/Runtime/Engine/Public/Slate/SceneViewport.h b/Engine/Source/Runtime/Engine/Public/Slate/SceneViewport.h
--- a/Engine/Source/Runtime/Engine/Public/Slate/SceneViewport.h
+++ b/Engine/Source/Runtime/Engine/Public/Slate/SceneViewport.h
@@ -41,6 +41,7 @@ public:
 
 	/** FViewport interface */
 	virtual void MoveWindow(int32 NewPosX, int32 NewPosY, int32 NewSizeX, int32 NewSizeY) override {}
+	ENGINE_API virtual void SetInitialSize(FIntPoint InitialSizeXY) override;
 	ENGINE_API virtual bool HasMouseCapture() const override;
 	ENGINE_API virtual bool HasFocus() const override;
 	ENGINE_API virtual bool IsForegroundWindow() const override;
@@ -220,7 +221,7 @@ public:
 	}
 
 	/** Updates the viewport RHI with a new size and fullscreen flag */
-	ENGINE_API virtual void UpdateViewportRHI(bool bDestroyed, uint32 NewSizeX, uint32 NewSizeY, EWindowMode::Type NewWindowMode, EPixelFormat PreferredPixelFormat) override;
+	ENGINE_API void UpdateViewportRHI(bool bDestroyed, uint32 NewSizeX, uint32 NewSizeY, EWindowMode::Type NewWindowMode, EPixelFormat PreferredPixelFormat);
 
 	/** ISlateViewport interface */
 	ENGINE_API virtual bool IsCameraCut() const override;
diff --git a/Engine/Source/Runtime/Engine/Public/UnrealClient.h b/Engine/Source/Runtime/Engine/Public/UnrealClient.h
--- a/Engine/Source/Runtime/Engine/Public/UnrealClient.h
+++ b/Engine/Source/Runtime/Engine/Public/UnrealClient.h
@@ -539,14 +539,14 @@ public:
 	/**
 	 *	Starts a new rendering frame. Called from the rendering thread.
 	 */
-	ENGINE_API virtual void	BeginRenderFrame(FRHICommandListImmediate& RHICmdList);
+	virtual void BeginRenderFrame(FRHICommandListImmediate& RHICmdList) = 0;
 
 	/**
 	 *	Ends a rendering frame. Called from the rendering thread.
 	 *	@param bPresent		Whether the frame should be presented to the screen
 	 *	@param bLockToVsync	Whether the GPU should block until VSYNC before presenting
 	 */
-	ENGINE_API virtual void	EndRenderFrame(FRHICommandListImmediate& RHICmdList, bool bPresent, bool bLockToVsync);
+	virtual void EndRenderFrame(FRHICommandListImmediate& RHICmdList, bool bPresent, bool bLockToVsync) = 0;
 
 	/**
 	 * Returns the GPU nodes on which to render this viewport.
@@ -730,7 +730,7 @@ public:
 	 *
 	 * @param InitialSizeXY	The initial size of the viewport
 	 */
-	ENGINE_API void SetInitialSize( FIntPoint InitialSizeXY );
+	virtual void SetInitialSize( FIntPoint InitialSizeXY ) = 0;
 
 	/** Returns true if the viewport is for play in editor */
 	bool IsPlayInEditorViewport() const
@@ -774,12 +774,6 @@ protected:
 	/** The viewport's client. */
 	FViewportClient* ViewportClient;
 
-	/**
-	 * Updates the viewport RHI with the current viewport state.
-	 * @param bDestroyed - True if the viewport has been destroyed.
-	 */
-	ENGINE_API virtual void UpdateViewportRHI(bool bDestroyed, uint32 NewSizeX, uint32 NewSizeY, EWindowMode::Type NewWindowMode, EPixelFormat PreferredPixelFormat);
-
 	/**
 	 * Take a high-resolution screenshot and save to disk.
 	 */
