Fix typo in assignation. The old code was behaving as expected since the expressions separated by a comma were already evaluating in the right order and there was no dependencies.

#Synced-CL 0

diff --git a/Engine/Source/Runtime/Engine/Private/LevelInstance/LevelInstanceSubsystem.cpp b/Engine/Source/Runtime/Engine/Private/LevelInstance/LevelInstanceSubsystem.cpp
--- a/Engine/Source/Runtime/Engine/Private/LevelInstance/LevelInstanceSubsystem.cpp
+++ b/Engine/Source/Runtime/Engine/Private/LevelInstance/LevelInstanceSubsystem.cpp
@@ -1440,7 +1440,7 @@ ILevelInstanceInterface* ULevelInstanceSubsystem::CreateLevelInstanceFrom(const
 		CreateNewStreamingLevelParams.bCreateWorldPartition = GetWorld()->IsPartitionedWorld();
 		CreateNewStreamingLevelParams.bEnableWorldPartitionStreaming = CreationParams.bEnableStreaming;
 		CreateNewStreamingLevelParams.ActorsToMove = &ActorsToMove;
-		CreateNewStreamingLevelParams.TemplateWorld = CreationParams.TemplateWorld,
+		CreateNewStreamingLevelParams.TemplateWorld = CreationParams.TemplateWorld;
 		CreateNewStreamingLevelParams.LevelStreamingCreatedCallback = [LevelInstanceActorGuid, ParentContainerInstance](ULevelStreaming* InLevelStreaming)
 		{
 			ULevelStreamingLevelInstanceEditor* LevelInstanceLevelStreaming = CastChecked<ULevelStreamingLevelInstanceEditor>(InLevelStreaming);
