FViewport::GetHitProxy shouldn't need to do anything if the viewport size is (0, 0)

#Synced-CL 50544822

diff --git a/Engine/Source/Runtime/Engine/Private/UnrealClient.cpp b/Engine/Source/Runtime/Engine/Private/UnrealClient.cpp
--- a/Engine/Source/Runtime/Engine/Private/UnrealClient.cpp
+++ b/Engine/Source/Runtime/Engine/Private/UnrealClient.cpp
@@ -2029,15 +2029,20 @@ void FViewport::GetHitProxyMap(FIntRect InRect,TArray<HHitProxy*>& OutMap)
 	}
 }
 
-HHitProxy* FViewport::GetHitProxy(int32 X,int32 Y)
+HHitProxy* FViewport::GetHitProxy(int32 X, int32 Y)
 {
+	const FIntPoint VPSize = GetSizeXY();
+	// if we haven't been initialized, we can't really hit anything.
+	if (VPSize == FIntPoint::ZeroValue)
+	{
+		return nullptr;
+	}
+
 	// Compute a HitProxySize x HitProxySize test region with the center at (X,Y).
 	int32	MinX = X - HitProxySize,
 			MinY = Y - HitProxySize,
 			MaxX = X + HitProxySize,
 			MaxY = Y + HitProxySize;
-
-	FIntPoint VPSize = GetSizeXY();
 	
 	// Clip the region to the viewport bounds.
 	MinX = FMath::Clamp(MinX, 0, VPSize.X - 1);
