Fix issue where editor could become unresponsive while compiling debug viewmode shaders.

#Synced-CL 50450100

diff --git a/Engine/Source/Runtime/Engine/Private/DebugViewModeHelpers.cpp b/Engine/Source/Runtime/Engine/Private/DebugViewModeHelpers.cpp
--- a/Engine/Source/Runtime/Engine/Private/DebugViewModeHelpers.cpp
+++ b/Engine/Source/Runtime/Engine/Private/DebugViewModeHelpers.cpp
@@ -377,6 +377,7 @@ bool CompileDebugViewModeShaders(EDebugViewShaderMode ShaderMode, EMaterialQuali
 		}
 
 		// Are we asked to cancel the operation?
+		ProgressTask->TickProgress();
 		if (GWarn->ReceivedUserCancel())
 		{
 			bAllMaterialsCompiledSuccessfully = false;
