SceneViewport pipelining fixes
- Keep a seperate "buffer index" for Game Thread vs Render Thread. This index is updated on the RT via a render command.
- Collapsing the three seperate "render target" arrays into a single array to simplify updates.
- Doing a number of minor code cleanups

#Synced-CL 50544822

diff --git a/Engine/Source/Runtime/Engine/Private/Slate/SceneViewport.cpp b/Engine/Source/Runtime/Engine/Private/Slate/SceneViewport.cpp
--- a/Engine/Source/Runtime/Engine/Private/Slate/SceneViewport.cpp
+++ b/Engine/Source/Runtime/Engine/Private/Slate/SceneViewport.cpp
@@ -59,8 +59,6 @@ FSceneViewport::FSceneViewport( FViewportClient* InViewportClient, TSharedPtr<SV
 	, MousePosBeforeHiddenDueToCapture( -1, -1 )
 	, RTTSize( 0, 0 )
 	, SceneTargetFormat( EPixelFormat::PF_A2B10G10R10 )
-	, CurrentBufferedTargetIndex(0)
-	, NextBufferedTargetIndex(0)
 	, NumTouches(0)
 	, DisplayColorGamut(EDisplayColorGamut::sRGB_D65)
 	, DisplayOutputFormat(EDisplayOutputFormat::SDR_sRGB)
@@ -1446,7 +1444,7 @@ FSlateShaderResource* FSceneViewport::GetViewportRenderTargetTexture() const
 	{
 		return RenderThreadSlateTexture;
 	}
-	return (BufferedSlateHandles.Num() != 0) ? BufferedSlateHandles[CurrentBufferedTargetIndex] : nullptr;
+	return NumBufferedFrames ? BufferedFrames[BufferIndex_GT].SlateHandle : nullptr;
 }
 
 bool FSceneViewport::IsStereoscopic3D() const
@@ -1929,13 +1927,13 @@ const FTextureRHIRef& FSceneViewport::GetRenderTargetTexture() const
 	return 	RenderTargetTextureRHI;
 }
 
-void FSceneViewport::SetRenderTargetTextureRenderThread(FTextureRHIRef& RT)
+void FSceneViewport::SetRenderTargetTextureRenderThread(FRHITexture* InRenderTargetTexture)
 {
 	check(IsInRenderingThread());
-	RenderTargetTextureRenderThreadRHI = RT;
-	if (RT.IsValid())
+	RenderTargetTextureRenderThreadRHI = InRenderTargetTexture;
+	if (InRenderTargetTexture)
 	{
-		RenderThreadSlateTexture->SetRHIRef(RenderTargetTextureRenderThreadRHI, RT->GetSizeX(), RT->GetSizeY());
+		RenderThreadSlateTexture->SetRHIRef(InRenderTargetTexture, InRenderTargetTexture->GetSizeX(), InRenderTargetTexture->GetSizeY());
 	}
 	else
 	{
@@ -1954,126 +1952,131 @@ void FSceneViewport::SetInitialSize(FIntPoint InitialSizeXY)
 
 void FSceneViewport::UpdateViewportRHI(bool bDestroyed, uint32 NewSizeX, uint32 NewSizeY, EWindowMode::Type NewWindowMode, EPixelFormat PreferredPixelFormat)
 {
+	check(IsInGameThread());
+
+	// Release the viewport's resources.
+	if (DebugCanvasDrawer)
 	{
-		check(IsInGameThread());
+		DebugCanvasDrawer->ReleaseInternalTexture();
+	}
+	BeginReleaseResource(this);
+	FlushRenderingCommands();
 
-		// Release the viewport's resources.
-		if (DebugCanvasDrawer)
-		{
-			DebugCanvasDrawer->ReleaseInternalTexture();
-		}
-		BeginReleaseResource(this);
-		FlushRenderingCommands();
+	// Update the viewport attributes.
+	// This is done AFTER the command flush done by UpdateViewportRHI, to avoid disrupting rendering thread accesses to the old viewport size.
+	SizeX = NewSizeX;
+	SizeY = NewSizeY;
+	WindowMode = NewWindowMode;
 
-		// Update the viewport attributes.
-		// This is done AFTER the command flush done by UpdateViewportRHI, to avoid disrupting rendering thread accesses to the old viewport size.
-		SizeX = NewSizeX;
-		SizeY = NewSizeY;
-		WindowMode = NewWindowMode;
+	if (!bDestroyed)
+	{
+		TSharedPtr<SWidget> PinnedViewport = ViewportWidget.Pin();
+		void* OSWindow = nullptr;
+		FVector2D WindowTopLeft = FVector2D(0.0f, 0.0f);
+		FVector2D WindowBottomRight = FVector2D(0.0f, 0.0f);
 
-		if( !bDestroyed )
+		if (PinnedViewport.IsValid())
 		{
-			TSharedPtr<SWidget> PinnedViewport = ViewportWidget.Pin();
-			void* OSWindow = nullptr;
-			FVector2D WindowTopLeft = FVector2D(0.0f, 0.0f);
-			FVector2D WindowBottomRight = FVector2D(0.0f, 0.0f);
-			if (PinnedViewport.IsValid())
+			TSharedPtr<SWindow> Window = FSlateApplication::Get().FindWidgetWindow(PinnedViewport.ToSharedRef());
+			if (Window.IsValid())
 			{
-				TSharedPtr<SWindow> Window = FSlateApplication::Get().FindWidgetWindow(PinnedViewport.ToSharedRef());
-				if (Window.IsValid())
-				{
-					OSWindow = Window->GetNativeWindow()->GetOSWindowHandle();
-					WindowTopLeft = Window->GetPositionInScreen();
-					WindowBottomRight = Window->GetPositionInScreen() + Window->GetSizeInScreen();
-				}
+				OSWindow = Window->GetNativeWindow()->GetOSWindowHandle();
+				WindowTopLeft = Window->GetPositionInScreen();
+				WindowBottomRight = Window->GetPositionInScreen() + Window->GetSizeInScreen();
 			}
-			FHDRMetaData OutHDRMetaData;
-			ComputeSceneViewportHDRMetaData(OutHDRMetaData, WindowTopLeft, WindowBottomRight, OSWindow, IsStereoRenderingAllowed());
-			DisplayOutputFormat = OutHDRMetaData.DisplayOutputFormat;
-			DisplayColorGamut = OutHDRMetaData.DisplayColorGamut;
-			MinimumLuminanceInNits = OutHDRMetaData.MinimumLuminanceInNits;
-			MaximumLuminanceInNits = OutHDRMetaData.MaximumLuminanceInNits;
-			MaximumFullFrameLuminanceInNits = OutHDRMetaData.MaximumFullFrameLuminanceInNits;
-			HDRPaperWhiteInNits = OutHDRMetaData.HDRPaperWhiteInNits;
-			LimitingColorSpace = OutHDRMetaData.LimitingColorSpace;
-			bHDRViewport = OutHDRMetaData.bHDRSupported;
-
-			BeginInitResource(this);
-			
-			FRenderCommandFence InitResourceFence;
-			InitResourceFence.BeginFence();
-			InitResourceFence.Wait();
+		}
 
-			FSlateRenderer* Renderer = FSlateApplication::Get().GetRenderer();
+		FHDRMetaData OutHDRMetaData;
+		ComputeSceneViewportHDRMetaData(OutHDRMetaData, WindowTopLeft, WindowBottomRight, OSWindow, IsStereoRenderingAllowed());
+		DisplayOutputFormat = OutHDRMetaData.DisplayOutputFormat;
+		DisplayColorGamut = OutHDRMetaData.DisplayColorGamut;
+		MinimumLuminanceInNits = OutHDRMetaData.MinimumLuminanceInNits;
+		MaximumLuminanceInNits = OutHDRMetaData.MaximumLuminanceInNits;
+		MaximumFullFrameLuminanceInNits = OutHDRMetaData.MaximumFullFrameLuminanceInNits;
+		HDRPaperWhiteInNits = OutHDRMetaData.HDRPaperWhiteInNits;
+		LimitingColorSpace = OutHDRMetaData.LimitingColorSpace;
+		bHDRViewport = OutHDRMetaData.bHDRSupported;
 
-			if (PinnedViewport.IsValid())
-			{
-				TSharedPtr<SWindow> Window = FSlateApplication::Get().FindWidgetWindow(PinnedViewport.ToSharedRef());
+		BeginInitResource(this);
 
-				if (UseSeparateRenderTarget())
-				{
-					uint32 TexSizeX = SizeX, TexSizeY = SizeY;
-					{
-						IStereoRenderTargetManager* const StereoRenderTargetManager = RetrieveStereoRenderTargetManager(IsStereoRenderingAllowed());
-						if (StereoRenderTargetManager)
-						{
-							StereoRenderTargetManager->CalculateRenderTargetSize(*this, TexSizeX, TexSizeY);
-						}
-					}
-					RTTSize = FIntPoint(TexSizeX, TexSizeY);
-				}
-			}
+		FRenderCommandFence InitResourceFence;
+		InitResourceFence.BeginFence();
+		InitResourceFence.Wait();
+
+		FSlateRenderer* Renderer = FSlateApplication::Get().GetRenderer();
+
+		if (PinnedViewport.IsValid())
+		{
+			TSharedPtr<SWindow> Window = FSlateApplication::Get().FindWidgetWindow(PinnedViewport.ToSharedRef());
 
-			if( !UseSeparateRenderTarget() )
+			if (UseSeparateRenderTarget())
 			{
-				TSharedPtr<SWindow> Window = FSlateApplication::Get().FindWidgetWindow(ViewportWidget.Pin().ToSharedRef());
-				void* ViewportResource = Renderer->GetViewportResource(*Window);
-				if( ViewportResource )
+				uint32 TexSizeX = SizeX, TexSizeY = SizeY;
+				if (IStereoRenderTargetManager* const StereoRenderTargetManager = RetrieveStereoRenderTargetManager(IsStereoRenderingAllowed()))
 				{
-					ViewportRHI = *((FViewportRHIRef*)ViewportResource);
+					StereoRenderTargetManager->CalculateRenderTargetSize(*this, TexSizeX, TexSizeY);
 				}
-				Renderer->UpdateFullscreenState(Window.ToSharedRef(), NewSizeX, NewSizeY);
+				RTTSize = FIntPoint(TexSizeX, TexSizeY);
 			}
+		}
 
-			ViewportResizedEvent.Broadcast(this, 0);
+		if (!UseSeparateRenderTarget())
+		{
+			TSharedPtr<SWindow> Window = FSlateApplication::Get().FindWidgetWindow(ViewportWidget.Pin().ToSharedRef());
+			if (void* ViewportResource = Renderer->GetViewportResource(*Window))
+			{
+				ViewportRHI = *((FViewportRHIRef*)ViewportResource);
+			}
+			Renderer->UpdateFullscreenState(Window.ToSharedRef(), NewSizeX, NewSizeY);
 		}
-		else
+
+		ViewportResizedEvent.Broadcast(this, 0);
+	}
+	else
+	{
+		TArray<FSlateRenderTargetRHI*> SlateHandlesToDelete;
+
+		SlateHandlesToDelete.Emplace(RenderThreadSlateTexture);
+		RenderThreadSlateTexture = nullptr;
+
+		for (FFrameData& FrameData : BufferedFrames)
 		{
-			// Enqueue a render command to delete the handle.  It must be deleted on the render thread after the resource is released
-			FSlateRenderTargetRHI* RenderThreadSlateTexturePtr = RenderThreadSlateTexture;
-			RenderThreadSlateTexture = nullptr;
+			SlateHandlesToDelete.Emplace(FrameData.SlateHandle);
+			FrameData.SlateHandle = nullptr;
+		}
 
-			ENQUEUE_RENDER_COMMAND(DeleteSlateRenderTarget)(
-				[BufferedSlateHandles = MoveTemp(BufferedSlateHandles), RenderThreadSlateTexturePtr](FRHICommandListImmediate& RHICmdList)
+		// Enqueue a render command to delete the slate handles. They must be deleted on the render thread after the resource is released.
+		ENQUEUE_RENDER_COMMAND(DeleteSlateRenderTarget)(
+			[SlateHandlesToDelete = MoveTemp(SlateHandlesToDelete)](FRHICommandListImmediate& RHICmdList)
+			{
+				for (FSlateRenderTargetRHI* SlateHandle : SlateHandlesToDelete)
 				{
-					for (int32 i = 0; i < BufferedSlateHandles.Num(); ++i)
-					{
-						delete BufferedSlateHandles[i];
-					}
+					delete SlateHandle;
+				}
+			});
 
-					delete RenderThreadSlateTexturePtr;
-				});
-		}
 	}
 }
 
 void FSceneViewport::EnqueueBeginRenderFrame(const bool bShouldPresent)
 {
-	check( IsInGameThread() );
+	check(IsInGameThread());
+	checkf(SizeX > 0 && SizeY > 0, TEXT("FSceneViewport::EnqueueBeginRenderFrame called before calling FSceneViewport::ResizeViewport, this is unsafe and will cause crashes."));
+
 	const bool bStereoRenderingAvailable = GEngine->StereoRenderingDevice.IsValid() && IsStereoRenderingAllowed();
 	const bool bStereoRenderingEnabled = bStereoRenderingAvailable && GEngine->StereoRenderingDevice->IsStereoEnabled();
 
-	IStereoRenderTargetManager* StereoRenderTargetManager = bStereoRenderingAvailable ? GEngine->StereoRenderingDevice->GetRenderTargetManager(): nullptr;
+	IStereoRenderTargetManager* StereoRenderTargetManager = bStereoRenderingAvailable ? GEngine->StereoRenderingDevice->GetRenderTargetManager() : nullptr;
 
-	CurrentBufferedTargetIndex = NextBufferedTargetIndex;
-	NextBufferedTargetIndex = (CurrentBufferedTargetIndex + 1) % BufferedSlateHandles.Num();
+	BufferIndex_GT = NextBufferIndex_GT;
+	NextBufferIndex_GT = NumBufferedFrames ? (BufferIndex_GT + 1) % NumBufferedFrames : 0;
 
 	// check if we need to reallocate rendertarget for HMD and update HMD rendering viewport 
 	if (bStereoRenderingAvailable)
 	{
-		bool bHMDWantsSeparateRenderTarget = StereoRenderTargetManager != nullptr ? StereoRenderTargetManager->ShouldUseSeparateRenderTarget() : false;
+		bool bHMDWantsSeparateRenderTarget = StereoRenderTargetManager ? StereoRenderTargetManager->ShouldUseSeparateRenderTarget() : false;
 		if (bHMDWantsSeparateRenderTarget != bForceSeparateRenderTarget ||
-		    (bHMDWantsSeparateRenderTarget && StereoRenderTargetManager->NeedReAllocateViewportRenderTarget(*this)))
+			(bHMDWantsSeparateRenderTarget && StereoRenderTargetManager->NeedReAllocateViewportRenderTarget(*this)))
 		{
 			// This will cause RT to be allocated (or freed)
 			bForceSeparateRenderTarget = bHMDWantsSeparateRenderTarget;
@@ -2084,15 +2087,15 @@ void FSceneViewport::EnqueueBeginRenderFrame(const bool bShouldPresent)
 		{
 			// We need to acquire a buffered texture from either the new RT or the existing one
 			int32 TextureIndex = StereoRenderTargetManager->AcquireColorTexture();
-			CurrentBufferedTargetIndex = TextureIndex < 0 ? CurrentBufferedTargetIndex : TextureIndex;
-			
+			BufferIndex_GT = TextureIndex < 0 ? BufferIndex_GT : TextureIndex;
+
 			StereoRenderTargetManager->AcquireDepthTexture();
 		}
 	}
 
-	if (BufferedRenderTargetsRHI[CurrentBufferedTargetIndex])
+	if (NumBufferedFrames && BufferedFrames[BufferIndex_GT].RenderTarget)
 	{
-		RenderTargetTextureRHI = BufferedRenderTargetsRHI[CurrentBufferedTargetIndex];
+		RenderTargetTextureRHI = BufferedFrames[BufferIndex_GT].RenderTarget;
 	}
 
 	DebugCanvasDrawer->InitDebugCanvas(GetClient(), GetClient()->GetWorld());
@@ -2100,19 +2103,17 @@ void FSceneViewport::EnqueueBeginRenderFrame(const bool bShouldPresent)
 	// Note: ViewportRHI is only updated on the game thread
 
 	// If we dont have the ViewportRHI then we need to get it before rendering
-	// Note, we need ViewportRHI even if UseSeparateRenderTarget() is true when stereo rendering
-	// is enabled.
+	// Note, we need ViewportRHI even if UseSeparateRenderTarget() is true when stereo rendering is enabled.
 	if (!IsValidRef(ViewportRHI) && (!UseSeparateRenderTarget() || bStereoRenderingEnabled))
 	{
 		// Get the viewport for this window from the renderer so we can render directly to the backbuffer
-		FSlateRenderer* Renderer = FSlateApplication::Get().GetRenderer();
 		if (ViewportWidget.IsValid())
 		{
-			auto WidgetWindow = FSlateApplication::Get().FindWidgetWindow(ViewportWidget.Pin().ToSharedRef());
+			TSharedPtr<SWindow> WidgetWindow = FSlateApplication::Get().FindWidgetWindow(ViewportWidget.Pin().ToSharedRef());
 			if (WidgetWindow.IsValid())
 			{
-				void* ViewportResource = Renderer->GetViewportResource(*WidgetWindow);
-				if (ViewportResource)
+				FSlateRenderer* Renderer = FSlateApplication::Get().GetRenderer();
+				if (void* ViewportResource = Renderer->GetViewportResource(*WidgetWindow))
 				{
 					ViewportRHI = *((FViewportRHIRef*)ViewportResource);
 				}
@@ -2120,18 +2121,18 @@ void FSceneViewport::EnqueueBeginRenderFrame(const bool bShouldPresent)
 		}
 	}
 
-
-	//set the rendertarget visible to the render thread
-	//must come before any render thread frame handling.
+	// set the render target visible to the render thread
+	// must come before any render thread frame handling.
 	ENQUEUE_RENDER_COMMAND(SetRenderThreadViewportTarget)(
-		[Viewport = this, RT = RenderTargetTextureRHI](FRHICommandListImmediate& RHICmdList) mutable
+		[Viewport = this, NewRenderTargetTexture = RenderTargetTextureRHI, NewTargetIndex = BufferIndex_GT](FRHICommandListImmediate& RHICmdList) mutable
 		{
-			Viewport->SetRenderTargetTextureRenderThread(RT);			
-		});		
+			Viewport->BufferIndex_RT = NewTargetIndex;
+			Viewport->SetRenderTargetTextureRenderThread(NewRenderTargetTexture);
+		});
 
 	FViewport::EnqueueBeginRenderFrame(bShouldPresent);
 
-	if (StereoRenderTargetManager != nullptr && bShouldPresent)
+	if (StereoRenderTargetManager && bShouldPresent)
 	{
 		StereoRenderTargetManager->UpdateViewport(UseSeparateRenderTarget(), *this, ViewportWidget.Pin().Get());
 	}
@@ -2139,12 +2140,13 @@ void FSceneViewport::EnqueueBeginRenderFrame(const bool bShouldPresent)
 
 void FSceneViewport::BeginRenderFrame(FRHICommandListImmediate& RHICmdList)
 {
-	check( IsInRenderingThread() );
+	check(IsInRenderingThread());
+
 	if (UseSeparateRenderTarget())
-	{		
+	{
 		RHICmdList.Transition(FRHITransitionInfo(RenderTargetTextureRenderThreadRHI, ERHIAccess::Unknown, ERHIAccess::RTV));
 	}
-	else if( IsValidRef( ViewportRHI ) ) 
+	else if (IsValidRef(ViewportRHI))
 	{
 		// Get the backbuffer render target to render directly to it
 		RenderTargetTextureRenderThreadRHI = RHIGetViewportBackBuffer(ViewportRHI);
@@ -2157,7 +2159,7 @@ void FSceneViewport::EndRenderFrame(FRHICommandListImmediate& RHICmdList, bool b
 	check( IsInRenderingThread() );
 	if (UseSeparateRenderTarget())
 	{
-		if (BufferedSlateHandles[CurrentBufferedTargetIndex])
+		if (BufferedFrames[BufferIndex_RT].SlateHandle)
 		{
 			RHICmdList.Transition(FRHITransitionInfo(RenderTargetTextureRenderThreadRHI, ERHIAccess::Unknown, ERHIAccess::SRVMask));
 		}
@@ -2256,18 +2258,16 @@ void FSceneViewport::OnPostResizeWindowBackbuffer(void* Backbuffer)
 {
 	check(IsInGameThread());
 
-	if(!UseSeparateRenderTarget() && !IsValidRef(ViewportRHI) && ViewportWidget.IsValid())
+	if (!UseSeparateRenderTarget() && !IsValidRef(ViewportRHI) && ViewportWidget.IsValid())
 	{
-		FSlateRenderer* Renderer = FSlateApplication::Get().GetRenderer();
-
 		TSharedPtr<SWindow> Window = FSlateApplication::Get().FindWidgetWindow(ViewportWidget.Pin().ToSharedRef());
 
 		// If the window is not valid then we are likely in a loading movie and the viewport is not attached to the window.  
 		// We'll have to wait until safe
-		if(Window.IsValid())
+		if (Window.IsValid())
 		{
-			void* ViewportResource = Renderer->GetViewportResource(*Window);
-			if (ViewportResource)
+			FSlateRenderer* Renderer = FSlateApplication::Get().GetRenderer();
+			if (void* ViewportResource = Renderer->GetViewportResource(*Window))
 			{
 				ViewportRHI = *((FViewportRHIRef*)ViewportResource);
 			}
@@ -2297,65 +2297,57 @@ void FSceneViewport::InitRHI(FRHICommandListBase& RHICmdList)
 
 	if (UseSeparateRenderTarget())
 	{
-		int32 NumBufferedFrames = 1;
-		TArray<FTextureRHIRef> BufferedRTRHI;
-		TArray<FTextureRHIRef> BufferedSRVRHI;
+		int32 NewNumBufferedFrames = 1;
+		TArray<FTextureRHIRef> ExternalRenderTargets;
+		TArray<FTextureRHIRef> ExternalResourceTextures;
 		
 		// @todo vreditor switch: This code needs to be called when switching between stereo/non when going immersive.  Seems to always work out that way anyway though? (Probably due to resize)
-		bool bHMDAllocatedSeparateRenderTargets = false;
-		IStereoRenderTargetManager* const StereoRenderTargetManager = RetrieveStereoRenderTargetManager(IsStereoRenderingAllowed());
-		if (StereoRenderTargetManager != nullptr)
+		bool bExternallyAllocatedRenderTargets = false;
+		if (IStereoRenderTargetManager* StereoRenderTargetManager = RetrieveStereoRenderTargetManager(IsStereoRenderingAllowed()))
 		{
 			StereoRenderTargetManager->CalculateRenderTargetSize(*this, TexSizeX, TexSizeY);
-			bHMDAllocatedSeparateRenderTargets = StereoRenderTargetManager->AllocateRenderTargetTextures(RHICmdList, TexSizeX, TexSizeY, SceneTargetFormat, 1, TexCreate_None, TexCreate_RenderTargetable, BufferedRTRHI, BufferedSRVRHI);
-			if (bHMDAllocatedSeparateRenderTargets)
+			bExternallyAllocatedRenderTargets = StereoRenderTargetManager->AllocateRenderTargetTextures(RHICmdList, TexSizeX, TexSizeY, SceneTargetFormat, 1, TexCreate_None, TexCreate_RenderTargetable, ExternalRenderTargets, ExternalResourceTextures);
+			if (bExternallyAllocatedRenderTargets)
 			{
-				check(BufferedRTRHI.Num() == BufferedSRVRHI.Num());
-				NumBufferedFrames = BufferedRTRHI.Num();
+				check(ExternalRenderTargets.Num() == ExternalResourceTextures.Num());
+				check(ExternalRenderTargets.Num() > 0);
+				NewNumBufferedFrames = ExternalRenderTargets.Num();
 			}
 		}
-		
-		check(BufferedSlateHandles.Num() == BufferedRenderTargetsRHI.Num() && BufferedSlateHandles.Num() == BufferedShaderResourceTexturesRHI.Num());
 
-		//clear existing entries
-		for (int32 i = 0; i < BufferedSlateHandles.Num(); ++i)
+		// clear existing entries
+		for (FFrameData& FrameData : BufferedFrames)
 		{
-			if (!BufferedSlateHandles[i])
+			if (!FrameData.SlateHandle)
 			{
-				BufferedSlateHandles[i] = new FSlateRenderTargetRHI(nullptr, 0, 0);
+				FrameData.SlateHandle = new FSlateRenderTargetRHI(nullptr, 0, 0);
 			}
-			BufferedRenderTargetsRHI[i] = nullptr;
-			BufferedShaderResourceTexturesRHI[i] = nullptr;
+			FrameData.RenderTarget = nullptr;
+			FrameData.ResourceTexture = nullptr;
 		}
 
-		if (BufferedSlateHandles.Num() < NumBufferedFrames)
+		if (BufferedFrames.Num() < NewNumBufferedFrames)
 		{
-			//add sufficient entries for buffering.
-			for (int32 i = BufferedSlateHandles.Num(); i < NumBufferedFrames; i++)
+			// Add sufficient entries for buffering.
+			for (int32 Index = BufferedFrames.Num(); Index < NewNumBufferedFrames; Index++)
 			{
-				BufferedSlateHandles.Add(new FSlateRenderTargetRHI(nullptr, 0, 0)); 
-				BufferedRenderTargetsRHI.Add(nullptr);
-				BufferedShaderResourceTexturesRHI.Add(nullptr);
+				BufferedFrames.Emplace(new FSlateRenderTargetRHI(nullptr, 0, 0));
 			}
 		}
-		else if (BufferedSlateHandles.Num() > NumBufferedFrames)
+		else if (BufferedFrames.Num() > NewNumBufferedFrames)
 		{
-			BufferedSlateHandles.SetNum(NumBufferedFrames);
-			BufferedRenderTargetsRHI.SetNum(NumBufferedFrames);
-			BufferedShaderResourceTexturesRHI.SetNum(NumBufferedFrames);
+			BufferedFrames.SetNum(NewNumBufferedFrames);
 		}
-		check(BufferedSlateHandles.Num() == BufferedRenderTargetsRHI.Num() && BufferedSlateHandles.Num() == BufferedShaderResourceTexturesRHI.Num());
 
-		FTextureRHIRef RTRHI;
-		FTextureRHIRef SRVRHI;
-
-		for (int32 i = 0; i < NumBufferedFrames; ++i)
+		for (int32 Index = 0; Index < NewNumBufferedFrames; Index++)
 		{
-			// try to allocate texture via StereoRenderingDevice; if not successful, use the default way
-			if (bHMDAllocatedSeparateRenderTargets)
+			FTextureRHIRef RenderTarget;
+			FTextureRHIRef ResourceTexture;
+
+			if (bExternallyAllocatedRenderTargets)
 			{
-				RTRHI = BufferedRTRHI[i];
-				SRVRHI = BufferedSRVRHI[i];
+				RenderTarget = MoveTemp(ExternalRenderTargets[Index]);
+				ResourceTexture = MoveTemp(ExternalResourceTextures[Index]);
 			}
 			else
 			{
@@ -2366,44 +2358,48 @@ void FSceneViewport::InitRHI(FRHICommandListBase& RHICmdList)
 					.SetFlags(ETextureCreateFlags::RenderTargetable | ETextureCreateFlags::ShaderResource)
 					.SetInitialState(ERHIAccess::SRVMask);
 
-				RTRHI = SRVRHI = RHICmdList.CreateTexture(Desc);
+				RenderTarget = ResourceTexture = RHICmdList.CreateTexture(Desc);
 			}
-			BufferedRenderTargetsRHI[i] = RTRHI;
-			BufferedShaderResourceTexturesRHI[i] = SRVRHI;
 
-			if (BufferedSlateHandles[i])
+			BufferedFrames[Index].RenderTarget = RenderTarget;
+			BufferedFrames[Index].ResourceTexture = ResourceTexture;
+
+			if (BufferedFrames[Index].SlateHandle)
 			{
-				BufferedSlateHandles[i]->SetRHIRef(SRVRHI, TexSizeX, TexSizeY);
+				BufferedFrames[Index].SlateHandle->SetRHIRef(ResourceTexture, TexSizeX, TexSizeY);
 			}
 		}
 
 		// clear out any extra entries we have hanging around
-		for (int32 i = NumBufferedFrames; i < BufferedSlateHandles.Num(); ++i)
+		for (int32 Index = NewNumBufferedFrames; Index < BufferedFrames.Num(); Index++)
 		{
-			if (BufferedSlateHandles[i])
+			if (BufferedFrames[Index].SlateHandle)
 			{
-				BufferedSlateHandles[i]->SetRHIRef(nullptr, 0, 0);
+				BufferedFrames[Index].SlateHandle->SetRHIRef(nullptr, 0, 0);
 			}
-			BufferedRenderTargetsRHI[i] = nullptr;
-			BufferedShaderResourceTexturesRHI[i] = nullptr;
+			BufferedFrames[Index].RenderTarget = nullptr;
+			BufferedFrames[Index].ResourceTexture = nullptr;
 		}
 
-		CurrentBufferedTargetIndex = 0;
-		NextBufferedTargetIndex = (CurrentBufferedTargetIndex + 1) % BufferedSlateHandles.Num();
-		RenderTargetTextureRHI = BufferedShaderResourceTexturesRHI[CurrentBufferedTargetIndex];
+		BufferIndex_GT = 0;
+		BufferIndex_RT = 0;
+		NextBufferIndex_GT = (BufferIndex_GT + 1) % NewNumBufferedFrames;
+		NumBufferedFrames = NewNumBufferedFrames;
+
+		RenderTargetTextureRHI = BufferedFrames[BufferIndex_GT].ResourceTexture;
 	}
 	else
 	{
-		check(BufferedSlateHandles.Num() == BufferedRenderTargetsRHI.Num() && BufferedSlateHandles.Num() == BufferedShaderResourceTexturesRHI.Num());
-		if (BufferedSlateHandles.Num() == 0)
+		if (BufferedFrames.Num() == 0)
 		{
-			BufferedSlateHandles.Add(nullptr);
-			BufferedRenderTargetsRHI.Add(nullptr);
-			BufferedShaderResourceTexturesRHI.Add(nullptr);
+			BufferedFrames.Emplace();
 		}
 
-		RenderTargetTextureRHI = nullptr;		
-		CurrentBufferedTargetIndex = NextBufferedTargetIndex = 0;
+		RenderTargetTextureRHI = nullptr;
+		BufferIndex_GT = 0;
+		BufferIndex_RT = 0;
+		NextBufferIndex_GT = 0;
+		NumBufferedFrames = 0;
 	}
 }
 
@@ -2413,11 +2409,11 @@ void FSceneViewport::ReleaseRHI()
 
 	ViewportRHI.SafeRelease();
 
-	for (int32 i = 0; i < BufferedSlateHandles.Num(); ++i)
+	for (FFrameData& FrameData : BufferedFrames)
 	{
-		if (BufferedSlateHandles[i])
+		if (FrameData.SlateHandle)
 		{
-			BufferedSlateHandles[i]->ReleaseRHI();
+			FrameData.SlateHandle->ReleaseRHI();
 		}
 	}
 	if (RenderThreadSlateTexture)
diff --git a/Engine/Source/Runtime/Engine/Public/Slate/SceneViewport.h b/Engine/Source/Runtime/Engine/Public/Slate/SceneViewport.h
--- a/Engine/Source/Runtime/Engine/Public/Slate/SceneViewport.h
+++ b/Engine/Source/Runtime/Engine/Public/Slate/SceneViewport.h
@@ -17,15 +17,15 @@
 class FCanvas;
 class FDebugCanvasDrawer;
 class FSlateRenderer;
+class FSlateRenderTargetRHI;
 class FSlateWindowElementList;
+class SDebugCanvas;
 class SViewport;
 class SWindow;
 
 /** Called in FSceneViewport::ResizeFrame after ResizeViewport*/
 DECLARE_DELEGATE_OneParam( FOnSceneViewportResize, FVector2D );
 
-class SViewport;
-
 extern const FName NAME_SceneViewport;
 
 /**
@@ -63,7 +63,7 @@ public:
 	/** Gets the proper RenderTarget based on the current thread*/
 	ENGINE_API virtual const FTextureRHIRef& GetRenderTargetTexture() const;
 
-	ENGINE_API virtual void SetRenderTargetTextureRenderThread(FTextureRHIRef& RT);
+	ENGINE_API virtual void SetRenderTargetTextureRenderThread(FRHITexture* InRenderTargetTexture);
 
 	/**
 	 * Captures or uncaptures the joystick
@@ -296,7 +296,7 @@ public:
 	};
 
 	/** Sets the debug canvas used to display FCanvas on top of this viewport */
-	ENGINE_API void SetDebugCanvas(TSharedPtr<class SDebugCanvas> InDebugCanvas);
+	ENGINE_API void SetDebugCanvas(TSharedPtr<SDebugCanvas> InDebugCanvas);
 
 	/** Adds a draw element for the debug canvas.  Called externally by a widget that manages where the debug canvas draws */
 	ENGINE_API void PaintDebugCanvas(const FGeometry& AllottedGeometry, FSlateWindowElementList& OutDrawElements, int32 LayerId) const;
@@ -365,9 +365,6 @@ private:
 	/** Utility function to figure out if we are currently a game viewport */
 	ENGINE_API bool IsCurrentlyGameViewport();
 
-	UE_DEPRECATED(5.5, "WindowRenderTargetUpdate is no longer used")
-	void WindowRenderTargetUpdate(FSlateRenderer* Renderer, SWindow* Window) {}
-
 	/** @return Returns true if we should always render to a separate render target (rather than rendering directly to the
 	    viewport backbuffer, taking into account any temporary requirements of head-mounted displays */
 	bool UseSeparateRenderTarget() const override
@@ -418,7 +415,7 @@ private:
 	/**	Whether the software cursor should be drawn in the viewport */
 	bool bIsSoftwareCursorVisible;	
 	/** Draws the debug canvas in Slate */
-	TSharedPtr<class FDebugCanvasDrawer, ESPMode::ThreadSafe> DebugCanvasDrawer;
+	TSharedPtr<FDebugCanvasDrawer, ESPMode::ThreadSafe> DebugCanvasDrawer;
 	/** The Slate viewport widget where this viewport is drawn */
 	TWeakPtr<SViewport> ViewportWidget;
 	/** Debug canvas widget we invalidate if our FCanvas has draw elements */
@@ -466,15 +463,29 @@ private:
 
 	/** Reprojection on some HMD RHI's requires ViewportTargets to be buffered */
 	/** The render target used by Slate to draw the viewport.  Can be null if this viewport renders directly to the backbuffer */
-	TArray<class FSlateRenderTargetRHI*> BufferedSlateHandles;
-	TArray<FTextureRHIRef> BufferedRenderTargetsRHI;
-	TArray<FTextureRHIRef> BufferedShaderResourceTexturesRHI;
+
+	struct FFrameData
+	{
+		FFrameData() = default;
+		FFrameData(FSlateRenderTargetRHI* InSlateHandle)
+			: SlateHandle(InSlateHandle)
+		{
+		}
+
+		FSlateRenderTargetRHI* SlateHandle = nullptr;
+		FTextureRHIRef         RenderTarget;
+		FTextureRHIRef         ResourceTexture;
+	};
+	TArray<FFrameData, TInlineAllocator<1>> BufferedFrames;
 
 	FTextureRHIRef RenderTargetTextureRenderThreadRHI;
-	class FSlateRenderTargetRHI* RenderThreadSlateTexture;
+	FSlateRenderTargetRHI* RenderThreadSlateTexture;
+
+	int32 BufferIndex_GT = 0;
+	int32 NextBufferIndex_GT = 0;
 
-	int32 CurrentBufferedTargetIndex;
-	int32 NextBufferedTargetIndex;
+	int32 BufferIndex_RT = 0;
+	int32 NumBufferedFrames = 0;
 
 	/** Tracks the number of touches currently active on the viewport */
 	int32 NumTouches;
