Fixed being able to attempt to toggle fullscreen on console platforms, resulting in a fatal log. Will remove the CVar in exchange for an constexpr-if in 40.00.

#rnx

diff --git a/Engine/Source/Runtime/Engine/Private/GameViewportClient.cpp b/Engine/Source/Runtime/Engine/Private/GameViewportClient.cpp
--- a/Engine/Source/Runtime/Engine/Private/GameViewportClient.cpp
+++ b/Engine/Source/Runtime/Engine/Private/GameViewportClient.cpp
@@ -171,6 +171,13 @@ static TAutoConsoleVariable<int32> CVarScreenshotTextureArrayIndex(
 	TEXT("Texture array slice to use for screenshots when r.Screenshot.TextureArrayMode = 0"),
 	ECVF_Default);
 
+static bool bForceAllowFullscreenToggle = false;
+static FAutoConsoleVariableRef ForceAllowFullscreenToggleCVar(
+	TEXT("r.ForceAllowFullscreenToggle"),
+	bForceAllowFullscreenToggle,
+	TEXT("Temporary CVar to allow rolling back logic that prevents full-screen toggling on non-desktop platforms"),
+	ECVF_Default);
+
 #if CSV_PROFILER_STATS
 struct FCsvLocalPlayer
 {
@@ -3991,6 +3998,11 @@ bool UGameViewportClient::SetDisplayConfiguration(const FIntPoint* Dimensions, E
 
 bool UGameViewportClient::HandleToggleFullscreenCommand()
 {
+	if (!FPlatformProperties::SupportsWindowedMode() && !bForceAllowFullscreenToggle)
+	{
+		return true;
+	}
+
 	if (!Viewport)
 	{
 		return true;
