Deprecate FViewportFrame
- Only FSceneViewport implemented this interface, and code using it already has access to the FSceneViewport type instead.

diff --git a/Engine/Plugins/Compositing/OpenColorIO/Source/OpenColorIO/Private/OpenColorIODisplayExtensionWrapper.cpp b/Engine/Plugins/Compositing/OpenColorIO/Source/OpenColorIO/Private/OpenColorIODisplayExtensionWrapper.cpp
--- a/Engine/Plugins/Compositing/OpenColorIO/Source/OpenColorIO/Private/OpenColorIODisplayExtensionWrapper.cpp
+++ b/Engine/Plugins/Compositing/OpenColorIO/Source/OpenColorIO/Private/OpenColorIODisplayExtensionWrapper.cpp
@@ -118,7 +118,7 @@ UOpenColorIODisplayExtensionWrapper* UOpenColorIODisplayExtensionWrapper::Create
 		// Activate the SVE if it is the game's primary viewport.
 		if (UGameEngine* GameEngine = Cast<UGameEngine>(GEngine))
 		{
-			if (GameEngine->SceneViewport.IsValid() && (GameEngine->SceneViewport->GetViewport() == Context.Viewport))
+			if (GameEngine->SceneViewport.IsValid() && (GameEngine->SceneViewport.Get() == Context.Viewport))
 			{
 				return TOptional<bool>(bIsEnabled);
 			}
diff --git a/Engine/Plugins/Experimental/VirtualProductionUtilities/Source/VPUtilities/Private/Libraries/VPRenderingBlueprintLibrary.cpp b/Engine/Plugins/Experimental/VirtualProductionUtilities/Source/VPUtilities/Private/Libraries/VPRenderingBlueprintLibrary.cpp
--- a/Engine/Plugins/Experimental/VirtualProductionUtilities/Source/VPUtilities/Private/Libraries/VPRenderingBlueprintLibrary.cpp
+++ b/Engine/Plugins/Experimental/VirtualProductionUtilities/Source/VPUtilities/Private/Libraries/VPRenderingBlueprintLibrary.cpp
@@ -79,7 +79,7 @@ void UVPRenderingBlueprintLibrary::GenerateSceneViewExtensionIsActiveFunctorForV
 		{
 			if (UGameEngine* GameEngine = Cast<UGameEngine>(GEngine))
 			{
-				if (GameEngine->SceneViewport.IsValid() && (GameEngine->SceneViewport->GetViewport() == Context.Viewport))
+				if (GameEngine->SceneViewport.IsValid() && (GameEngine->SceneViewport.Get() == Context.Viewport))
 				{
 					return TOptional<bool>(true);
 				}
diff --git a/Engine/Plugins/Media/MediaFrameworkUtilities/Source/MediaFrameworkUtilitiesEditor/Private/AssetEditor/SMediaProfileMediaItemDisplay.cpp b/Engine/Plugins/Media/MediaFrameworkUtilities/Source/MediaFrameworkUtilitiesEditor/Private/AssetEditor/SMediaProfileMediaItemDisplay.cpp
--- a/Engine/Plugins/Media/MediaFrameworkUtilities/Source/MediaFrameworkUtilitiesEditor/Private/AssetEditor/SMediaProfileMediaItemDisplay.cpp
+++ b/Engine/Plugins/Media/MediaFrameworkUtilities/Source/MediaFrameworkUtilitiesEditor/Private/AssetEditor/SMediaProfileMediaItemDisplay.cpp
@@ -714,10 +714,9 @@ private:
 			return nullptr;
 		}
 
-		FViewport* Viewport = ViewportClient->Viewport;
-		if (Viewport && Viewport->GetViewportType() == TEXT("SceneViewport"))
+		if (FViewport* Viewport = ViewportClient->Viewport)
 		{
-			return static_cast<FSceneViewport*>(Viewport);
+			return Viewport->AsSceneViewport();
 		}
 		
 		return nullptr;
diff --git a/Engine/Plugins/Media/MediaIOFramework/Source/MediaIOCore/Private/MediaCapture.cpp b/Engine/Plugins/Media/MediaIOFramework/Source/MediaIOCore/Private/MediaCapture.cpp
--- a/Engine/Plugins/Media/MediaIOFramework/Source/MediaIOCore/Private/MediaCapture.cpp
+++ b/Engine/Plugins/Media/MediaIOFramework/Source/MediaIOCore/Private/MediaCapture.cpp
@@ -336,7 +336,7 @@ bool UMediaCapture::StartSourceCapture(TSharedPtr<UE::MediaCapture::Private::FCa
 			FRenderTarget* RenderTarget = nullptr;
 			if (CaptureSource->GetSceneViewport())
 			{
-				RenderTarget = CaptureSource->GetSceneViewport()->GetViewport();
+				RenderTarget = CaptureSource->GetSceneViewport().Get();
 			}
 			else if (CaptureSource->GetRenderTarget())
 			{
diff --git a/Engine/Plugins/NetcodeUnitTest/NetcodeUnitTest/Source/NetcodeUnitTest/Private/UnitTestCommandlet.cpp b/Engine/Plugins/NetcodeUnitTest/NetcodeUnitTest/Source/NetcodeUnitTest/Private/UnitTestCommandlet.cpp
--- a/Engine/Plugins/NetcodeUnitTest/NetcodeUnitTest/Source/NetcodeUnitTest/Private/UnitTestCommandlet.cpp
+++ b/Engine/Plugins/NetcodeUnitTest/NetcodeUnitTest/Source/NetcodeUnitTest/Private/UnitTestCommandlet.cpp
@@ -107,7 +107,7 @@ int32 UUnitTestCommandlet::Main(const FString& Params)
 				NewViewport->SetViewportOverlayWidget(nullptr, DudOverlay);
 
 				// Set the internal FViewport, for the new game viewport, to avoid another bit of auto-exit code
-				NewViewport->Viewport = new FDummyViewport(NewViewport);
+				NewViewport->Viewport = new FSceneViewport(NewViewport, nullptr);
 
 				// Set the main engine world context game viewport, to match the newly created viewport, in order to prevent crashes
 				CurContext->GameViewport = NewViewport;
diff --git a/Engine/Plugins/VirtualProduction/SceneState/Source/SceneStateGameplayEditor/Private/SceneStateGameplayContextEditor.cpp b/Engine/Plugins/VirtualProduction/SceneState/Source/SceneStateGameplayEditor/Private/SceneStateGameplayContextEditor.cpp
--- a/Engine/Plugins/VirtualProduction/SceneState/Source/SceneStateGameplayEditor/Private/SceneStateGameplayContextEditor.cpp
+++ b/Engine/Plugins/VirtualProduction/SceneState/Source/SceneStateGameplayEditor/Private/SceneStateGameplayContextEditor.cpp
@@ -36,7 +36,7 @@ class SGameViewport : public SCompoundWidget
 		ViewportWidget->SetViewportInterface(SceneViewport.ToSharedRef());
 
 		ViewportClient.Reset(InViewportClient);
-		ViewportClient->SetViewportFrame(SceneViewport.Get());
+		ViewportClient->SetViewport(SceneViewport.Get());
 
 		ChildSlot
 		[
diff --git a/Engine/Source/Editor/LevelEditor/Private/SLevelViewport.cpp b/Engine/Source/Editor/LevelEditor/Private/SLevelViewport.cpp
--- a/Engine/Source/Editor/LevelEditor/Private/SLevelViewport.cpp
+++ b/Engine/Source/Editor/LevelEditor/Private/SLevelViewport.cpp
@@ -2929,7 +2929,7 @@ bool SLevelViewport::IsFoliageTypeVisible(TWeakObjectPtr<UFoliageType> InFoliage
 
 FViewport* SLevelViewport::GetActiveViewport() 
 { 
-	return ActiveViewport->GetViewport(); 
+	return ActiveViewport.Get();
 }
 
 void SLevelViewport::OnFocusViewportToSelection() 
diff --git a/Engine/Source/Editor/UnrealEd/Private/PlayLevel.cpp b/Engine/Source/Editor/UnrealEd/Private/PlayLevel.cpp
--- a/Engine/Source/Editor/UnrealEd/Private/PlayLevel.cpp
+++ b/Engine/Source/Editor/UnrealEd/Private/PlayLevel.cpp
@@ -3476,9 +3476,8 @@ TSharedRef<SPIEViewport> UEditorEngine::GeneratePIEViewportWindow(const FRequest
 
 	InSlateInfo.SlatePlayInEditorWindow = PieWindow;
 
-	// Let the viewport client know what viewport is using it.  We need to set the Viewport Frame as 
-	// well (which in turn sets the viewport) so that SetRes command will work.
-	InViewportClient->SetViewportFrame(InSlateInfo.SlatePlayInEditorWindowViewport.Get());
+	// Let the viewport client know what viewport is using it.
+	InViewportClient->SetViewport(InSlateInfo.SlatePlayInEditorWindowViewport.Get());
 	// Mark the viewport as PIE viewport
 	InViewportClient->Viewport->SetPlayInEditorViewport(InViewportClient->bIsPlayInEditorViewport);
 
diff --git a/Engine/Source/Runtime/Engine/Classes/Engine/GameViewportClient.h b/Engine/Source/Runtime/Engine/Classes/Engine/GameViewportClient.h
--- a/Engine/Source/Runtime/Engine/Classes/Engine/GameViewportClient.h
+++ b/Engine/Source/Runtime/Engine/Classes/Engine/GameViewportClient.h
@@ -20,13 +20,13 @@
 #include "UObject/SoftObjectPath.h"
 #include "StereoRendering.h"
 #include "AudioDeviceHandle.h"
+#include "Slate/SceneViewport.h"
 
 #include "GameViewportClient.generated.h"
 
 class FCanvas;
 class FSceneView;
 class FSceneViewport;
-class FViewportFrame;
 class IGameLayerManager;
 class SOverlay;
 class SViewport;
@@ -325,13 +325,17 @@ public:
 	 * Set this GameViewportClient's viewport and viewport frame to the viewport specified
 	 * @param	InViewportFrame	The viewportframe to set
 	 */
-	ENGINE_API virtual void SetViewportFrame( FViewportFrame* InViewportFrame );
+	UE_DEPRECATED(5.8, "UGameViewportClient::SetViewportFrame and FViewportFrame are deprecated. Use FSceneViewport/SetViewport.")
+	inline void SetViewportFrame(FSceneViewport* InViewportFrame)
+	{
+		SetViewport(InViewportFrame);
+	}
 
 	/**
 	 * Set this GameViewportClient's viewport to the viewport specified
 	 * @param	InViewportFrame	The viewport to set
 	 */
-	ENGINE_API virtual void SetViewport( FViewport* InViewportFrame );
+	ENGINE_API virtual void SetViewport( FSceneViewport* InViewportFrame );
 
 	/** Assigns the viewport overlay widget to use for this viewport client.  Should only be called when first created */
 	ENGINE_API void SetViewportOverlayWidget(TSharedPtr< SWindow > InWindow, TSharedRef<SOverlay> InViewportOverlayWidget);
@@ -681,11 +685,15 @@ public:
 	/** The show flags used by the viewport's players. */
 	FEngineShowFlags EngineShowFlags;
 
-	/** The platform-specific viewport which this viewport client is attached to. */
-	FViewport* Viewport;
+	union
+	{
+		/** The platform-specific viewport which this viewport client is attached to. */
+		FSceneViewport* Viewport;
 
-	/** The platform-specific viewport frame which this viewport is contained by. */
-	FViewportFrame* ViewportFrame;
+		/** The platform-specific viewport frame which this viewport is contained by. */
+		UE_DEPRECATED(5.8, "ViewportFrame is deprecated. Use the Viewport member instead.")
+		FSceneViewport* ViewportFrame;
+	};
 
 	/**
 	 * Controls suppression of the blue transition text messages 
diff --git a/Engine/Source/Runtime/Engine/Private/GameEngine.cpp b/Engine/Source/Runtime/Engine/Private/GameEngine.cpp
--- a/Engine/Source/Runtime/Engine/Private/GameEngine.cpp
+++ b/Engine/Source/Runtime/Engine/Private/GameEngine.cpp
@@ -300,11 +300,9 @@ void UGameEngine::CreateGameViewport( UGameViewportClient* GameViewportClient )
 	// The viewport widget needs an interface so it knows what should render
 	GameViewportWidgetRef->SetViewportInterface( SceneViewport.ToSharedRef() );
 
-	FSceneViewport* ViewportFrame = SceneViewport.Get();
+	GameViewport->SetViewport(SceneViewport.Get());
 
-	GameViewport->SetViewportFrame(ViewportFrame);
-
-	GameViewport->GetGameLayerManager()->SetSceneViewport(ViewportFrame);
+	GameViewport->GetGameLayerManager()->SetSceneViewport(SceneViewport.Get());
 
 	FViewport::ViewportResizedEvent.AddUObject(this, &UGameEngine::OnViewportResized);
 }
@@ -816,7 +814,7 @@ void UGameEngine::OnGameWindowClosed( const TSharedRef<SWindow>& WindowBeingClos
 {
 	FSlateApplication::Get().UnregisterGameViewport();
 	// This will shutdown the game
-	GameViewport->CloseRequested( SceneViewport->GetViewport() );
+	GameViewport->CloseRequested(SceneViewport.Get());
 	SceneViewport.Reset();
 }
 
diff --git a/Engine/Source/Runtime/Engine/Private/GameUserSettings.cpp b/Engine/Source/Runtime/Engine/Private/GameUserSettings.cpp
--- a/Engine/Source/Runtime/Engine/Private/GameUserSettings.cpp
+++ b/Engine/Source/Runtime/Engine/Private/GameUserSettings.cpp
@@ -218,7 +218,7 @@ bool UGameUserSettings::IsDynamicResolutionEnabled() const
 bool UGameUserSettings::IsScreenResolutionDirty() const
 {
 	bool bIsDirty = false;
-	if (GEngine && GEngine->GameViewport && GEngine->GameViewport->ViewportFrame)
+	if (GEngine && GEngine->GameViewport && GEngine->GameViewport->Viewport)
 	{
 		bIsDirty = (ResolutionSizeX != GSystemResolution.ResX || ResolutionSizeY != GSystemResolution.ResY) ? true : false;
 	}
@@ -228,7 +228,7 @@ bool UGameUserSettings::IsScreenResolutionDirty() const
 bool UGameUserSettings::IsFullscreenModeDirty() const
 {
 	bool bIsDirty = false;
-	if (GEngine && GEngine->GameViewport && GEngine->GameViewport->ViewportFrame)
+	if (GEngine && GEngine->GameViewport && GEngine->GameViewport->Viewport)
 	{
 		EWindowMode::Type CurrentFullscreenMode = GEngine->GameViewport->Viewport->GetWindowMode();
 		EWindowMode::Type NewFullscreenMode = GetFullscreenMode();
@@ -240,7 +240,7 @@ bool UGameUserSettings::IsFullscreenModeDirty() const
 bool UGameUserSettings::IsVSyncDirty() const
 {
 	bool bIsDirty = false;
-	if (GEngine && GEngine->GameViewport && GEngine->GameViewport->ViewportFrame)
+	if (GEngine && GEngine->GameViewport && GEngine->GameViewport->Viewport)
 	{
 		static const auto CVar = IConsoleManager::Get().FindTConsoleVariableDataInt(TEXT("r.VSync"));
 		bIsDirty = (bUseVSync != (CVar->GetValueOnAnyThread() != 0));
@@ -251,7 +251,7 @@ bool UGameUserSettings::IsVSyncDirty() const
 bool UGameUserSettings::IsDynamicResolutionDirty() const
 {
 	bool bIsDirty = false;
-	if (GEngine && GEngine->GameViewport && GEngine->GameViewport->ViewportFrame)
+	if (GEngine && GEngine->GameViewport && GEngine->GameViewport->Viewport)
 	{
 		bIsDirty = (bUseDynamicResolution != GEngine->GetDynamicResolutionUserSetting());
 	}
diff --git a/Engine/Source/Runtime/Engine/Private/GameViewportClient.cpp b/Engine/Source/Runtime/Engine/Private/GameViewportClient.cpp
--- a/Engine/Source/Runtime/Engine/Private/GameViewportClient.cpp
+++ b/Engine/Source/Runtime/Engine/Private/GameViewportClient.cpp
@@ -485,20 +485,12 @@ FSceneViewport* UGameViewportClient::CreateGameViewport(TSharedPtr<SViewport> In
 
 FSceneViewport* UGameViewportClient::GetGameViewport()
 {
-	if (Viewport && Viewport->GetViewportType() == NAME_SceneViewport)
-	{
-		return static_cast<FSceneViewport*>(Viewport);
-	}
-	return nullptr;
+	return Viewport;
 }
 
 const FSceneViewport* UGameViewportClient::GetGameViewport() const
 {
-	if (Viewport && Viewport->GetViewportType() == NAME_SceneViewport)
-	{
-		return static_cast<FSceneViewport*>(Viewport);
-	}
-	return nullptr;
+	return Viewport;
 }
 
 
@@ -1253,16 +1245,9 @@ TOptional<TSharedRef<SWidget>> UGameViewportClient::MapCursor(FViewport* InViewp
 	return TOptional<TSharedRef<SWidget>>();
 }
 
-void UGameViewportClient::SetViewportFrame( FViewportFrame* InViewportFrame )
-{
-	ViewportFrame = InViewportFrame;
-	SetViewport( ViewportFrame ? ViewportFrame->GetViewport() : NULL );
-}
-
-
-void UGameViewportClient::SetViewport( FViewport* InViewport )
+void UGameViewportClient::SetViewport( FSceneViewport* InViewport )
 {
-	FViewport* PreviousViewport = Viewport;
+	FSceneViewport* PreviousViewport = Viewport;
 	Viewport = InViewport;
 
 	if ( PreviousViewport == NULL && Viewport != NULL )
@@ -2575,7 +2560,7 @@ void UGameViewportClient::CloseRequested(FViewport* InViewport)
 	// broadcast close request to anyone that registered an interest
 	CloseRequestedDelegate.Broadcast(InViewport);
 
-	SetViewportFrame(NULL);
+	SetViewport(nullptr);
 
 	TSharedPtr< IGameLayerManager > GameLayerManager(GameLayerManagerPtr.Pin());
 	if (GameLayerManager.IsValid())
@@ -3958,7 +3943,7 @@ bool UGameViewportClient::HandlePreCacheCommand( const TCHAR* Cmd, FOutputDevice
 
 bool UGameViewportClient::SetDisplayConfiguration(const FIntPoint* Dimensions, EWindowMode::Type WindowMode)
 {
-	if (Viewport == NULL || ViewportFrame == NULL)
+	if (Viewport == nullptr)
 	{
 		return true;
 	}
@@ -4081,7 +4066,7 @@ bool UGameViewportClient::HandleToggleFullscreenCommand()
 
 bool UGameViewportClient::HandleSetResCommand( const TCHAR* Cmd, FOutputDevice& Ar )
 {
-	if(Viewport && ViewportFrame)
+	if (Viewport)
 	{
 		int32 X=FCString::Atoi(Cmd);
 		const TCHAR* CmdTemp = FCString::Strchr(Cmd,'x') ? FCString::Strchr(Cmd,'x')+1 : FCString::Strchr(Cmd,'X') ? FCString::Strchr(Cmd,'X')+1 : TEXT("");
diff --git a/Engine/Source/Runtime/Engine/Private/SceneView.cpp b/Engine/Source/Runtime/Engine/Private/SceneView.cpp
--- a/Engine/Source/Runtime/Engine/Private/SceneView.cpp
+++ b/Engine/Source/Runtime/Engine/Private/SceneView.cpp
@@ -2356,7 +2356,7 @@ void FSceneView::EndFinalPostprocessSettings(const FSceneViewInitOptions& ViewIn
 		{
 			// ...and if this is the viewport associated with the highres screenshot UI
 			auto ConfigViewport = Config.TargetViewport.Pin();
-			if (ConfigViewport.IsValid() && Family->RenderTarget == ConfigViewport->GetViewport())
+			if (ConfigViewport.IsValid() && Family->RenderTarget == ConfigViewport.Get())
 			{
 				static const FName ParamName = "RegionRect";
 				FLinearColor NormalizedCaptureRegion;
diff --git a/Engine/Source/Runtime/Engine/Private/Slate/SceneViewport.cpp b/Engine/Source/Runtime/Engine/Private/Slate/SceneViewport.cpp
--- a/Engine/Source/Runtime/Engine/Private/Slate/SceneViewport.cpp
+++ b/Engine/Source/Runtime/Engine/Private/Slate/SceneViewport.cpp
@@ -31,8 +31,6 @@ DEFINE_LOG_CATEGORY(LogViewport);
 
 extern EWindowMode::Type GetWindowModeType(EWindowMode::Type WindowMode);
 
-const FName NAME_SceneViewport = FName("SceneViewport");
-
 FSceneViewport::FSceneViewport( FViewportClient* InViewportClient, TSharedPtr<SViewport> InViewportWidget )
 	: FViewport( InViewportClient )
 	, CurrentReplyState( FReply::Unhandled() )
@@ -69,7 +67,6 @@ FSceneViewport::FSceneViewport( FViewportClient* InViewportClient, TSharedPtr<SV
 	, LimitingColorSpace(UE::Color::EColorSpace::sRGB)
 {
 	bIsSlateViewport = true;
-	ViewportType = NAME_SceneViewport;
 	RenderThreadSlateTexture = new FSlateRenderTargetRHI(nullptr, 0, 0);
 
 	if (InViewportClient)
diff --git a/Engine/Source/Runtime/Engine/Private/UnrealClient.cpp b/Engine/Source/Runtime/Engine/Private/UnrealClient.cpp
--- a/Engine/Source/Runtime/Engine/Private/UnrealClient.cpp
+++ b/Engine/Source/Runtime/Engine/Private/UnrealClient.cpp
@@ -55,8 +55,6 @@ static FAutoConsoleVariableRef CVarShowTimeStampsInStatUnit(
 bool FViewport::bIsGameRenderingEnabled = true;
 int32 FViewport::PresentAndStopMovieDelay = 0;
 
-static const FName NAME_DummyViewport = FName("DummyViewport");
-
 bool FRenderTarget::ReadPixels(TArray< FColor >& OutImageData, FReadSurfaceDataFlags InFlags, FIntRect InSrcRect)
 {
 	if(InSrcRect == FIntRect(0, 0, 0, 0))
@@ -1474,19 +1472,14 @@ int32 FStatHitchesData::DrawStat(FViewport* InViewport, FCanvas* InCanvas, int32
 /** Send when a viewport is resized */
 FViewport::FOnViewportResized FViewport::ViewportResizedEvent;
 
-FViewport::FViewport(FViewportClient* InViewportClient):
-	ViewportClient(InViewportClient),
-	InitialPositionX(0),
-	InitialPositionY(0),
-	SizeX(0),
-	SizeY(0),
-	WindowMode(IsRunningGame() ? GSystemResolution.WindowMode : EWindowMode::Windowed),
-	bHitProxiesCached(false),
-	bHasRequestedToggleFreeze(false),
-	bIsSlateViewport(false),
-	bIsHDR(false),
-	ViewportType(NAME_None),
-	bTakeHighResScreenShot(false)
+FViewport::FViewport(FViewportClient* InViewportClient)
+	: ViewportClient(InViewportClient)
+	, WindowMode(IsRunningGame() ? GSystemResolution.WindowMode : EWindowMode::Windowed)
+	, bHitProxiesCached(false)
+	, bHasRequestedToggleFreeze(false)
+	, bIsSlateViewport(false)
+	, bIsHDR(false)
+	, bTakeHighResScreenShot(false)
 {
 	//initialize the hit proxy kernel
 	HitProxySize = 5;
@@ -2642,10 +2635,9 @@ FDummyViewport::FDummyViewport(FViewportClient* InViewportClient)
 	: FViewport(InViewportClient)
 	, DebugCanvas(NULL)
 {
-	ViewportType = NAME_DummyViewport;
 	UWorld* CurWorld = (InViewportClient != NULL ? InViewportClient->GetWorld() : NULL);
 	DebugCanvas = new FCanvas(this, NULL, CurWorld, (CurWorld != NULL ? CurWorld->GetFeatureLevel() : GMaxRHIFeatureLevel));
-		
+
 	DebugCanvas->SetAllowedModes(0);
 }
 
diff --git a/Engine/Source/Runtime/Engine/Private/UnrealEngine.cpp b/Engine/Source/Runtime/Engine/Private/UnrealEngine.cpp
--- a/Engine/Source/Runtime/Engine/Private/UnrealEngine.cpp
+++ b/Engine/Source/Runtime/Engine/Private/UnrealEngine.cpp
@@ -1131,11 +1131,11 @@ void SystemResolutionSinkCallback()
 			// tell anyone listening about the change
 			FCoreDelegates::OnSystemResolutionChanged.Broadcast(ResX, ResY);
 
-			if (GEngine && GEngine->GameViewport && GEngine->GameViewport->ViewportFrame)
+			if (GEngine && GEngine->GameViewport && GEngine->GameViewport->Viewport)
 			{
 				FPlatformMisc::LowLevelOutputDebugStringf(TEXT("Resizing viewport due to setres change, %d x %d, MonitorIndex=%d\n"), ResX, ResY, MonitorIndex);
-				GEngine->GameViewport->ViewportFrame->SetMonitorIndexForResize(MonitorIndex);
-				GEngine->GameViewport->ViewportFrame->ResizeFrame(ResX, ResY, WindowMode);
+				GEngine->GameViewport->Viewport->SetMonitorIndexForResize(MonitorIndex);
+				GEngine->GameViewport->Viewport->ResizeFrame(ResX, ResY, WindowMode);
 			}
 		}
 	}
diff --git a/Engine/Source/Runtime/Engine/Public/DummyViewport.h b/Engine/Source/Runtime/Engine/Public/DummyViewport.h
--- a/Engine/Source/Runtime/Engine/Public/DummyViewport.h
+++ b/Engine/Source/Runtime/Engine/Public/DummyViewport.h
@@ -18,13 +18,12 @@ public:
 	//~ Begin FViewport Interface
 	virtual void BeginRenderFrame(FRHICommandListImmediate& RHICmdList) override
 	{
-		check( IsInRenderingThread() );
-		//SetRenderTarget(RHICmdList,  RenderTargetTextureRHI,  FTextureRHIRef() );
+		check(IsInRenderingThread());
 	};
 
 	virtual void EndRenderFrame(FRHICommandListImmediate& RHICmdList, bool bPresent, bool bLockToVsync) override
 	{
-		check( IsInRenderingThread() );
+		check(IsInRenderingThread());
 	}
 
 	void SetupHDR(EDisplayColorGamut InDisplayColorGamut, EDisplayOutputFormat InDisplayOutputFormat, bool bInSceneHDREnabled)
@@ -57,7 +56,6 @@ public:
 	virtual FIntPoint ViewportToVirtualDesktopPixel(FVector2D ViewportCoordinate) const override { return FIntPoint::ZeroValue; }
 	virtual void InvalidateDisplay() override { }
 	virtual void DeferInvalidateHitProxy() override { }
-	virtual FViewportFrame* GetViewportFrame() override { return 0; }
 	virtual FCanvas* GetDebugCanvas() override { return DebugCanvas; }
 	ENGINE_API virtual EDisplayColorGamut GetDisplayColorGamut() const;
 	ENGINE_API virtual EDisplayOutputFormat GetDisplayOutputFormat() const;
diff --git a/Engine/Source/Runtime/Engine/Public/Slate/SceneViewport.h b/Engine/Source/Runtime/Engine/Public/Slate/SceneViewport.h
--- a/Engine/Source/Runtime/Engine/Public/Slate/SceneViewport.h
+++ b/Engine/Source/Runtime/Engine/Public/Slate/SceneViewport.h
@@ -26,12 +26,10 @@ class SWindow;
 /** Called in FSceneViewport::ResizeFrame after ResizeViewport*/
 DECLARE_DELEGATE_OneParam( FOnSceneViewportResize, FVector2D );
 
-extern const FName NAME_SceneViewport;
-
 /**
  * A viewport for use with Slate SViewport widgets.
  */
-class FSceneViewport : public FViewportFrame, public FViewport, public ISlateViewport
+class FSceneViewport : public FViewport, public ISlateViewport
 {
 public:
 	ENGINE_API FSceneViewport( FViewportClient* InViewportClient, TSharedPtr<SViewport> InViewportWidget );
@@ -129,8 +127,15 @@ public:
 	ENGINE_API virtual void DeferInvalidateHitProxy() override;
 
 	/** FViewportFrame interface */
-	virtual FViewport* GetViewport() override { return this; }
+	UE_DEPRECATED(5.8, "FSceneViewport::GetViewport() is deprecated because FViewportFrame is also deprecated. This function call is unnecessary and can be removed.")
+	FViewport* GetViewport() { return this; }
+
+PRAGMA_DISABLE_DEPRECATION_WARNINGS
+	UE_DEPRECATED(5.8, "FSceneViewport::GetViewportFrame() and FViewportFrame are deprecated. Use the scene viewport directly.")
 	virtual FViewportFrame* GetViewportFrame() override { return this; }
+PRAGMA_ENABLE_DEPRECATION_WARNINGS
+
+	virtual FSceneViewport* AsSceneViewport() override { return this; }
 
 	/** @return The viewport widget being used */
 	TWeakPtr<SViewport> GetViewportWidget() const { return ViewportWidget; }
@@ -183,7 +188,7 @@ public:
 	 */
 	virtual bool RequiresVsync() const override { return bRequiresVsync; }
 
-	virtual void SetMonitorIndexForResize(int32 MonitorIndex) override { MonitorIndexForResize = MonitorIndex; }
+	void SetMonitorIndexForResize(int32 MonitorIndex) { MonitorIndexForResize = MonitorIndex; }
 
 	/**
 	 * Called to resize the actual window where this viewport resides
@@ -192,7 +197,7 @@ public:
 	 * @param NewSizeY		The new height of the viewport
 	 * @param NewWindowMode	 What window mode should the viewport be resized to
 	 */
-	ENGINE_API virtual void ResizeFrame(uint32 NewSizeX,uint32 NewSizeY,EWindowMode::Type NewWindowMode) override;
+	ENGINE_API void ResizeFrame(uint32 NewSizeX,uint32 NewSizeY,EWindowMode::Type NewWindowMode);
 
 	/**
 	 *	Sets the Viewport resize delegate.
diff --git a/Engine/Source/Runtime/Engine/Public/UnrealClient.h b/Engine/Source/Runtime/Engine/Public/UnrealClient.h
--- a/Engine/Source/Runtime/Engine/Public/UnrealClient.h
+++ b/Engine/Source/Runtime/Engine/Public/UnrealClient.h
@@ -178,18 +178,7 @@ protected:
 	FTextureRHIRef RenderTargetTextureRHI;
 };
 
-
-/**
-* An interface to the platform-specific implementation of a UI frame for a viewport.
-*/
-class FViewportFrame
-{
-public:
-
-	virtual class FViewport* GetViewport() = 0;
-	virtual void SetMonitorIndexForResize(int32 MonitorIndex) = 0;
-	virtual void ResizeFrame(uint32 NewSizeX,uint32 NewSizeY,EWindowMode::Type NewWindowMode) = 0;
-};
+using FViewportFrame UE_DEPRECATED(5.8, "FViewportFrame is deprecated. Use FSceneViewport instead.") = class FSceneViewport;
 
 /**
 * The maximum size that the hit proxy kernel is allowed to be set to
@@ -645,7 +634,13 @@ public:
 	 * Retrieves the interface to the viewport's frame, if it has one.
 	 * @return The viewport's frame interface.
 	 */
-	virtual FViewportFrame* GetViewportFrame() = 0;
+PRAGMA_DISABLE_DEPRECATION_WARNINGS
+	UE_DEPRECATED(5.8, "FViewport::GetViewportFrame() and FViewportFrame are deprecated. Use FSceneViewport instead.")
+	virtual FViewportFrame* GetViewportFrame()
+	{
+		return nullptr;
+	}
+PRAGMA_ENABLE_DEPRECATION_WARNINGS
 
 	/**
 	 * Calculates the view inside the viewport when the aspect ratio is locked.
@@ -767,7 +762,8 @@ public:
 	/** Returns dimensions of RenderTarget texture. Can be called on a game thread. */
 	virtual FIntPoint GetRenderTargetTextureSizeXY() const { return GetSizeXY(); }
 
-	inline FName GetViewportType() const { return ViewportType; }
+	/** Casts this instance to an FSceneViewport, or returns nullptr if this is not a scene viewport. */
+	virtual class FSceneViewport* AsSceneViewport() { return nullptr; }
 
 protected:
 
@@ -839,16 +835,16 @@ protected:
 	FViewportRHIRef ViewportRHI;
 
 	/** The initial position of the viewport. */
-	uint32 InitialPositionX;
+	uint32 InitialPositionX = 0;
 
 	/** The initial position of the viewport. */
-	uint32 InitialPositionY;
+	uint32 InitialPositionY = 0;
 
 	/** The width of the viewport. */
-	uint32 SizeX;
+	uint32 SizeX = 0;
 
 	/** The height of the viewport. */
-	uint32 SizeY;
+	uint32 SizeY = 0;
 
 	/** The size of the region to check hit proxies */
 	uint32 HitProxySize;
@@ -874,9 +870,6 @@ protected:
 	/** If true this viewport is being displayed on a HDR monitor */
 	uint32 bIsHDR : 1;
 
-	/** Used internally for testing runtime instance type before casting */
-	FName ViewportType;
-
 	/** true if we should draw game viewports (has no effect on Editor viewports) */
 	ENGINE_API static bool bIsGameRenderingEnabled;
 
diff --git a/Engine/Source/Runtime/GameMenuBuilder/Private/SGameMenuPageWidget.cpp b/Engine/Source/Runtime/GameMenuBuilder/Private/SGameMenuPageWidget.cpp
--- a/Engine/Source/Runtime/GameMenuBuilder/Private/SGameMenuPageWidget.cpp
+++ b/Engine/Source/Runtime/GameMenuBuilder/Private/SGameMenuPageWidget.cpp
@@ -658,13 +658,9 @@ void SGameMenuPageWidget::Tick( const FGeometry& AllottedGeometry, const double
 
 	if (GEngine && GEngine->GameViewport )
 	{
-		if ( GEngine->GameViewport->ViewportFrame != nullptr )
+		if ( GEngine->GameViewport->Viewport != nullptr )
 		{
-			FViewport* Viewport = GEngine->GameViewport->ViewportFrame->GetViewport();
-			if (Viewport)
-			{
-				ScreenRes = Viewport->GetSizeXY();
-			}
+			ScreenRes = GEngine->GameViewport->Viewport->GetSizeXY();
 		}
 		else
 		{
