[Backout] - CL50486700
#rnx
Original CL Desc
-----------------------------------------------------------------
This change fixes a bug where focus was lost in the Details panel after editing a property of a Blueprint-added component, due to reconstruction of the SCS (Simple Construction Script).

Problem:
When a property is edited, the SCS component is destroyed and rebuilt. This destroys all associated widgets, including the one that previously had keyboard focus. As a result, pressing Enter would not refocus the same field or any appropriate fallback widget, breaking expected keyboard flow.

Fix Details:
- The focus path is now saved as an array of widget type + index pairs before reconstruction.
- After the widget tree rebuilds, the new widget tree is traversed using this index path to identify a new matching widget. If it fails to find the exact path, it will abort.
- This allows the system to restore focus to the last focusable element if the index path is exactly the same.

Validation:
- Verified consistent focus behavior across property edits in the Details panel for Blueprint-added components (SpinBoxes, CheckBoxes, TextBoxes, ComboBoxes).

#Synced-CL 50547624

#Synced-CL 50547624

diff --git a/Engine/Source/Runtime/Engine/Private/ActorConstruction.cpp b/Engine/Source/Runtime/Engine/Private/ActorConstruction.cpp
--- a/Engine/Source/Runtime/Engine/Private/ActorConstruction.cpp
+++ b/Engine/Source/Runtime/Engine/Private/ActorConstruction.cpp
@@ -20,7 +20,6 @@
 
 #if WITH_EDITOR
 #include "Editor.h"
-#include "EditorWidgetFocusUtils.h"
 #else
 #include "Engine/World.h"
 #include "UObject/Package.h"
@@ -249,7 +248,6 @@ bool AActor::HasNonTrivialUserConstructionScript() const
 	return false;
 }
 
-
 #if WITH_EDITOR
 void AActor::RerunConstructionScripts()
 {
@@ -577,9 +575,6 @@ void AActor::RerunConstructionScripts()
 			}
 		}
 
-		using namespace WidgetFocusRestorationUtils;
-		TOptional<FFocusedWidgetData> FocusedWidgetData = GetFocusedWidgetData();
-		
 		// Destroy existing components
 		DestroyConstructedComponents();
 
@@ -801,8 +796,6 @@ void AActor::RerunConstructionScripts()
 			CurrentTransactionAnnotation = nullptr;
 		}
 
-		RestoreFocusedWidget(FocusedWidgetData);
-
 		// Remove the name to SCS node maps now that we're done constructing
 		for (const UBlueprintGeneratedClass* BPClass : ParentBPClassStack)
 		{
diff --git a/Engine/Source/Runtime/Engine/Private/EditorWidgetFocusUtils.cpp b/Engine/Source/Runtime/Engine/Private/EditorWidgetFocusUtils.cpp
deleted file mode 100644
diff --git a/Engine/Source/Runtime/Engine/Private/EditorWidgetFocusUtils.h b/Engine/Source/Runtime/Engine/Private/EditorWidgetFocusUtils.h
deleted file mode 100644
