deprecating FSoundWaveData interface on FSoundWaveProxy. Proxy becomes broker for USoundWave to live-update FSoundWaveData. Clients updated to use FSoundWaveData instead of the proxy.

#Synced-CL 50406897

#Synced-CL 50406897

diff --git a/Engine/Plugins/Experimental/MetasoundExperimental/Source/MetasoundExperimentalEngineRuntime/Private/MetasoundGranulatorNode.cpp b/Engine/Plugins/Experimental/MetasoundExperimental/Source/MetasoundExperimentalEngineRuntime/Private/MetasoundGranulatorNode.cpp
--- a/Engine/Plugins/Experimental/MetasoundExperimental/Source/MetasoundExperimentalEngineRuntime/Private/MetasoundGranulatorNode.cpp
+++ b/Engine/Plugins/Experimental/MetasoundExperimental/Source/MetasoundExperimentalEngineRuntime/Private/MetasoundGranulatorNode.cpp
@@ -277,23 +277,34 @@ namespace Metasound
 			check(!bPlaying);
 
  			CurrentWaveAsset = *WaveAsset;
-			SoundWaveProxyPtr = CurrentWaveAsset.GetSoundWaveProxy();
-			if (!SoundWaveProxyPtr)
+			const TSharedPtr<const FSoundWaveProxy> Proxy = CurrentWaveAsset.GetWaveProxy();
+
+			if (Proxy)
+			{
+				SoundWaveDataPtr = Proxy->GetSoundWaveDataRef();
+			}
+			else
 			{
+				SoundWaveDataPtr.Reset();
 				return false;
 			}
 
-			SourceFileSampleRate = SoundWaveProxyPtr->GetSampleRate();
-			SourceFileDurationSeconds = SoundWaveProxyPtr->GetDuration();
-			SourceFileNumFrames = SoundWaveProxyPtr->GetNumFrames();
-			SourceFileNumChannels = SoundWaveProxyPtr->GetNumChannels();
+			SourceFileSampleRate = SoundWaveDataPtr->GetSampleRate();
+			SourceFileDurationSeconds = SoundWaveDataPtr->GetDuration();
+			SourceFileNumFrames = SoundWaveDataPtr->GetNumFrames();
+			SourceFileNumChannels = SoundWaveDataPtr->GetNumChannels();
 
 			FSoundWaveProxyReader::FSettings ProxyReaderSettings;
 			ProxyReaderSettings.MaxDecodeSizeInFrames = DecodedAudioSizeInSeconds * SourceFileSampleRate;
 			ProxyReaderSettings.StartTimeInSeconds = 0.0f;
 			ProxyReaderSettings.bIsLooping = true;
 			ProxyReaderSettings.LoopDurationInSeconds = SourceFileDurationSeconds;
-			SoundWaveProxyReaderPtr = FSoundWaveProxyReader::Create(SoundWaveProxyPtr.ToSharedRef(), ProxyReaderSettings);
+			SoundWaveProxyReaderPtr = FSoundWaveProxyReader::Create(SoundWaveDataPtr.ToSharedRef(), ProxyReaderSettings);
+			
+			if (!SoundWaveProxyReaderPtr)
+			{
+				return false;
+			}
 
 			DecodedChunks.Reset();
 			DecodedChunks.AddDefaulted(3);
@@ -342,7 +353,7 @@ namespace Metasound
 			check(SourceFileSampleRate > 0.0f);
 			check(InDecoderSeekTimeSeconds >= 0.0f);
 			check(SourceFileNumChannels > 0);
-			check(SoundWaveProxyPtr.IsValid());
+			check(SoundWaveDataPtr.IsValid());
 
 			check(SoundWaveProxyReaderPtr.IsValid());
 			float DecoderSeekTimeSeconds = FMath::Max(InDecoderSeekTimeSeconds - 0.5f * DecodedAudioSizeInSeconds, 0.0f);
@@ -930,7 +941,7 @@ namespace Metasound
 		int32 GrainCount = 0;
 
 		FWaveAsset CurrentWaveAsset;
-		FSoundWaveProxyPtr SoundWaveProxyPtr;
+		TSharedPtr<const FSoundWaveData> SoundWaveDataPtr;
 		TUniquePtr<FSoundWaveProxyReader> SoundWaveProxyReaderPtr;
 		TUniquePtr<Audio::ISampleRateConverter> SRC;
 		EGranularEnvelope GranularEnvelopeType = EGranularEnvelope::Hann;
diff --git a/Engine/Plugins/Runtime/Harmonix/Source/HarmonixDsp/Private/HarmonixDsp/AudioData/StreamingAudioRendererV2.cpp b/Engine/Plugins/Runtime/Harmonix/Source/HarmonixDsp/Private/HarmonixDsp/AudioData/StreamingAudioRendererV2.cpp
--- a/Engine/Plugins/Runtime/Harmonix/Source/HarmonixDsp/Private/HarmonixDsp/AudioData/StreamingAudioRendererV2.cpp
+++ b/Engine/Plugins/Runtime/Harmonix/Source/HarmonixDsp/Private/HarmonixDsp/AudioData/StreamingAudioRendererV2.cpp
@@ -29,16 +29,17 @@ void FStreamingAudioRendererV2::Reset()
 	Shifter = nullptr;
 }
 
-void FStreamingAudioRendererV2::SetAudioData(TSharedRef<FSoundWaveProxy> InSoundWave, const FSettings& InSettings)
+void FStreamingAudioRendererV2::SetSoundWaveData(const TSharedRef<const FSoundWaveData>& InSoundWaveData, const FSettings& InSettings)
 {
-	SoundWaveProxy = InSoundWave.ToSharedPtr();
+	SoundWaveData = InSoundWaveData;
+	check(SoundWaveData.IsValid());
 
 	WaveProxyReader.Reset();
-	WaveProxyReader = CreateProxyReader(InSoundWave);
+	WaveProxyReader = CreateProxyReader(SoundWaveData.ToSharedRef());
 
 	check(WaveProxyReader.IsValid());
 
-	int32 WaveProxyNumChannels = SoundWaveProxy->GetNumChannels();
+	int32 WaveProxyNumChannels = SoundWaveData->GetNumChannels();
 
 	int32 DecodeBufferSize = WaveProxyNumChannels * DeinterleaveBlockSizeInFrames;
 	DecodeBuffer.Reset(DecodeBufferSize);
@@ -60,13 +61,13 @@ void FStreamingAudioRendererV2::SetAudioData(TSharedRef<FSoundWaveProxy> InSound
 
 	if (Shifter)
 	{
-		Shifter->SetSampleSourceReset(SoundWaveProxy, AsShared(), InSettings.ShifterSettings);
+		Shifter->SetSampleSourceReset(SoundWaveData.ToSharedRef(), AsShared(), InSettings.ShifterSettings);
 	}
 }
 
-const TSharedPtr<FSoundWaveProxy> FStreamingAudioRendererV2::GetAudioData() const
+TSharedPtr<const FSoundWaveData> FStreamingAudioRendererV2::GetSoundWaveData() const
 {
-	return SoundWaveProxy;
+	return SoundWaveData;
 }
 
 void FStreamingAudioRendererV2::MigrateToSampler(const FFusionSampler* InSampler)
@@ -84,7 +85,7 @@ void FStreamingAudioRendererV2::SetFrame(uint32 InFrameNum)
 
 double FStreamingAudioRendererV2::Render(TAudioBuffer<float>& OutBuffer, double InPos, int32 InMaxFrame, double InResampleInc, double InPitchShift, double InSpeed, bool MaintainPitchWhenSpeedChanges, bool InShouldHonorLoopPoints, const FGainMatrix& InGain)
 {
-	if (!SoundWaveProxy || SoundWaveProxy->GetNumFrames() == 0)
+	if (!SoundWaveData || SoundWaveData->GetNumFrames() == 0)
 	{
 		OutBuffer.ZeroValidFrames();
 		return InPos;
@@ -107,7 +108,7 @@ double FStreamingAudioRendererV2::RenderInternal(TAudioBuffer<float>& OutBuffer,
 {
 	check(OutBuffer.GetNumValidFrames() <= AudioRendering::kMicroSliceSize);
 
-	if (!SoundWaveProxy || SoundWaveProxy->GetNumFrames() < FMath::FloorToInt32(InPos))
+	if (!SoundWaveData || SoundWaveData->GetNumFrames() < FMath::FloorToInt32(InPos))
 	{
 		OutBuffer.ZeroValidFrames();
 		return InPos;
@@ -122,7 +123,7 @@ double FStreamingAudioRendererV2::RenderInternal(TAudioBuffer<float>& OutBuffer,
 	{
 		RenderMultiChannelRoutedUnshifted(OutBuffer, LerpArray, NumOutFrames, InGain, InInc, InShouldHonorLoopPoints);
 	}
-	else if (SoundWaveProxy->GetNumChannels() <= 2)
+	else if (SoundWaveData->GetNumChannels() <= 2)
 	{
 		RenderSimpleUnshifted(OutBuffer, LerpArray, NumOutFrames, InGain, InInc, InShouldHonorLoopPoints);
 	}
@@ -136,7 +137,7 @@ double FStreamingAudioRendererV2::RenderInternal(TAudioBuffer<float>& OutBuffer,
 
 double FStreamingAudioRendererV2::RenderUnshifted(TAudioBuffer<float>& OutBuffer, double InPos, int32 InMaxFrame, double InInc, bool InShouldHonorLoopPoints, const FGainMatrix& InGain)
 {
-	if (!SoundWaveProxy)
+	if (!SoundWaveData)
 	{
 		OutBuffer.ZeroValidFrames();
 		return InPos;
@@ -158,13 +159,18 @@ int32 FStreamingAudioRendererV2::CalculateNumFramesNeeded(const FLerpData* LerpD
 		return 1 + LerpData[NumPoints - 1].PosB - LerpData[0].PosA;
 	}
 
+	if (!SoundWaveData.IsValid())
+	{
+		return 0;
+	}
+
 	// Fail-safe for if we got what appears to be a loop, but we don't have a loop in the audio data
-	if (SoundWaveProxy->GetLoopRegions().IsEmpty())
+	if (SoundWaveData->GetLoopRegions().IsEmpty())
 	{
-		return SoundWaveProxy->GetNumFrames() - LerpData[0].PosA;
+		return SoundWaveData->GetNumFrames() - LerpData[0].PosA;
 	}
 
-	const FSoundWaveCuePoint& LoopRegion = SoundWaveProxy->GetLoopRegions()[0];
+	const FSoundWaveCuePoint& LoopRegion = SoundWaveData->GetLoopRegions()[0];
 
 	// loop in lerp. 
 	// we need all the samples from LerpData[0].PosA to the end of the loop, 
@@ -185,7 +191,7 @@ void FStreamingAudioRendererV2::RenderSimpleUnshifted(TAudioBuffer<float>& OutBu
 	OutBuffer.ZeroValidFrames();
 
 	int32 NumOutChannels = OutBuffer.GetNumValidChannels();
-	int32 NumInputChannels = SoundWaveProxy->GetNumChannels();
+	int32 NumInputChannels = SoundWaveData->GetNumChannels();
 
 	uint32 StartFrameIndex = LerpArray[0].PosA;
 	int32 NumSourceFramesNeeded = CalculateNumFramesNeeded(LerpArray, InNumFrames);
@@ -241,12 +247,12 @@ void FStreamingAudioRendererV2::RenderMultiChannelUnshifted(TAudioBuffer<float>&
 	OutBuffer.ZeroValidFrames();
 
 	int32 NumOutChannels = OutBuffer.GetNumValidChannels();
-	int32 NumInputChannels = SoundWaveProxy->GetNumChannels();
+	int32 NumInputChannels = SoundWaveData->GetNumChannels();
 
 	// interleaved
 	Audio::FAlignedFloatBuffer ResampleBuffer;
 	constexpr int32 kNumResampleChannels = 2;
-	ResampleBuffer.SetNum(kNumResampleChannels * InNumFrames);
+	ResampleBuffer.SetNumZeroed(kNumResampleChannels * InNumFrames);
 
 	uint32 StartFrameIndex = LerpArray[0].PosA;
 	int32 NumSourceFramesNeeded = CalculateNumFramesNeeded(LerpArray, InNumFrames);
@@ -314,7 +320,7 @@ void FStreamingAudioRendererV2::RenderMultiChannelRoutedUnshifted(TAudioBuffer<f
 	OutBuffer.ZeroValidFrames();
 
 	int32 NumOutChannels = OutBuffer.GetNumValidChannels();
-	int32 NumInputChannels = SoundWaveProxy->GetNumChannels();
+	int32 NumInputChannels = SoundWaveData->GetNumChannels();
 
 	// interleaved
 	Audio::FAlignedFloatBuffer ResampleBuffer;
@@ -503,31 +509,32 @@ uint32 FStreamingAudioRendererV2::GetSourceAudioFrameIndex()
 
 }
 
-TUniquePtr<FSoundWaveProxyReader> FStreamingAudioRendererV2::CreateProxyReader(TSharedRef<FSoundWaveProxy> WaveProxy)
+TUniquePtr<FSoundWaveProxyReader> FStreamingAudioRendererV2::CreateProxyReader(TSharedRef<const FSoundWaveData> WaveData)
 {
 	FSoundWaveProxyReader::FSettings WaveReaderSettings;
 	WaveReaderSettings.MaxDecodeSizeInFrames = MaxDecodeSizeInFrames;
 	WaveReaderSettings.StartTimeInSeconds = 0.0f;
 	WaveReaderSettings.bMaintainAudioSync = true;
-	return FSoundWaveProxyReader::Create(WaveProxy, WaveReaderSettings);	
+	return FSoundWaveProxyReader::Create(WaveData, WaveReaderSettings);	
 }
 
 bool FStreamingAudioRendererV2::HasLoopSection() const
 {
-	check(SoundWaveProxy);
-	return SoundWaveProxy->GetLoopRegions().Num() > 0;
+	check(SoundWaveData.IsValid());
+	return SoundWaveData->GetLoopRegions().Num() > 0;
 }
 
 uint32 FStreamingAudioRendererV2::GetLoopStartFrame() const
 {
 	check(HasLoopSection());
-	return SoundWaveProxy->GetLoopRegions()[0].FramePosition;
+	check(SoundWaveData.IsValid());
+	return SoundWaveData->GetLoopRegions()[0].FramePosition;
 }
 
 uint32 FStreamingAudioRendererV2::GetLoopEndFrame() const
 {
-	check(HasLoopSection())
-	const FSoundWaveCuePoint& LoopRegion = SoundWaveProxy->GetLoopRegions()[0];
+	check(HasLoopSection());
+	const FSoundWaveCuePoint& LoopRegion = SoundWaveData->GetLoopRegions()[0];
 	return LoopRegion.FramePosition + LoopRegion.FrameLength;
 }
 
@@ -654,6 +661,6 @@ void FStreamingAudioRendererV2::GenerateSourceAudioInternal(uint32 StartFrameInd
 		FMemory::Memzero(&OutAudioData[BufferIdx], sizeof(float) * NumSamplesRequested);
 
 		UE_LOG(LogHarmonixStreamingAudioRendererV2, Verbose, TEXT("%s: Failed to generated samples: StartFrameIndex: %d, NumSamplesRequested: %d"), 
-			*SoundWaveProxy->GetFName().ToString(), StartFrameIndex, NumSamplesRequested);
+			   *SoundWaveData->GetFName().ToString(), StartFrameIndex, NumSamplesRequested);
 	}
 }
diff --git a/Engine/Plugins/Runtime/Harmonix/Source/HarmonixDsp/Private/HarmonixDsp/AudioData/StreamingAudioRendererV2.h b/Engine/Plugins/Runtime/Harmonix/Source/HarmonixDsp/Private/HarmonixDsp/AudioData/StreamingAudioRendererV2.h
--- a/Engine/Plugins/Runtime/Harmonix/Source/HarmonixDsp/Private/HarmonixDsp/AudioData/StreamingAudioRendererV2.h
+++ b/Engine/Plugins/Runtime/Harmonix/Source/HarmonixDsp/Private/HarmonixDsp/AudioData/StreamingAudioRendererV2.h
@@ -13,6 +13,7 @@
 
 #define UE_API HARMONIXDSP_API
 
+class FSoundWaveData;
 class FSoundWaveProxy;
 class FSoundWaveProxyReader;
 class FFusionSampler;
@@ -30,8 +31,10 @@ public:
 	UE_API ~FStreamingAudioRendererV2();
 
 	UE_API virtual void Reset() override;
-	UE_API virtual void SetAudioData(TSharedRef<FSoundWaveProxy> SoundWaveProxy, const FSettings& InSettings) override;
-	UE_API virtual const TSharedPtr<FSoundWaveProxy> GetAudioData() const override;
+
+	UE_API virtual void SetSoundWaveData(const TSharedRef<const FSoundWaveData>& InSoundWaveData, const FSettings& InSettings) override;
+
+	UE_API virtual TSharedPtr<const FSoundWaveData> GetSoundWaveData() const override;
 	
 	UE_API virtual void MigrateToSampler(const FFusionSampler* InSampler) override;
 
@@ -60,7 +63,7 @@ private:
 
 	uint32 GetSourceAudioFrameIndex();
 
-	static TUniquePtr<FSoundWaveProxyReader> CreateProxyReader(TSharedRef<FSoundWaveProxy> WaveProxy);
+	static TUniquePtr<FSoundWaveProxyReader> CreateProxyReader(TSharedRef<const FSoundWaveData> WaveData);
 
 	bool HasLoopSection() const;
 	uint32 GetLoopStartFrame() const;
@@ -72,8 +75,8 @@ private:
 
 	// Ref to the actual streaming audio data
 	// this data is a shared instance of audio data
-	// it will get loaded on construction 
-	TSharedPtr<FSoundWaveProxy> SoundWaveProxy;
+	// it is set by SetSoundWaveData()
+	TSharedPtr<const FSoundWaveData> SoundWaveData;
 
 	TUniquePtr<FSoundWaveProxyReader> WaveProxyReader;
 
diff --git a/Engine/Plugins/Runtime/Harmonix/Source/HarmonixDsp/Private/HarmonixDsp/AudioDataRenderer.cpp b/Engine/Plugins/Runtime/Harmonix/Source/HarmonixDsp/Private/HarmonixDsp/AudioDataRenderer.cpp
--- a/Engine/Plugins/Runtime/Harmonix/Source/HarmonixDsp/Private/HarmonixDsp/AudioDataRenderer.cpp
+++ b/Engine/Plugins/Runtime/Harmonix/Source/HarmonixDsp/Private/HarmonixDsp/AudioDataRenderer.cpp
@@ -3,20 +3,55 @@
 #include "HarmonixDsp/AudioDataRenderer.h"
 #include "Sound/SoundWave.h"
 
+DEFINE_LOG_CATEGORY(LogHarmonixAudioDataRenderer);
+
+void IAudioDataRenderer::SetSoundWaveData(const TSharedRef<const FSoundWaveData>& InSoundWaveData, const FSettings& InSettings)
+{
+	static bool bDidWarn = false;
+	if (!bDidWarn)
+	{
+		UE_LOG(LogHarmonixAudioDataRenderer, Warning, TEXT("IAudioDataRenderer::SetSoundWaveData() must be overridden, and will be pure virtual in a future release."));
+		bDidWarn = true;
+	}
+	
+	PRAGMA_DISABLE_DEPRECATION_WARNINGS
+	SetAudioData(ConstCastSharedRef<FSoundWaveProxy>(FSoundWaveProxy::CreateNonUpdatingCopy(InSoundWaveData)), InSettings);
+	PRAGMA_ENABLE_DEPRECATION_WARNINGS
+}
+
+TSharedPtr<const FSoundWaveData> IAudioDataRenderer::GetSoundWaveData() const
+{
+	static bool bDidWarn = false;
+	if (!bDidWarn)
+	{
+		UE_LOG(LogHarmonixAudioDataRenderer, Warning, TEXT("IAudioDataRenderer::GetSoundWaveData() must be overridden, and will be pure virtual in a future release."));
+		bDidWarn = true;
+	}
+	
+	PRAGMA_DISABLE_DEPRECATION_WARNINGS
+	if (const TSharedPtr<FSoundWaveProxy> Proxy = GetAudioData())
+	{
+		return Proxy->GetSoundWaveDataRef().ToSharedPtr();
+	}
+	PRAGMA_ENABLE_DEPRECATION_WARNINGS
+
+	return {};
+}
+
 double IAudioDataRenderer::CalculateLerpData(FLerpData* LerpArray, int32 InNumPoints, uint32 InNumOutFrames, double InPos, int32 InMaxFrame, bool InHonorLoopPoints, double InInc) const
 {
 	check(InNumPoints >= (int32)InNumOutFrames);
 	FMemory::Memset(LerpArray, 0, sizeof(FLerpData) * InNumPoints);
 
-	const TSharedPtr<FSoundWaveProxy> SoundWaveProxy = GetAudioData();
-	if (!SoundWaveProxy)
+	TSharedPtr<const FSoundWaveData> SoundWaveData = GetSoundWaveData();
+	if (!SoundWaveData)
 	{
 		return 0.0;
 	}
 
-	if (InMaxFrame < 0 || InMaxFrame > (int32)SoundWaveProxy->GetNumFrames())
+	if (InMaxFrame < 0 || InMaxFrame > (int32)SoundWaveData->GetNumFrames())
 	{
-		InMaxFrame = (int32)SoundWaveProxy->GetNumFrames();
+		InMaxFrame = (int32)SoundWaveData->GetNumFrames();
 	}
 
 	double NumSourceFrames = (double)InMaxFrame;
@@ -25,12 +60,13 @@ double IAudioDataRenderer::CalculateLerpData(FLerpData* LerpArray, int32 InNumPo
 	uint32 LoopEndFrameIndex = 0;
 	bool IsLooping = false;
 	bool HasTailSection = false;
-	if (SoundWaveProxy->GetLoopRegions().Num() > 0)
+	if (SoundWaveData->GetLoopRegions().Num() > 0)
 	{
-		const FSoundWaveCuePoint& LoopRegion = SoundWaveProxy->GetLoopRegions()[0];
+		const uint32 NumFramesClamped = FMath::Max(1, SoundWaveData->GetNumFrames()) - 1;
+		const FSoundWaveCuePoint& LoopRegion = SoundWaveData->GetLoopRegions()[0];
 		LoopStartFrameIndex = LoopRegion.FramePosition;
 		LoopEndFrameIndex = LoopRegion.FramePosition + LoopRegion.FrameLength;
-		HasTailSection = LoopEndFrameIndex != SoundWaveProxy->GetNumFrames() - 1;
+		HasTailSection = LoopEndFrameIndex != NumFramesClamped;
 		IsLooping = (!HasTailSection || (HasTailSection && InHonorLoopPoints));
 	}
 
diff --git a/Engine/Plugins/Runtime/Harmonix/Source/HarmonixDsp/Private/HarmonixDsp/FusionSampler/FusionSampler.cpp b/Engine/Plugins/Runtime/Harmonix/Source/HarmonixDsp/Private/HarmonixDsp/FusionSampler/FusionSampler.cpp
--- a/Engine/Plugins/Runtime/Harmonix/Source/HarmonixDsp/Private/HarmonixDsp/FusionSampler/FusionSampler.cpp
+++ b/Engine/Plugins/Runtime/Harmonix/Source/HarmonixDsp/Private/HarmonixDsp/FusionSampler/FusionSampler.cpp
@@ -383,7 +383,7 @@ int32 FFusionSampler::GatherMatchingKeyzones(uint8 TransposedNote, uint8 InVeloc
 	for (int32 KeyzoneIdx = 0; KeyzoneIdx < InKeyzones.Num(); ++KeyzoneIdx)
 	{
 		const FKeyzoneSettings& Keyzone = InKeyzones[KeyzoneIdx];
-		if (!Keyzone.SoundWaveProxy)
+		if (!Keyzone.SoundWaveData)
 		{
 			continue;
 		}
diff --git a/Engine/Plugins/Runtime/Harmonix/Source/HarmonixDsp/Private/HarmonixDsp/FusionSampler/FusionVoice.cpp b/Engine/Plugins/Runtime/Harmonix/Source/HarmonixDsp/Private/HarmonixDsp/FusionSampler/FusionVoice.cpp
--- a/Engine/Plugins/Runtime/Harmonix/Source/HarmonixDsp/Private/HarmonixDsp/FusionSampler/FusionVoice.cpp
+++ b/Engine/Plugins/Runtime/Harmonix/Source/HarmonixDsp/Private/HarmonixDsp/FusionSampler/FusionVoice.cpp
@@ -66,7 +66,7 @@ bool FFusionVoice::AssignIDs(FFusionSampler* InSampler, const FKeyzoneSettings*
 {
 	check(InSampler);
 	check(InKeyZone);
-	check(InKeyZone->SoundWaveProxy);
+	check(InKeyZone->SoundWaveData);
 
 #if CPUPROFILERTRACE_ENABLED
 	const IConsoleVariable* ProfilingAllGraphsCheat = IConsoleManager::Get().FindConsoleVariable(TEXT("au.MetaSound.ProfileAllGraphs"));
@@ -117,8 +117,7 @@ bool FFusionVoice::AssignIDs(FFusionSampler* InSampler, const FKeyzoneSettings*
 	// so in that case set NumInChannels to 0 here and we'll sort this
 	// out later...
 	// 0 because we need to configure to the output buffer we are asked to fill
-	
-	int32 NumInChannels = InKeyZone->TrackMap.Num() == 0 ? KeyZone->SoundWaveProxy->GetNumChannels() : 0; 
+	int32 NumInChannels = InKeyZone->TrackMap.Num() == 0 ? KeyZone->SoundWaveData->GetNumChannels() : 0; 
 	EAudioBufferChannelLayout ChannelLayout = HarmonixDsp::FAudioBuffer::GetDefaultChannelLayoutForChannelCount(NumInChannels);
 	ESpeakerMask::Type ChannelMask = HarmonixDsp::FAudioBuffer::GetChannelMaskForNumChannels(NumInChannels);
 	Panner.Setup(KeyZone->Pan,
@@ -131,7 +130,7 @@ bool FFusionVoice::AssignIDs(FFusionSampler* InSampler, const FKeyzoneSettings*
 
 	// compare the audio file's sample rate to the output sample rate
 	// so we can make them match
-	double FileSamplesPerSecond = InKeyZone->SoundWaveProxy->GetSampleRate();
+	double FileSamplesPerSecond = KeyZone->SoundWaveData->GetSampleRate();
 	FileToOutputSampleRatio = FileSamplesPerSecond * SecondsPerSample;
 
 	IAudioDataRenderer::FSettings Settings;
@@ -140,9 +139,9 @@ bool FFusionVoice::AssignIDs(FFusionSampler* InSampler, const FKeyzoneSettings*
 	Settings.TrackChannelInfo = &InKeyZone->TrackMap;
 	Settings.Sampler = MySampler;
 	//TRACE_BOOKMARK(TEXT("Starting Waveform: %s"), *KeyZone->AudioSample->GetName().ToString());
-	ActiveRenderer->SetAudioData(InKeyZone->SoundWaveProxy.ToSharedRef(), Settings);
+	ActiveRenderer->SetSoundWaveData(KeyZone->SoundWaveData.ToSharedRef(), Settings);
 
-	OutputBuffer.SetNumValidChannels(InKeyZone->SoundWaveProxy->GetNumChannels());
+	OutputBuffer.SetNumValidChannels(KeyZone->SoundWaveData->GetNumChannels());
 
 	VoiceID = InVoiceID;
 	bWaitingForAttack = true;
@@ -160,7 +159,7 @@ void FFusionVoice::Attack()
 {
 	check(MySampler);
 	checkSlow(KeyZone);
-	checkSlow(KeyZone->SoundWaveProxy);
+	checkSlow(KeyZone->SoundWaveData);
 
 	AdsrVolume.Attack();
 	AdsrAssignable.Attack();
@@ -284,7 +283,7 @@ void FFusionVoice::SetSampler(FFusionSampler* InSampler)
 
 void FFusionVoice::SetPitchShiftForMidiNote(uint8 InNote)
 {
-	check(KeyZone);
+	check(KeyZone && KeyZone->SoundWaveData);
 
 	// this is the note (pitch) we are supposedly playing
 	// although the exact pitch will vary depending on various modulations
@@ -396,14 +395,14 @@ void FFusionVoice::AttackWithTargetNote(uint8 InMidiNoteNumber, float InGain, in
 	double KeyzoneStartOffset = (KeyZone->SampleStartOffset == -1) ?
 		0 :
 		KeyZone->SampleStartOffset;
-
+	
 	SamplePos = KeyzoneStartOffset + (OutputSampleOffset * FileToOutputSampleRatio);
 	StartPos = SamplePos;
 	EndOfSampleData = (KeyZone->SampleEndOffset == -1) ?
-		KeyZone->SoundWaveProxy->GetNumFrames() :
+		KeyZone->SoundWaveData->GetNumFrames() :
 		KeyZone->SampleEndOffset;
 
-	bool HasLoopSection = KeyZone->SoundWaveProxy->GetLoopRegions().Num() > 0;
+	bool HasLoopSection = KeyZone->SoundWaveData->GetLoopRegions().Num() > 0;
 	if (SamplePos > EndOfSampleData && !HasLoopSection)
 	{
 		Kill();
@@ -436,7 +435,7 @@ void FFusionVoice::PrepareWithPitchOffsetAndGain(double InPitchOffsetCents, floa
 
 	checkSlow(MySampler);
 	checkSlow(KeyZone);
-	checkSlow(KeyZone->SoundWaveProxy);
+	checkSlow(KeyZone->SoundWaveData);
 
 	for (int32 Idx = 0; Idx < kNumLfos; ++Idx)
 	{
@@ -675,17 +674,7 @@ uint32 FFusionVoice::Process(uint32 InSliceIndex, uint32 InSubsliceIndex, float*
 
 	using namespace Harmonix::Dsp::Modulators;
 
-	if (!IsInUse())
-	{
-		return 0;
-	}
-
-	if (!MySampler || !KeyZone)
-	{
-		return 0;
-	}
-
-	if (!KeyZone->SoundWaveProxy)
+	if (!IsInUse() || !MySampler || !KeyZone || !KeyZone->SoundWaveData)
 	{
 		return 0;
 	}
@@ -777,8 +766,8 @@ uint32 FFusionVoice::Process(uint32 InSliceIndex, uint32 InSubsliceIndex, float*
 				{
 					CurrentSampleFrame -= KeyZone->SampleStartOffset;
 				}
-
-				float ActualElapsedMs = (CurrentSampleFrame * 1000.0f) / KeyZone->SoundWaveProxy->GetSampleRate();
+				
+				float ActualElapsedMs = (CurrentSampleFrame * 1000.0f) / FMath::Max(1, KeyZone->SoundWaveData->GetSampleRate());
 				float ErrorMs = ActualElapsedMs - ExpectedElapsedMs;
 
 				if (!bHasRenderedAnySamples)
@@ -788,11 +777,11 @@ uint32 FFusionVoice::Process(uint32 InSliceIndex, uint32 InSubsliceIndex, float*
 						UE_LOG(LogFusionVoice, Verbose, TEXT("Adjusting start time due to error of %f ms"), ErrorMs);
 					}
 
-					SamplePos = ExpectedElapsedMs / 1000.0f * (float)KeyZone->SoundWaveProxy->GetSampleRate();
+					SamplePos = ExpectedElapsedMs / 1000.0f * (float)KeyZone->SoundWaveData->GetSampleRate();
 					
 					if (SamplePos > 0.0)
 					{
-						UE_LOG(LogFusionVoice, Verbose, TEXT("SHAVING %f ms from the beginning of %s"), ErrorMs, *KeyZone->SoundWaveProxy->GetFName().ToString());
+						UE_LOG(LogFusionVoice, Verbose, TEXT("SHAVING %f ms from the beginning of %s"), ErrorMs, *KeyZone->SoundWaveData->GetFName().ToString());
 					}
 
 					if (SamplePos < 0.0)
@@ -803,7 +792,7 @@ uint32 FFusionVoice::Process(uint32 InSliceIndex, uint32 InSubsliceIndex, float*
 					{
 						SamplePos += KeyZone->SampleStartOffset;
 					}
-					if (SamplePos > EndOfSampleData && KeyZone->SoundWaveProxy->GetLoopRegions().IsEmpty())
+					if (SamplePos > EndOfSampleData && KeyZone->SoundWaveData->GetLoopRegions().IsEmpty())
 					{
 						Kill();
 						return 0;
@@ -815,7 +804,7 @@ uint32 FFusionVoice::Process(uint32 InSliceIndex, uint32 InSubsliceIndex, float*
 				{
 					if (ErrorMs > 20.0f)
 					{
-						UE_LOG(LogFusionVoice, Verbose, TEXT("Fixing %s offset -> %f\n"), *KeyZone->SoundWaveProxy->GetFName().ToString(), ErrorMs);
+						UE_LOG(LogFusionVoice, Verbose, TEXT("Fixing %s offset -> %f\n"), *KeyZone->SoundWaveData->GetFName().ToString(), ErrorMs);
 					}
 					CurrentVso = -(FusionVoice::kVsoAmount * RenderSpeed);
 				}
@@ -823,7 +812,7 @@ uint32 FFusionVoice::Process(uint32 InSliceIndex, uint32 InSubsliceIndex, float*
 				{
 					if (ErrorMs < -20.0f)
 					{
-						UE_LOG(LogFusionVoice, Verbose, TEXT("Fixing %s offset -> %f\n"), *KeyZone->SoundWaveProxy->GetFName().ToString(), ErrorMs);
+						UE_LOG(LogFusionVoice, Verbose, TEXT("Fixing %s offset -> %f\n"), *KeyZone->SoundWaveData->GetFName().ToString(), ErrorMs);
 					}
 					CurrentVso = FusionVoice::kVsoAmount * RenderSpeed;
 				}
@@ -855,7 +844,7 @@ uint32 FFusionVoice::Process(uint32 InSliceIndex, uint32 InSubsliceIndex, float*
 
 		// We probably have a trackmap, and this is the first render call,
 		// so we have to configure the panner for the output buffer setup...
-		int32 NumInChannels = KeyZone->SoundWaveProxy->GetNumChannels();
+		int32 NumInChannels = KeyZone->SoundWaveData->GetNumChannels();
 		if (NumInChannels != Panner.GetCurrentGainMatrix().GetNumInChannels())
 		{
 			Panner.Configure(NumInChannels, OutputBuffer);
diff --git a/Engine/Plugins/Runtime/Harmonix/Source/HarmonixDsp/Private/HarmonixDsp/FusionSampler/FusionVoicePool.cpp b/Engine/Plugins/Runtime/Harmonix/Source/HarmonixDsp/Private/HarmonixDsp/FusionSampler/FusionVoicePool.cpp
--- a/Engine/Plugins/Runtime/Harmonix/Source/HarmonixDsp/Private/HarmonixDsp/FusionSampler/FusionVoicePool.cpp
+++ b/Engine/Plugins/Runtime/Harmonix/Source/HarmonixDsp/Private/HarmonixDsp/FusionSampler/FusionVoicePool.cpp
@@ -294,7 +294,7 @@ FFusionVoice* FFusionVoicePool::GetFreeVoice(
 		return nullptr;
 	}
 
-	if (!InKeyzone->SoundWaveProxy)
+	if (!InKeyzone->SoundWaveData)
 	{
 		UE_LOG(LogFusionVoicePool, Warning, TEXT("Asked to allocate a InSampler voice for an invalid sample"));
 		return nullptr;
@@ -308,7 +308,7 @@ FFusionVoice* FFusionVoicePool::GetFreeVoice(
 			FName FactoryName = InKeyzone->TimeStretchConfig.PitchShifter.Name;
 			UE_LOG(LogFusionVoicePool, Warning,
 				TEXT("Attempting to play a fusion InKeyzone that is set to \"maintain time\", but no shifter factory was available for assigned Pitch Shifter: %s. Check project configuration! (%s)"),
-				*FactoryName.ToString(), *InKeyzone->SoundWaveProxy->GetFName().ToString());
+				*FactoryName.ToString(), *InKeyzone->SoundWaveData->GetFName().ToString());
 			return nullptr;
 		}
 
diff --git a/Engine/Plugins/Runtime/Harmonix/Source/HarmonixDsp/Private/HarmonixDsp/FusionSampler/Settings/KeyzoneSettings.cpp b/Engine/Plugins/Runtime/Harmonix/Source/HarmonixDsp/Private/HarmonixDsp/FusionSampler/Settings/KeyzoneSettings.cpp
--- a/Engine/Plugins/Runtime/Harmonix/Source/HarmonixDsp/Private/HarmonixDsp/FusionSampler/Settings/KeyzoneSettings.cpp
+++ b/Engine/Plugins/Runtime/Harmonix/Source/HarmonixDsp/Private/HarmonixDsp/FusionSampler/Settings/KeyzoneSettings.cpp
@@ -19,9 +19,9 @@ FKeyzoneSettings::FKeyzoneSettings()
 
 FKeyzoneSettings::~FKeyzoneSettings()
 {
-	if (SoundWaveProxy && SoundWaveProxy->GetLoadingBehavior() == ESoundWaveLoadingBehavior::ForceInline)
+	if (SoundWaveData.IsValid() && SoundWaveData->GetLoadingBehavior() == ESoundWaveLoadingBehavior::ForceInline)
 	{
-		IStreamingManager::Get().GetAudioStreamingManager().RemoveForceInlineSoundWave(SoundWaveProxy);
+		IStreamingManager::Get().GetAudioStreamingManager().RemoveForceInlineSoundWave(SoundWaveData.ToSharedRef());
 	}
 }
 
@@ -37,19 +37,20 @@ void FKeyzoneSettings::InitProxyData(const Audio::FProxyDataInitParams& InitPara
 		return;
 	}
 
-	if (SoundWaveProxy && SoundWaveProxy->GetLoadingBehavior() == ESoundWaveLoadingBehavior::ForceInline)
+	if (SoundWaveData.IsValid() && SoundWaveData->GetLoadingBehavior() == ESoundWaveLoadingBehavior::ForceInline)
 	{
 		// I don't think we'de ever hit this code, but let's just do this to be sure.
-		IStreamingManager::Get().GetAudioStreamingManager().RemoveForceInlineSoundWave(SoundWaveProxy);
+		IStreamingManager::Get().GetAudioStreamingManager().RemoveForceInlineSoundWave(SoundWaveData.ToSharedRef());
 	}
 	
 	TSharedPtr<Audio::IProxyData> AudioProxy = SoundWave->CreateProxyData(InitParams);
-	SoundWaveProxy = StaticCastSharedPtr<FSoundWaveProxy>(AudioProxy);
-	check(SoundWaveProxy);
+	TSharedPtr<FSoundWaveProxy> TempSoundWaveProxy = StaticCastSharedPtr<FSoundWaveProxy>(AudioProxy);
+	SoundWaveData = TempSoundWaveProxy->GetSoundWaveDataRef();
+	check(SoundWaveData);
 		
-	if (SoundWaveProxy->GetLoadingBehavior() == ESoundWaveLoadingBehavior::ForceInline)
+	if (SoundWaveData->GetLoadingBehavior() == ESoundWaveLoadingBehavior::ForceInline)
 	{
-		IStreamingManager::Get().GetAudioStreamingManager().AddForceInlineSoundWave(SoundWaveProxy);
+		IStreamingManager::Get().GetAudioStreamingManager().AddForceInlineSoundWave(SoundWaveData.ToSharedRef());
 	}
 
 	// nulling out sound wave asset for this Keyzone
diff --git a/Engine/Plugins/Runtime/Harmonix/Source/HarmonixDsp/Private/HarmonixDsp/Stretchers/SmbPitchShifter.cpp b/Engine/Plugins/Runtime/Harmonix/Source/HarmonixDsp/Private/HarmonixDsp/Stretchers/SmbPitchShifter.cpp
--- a/Engine/Plugins/Runtime/Harmonix/Source/HarmonixDsp/Private/HarmonixDsp/Stretchers/SmbPitchShifter.cpp
+++ b/Engine/Plugins/Runtime/Harmonix/Source/HarmonixDsp/Private/HarmonixDsp/Stretchers/SmbPitchShifter.cpp
@@ -189,9 +189,9 @@ void FSmbPitchShifter::SetSampleRateAndReset(float InSampleRate)
 
 void FSmbPitchShifter::SetupAndResetImpl(const FTimeStretchConfig* ShifterSettings)
 {
-	check(MyAudioData);
-	SampleRate = MyAudioData->GetSampleRate();
-	FreqPerFftBin = (double)MyAudioData->GetSampleRate() / (double)FftFrameSize;
+	check(MySoundWaveData);
+	SampleRate = MySoundWaveData->GetSampleRate();
+	FreqPerFftBin = (double)MySoundWaveData->GetSampleRate() / (double)FftFrameSize;
 	Reset();
 }
 
diff --git a/Engine/Plugins/Runtime/Harmonix/Source/HarmonixDsp/Public/HarmonixDsp/AudioDataRenderer.h b/Engine/Plugins/Runtime/Harmonix/Source/HarmonixDsp/Public/HarmonixDsp/AudioDataRenderer.h
--- a/Engine/Plugins/Runtime/Harmonix/Source/HarmonixDsp/Public/HarmonixDsp/AudioDataRenderer.h
+++ b/Engine/Plugins/Runtime/Harmonix/Source/HarmonixDsp/Public/HarmonixDsp/AudioDataRenderer.h
@@ -5,12 +5,15 @@
 
 #include "HarmonixDsp/Streaming/TrackChannelInfo.h"
 
+DECLARE_LOG_CATEGORY_EXTERN(LogHarmonixAudioDataRenderer, Log, All);
+
 class IStretcherAndPitchShifter;
 struct FTrackChannelInfo;
 template<typename Type> 
 class TAudioBuffer;
 class FGainMatrix;
 class FFusionSampler;
+class FSoundWaveData;
 class FSoundWaveProxy;
 struct FTimeStretchConfig;
 
@@ -29,11 +32,17 @@ public:
 
 	virtual ~IAudioDataRenderer() = default;
 
-	virtual void SetAudioData(TSharedRef<FSoundWaveProxy> SoundWaveProxy, const FSettings& InSettings) = 0;
+	UE_DEPRECATED(5.8, "Mutable get/set is deprecated, please call SetSoundWaveData() instead")
+	virtual void SetAudioData(TSharedRef<FSoundWaveProxy> SoundWaveProxy, const FSettings& InSettings) {};
+
+	UE_DEPRECATED(5.8, "Mutable get/set is deprecated, please call GetSoundWaveData() instead")
+	virtual const TSharedPtr<FSoundWaveProxy> GetAudioData() const { return {}; };
+
+	virtual void SetSoundWaveData(const TSharedRef<const FSoundWaveData>& InSoundWaveData, const FSettings& InSettings);
+	virtual TSharedPtr<const FSoundWaveData> GetSoundWaveData() const;
 
 	virtual void Reset() = 0;
 	
-	virtual const TSharedPtr<FSoundWaveProxy> GetAudioData() const = 0;
 
 	virtual void MigrateToSampler(const FFusionSampler* InSampler) = 0;
 
diff --git a/Engine/Plugins/Runtime/Harmonix/Source/HarmonixDsp/Public/HarmonixDsp/FusionSampler/Settings/KeyzoneSettings.h b/Engine/Plugins/Runtime/Harmonix/Source/HarmonixDsp/Public/HarmonixDsp/FusionSampler/Settings/KeyzoneSettings.h
--- a/Engine/Plugins/Runtime/Harmonix/Source/HarmonixDsp/Public/HarmonixDsp/FusionSampler/Settings/KeyzoneSettings.h
+++ b/Engine/Plugins/Runtime/Harmonix/Source/HarmonixDsp/Public/HarmonixDsp/FusionSampler/Settings/KeyzoneSettings.h
@@ -69,13 +69,13 @@ struct FKeyzoneSettings
 	}
 
 public:
-	
 	// Not null when this keyzone is not being proxied
-	// Made private since we generally want the Proxy Data
 	UPROPERTY(EditDefaultsOnly, Category = "Audio Sample")
 	TObjectPtr<USoundWave> SoundWave;
 
 	TSharedPtr<FSoundWaveProxy> SoundWaveProxy;
+	TSharedPtr<const FSoundWaveData> SoundWaveData;
+
 
 	UPROPERTY(EditDefaultsOnly, Category = "Settings", Meta = (UIMin=0, UIMax=127, ClampMin=0, ClampMax=127))
 	int8 RootNote = 60;
diff --git a/Engine/Plugins/Runtime/Harmonix/Source/HarmonixDsp/Public/HarmonixDsp/StretcherAndPitchShifter.h b/Engine/Plugins/Runtime/Harmonix/Source/HarmonixDsp/Public/HarmonixDsp/StretcherAndPitchShifter.h
--- a/Engine/Plugins/Runtime/Harmonix/Source/HarmonixDsp/Public/HarmonixDsp/StretcherAndPitchShifter.h
+++ b/Engine/Plugins/Runtime/Harmonix/Source/HarmonixDsp/Public/HarmonixDsp/StretcherAndPitchShifter.h
@@ -6,6 +6,7 @@
 #include "HarmonixDsp/AudioBuffer.h"
 #include "HarmonixDsp/PitchShifterName.h"
 #include "HarmonixDsp/Containers/TypedParameter.h"
+#include "Sound/SoundWave.h"
 
 #include "StretcherAndPitchShifter.generated.h"
 
@@ -40,7 +41,10 @@ class IStretcherAndPitchShifter
 {
 public:
 	IStretcherAndPitchShifter(FName InFactoryName) : MyFactoryName(InFactoryName) {}
+	
+	PRAGMA_DISABLE_DEPRECATION_WARNINGS
 	virtual ~IStretcherAndPitchShifter() = default;
+	PRAGMA_ENABLE_DEPRECATION_WARNINGS
 
 	FName GetFactoryName() const { return MyFactoryName; };
 
@@ -53,9 +57,27 @@ public:
 	virtual bool HasCurrentSampleFrame() const { return false; }
 	virtual int32  GetCurrentSampleFrame() const { return 0; }
 
-	void SetSampleSourceReset(TSharedPtr<FSoundWaveProxy> InSoundWave, TSharedPtr<IAudioDataRenderer, ESPMode::ThreadSafe> InAudioRenderer, const FTimeStretchConfig* ShifterSettings)
+
+	void SetSampleSourceReset(const TSharedRef<const FSoundWaveData> InSoundWaveData, TSharedPtr<IAudioDataRenderer, ESPMode::ThreadSafe> InAudioRenderer, const FTimeStretchConfig* ShifterSettings)
 	{ 
-		MyAudioData = InSoundWave;
+		MySoundWaveData = InSoundWaveData;
+
+		MyAudioRenderer = InAudioRenderer;
+
+		SetupAndResetImpl(ShifterSettings);
+	}
+
+	UE_DEPRECATED(5.8, "Please use the TSharedRef<const FSoundWaveData> overload instead")
+	void SetSampleSourceReset(TSharedPtr<FSoundWaveProxy> InSoundWave, TSharedPtr<IAudioDataRenderer, ESPMode::ThreadSafe> InAudioRenderer, const FTimeStretchConfig* ShifterSettings)
+	{
+		if (InSoundWave.IsValid())
+		{
+			SetSampleSourceReset(InSoundWave->GetSoundWaveDataRef(), InAudioRenderer, ShifterSettings);
+
+			return;
+		}
+
+		MySoundWaveData.Reset();
 		MyAudioRenderer = InAudioRenderer;
 
 		SetupAndResetImpl(ShifterSettings);
@@ -74,7 +96,11 @@ public:
 
 	virtual void Cleanup() = 0;
 protected:
+	TSharedPtr<const FSoundWaveData> MySoundWaveData;
+	
+	UE_DEPRECATED(5.8, "Please use MySoundWaveData instead")
 	TSharedPtr<FSoundWaveProxy> MyAudioData;
+	
 	TSharedPtr<IAudioDataRenderer, ESPMode::ThreadSafe> MyAudioRenderer;
 
 	virtual void SetupAndResetImpl(const FTimeStretchConfig* ShifterSettings) { }
diff --git a/Engine/Plugins/Runtime/Metasound/Source/MetasoundEngine/Private/MetasoundWave.cpp b/Engine/Plugins/Runtime/Metasound/Source/MetasoundEngine/Private/MetasoundWave.cpp
--- a/Engine/Plugins/Runtime/Metasound/Source/MetasoundEngine/Private/MetasoundWave.cpp
+++ b/Engine/Plugins/Runtime/Metasound/Source/MetasoundEngine/Private/MetasoundWave.cpp
@@ -21,17 +21,19 @@ FAutoConsoleVariableRef CVarDisableMetasoundWaveAssetCachePriming(
 
 namespace Metasound
 {
-	FWaveAsset::FWaveAsset(const TSharedPtr<Audio::IProxyData>& InInitData)
+	FWaveAsset::FWaveAsset(const TSharedPtr<const Audio::IProxyData>& InInitData)
 	{
 		if (InInitData.IsValid())
 		{
 			if (InInitData->CheckTypeCast<FSoundWaveProxy>())
 			{
-				// should we be getting handed a SharedPtr here?
-				SoundWaveProxy = MakeShared<FSoundWaveProxy, ESPMode::ThreadSafe>(InInitData->GetAs<FSoundWaveProxy>());
+				// todo: remove const cast once deprecated getters can be removed and SoundWaveProxy is const.
+				SoundWaveProxy = ConstCastSharedPtr<FSoundWaveProxy>(InInitData->GetAs<const FSoundWaveProxy>().AsShared().ToSharedPtr());
 
 				if (ensureAlways(SoundWaveProxy.IsValid()))
 				{
+					TSharedRef<const FSoundWaveData> SoundWaveData = SoundWaveProxy->GetSoundWaveDataRef();
+					
 					// TODO HACK: Prime the sound for playback.
 					//
 					// Preferably playback latency would be controlled externally.
@@ -50,14 +52,14 @@ namespace Metasound
 					// time that the decoder attempts to decode audio.
 					if (0 == DisableMetasoundWaveAssetCachePriming)
 					{
-						if (SoundWaveProxy->IsStreaming())
+						if (SoundWaveData->IsStreaming())
 						{
-							if (SoundWaveProxy->GetNumChunks() > 1)
+							if (SoundWaveData->GetNumChunks() > 1)
 							{
 								if (FStreamingManagerCollection* StreamingMgr = IStreamingManager::Get_Concurrent())
 								{
 									IAudioStreamingManager& AudioStreamingMgr = StreamingMgr->GetAudioStreamingManager();
-									AudioStreamingMgr.RequestChunk(SoundWaveProxy, 1, [](EAudioChunkLoadResult) {});
+									AudioStreamingMgr.RequestChunk(SoundWaveData->AsShared(), 1, [](EAudioChunkLoadResult) {});
 								}
 							}
 						}
diff --git a/Engine/Plugins/Runtime/Metasound/Source/MetasoundEngine/Private/MetasoundWaveInfoNode.cpp b/Engine/Plugins/Runtime/Metasound/Source/MetasoundEngine/Private/MetasoundWaveInfoNode.cpp
--- a/Engine/Plugins/Runtime/Metasound/Source/MetasoundEngine/Private/MetasoundWaveInfoNode.cpp
+++ b/Engine/Plugins/Runtime/Metasound/Source/MetasoundEngine/Private/MetasoundWaveInfoNode.cpp
@@ -60,7 +60,7 @@ namespace Metasound
 		FStringWriteRef PathOutput;
 
 		// other
-		FSoundWaveProxyPtr SoundWaveProxy;
+		TSharedPtr<const FSoundWaveData> SoundWaveData;
 
 	}; // class FWaveInfoNodeOperator
 
@@ -143,21 +143,27 @@ namespace Metasound
 
 	void FWaveInfoNodeOperator::Execute()
 	{
-		FSoundWaveProxyPtr ProxyPtr = (*WaveAsset).GetSoundWaveProxy();
+		TSharedPtr<const FSoundWaveProxy> ProxyPtr = (*WaveAsset).GetWaveProxy();
+		TSharedRef<const FSoundWaveData> DataPtr;
+		if (ProxyPtr)
+		{
+			DataPtr = ProxyPtr->GetSoundWaveDataRef();
+		}
+		
 
-		if (SoundWaveProxy != ProxyPtr)
+		if (SoundWaveData != DataPtr)
 		{
-			SoundWaveProxy = ProxyPtr;
+			SoundWaveData = DataPtr;
 
-			if (SoundWaveProxy.IsValid() && (*WaveAsset).IsSoundWaveValid())
+			if (SoundWaveData.IsValid() && (*WaveAsset).IsSoundWaveValid())
 			{
-				*DurationSeconds = FTime::FromSeconds((*WaveAsset)->GetDuration());
+				*DurationSeconds = FTime::FromSeconds(DataPtr->GetDuration());
 
-				FName AssetName = SoundWaveProxy->GetFName();
+				FName AssetName = SoundWaveData->GetFName();
 				*NameOutput = AssetName.ToString();
 
 				FString FullPath;
-				SoundWaveProxy->GetPackageName().AppendString(FullPath);
+				SoundWaveData->GetPackageName().AppendString(FullPath);
 				if (!AssetName.IsNone())
 				{
 					FullPath += TEXT(".");
diff --git a/Engine/Plugins/Runtime/Metasound/Source/MetasoundEngine/Private/MetasoundWavePlayerNode.cpp b/Engine/Plugins/Runtime/Metasound/Source/MetasoundEngine/Private/MetasoundWavePlayerNode.cpp
--- a/Engine/Plugins/Runtime/Metasound/Source/MetasoundEngine/Private/MetasoundWavePlayerNode.cpp
+++ b/Engine/Plugins/Runtime/Metasound/Source/MetasoundEngine/Private/MetasoundWavePlayerNode.cpp
@@ -533,12 +533,12 @@ namespace Metasound
 			// Reset flag to track render cost.
 			bDidWaveRenderThisBlock = false;
 
-#if PLATFORM_DESKTOP && !UE_BUILD_SHIPPING
-			if (bIsPlaying && WaveAsset->IsSoundWaveValid() && WaveAsset->GetSoundWaveProxy()->TryGetLatest())
+ #if UE_USOUNDWAVE_HOT_RELOADING_ENABLED
+			if (bIsPlaying && WaveAsset->IsSoundWaveValid() && !WaveAsset->GetWaveProxy()->IsLatestData(CurrentSoundWaveData.Get()))
 			{
 				StartPlaying();
 			}
-#endif // #if PLATFORM_DESKTOP && !UE_BUILD_SHIPPING
+ #endif // #if UE_USOUNDWAVE_HOT_RELOADING_ENABLED
 
 			// Update wave proxy reader with any new looping bounds. 
 			if (WaveProxyReader.IsValid())
@@ -813,7 +813,7 @@ namespace Metasound
 			// Copy the wave asset off on init in case the user changes it while we're playing it.
 			// We'll only check for new wave assets when the current one finishes for sample accurate concantenation
 			CurrentWaveAsset = *WaveAsset;
-			FSoundWaveProxyPtr WaveProxy = CurrentWaveAsset.GetSoundWaveProxy();
+			TSharedPtr<const FSoundWaveProxy> WaveProxy = CurrentWaveAsset.GetWaveProxy();
 
 			bOnNearlyDoneTriggeredForWave = false;
 			bIsPlaying = false;
@@ -828,14 +828,12 @@ namespace Metasound
 			if (WaveProxy.IsValid())
 			{
 				// Make sure we have up-to-date SoundWaveData
-#if PLATFORM_DESKTOP
-				WaveProxy->TryGetLatest();
-#endif // #if PLATFORM_DESKTOP
+				CurrentSoundWaveData = WaveProxy->GetSoundWaveDataRef();
 
-				UE_LOG(LogMetaSound, Verbose, TEXT("Starting Sound: '%s'"), *CurrentWaveAsset->GetFName().ToString());
+				UE_LOG(LogMetaSound, Verbose, TEXT("Starting Sound: '%s'"), *CurrentSoundWaveData->GetFName().ToString());
 
 				// Create local sorted copy of cue points and loop regions.
-				SortedCuePoints = WaveProxy->GetCuePointsAndLoopRegions();
+				SortedCuePoints = CurrentSoundWaveData->GetCuePointsAndLoopRegions();
 				Algo::SortBy(SortedCuePoints, WavePlayerNodePrivate::GetCuePointFrame);
 				
 				// Create the wave proxy reader.
@@ -849,7 +847,7 @@ namespace Metasound
 				WaveReaderSettings.bIsLooping = *bLoop;
 				WaveReaderSettings.bMaintainAudioSync = bMaintainAudioSync;
 
-				WaveProxyReader = FSoundWaveProxyReader::Create(WaveProxy.ToSharedRef(), WaveReaderSettings);
+				WaveProxyReader = FSoundWaveProxyReader::Create(CurrentSoundWaveData.ToSharedRef(), WaveReaderSettings);
 
 				DeinterleaveBufferBlockSizeInFrames = FMath::IsPowerOfTwo(DeinterleaveBlockSizeInFrames) ? 
 					DeinterleaveBlockSizeInFrames : FMath::RoundUpToPowerOfTwo(DeinterleaveBlockSizeInFrames);
@@ -1135,7 +1133,7 @@ namespace Metasound
 				int32 NumFramesProduced = Resampler->ProcessAndConsumeAudio(SourceCircularBuffer, OutBuffer);
 				if (NumFramesProduced < 1)
 				{
-					UE_LOG(LogMetaSound, Error, TEXT("Aborting currently playing metasound wave %s. Failed to produce any resampled audio frames with %d input frames and a frame ratio of %f."), *CurrentWaveAsset->GetFName().ToString(), NumSourceFramesAvailable, GetFrameRatio());
+					UE_LOG(LogMetaSound, Error, TEXT("Aborting currently playing metasound wave %s. Failed to produce any resampled audio frames with %d input frames and a frame ratio of %f."), *CurrentSoundWaveData->GetFName().ToString(), NumSourceFramesAvailable, GetFrameRatio());
 					bIsPlaying = false;
 					break;
 				}
@@ -1187,6 +1185,7 @@ namespace Metasound
 		TUniquePtr<Audio::FMultichannelLinearResampler> Resampler;
 
 		FWaveAsset CurrentWaveAsset;
+		TSharedPtr<const FSoundWaveData> CurrentSoundWaveData;
 		TArray<FSoundWaveCuePoint> SortedCuePoints;
 		Audio::FAlignedFloatBuffer InterleavedBuffer;
 		Audio::FMultichannelBuffer DeinterleavedBuffer;
diff --git a/Engine/Plugins/Runtime/Metasound/Source/MetasoundEngine/Private/MetasoundWavePlayerNodeV2.cpp b/Engine/Plugins/Runtime/Metasound/Source/MetasoundEngine/Private/MetasoundWavePlayerNodeV2.cpp
--- a/Engine/Plugins/Runtime/Metasound/Source/MetasoundEngine/Private/MetasoundWavePlayerNodeV2.cpp
+++ b/Engine/Plugins/Runtime/Metasound/Source/MetasoundEngine/Private/MetasoundWavePlayerNodeV2.cpp
@@ -269,6 +269,8 @@ namespace Metasound
 			ProxyPlayer->Reset(PlayerSettings);
 
 			CurrentWaveAsset = FWaveAsset();
+			CurrentSoundWaveData.Reset();
+
 			bIsPlaying = false;
 			bDidRenderThisBlock = false;
 		}
@@ -402,13 +404,15 @@ namespace Metasound
 
 		void ExecuteSubblocks()
 		{
-#if PLATFORM_DESKTOP && !UE_BUILD_SHIPPING
+#if UE_USOUNDWAVE_HOT_RELOADING_ENABLED
 			// hot reload sound
-			if (bIsPlaying && Inputs.WaveAsset->IsSoundWaveValid() && Inputs.WaveAsset->GetSoundWaveProxy()->TryGetLatest())
+			if (bIsPlaying && Inputs.WaveAsset->IsSoundWaveValid()
+				&& Inputs.WaveAsset->GetWaveProxy()
+				&& !Inputs.WaveAsset->GetWaveProxy()->IsLatestData(CurrentSoundWaveData.Get()))
 			{
 				StartPlaying();
 			}
-#endif // PLATFORM_DESKTOP && !UE_BUILD_SHIPPING
+#endif // UE_USOUNDWAVE_HOT_RELOADING_ENABLED
 			
 			// Parse triggers and render audio
 			int32 PlayTrigIndex = 0;
@@ -495,14 +499,24 @@ namespace Metasound
 		{
 			check(ProxyPlayer.IsValid());
 			CurrentWaveAsset = *Inputs.WaveAsset;
-			if (ProxyPlayer->SetSoundWave(CurrentWaveAsset.GetSoundWaveProxy()))
+			TSharedPtr<const FSoundWaveProxy> Proxy = CurrentWaveAsset.GetWaveProxy();
+			if(Proxy)
+			{
+				CurrentSoundWaveData = Proxy->GetSoundWaveDataRef().ToSharedPtr();
+			}
+			else
+			{
+				CurrentSoundWaveData.Reset();
+			}
+			
+			if (CurrentSoundWaveData && ProxyPlayer->SetSoundWave(CurrentSoundWaveData.ToSharedRef()))
 			{
 				ProxyPlayer->SetSpeed(CalculateSpeedFromPitchShift());
 				ProxyPlayer->SetLoop(*Inputs.Loop, Inputs.LoopStartTime->GetSeconds(), Inputs.LoopDuration->GetSeconds());
 				ProxyPlayer->SeekToTime(Inputs.StartTime->GetSeconds());
 			}
 
-			bIsPlaying = ProxyPlayer.IsValid() && ProxyPlayer->IsPlayerValid() && !ProxyPlayer->IsFinished();
+			bIsPlaying = Proxy.IsValid() && ProxyPlayer.IsValid() && ProxyPlayer->IsPlayerValid() && !ProxyPlayer->IsFinished();
 		}
 
 		void RenderFrameRange(int32 StartFrame, int32 EndFrame)
@@ -580,6 +594,7 @@ namespace Metasound
 		bool bDidRenderThisBlock = false;
 
 		FWaveAsset CurrentWaveAsset;
+		TSharedPtr<const FSoundWaveData> CurrentSoundWaveData;
 		TUniquePtr<FSoundWaveProxyPlayer> ProxyPlayer;
 		TArray<TArrayView<float>> MultiChannelOutputView;
 	};
diff --git a/Engine/Plugins/Runtime/Metasound/Source/MetasoundEngine/Public/MetasoundWave.h b/Engine/Plugins/Runtime/Metasound/Source/MetasoundEngine/Public/MetasoundWave.h
--- a/Engine/Plugins/Runtime/Metasound/Source/MetasoundEngine/Public/MetasoundWave.h
+++ b/Engine/Plugins/Runtime/Metasound/Source/MetasoundEngine/Public/MetasoundWave.h
@@ -24,30 +24,39 @@ namespace Metasound
 		static constexpr bool Value = true;
 	};
 
-	// Metasound data type that holds onto a weak ptr. Mostly used as a placeholder until we have a proper proxy type.
+	// Metasound data type that holds onto a shared ptr. Mostly used as a placeholder until we have a proper proxy type.
 	class FWaveAsset
 	{
-		FSoundWaveProxyPtr SoundWaveProxy;
+		// todo: this can be TSharedPtr<const FSoundWaveProxy> once depprecated functions are removed
+		TSharedPtr<FSoundWaveProxy, ESPMode::ThreadSafe> SoundWaveProxy;
+	
 	public:
 
 		FWaveAsset() = default;
 		FWaveAsset(const FWaveAsset&) = default;
 		FWaveAsset& operator=(const FWaveAsset& Other) = default;
 
-		UE_API FWaveAsset(const TSharedPtr<Audio::IProxyData>& InInitData);
+		UE_API FWaveAsset(const TSharedPtr<const Audio::IProxyData>& InInitData);
 
 		UE_API bool IsSoundWaveValid() const;
 
+		UE_DEPRECATED(5.8, "Use GetWaveProxy() instead, mutable proxy access is deprecated")
 		const FSoundWaveProxyPtr& GetSoundWaveProxy() const
 		{
 			return SoundWaveProxy;
 		}
 
+		const TSharedPtr<const FSoundWaveProxy> GetWaveProxy() const 
+		{
+			return SoundWaveProxy;
+		}
+
 		const FSoundWaveProxy* operator->() const
 		{
 			return SoundWaveProxy.Get();
 		}
 
+		UE_DEPRECATED(5.8, "Use GetWaveProxy() for non-const objects instead. (mutable proxy access is deprecated)")
 		FSoundWaveProxy* operator->()
 		{
 			return SoundWaveProxy.Get();
@@ -57,7 +66,9 @@ namespace Metasound
 		{
 			if (InWaveAsset.IsSoundWaveValid())
 			{
-				return GetTypeHash(*InWaveAsset.GetSoundWaveProxy());
+				const TSharedPtr<const FSoundWaveProxy> Proxy = InWaveAsset.GetWaveProxy();
+				check(Proxy);
+				return GetTypeHash(*Proxy);
 			}
 			return INDEX_NONE;
 		}
diff --git a/Engine/Source/Runtime/AdpcmAudioDecoder/Module/Private/ADPCMAudioInfo.cpp b/Engine/Source/Runtime/AdpcmAudioDecoder/Module/Private/ADPCMAudioInfo.cpp
--- a/Engine/Source/Runtime/AdpcmAudioDecoder/Module/Private/ADPCMAudioInfo.cpp
+++ b/Engine/Source/Runtime/AdpcmAudioDecoder/Module/Private/ADPCMAudioInfo.cpp
@@ -123,12 +123,14 @@ void FADPCMAudioInfo::SeekToFrameInternal(const uint32 InSeekFrame)
 
 	// Calculate block index & force SeekTime to be in bounds.
 	check(WaveInfo.pSamplesPerSec != nullptr);
+	check(TotalSamplesPerChannel > 1);
+
 	uint32 SeekedSamples = InSeekFrame;
 	TotalSamplesStreamed = FMath::Min<uint32>(SeekedSamples, TotalSamplesPerChannel - 1);
 
 	const uint32 HeaderOffset = static_cast<uint32>(WaveInfo.SampleDataStart - SrcBufferData);
 	
-	if (!StreamingSoundWave)
+	if (!StreamingSoundWaveData)
 	{
 		// For the non streaming case:
 		if (Format == WAVE_FORMAT_LPCM)
@@ -151,13 +153,13 @@ void FADPCMAudioInfo::SeekToFrameInternal(const uint32 InSeekFrame)
 	}
 	else
 	{
-		if (StreamingSoundWave->GetNumChunks() == 0)
+		if (StreamingSoundWaveData->GetNumChunks() == 0)
 		{
 			UE_LOG(LogAudio, Error, TEXT("Entered streaming seek path with a non-streaming sound!"));
 			return;
 		}
 
-		const uint32 TotalStreamingChunks = StreamingSoundWave->GetNumChunks();
+		const uint32 TotalStreamingChunks = StreamingSoundWaveData->GetNumChunks();
 
 		if (Format == WAVE_FORMAT_ADPCM)
 		{
@@ -168,7 +170,7 @@ void FADPCMAudioInfo::SeekToFrameInternal(const uint32 InSeekFrame)
 			const int32 ChannelBlockSize = BlockSize * NumChannels;
 			for (uint32 BlockIndex = 0; BlockIndex < CurrentCompressedBlockIndex; ++BlockIndex)
 			{
-				const uint32 SizeOfChunk = StreamingSoundWave->GetSizeOfChunk(CurrentChunkIndex);
+				const uint32 SizeOfChunk = StreamingSoundWaveData->GetSizeOfChunk(CurrentChunkIndex);
 				if (CurrentChunkBufferOffset + ChannelBlockSize >= SizeOfChunk)
 				{
 					const uint32 RemainderAfterEndOfChunk = CurrentChunkBufferOffset + ChannelBlockSize - SizeOfChunk;
@@ -196,7 +198,7 @@ void FADPCMAudioInfo::SeekToFrameInternal(const uint32 InSeekFrame)
 			const uint32 TotalByteSizeToSeek = TotalSamplesStreamed * ChannelBlockSize;	
 			
 			uint32 TargetChunkIndex = FirstChunkSampleDataIndex;
-			uint32 SizeOfTargetChunk = StreamingSoundWave->GetSizeOfChunk(TargetChunkIndex);
+			uint32 SizeOfTargetChunk = StreamingSoundWaveData->GetSizeOfChunk(TargetChunkIndex);
 			
 			uint32 CurrentBytes = 0;
 			while (CurrentBytes + SizeOfTargetChunk < TotalByteSizeToSeek)
@@ -212,7 +214,7 @@ void FADPCMAudioInfo::SeekToFrameInternal(const uint32 InSeekFrame)
 				++TargetChunkIndex;
 				CurrentBytes += SizeOfTargetChunk;
 				CurrentChunkBufferOffset -= SizeOfTargetChunk;
-				SizeOfTargetChunk = StreamingSoundWave->GetSizeOfChunk(TargetChunkIndex);
+				SizeOfTargetChunk = StreamingSoundWaveData->GetSizeOfChunk(TargetChunkIndex);
 			}
 
 			check(CurrentBytes <= TotalByteSizeToSeek);
@@ -479,21 +481,17 @@ void FADPCMAudioInfo::ResetSeekState()
 	bNewSeekRequest = false;
 }
 
-bool FADPCMAudioInfo::StreamCompressedInfoInternal(const FSoundWaveProxyPtr& InWaveProxy, struct FSoundQualityInfo* QualityInfo)
+bool FADPCMAudioInfo::StreamCompressedInfoInternal(const TSharedRef<const FSoundWaveData>& InWaveData, struct FSoundQualityInfo* QualityInfo)
 {
 	FScopeLock ScopeLock(&CurCompressedChunkHandleCriticalSection);
 
 	check(QualityInfo);
-	check(StreamingSoundWave == InWaveProxy);
-	if (!ensure(InWaveProxy.IsValid()))
-	{
-		return false;
-	}
+	check(StreamingSoundWaveData == InWaveData);
 
 	CurrentChunkIndex = 0;
 
 	// Get the first chunk of audio data (should already be loaded)
-	uint8 const* ChunkData = GetLoadedChunk(InWaveProxy, CurrentChunkIndex, CurrentChunkDataSize);
+	uint8 const* ChunkData = GetLoadedChunk(*InWaveData, CurrentChunkIndex, CurrentChunkDataSize);
 
 	if (ChunkData == nullptr)
 	{
@@ -519,7 +517,7 @@ bool FADPCMAudioInfo::StreamCompressedInfoInternal(const FSoundWaveProxyPtr& InW
 	if (((uint32)SampleDataOffset) >= CurrentChunkDataSize)
 	{
 		++CurrentChunkIndex;
-		ChunkData = GetLoadedChunk(InWaveProxy, CurrentChunkIndex, CurrentChunkDataSize);
+		ChunkData = GetLoadedChunk(*InWaveData, CurrentChunkIndex, CurrentChunkDataSize);
 		FirstChunkSampleDataIndex = CurrentChunkIndex;
 		FirstChunkSampleDataOffset = 0;
 	}
@@ -676,7 +674,7 @@ bool FADPCMAudioInfo::StreamCompressedData(uint8* Destination, bool bLooping, ui
 						bSeekedFowardToNextChunk = false;
 					}
 
-					if (CurrentChunkIndex >= StreamingSoundWave->GetNumChunks())
+					if (CurrentChunkIndex >= StreamingSoundWaveData->GetNumChunks())
 					{
 						// If this is the case, we've seeked to the end of the file.
 						ReachedEndOfSamples = true;
@@ -694,7 +692,7 @@ bool FADPCMAudioInfo::StreamCompressedData(uint8* Destination, bool bLooping, ui
 					}
 
 					// Request the next chunk of data from the streaming engine
-					CurCompressedChunkData = GetLoadedChunk(StreamingSoundWave, CurrentChunkIndex, CurrentChunkDataSize);
+					CurCompressedChunkData = GetLoadedChunk(*StreamingSoundWaveData, CurrentChunkIndex, CurrentChunkDataSize);
 
 					if(CurCompressedChunkData == nullptr)
 					{
@@ -740,7 +738,7 @@ bool FADPCMAudioInfo::StreamCompressedData(uint8* Destination, bool bLooping, ui
 							CurrentChunkBufferOffset += NumCompressedBytesForCallback;
 							TotalSamplesStreamed += NumDecompressedSamplesForCallback;
 
-							const uint32 SizeOfCurrentChunk = StreamingSoundWave->GetSizeOfChunk(CurrentChunkIndex);
+							const uint32 SizeOfCurrentChunk = StreamingSoundWaveData->GetSizeOfChunk(CurrentChunkIndex);
 							if (CurrentChunkBufferOffset >= SizeOfCurrentChunk)
 							{
 								// If this is the case, we've seeked forward past the end of the chunk.
@@ -754,7 +752,7 @@ bool FADPCMAudioInfo::StreamCompressedData(uint8* Destination, bool bLooping, ui
 								CurrentChunkBufferOffset -= SizeOfCurrentChunk;
 
 								// Ensure that we have not moved past the size of our new chunk with our current callback.
-								ensureAlways(CurrentChunkBufferOffset < StreamingSoundWave->GetSizeOfChunk(CurrentChunkIndex));
+								ensureAlways(CurrentChunkBufferOffset < StreamingSoundWaveData->GetSizeOfChunk(CurrentChunkIndex));
 
 								bSeekedFowardToNextChunk = true;
 
@@ -763,7 +761,7 @@ bool FADPCMAudioInfo::StreamCompressedData(uint8* Destination, bool bLooping, ui
 									CurrentChunkBufferOffset = AlignDown(CurrentChunkBufferOffset, CompressedBlockSize);
 								}
 
-								if (CurrentChunkIndex >= StreamingSoundWave->GetNumChunks())
+								if (CurrentChunkIndex >= StreamingSoundWaveData->GetNumChunks())
 								{
 									// If this is the case, we've seeked to the end of the file.
 									ReachedEndOfSamples = true;
@@ -881,7 +879,7 @@ bool FADPCMAudioInfo::StreamCompressedData(uint8* Destination, bool bLooping, ui
 				}
 
 				// Check for the end of the audio samples and loop if needed
-				if (CurrentChunkIndex >= StreamingSoundWave->GetNumChunks())
+				if (CurrentChunkIndex >= StreamingSoundWaveData->GetNumChunks())
 				{
 					ReachedEndOfSamples = true;
 					CurrentChunkIndex = FirstChunkSampleDataIndex;
@@ -897,7 +895,7 @@ bool FADPCMAudioInfo::StreamCompressedData(uint8* Destination, bool bLooping, ui
 				}
 
 				// Request the next chunk of data from the streaming engine
-				CurCompressedChunkData = GetLoadedChunk(StreamingSoundWave, CurrentChunkIndex, CurrentChunkDataSize);
+				CurCompressedChunkData = GetLoadedChunk(*StreamingSoundWaveData, CurrentChunkIndex, CurrentChunkDataSize);
 
 				if(CurCompressedChunkData == nullptr)
 				{
@@ -1011,7 +1009,7 @@ bool FADPCMAudioInfo::ReleaseStreamChunk(bool bBlockUntilReleased)
 	}
 }
 
-const uint8* FADPCMAudioInfo::GetLoadedChunk(const FSoundWaveProxyPtr& InSoundWave, uint32 ChunkIndex, uint32& OutChunkSize)
+const uint8* FADPCMAudioInfo::GetLoadedChunk(const FSoundWaveData& InSoundWaveData, uint32 ChunkIndex, uint32& OutChunkSize)
 {
 	if (ChunkIndex != 0 && FMath::RandRange(0.0f, 1.0f) < ChanceForIntentionalChunkMissCVar)
 	{
@@ -1020,14 +1018,14 @@ const uint8* FADPCMAudioInfo::GetLoadedChunk(const FSoundWaveProxyPtr& InSoundWa
 		return nullptr;
 	}
 
-	if (!ensure(InSoundWave.IsValid()) || ChunkIndex >= InSoundWave->GetNumChunks())
+	if (!ensure(ChunkIndex < InSoundWaveData.GetNumChunks()))
 	{
 		OutChunkSize = 0;
 		return nullptr;
 	}
 	else if (ChunkIndex == 0)
 	{
-		TArrayView<const uint8> ZerothChunk = FSoundWaveProxy::GetZerothChunk(InSoundWave, true);
+		TArrayView<const uint8> ZerothChunk = InSoundWaveData.GetZerothChunk(true);
 		OutChunkSize = ZerothChunk.Num();
 		PreviouslyRequestedChunkIndex = ChunkIndex;
 		return ZerothChunk.GetData();
@@ -1037,7 +1035,7 @@ const uint8* FADPCMAudioInfo::GetLoadedChunk(const FSoundWaveProxyPtr& InSoundWa
 		const bool bIsSeekingOrLooping = bSeekPendingRead || ChunkIndex == FirstChunkSampleDataIndex;
 		ensureAlwaysMsgf(bIsSeekingOrLooping || ChunkIndex == PreviouslyRequestedChunkIndex || ChunkIndex == PreviouslyRequestedChunkIndex + 1, TEXT("ADPCM playback error! We skipped from the end of chunk %d to the beginning of chunk %d."), PreviouslyRequestedChunkIndex, ChunkIndex);
 
-		CurCompressedChunkHandle = IStreamingManager::Get().GetAudioStreamingManager().GetLoadedChunk(InSoundWave, ChunkIndex, false, true);
+		CurCompressedChunkHandle = IStreamingManager::Get().GetAudioStreamingManager().GetLoadedChunk(InSoundWaveData.AsShared(), ChunkIndex, false, true);
 		OutChunkSize = CurCompressedChunkHandle.Num();
 
 		PreviouslyRequestedChunkIndex = ChunkIndex;
diff --git a/Engine/Source/Runtime/AdpcmAudioDecoder/Module/Private/PcmAudioInfoHybrid.cpp b/Engine/Source/Runtime/AdpcmAudioDecoder/Module/Private/PcmAudioInfoHybrid.cpp
--- a/Engine/Source/Runtime/AdpcmAudioDecoder/Module/Private/PcmAudioInfoHybrid.cpp
+++ b/Engine/Source/Runtime/AdpcmAudioDecoder/Module/Private/PcmAudioInfoHybrid.cpp
@@ -342,7 +342,7 @@ bool FPcmAudioInfoHybrid::AbsPositionToChunkIndexAndOffset(const uint64 InAbsPos
 	// But we could make this faster by just special casing zero, first, rest, and last.
 	
 	// If we're streaming, grab the streaming wave instance.
-	if (const FSoundWaveProxyPtr& Wave = GetStreamingSoundWave(); Wave.IsValid() && Wave->IsStreaming())
+	if (const TSharedPtr<const FSoundWaveData>& Wave = GetStreamingSoundWaveData(); Wave.IsValid() && Wave->IsStreaming())
 	{
 		const int32 NumChunks = Wave->GetNumChunks();
 		uint64 AbsChunkOffset = AudioDataOffset;
@@ -369,7 +369,7 @@ bool FPcmAudioInfoHybrid::ParseHeader(const uint8* InSrcBufferData, uint32 InSrc
 	CurrentSampleCount = 0;
 
 	// We only have the header if we're streaming, so just parse that.
-	const bool bParseJustHeader = GetStreamingSoundWave().IsValid() && GetStreamingSoundWave()->IsStreaming();
+	const bool bParseJustHeader = GetStreamingSoundWaveData().IsValid() && GetStreamingSoundWaveData()->IsStreaming();
 
 	void* FormatHeader = nullptr;
 	FWaveModInfo WaveInfo;
@@ -473,6 +473,7 @@ void FPcmAudioInfoHybrid::SeekToAbs(const uint64 InAbsPosition, const uint64 InS
 
 void FPcmAudioInfoHybrid::FHybridImpl::SeekToTime(const float InSeekTime)
 {
+	check(TotalSamplesPerChannel > 0);
 	const uint32 SeekFrame = InSeekTime * SampleRate;
 	const uint32 SeekFrameClamped = FMath::Clamp<uint32>( SeekFrame, 0, TotalSamplesPerChannel - 1);
 	SeekToFrame(SeekFrameClamped);
diff --git a/Engine/Source/Runtime/AdpcmAudioDecoder/Module/Private/PcmAudioInfoHybrid.h b/Engine/Source/Runtime/AdpcmAudioDecoder/Module/Private/PcmAudioInfoHybrid.h
--- a/Engine/Source/Runtime/AdpcmAudioDecoder/Module/Private/PcmAudioInfoHybrid.h
+++ b/Engine/Source/Runtime/AdpcmAudioDecoder/Module/Private/PcmAudioInfoHybrid.h
@@ -78,7 +78,7 @@ protected:
 public:
 	bool IsStreaming() const
 	{
-		return GetStreamingSoundWave().IsValid() && GetStreamingSoundWave()->IsStreaming();
+		return GetStreamingSoundWaveData().IsValid() && GetStreamingSoundWaveData()->IsStreaming();
 	}
 	void SeekToAbs(const uint64 InAbsPosition, const uint64 InSeekFrame);	
 	void SetErrorFlag() const { bHasError = true; }
@@ -93,9 +93,11 @@ private:
 	
 	void SetSeekBlockIndex(const uint32 InStreamSeekBlockIndex)
 	{
-		ensure(GetStreamingSoundWave().IsValid() && GetStreamingSoundWave()->IsStreaming()
-			&& InStreamSeekBlockIndex < GetStreamingSoundWave()->GetNumChunks());
-		StreamSeekBlockIndex = InStreamSeekBlockIndex;
+		if (ensure(GetStreamingSoundWaveData().IsValid() && GetStreamingSoundWaveData()->IsStreaming()
+			&& InStreamSeekBlockIndex < GetStreamingSoundWaveData()->GetNumChunks()))
+		{
+			StreamSeekBlockIndex = InStreamSeekBlockIndex;
+		}
 	};
 	void SetSeekBlockOffset(const int32 InStreamSeekBlockOffset)
 	{
diff --git a/Engine/Source/Runtime/AdpcmAudioDecoder/Module/Public/Decoders/ADPCMAudioInfo.h b/Engine/Source/Runtime/AdpcmAudioDecoder/Module/Public/Decoders/ADPCMAudioInfo.h
--- a/Engine/Source/Runtime/AdpcmAudioDecoder/Module/Public/Decoders/ADPCMAudioInfo.h
+++ b/Engine/Source/Runtime/AdpcmAudioDecoder/Module/Public/Decoders/ADPCMAudioInfo.h
@@ -107,7 +107,7 @@ public:
 
 	// Additional overrides for streaming
 	virtual bool SupportsStreaming() const override {return true;}
-	virtual bool StreamCompressedInfoInternal(const FSoundWaveProxyPtr& InWaveProxy, struct FSoundQualityInfo* QualityInfo) override;
+	virtual bool StreamCompressedInfoInternal(const TSharedRef<const FSoundWaveData>& InWaveData, struct FSoundQualityInfo* QualityInfo) override;
 	virtual bool StreamCompressedData(uint8* Destination, bool bLooping, uint32 BufferSize, int32& OutNumBytesStreamed) override;
 	virtual int32 GetCurrentChunkIndex() const override
 	{
@@ -139,7 +139,7 @@ private:
 	// Wrapper function that returns a pointer to the currently used compressed data.
 	// If a non-zero chunk is requested, this function also aquires a reference to that chunk
 	// until we move on to a different chunk.
-	const uint8* GetLoadedChunk(const FSoundWaveProxyPtr& InSoundWave, uint32 ChunkIndex, uint32& OutChunkSize);
+	const uint8* GetLoadedChunk(const FSoundWaveData& InSoundWaveData, uint32 ChunkIndex, uint32& OutChunkSize);
 
 	int32 NumConsecutiveReadFailiures;
 
diff --git a/Engine/Source/Runtime/AudioMixer/Classes/Generators/SoundWaveScrubber.h b/Engine/Source/Runtime/AudioMixer/Classes/Generators/SoundWaveScrubber.h
--- a/Engine/Source/Runtime/AudioMixer/Classes/Generators/SoundWaveScrubber.h
+++ b/Engine/Source/Runtime/AudioMixer/Classes/Generators/SoundWaveScrubber.h
@@ -27,6 +27,9 @@ public:
 	UE_API ~FSoundWaveScrubber();
 
 	// Initialize the sound wave scrubber
+	UE_API void Init(TSharedRef<const FSoundWaveProxy> InSoundWaveProxyRef, float InSampleRate, int32 InNumChannels, float InPlayheadTimeSeconds = 0.0f);
+
+	UE_DEPRECATED(5.8, "Please use the TSharedRef<const FSoundWaveProxy> overload instead")
 	UE_API void Init(FSoundWaveProxyPtr InSoundWaveProxyPtr, float InSampleRate, int32 InNumChannels, float InPlayheadTimeSeconds = 0.0f);
 
 	// Sets whether or not to scrub the audio file. If false, then the output of the scrubber will be silence.
@@ -85,8 +88,8 @@ private:
 	// Place to store decoded chunks. We need 2 chunks for actively playing grains plus a free chunk.
 	TArray<FDecodedDataChunk> DecodedChunks;
 
-	// Sound wave proxy to safely query and use the referenced sound wave asset
-	FSoundWaveProxyPtr SoundWaveProxyPtr;
+	// Sound wave data to safely query and use the referenced sound wave asset
+	TSharedPtr<const FSoundWaveData> SoundWaveDataPtr;
 	// Sound wave proxy reader is a decoder we can use to decode audio into chunks
 	TUniquePtr<FSoundWaveProxyReader> SoundWaveProxyReaderPtr;
 	// Simple SRC interface. Decoded audio is SRC'd to match the audio mixer sample rate
@@ -173,6 +176,9 @@ private:
 class FSoundWaveScrubberGenerator : public ISoundGenerator
 {
 public:
+	UE_API void Init(TSharedRef<const FSoundWaveProxy> InSoundWaveProxyRef, float InSampleRate, int32 InNumChannels, float InPlayheadTimeSeconds = 0.0f);
+
+	UE_DEPRECATED(5.8, "Please use the TSharedRef<const FSoundWaveProxy> overload instead")
 	UE_API void Init(FSoundWaveProxyPtr InSoundWaveProxyPtr, float InSampleRate, int32 InNumChannels, float InPlayheadTimeSeconds = 0.0f);
 
 	// ~ISoundGenerator Begin
diff --git a/Engine/Source/Runtime/AudioMixer/Private/AudioMixerBlueprintLibrary.cpp b/Engine/Source/Runtime/AudioMixer/Private/AudioMixerBlueprintLibrary.cpp
--- a/Engine/Source/Runtime/AudioMixer/Private/AudioMixerBlueprintLibrary.cpp
+++ b/Engine/Source/Runtime/AudioMixer/Private/AudioMixerBlueprintLibrary.cpp
@@ -602,23 +602,26 @@ void UAudioMixerBlueprintLibrary::PrimeSoundForPlayback(USoundWave* SoundWave, c
 	if (!SoundWave)
 	{
 		UE_LOG(LogAudioMixer, Warning, TEXT("Prime Sound For Playback called with a null SoundWave pointer."));
+		OnLoadCompletion.ExecuteIfBound(SoundWave, true);
+
+		return;
 	}
-	else
+
+	TSharedRef<const FSoundWaveProxy> SoundWaveProxy = SoundWave->GetSoundWaveProxy();
+
+	IStreamingManager::Get().GetAudioStreamingManager().RequestChunk(SoundWaveProxy->GetSoundWaveDataRef(), 1, [OnLoadCompletion, SoundWave](EAudioChunkLoadResult InResult)
 	{
-		IStreamingManager::Get().GetAudioStreamingManager().RequestChunk(SoundWave->CreateSoundWaveProxy(), 1, [OnLoadCompletion, SoundWave](EAudioChunkLoadResult InResult)
-		{
-			AsyncTask(ENamedThreads::GameThread, [OnLoadCompletion, SoundWave, InResult]() {
-				if (InResult == EAudioChunkLoadResult::Completed || InResult == EAudioChunkLoadResult::AlreadyLoaded)
-				{
-					OnLoadCompletion.ExecuteIfBound(SoundWave, false);
-				}
-				else
-				{
-					OnLoadCompletion.ExecuteIfBound(SoundWave, true);
-				}
-			});
+		AsyncTask(ENamedThreads::GameThread, [OnLoadCompletion, SoundWave, InResult]() {
+			if (InResult == EAudioChunkLoadResult::Completed || InResult == EAudioChunkLoadResult::AlreadyLoaded)
+			{
+				OnLoadCompletion.ExecuteIfBound(SoundWave, false);
+			}
+			else
+			{
+				OnLoadCompletion.ExecuteIfBound(SoundWave, true);
+			}
 		});
-	}
+	});
 }
 
 void UAudioMixerBlueprintLibrary::PrimeSoundCueForPlayback(USoundCue* SoundCue)
diff --git a/Engine/Source/Runtime/AudioMixer/Private/Generators/SoundWaveScrubber.cpp b/Engine/Source/Runtime/AudioMixer/Private/Generators/SoundWaveScrubber.cpp
--- a/Engine/Source/Runtime/AudioMixer/Private/Generators/SoundWaveScrubber.cpp
+++ b/Engine/Source/Runtime/AudioMixer/Private/Generators/SoundWaveScrubber.cpp
@@ -21,14 +21,15 @@ FSoundWaveScrubber::~FSoundWaveScrubber()
 {
 }
 
-void FSoundWaveScrubber::Init(FSoundWaveProxyPtr InSoundWaveProxyPtr, float InSampleRate, int32 InNumChannels, float InPlayheadTimeSeconds)
+void FSoundWaveScrubber::Init(TSharedRef<const FSoundWaveProxy> InSoundWaveProxyRef, float InSampleRate, int32 InNumChannels, float InPlayheadTimeSeconds)
 {
-	check(InSoundWaveProxyPtr.IsValid());
-	SoundWaveProxyPtr = InSoundWaveProxyPtr;
+	SoundWaveDataPtr = InSoundWaveProxyRef->GetSoundWaveDataRef();
+	check(SoundWaveDataPtr.IsValid());
+	check(InSampleRate > 0);
 
 	AudioMixerSampleRate = InSampleRate;
-	SourceFileSampleRate = SoundWaveProxyPtr->GetSampleRate();
-	SourceFileDurationSeconds = SoundWaveProxyPtr->GetDuration();
+	SourceFileSampleRate = SoundWaveDataPtr->GetSampleRate();
+	SourceFileDurationSeconds = SoundWaveDataPtr->GetDuration();
 	NumChannels = InNumChannels;
 
 	TargetGrainDurationRange = { 0.4f, 0.05f };
@@ -50,6 +51,17 @@ void FSoundWaveScrubber::Init(FSoundWaveProxyPtr InSoundWaveProxyPtr, float InSa
 	DecodeToDataChunk(DecodedChunks[0], DecoderSeekTimeSeconds);
 }
 
+void FSoundWaveScrubber::Init(FSoundWaveProxyPtr InSoundWaveProxyPtr, float InSampleRate, int32 InNumChannels, float InPlayheadTimeSeconds)
+{
+	if (!ensure(InSoundWaveProxyPtr))
+	{
+		return;
+	}
+
+	TSharedRef<const FSoundWaveProxy> ProxyRef = InSoundWaveProxyPtr->AsShared();
+	this->Init(ProxyRef, InSampleRate, InNumChannels, InPlayheadTimeSeconds);
+}
+
 int32 FSoundWaveScrubber::GetDecodedDataChunkIndexForCurrentReadIndex(int32 InReadFrameIndex)
 {
 	for (int32 i = 0; i < DecodedChunks.Num(); ++i)
@@ -106,7 +118,7 @@ void FSoundWaveScrubber::DecodeToDataChunk(FDecodedDataChunk& InOutDataChunk, fl
 	check(SourceFileSampleRate > 0.0f);
 	check(InDecoderSeekTimeSeconds >= 0.0f);
 	check(NumChannels > 0);
-	check(SoundWaveProxyPtr.IsValid());
+	check(SoundWaveDataPtr.IsValid());
 
 	if (!SoundWaveProxyReaderPtr.IsValid())
 	{
@@ -117,12 +129,12 @@ void FSoundWaveScrubber::DecodeToDataChunk(FDecodedDataChunk& InOutDataChunk, fl
 		FSoundWaveProxyReader::FSettings ProxyReaderSettings;
 		ProxyReaderSettings.MaxDecodeSizeInFrames = DecodedAudioSizeInSeconds * SourceFileSampleRate;
 		ProxyReaderSettings.StartTimeInSeconds = InDecoderSeekTimeSeconds;
-		SoundWaveProxyReaderPtr = FSoundWaveProxyReader::Create(SoundWaveProxyPtr.ToSharedRef(), ProxyReaderSettings);
+		SoundWaveProxyReaderPtr = FSoundWaveProxyReader::Create(SoundWaveDataPtr.ToSharedRef(), ProxyReaderSettings);
 
 		if (!SoundWaveProxyReaderPtr.IsValid())
 		{
 			bHasErrorWithDecoder = true;
-			UE_LOG(LogAudioMixer, Warning, TEXT("Unable to make a sound wave proxy reader for sound wave '%s' in the sound wave scrubber."), *SoundWaveProxyPtr->GetFName().ToString());
+			UE_LOG(LogAudioMixer, Warning, TEXT("Unable to make a sound wave proxy reader for sound wave '%s' in the sound wave scrubber."), *SoundWaveDataPtr->GetFName().ToString());
 			return;
 		}
 	}
@@ -227,7 +239,7 @@ int32 FSoundWaveScrubber::RenderAudio(TArrayView<float>& OutAudio)
 
 	// If we have an error with our decoder, we don't need to render audio
 	// To avoid spamming, we'll allow it to render silence
-	if (bHasErrorWithDecoder)
+	if (bHasErrorWithDecoder || !SoundWaveDataPtr.IsValid())
 	{
 		return OutAudio.Num();
 	}
@@ -526,10 +538,21 @@ void FSoundWaveScrubber::RenderActiveGrains(TArrayView<float>& OutAudio, int32 I
 	}
 }
 
-void FSoundWaveScrubberGenerator::Init(FSoundWaveProxyPtr InSoundWaveProxyPtr, float InSampleRate, int32 InNumChannels, float InPlayheadTimeSeconds)
+void FSoundWaveScrubberGenerator::Init(TSharedRef<const FSoundWaveProxy> InSoundWaveProxyRef, float InSampleRate, int32 InNumChannels, float InPlayheadTimeSeconds)
 {
 	NumChannels = InNumChannels;
-	SoundWaveScrubber.Init(InSoundWaveProxyPtr, InSampleRate, InNumChannels, InPlayheadTimeSeconds);
+	SoundWaveScrubber.Init(InSoundWaveProxyRef, InSampleRate, InNumChannels, InPlayheadTimeSeconds);
+}
+
+void FSoundWaveScrubberGenerator::Init(FSoundWaveProxyPtr InSoundWaveProxyPtr, float InSampleRate, int32 InNumChannels, float InPlayheadTimeSeconds)
+{
+	if (!ensure(InSoundWaveProxyPtr))
+	{
+		return;
+	}
+
+	TSharedRef<const FSoundWaveProxy> ProxyRef = InSoundWaveProxyPtr->AsShared();
+	this->Init(ProxyRef, InSampleRate, InNumChannels, InPlayheadTimeSeconds);
 }
 
 void FSoundWaveScrubberGenerator::SetIsScrubbing(bool bInIsScrubbing)
@@ -584,12 +607,12 @@ ISoundGeneratorPtr UScrubbedSound::CreateSoundGenerator(const FSoundGeneratorIni
 
 	if (SoundWaveToScrub)
 	{
-		FSoundWaveProxyPtr SoundWaveProxyPtr = SoundWaveToScrub->CreateSoundWaveProxy();
+		TSharedRef<const FSoundWaveProxy> SoundWaveProxyRef = SoundWaveToScrub->GetSoundWaveProxy();
 
 		SoundWaveScrubber = ISoundGeneratorPtr(new FSoundWaveScrubberGenerator());
 
 		FSoundWaveScrubberGenerator* Scrubber = static_cast<FSoundWaveScrubberGenerator*>(SoundWaveScrubber.Get());
-		Scrubber->Init(SoundWaveProxyPtr, InParams.SampleRate, NumChannels, PlayheadTimeSeconds);
+		Scrubber->Init(SoundWaveProxyRef, InParams.SampleRate, NumChannels, PlayheadTimeSeconds);
 		Scrubber->SetIsScrubbing(bIsScrubbing);
 		Scrubber->SetIsScrubbingWhileStationary(bScrubWhileStationary);
 		Scrubber->SetGrainDurationRange(GrainDurationRange);
diff --git a/Engine/Source/Runtime/BinkAudioDecoder/Module/Private/BinkAudioInfo.cpp b/Engine/Source/Runtime/BinkAudioDecoder/Module/Private/BinkAudioInfo.cpp
--- a/Engine/Source/Runtime/BinkAudioDecoder/Module/Private/BinkAudioInfo.cpp
+++ b/Engine/Source/Runtime/BinkAudioDecoder/Module/Private/BinkAudioInfo.cpp
@@ -190,20 +190,20 @@ void FBinkAudioInfo::SeekToFrame(const uint32 InFrameNum)
 	//
 	// Here we need to set up the data we get to point at the right spot.
 	//
-	if (StreamingSoundWave == nullptr)
+	if (StreamingSoundWaveData == nullptr)
 	{
 		// If we aren't streaming we can just go directly to the offset we need.
 		this->SrcBufferOffset = OffsetToBlock;
 	}
 	else
 	{
-		const uint32 TotalStreamingChunks = StreamingSoundWave->GetNumChunks();
+		const uint32 TotalStreamingChunks = StreamingSoundWaveData->GetNumChunks();
 		uint32 RemnOffset = OffsetToBlock;
 
 		// Find the chunk and offset to the block we need.
 		for (uint32 BlockIndex = 0; BlockIndex < TotalStreamingChunks; ++BlockIndex)
 		{
-			const uint32 SizeOfChunk = StreamingSoundWave->GetSizeOfChunk(BlockIndex);
+			const uint32 SizeOfChunk = StreamingSoundWaveData->GetSizeOfChunk(BlockIndex);
 
 			if (SizeOfChunk > RemnOffset)
 			{
diff --git a/Engine/Source/Runtime/Engine/Classes/Sound/SoundWave.h b/Engine/Source/Runtime/Engine/Classes/Sound/SoundWave.h
--- a/Engine/Source/Runtime/Engine/Classes/Sound/SoundWave.h
+++ b/Engine/Source/Runtime/Engine/Classes/Sound/SoundWave.h
@@ -6,6 +6,10 @@
  * Playable sound object for raw wave files
  */
 
+#ifndef UE_USOUNDWAVE_HOT_RELOADING_ENABLED
+#define UE_USOUNDWAVE_HOT_RELOADING_ENABLED PLATFORM_DESKTOP && !UE_BUILD_SHIPPING
+#endif // #ifndef UE_USOUNDWAVE_HOT_RELOADING_ENABLED
+
 #include "CoreMinimal.h"
 #include "Containers/IndirectArray.h"
 #include "Sound/AudioSettings.h"
@@ -412,6 +416,7 @@ struct ISoundWaveClient
 	virtual bool OnIsReadyForFinishDestroy(class USoundWave* Wave) const = 0;
 	virtual void OnFinishDestroy(class USoundWave* Wave) = 0;
 };
+
 UCLASS(hidecategories=Object, editinlinenew, BlueprintType, meta= (LoadBehavior = "LazyOnDemand"), MinimalAPI)
 class USoundWave : public USoundBase, public IAudioProxyDataFactory, public IInterface_AsyncCompilation
 {
@@ -516,8 +521,19 @@ public:
 	ENGINE_API void ResetInvalidTemplateArguments();
 
 private:
-	// cached proxy
-	FSoundWaveProxyPtr Proxy{ nullptr };
+	/** if UE_USOUNDWAVE_HOT_RELOADING_ENABLED is true, the sound wave will have
+		its own, un-published copy of sound wave data (TSharedRef<FSoundWaveData> WorkingSoundWaveData)
+		Otherwise, we use and make changes to the Proxy->PublishedSoundWaveData directly.
+
+		This getter will return the correct one.
+	*/
+	const TSharedRef<FSoundWaveData>& GetWorkingSoundWaveData() const;
+
+	TSharedRef<FSoundWaveProxy> Proxy = MakeShared<FSoundWaveProxy>();
+	
+#if UE_USOUNDWAVE_HOT_RELOADING_ENABLED
+	TSharedRef<FSoundWaveData> WorkingSoundWaveData = MakeShared<FSoundWaveData>();
+#endif // #if UE_USOUNDWAVE_HOT_RELOADING_ENABLED
 
 public:
 
@@ -1156,9 +1172,6 @@ private:
 	/**  Utility function used internally to wait or poll a task while maintaining thread-safety. */
 	ENGINE_API bool WaitAsyncTaskWithTimeout(float InTimeoutInSeconds);
 
-	/** Creates and initializes a new FSoundWaveData and FSoundWaveProxy. */
-	void CreateNewSoundWaveData();
-
 public:
 #endif // WITH_EDITOR
 
@@ -1186,8 +1199,11 @@ public:
 	virtual bool IsProcedurallyGenerated() const override { return bProcedural; }
 	//~ End USoundBase Interface.
 
+	UE_DEPRECATED(5.8, "Mutable access to proxy is deprecated. Use GetSoundWaveProxy()")
 	ENGINE_API FSoundWaveProxyPtr CreateSoundWaveProxy();
 
+	ENGINE_API TSharedRef<const FSoundWaveProxy> GetSoundWaveProxy();
+
 	//~Begin IAudioProxyDataFactory Interface.
 	ENGINE_API virtual TSharedPtr<Audio::IProxyData> CreateProxyData(const Audio::FProxyDataInitParams& InitParams) override;
 	//~ End IAudioProxyDataFactory Interface.
@@ -1567,19 +1583,12 @@ public:
 	ENGINE_API void OverrideWaveData(int32 InNumChannels, int32 InNumFrames, int32 InSampleRate);
 
 private:
-#if PLATFORM_DESKTOP
-	TSharedRef<TDataSnapshotProducer<FSoundWaveData>> SoundWaveData = MakeShared<TDataSnapshotProducer<FSoundWaveData>>();
-	TDataSnapshotProducer<FSoundWaveData>& SoundWaveDataPtr = *SoundWaveData;
-#else // #if PLATFORM_DESKTOP
-	TSharedPtr<FSoundWaveData, ESPMode::ThreadSafe> SoundWaveDataPtr{ MakeShared<FSoundWaveData>() };
-#endif // #if PLATFORM_DESKTOP
-
 	friend class FSoundWaveProxy;
 	friend class USoundFactory;
 };
 
 
-class FSoundWaveData
+class FSoundWaveData : public TSharedFromThis<FSoundWaveData>
 {
 public:
 	UE_NONCOPYABLE(FSoundWaveData);
@@ -1598,11 +1607,13 @@ public:
 	{
 	}
 
+	FSoundWaveData(const USoundWave& InWave);
+
 	// dtor
 	ENGINE_API ~FSoundWaveData();
 
 	// non-const arg for internal call to FAudioDevice::GetRuntimeFormat
-	ENGINE_API void InitializeDataFromSoundWave(USoundWave& InWave);
+	ENGINE_API void InitializeDataFromSoundWave(const USoundWave& InWave);
 
 	ENGINE_API void OverrideRuntimeFormat(const FName& InRuntimeFormat);
 
@@ -1623,7 +1634,10 @@ public:
 	int32 GetNumFrames() const { return NumFrames; }
 	float GetDuration() const { return Duration; }
 
-	ENGINE_API void ReleaseCompressedAudio();
+	ENGINE_API void ReleaseCompressedAudio() const;
+
+	ENGINE_API const TArray<FSoundWaveCuePoint> GetCuePointsAndLoopRegions() const;
+	ENGINE_API TArrayView<const uint8> GetZerothChunk(bool bForImmediatePlayback = false) const;
 
 	ENGINE_API void SetError(const TCHAR* InErrorMsg=nullptr);  
 	ENGINE_API bool HasError() const;
@@ -1641,9 +1655,9 @@ public:
 	bool IsTemplate() const { return bIsTemplate; }
 
  	ENGINE_API bool HasCompressedData(FName Format, ITargetPlatform* TargetPlatform = USoundWave::GetRunningPlatform()) const;
- 	ENGINE_API FByteBulkData* GetCompressedData(FName Format, const FPlatformAudioCookOverrides* CompressionOverrides = USoundWave::GetPlatformCompressionOverridesForCurrentPlatform());
+ 	ENGINE_API FByteBulkData* GetCompressedData(FName Format, const FPlatformAudioCookOverrides* CompressionOverrides = USoundWave::GetPlatformCompressionOverridesForCurrentPlatform()) const;
  
-	ENGINE_API bool GetChunkData(int32 ChunkIndex, uint8** OutChunkData, bool bMakeSureChunkIsLoaded = false);
+	ENGINE_API bool GetChunkData(int32 ChunkIndex, uint8** OutChunkData, bool bMakeSureChunkIsLoaded = false) const;
  
  	ENGINE_API bool IsZerothChunkDataLoaded() const;
  	ENGINE_API const TArrayView<uint8> GetZerothChunkDataView() const;
@@ -1654,21 +1668,22 @@ public:
 	// Returns true if the zeroth chunk is loaded, or attempts to load it if not already loaded,
 	// returning true if the load was successful. Can return false if either an error was encountered
 	// in attempting to load the chunk or if stream caching is not enabled for the given sound.
-	ENGINE_API bool LoadZerothChunk();
+	// (const because this call is used to lazy-init zeroth chunk in GetZerothChunkData())
+	ENGINE_API bool LoadZerothChunk() const;
  
  #if WITH_EDITOR
  	ENGINE_API int32 GetCurrentChunkRevision() const;
  #endif // #if WITH_EDITOR
  
- 	ENGINE_API FStreamedAudioChunk& GetChunk(uint32 ChunkIndex);
- 	ENGINE_API int32 GetChunkFromDDC(int32 ChunkIndex, uint8** OutChunkData, bool bMakeSureChunkIsLoaded);
+ 	ENGINE_API FStreamedAudioChunk& GetChunk(uint32 ChunkIndex) const;
+ 	ENGINE_API int32 GetChunkFromDDC(int32 ChunkIndex, uint8** OutChunkData, bool bMakeSureChunkIsLoaded) const;
 
 #if WITH_EDITORONLY_DATA
  	ENGINE_API FString GetDerivedDataKey() const;
 	FPerPlatformFloat GetSizeOfFirstAudioChunkInSeconds() const { return SizeOfFirstAudioChunkInSeconds; }
 #endif // #if WITH_EDITORONLY_DATA
 
-	int32 GetResourceSize() { return ResourceSize; }
+	int32 GetResourceSize() const { return ResourceSize; }
 	const uint8* GetResourceData() const { return ResourceData.GetView().GetData(); }
 	uint32 GetNumChannels() { return NumChannels; }
 
@@ -1688,7 +1703,8 @@ private:
 	ENGINE_API void CacheRuntimeFormatDependentState(const FName InFormatName);
 
 	/** Zeroth Chunk of audio for sources that use Load On Demand. */
-	FBulkDataBuffer<uint8> ZerothChunkData;
+	// mutable for lazy-init in GetZerothChunkData()
+	mutable FBulkDataBuffer<uint8> ZerothChunkData;
 #if WITH_EDITOR
 	mutable FCriticalSection LoadZerothChunkDataCriticalSection;
 #endif
@@ -1696,12 +1712,16 @@ private:
 	/* Accessor to get the zeroth chunk which might perform additional work in editor to handle async tasks. */
 	ENGINE_API FBulkDataBuffer<uint8>& GetZerothChunkData() const;
 
+
+	// The following members are mutable to allow lazy initialization in const getters
+	// (This way we can hand out immutable refrences to objects of this type)
+
 	/** The streaming derived data for this sound on this platform. */
-	FStreamedAudioPlatformData RunningPlatformData;
+	mutable FStreamedAudioPlatformData RunningPlatformData;
 
-	FAudioChunkHandle FirstChunk;
+	mutable FAudioChunkHandle FirstChunk;
 
-	FFormatContainer CompressedFormatData;
+	mutable FFormatContainer CompressedFormatData;
 
 	int32 ResourceSize = 0;
 
@@ -1741,6 +1761,7 @@ private:
 	uint8 bShouldUseStreamCaching : 1;
 	uint8 bLoadingBehaviorOverridden : 1;
 	uint8 bHasError : 1;
+	
 
 #if WITH_EDITOR
 	uint8 bLoadedFromCookedData : 1;
@@ -1749,97 +1770,159 @@ private:
 	friend class USoundWave;
 }; // class FSoundWaveData
 
+
+
+// i.e.: "FSoundWaveProxy is a shared object between USoundWave and other systems which utilize FSoundWaveData"
+/**
+	The proxy acts as a broker between the USoundWave and Proxy Readers.
+	if UE_USOUNDWAVE_HOT_RELOADING_ENABLED is true, then proxy clients can use this to query the availablity
+	of new FSoundWaveData, and optionally update to that data.
+
+	This lets editor changes and re-imports take effect on even playing sounds.
+
+	In these scenarios, the owning USoundWave has its "working" (un-published) copy of FSoundWaveData
+	For changes, and then calls "Publish()" to make this available
+*/
 using FSoundWaveProxyPtr = TSharedPtr<FSoundWaveProxy, ESPMode::ThreadSafe>;
 class FSoundWaveProxy : public Audio::TProxyData<FSoundWaveProxy>
+					  , public TSharedFromThis<FSoundWaveProxy, ESPMode::ThreadSafe>	
 {
 public:
 	IMPL_AUDIOPROXY_CLASS(FSoundWaveProxy);
 
+	ENGINE_API FSoundWaveProxy();
+	
 	ENGINE_API explicit FSoundWaveProxy(USoundWave* InWave);
 
-	ENGINE_API FSoundWaveProxy(const FSoundWaveProxy& Other);
-
 	ENGINE_API ~FSoundWaveProxy();
+	
+	static ENGINE_API TSharedRef<FSoundWaveProxy> CreateNonUpdatingCopy(const TSharedRef<const FSoundWaveProxy>& Other);
+	static ENGINE_API TSharedRef<const FSoundWaveProxy> CreateNonUpdatingCopy(const TSharedRef<const FSoundWaveData>& Other);
+	
+	friend inline uint32 GetTypeHash(const FSoundWaveProxy& InProxy)
+	{
+		return InProxy.TypeHash;
+	}
 
 	// USoundWave Interface
+	
+	UE_DEPRECATED(5.8, "The FSoundWaveProxy is now just a way to get a reference to the latest FSoundWaveData. Clients should hold onto the return value of GetSoundWaveDataRef() instead")
+	static ENGINE_API TArrayView<const uint8> GetZerothChunk(const FSoundWaveProxyPtr& SoundWaveProxy, bool bForImmediatePlayback = false);
+	UE_DEPRECATED(5.8, "The FSoundWaveProxy is now just a way to get a reference to the latest FSoundWaveData. Clients should hold onto the return value of GetSoundWaveDataRef() instead")
+	ENGINE_API const TArray<FSoundWaveCuePoint> GetCuePointsAndLoopRegions() const;
+	UE_DEPRECATED(5.8, "The FSoundWaveProxy is now just a way to get a reference to the latest FSoundWaveData. Clients should hold onto the return value of GetSoundWaveDataRef() instead")
 	ENGINE_API void ReleaseCompressedAudio();
 
 	// Returns true if the zeroth chunk is loaded, or attempts to load it if not already loaded,
 	// returning true if the load was successful. Can return false if either an error was encountered
-	// in attempting to load the chunk or if stream caching is not enabled for the given sound.
+	// in attempting to load the chunk or if stream caching is not enabled for the given sound
+	UE_DEPRECATED(5.8, "The FSoundWaveProxy is now just a way to get a reference to the latest FSoundWaveData. Clients should hold onto the return value of GetSoundWaveDataRef() instead")
 	ENGINE_API bool LoadZerothChunk();
+	UE_DEPRECATED(5.8, "The FSoundWaveProxy is now just a way to get a reference to the latest FSoundWaveData. Clients should hold onto the return value of GetSoundWaveDataRef() instead")
 	ENGINE_API bool GetChunkData(int32 ChunkIndex, uint8** OutChunkData, bool bMakeSureChunkIsLoaded = false);
-
+	
 	// Getters
+	UE_DEPRECATED(5.8, "The FSoundWaveProxy is now just a way to get a reference to the latest FSoundWaveData. Clients should hold onto the return value of GetSoundWaveDataRef() instead")
 	ENGINE_API const FName& GetFName() const;
+	UE_DEPRECATED(5.8, "The FSoundWaveProxy is now just a way to get a reference to the latest FSoundWaveData. Clients should hold onto the return value of GetSoundWaveDataRef() instead")
 	ENGINE_API const FName& GetPackageName() const;
+	UE_DEPRECATED(5.8, "The FSoundWaveProxy is now just a way to get a reference to the latest FSoundWaveData. Clients should hold onto the return value of GetSoundWaveDataRef() instead")
 	ENGINE_API const FName& GetRuntimeFormat() const;
+	UE_DEPRECATED(5.8, "The FSoundWaveProxy is now just a way to get a reference to the latest FSoundWaveData. Clients should hold onto the return value of GetSoundWaveDataRef() instead")
 	ENGINE_API const FObjectKey& GetFObjectKey() const;
-
+	
+	UE_DEPRECATED(5.8, "The FSoundWaveProxy is now just a way to get a reference to the latest FSoundWaveData. Clients should hold onto the return value of GetSoundWaveDataRef() instead")
 	ENGINE_API float GetDuration() const;
+	UE_DEPRECATED(5.8, "The FSoundWaveProxy is now just a way to get a reference to the latest FSoundWaveData. Clients should hold onto the return value of GetSoundWaveDataRef() instead")
 	ENGINE_API float GetSampleRate() const;
-
+	
+	UE_DEPRECATED(5.8, "The FSoundWaveProxy is now just a way to get a reference to the latest FSoundWaveData. Clients should hold onto the return value of GetSoundWaveDataRef() instead")
 	ENGINE_API int32 GetNumFrames() const;
+	UE_DEPRECATED(5.8, "The FSoundWaveProxy is now just a way to get a reference to the latest FSoundWaveData. Clients should hold onto the return value of GetSoundWaveDataRef() instead")
 	ENGINE_API uint32 GetNumChunks() const;
+	UE_DEPRECATED(5.8, "The FSoundWaveProxy is now just a way to get a reference to the latest FSoundWaveData. Clients should hold onto the return value of GetSoundWaveDataRef() instead")
 	ENGINE_API uint32 GetNumChannels() const;
+	UE_DEPRECATED(5.8, "The FSoundWaveProxy is now just a way to get a reference to the latest FSoundWaveData. Clients should hold onto the return value of GetSoundWaveDataRef() instead")
 	ENGINE_API uint32 GetSizeOfChunk(uint32 ChunkIndex) const;
+	UE_DEPRECATED(5.8, "The FSoundWaveProxy is now just a way to get a reference to the latest FSoundWaveData. Clients should hold onto the return value of GetSoundWaveDataRef() instead")
 	ENGINE_API const TArray<FSoundWaveCuePoint>& GetCuePoints() const;
+	UE_DEPRECATED(5.8, "The FSoundWaveProxy is now just a way to get a reference to the latest FSoundWaveData. Clients should hold onto the return value of GetSoundWaveDataRef() instead")
 	ENGINE_API const TArray<FSoundWaveCuePoint>& GetLoopRegions() const;
-	ENGINE_API const TArray<FSoundWaveCuePoint> GetCuePointsAndLoopRegions() const;
-
+	
+	UE_DEPRECATED(5.8, "The FSoundWaveProxy is now just a way to get a reference to the latest FSoundWaveData. Clients should hold onto the return value of GetSoundWaveDataRef() instead")
 	ENGINE_API bool IsLooping() const;
+	UE_DEPRECATED(5.8, "The FSoundWaveProxy is now just a way to get a reference to the latest FSoundWaveData. Clients should hold onto the return value of GetSoundWaveDataRef() instead")
 	ENGINE_API bool IsTemplate() const;
+	UE_DEPRECATED(5.8, "The FSoundWaveProxy is now just a way to get a reference to the latest FSoundWaveData. Clients should hold onto the return value of GetSoundWaveDataRef() instead")
 	ENGINE_API bool IsStreaming() const;
+	UE_DEPRECATED(5.8, "The FSoundWaveProxy is now just a way to get a reference to the latest FSoundWaveData. Clients should hold onto the return value of GetSoundWaveDataRef() instead")
 	ENGINE_API bool IsRetainingAudio() const;
+	UE_DEPRECATED(5.8, "The FSoundWaveProxy is now just a way to get a reference to the latest FSoundWaveData. Clients should hold onto the return value of GetSoundWaveDataRef() instead")
 	ENGINE_API bool ShouldUseStreamCaching() const;
+	UE_DEPRECATED(5.8, "The FSoundWaveProxy is now just a way to get a reference to the latest FSoundWaveData. Clients should hold onto the return value of GetSoundWaveDataRef() instead")
 	ENGINE_API bool IsSeekable() const;
+	UE_DEPRECATED(5.8, "The FSoundWaveProxy is now just a way to get a reference to the latest FSoundWaveData. Clients should hold onto the return value of GetSoundWaveDataRef() instead")
 	ENGINE_API bool IsZerothChunkDataLoaded() const;
+	UE_DEPRECATED(5.8, "The FSoundWaveProxy is now just a way to get a reference to the latest FSoundWaveData. Clients should hold onto the return value of GetSoundWaveDataRef() instead")
 	ENGINE_API bool WasLoadingBehaviorOverridden() const;
+	UE_DEPRECATED(5.8, "The FSoundWaveProxy is now just a way to get a reference to the latest FSoundWaveData. Clients should hold onto the return value of GetSoundWaveDataRef() instead")
 	ENGINE_API bool HasCompressedData(FName Format, ITargetPlatform* TargetPlatform = USoundWave::GetRunningPlatform()) const;
-
+	
+	UE_DEPRECATED(5.8, "The FSoundWaveProxy is now just a way to get a reference to the latest FSoundWaveData. Clients should hold onto the return value of GetSoundWaveDataRef() instead")
 	ENGINE_API ESoundWaveLoadingBehavior GetLoadingBehavior() const;
-
+	
+	UE_DEPRECATED(5.8, "The FSoundWaveProxy is now just a way to get a reference to the latest FSoundWaveData. Clients should hold onto the return value of GetSoundWaveDataRef() instead")
 	ENGINE_API const TArrayView<uint8> GetZerothChunkDataView() const;
-
+	
+	UE_DEPRECATED(5.8, "The FSoundWaveProxy is now just a way to get a reference to the latest FSoundWaveData. Clients should hold onto the return value of GetSoundWaveDataRef() instead")
 	ENGINE_API FByteBulkData* GetCompressedData(FName Format, const FPlatformAudioCookOverrides* CompressionOverrides = USoundWave::GetPlatformCompressionOverridesForCurrentPlatform());
 
-	static ENGINE_API TArrayView<const uint8> GetZerothChunk(const FSoundWaveProxyPtr& SoundWaveProxy, bool bForImmediatePlayback = false);
-
-
 #if WITH_EDITOR
+	UE_DEPRECATED(5.8, "The FSoundWaveProxy is now just a way to get a reference to the latest FSoundWaveData. Clients should hold onto the return value of GetSoundWaveDataRef() instead")
 	ENGINE_API int32 GetCurrentChunkRevision() const;
 #endif // #if WITH_EDITOR
 
+	UE_DEPRECATED(5.8, "The FSoundWaveProxy is now just a way to get a reference to the latest FSoundWaveData. Clients should hold onto the return value of GetSoundWaveDataRef() instead")
 	ENGINE_API FStreamedAudioChunk& GetChunk(uint32 ChunkIndex);
+	UE_DEPRECATED(5.8, "The FSoundWaveProxy is now just a way to get a reference to the latest FSoundWaveData. Clients should hold onto the return value of GetSoundWaveDataRef() instead")
 	ENGINE_API int32 GetChunkFromDDC(int32 ChunkIndex, uint8** OutChunkData, bool bMakeSureChunkIsLoaded);
 
 #if WITH_EDITORONLY_DATA
+	UE_DEPRECATED(5.8, "The FSoundWaveProxy is now just a way to get a reference to the latest FSoundWaveData. Clients should hold onto the return value of GetSoundWaveDataRef() instead")
 	ENGINE_API FString GetDerivedDataKey() const;
 #endif // #if WITH_EDITORONLY_DATA
 
+	UE_DEPRECATED(5.8, "The FSoundWaveProxy is now just a way to get a reference to the latest FSoundWaveData. Clients should hold onto the return value of GetSoundWaveDataRef() instead")
 	ENGINE_API int32 GetResourceSize() const;
+
+	UE_DEPRECATED(5.8, "The FSoundWaveProxy is now just a way to get a reference to the latest FSoundWaveData. Clients should hold onto the return value of GetSoundWaveDataRef() instead")
 	ENGINE_API const uint8* GetResourceData() const;
 
-	ENGINE_API const FSoundWavePtr GetSoundWaveData();
-	
-	friend inline uint32 GetTypeHash(const FSoundWaveProxy& InProxy)
+	UE_DEPRECATED(5.8, "mutable access to sound wave data is deprecated. please use GetSoundWaveDataRef()")
+	TSharedPtr<FSoundWaveData> GetSoundWaveData() const
 	{
-		return InProxy.TypeHash;
+		return ConstCastSharedRef<FSoundWaveData>(PublishedSoundWaveData).ToSharedPtr();
 	}
+	
+	ENGINE_API TSharedRef<const FSoundWaveData> GetSoundWaveDataRef() const;
 
-#if PLATFORM_DESKTOP
-	bool TryGetLatest()
-	{
-		return SoundWaveDataPtr.GetLatest();
-	}
-#endif // #if PLATFORM_DESKTOP
+#if UE_USOUNDWAVE_HOT_RELOADING_ENABLED
+	bool ENGINE_API IsLatestData(const FSoundWaveData* InData) const;
+#endif // #if UE_USOUNDWAVE_HOT_RELOADING_ENABLED
 	
 private:
-#if PLATFORM_DESKTOP
-	TDataSnapshot<FSoundWaveData> SoundWaveDataPtr;
-#else // #if PLATFORM_DESKTOP
-	TSharedPtr<FSoundWaveData, ESPMode::ThreadSafe> SoundWaveDataPtr;
-#endif // #if PLATFORM_DESKTOP
+	friend USoundWave;
+	void InitializeFromWave(USoundWave& InWave);
+
+#if UE_USOUNDWAVE_HOT_RELOADING_ENABLED	
+	void Publish(const TSharedRef<FSoundWaveData>& InNewData);
+#endif // #if UE_USOUNDWAVE_HOT_RELOADING_ENABLED
+	
+	TSharedRef<FSoundWaveData> PublishedSoundWaveData = MakeShared<FSoundWaveData>();
+
+#if UE_USOUNDWAVE_HOT_RELOADING_ENABLED
+	mutable FRWLock PublishedSoundWaveDataLock;
+#endif // #if UE_USOUNDWAVE_HOT_RELOADING_ENABLED
 	
 	uint32 TypeHash = uint32(INDEX_NONE);
 };
diff --git a/Engine/Source/Runtime/Engine/Classes/Sound/SoundWaveProxyPlayer.h b/Engine/Source/Runtime/Engine/Classes/Sound/SoundWaveProxyPlayer.h
--- a/Engine/Source/Runtime/Engine/Classes/Sound/SoundWaveProxyPlayer.h
+++ b/Engine/Source/Runtime/Engine/Classes/Sound/SoundWaveProxyPlayer.h
@@ -83,15 +83,17 @@ public:
 	UE_API static TUniquePtr<FSoundWaveProxyPlayer> Create(const FSettings& InSettings);
 
 	/**
-	 * Set the SoundWave Proxy to play
+	 * Set the SoundWave Data to play
 	 * 
-	 * if Null, will reset the Proxy Player and return 'false'
-	 * if the WaveProxy has invalid data or is unabled to be decoded,
+	* if the FSoundWaveData has invalid data or is unabled to be decoded,
 	 * the player will be reset and this will return 'false'
 	 * 
-	 * @param InWaveProxy 
+	 * @param InWaveData 
 	 * @return whether the player is valid for playback
 	 */
+	UE_API bool SetSoundWave(const TSharedRef<const FSoundWaveData>& InWaveData); 
+	
+	UE_DEPRECATED(5.8, "Please use the TSharedRef<const FSoundWaveData>& overload instead")
 	UE_API bool SetSoundWave(FSoundWaveProxyPtr InWaveProxy);
 
 	/**
@@ -282,12 +284,12 @@ private:
 	FSoundWaveProxyPlayer(const FSettings& InSettings);
 
 	/**
-	 * Check wether the Wave Proxy has valid data (valid Sample Rate, Num Channels, Num Frames)
+	 * Check whether the Wave Proxy has valid data (valid Sample Rate, Num Channels, Num Frames)
 	 * 
-	 * @param InWaveProxy 
+	 * @param InWaveData
 	 * @return whether the WaveProxy is valid
 	 */
-	static bool IsValidSoundWaveProxy(FSoundWaveProxyRef InWaveProxy);
+	static bool IsValidSoundWaveData(const FSoundWaveData& InWaveData);
 
 	// 
 	/**
@@ -342,15 +344,14 @@ private:
 	/**
 	 * Initialize the Decoder with the given wave proxy
 	 * also initialize the decode buffers
-	 * @param WaveProxy
 	 * @param MaxDecodeSizeInFrames 
 	 * @return whether the decoder was successfully initialized
 	 */
-	bool InitializeDecoder(const FSoundWaveProxyRef& InWaveProxy, int32 MaxDecodeSizeInFrames);
+	bool InitializeDecoder(int32 MaxDecodeSizeInFrames);
 
 	void InitializeResampler(int32 NumChannels);
 
-	FSoundWaveProxyPtr SoundWave;
+	TSharedPtr<const FSoundWaveData> SoundWaveDataPtr;
 	FSettings Settings;
 	TUniquePtr<ICompressedAudioInfo> CompressedAudioInfo;
 	TUniquePtr<Audio::IConvertDeinterleave> ConvertDeinterleave;
diff --git a/Engine/Source/Runtime/Engine/Classes/Sound/SoundWaveProxyReader.h b/Engine/Source/Runtime/Engine/Classes/Sound/SoundWaveProxyReader.h
--- a/Engine/Source/Runtime/Engine/Classes/Sound/SoundWaveProxyReader.h
+++ b/Engine/Source/Runtime/Engine/Classes/Sound/SoundWaveProxyReader.h
@@ -21,6 +21,7 @@ class FSoundWaveProxyReader
 
 public:
 	using FSoundWaveProxyRef = TSharedRef<FSoundWaveProxy, ESPMode::ThreadSafe>;
+	using FSoundWaveDataRef = TSharedRef<const FSoundWaveData, ESPMode::ThreadSafe>;
 
 	/** Minimum number of frames to decode per a call to the decoder.  */
 	static constexpr uint32 DefaultMinDecodeSizeInFrames = 128;
@@ -58,17 +59,24 @@ public:
 private:
 	/** Construct a wave proxy reader.
 	 *
-	 * @param InWaveProxy - A TSharedRef of a FSoundWaveProxy which is to be played.
+	 * @param InWaveData - A TSharedRef of a const FSoundWaveData which is to be played.
 	 * @param InSettings - Reader settings.
 	 */
+	FSoundWaveProxyReader(FSoundWaveDataRef InWaveData, const FSettings& InSettings);
+	
+	UE_DEPRECATED(5.8, "Please use the FSoundWaveDataRef overload instead")
 	FSoundWaveProxyReader(FSoundWaveProxyRef InWaveProxy, const FSettings& InSettings);
+	
 
 public:
 	/** Create a wave proxy reader.
 	 *
-	 * @param InWaveProxy - A TSharedRef of a FSoundWaveProxy which is to be played.
+	 * @param InWaveProxy - A TSharedRef of a const FSoundWaveData which is to be played.
 	 * @param InSettings - Reader settings.
 	 */
+	static ENGINE_API TUniquePtr<FSoundWaveProxyReader> Create(FSoundWaveDataRef InWaveProxy, const FSettings& InSettings);
+	
+	UE_DEPRECATED(5.8, "Please use the FSoundWaveDataRef overload instead")
 	static ENGINE_API TUniquePtr<FSoundWaveProxyReader> Create(FSoundWaveProxyRef InWaveProxy, const FSettings& InSettings);
 
 	/** Set whether the reader should loop the audio or not. */
@@ -181,7 +189,7 @@ private:
 	float ClampLoopDuration(float InDurationInSeconds);
 	void UpdateLoopBoundaries();
 
-	FSoundWaveProxyRef WaveProxy;
+	FSoundWaveDataRef WaveDataRef;
 
 	FSettings Settings;
 
