[Backout] - CL50566075
#rnx
Original CL Desc
-----------------------------------------------------------------
deprecating FSoundWaveData interface on FSoundWaveProxy. Proxy becomes broker for USoundWave to live-update FSoundWaveData. Clients updated to use FSoundWaveData instead of the proxy.

diff --git a/Engine/Plugins/Experimental/MetasoundExperimental/Source/MetasoundExperimentalEngineRuntime/Private/MetasoundGranulatorNode.cpp b/Engine/Plugins/Experimental/MetasoundExperimental/Source/MetasoundExperimentalEngineRuntime/Private/MetasoundGranulatorNode.cpp
--- a/Engine/Plugins/Experimental/MetasoundExperimental/Source/MetasoundExperimentalEngineRuntime/Private/MetasoundGranulatorNode.cpp
+++ b/Engine/Plugins/Experimental/MetasoundExperimental/Source/MetasoundExperimentalEngineRuntime/Private/MetasoundGranulatorNode.cpp
@@ -277,34 +277,23 @@ namespace Metasound
 			check(!bPlaying);
 
  			CurrentWaveAsset = *WaveAsset;
-			const TSharedPtr<const FSoundWaveProxy> Proxy = CurrentWaveAsset.GetWaveProxy();
-
-			if (Proxy)
-			{
-				SoundWaveDataPtr = Proxy->GetSoundWaveDataRef();
-			}
-			else
+			SoundWaveProxyPtr = CurrentWaveAsset.GetSoundWaveProxy();
+			if (!SoundWaveProxyPtr)
 			{
-				SoundWaveDataPtr.Reset();
 				return false;
 			}
 
-			SourceFileSampleRate = SoundWaveDataPtr->GetSampleRate();
-			SourceFileDurationSeconds = SoundWaveDataPtr->GetDuration();
-			SourceFileNumFrames = SoundWaveDataPtr->GetNumFrames();
-			SourceFileNumChannels = SoundWaveDataPtr->GetNumChannels();
+			SourceFileSampleRate = SoundWaveProxyPtr->GetSampleRate();
+			SourceFileDurationSeconds = SoundWaveProxyPtr->GetDuration();
+			SourceFileNumFrames = SoundWaveProxyPtr->GetNumFrames();
+			SourceFileNumChannels = SoundWaveProxyPtr->GetNumChannels();
 
 			FSoundWaveProxyReader::FSettings ProxyReaderSettings;
 			ProxyReaderSettings.MaxDecodeSizeInFrames = DecodedAudioSizeInSeconds * SourceFileSampleRate;
 			ProxyReaderSettings.StartTimeInSeconds = 0.0f;
 			ProxyReaderSettings.bIsLooping = true;
 			ProxyReaderSettings.LoopDurationInSeconds = SourceFileDurationSeconds;
-			SoundWaveProxyReaderPtr = FSoundWaveProxyReader::Create(SoundWaveDataPtr.ToSharedRef(), ProxyReaderSettings);
-			
-			if (!SoundWaveProxyReaderPtr)
-			{
-				return false;
-			}
+			SoundWaveProxyReaderPtr = FSoundWaveProxyReader::Create(SoundWaveProxyPtr.ToSharedRef(), ProxyReaderSettings);
 
 			DecodedChunks.Reset();
 			DecodedChunks.AddDefaulted(3);
@@ -353,7 +342,7 @@ namespace Metasound
 			check(SourceFileSampleRate > 0.0f);
 			check(InDecoderSeekTimeSeconds >= 0.0f);
 			check(SourceFileNumChannels > 0);
-			check(SoundWaveDataPtr.IsValid());
+			check(SoundWaveProxyPtr.IsValid());
 
 			check(SoundWaveProxyReaderPtr.IsValid());
 			float DecoderSeekTimeSeconds = FMath::Max(InDecoderSeekTimeSeconds - 0.5f * DecodedAudioSizeInSeconds, 0.0f);
@@ -941,7 +930,7 @@ namespace Metasound
 		int32 GrainCount = 0;
 
 		FWaveAsset CurrentWaveAsset;
-		TSharedPtr<const FSoundWaveData> SoundWaveDataPtr;
+		FSoundWaveProxyPtr SoundWaveProxyPtr;
 		TUniquePtr<FSoundWaveProxyReader> SoundWaveProxyReaderPtr;
 		TUniquePtr<Audio::ISampleRateConverter> SRC;
 		EGranularEnvelope GranularEnvelopeType = EGranularEnvelope::Hann;
diff --git a/Engine/Plugins/Runtime/Harmonix/Source/HarmonixDsp/Private/HarmonixDsp/AudioData/StreamingAudioRendererV2.cpp b/Engine/Plugins/Runtime/Harmonix/Source/HarmonixDsp/Private/HarmonixDsp/AudioData/StreamingAudioRendererV2.cpp
--- a/Engine/Plugins/Runtime/Harmonix/Source/HarmonixDsp/Private/HarmonixDsp/AudioData/StreamingAudioRendererV2.cpp
+++ b/Engine/Plugins/Runtime/Harmonix/Source/HarmonixDsp/Private/HarmonixDsp/AudioData/StreamingAudioRendererV2.cpp
@@ -29,17 +29,16 @@ void FStreamingAudioRendererV2::Reset()
 	Shifter = nullptr;
 }
 
-void FStreamingAudioRendererV2::SetSoundWaveData(const TSharedRef<const FSoundWaveData>& InSoundWaveData, const FSettings& InSettings)
+void FStreamingAudioRendererV2::SetAudioData(TSharedRef<FSoundWaveProxy> InSoundWave, const FSettings& InSettings)
 {
-	SoundWaveData = InSoundWaveData;
-	check(SoundWaveData.IsValid());
+	SoundWaveProxy = InSoundWave.ToSharedPtr();
 
 	WaveProxyReader.Reset();
-	WaveProxyReader = CreateProxyReader(SoundWaveData.ToSharedRef());
+	WaveProxyReader = CreateProxyReader(InSoundWave);
 
 	check(WaveProxyReader.IsValid());
 
-	int32 WaveProxyNumChannels = SoundWaveData->GetNumChannels();
+	int32 WaveProxyNumChannels = SoundWaveProxy->GetNumChannels();
 
 	int32 DecodeBufferSize = WaveProxyNumChannels * DeinterleaveBlockSizeInFrames;
 	DecodeBuffer.Reset(DecodeBufferSize);
@@ -61,13 +60,13 @@ void FStreamingAudioRendererV2::SetSoundWaveData(const TSharedRef<const FSoundWa
 
 	if (Shifter)
 	{
-		Shifter->SetSampleSourceReset(SoundWaveData.ToSharedRef(), AsShared(), InSettings.ShifterSettings);
+		Shifter->SetSampleSourceReset(SoundWaveProxy, AsShared(), InSettings.ShifterSettings);
 	}
 }
 
-TSharedPtr<const FSoundWaveData> FStreamingAudioRendererV2::GetSoundWaveData() const
+const TSharedPtr<FSoundWaveProxy> FStreamingAudioRendererV2::GetAudioData() const
 {
-	return SoundWaveData;
+	return SoundWaveProxy;
 }
 
 void FStreamingAudioRendererV2::MigrateToSampler(const FFusionSampler* InSampler)
@@ -85,7 +84,7 @@ void FStreamingAudioRendererV2::SetFrame(uint32 InFrameNum)
 
 double FStreamingAudioRendererV2::Render(TAudioBuffer<float>& OutBuffer, double InPos, int32 InMaxFrame, double InResampleInc, double InPitchShift, double InSpeed, bool MaintainPitchWhenSpeedChanges, bool InShouldHonorLoopPoints, const FGainMatrix& InGain)
 {
-	if (!SoundWaveData || SoundWaveData->GetNumFrames() == 0)
+	if (!SoundWaveProxy || SoundWaveProxy->GetNumFrames() == 0)
 	{
 		OutBuffer.ZeroValidFrames();
 		return InPos;
@@ -108,7 +107,7 @@ double FStreamingAudioRendererV2::RenderInternal(TAudioBuffer<float>& OutBuffer,
 {
 	check(OutBuffer.GetNumValidFrames() <= AudioRendering::kMicroSliceSize);
 
-	if (!SoundWaveData || SoundWaveData->GetNumFrames() < FMath::FloorToInt32(InPos))
+	if (!SoundWaveProxy || SoundWaveProxy->GetNumFrames() < FMath::FloorToInt32(InPos))
 	{
 		OutBuffer.ZeroValidFrames();
 		return InPos;
@@ -123,7 +122,7 @@ double FStreamingAudioRendererV2::RenderInternal(TAudioBuffer<float>& OutBuffer,
 	{
 		RenderMultiChannelRoutedUnshifted(OutBuffer, LerpArray, NumOutFrames, InGain, InInc, InShouldHonorLoopPoints);
 	}
-	else if (SoundWaveData->GetNumChannels() <= 2)
+	else if (SoundWaveProxy->GetNumChannels() <= 2)
 	{
 		RenderSimpleUnshifted(OutBuffer, LerpArray, NumOutFrames, InGain, InInc, InShouldHonorLoopPoints);
 	}
@@ -137,7 +136,7 @@ double FStreamingAudioRendererV2::RenderInternal(TAudioBuffer<float>& OutBuffer,
 
 double FStreamingAudioRendererV2::RenderUnshifted(TAudioBuffer<float>& OutBuffer, double InPos, int32 InMaxFrame, double InInc, bool InShouldHonorLoopPoints, const FGainMatrix& InGain)
 {
-	if (!SoundWaveData)
+	if (!SoundWaveProxy)
 	{
 		OutBuffer.ZeroValidFrames();
 		return InPos;
@@ -159,18 +158,13 @@ int32 FStreamingAudioRendererV2::CalculateNumFramesNeeded(const FLerpData* LerpD
 		return 1 + LerpData[NumPoints - 1].PosB - LerpData[0].PosA;
 	}
 
-	if (!SoundWaveData.IsValid())
-	{
-		return 0;
-	}
-
 	// Fail-safe for if we got what appears to be a loop, but we don't have a loop in the audio data
-	if (SoundWaveData->GetLoopRegions().IsEmpty())
+	if (SoundWaveProxy->GetLoopRegions().IsEmpty())
 	{
-		return SoundWaveData->GetNumFrames() - LerpData[0].PosA;
+		return SoundWaveProxy->GetNumFrames() - LerpData[0].PosA;
 	}
 
-	const FSoundWaveCuePoint& LoopRegion = SoundWaveData->GetLoopRegions()[0];
+	const FSoundWaveCuePoint& LoopRegion = SoundWaveProxy->GetLoopRegions()[0];
 
 	// loop in lerp. 
 	// we need all the samples from LerpData[0].PosA to the end of the loop, 
@@ -191,7 +185,7 @@ void FStreamingAudioRendererV2::RenderSimpleUnshifted(TAudioBuffer<float>& OutBu
 	OutBuffer.ZeroValidFrames();
 
 	int32 NumOutChannels = OutBuffer.GetNumValidChannels();
-	int32 NumInputChannels = SoundWaveData->GetNumChannels();
+	int32 NumInputChannels = SoundWaveProxy->GetNumChannels();
 
 	uint32 StartFrameIndex = LerpArray[0].PosA;
 	int32 NumSourceFramesNeeded = CalculateNumFramesNeeded(LerpArray, InNumFrames);
@@ -247,12 +241,12 @@ void FStreamingAudioRendererV2::RenderMultiChannelUnshifted(TAudioBuffer<float>&
 	OutBuffer.ZeroValidFrames();
 
 	int32 NumOutChannels = OutBuffer.GetNumValidChannels();
-	int32 NumInputChannels = SoundWaveData->GetNumChannels();
+	int32 NumInputChannels = SoundWaveProxy->GetNumChannels();
 
 	// interleaved
 	Audio::FAlignedFloatBuffer ResampleBuffer;
 	constexpr int32 kNumResampleChannels = 2;
-	ResampleBuffer.SetNumZeroed(kNumResampleChannels * InNumFrames);
+	ResampleBuffer.SetNum(kNumResampleChannels * InNumFrames);
 
 	uint32 StartFrameIndex = LerpArray[0].PosA;
 	int32 NumSourceFramesNeeded = CalculateNumFramesNeeded(LerpArray, InNumFrames);
@@ -320,7 +314,7 @@ void FStreamingAudioRendererV2::RenderMultiChannelRoutedUnshifted(TAudioBuffer<f
 	OutBuffer.ZeroValidFrames();
 
 	int32 NumOutChannels = OutBuffer.GetNumValidChannels();
-	int32 NumInputChannels = SoundWaveData->GetNumChannels();
+	int32 NumInputChannels = SoundWaveProxy->GetNumChannels();
 
 	// interleaved
 	Audio::FAlignedFloatBuffer ResampleBuffer;
@@ -509,32 +503,31 @@ uint32 FStreamingAudioRendererV2::GetSourceAudioFrameIndex()
 
 }
 
-TUniquePtr<FSoundWaveProxyReader> FStreamingAudioRendererV2::CreateProxyReader(TSharedRef<const FSoundWaveData> WaveData)
+TUniquePtr<FSoundWaveProxyReader> FStreamingAudioRendererV2::CreateProxyReader(TSharedRef<FSoundWaveProxy> WaveProxy)
 {
 	FSoundWaveProxyReader::FSettings WaveReaderSettings;
 	WaveReaderSettings.MaxDecodeSizeInFrames = MaxDecodeSizeInFrames;
 	WaveReaderSettings.StartTimeInSeconds = 0.0f;
 	WaveReaderSettings.bMaintainAudioSync = true;
-	return FSoundWaveProxyReader::Create(WaveData, WaveReaderSettings);	
+	return FSoundWaveProxyReader::Create(WaveProxy, WaveReaderSettings);	
 }
 
 bool FStreamingAudioRendererV2::HasLoopSection() const
 {
-	check(SoundWaveData.IsValid());
-	return SoundWaveData->GetLoopRegions().Num() > 0;
+	check(SoundWaveProxy);
+	return SoundWaveProxy->GetLoopRegions().Num() > 0;
 }
 
 uint32 FStreamingAudioRendererV2::GetLoopStartFrame() const
 {
 	check(HasLoopSection());
-	check(SoundWaveData.IsValid());
-	return SoundWaveData->GetLoopRegions()[0].FramePosition;
+	return SoundWaveProxy->GetLoopRegions()[0].FramePosition;
 }
 
 uint32 FStreamingAudioRendererV2::GetLoopEndFrame() const
 {
-	check(HasLoopSection());
-	const FSoundWaveCuePoint& LoopRegion = SoundWaveData->GetLoopRegions()[0];
+	check(HasLoopSection())
+	const FSoundWaveCuePoint& LoopRegion = SoundWaveProxy->GetLoopRegions()[0];
 	return LoopRegion.FramePosition + LoopRegion.FrameLength;
 }
 
@@ -661,6 +654,6 @@ void FStreamingAudioRendererV2::GenerateSourceAudioInternal(uint32 StartFrameInd
 		FMemory::Memzero(&OutAudioData[BufferIdx], sizeof(float) * NumSamplesRequested);
 
 		UE_LOG(LogHarmonixStreamingAudioRendererV2, Verbose, TEXT("%s: Failed to generated samples: StartFrameIndex: %d, NumSamplesRequested: %d"), 
-			   *SoundWaveData->GetFName().ToString(), StartFrameIndex, NumSamplesRequested);
+			*SoundWaveProxy->GetFName().ToString(), StartFrameIndex, NumSamplesRequested);
 	}
 }
diff --git a/Engine/Plugins/Runtime/Harmonix/Source/HarmonixDsp/Private/HarmonixDsp/AudioData/StreamingAudioRendererV2.h b/Engine/Plugins/Runtime/Harmonix/Source/HarmonixDsp/Private/HarmonixDsp/AudioData/StreamingAudioRendererV2.h
--- a/Engine/Plugins/Runtime/Harmonix/Source/HarmonixDsp/Private/HarmonixDsp/AudioData/StreamingAudioRendererV2.h
+++ b/Engine/Plugins/Runtime/Harmonix/Source/HarmonixDsp/Private/HarmonixDsp/AudioData/StreamingAudioRendererV2.h
@@ -13,7 +13,6 @@
 
 #define UE_API HARMONIXDSP_API
 
-class FSoundWaveData;
 class FSoundWaveProxy;
 class FSoundWaveProxyReader;
 class FFusionSampler;
@@ -31,10 +30,8 @@ public:
 	UE_API ~FStreamingAudioRendererV2();
 
 	UE_API virtual void Reset() override;
-
-	UE_API virtual void SetSoundWaveData(const TSharedRef<const FSoundWaveData>& InSoundWaveData, const FSettings& InSettings) override;
-
-	UE_API virtual TSharedPtr<const FSoundWaveData> GetSoundWaveData() const override;
+	UE_API virtual void SetAudioData(TSharedRef<FSoundWaveProxy> SoundWaveProxy, const FSettings& InSettings) override;
+	UE_API virtual const TSharedPtr<FSoundWaveProxy> GetAudioData() const override;
 	
 	UE_API virtual void MigrateToSampler(const FFusionSampler* InSampler) override;
 
@@ -63,7 +60,7 @@ private:
 
 	uint32 GetSourceAudioFrameIndex();
 
-	static TUniquePtr<FSoundWaveProxyReader> CreateProxyReader(TSharedRef<const FSoundWaveData> WaveData);
+	static TUniquePtr<FSoundWaveProxyReader> CreateProxyReader(TSharedRef<FSoundWaveProxy> WaveProxy);
 
 	bool HasLoopSection() const;
 	uint32 GetLoopStartFrame() const;
@@ -75,8 +72,8 @@ private:
 
 	// Ref to the actual streaming audio data
 	// this data is a shared instance of audio data
-	// it is set by SetSoundWaveData()
-	TSharedPtr<const FSoundWaveData> SoundWaveData;
+	// it will get loaded on construction 
+	TSharedPtr<FSoundWaveProxy> SoundWaveProxy;
 
 	TUniquePtr<FSoundWaveProxyReader> WaveProxyReader;
 
diff --git a/Engine/Plugins/Runtime/Harmonix/Source/HarmonixDsp/Private/HarmonixDsp/AudioDataRenderer.cpp b/Engine/Plugins/Runtime/Harmonix/Source/HarmonixDsp/Private/HarmonixDsp/AudioDataRenderer.cpp
--- a/Engine/Plugins/Runtime/Harmonix/Source/HarmonixDsp/Private/HarmonixDsp/AudioDataRenderer.cpp
+++ b/Engine/Plugins/Runtime/Harmonix/Source/HarmonixDsp/Private/HarmonixDsp/AudioDataRenderer.cpp
@@ -3,55 +3,20 @@
 #include "HarmonixDsp/AudioDataRenderer.h"
 #include "Sound/SoundWave.h"
 
-DEFINE_LOG_CATEGORY(LogHarmonixAudioDataRenderer);
-
-void IAudioDataRenderer::SetSoundWaveData(const TSharedRef<const FSoundWaveData>& InSoundWaveData, const FSettings& InSettings)
-{
-	static bool bDidWarn = false;
-	if (!bDidWarn)
-	{
-		UE_LOG(LogHarmonixAudioDataRenderer, Warning, TEXT("IAudioDataRenderer::SetSoundWaveData() must be overridden, and will be pure virtual in a future release."));
-		bDidWarn = true;
-	}
-	
-	PRAGMA_DISABLE_DEPRECATION_WARNINGS
-	SetAudioData(ConstCastSharedRef<FSoundWaveProxy>(FSoundWaveProxy::CreateNonUpdatingCopy(InSoundWaveData)), InSettings);
-	PRAGMA_ENABLE_DEPRECATION_WARNINGS
-}
-
-TSharedPtr<const FSoundWaveData> IAudioDataRenderer::GetSoundWaveData() const
-{
-	static bool bDidWarn = false;
-	if (!bDidWarn)
-	{
-		UE_LOG(LogHarmonixAudioDataRenderer, Warning, TEXT("IAudioDataRenderer::GetSoundWaveData() must be overridden, and will be pure virtual in a future release."));
-		bDidWarn = true;
-	}
-	
-	PRAGMA_DISABLE_DEPRECATION_WARNINGS
-	if (const TSharedPtr<FSoundWaveProxy> Proxy = GetAudioData())
-	{
-		return Proxy->GetSoundWaveDataRef().ToSharedPtr();
-	}
-	PRAGMA_ENABLE_DEPRECATION_WARNINGS
-
-	return {};
-}
-
 double IAudioDataRenderer::CalculateLerpData(FLerpData* LerpArray, int32 InNumPoints, uint32 InNumOutFrames, double InPos, int32 InMaxFrame, bool InHonorLoopPoints, double InInc) const
 {
 	check(InNumPoints >= (int32)InNumOutFrames);
 	FMemory::Memset(LerpArray, 0, sizeof(FLerpData) * InNumPoints);
 
-	TSharedPtr<const FSoundWaveData> SoundWaveData = GetSoundWaveData();
-	if (!SoundWaveData)
+	const TSharedPtr<FSoundWaveProxy> SoundWaveProxy = GetAudioData();
+	if (!SoundWaveProxy)
 	{
 		return 0.0;
 	}
 
-	if (InMaxFrame < 0 || InMaxFrame > (int32)SoundWaveData->GetNumFrames())
+	if (InMaxFrame < 0 || InMaxFrame > (int32)SoundWaveProxy->GetNumFrames())
 	{
-		InMaxFrame = (int32)SoundWaveData->GetNumFrames();
+		InMaxFrame = (int32)SoundWaveProxy->GetNumFrames();
 	}
 
 	double NumSourceFrames = (double)InMaxFrame;
@@ -60,13 +25,12 @@ double IAudioDataRenderer::CalculateLerpData(FLerpData* LerpArray, int32 InNumPo
 	uint32 LoopEndFrameIndex = 0;
 	bool IsLooping = false;
 	bool HasTailSection = false;
-	if (SoundWaveData->GetLoopRegions().Num() > 0)
+	if (SoundWaveProxy->GetLoopRegions().Num() > 0)
 	{
-		const uint32 NumFramesClamped = FMath::Max(1, SoundWaveData->GetNumFrames()) - 1;
-		const FSoundWaveCuePoint& LoopRegion = SoundWaveData->GetLoopRegions()[0];
+		const FSoundWaveCuePoint& LoopRegion = SoundWaveProxy->GetLoopRegions()[0];
 		LoopStartFrameIndex = LoopRegion.FramePosition;
 		LoopEndFrameIndex = LoopRegion.FramePosition + LoopRegion.FrameLength;
-		HasTailSection = LoopEndFrameIndex != NumFramesClamped;
+		HasTailSection = LoopEndFrameIndex != SoundWaveProxy->GetNumFrames() - 1;
 		IsLooping = (!HasTailSection || (HasTailSection && InHonorLoopPoints));
 	}
 
diff --git a/Engine/Plugins/Runtime/Harmonix/Source/HarmonixDsp/Private/HarmonixDsp/FusionSampler/FusionSampler.cpp b/Engine/Plugins/Runtime/Harmonix/Source/HarmonixDsp/Private/HarmonixDsp/FusionSampler/FusionSampler.cpp
--- a/Engine/Plugins/Runtime/Harmonix/Source/HarmonixDsp/Private/HarmonixDsp/FusionSampler/FusionSampler.cpp
+++ b/Engine/Plugins/Runtime/Harmonix/Source/HarmonixDsp/Private/HarmonixDsp/FusionSampler/FusionSampler.cpp
@@ -383,7 +383,7 @@ int32 FFusionSampler::GatherMatchingKeyzones(uint8 TransposedNote, uint8 InVeloc
 	for (int32 KeyzoneIdx = 0; KeyzoneIdx < InKeyzones.Num(); ++KeyzoneIdx)
 	{
 		const FKeyzoneSettings& Keyzone = InKeyzones[KeyzoneIdx];
-		if (!Keyzone.SoundWaveData)
+		if (!Keyzone.SoundWaveProxy)
 		{
 			continue;
 		}
diff --git a/Engine/Plugins/Runtime/Harmonix/Source/HarmonixDsp/Private/HarmonixDsp/FusionSampler/FusionVoice.cpp b/Engine/Plugins/Runtime/Harmonix/Source/HarmonixDsp/Private/HarmonixDsp/FusionSampler/FusionVoice.cpp
--- a/Engine/Plugins/Runtime/Harmonix/Source/HarmonixDsp/Private/HarmonixDsp/FusionSampler/FusionVoice.cpp
+++ b/Engine/Plugins/Runtime/Harmonix/Source/HarmonixDsp/Private/HarmonixDsp/FusionSampler/FusionVoice.cpp
@@ -66,7 +66,7 @@ bool FFusionVoice::AssignIDs(FFusionSampler* InSampler, const FKeyzoneSettings*
 {
 	check(InSampler);
 	check(InKeyZone);
-	check(InKeyZone->SoundWaveData);
+	check(InKeyZone->SoundWaveProxy);
 
 #if CPUPROFILERTRACE_ENABLED
 	const IConsoleVariable* ProfilingAllGraphsCheat = IConsoleManager::Get().FindConsoleVariable(TEXT("au.MetaSound.ProfileAllGraphs"));
@@ -117,7 +117,8 @@ bool FFusionVoice::AssignIDs(FFusionSampler* InSampler, const FKeyzoneSettings*
 	// so in that case set NumInChannels to 0 here and we'll sort this
 	// out later...
 	// 0 because we need to configure to the output buffer we are asked to fill
-	int32 NumInChannels = InKeyZone->TrackMap.Num() == 0 ? KeyZone->SoundWaveData->GetNumChannels() : 0; 
+	
+	int32 NumInChannels = InKeyZone->TrackMap.Num() == 0 ? KeyZone->SoundWaveProxy->GetNumChannels() : 0; 
 	EAudioBufferChannelLayout ChannelLayout = HarmonixDsp::FAudioBuffer::GetDefaultChannelLayoutForChannelCount(NumInChannels);
 	ESpeakerMask::Type ChannelMask = HarmonixDsp::FAudioBuffer::GetChannelMaskForNumChannels(NumInChannels);
 	Panner.Setup(KeyZone->Pan,
@@ -130,7 +131,7 @@ bool FFusionVoice::AssignIDs(FFusionSampler* InSampler, const FKeyzoneSettings*
 
 	// compare the audio file's sample rate to the output sample rate
 	// so we can make them match
-	double FileSamplesPerSecond = KeyZone->SoundWaveData->GetSampleRate();
+	double FileSamplesPerSecond = InKeyZone->SoundWaveProxy->GetSampleRate();
 	FileToOutputSampleRatio = FileSamplesPerSecond * SecondsPerSample;
 
 	IAudioDataRenderer::FSettings Settings;
@@ -139,9 +140,9 @@ bool FFusionVoice::AssignIDs(FFusionSampler* InSampler, const FKeyzoneSettings*
 	Settings.TrackChannelInfo = &InKeyZone->TrackMap;
 	Settings.Sampler = MySampler;
 	//TRACE_BOOKMARK(TEXT("Starting Waveform: %s"), *KeyZone->AudioSample->GetName().ToString());
-	ActiveRenderer->SetSoundWaveData(KeyZone->SoundWaveData.ToSharedRef(), Settings);
+	ActiveRenderer->SetAudioData(InKeyZone->SoundWaveProxy.ToSharedRef(), Settings);
 
-	OutputBuffer.SetNumValidChannels(KeyZone->SoundWaveData->GetNumChannels());
+	OutputBuffer.SetNumValidChannels(InKeyZone->SoundWaveProxy->GetNumChannels());
 
 	VoiceID = InVoiceID;
 	bWaitingForAttack = true;
@@ -159,7 +160,7 @@ void FFusionVoice::Attack()
 {
 	check(MySampler);
 	checkSlow(KeyZone);
-	checkSlow(KeyZone->SoundWaveData);
+	checkSlow(KeyZone->SoundWaveProxy);
 
 	AdsrVolume.Attack();
 	AdsrAssignable.Attack();
@@ -283,7 +284,7 @@ void FFusionVoice::SetSampler(FFusionSampler* InSampler)
 
 void FFusionVoice::SetPitchShiftForMidiNote(uint8 InNote)
 {
-	check(KeyZone && KeyZone->SoundWaveData);
+	check(KeyZone);
 
 	// this is the note (pitch) we are supposedly playing
 	// although the exact pitch will vary depending on various modulations
@@ -395,14 +396,14 @@ void FFusionVoice::AttackWithTargetNote(uint8 InMidiNoteNumber, float InGain, in
 	double KeyzoneStartOffset = (KeyZone->SampleStartOffset == -1) ?
 		0 :
 		KeyZone->SampleStartOffset;
-	
+
 	SamplePos = KeyzoneStartOffset + (OutputSampleOffset * FileToOutputSampleRatio);
 	StartPos = SamplePos;
 	EndOfSampleData = (KeyZone->SampleEndOffset == -1) ?
-		KeyZone->SoundWaveData->GetNumFrames() :
+		KeyZone->SoundWaveProxy->GetNumFrames() :
 		KeyZone->SampleEndOffset;
 
-	bool HasLoopSection = KeyZone->SoundWaveData->GetLoopRegions().Num() > 0;
+	bool HasLoopSection = KeyZone->SoundWaveProxy->GetLoopRegions().Num() > 0;
 	if (SamplePos > EndOfSampleData && !HasLoopSection)
 	{
 		Kill();
@@ -435,7 +436,7 @@ void FFusionVoice::PrepareWithPitchOffsetAndGain(double InPitchOffsetCents, floa
 
 	checkSlow(MySampler);
 	checkSlow(KeyZone);
-	checkSlow(KeyZone->SoundWaveData);
+	checkSlow(KeyZone->SoundWaveProxy);
 
 	for (int32 Idx = 0; Idx < kNumLfos; ++Idx)
 	{
@@ -674,7 +675,17 @@ uint32 FFusionVoice::Process(uint32 InSliceIndex, uint32 InSubsliceIndex, float*
 
 	using namespace Harmonix::Dsp::Modulators;
 
-	if (!IsInUse() || !MySampler || !KeyZone || !KeyZone->SoundWaveData)
+	if (!IsInUse())
+	{
+		return 0;
+	}
+
+	if (!MySampler || !KeyZone)
+	{
+		return 0;
+	}
+
+	if (!KeyZone->SoundWaveProxy)
 	{
 		return 0;
 	}
@@ -766,8 +777,8 @@ uint32 FFusionVoice::Process(uint32 InSliceIndex, uint32 InSubsliceIndex, float*
 				{
 					CurrentSampleFrame -= KeyZone->SampleStartOffset;
 				}
-				
-				float ActualElapsedMs = (CurrentSampleFrame * 1000.0f) / FMath::Max(1, KeyZone->SoundWaveData->GetSampleRate());
+
+				float ActualElapsedMs = (CurrentSampleFrame * 1000.0f) / KeyZone->SoundWaveProxy->GetSampleRate();
 				float ErrorMs = ActualElapsedMs - ExpectedElapsedMs;
 
 				if (!bHasRenderedAnySamples)
@@ -777,11 +788,11 @@ uint32 FFusionVoice::Process(uint32 InSliceIndex, uint32 InSubsliceIndex, float*
 						UE_LOG(LogFusionVoice, Verbose, TEXT("Adjusting start time due to error of %f ms"), ErrorMs);
 					}
 
-					SamplePos = ExpectedElapsedMs / 1000.0f * (float)KeyZone->SoundWaveData->GetSampleRate();
+					SamplePos = ExpectedElapsedMs / 1000.0f * (float)KeyZone->SoundWaveProxy->GetSampleRate();
 					
 					if (SamplePos > 0.0)
 					{
-						UE_LOG(LogFusionVoice, Verbose, TEXT("SHAVING %f ms from the beginning of %s"), ErrorMs, *KeyZone->SoundWaveData->GetFName().ToString());
+						UE_LOG(LogFusionVoice, Verbose, TEXT("SHAVING %f ms from the beginning of %s"), ErrorMs, *KeyZone->SoundWaveProxy->GetFName().ToString());
 					}
 
 					if (SamplePos < 0.0)
@@ -792,7 +803,7 @@ uint32 FFusionVoice::Process(uint32 InSliceIndex, uint32 InSubsliceIndex, float*
 					{
 						SamplePos += KeyZone->SampleStartOffset;
 					}
-					if (SamplePos > EndOfSampleData && KeyZone->SoundWaveData->GetLoopRegions().IsEmpty())
+					if (SamplePos > EndOfSampleData && KeyZone->SoundWaveProxy->GetLoopRegions().IsEmpty())
 					{
 						Kill();
 						return 0;
@@ -804,7 +815,7 @@ uint32 FFusionVoice::Process(uint32 InSliceIndex, uint32 InSubsliceIndex, float*
 				{
 					if (ErrorMs > 20.0f)
 					{
-						UE_LOG(LogFusionVoice, Verbose, TEXT("Fixing %s offset -> %f\n"), *KeyZone->SoundWaveData->GetFName().ToString(), ErrorMs);
+						UE_LOG(LogFusionVoice, Verbose, TEXT("Fixing %s offset -> %f\n"), *KeyZone->SoundWaveProxy->GetFName().ToString(), ErrorMs);
 					}
 					CurrentVso = -(FusionVoice::kVsoAmount * RenderSpeed);
 				}
@@ -812,7 +823,7 @@ uint32 FFusionVoice::Process(uint32 InSliceIndex, uint32 InSubsliceIndex, float*
 				{
 					if (ErrorMs < -20.0f)
 					{
-						UE_LOG(LogFusionVoice, Verbose, TEXT("Fixing %s offset -> %f\n"), *KeyZone->SoundWaveData->GetFName().ToString(), ErrorMs);
+						UE_LOG(LogFusionVoice, Verbose, TEXT("Fixing %s offset -> %f\n"), *KeyZone->SoundWaveProxy->GetFName().ToString(), ErrorMs);
 					}
 					CurrentVso = FusionVoice::kVsoAmount * RenderSpeed;
 				}
@@ -844,7 +855,7 @@ uint32 FFusionVoice::Process(uint32 InSliceIndex, uint32 InSubsliceIndex, float*
 
 		// We probably have a trackmap, and this is the first render call,
 		// so we have to configure the panner for the output buffer setup...
-		int32 NumInChannels = KeyZone->SoundWaveData->GetNumChannels();
+		int32 NumInChannels = KeyZone->SoundWaveProxy->GetNumChannels();
 		if (NumInChannels != Panner.GetCurrentGainMatrix().GetNumInChannels())
 		{
 			Panner.Configure(NumInChannels, OutputBuffer);
diff --git a/Engine/Plugins/Runtime/Harmonix/Source/HarmonixDsp/Private/HarmonixDsp/FusionSampler/FusionVoicePool.cpp b/Engine/Plugins/Runtime/Harmonix/Source/HarmonixDsp/Private/HarmonixDsp/FusionSampler/FusionVoicePool.cpp
--- a/Engine/Plugins/Runtime/Harmonix/Source/HarmonixDsp/Private/HarmonixDsp/FusionSampler/FusionVoicePool.cpp
+++ b/Engine/Plugins/Runtime/Harmonix/Source/HarmonixDsp/Private/HarmonixDsp/FusionSampler/FusionVoicePool.cpp
@@ -294,7 +294,7 @@ FFusionVoice* FFusionVoicePool::GetFreeVoice(
 		return nullptr;
 	}
 
-	if (!InKeyzone->SoundWaveData)
+	if (!InKeyzone->SoundWaveProxy)
 	{
 		UE_LOG(LogFusionVoicePool, Warning, TEXT("Asked to allocate a InSampler voice for an invalid sample"));
 		return nullptr;
@@ -308,7 +308,7 @@ FFusionVoice* FFusionVoicePool::GetFreeVoice(
 			FName FactoryName = InKeyzone->TimeStretchConfig.PitchShifter.Name;
 			UE_LOG(LogFusionVoicePool, Warning,
 				TEXT("Attempting to play a fusion InKeyzone that is set to \"maintain time\", but no shifter factory was available for assigned Pitch Shifter: %s. Check project configuration! (%s)"),
-				*FactoryName.ToString(), *InKeyzone->SoundWaveData->GetFName().ToString());
+				*FactoryName.ToString(), *InKeyzone->SoundWaveProxy->GetFName().ToString());
 			return nullptr;
 		}
 
diff --git a/Engine/Plugins/Runtime/Harmonix/Source/HarmonixDsp/Private/HarmonixDsp/FusionSampler/Settings/KeyzoneSettings.cpp b/Engine/Plugins/Runtime/Harmonix/Source/HarmonixDsp/Private/HarmonixDsp/FusionSampler/Settings/KeyzoneSettings.cpp
--- a/Engine/Plugins/Runtime/Harmonix/Source/HarmonixDsp/Private/HarmonixDsp/FusionSampler/Settings/KeyzoneSettings.cpp
+++ b/Engine/Plugins/Runtime/Harmonix/Source/HarmonixDsp/Private/HarmonixDsp/FusionSampler/Settings/KeyzoneSettings.cpp
@@ -19,9 +19,9 @@ FKeyzoneSettings::FKeyzoneSettings()
 
 FKeyzoneSettings::~FKeyzoneSettings()
 {
-	if (SoundWaveData.IsValid() && SoundWaveData->GetLoadingBehavior() == ESoundWaveLoadingBehavior::ForceInline)
+	if (SoundWaveProxy && SoundWaveProxy->GetLoadingBehavior() == ESoundWaveLoadingBehavior::ForceInline)
 	{
-		IStreamingManager::Get().GetAudioStreamingManager().RemoveForceInlineSoundWave(SoundWaveData.ToSharedRef());
+		IStreamingManager::Get().GetAudioStreamingManager().RemoveForceInlineSoundWave(SoundWaveProxy);
 	}
 }
 
@@ -37,20 +37,19 @@ void FKeyzoneSettings::InitProxyData(const Audio::FProxyDataInitParams& InitPara
 		return;
 	}
 
-	if (SoundWaveData.IsValid() && SoundWaveData->GetLoadingBehavior() == ESoundWaveLoadingBehavior::ForceInline)
+	if (SoundWaveProxy && SoundWaveProxy->GetLoadingBehavior() == ESoundWaveLoadingBehavior::ForceInline)
 	{
 		// I don't think we'de ever hit this code, but let's just do this to be sure.
-		IStreamingManager::Get().GetAudioStreamingManager().RemoveForceInlineSoundWave(SoundWaveData.ToSharedRef());
+		IStreamingManager::Get().GetAudioStreamingManager().RemoveForceInlineSoundWave(SoundWaveProxy);
 	}
 	
 	TSharedPtr<Audio::IProxyData> AudioProxy = SoundWave->CreateProxyData(InitParams);
-	TSharedPtr<FSoundWaveProxy> TempSoundWaveProxy = StaticCastSharedPtr<FSoundWaveProxy>(AudioProxy);
-	SoundWaveData = TempSoundWaveProxy->GetSoundWaveDataRef();
-	check(SoundWaveData);
+	SoundWaveProxy = StaticCastSharedPtr<FSoundWaveProxy>(AudioProxy);
+	check(SoundWaveProxy);
 		
-	if (SoundWaveData->GetLoadingBehavior() == ESoundWaveLoadingBehavior::ForceInline)
+	if (SoundWaveProxy->GetLoadingBehavior() == ESoundWaveLoadingBehavior::ForceInline)
 	{
-		IStreamingManager::Get().GetAudioStreamingManager().AddForceInlineSoundWave(SoundWaveData.ToSharedRef());
+		IStreamingManager::Get().GetAudioStreamingManager().AddForceInlineSoundWave(SoundWaveProxy);
 	}
 
 	// nulling out sound wave asset for this Keyzone
diff --git a/Engine/Plugins/Runtime/Harmonix/Source/HarmonixDsp/Private/HarmonixDsp/Stretchers/SmbPitchShifter.cpp b/Engine/Plugins/Runtime/Harmonix/Source/HarmonixDsp/Private/HarmonixDsp/Stretchers/SmbPitchShifter.cpp
--- a/Engine/Plugins/Runtime/Harmonix/Source/HarmonixDsp/Private/HarmonixDsp/Stretchers/SmbPitchShifter.cpp
+++ b/Engine/Plugins/Runtime/Harmonix/Source/HarmonixDsp/Private/HarmonixDsp/Stretchers/SmbPitchShifter.cpp
@@ -189,9 +189,9 @@ void FSmbPitchShifter::SetSampleRateAndReset(float InSampleRate)
 
 void FSmbPitchShifter::SetupAndResetImpl(const FTimeStretchConfig* ShifterSettings)
 {
-	check(MySoundWaveData);
-	SampleRate = MySoundWaveData->GetSampleRate();
-	FreqPerFftBin = (double)MySoundWaveData->GetSampleRate() / (double)FftFrameSize;
+	check(MyAudioData);
+	SampleRate = MyAudioData->GetSampleRate();
+	FreqPerFftBin = (double)MyAudioData->GetSampleRate() / (double)FftFrameSize;
 	Reset();
 }
 
diff --git a/Engine/Plugins/Runtime/Harmonix/Source/HarmonixDsp/Public/HarmonixDsp/AudioDataRenderer.h b/Engine/Plugins/Runtime/Harmonix/Source/HarmonixDsp/Public/HarmonixDsp/AudioDataRenderer.h
--- a/Engine/Plugins/Runtime/Harmonix/Source/HarmonixDsp/Public/HarmonixDsp/AudioDataRenderer.h
+++ b/Engine/Plugins/Runtime/Harmonix/Source/HarmonixDsp/Public/HarmonixDsp/AudioDataRenderer.h
@@ -5,15 +5,12 @@
 
 #include "HarmonixDsp/Streaming/TrackChannelInfo.h"
 
-DECLARE_LOG_CATEGORY_EXTERN(LogHarmonixAudioDataRenderer, Log, All);
-
 class IStretcherAndPitchShifter;
 struct FTrackChannelInfo;
 template<typename Type> 
 class TAudioBuffer;
 class FGainMatrix;
 class FFusionSampler;
-class FSoundWaveData;
 class FSoundWaveProxy;
 struct FTimeStretchConfig;
 
@@ -32,17 +29,11 @@ public:
 
 	virtual ~IAudioDataRenderer() = default;
 
-	UE_DEPRECATED(5.8, "Mutable get/set is deprecated, please call SetSoundWaveData() instead")
-	virtual void SetAudioData(TSharedRef<FSoundWaveProxy> SoundWaveProxy, const FSettings& InSettings) {};
-
-	UE_DEPRECATED(5.8, "Mutable get/set is deprecated, please call GetSoundWaveData() instead")
-	virtual const TSharedPtr<FSoundWaveProxy> GetAudioData() const { return {}; };
-
-	virtual void SetSoundWaveData(const TSharedRef<const FSoundWaveData>& InSoundWaveData, const FSettings& InSettings);
-	virtual TSharedPtr<const FSoundWaveData> GetSoundWaveData() const;
+	virtual void SetAudioData(TSharedRef<FSoundWaveProxy> SoundWaveProxy, const FSettings& InSettings) = 0;
 
 	virtual void Reset() = 0;
 	
+	virtual const TSharedPtr<FSoundWaveProxy> GetAudioData() const = 0;
 
 	virtual void MigrateToSampler(const FFusionSampler* InSampler) = 0;
 
diff --git a/Engine/Plugins/Runtime/Harmonix/Source/HarmonixDsp/Public/HarmonixDsp/FusionSampler/Settings/KeyzoneSettings.h b/Engine/Plugins/Runtime/Harmonix/Source/HarmonixDsp/Public/HarmonixDsp/FusionSampler/Settings/KeyzoneSettings.h
--- a/Engine/Plugins/Runtime/Harmonix/Source/HarmonixDsp/Public/HarmonixDsp/FusionSampler/Settings/KeyzoneSettings.h
+++ b/Engine/Plugins/Runtime/Harmonix/Source/HarmonixDsp/Public/HarmonixDsp/FusionSampler/Settings/KeyzoneSettings.h
@@ -69,13 +69,13 @@ struct FKeyzoneSettings
 	}
 
 public:
+	
 	// Not null when this keyzone is not being proxied
+	// Made private since we generally want the Proxy Data
 	UPROPERTY(EditDefaultsOnly, Category = "Audio Sample")
 	TObjectPtr<USoundWave> SoundWave;
 
 	TSharedPtr<FSoundWaveProxy> SoundWaveProxy;
-	TSharedPtr<const FSoundWaveData> SoundWaveData;
-
 
 	UPROPERTY(EditDefaultsOnly, Category = "Settings", Meta = (UIMin=0, UIMax=127, ClampMin=0, ClampMax=127))
 	int8 RootNote = 60;
diff --git a/Engine/Plugins/Runtime/Harmonix/Source/HarmonixDsp/Public/HarmonixDsp/StretcherAndPitchShifter.h b/Engine/Plugins/Runtime/Harmonix/Source/HarmonixDsp/Public/HarmonixDsp/StretcherAndPitchShifter.h
--- a/Engine/Plugins/Runtime/Harmonix/Source/HarmonixDsp/Public/HarmonixDsp/StretcherAndPitchShifter.h
+++ b/Engine/Plugins/Runtime/Harmonix/Source/HarmonixDsp/Public/HarmonixDsp/StretcherAndPitchShifter.h
@@ -6,7 +6,6 @@
 #include "HarmonixDsp/AudioBuffer.h"
 #include "HarmonixDsp/PitchShifterName.h"
 #include "HarmonixDsp/Containers/TypedParameter.h"
-#include "Sound/SoundWave.h"
 
 #include "StretcherAndPitchShifter.generated.h"
 
@@ -41,10 +40,7 @@ class IStretcherAndPitchShifter
 {
 public:
 	IStretcherAndPitchShifter(FName InFactoryName) : MyFactoryName(InFactoryName) {}
-	
-	PRAGMA_DISABLE_DEPRECATION_WARNINGS
 	virtual ~IStretcherAndPitchShifter() = default;
-	PRAGMA_ENABLE_DEPRECATION_WARNINGS
 
 	FName GetFactoryName() const { return MyFactoryName; };
 
@@ -57,27 +53,9 @@ public:
 	virtual bool HasCurrentSampleFrame() const { return false; }
 	virtual int32  GetCurrentSampleFrame() const { return 0; }
 
-
-	void SetSampleSourceReset(const TSharedRef<const FSoundWaveData> InSoundWaveData, TSharedPtr<IAudioDataRenderer, ESPMode::ThreadSafe> InAudioRenderer, const FTimeStretchConfig* ShifterSettings)
-	{ 
-		MySoundWaveData = InSoundWaveData;
-
-		MyAudioRenderer = InAudioRenderer;
-
-		SetupAndResetImpl(ShifterSettings);
-	}
-
-	UE_DEPRECATED(5.8, "Please use the TSharedRef<const FSoundWaveData> overload instead")
 	void SetSampleSourceReset(TSharedPtr<FSoundWaveProxy> InSoundWave, TSharedPtr<IAudioDataRenderer, ESPMode::ThreadSafe> InAudioRenderer, const FTimeStretchConfig* ShifterSettings)
-	{
-		if (InSoundWave.IsValid())
-		{
-			SetSampleSourceReset(InSoundWave->GetSoundWaveDataRef(), InAudioRenderer, ShifterSettings);
-
-			return;
-		}
-
-		MySoundWaveData.Reset();
+	{ 
+		MyAudioData = InSoundWave;
 		MyAudioRenderer = InAudioRenderer;
 
 		SetupAndResetImpl(ShifterSettings);
@@ -96,11 +74,7 @@ public:
 
 	virtual void Cleanup() = 0;
 protected:
-	TSharedPtr<const FSoundWaveData> MySoundWaveData;
-	
-	UE_DEPRECATED(5.8, "Please use MySoundWaveData instead")
 	TSharedPtr<FSoundWaveProxy> MyAudioData;
-	
 	TSharedPtr<IAudioDataRenderer, ESPMode::ThreadSafe> MyAudioRenderer;
 
 	virtual void SetupAndResetImpl(const FTimeStretchConfig* ShifterSettings) { }
diff --git a/Engine/Plugins/Runtime/Metasound/Source/MetasoundEngine/Private/MetasoundWave.cpp b/Engine/Plugins/Runtime/Metasound/Source/MetasoundEngine/Private/MetasoundWave.cpp
--- a/Engine/Plugins/Runtime/Metasound/Source/MetasoundEngine/Private/MetasoundWave.cpp
+++ b/Engine/Plugins/Runtime/Metasound/Source/MetasoundEngine/Private/MetasoundWave.cpp
@@ -21,19 +21,17 @@ FAutoConsoleVariableRef CVarDisableMetasoundWaveAssetCachePriming(
 
 namespace Metasound
 {
-	FWaveAsset::FWaveAsset(const TSharedPtr<const Audio::IProxyData>& InInitData)
+	FWaveAsset::FWaveAsset(const TSharedPtr<Audio::IProxyData>& InInitData)
 	{
 		if (InInitData.IsValid())
 		{
 			if (InInitData->CheckTypeCast<FSoundWaveProxy>())
 			{
-				// todo: remove const cast once deprecated getters can be removed and SoundWaveProxy is const.
-				SoundWaveProxy = ConstCastSharedPtr<FSoundWaveProxy>(InInitData->GetAs<const FSoundWaveProxy>().AsShared().ToSharedPtr());
+				// should we be getting handed a SharedPtr here?
+				SoundWaveProxy = MakeShared<FSoundWaveProxy, ESPMode::ThreadSafe>(InInitData->GetAs<FSoundWaveProxy>());
 
 				if (ensureAlways(SoundWaveProxy.IsValid()))
 				{
-					TSharedRef<const FSoundWaveData> SoundWaveData = SoundWaveProxy->GetSoundWaveDataRef();
-					
 					// TODO HACK: Prime the sound for playback.
 					//
 					// Preferably playback latency would be controlled externally.
@@ -52,14 +50,14 @@ namespace Metasound
 					// time that the decoder attempts to decode audio.
 					if (0 == DisableMetasoundWaveAssetCachePriming)
 					{
-						if (SoundWaveData->IsStreaming())
+						if (SoundWaveProxy->IsStreaming())
 						{
-							if (SoundWaveData->GetNumChunks() > 1)
+							if (SoundWaveProxy->GetNumChunks() > 1)
 							{
 								if (FStreamingManagerCollection* StreamingMgr = IStreamingManager::Get_Concurrent())
 								{
 									IAudioStreamingManager& AudioStreamingMgr = StreamingMgr->GetAudioStreamingManager();
-									AudioStreamingMgr.RequestChunk(SoundWaveData->AsShared(), 1, [](EAudioChunkLoadResult) {});
+									AudioStreamingMgr.RequestChunk(SoundWaveProxy, 1, [](EAudioChunkLoadResult) {});
 								}
 							}
 						}
diff --git a/Engine/Plugins/Runtime/Metasound/Source/MetasoundEngine/Private/MetasoundWaveInfoNode.cpp b/Engine/Plugins/Runtime/Metasound/Source/MetasoundEngine/Private/MetasoundWaveInfoNode.cpp
--- a/Engine/Plugins/Runtime/Metasound/Source/MetasoundEngine/Private/MetasoundWaveInfoNode.cpp
+++ b/Engine/Plugins/Runtime/Metasound/Source/MetasoundEngine/Private/MetasoundWaveInfoNode.cpp
@@ -60,7 +60,7 @@ namespace Metasound
 		FStringWriteRef PathOutput;
 
 		// other
-		TSharedPtr<const FSoundWaveData> SoundWaveData;
+		FSoundWaveProxyPtr SoundWaveProxy;
 
 	}; // class FWaveInfoNodeOperator
 
@@ -143,27 +143,21 @@ namespace Metasound
 
 	void FWaveInfoNodeOperator::Execute()
 	{
-		TSharedPtr<const FSoundWaveProxy> ProxyPtr = (*WaveAsset).GetWaveProxy();
-		TSharedRef<const FSoundWaveData> DataPtr;
-		if (ProxyPtr)
-		{
-			DataPtr = ProxyPtr->GetSoundWaveDataRef();
-		}
-		
+		FSoundWaveProxyPtr ProxyPtr = (*WaveAsset).GetSoundWaveProxy();
 
-		if (SoundWaveData != DataPtr)
+		if (SoundWaveProxy != ProxyPtr)
 		{
-			SoundWaveData = DataPtr;
+			SoundWaveProxy = ProxyPtr;
 
-			if (SoundWaveData.IsValid() && (*WaveAsset).IsSoundWaveValid())
+			if (SoundWaveProxy.IsValid() && (*WaveAsset).IsSoundWaveValid())
 			{
-				*DurationSeconds = FTime::FromSeconds(DataPtr->GetDuration());
+				*DurationSeconds = FTime::FromSeconds((*WaveAsset)->GetDuration());
 
-				FName AssetName = SoundWaveData->GetFName();
+				FName AssetName = SoundWaveProxy->GetFName();
 				*NameOutput = AssetName.ToString();
 
 				FString FullPath;
-				SoundWaveData->GetPackageName().AppendString(FullPath);
+				SoundWaveProxy->GetPackageName().AppendString(FullPath);
 				if (!AssetName.IsNone())
 				{
 					FullPath += TEXT(".");
diff --git a/Engine/Plugins/Runtime/Metasound/Source/MetasoundEngine/Private/MetasoundWavePlayerNode.cpp b/Engine/Plugins/Runtime/Metasound/Source/MetasoundEngine/Private/MetasoundWavePlayerNode.cpp
--- a/Engine/Plugins/Runtime/Metasound/Source/MetasoundEngine/Private/MetasoundWavePlayerNode.cpp
+++ b/Engine/Plugins/Runtime/Metasound/Source/MetasoundEngine/Private/MetasoundWavePlayerNode.cpp
@@ -533,12 +533,12 @@ namespace Metasound
 			// Reset flag to track render cost.
 			bDidWaveRenderThisBlock = false;
 
- #if UE_USOUNDWAVE_HOT_RELOADING_ENABLED
-			if (bIsPlaying && WaveAsset->IsSoundWaveValid() && !WaveAsset->GetWaveProxy()->IsLatestData(CurrentSoundWaveData.Get()))
+#if PLATFORM_DESKTOP && !UE_BUILD_SHIPPING
+			if (bIsPlaying && WaveAsset->IsSoundWaveValid() && WaveAsset->GetSoundWaveProxy()->TryGetLatest())
 			{
 				StartPlaying();
 			}
- #endif // #if UE_USOUNDWAVE_HOT_RELOADING_ENABLED
+#endif // #if PLATFORM_DESKTOP && !UE_BUILD_SHIPPING
 
 			// Update wave proxy reader with any new looping bounds. 
 			if (WaveProxyReader.IsValid())
@@ -813,7 +813,7 @@ namespace Metasound
 			// Copy the wave asset off on init in case the user changes it while we're playing it.
 			// We'll only check for new wave assets when the current one finishes for sample accurate concantenation
 			CurrentWaveAsset = *WaveAsset;
-			TSharedPtr<const FSoundWaveProxy> WaveProxy = CurrentWaveAsset.GetWaveProxy();
+			FSoundWaveProxyPtr WaveProxy = CurrentWaveAsset.GetSoundWaveProxy();
 
 			bOnNearlyDoneTriggeredForWave = false;
 			bIsPlaying = false;
@@ -828,12 +828,14 @@ namespace Metasound
 			if (WaveProxy.IsValid())
 			{
 				// Make sure we have up-to-date SoundWaveData
-				CurrentSoundWaveData = WaveProxy->GetSoundWaveDataRef();
+#if PLATFORM_DESKTOP
+				WaveProxy->TryGetLatest();
+#endif // #if PLATFORM_DESKTOP
 
-				UE_LOG(LogMetaSound, Verbose, TEXT("Starting Sound: '%s'"), *CurrentSoundWaveData->GetFName().ToString());
+				UE_LOG(LogMetaSound, Verbose, TEXT("Starting Sound: '%s'"), *CurrentWaveAsset->GetFName().ToString());
 
 				// Create local sorted copy of cue points and loop regions.
-				SortedCuePoints = CurrentSoundWaveData->GetCuePointsAndLoopRegions();
+				SortedCuePoints = WaveProxy->GetCuePointsAndLoopRegions();
 				Algo::SortBy(SortedCuePoints, WavePlayerNodePrivate::GetCuePointFrame);
 				
 				// Create the wave proxy reader.
@@ -847,7 +849,7 @@ namespace Metasound
 				WaveReaderSettings.bIsLooping = *bLoop;
 				WaveReaderSettings.bMaintainAudioSync = bMaintainAudioSync;
 
-				WaveProxyReader = FSoundWaveProxyReader::Create(CurrentSoundWaveData.ToSharedRef(), WaveReaderSettings);
+				WaveProxyReader = FSoundWaveProxyReader::Create(WaveProxy.ToSharedRef(), WaveReaderSettings);
 
 				DeinterleaveBufferBlockSizeInFrames = FMath::IsPowerOfTwo(DeinterleaveBlockSizeInFrames) ? 
 					DeinterleaveBlockSizeInFrames : FMath::RoundUpToPowerOfTwo(DeinterleaveBlockSizeInFrames);
@@ -1133,7 +1135,7 @@ namespace Metasound
 				int32 NumFramesProduced = Resampler->ProcessAndConsumeAudio(SourceCircularBuffer, OutBuffer);
 				if (NumFramesProduced < 1)
 				{
-					UE_LOG(LogMetaSound, Error, TEXT("Aborting currently playing metasound wave %s. Failed to produce any resampled audio frames with %d input frames and a frame ratio of %f."), *CurrentSoundWaveData->GetFName().ToString(), NumSourceFramesAvailable, GetFrameRatio());
+					UE_LOG(LogMetaSound, Error, TEXT("Aborting currently playing metasound wave %s. Failed to produce any resampled audio frames with %d input frames and a frame ratio of %f."), *CurrentWaveAsset->GetFName().ToString(), NumSourceFramesAvailable, GetFrameRatio());
 					bIsPlaying = false;
 					break;
 				}
@@ -1185,7 +1187,6 @@ namespace Metasound
 		TUniquePtr<Audio::FMultichannelLinearResampler> Resampler;
 
 		FWaveAsset CurrentWaveAsset;
-		TSharedPtr<const FSoundWaveData> CurrentSoundWaveData;
 		TArray<FSoundWaveCuePoint> SortedCuePoints;
 		Audio::FAlignedFloatBuffer InterleavedBuffer;
 		Audio::FMultichannelBuffer DeinterleavedBuffer;
diff --git a/Engine/Plugins/Runtime/Metasound/Source/MetasoundEngine/Private/MetasoundWavePlayerNodeV2.cpp b/Engine/Plugins/Runtime/Metasound/Source/MetasoundEngine/Private/MetasoundWavePlayerNodeV2.cpp
--- a/Engine/Plugins/Runtime/Metasound/Source/MetasoundEngine/Private/MetasoundWavePlayerNodeV2.cpp
+++ b/Engine/Plugins/Runtime/Metasound/Source/MetasoundEngine/Private/MetasoundWavePlayerNodeV2.cpp
@@ -269,8 +269,6 @@ namespace Metasound
 			ProxyPlayer->Reset(PlayerSettings);
 
 			CurrentWaveAsset = FWaveAsset();
-			CurrentSoundWaveData.Reset();
-
 			bIsPlaying = false;
 			bDidRenderThisBlock = false;
 		}
@@ -404,15 +402,13 @@ namespace Metasound
 
 		void ExecuteSubblocks()
 		{
-#if UE_USOUNDWAVE_HOT_RELOADING_ENABLED
+#if PLATFORM_DESKTOP && !UE_BUILD_SHIPPING
 			// hot reload sound
-			if (bIsPlaying && Inputs.WaveAsset->IsSoundWaveValid()
-				&& Inputs.WaveAsset->GetWaveProxy()
-				&& !Inputs.WaveAsset->GetWaveProxy()->IsLatestData(CurrentSoundWaveData.Get()))
+			if (bIsPlaying && Inputs.WaveAsset->IsSoundWaveValid() && Inputs.WaveAsset->GetSoundWaveProxy()->TryGetLatest())
 			{
 				StartPlaying();
 			}
-#endif // UE_USOUNDWAVE_HOT_RELOADING_ENABLED
+#endif // PLATFORM_DESKTOP && !UE_BUILD_SHIPPING
 			
 			// Parse triggers and render audio
 			int32 PlayTrigIndex = 0;
@@ -499,24 +495,14 @@ namespace Metasound
 		{
 			check(ProxyPlayer.IsValid());
 			CurrentWaveAsset = *Inputs.WaveAsset;
-			TSharedPtr<const FSoundWaveProxy> Proxy = CurrentWaveAsset.GetWaveProxy();
-			if(Proxy)
-			{
-				CurrentSoundWaveData = Proxy->GetSoundWaveDataRef().ToSharedPtr();
-			}
-			else
-			{
-				CurrentSoundWaveData.Reset();
-			}
-			
-			if (CurrentSoundWaveData && ProxyPlayer->SetSoundWave(CurrentSoundWaveData.ToSharedRef()))
+			if (ProxyPlayer->SetSoundWave(CurrentWaveAsset.GetSoundWaveProxy()))
 			{
 				ProxyPlayer->SetSpeed(CalculateSpeedFromPitchShift());
 				ProxyPlayer->SetLoop(*Inputs.Loop, Inputs.LoopStartTime->GetSeconds(), Inputs.LoopDuration->GetSeconds());
 				ProxyPlayer->SeekToTime(Inputs.StartTime->GetSeconds());
 			}
 
-			bIsPlaying = Proxy.IsValid() && ProxyPlayer.IsValid() && ProxyPlayer->IsPlayerValid() && !ProxyPlayer->IsFinished();
+			bIsPlaying = ProxyPlayer.IsValid() && ProxyPlayer->IsPlayerValid() && !ProxyPlayer->IsFinished();
 		}
 
 		void RenderFrameRange(int32 StartFrame, int32 EndFrame)
@@ -594,7 +580,6 @@ namespace Metasound
 		bool bDidRenderThisBlock = false;
 
 		FWaveAsset CurrentWaveAsset;
-		TSharedPtr<const FSoundWaveData> CurrentSoundWaveData;
 		TUniquePtr<FSoundWaveProxyPlayer> ProxyPlayer;
 		TArray<TArrayView<float>> MultiChannelOutputView;
 	};
diff --git a/Engine/Plugins/Runtime/Metasound/Source/MetasoundEngine/Public/MetasoundWave.h b/Engine/Plugins/Runtime/Metasound/Source/MetasoundEngine/Public/MetasoundWave.h
--- a/Engine/Plugins/Runtime/Metasound/Source/MetasoundEngine/Public/MetasoundWave.h
+++ b/Engine/Plugins/Runtime/Metasound/Source/MetasoundEngine/Public/MetasoundWave.h
@@ -24,39 +24,30 @@ namespace Metasound
 		static constexpr bool Value = true;
 	};
 
-	// Metasound data type that holds onto a shared ptr. Mostly used as a placeholder until we have a proper proxy type.
+	// Metasound data type that holds onto a weak ptr. Mostly used as a placeholder until we have a proper proxy type.
 	class FWaveAsset
 	{
-		// todo: this can be TSharedPtr<const FSoundWaveProxy> once depprecated functions are removed
-		TSharedPtr<FSoundWaveProxy, ESPMode::ThreadSafe> SoundWaveProxy;
-	
+		FSoundWaveProxyPtr SoundWaveProxy;
 	public:
 
 		FWaveAsset() = default;
 		FWaveAsset(const FWaveAsset&) = default;
 		FWaveAsset& operator=(const FWaveAsset& Other) = default;
 
-		UE_API FWaveAsset(const TSharedPtr<const Audio::IProxyData>& InInitData);
+		UE_API FWaveAsset(const TSharedPtr<Audio::IProxyData>& InInitData);
 
 		UE_API bool IsSoundWaveValid() const;
 
-		UE_DEPRECATED(5.8, "Use GetWaveProxy() instead, mutable proxy access is deprecated")
 		const FSoundWaveProxyPtr& GetSoundWaveProxy() const
 		{
 			return SoundWaveProxy;
 		}
 
-		const TSharedPtr<const FSoundWaveProxy> GetWaveProxy() const 
-		{
-			return SoundWaveProxy;
-		}
-
 		const FSoundWaveProxy* operator->() const
 		{
 			return SoundWaveProxy.Get();
 		}
 
-		UE_DEPRECATED(5.8, "Use GetWaveProxy() for non-const objects instead. (mutable proxy access is deprecated)")
 		FSoundWaveProxy* operator->()
 		{
 			return SoundWaveProxy.Get();
@@ -66,9 +57,7 @@ namespace Metasound
 		{
 			if (InWaveAsset.IsSoundWaveValid())
 			{
-				const TSharedPtr<const FSoundWaveProxy> Proxy = InWaveAsset.GetWaveProxy();
-				check(Proxy);
-				return GetTypeHash(*Proxy);
+				return GetTypeHash(*InWaveAsset.GetSoundWaveProxy());
 			}
 			return INDEX_NONE;
 		}
diff --git a/Engine/Source/Runtime/AdpcmAudioDecoder/Module/Private/ADPCMAudioInfo.cpp b/Engine/Source/Runtime/AdpcmAudioDecoder/Module/Private/ADPCMAudioInfo.cpp
--- a/Engine/Source/Runtime/AdpcmAudioDecoder/Module/Private/ADPCMAudioInfo.cpp
+++ b/Engine/Source/Runtime/AdpcmAudioDecoder/Module/Private/ADPCMAudioInfo.cpp
@@ -123,14 +123,12 @@ void FADPCMAudioInfo::SeekToFrameInternal(const uint32 InSeekFrame)
 
 	// Calculate block index & force SeekTime to be in bounds.
 	check(WaveInfo.pSamplesPerSec != nullptr);
-	check(TotalSamplesPerChannel > 1);
-
 	uint32 SeekedSamples = InSeekFrame;
 	TotalSamplesStreamed = FMath::Min<uint32>(SeekedSamples, TotalSamplesPerChannel - 1);
 
 	const uint32 HeaderOffset = static_cast<uint32>(WaveInfo.SampleDataStart - SrcBufferData);
 	
-	if (!StreamingSoundWaveData)
+	if (!StreamingSoundWave)
 	{
 		// For the non streaming case:
 		if (Format == WAVE_FORMAT_LPCM)
@@ -153,13 +151,13 @@ void FADPCMAudioInfo::SeekToFrameInternal(const uint32 InSeekFrame)
 	}
 	else
 	{
-		if (StreamingSoundWaveData->GetNumChunks() == 0)
+		if (StreamingSoundWave->GetNumChunks() == 0)
 		{
 			UE_LOG(LogAudio, Error, TEXT("Entered streaming seek path with a non-streaming sound!"));
 			return;
 		}
 
-		const uint32 TotalStreamingChunks = StreamingSoundWaveData->GetNumChunks();
+		const uint32 TotalStreamingChunks = StreamingSoundWave->GetNumChunks();
 
 		if (Format == WAVE_FORMAT_ADPCM)
 		{
@@ -170,7 +168,7 @@ void FADPCMAudioInfo::SeekToFrameInternal(const uint32 InSeekFrame)
 			const int32 ChannelBlockSize = BlockSize * NumChannels;
 			for (uint32 BlockIndex = 0; BlockIndex < CurrentCompressedBlockIndex; ++BlockIndex)
 			{
-				const uint32 SizeOfChunk = StreamingSoundWaveData->GetSizeOfChunk(CurrentChunkIndex);
+				const uint32 SizeOfChunk = StreamingSoundWave->GetSizeOfChunk(CurrentChunkIndex);
 				if (CurrentChunkBufferOffset + ChannelBlockSize >= SizeOfChunk)
 				{
 					const uint32 RemainderAfterEndOfChunk = CurrentChunkBufferOffset + ChannelBlockSize - SizeOfChunk;
@@ -198,7 +196,7 @@ void FADPCMAudioInfo::SeekToFrameInternal(const uint32 InSeekFrame)
 			const uint32 TotalByteSizeToSeek = TotalSamplesStreamed * ChannelBlockSize;	
 			
 			uint32 TargetChunkIndex = FirstChunkSampleDataIndex;
-			uint32 SizeOfTargetChunk = StreamingSoundWaveData->GetSizeOfChunk(TargetChunkIndex);
+			uint32 SizeOfTargetChunk = StreamingSoundWave->GetSizeOfChunk(TargetChunkIndex);
 			
 			uint32 CurrentBytes = 0;
 			while (CurrentBytes + SizeOfTargetChunk < TotalByteSizeToSeek)
@@ -214,7 +212,7 @@ void FADPCMAudioInfo::SeekToFrameInternal(const uint32 InSeekFrame)
 				++TargetChunkIndex;
 				CurrentBytes += SizeOfTargetChunk;
 				CurrentChunkBufferOffset -= SizeOfTargetChunk;
-				SizeOfTargetChunk = StreamingSoundWaveData->GetSizeOfChunk(TargetChunkIndex);
+				SizeOfTargetChunk = StreamingSoundWave->GetSizeOfChunk(TargetChunkIndex);
 			}
 
 			check(CurrentBytes <= TotalByteSizeToSeek);
@@ -481,17 +479,21 @@ void FADPCMAudioInfo::ResetSeekState()
 	bNewSeekRequest = false;
 }
 
-bool FADPCMAudioInfo::StreamCompressedInfoInternal(const TSharedRef<const FSoundWaveData>& InWaveData, struct FSoundQualityInfo* QualityInfo)
+bool FADPCMAudioInfo::StreamCompressedInfoInternal(const FSoundWaveProxyPtr& InWaveProxy, struct FSoundQualityInfo* QualityInfo)
 {
 	FScopeLock ScopeLock(&CurCompressedChunkHandleCriticalSection);
 
 	check(QualityInfo);
-	check(StreamingSoundWaveData == InWaveData);
+	check(StreamingSoundWave == InWaveProxy);
+	if (!ensure(InWaveProxy.IsValid()))
+	{
+		return false;
+	}
 
 	CurrentChunkIndex = 0;
 
 	// Get the first chunk of audio data (should already be loaded)
-	uint8 const* ChunkData = GetLoadedChunk(*InWaveData, CurrentChunkIndex, CurrentChunkDataSize);
+	uint8 const* ChunkData = GetLoadedChunk(InWaveProxy, CurrentChunkIndex, CurrentChunkDataSize);
 
 	if (ChunkData == nullptr)
 	{
@@ -517,7 +519,7 @@ bool FADPCMAudioInfo::StreamCompressedInfoInternal(const TSharedRef<const FSound
 	if (((uint32)SampleDataOffset) >= CurrentChunkDataSize)
 	{
 		++CurrentChunkIndex;
-		ChunkData = GetLoadedChunk(*InWaveData, CurrentChunkIndex, CurrentChunkDataSize);
+		ChunkData = GetLoadedChunk(InWaveProxy, CurrentChunkIndex, CurrentChunkDataSize);
 		FirstChunkSampleDataIndex = CurrentChunkIndex;
 		FirstChunkSampleDataOffset = 0;
 	}
@@ -674,7 +676,7 @@ bool FADPCMAudioInfo::StreamCompressedData(uint8* Destination, bool bLooping, ui
 						bSeekedFowardToNextChunk = false;
 					}
 
-					if (CurrentChunkIndex >= StreamingSoundWaveData->GetNumChunks())
+					if (CurrentChunkIndex >= StreamingSoundWave->GetNumChunks())
 					{
 						// If this is the case, we've seeked to the end of the file.
 						ReachedEndOfSamples = true;
@@ -692,7 +694,7 @@ bool FADPCMAudioInfo::StreamCompressedData(uint8* Destination, bool bLooping, ui
 					}
 
 					// Request the next chunk of data from the streaming engine
-					CurCompressedChunkData = GetLoadedChunk(*StreamingSoundWaveData, CurrentChunkIndex, CurrentChunkDataSize);
+					CurCompressedChunkData = GetLoadedChunk(StreamingSoundWave, CurrentChunkIndex, CurrentChunkDataSize);
 
 					if(CurCompressedChunkData == nullptr)
 					{
@@ -738,7 +740,7 @@ bool FADPCMAudioInfo::StreamCompressedData(uint8* Destination, bool bLooping, ui
 							CurrentChunkBufferOffset += NumCompressedBytesForCallback;
 							TotalSamplesStreamed += NumDecompressedSamplesForCallback;
 
-							const uint32 SizeOfCurrentChunk = StreamingSoundWaveData->GetSizeOfChunk(CurrentChunkIndex);
+							const uint32 SizeOfCurrentChunk = StreamingSoundWave->GetSizeOfChunk(CurrentChunkIndex);
 							if (CurrentChunkBufferOffset >= SizeOfCurrentChunk)
 							{
 								// If this is the case, we've seeked forward past the end of the chunk.
@@ -752,7 +754,7 @@ bool FADPCMAudioInfo::StreamCompressedData(uint8* Destination, bool bLooping, ui
 								CurrentChunkBufferOffset -= SizeOfCurrentChunk;
 
 								// Ensure that we have not moved past the size of our new chunk with our current callback.
-								ensureAlways(CurrentChunkBufferOffset < StreamingSoundWaveData->GetSizeOfChunk(CurrentChunkIndex));
+								ensureAlways(CurrentChunkBufferOffset < StreamingSoundWave->GetSizeOfChunk(CurrentChunkIndex));
 
 								bSeekedFowardToNextChunk = true;
 
@@ -761,7 +763,7 @@ bool FADPCMAudioInfo::StreamCompressedData(uint8* Destination, bool bLooping, ui
 									CurrentChunkBufferOffset = AlignDown(CurrentChunkBufferOffset, CompressedBlockSize);
 								}
 
-								if (CurrentChunkIndex >= StreamingSoundWaveData->GetNumChunks())
+								if (CurrentChunkIndex >= StreamingSoundWave->GetNumChunks())
 								{
 									// If this is the case, we've seeked to the end of the file.
 									ReachedEndOfSamples = true;
@@ -879,7 +881,7 @@ bool FADPCMAudioInfo::StreamCompressedData(uint8* Destination, bool bLooping, ui
 				}
 
 				// Check for the end of the audio samples and loop if needed
-				if (CurrentChunkIndex >= StreamingSoundWaveData->GetNumChunks())
+				if (CurrentChunkIndex >= StreamingSoundWave->GetNumChunks())
 				{
 					ReachedEndOfSamples = true;
 					CurrentChunkIndex = FirstChunkSampleDataIndex;
@@ -895,7 +897,7 @@ bool FADPCMAudioInfo::StreamCompressedData(uint8* Destination, bool bLooping, ui
 				}
 
 				// Request the next chunk of data from the streaming engine
-				CurCompressedChunkData = GetLoadedChunk(*StreamingSoundWaveData, CurrentChunkIndex, CurrentChunkDataSize);
+				CurCompressedChunkData = GetLoadedChunk(StreamingSoundWave, CurrentChunkIndex, CurrentChunkDataSize);
 
 				if(CurCompressedChunkData == nullptr)
 				{
@@ -1009,7 +1011,7 @@ bool FADPCMAudioInfo::ReleaseStreamChunk(bool bBlockUntilReleased)
 	}
 }
 
-const uint8* FADPCMAudioInfo::GetLoadedChunk(const FSoundWaveData& InSoundWaveData, uint32 ChunkIndex, uint32& OutChunkSize)
+const uint8* FADPCMAudioInfo::GetLoadedChunk(const FSoundWaveProxyPtr& InSoundWave, uint32 ChunkIndex, uint32& OutChunkSize)
 {
 	if (ChunkIndex != 0 && FMath::RandRange(0.0f, 1.0f) < ChanceForIntentionalChunkMissCVar)
 	{
@@ -1018,14 +1020,14 @@ const uint8* FADPCMAudioInfo::GetLoadedChunk(const FSoundWaveData& InSoundWaveDa
 		return nullptr;
 	}
 
-	if (!ensure(ChunkIndex < InSoundWaveData.GetNumChunks()))
+	if (!ensure(InSoundWave.IsValid()) || ChunkIndex >= InSoundWave->GetNumChunks())
 	{
 		OutChunkSize = 0;
 		return nullptr;
 	}
 	else if (ChunkIndex == 0)
 	{
-		TArrayView<const uint8> ZerothChunk = InSoundWaveData.GetZerothChunk(true);
+		TArrayView<const uint8> ZerothChunk = FSoundWaveProxy::GetZerothChunk(InSoundWave, true);
 		OutChunkSize = ZerothChunk.Num();
 		PreviouslyRequestedChunkIndex = ChunkIndex;
 		return ZerothChunk.GetData();
@@ -1035,7 +1037,7 @@ const uint8* FADPCMAudioInfo::GetLoadedChunk(const FSoundWaveData& InSoundWaveDa
 		const bool bIsSeekingOrLooping = bSeekPendingRead || ChunkIndex == FirstChunkSampleDataIndex;
 		ensureAlwaysMsgf(bIsSeekingOrLooping || ChunkIndex == PreviouslyRequestedChunkIndex || ChunkIndex == PreviouslyRequestedChunkIndex + 1, TEXT("ADPCM playback error! We skipped from the end of chunk %d to the beginning of chunk %d."), PreviouslyRequestedChunkIndex, ChunkIndex);
 
-		CurCompressedChunkHandle = IStreamingManager::Get().GetAudioStreamingManager().GetLoadedChunk(InSoundWaveData.AsShared(), ChunkIndex, false, true);
+		CurCompressedChunkHandle = IStreamingManager::Get().GetAudioStreamingManager().GetLoadedChunk(InSoundWave, ChunkIndex, false, true);
 		OutChunkSize = CurCompressedChunkHandle.Num();
 
 		PreviouslyRequestedChunkIndex = ChunkIndex;
diff --git a/Engine/Source/Runtime/AdpcmAudioDecoder/Module/Private/PcmAudioInfoHybrid.cpp b/Engine/Source/Runtime/AdpcmAudioDecoder/Module/Private/PcmAudioInfoHybrid.cpp
--- a/Engine/Source/Runtime/AdpcmAudioDecoder/Module/Private/PcmAudioInfoHybrid.cpp
+++ b/Engine/Source/Runtime/AdpcmAudioDecoder/Module/Private/PcmAudioInfoHybrid.cpp
@@ -342,7 +342,7 @@ bool FPcmAudioInfoHybrid::AbsPositionToChunkIndexAndOffset(const uint64 InAbsPos
 	// But we could make this faster by just special casing zero, first, rest, and last.
 	
 	// If we're streaming, grab the streaming wave instance.
-	if (const TSharedPtr<const FSoundWaveData>& Wave = GetStreamingSoundWaveData(); Wave.IsValid() && Wave->IsStreaming())
+	if (const FSoundWaveProxyPtr& Wave = GetStreamingSoundWave(); Wave.IsValid() && Wave->IsStreaming())
 	{
 		const int32 NumChunks = Wave->GetNumChunks();
 		uint64 AbsChunkOffset = AudioDataOffset;
@@ -369,7 +369,7 @@ bool FPcmAudioInfoHybrid::ParseHeader(const uint8* InSrcBufferData, uint32 InSrc
 	CurrentSampleCount = 0;
 
 	// We only have the header if we're streaming, so just parse that.
-	const bool bParseJustHeader = GetStreamingSoundWaveData().IsValid() && GetStreamingSoundWaveData()->IsStreaming();
+	const bool bParseJustHeader = GetStreamingSoundWave().IsValid() && GetStreamingSoundWave()->IsStreaming();
 
 	void* FormatHeader = nullptr;
 	FWaveModInfo WaveInfo;
@@ -473,7 +473,6 @@ void FPcmAudioInfoHybrid::SeekToAbs(const uint64 InAbsPosition, const uint64 InS
 
 void FPcmAudioInfoHybrid::FHybridImpl::SeekToTime(const float InSeekTime)
 {
-	check(TotalSamplesPerChannel > 0);
 	const uint32 SeekFrame = InSeekTime * SampleRate;
 	const uint32 SeekFrameClamped = FMath::Clamp<uint32>( SeekFrame, 0, TotalSamplesPerChannel - 1);
 	SeekToFrame(SeekFrameClamped);
diff --git a/Engine/Source/Runtime/AdpcmAudioDecoder/Module/Private/PcmAudioInfoHybrid.h b/Engine/Source/Runtime/AdpcmAudioDecoder/Module/Private/PcmAudioInfoHybrid.h
--- a/Engine/Source/Runtime/AdpcmAudioDecoder/Module/Private/PcmAudioInfoHybrid.h
+++ b/Engine/Source/Runtime/AdpcmAudioDecoder/Module/Private/PcmAudioInfoHybrid.h
@@ -78,7 +78,7 @@ protected:
 public:
 	bool IsStreaming() const
 	{
-		return GetStreamingSoundWaveData().IsValid() && GetStreamingSoundWaveData()->IsStreaming();
+		return GetStreamingSoundWave().IsValid() && GetStreamingSoundWave()->IsStreaming();
 	}
 	void SeekToAbs(const uint64 InAbsPosition, const uint64 InSeekFrame);	
 	void SetErrorFlag() const { bHasError = true; }
@@ -93,11 +93,9 @@ private:
 	
 	void SetSeekBlockIndex(const uint32 InStreamSeekBlockIndex)
 	{
-		if (ensure(GetStreamingSoundWaveData().IsValid() && GetStreamingSoundWaveData()->IsStreaming()
-			&& InStreamSeekBlockIndex < GetStreamingSoundWaveData()->GetNumChunks()))
-		{
-			StreamSeekBlockIndex = InStreamSeekBlockIndex;
-		}
+		ensure(GetStreamingSoundWave().IsValid() && GetStreamingSoundWave()->IsStreaming()
+			&& InStreamSeekBlockIndex < GetStreamingSoundWave()->GetNumChunks());
+		StreamSeekBlockIndex = InStreamSeekBlockIndex;
 	};
 	void SetSeekBlockOffset(const int32 InStreamSeekBlockOffset)
 	{
diff --git a/Engine/Source/Runtime/AdpcmAudioDecoder/Module/Public/Decoders/ADPCMAudioInfo.h b/Engine/Source/Runtime/AdpcmAudioDecoder/Module/Public/Decoders/ADPCMAudioInfo.h
--- a/Engine/Source/Runtime/AdpcmAudioDecoder/Module/Public/Decoders/ADPCMAudioInfo.h
+++ b/Engine/Source/Runtime/AdpcmAudioDecoder/Module/Public/Decoders/ADPCMAudioInfo.h
@@ -107,7 +107,7 @@ public:
 
 	// Additional overrides for streaming
 	virtual bool SupportsStreaming() const override {return true;}
-	virtual bool StreamCompressedInfoInternal(const TSharedRef<const FSoundWaveData>& InWaveData, struct FSoundQualityInfo* QualityInfo) override;
+	virtual bool StreamCompressedInfoInternal(const FSoundWaveProxyPtr& InWaveProxy, struct FSoundQualityInfo* QualityInfo) override;
 	virtual bool StreamCompressedData(uint8* Destination, bool bLooping, uint32 BufferSize, int32& OutNumBytesStreamed) override;
 	virtual int32 GetCurrentChunkIndex() const override
 	{
@@ -139,7 +139,7 @@ private:
 	// Wrapper function that returns a pointer to the currently used compressed data.
 	// If a non-zero chunk is requested, this function also aquires a reference to that chunk
 	// until we move on to a different chunk.
-	const uint8* GetLoadedChunk(const FSoundWaveData& InSoundWaveData, uint32 ChunkIndex, uint32& OutChunkSize);
+	const uint8* GetLoadedChunk(const FSoundWaveProxyPtr& InSoundWave, uint32 ChunkIndex, uint32& OutChunkSize);
 
 	int32 NumConsecutiveReadFailiures;
 
diff --git a/Engine/Source/Runtime/AudioMixer/Classes/Generators/SoundWaveScrubber.h b/Engine/Source/Runtime/AudioMixer/Classes/Generators/SoundWaveScrubber.h
--- a/Engine/Source/Runtime/AudioMixer/Classes/Generators/SoundWaveScrubber.h
+++ b/Engine/Source/Runtime/AudioMixer/Classes/Generators/SoundWaveScrubber.h
@@ -27,9 +27,6 @@ public:
 	UE_API ~FSoundWaveScrubber();
 
 	// Initialize the sound wave scrubber
-	UE_API void Init(TSharedRef<const FSoundWaveProxy> InSoundWaveProxyRef, float InSampleRate, int32 InNumChannels, float InPlayheadTimeSeconds = 0.0f);
-
-	UE_DEPRECATED(5.8, "Please use the TSharedRef<const FSoundWaveProxy> overload instead")
 	UE_API void Init(FSoundWaveProxyPtr InSoundWaveProxyPtr, float InSampleRate, int32 InNumChannels, float InPlayheadTimeSeconds = 0.0f);
 
 	// Sets whether or not to scrub the audio file. If false, then the output of the scrubber will be silence.
@@ -88,8 +85,8 @@ private:
 	// Place to store decoded chunks. We need 2 chunks for actively playing grains plus a free chunk.
 	TArray<FDecodedDataChunk> DecodedChunks;
 
-	// Sound wave data to safely query and use the referenced sound wave asset
-	TSharedPtr<const FSoundWaveData> SoundWaveDataPtr;
+	// Sound wave proxy to safely query and use the referenced sound wave asset
+	FSoundWaveProxyPtr SoundWaveProxyPtr;
 	// Sound wave proxy reader is a decoder we can use to decode audio into chunks
 	TUniquePtr<FSoundWaveProxyReader> SoundWaveProxyReaderPtr;
 	// Simple SRC interface. Decoded audio is SRC'd to match the audio mixer sample rate
@@ -176,9 +173,6 @@ private:
 class FSoundWaveScrubberGenerator : public ISoundGenerator
 {
 public:
-	UE_API void Init(TSharedRef<const FSoundWaveProxy> InSoundWaveProxyRef, float InSampleRate, int32 InNumChannels, float InPlayheadTimeSeconds = 0.0f);
-
-	UE_DEPRECATED(5.8, "Please use the TSharedRef<const FSoundWaveProxy> overload instead")
 	UE_API void Init(FSoundWaveProxyPtr InSoundWaveProxyPtr, float InSampleRate, int32 InNumChannels, float InPlayheadTimeSeconds = 0.0f);
 
 	// ~ISoundGenerator Begin
diff --git a/Engine/Source/Runtime/AudioMixer/Private/AudioMixerBlueprintLibrary.cpp b/Engine/Source/Runtime/AudioMixer/Private/AudioMixerBlueprintLibrary.cpp
--- a/Engine/Source/Runtime/AudioMixer/Private/AudioMixerBlueprintLibrary.cpp
+++ b/Engine/Source/Runtime/AudioMixer/Private/AudioMixerBlueprintLibrary.cpp
@@ -602,26 +602,23 @@ void UAudioMixerBlueprintLibrary::PrimeSoundForPlayback(USoundWave* SoundWave, c
 	if (!SoundWave)
 	{
 		UE_LOG(LogAudioMixer, Warning, TEXT("Prime Sound For Playback called with a null SoundWave pointer."));
-		OnLoadCompletion.ExecuteIfBound(SoundWave, true);
-
-		return;
 	}
-
-	TSharedRef<const FSoundWaveProxy> SoundWaveProxy = SoundWave->GetSoundWaveProxy();
-
-	IStreamingManager::Get().GetAudioStreamingManager().RequestChunk(SoundWaveProxy->GetSoundWaveDataRef(), 1, [OnLoadCompletion, SoundWave](EAudioChunkLoadResult InResult)
+	else
 	{
-		AsyncTask(ENamedThreads::GameThread, [OnLoadCompletion, SoundWave, InResult]() {
-			if (InResult == EAudioChunkLoadResult::Completed || InResult == EAudioChunkLoadResult::AlreadyLoaded)
-			{
-				OnLoadCompletion.ExecuteIfBound(SoundWave, false);
-			}
-			else
-			{
-				OnLoadCompletion.ExecuteIfBound(SoundWave, true);
-			}
+		IStreamingManager::Get().GetAudioStreamingManager().RequestChunk(SoundWave->CreateSoundWaveProxy(), 1, [OnLoadCompletion, SoundWave](EAudioChunkLoadResult InResult)
+		{
+			AsyncTask(ENamedThreads::GameThread, [OnLoadCompletion, SoundWave, InResult]() {
+				if (InResult == EAudioChunkLoadResult::Completed || InResult == EAudioChunkLoadResult::AlreadyLoaded)
+				{
+					OnLoadCompletion.ExecuteIfBound(SoundWave, false);
+				}
+				else
+				{
+					OnLoadCompletion.ExecuteIfBound(SoundWave, true);
+				}
+			});
 		});
-	});
+	}
 }
 
 void UAudioMixerBlueprintLibrary::PrimeSoundCueForPlayback(USoundCue* SoundCue)
diff --git a/Engine/Source/Runtime/AudioMixer/Private/Generators/SoundWaveScrubber.cpp b/Engine/Source/Runtime/AudioMixer/Private/Generators/SoundWaveScrubber.cpp
--- a/Engine/Source/Runtime/AudioMixer/Private/Generators/SoundWaveScrubber.cpp
+++ b/Engine/Source/Runtime/AudioMixer/Private/Generators/SoundWaveScrubber.cpp
@@ -21,15 +21,14 @@ FSoundWaveScrubber::~FSoundWaveScrubber()
 {
 }
 
-void FSoundWaveScrubber::Init(TSharedRef<const FSoundWaveProxy> InSoundWaveProxyRef, float InSampleRate, int32 InNumChannels, float InPlayheadTimeSeconds)
+void FSoundWaveScrubber::Init(FSoundWaveProxyPtr InSoundWaveProxyPtr, float InSampleRate, int32 InNumChannels, float InPlayheadTimeSeconds)
 {
-	SoundWaveDataPtr = InSoundWaveProxyRef->GetSoundWaveDataRef();
-	check(SoundWaveDataPtr.IsValid());
-	check(InSampleRate > 0);
+	check(InSoundWaveProxyPtr.IsValid());
+	SoundWaveProxyPtr = InSoundWaveProxyPtr;
 
 	AudioMixerSampleRate = InSampleRate;
-	SourceFileSampleRate = SoundWaveDataPtr->GetSampleRate();
-	SourceFileDurationSeconds = SoundWaveDataPtr->GetDuration();
+	SourceFileSampleRate = SoundWaveProxyPtr->GetSampleRate();
+	SourceFileDurationSeconds = SoundWaveProxyPtr->GetDuration();
 	NumChannels = InNumChannels;
 
 	TargetGrainDurationRange = { 0.4f, 0.05f };
@@ -51,17 +50,6 @@ void FSoundWaveScrubber::Init(TSharedRef<const FSoundWaveProxy> InSoundWaveProxy
 	DecodeToDataChunk(DecodedChunks[0], DecoderSeekTimeSeconds);
 }
 
-void FSoundWaveScrubber::Init(FSoundWaveProxyPtr InSoundWaveProxyPtr, float InSampleRate, int32 InNumChannels, float InPlayheadTimeSeconds)
-{
-	if (!ensure(InSoundWaveProxyPtr))
-	{
-		return;
-	}
-
-	TSharedRef<const FSoundWaveProxy> ProxyRef = InSoundWaveProxyPtr->AsShared();
-	this->Init(ProxyRef, InSampleRate, InNumChannels, InPlayheadTimeSeconds);
-}
-
 int32 FSoundWaveScrubber::GetDecodedDataChunkIndexForCurrentReadIndex(int32 InReadFrameIndex)
 {
 	for (int32 i = 0; i < DecodedChunks.Num(); ++i)
@@ -118,7 +106,7 @@ void FSoundWaveScrubber::DecodeToDataChunk(FDecodedDataChunk& InOutDataChunk, fl
 	check(SourceFileSampleRate > 0.0f);
 	check(InDecoderSeekTimeSeconds >= 0.0f);
 	check(NumChannels > 0);
-	check(SoundWaveDataPtr.IsValid());
+	check(SoundWaveProxyPtr.IsValid());
 
 	if (!SoundWaveProxyReaderPtr.IsValid())
 	{
@@ -129,12 +117,12 @@ void FSoundWaveScrubber::DecodeToDataChunk(FDecodedDataChunk& InOutDataChunk, fl
 		FSoundWaveProxyReader::FSettings ProxyReaderSettings;
 		ProxyReaderSettings.MaxDecodeSizeInFrames = DecodedAudioSizeInSeconds * SourceFileSampleRate;
 		ProxyReaderSettings.StartTimeInSeconds = InDecoderSeekTimeSeconds;
-		SoundWaveProxyReaderPtr = FSoundWaveProxyReader::Create(SoundWaveDataPtr.ToSharedRef(), ProxyReaderSettings);
+		SoundWaveProxyReaderPtr = FSoundWaveProxyReader::Create(SoundWaveProxyPtr.ToSharedRef(), ProxyReaderSettings);
 
 		if (!SoundWaveProxyReaderPtr.IsValid())
 		{
 			bHasErrorWithDecoder = true;
-			UE_LOG(LogAudioMixer, Warning, TEXT("Unable to make a sound wave proxy reader for sound wave '%s' in the sound wave scrubber."), *SoundWaveDataPtr->GetFName().ToString());
+			UE_LOG(LogAudioMixer, Warning, TEXT("Unable to make a sound wave proxy reader for sound wave '%s' in the sound wave scrubber."), *SoundWaveProxyPtr->GetFName().ToString());
 			return;
 		}
 	}
@@ -239,7 +227,7 @@ int32 FSoundWaveScrubber::RenderAudio(TArrayView<float>& OutAudio)
 
 	// If we have an error with our decoder, we don't need to render audio
 	// To avoid spamming, we'll allow it to render silence
-	if (bHasErrorWithDecoder || !SoundWaveDataPtr.IsValid())
+	if (bHasErrorWithDecoder)
 	{
 		return OutAudio.Num();
 	}
@@ -538,21 +526,10 @@ void FSoundWaveScrubber::RenderActiveGrains(TArrayView<float>& OutAudio, int32 I
 	}
 }
 
-void FSoundWaveScrubberGenerator::Init(TSharedRef<const FSoundWaveProxy> InSoundWaveProxyRef, float InSampleRate, int32 InNumChannels, float InPlayheadTimeSeconds)
-{
-	NumChannels = InNumChannels;
-	SoundWaveScrubber.Init(InSoundWaveProxyRef, InSampleRate, InNumChannels, InPlayheadTimeSeconds);
-}
-
 void FSoundWaveScrubberGenerator::Init(FSoundWaveProxyPtr InSoundWaveProxyPtr, float InSampleRate, int32 InNumChannels, float InPlayheadTimeSeconds)
 {
-	if (!ensure(InSoundWaveProxyPtr))
-	{
-		return;
-	}
-
-	TSharedRef<const FSoundWaveProxy> ProxyRef = InSoundWaveProxyPtr->AsShared();
-	this->Init(ProxyRef, InSampleRate, InNumChannels, InPlayheadTimeSeconds);
+	NumChannels = InNumChannels;
+	SoundWaveScrubber.Init(InSoundWaveProxyPtr, InSampleRate, InNumChannels, InPlayheadTimeSeconds);
 }
 
 void FSoundWaveScrubberGenerator::SetIsScrubbing(bool bInIsScrubbing)
@@ -607,12 +584,12 @@ ISoundGeneratorPtr UScrubbedSound::CreateSoundGenerator(const FSoundGeneratorIni
 
 	if (SoundWaveToScrub)
 	{
-		TSharedRef<const FSoundWaveProxy> SoundWaveProxyRef = SoundWaveToScrub->GetSoundWaveProxy();
+		FSoundWaveProxyPtr SoundWaveProxyPtr = SoundWaveToScrub->CreateSoundWaveProxy();
 
 		SoundWaveScrubber = ISoundGeneratorPtr(new FSoundWaveScrubberGenerator());
 
 		FSoundWaveScrubberGenerator* Scrubber = static_cast<FSoundWaveScrubberGenerator*>(SoundWaveScrubber.Get());
-		Scrubber->Init(SoundWaveProxyRef, InParams.SampleRate, NumChannels, PlayheadTimeSeconds);
+		Scrubber->Init(SoundWaveProxyPtr, InParams.SampleRate, NumChannels, PlayheadTimeSeconds);
 		Scrubber->SetIsScrubbing(bIsScrubbing);
 		Scrubber->SetIsScrubbingWhileStationary(bScrubWhileStationary);
 		Scrubber->SetGrainDurationRange(GrainDurationRange);
diff --git a/Engine/Source/Runtime/BinkAudioDecoder/Module/Private/BinkAudioInfo.cpp b/Engine/Source/Runtime/BinkAudioDecoder/Module/Private/BinkAudioInfo.cpp
--- a/Engine/Source/Runtime/BinkAudioDecoder/Module/Private/BinkAudioInfo.cpp
+++ b/Engine/Source/Runtime/BinkAudioDecoder/Module/Private/BinkAudioInfo.cpp
@@ -190,20 +190,20 @@ void FBinkAudioInfo::SeekToFrame(const uint32 InFrameNum)
 	//
 	// Here we need to set up the data we get to point at the right spot.
 	//
-	if (StreamingSoundWaveData == nullptr)
+	if (StreamingSoundWave == nullptr)
 	{
 		// If we aren't streaming we can just go directly to the offset we need.
 		this->SrcBufferOffset = OffsetToBlock;
 	}
 	else
 	{
-		const uint32 TotalStreamingChunks = StreamingSoundWaveData->GetNumChunks();
+		const uint32 TotalStreamingChunks = StreamingSoundWave->GetNumChunks();
 		uint32 RemnOffset = OffsetToBlock;
 
 		// Find the chunk and offset to the block we need.
 		for (uint32 BlockIndex = 0; BlockIndex < TotalStreamingChunks; ++BlockIndex)
 		{
-			const uint32 SizeOfChunk = StreamingSoundWaveData->GetSizeOfChunk(BlockIndex);
+			const uint32 SizeOfChunk = StreamingSoundWave->GetSizeOfChunk(BlockIndex);
 
 			if (SizeOfChunk > RemnOffset)
 			{
diff --git a/Engine/Source/Runtime/Engine/Classes/Sound/SoundWave.h b/Engine/Source/Runtime/Engine/Classes/Sound/SoundWave.h
--- a/Engine/Source/Runtime/Engine/Classes/Sound/SoundWave.h
+++ b/Engine/Source/Runtime/Engine/Classes/Sound/SoundWave.h
@@ -6,10 +6,6 @@
  * Playable sound object for raw wave files
  */
 
-#ifndef UE_USOUNDWAVE_HOT_RELOADING_ENABLED
-#define UE_USOUNDWAVE_HOT_RELOADING_ENABLED PLATFORM_DESKTOP && !UE_BUILD_SHIPPING
-#endif // #ifndef UE_USOUNDWAVE_HOT_RELOADING_ENABLED
-
 #include "CoreMinimal.h"
 #include "Containers/IndirectArray.h"
 #include "Sound/AudioSettings.h"
@@ -416,7 +412,6 @@ struct ISoundWaveClient
 	virtual bool OnIsReadyForFinishDestroy(class USoundWave* Wave) const = 0;
 	virtual void OnFinishDestroy(class USoundWave* Wave) = 0;
 };
-
 UCLASS(hidecategories=Object, editinlinenew, BlueprintType, meta= (LoadBehavior = "LazyOnDemand"), MinimalAPI)
 class USoundWave : public USoundBase, public IAudioProxyDataFactory, public IInterface_AsyncCompilation
 {
@@ -521,19 +516,8 @@ public:
 	ENGINE_API void ResetInvalidTemplateArguments();
 
 private:
-	/** if UE_USOUNDWAVE_HOT_RELOADING_ENABLED is true, the sound wave will have
-		its own, un-published copy of sound wave data (TSharedRef<FSoundWaveData> WorkingSoundWaveData)
-		Otherwise, we use and make changes to the Proxy->PublishedSoundWaveData directly.
-
-		This getter will return the correct one.
-	*/
-	const TSharedRef<FSoundWaveData>& GetWorkingSoundWaveData() const;
-
-	TSharedRef<FSoundWaveProxy> Proxy = MakeShared<FSoundWaveProxy>();
-	
-#if UE_USOUNDWAVE_HOT_RELOADING_ENABLED
-	TSharedRef<FSoundWaveData> WorkingSoundWaveData = MakeShared<FSoundWaveData>();
-#endif // #if UE_USOUNDWAVE_HOT_RELOADING_ENABLED
+	// cached proxy
+	FSoundWaveProxyPtr Proxy{ nullptr };
 
 public:
 
@@ -1172,6 +1156,9 @@ private:
 	/**  Utility function used internally to wait or poll a task while maintaining thread-safety. */
 	ENGINE_API bool WaitAsyncTaskWithTimeout(float InTimeoutInSeconds);
 
+	/** Creates and initializes a new FSoundWaveData and FSoundWaveProxy. */
+	void CreateNewSoundWaveData();
+
 public:
 #endif // WITH_EDITOR
 
@@ -1199,11 +1186,8 @@ public:
 	virtual bool IsProcedurallyGenerated() const override { return bProcedural; }
 	//~ End USoundBase Interface.
 
-	UE_DEPRECATED(5.8, "Mutable access to proxy is deprecated. Use GetSoundWaveProxy()")
 	ENGINE_API FSoundWaveProxyPtr CreateSoundWaveProxy();
 
-	ENGINE_API TSharedRef<const FSoundWaveProxy> GetSoundWaveProxy();
-
 	//~Begin IAudioProxyDataFactory Interface.
 	ENGINE_API virtual TSharedPtr<Audio::IProxyData> CreateProxyData(const Audio::FProxyDataInitParams& InitParams) override;
 	//~ End IAudioProxyDataFactory Interface.
@@ -1583,12 +1567,19 @@ public:
 	ENGINE_API void OverrideWaveData(int32 InNumChannels, int32 InNumFrames, int32 InSampleRate);
 
 private:
+#if PLATFORM_DESKTOP
+	TSharedRef<TDataSnapshotProducer<FSoundWaveData>> SoundWaveData = MakeShared<TDataSnapshotProducer<FSoundWaveData>>();
+	TDataSnapshotProducer<FSoundWaveData>& SoundWaveDataPtr = *SoundWaveData;
+#else // #if PLATFORM_DESKTOP
+	TSharedPtr<FSoundWaveData, ESPMode::ThreadSafe> SoundWaveDataPtr{ MakeShared<FSoundWaveData>() };
+#endif // #if PLATFORM_DESKTOP
+
 	friend class FSoundWaveProxy;
 	friend class USoundFactory;
 };
 
 
-class FSoundWaveData : public TSharedFromThis<FSoundWaveData>
+class FSoundWaveData
 {
 public:
 	UE_NONCOPYABLE(FSoundWaveData);
@@ -1607,13 +1598,11 @@ public:
 	{
 	}
 
-	FSoundWaveData(const USoundWave& InWave);
-
 	// dtor
 	ENGINE_API ~FSoundWaveData();
 
 	// non-const arg for internal call to FAudioDevice::GetRuntimeFormat
-	ENGINE_API void InitializeDataFromSoundWave(const USoundWave& InWave);
+	ENGINE_API void InitializeDataFromSoundWave(USoundWave& InWave);
 
 	ENGINE_API void OverrideRuntimeFormat(const FName& InRuntimeFormat);
 
@@ -1634,10 +1623,7 @@ public:
 	int32 GetNumFrames() const { return NumFrames; }
 	float GetDuration() const { return Duration; }
 
-	ENGINE_API void ReleaseCompressedAudio() const;
-
-	ENGINE_API const TArray<FSoundWaveCuePoint> GetCuePointsAndLoopRegions() const;
-	ENGINE_API TArrayView<const uint8> GetZerothChunk(bool bForImmediatePlayback = false) const;
+	ENGINE_API void ReleaseCompressedAudio();
 
 	ENGINE_API void SetError(const TCHAR* InErrorMsg=nullptr);  
 	ENGINE_API bool HasError() const;
@@ -1655,9 +1641,9 @@ public:
 	bool IsTemplate() const { return bIsTemplate; }
 
  	ENGINE_API bool HasCompressedData(FName Format, ITargetPlatform* TargetPlatform = USoundWave::GetRunningPlatform()) const;
- 	ENGINE_API FByteBulkData* GetCompressedData(FName Format, const FPlatformAudioCookOverrides* CompressionOverrides = USoundWave::GetPlatformCompressionOverridesForCurrentPlatform()) const;
+ 	ENGINE_API FByteBulkData* GetCompressedData(FName Format, const FPlatformAudioCookOverrides* CompressionOverrides = USoundWave::GetPlatformCompressionOverridesForCurrentPlatform());
  
-	ENGINE_API bool GetChunkData(int32 ChunkIndex, uint8** OutChunkData, bool bMakeSureChunkIsLoaded = false) const;
+	ENGINE_API bool GetChunkData(int32 ChunkIndex, uint8** OutChunkData, bool bMakeSureChunkIsLoaded = false);
  
  	ENGINE_API bool IsZerothChunkDataLoaded() const;
  	ENGINE_API const TArrayView<uint8> GetZerothChunkDataView() const;
@@ -1668,22 +1654,21 @@ public:
 	// Returns true if the zeroth chunk is loaded, or attempts to load it if not already loaded,
 	// returning true if the load was successful. Can return false if either an error was encountered
 	// in attempting to load the chunk or if stream caching is not enabled for the given sound.
-	// (const because this call is used to lazy-init zeroth chunk in GetZerothChunkData())
-	ENGINE_API bool LoadZerothChunk() const;
+	ENGINE_API bool LoadZerothChunk();
  
  #if WITH_EDITOR
  	ENGINE_API int32 GetCurrentChunkRevision() const;
  #endif // #if WITH_EDITOR
  
- 	ENGINE_API FStreamedAudioChunk& GetChunk(uint32 ChunkIndex) const;
- 	ENGINE_API int32 GetChunkFromDDC(int32 ChunkIndex, uint8** OutChunkData, bool bMakeSureChunkIsLoaded) const;
+ 	ENGINE_API FStreamedAudioChunk& GetChunk(uint32 ChunkIndex);
+ 	ENGINE_API int32 GetChunkFromDDC(int32 ChunkIndex, uint8** OutChunkData, bool bMakeSureChunkIsLoaded);
 
 #if WITH_EDITORONLY_DATA
  	ENGINE_API FString GetDerivedDataKey() const;
 	FPerPlatformFloat GetSizeOfFirstAudioChunkInSeconds() const { return SizeOfFirstAudioChunkInSeconds; }
 #endif // #if WITH_EDITORONLY_DATA
 
-	int32 GetResourceSize() const { return ResourceSize; }
+	int32 GetResourceSize() { return ResourceSize; }
 	const uint8* GetResourceData() const { return ResourceData.GetView().GetData(); }
 	uint32 GetNumChannels() { return NumChannels; }
 
@@ -1703,8 +1688,7 @@ private:
 	ENGINE_API void CacheRuntimeFormatDependentState(const FName InFormatName);
 
 	/** Zeroth Chunk of audio for sources that use Load On Demand. */
-	// mutable for lazy-init in GetZerothChunkData()
-	mutable FBulkDataBuffer<uint8> ZerothChunkData;
+	FBulkDataBuffer<uint8> ZerothChunkData;
 #if WITH_EDITOR
 	mutable FCriticalSection LoadZerothChunkDataCriticalSection;
 #endif
@@ -1712,16 +1696,12 @@ private:
 	/* Accessor to get the zeroth chunk which might perform additional work in editor to handle async tasks. */
 	ENGINE_API FBulkDataBuffer<uint8>& GetZerothChunkData() const;
 
-
-	// The following members are mutable to allow lazy initialization in const getters
-	// (This way we can hand out immutable refrences to objects of this type)
-
 	/** The streaming derived data for this sound on this platform. */
-	mutable FStreamedAudioPlatformData RunningPlatformData;
+	FStreamedAudioPlatformData RunningPlatformData;
 
-	mutable FAudioChunkHandle FirstChunk;
+	FAudioChunkHandle FirstChunk;
 
-	mutable FFormatContainer CompressedFormatData;
+	FFormatContainer CompressedFormatData;
 
 	int32 ResourceSize = 0;
 
@@ -1761,7 +1741,6 @@ private:
 	uint8 bShouldUseStreamCaching : 1;
 	uint8 bLoadingBehaviorOverridden : 1;
 	uint8 bHasError : 1;
-	
 
 #if WITH_EDITOR
 	uint8 bLoadedFromCookedData : 1;
@@ -1770,159 +1749,97 @@ private:
 	friend class USoundWave;
 }; // class FSoundWaveData
 
-
-
-// i.e.: "FSoundWaveProxy is a shared object between USoundWave and other systems which utilize FSoundWaveData"
-/**
-	The proxy acts as a broker between the USoundWave and Proxy Readers.
-	if UE_USOUNDWAVE_HOT_RELOADING_ENABLED is true, then proxy clients can use this to query the availablity
-	of new FSoundWaveData, and optionally update to that data.
-
-	This lets editor changes and re-imports take effect on even playing sounds.
-
-	In these scenarios, the owning USoundWave has its "working" (un-published) copy of FSoundWaveData
-	For changes, and then calls "Publish()" to make this available
-*/
 using FSoundWaveProxyPtr = TSharedPtr<FSoundWaveProxy, ESPMode::ThreadSafe>;
 class FSoundWaveProxy : public Audio::TProxyData<FSoundWaveProxy>
-					  , public TSharedFromThis<FSoundWaveProxy, ESPMode::ThreadSafe>	
 {
 public:
 	IMPL_AUDIOPROXY_CLASS(FSoundWaveProxy);
 
-	ENGINE_API FSoundWaveProxy();
-	
 	ENGINE_API explicit FSoundWaveProxy(USoundWave* InWave);
 
+	ENGINE_API FSoundWaveProxy(const FSoundWaveProxy& Other);
+
 	ENGINE_API ~FSoundWaveProxy();
-	
-	static ENGINE_API TSharedRef<FSoundWaveProxy> CreateNonUpdatingCopy(const TSharedRef<const FSoundWaveProxy>& Other);
-	static ENGINE_API TSharedRef<const FSoundWaveProxy> CreateNonUpdatingCopy(const TSharedRef<const FSoundWaveData>& Other);
-	
-	friend inline uint32 GetTypeHash(const FSoundWaveProxy& InProxy)
-	{
-		return InProxy.TypeHash;
-	}
 
 	// USoundWave Interface
-	
-	UE_DEPRECATED(5.8, "The FSoundWaveProxy is now just a way to get a reference to the latest FSoundWaveData. Clients should hold onto the return value of GetSoundWaveDataRef() instead")
-	static ENGINE_API TArrayView<const uint8> GetZerothChunk(const FSoundWaveProxyPtr& SoundWaveProxy, bool bForImmediatePlayback = false);
-	UE_DEPRECATED(5.8, "The FSoundWaveProxy is now just a way to get a reference to the latest FSoundWaveData. Clients should hold onto the return value of GetSoundWaveDataRef() instead")
-	ENGINE_API const TArray<FSoundWaveCuePoint> GetCuePointsAndLoopRegions() const;
-	UE_DEPRECATED(5.8, "The FSoundWaveProxy is now just a way to get a reference to the latest FSoundWaveData. Clients should hold onto the return value of GetSoundWaveDataRef() instead")
 	ENGINE_API void ReleaseCompressedAudio();
 
 	// Returns true if the zeroth chunk is loaded, or attempts to load it if not already loaded,
 	// returning true if the load was successful. Can return false if either an error was encountered
-	// in attempting to load the chunk or if stream caching is not enabled for the given sound
-	UE_DEPRECATED(5.8, "The FSoundWaveProxy is now just a way to get a reference to the latest FSoundWaveData. Clients should hold onto the return value of GetSoundWaveDataRef() instead")
+	// in attempting to load the chunk or if stream caching is not enabled for the given sound.
 	ENGINE_API bool LoadZerothChunk();
-	UE_DEPRECATED(5.8, "The FSoundWaveProxy is now just a way to get a reference to the latest FSoundWaveData. Clients should hold onto the return value of GetSoundWaveDataRef() instead")
 	ENGINE_API bool GetChunkData(int32 ChunkIndex, uint8** OutChunkData, bool bMakeSureChunkIsLoaded = false);
-	
+
 	// Getters
-	UE_DEPRECATED(5.8, "The FSoundWaveProxy is now just a way to get a reference to the latest FSoundWaveData. Clients should hold onto the return value of GetSoundWaveDataRef() instead")
 	ENGINE_API const FName& GetFName() const;
-	UE_DEPRECATED(5.8, "The FSoundWaveProxy is now just a way to get a reference to the latest FSoundWaveData. Clients should hold onto the return value of GetSoundWaveDataRef() instead")
 	ENGINE_API const FName& GetPackageName() const;
-	UE_DEPRECATED(5.8, "The FSoundWaveProxy is now just a way to get a reference to the latest FSoundWaveData. Clients should hold onto the return value of GetSoundWaveDataRef() instead")
 	ENGINE_API const FName& GetRuntimeFormat() const;
-	UE_DEPRECATED(5.8, "The FSoundWaveProxy is now just a way to get a reference to the latest FSoundWaveData. Clients should hold onto the return value of GetSoundWaveDataRef() instead")
 	ENGINE_API const FObjectKey& GetFObjectKey() const;
-	
-	UE_DEPRECATED(5.8, "The FSoundWaveProxy is now just a way to get a reference to the latest FSoundWaveData. Clients should hold onto the return value of GetSoundWaveDataRef() instead")
+
 	ENGINE_API float GetDuration() const;
-	UE_DEPRECATED(5.8, "The FSoundWaveProxy is now just a way to get a reference to the latest FSoundWaveData. Clients should hold onto the return value of GetSoundWaveDataRef() instead")
 	ENGINE_API float GetSampleRate() const;
-	
-	UE_DEPRECATED(5.8, "The FSoundWaveProxy is now just a way to get a reference to the latest FSoundWaveData. Clients should hold onto the return value of GetSoundWaveDataRef() instead")
+
 	ENGINE_API int32 GetNumFrames() const;
-	UE_DEPRECATED(5.8, "The FSoundWaveProxy is now just a way to get a reference to the latest FSoundWaveData. Clients should hold onto the return value of GetSoundWaveDataRef() instead")
 	ENGINE_API uint32 GetNumChunks() const;
-	UE_DEPRECATED(5.8, "The FSoundWaveProxy is now just a way to get a reference to the latest FSoundWaveData. Clients should hold onto the return value of GetSoundWaveDataRef() instead")
 	ENGINE_API uint32 GetNumChannels() const;
-	UE_DEPRECATED(5.8, "The FSoundWaveProxy is now just a way to get a reference to the latest FSoundWaveData. Clients should hold onto the return value of GetSoundWaveDataRef() instead")
 	ENGINE_API uint32 GetSizeOfChunk(uint32 ChunkIndex) const;
-	UE_DEPRECATED(5.8, "The FSoundWaveProxy is now just a way to get a reference to the latest FSoundWaveData. Clients should hold onto the return value of GetSoundWaveDataRef() instead")
 	ENGINE_API const TArray<FSoundWaveCuePoint>& GetCuePoints() const;
-	UE_DEPRECATED(5.8, "The FSoundWaveProxy is now just a way to get a reference to the latest FSoundWaveData. Clients should hold onto the return value of GetSoundWaveDataRef() instead")
 	ENGINE_API const TArray<FSoundWaveCuePoint>& GetLoopRegions() const;
-	
-	UE_DEPRECATED(5.8, "The FSoundWaveProxy is now just a way to get a reference to the latest FSoundWaveData. Clients should hold onto the return value of GetSoundWaveDataRef() instead")
+	ENGINE_API const TArray<FSoundWaveCuePoint> GetCuePointsAndLoopRegions() const;
+
 	ENGINE_API bool IsLooping() const;
-	UE_DEPRECATED(5.8, "The FSoundWaveProxy is now just a way to get a reference to the latest FSoundWaveData. Clients should hold onto the return value of GetSoundWaveDataRef() instead")
 	ENGINE_API bool IsTemplate() const;
-	UE_DEPRECATED(5.8, "The FSoundWaveProxy is now just a way to get a reference to the latest FSoundWaveData. Clients should hold onto the return value of GetSoundWaveDataRef() instead")
 	ENGINE_API bool IsStreaming() const;
-	UE_DEPRECATED(5.8, "The FSoundWaveProxy is now just a way to get a reference to the latest FSoundWaveData. Clients should hold onto the return value of GetSoundWaveDataRef() instead")
 	ENGINE_API bool IsRetainingAudio() const;
-	UE_DEPRECATED(5.8, "The FSoundWaveProxy is now just a way to get a reference to the latest FSoundWaveData. Clients should hold onto the return value of GetSoundWaveDataRef() instead")
 	ENGINE_API bool ShouldUseStreamCaching() const;
-	UE_DEPRECATED(5.8, "The FSoundWaveProxy is now just a way to get a reference to the latest FSoundWaveData. Clients should hold onto the return value of GetSoundWaveDataRef() instead")
 	ENGINE_API bool IsSeekable() const;
-	UE_DEPRECATED(5.8, "The FSoundWaveProxy is now just a way to get a reference to the latest FSoundWaveData. Clients should hold onto the return value of GetSoundWaveDataRef() instead")
 	ENGINE_API bool IsZerothChunkDataLoaded() const;
-	UE_DEPRECATED(5.8, "The FSoundWaveProxy is now just a way to get a reference to the latest FSoundWaveData. Clients should hold onto the return value of GetSoundWaveDataRef() instead")
 	ENGINE_API bool WasLoadingBehaviorOverridden() const;
-	UE_DEPRECATED(5.8, "The FSoundWaveProxy is now just a way to get a reference to the latest FSoundWaveData. Clients should hold onto the return value of GetSoundWaveDataRef() instead")
 	ENGINE_API bool HasCompressedData(FName Format, ITargetPlatform* TargetPlatform = USoundWave::GetRunningPlatform()) const;
-	
-	UE_DEPRECATED(5.8, "The FSoundWaveProxy is now just a way to get a reference to the latest FSoundWaveData. Clients should hold onto the return value of GetSoundWaveDataRef() instead")
+
 	ENGINE_API ESoundWaveLoadingBehavior GetLoadingBehavior() const;
-	
-	UE_DEPRECATED(5.8, "The FSoundWaveProxy is now just a way to get a reference to the latest FSoundWaveData. Clients should hold onto the return value of GetSoundWaveDataRef() instead")
+
 	ENGINE_API const TArrayView<uint8> GetZerothChunkDataView() const;
-	
-	UE_DEPRECATED(5.8, "The FSoundWaveProxy is now just a way to get a reference to the latest FSoundWaveData. Clients should hold onto the return value of GetSoundWaveDataRef() instead")
+
 	ENGINE_API FByteBulkData* GetCompressedData(FName Format, const FPlatformAudioCookOverrides* CompressionOverrides = USoundWave::GetPlatformCompressionOverridesForCurrentPlatform());
 
+	static ENGINE_API TArrayView<const uint8> GetZerothChunk(const FSoundWaveProxyPtr& SoundWaveProxy, bool bForImmediatePlayback = false);
+
+
 #if WITH_EDITOR
-	UE_DEPRECATED(5.8, "The FSoundWaveProxy is now just a way to get a reference to the latest FSoundWaveData. Clients should hold onto the return value of GetSoundWaveDataRef() instead")
 	ENGINE_API int32 GetCurrentChunkRevision() const;
 #endif // #if WITH_EDITOR
 
-	UE_DEPRECATED(5.8, "The FSoundWaveProxy is now just a way to get a reference to the latest FSoundWaveData. Clients should hold onto the return value of GetSoundWaveDataRef() instead")
 	ENGINE_API FStreamedAudioChunk& GetChunk(uint32 ChunkIndex);
-	UE_DEPRECATED(5.8, "The FSoundWaveProxy is now just a way to get a reference to the latest FSoundWaveData. Clients should hold onto the return value of GetSoundWaveDataRef() instead")
 	ENGINE_API int32 GetChunkFromDDC(int32 ChunkIndex, uint8** OutChunkData, bool bMakeSureChunkIsLoaded);
 
 #if WITH_EDITORONLY_DATA
-	UE_DEPRECATED(5.8, "The FSoundWaveProxy is now just a way to get a reference to the latest FSoundWaveData. Clients should hold onto the return value of GetSoundWaveDataRef() instead")
 	ENGINE_API FString GetDerivedDataKey() const;
 #endif // #if WITH_EDITORONLY_DATA
 
-	UE_DEPRECATED(5.8, "The FSoundWaveProxy is now just a way to get a reference to the latest FSoundWaveData. Clients should hold onto the return value of GetSoundWaveDataRef() instead")
 	ENGINE_API int32 GetResourceSize() const;
-
-	UE_DEPRECATED(5.8, "The FSoundWaveProxy is now just a way to get a reference to the latest FSoundWaveData. Clients should hold onto the return value of GetSoundWaveDataRef() instead")
 	ENGINE_API const uint8* GetResourceData() const;
 
-	UE_DEPRECATED(5.8, "mutable access to sound wave data is deprecated. please use GetSoundWaveDataRef()")
-	TSharedPtr<FSoundWaveData> GetSoundWaveData() const
+	ENGINE_API const FSoundWavePtr GetSoundWaveData();
+	
+	friend inline uint32 GetTypeHash(const FSoundWaveProxy& InProxy)
 	{
-		return ConstCastSharedRef<FSoundWaveData>(PublishedSoundWaveData).ToSharedPtr();
+		return InProxy.TypeHash;
 	}
-	
-	ENGINE_API TSharedRef<const FSoundWaveData> GetSoundWaveDataRef() const;
 
-#if UE_USOUNDWAVE_HOT_RELOADING_ENABLED
-	bool ENGINE_API IsLatestData(const FSoundWaveData* InData) const;
-#endif // #if UE_USOUNDWAVE_HOT_RELOADING_ENABLED
+#if PLATFORM_DESKTOP
+	bool TryGetLatest()
+	{
+		return SoundWaveDataPtr.GetLatest();
+	}
+#endif // #if PLATFORM_DESKTOP
 	
 private:
-	friend USoundWave;
-	void InitializeFromWave(USoundWave& InWave);
-
-#if UE_USOUNDWAVE_HOT_RELOADING_ENABLED	
-	void Publish(const TSharedRef<FSoundWaveData>& InNewData);
-#endif // #if UE_USOUNDWAVE_HOT_RELOADING_ENABLED
-	
-	TSharedRef<FSoundWaveData> PublishedSoundWaveData = MakeShared<FSoundWaveData>();
-
-#if UE_USOUNDWAVE_HOT_RELOADING_ENABLED
-	mutable FRWLock PublishedSoundWaveDataLock;
-#endif // #if UE_USOUNDWAVE_HOT_RELOADING_ENABLED
+#if PLATFORM_DESKTOP
+	TDataSnapshot<FSoundWaveData> SoundWaveDataPtr;
+#else // #if PLATFORM_DESKTOP
+	TSharedPtr<FSoundWaveData, ESPMode::ThreadSafe> SoundWaveDataPtr;
+#endif // #if PLATFORM_DESKTOP
 	
 	uint32 TypeHash = uint32(INDEX_NONE);
 };
diff --git a/Engine/Source/Runtime/Engine/Classes/Sound/SoundWaveProxyPlayer.h b/Engine/Source/Runtime/Engine/Classes/Sound/SoundWaveProxyPlayer.h
--- a/Engine/Source/Runtime/Engine/Classes/Sound/SoundWaveProxyPlayer.h
+++ b/Engine/Source/Runtime/Engine/Classes/Sound/SoundWaveProxyPlayer.h
@@ -83,17 +83,15 @@ public:
 	UE_API static TUniquePtr<FSoundWaveProxyPlayer> Create(const FSettings& InSettings);
 
 	/**
-	 * Set the SoundWave Data to play
+	 * Set the SoundWave Proxy to play
 	 * 
-	* if the FSoundWaveData has invalid data or is unabled to be decoded,
+	 * if Null, will reset the Proxy Player and return 'false'
+	 * if the WaveProxy has invalid data or is unabled to be decoded,
 	 * the player will be reset and this will return 'false'
 	 * 
-	 * @param InWaveData 
+	 * @param InWaveProxy 
 	 * @return whether the player is valid for playback
 	 */
-	UE_API bool SetSoundWave(const TSharedRef<const FSoundWaveData>& InWaveData); 
-	
-	UE_DEPRECATED(5.8, "Please use the TSharedRef<const FSoundWaveData>& overload instead")
 	UE_API bool SetSoundWave(FSoundWaveProxyPtr InWaveProxy);
 
 	/**
@@ -284,12 +282,12 @@ private:
 	FSoundWaveProxyPlayer(const FSettings& InSettings);
 
 	/**
-	 * Check whether the Wave Proxy has valid data (valid Sample Rate, Num Channels, Num Frames)
+	 * Check wether the Wave Proxy has valid data (valid Sample Rate, Num Channels, Num Frames)
 	 * 
-	 * @param InWaveData
+	 * @param InWaveProxy 
 	 * @return whether the WaveProxy is valid
 	 */
-	static bool IsValidSoundWaveData(const FSoundWaveData& InWaveData);
+	static bool IsValidSoundWaveProxy(FSoundWaveProxyRef InWaveProxy);
 
 	// 
 	/**
@@ -344,14 +342,15 @@ private:
 	/**
 	 * Initialize the Decoder with the given wave proxy
 	 * also initialize the decode buffers
+	 * @param WaveProxy
 	 * @param MaxDecodeSizeInFrames 
 	 * @return whether the decoder was successfully initialized
 	 */
-	bool InitializeDecoder(int32 MaxDecodeSizeInFrames);
+	bool InitializeDecoder(const FSoundWaveProxyRef& InWaveProxy, int32 MaxDecodeSizeInFrames);
 
 	void InitializeResampler(int32 NumChannels);
 
-	TSharedPtr<const FSoundWaveData> SoundWaveDataPtr;
+	FSoundWaveProxyPtr SoundWave;
 	FSettings Settings;
 	TUniquePtr<ICompressedAudioInfo> CompressedAudioInfo;
 	TUniquePtr<Audio::IConvertDeinterleave> ConvertDeinterleave;
diff --git a/Engine/Source/Runtime/Engine/Classes/Sound/SoundWaveProxyReader.h b/Engine/Source/Runtime/Engine/Classes/Sound/SoundWaveProxyReader.h
--- a/Engine/Source/Runtime/Engine/Classes/Sound/SoundWaveProxyReader.h
+++ b/Engine/Source/Runtime/Engine/Classes/Sound/SoundWaveProxyReader.h
@@ -21,7 +21,6 @@ class FSoundWaveProxyReader
 
 public:
 	using FSoundWaveProxyRef = TSharedRef<FSoundWaveProxy, ESPMode::ThreadSafe>;
-	using FSoundWaveDataRef = TSharedRef<const FSoundWaveData, ESPMode::ThreadSafe>;
 
 	/** Minimum number of frames to decode per a call to the decoder.  */
 	static constexpr uint32 DefaultMinDecodeSizeInFrames = 128;
@@ -59,24 +58,17 @@ public:
 private:
 	/** Construct a wave proxy reader.
 	 *
-	 * @param InWaveData - A TSharedRef of a const FSoundWaveData which is to be played.
+	 * @param InWaveProxy - A TSharedRef of a FSoundWaveProxy which is to be played.
 	 * @param InSettings - Reader settings.
 	 */
-	FSoundWaveProxyReader(FSoundWaveDataRef InWaveData, const FSettings& InSettings);
-	
-	UE_DEPRECATED(5.8, "Please use the FSoundWaveDataRef overload instead")
 	FSoundWaveProxyReader(FSoundWaveProxyRef InWaveProxy, const FSettings& InSettings);
-	
 
 public:
 	/** Create a wave proxy reader.
 	 *
-	 * @param InWaveProxy - A TSharedRef of a const FSoundWaveData which is to be played.
+	 * @param InWaveProxy - A TSharedRef of a FSoundWaveProxy which is to be played.
 	 * @param InSettings - Reader settings.
 	 */
-	static ENGINE_API TUniquePtr<FSoundWaveProxyReader> Create(FSoundWaveDataRef InWaveProxy, const FSettings& InSettings);
-	
-	UE_DEPRECATED(5.8, "Please use the FSoundWaveDataRef overload instead")
 	static ENGINE_API TUniquePtr<FSoundWaveProxyReader> Create(FSoundWaveProxyRef InWaveProxy, const FSettings& InSettings);
 
 	/** Set whether the reader should loop the audio or not. */
@@ -189,7 +181,7 @@ private:
 	float ClampLoopDuration(float InDurationInSeconds);
 	void UpdateLoopBoundaries();
 
-	FSoundWaveDataRef WaveDataRef;
+	FSoundWaveProxyRef WaveProxy;
 
 	FSettings Settings;
 
