Added a constexpr-if to compile out functionality to toggle fullscreen on non-desktop platforms.

#Synced-CL 50560025

diff --git a/Engine/Source/Runtime/Engine/Private/GameViewportClient.cpp b/Engine/Source/Runtime/Engine/Private/GameViewportClient.cpp
--- a/Engine/Source/Runtime/Engine/Private/GameViewportClient.cpp
+++ b/Engine/Source/Runtime/Engine/Private/GameViewportClient.cpp
@@ -171,13 +171,6 @@ static TAutoConsoleVariable<int32> CVarScreenshotTextureArrayIndex(
 	TEXT("Texture array slice to use for screenshots when r.Screenshot.TextureArrayMode = 0"),
 	ECVF_Default);
 
-static bool bForceAllowFullscreenToggle = false;
-static FAutoConsoleVariableRef ForceAllowFullscreenToggleCVar(
-	TEXT("r.ForceAllowFullscreenToggle"),
-	bForceAllowFullscreenToggle,
-	TEXT("Temporary CVar to allow rolling back logic that prevents full-screen toggling on non-desktop platforms"),
-	ECVF_Default);
-
 #if CSV_PROFILER_STATS
 struct FCsvLocalPlayer
 {
@@ -3983,84 +3976,82 @@ bool UGameViewportClient::SetDisplayConfiguration(const FIntPoint* Dimensions, E
 
 bool UGameViewportClient::HandleToggleFullscreenCommand()
 {
-	if (!FPlatformProperties::SupportsWindowedMode() && !bForceAllowFullscreenToggle)
+	if constexpr (FPlatformProperties::SupportsWindowedMode())
 	{
-		return true;
-	}
-
-	if (!Viewport)
-	{
-		return true;
-	}
+		if (!Viewport)
+		{
+			return true;
+		}
 
-	static auto CVar = IConsoleManager::Get().FindTConsoleVariableDataInt(TEXT("r.FullScreenMode"));
-	check(CVar);
-	EWindowMode::Type FullScreenMode = Viewport->IsFullscreen()
-		? EWindowMode::Windowed
-		: (((CVar->GetValueOnGameThread() == 0) && Viewport->IsExclusiveFullscreenSupported())
-			? EWindowMode::Fullscreen
-			: EWindowMode::WindowedFullscreen);
+		static auto CVar = IConsoleManager::Get().FindTConsoleVariableDataInt(TEXT("r.FullScreenMode"));
+		check(CVar);
+		EWindowMode::Type FullScreenMode = Viewport->IsFullscreen()
+			? EWindowMode::Windowed
+			: (((CVar->GetValueOnGameThread() == 0) && Viewport->IsExclusiveFullscreenSupported())
+				? EWindowMode::Fullscreen
+				: EWindowMode::WindowedFullscreen);
 
-	int32 ResolutionX = GSystemResolution.ResX;
-	int32 ResolutionY = GSystemResolution.ResY;
-	const int32 MonitorIndex = Viewport->GetCurrentMonitorIndex();
-	bool bNewModeApplied = false;
+		int32 ResolutionX = GSystemResolution.ResX;
+		int32 ResolutionY = GSystemResolution.ResY;
+		const int32 MonitorIndex = Viewport->GetCurrentMonitorIndex();
+		bool bNewModeApplied = false;
 
-	// Make sure the user's settings are updated after pressing Alt+Enter to toggle fullscreen.  Note
-	// that we don't need to "apply" the setting change, as we already did that above directly.
-	UGameEngine* const GameEngine = Cast<UGameEngine>( GEngine );
-	if (GameEngine != nullptr)
-	{
-		UGameUserSettings* const UserSettings = GameEngine->GetGameUserSettings();
-		if (UserSettings != nullptr)
+		// Make sure the user's settings are updated after pressing Alt+Enter to toggle fullscreen.  Note
+		// that we don't need to "apply" the setting change, as we already did that above directly.
+		UGameEngine* const GameEngine = Cast<UGameEngine>( GEngine );
+		if (GameEngine != nullptr)
 		{
-			ResolutionX = UserSettings->GetScreenResolution().X;
-			ResolutionY = UserSettings->GetScreenResolution().Y;
-
-			const FDisplayMetrics& DisplayMetrics = FSlateApplication::Get().GetCachedDisplayMetricsByRef();
-			if (DisplayMetrics.MonitorInfo.IsValidIndex(MonitorIndex))
+			UGameUserSettings* const UserSettings = GameEngine->GetGameUserSettings();
+			if (UserSettings != nullptr)
 			{
-				const FMonitorInfo& Monitor = DisplayMetrics.MonitorInfo[MonitorIndex];
-				if (FullScreenMode == EWindowMode::Fullscreen)
+				ResolutionX = UserSettings->GetScreenResolution().X;
+				ResolutionY = UserSettings->GetScreenResolution().Y;
+
+				const FDisplayMetrics& DisplayMetrics = FSlateApplication::Get().GetCachedDisplayMetricsByRef();
+				if (DisplayMetrics.MonitorInfo.IsValidIndex(MonitorIndex))
 				{
-					ResolutionX = Monitor.MaxResolution.X;
-					ResolutionY = Monitor.MaxResolution.Y;
-					if ((ResolutionX == 0) || (ResolutionY == 0))
+					const FMonitorInfo& Monitor = DisplayMetrics.MonitorInfo[MonitorIndex];
+					if (FullScreenMode == EWindowMode::Fullscreen)
 					{
-						ResolutionX = Monitor.NativeWidth;
-						ResolutionY = Monitor.NativeHeight;
+						ResolutionX = Monitor.MaxResolution.X;
+						ResolutionY = Monitor.MaxResolution.Y;
+						if ((ResolutionX == 0) || (ResolutionY == 0))
+						{
+							ResolutionX = Monitor.NativeWidth;
+							ResolutionY = Monitor.NativeHeight;
+						}
 					}
-				}
-				else
-				{
-					const FPlatformRect& DisplayRect = Monitor.DisplayRect;
-					ResolutionX = DisplayRect.Right - DisplayRect.Left;
-					ResolutionY = DisplayRect.Bottom - DisplayRect.Top;
+					else
+					{
+						const FPlatformRect& DisplayRect = Monitor.DisplayRect;
+						ResolutionX = DisplayRect.Right - DisplayRect.Left;
+						ResolutionY = DisplayRect.Bottom - DisplayRect.Top;
+					}
+
+					UserSettings->SetDisplayProperties(Monitor.ID, MonitorIndex);
 				}
 
-				UserSettings->SetDisplayProperties(Monitor.ID, MonitorIndex);
-			}
+				// Ensure that our desired screen size will fit on the display
+				UGameEngine::DetermineGameWindowResolution(ResolutionX, ResolutionY, FullScreenMode, true, MonitorIndex);
 
-			// Ensure that our desired screen size will fit on the display
-			UGameEngine::DetermineGameWindowResolution(ResolutionX, ResolutionY, FullScreenMode, true, MonitorIndex);
+				UserSettings->SetScreenResolution(FIntPoint(ResolutionX, ResolutionY));
+				UserSettings->SetFullscreenMode(FullScreenMode);
+				UserSettings->ConfirmVideoMode();
+				UserSettings->ApplySettings(false);
+				bNewModeApplied = true;
+			}
+		}
 
-			UserSettings->SetScreenResolution(FIntPoint(ResolutionX, ResolutionY));
-			UserSettings->SetFullscreenMode(FullScreenMode);
-			UserSettings->ConfirmVideoMode();
-			UserSettings->ApplySettings(false);
-			bNewModeApplied = true;
+		if (!bNewModeApplied)
+		{
+			FSystemResolution::SetMonitorIndex(MonitorIndex);
+			UE_LOG(LogViewport, Display, TEXT("UGameViewportClient::HandleToggleFullscreenCommand: calling RequestResolutionChange %dx%d %s"), ResolutionX, ResolutionY, LexToString(FullScreenMode));
+			FSystemResolution::RequestResolutionChange(ResolutionX, ResolutionY, FullScreenMode);
 		}
-	}
 
-	if (!bNewModeApplied)
-	{
-		FSystemResolution::SetMonitorIndex(MonitorIndex);
-		UE_LOG(LogViewport, Display, TEXT("UGameViewportClient::HandleToggleFullscreenCommand: calling RequestResolutionChange %dx%d %s"), ResolutionX, ResolutionY, LexToString(FullScreenMode));
-		FSystemResolution::RequestResolutionChange(ResolutionX, ResolutionY, FullScreenMode);
+		ToggleFullscreenDelegate.Broadcast(FullScreenMode != EWindowMode::Windowed);
 	}
 
-	ToggleFullscreenDelegate.Broadcast(FullScreenMode != EWindowMode::Windowed);
-
 	return true;
 }
 
