Only patch animsequence tracks if it isn't already queued for a full patch.

#Synced-CL 50568424

diff --git a/Engine/Source/Runtime/Engine/Classes/Animation/AnimSequenceTransformProviderData.h b/Engine/Source/Runtime/Engine/Classes/Animation/AnimSequenceTransformProviderData.h
--- a/Engine/Source/Runtime/Engine/Classes/Animation/AnimSequenceTransformProviderData.h
+++ b/Engine/Source/Runtime/Engine/Classes/Animation/AnimSequenceTransformProviderData.h
@@ -122,6 +122,7 @@ public:
 		Data.SequenceIndex = AutoPlayData.SequenceIndex;
 		Data.bAutoPlay = 1;
 		Data.LoopMode = static_cast<uint32>(AutoPlayData.LoopMode);
+		Data.bFullPatch = 1;
 		Data.PatchFlags = 0;
 		Data.bHasPreviousPosition = false;
 		Data.Auto.ReferenceTimestamp = GetTimestamp(WorldTime, AutoPlayData.Position, AutoPlayData.PlayRate);
@@ -146,6 +147,7 @@ public:
 		Data.SequenceIndex = ManualData.SequenceIndex;
 		Data.bAutoPlay = 0;
 		Data.LoopMode = static_cast<uint32>(ManualData.LoopMode);
+		Data.bFullPatch = 1;
 		Data.PatchFlags = 0;
 		Data.bHasPreviousPosition = ManualData.bHasPreviousPosition;
 		Data.Manual.Position = GetPosition(ManualData.LoopMode, ManualData.Position, SequenceLength);
@@ -204,6 +206,12 @@ public:
 	void ResetPatchFlags()
 	{
 		PatchFlags = 0;
+		bFullPatch = 0;
+	}
+
+	bool IsFullPatch() const
+	{
+		return bFullPatch;
 	}
 
 	EAnimSequenceTrackPatchFlags GetPatchFlags() const
@@ -280,8 +288,8 @@ private:
 	uint32 LoopMode : 1;
 	uint32 bAutoPlay : 1;
 	uint32 bHasPreviousPosition : 1;
+	uint32 bFullPatch : 1;
 	uint32 PatchFlags : 8;
-	uint32 Unused : 1 = 0;
 
 	union
 	{
diff --git a/Engine/Source/Runtime/Engine/Private/Animation/AnimSequenceTransformProviderData.cpp b/Engine/Source/Runtime/Engine/Private/Animation/AnimSequenceTransformProviderData.cpp
--- a/Engine/Source/Runtime/Engine/Private/Animation/AnimSequenceTransformProviderData.cpp
+++ b/Engine/Source/Runtime/Engine/Private/Animation/AnimSequenceTransformProviderData.cpp
@@ -87,6 +87,7 @@ void FAnimSequenceTrackPackedData::Serialize(FArchive& Ar)
 		LoopMode = (PackedBits >> LoopModeShift) & LoopModeMask;
 		bAutoPlay = (PackedBits >> AutoPlayShift) & AutoPlayMask;
 		bHasPreviousPosition = (PackedBits >> HasPreviousPositionShift) & HasPreviousPositionMask;
+		bFullPatch = 1;
 		PatchFlags = 0;
 
 		if (Ar.CustomVer(FFortniteMainBranchObjectVersion::GUID) >= FFortniteMainBranchObjectVersion::AnimSequenceTrackTimestamps)
@@ -130,7 +131,7 @@ void FAnimSequenceTrackPackedData::Serialize(FArchive& Ar)
 void FAnimSequenceTrackRenderData::Patch(const FAnimSequenceTrackPackedData& Data, double WorldTime, bool bWasActive)
 {
 	// Only patching specific flags.
-	if (Data.GetPatchFlags() != EAnimSequenceTrackPatchFlags::None)
+	if (!Data.IsFullPatch() && Data.GetPatchFlags() != EAnimSequenceTrackPatchFlags::None)
 	{
 		if (EnumHasAnyFlags(Data.GetPatchFlags(), EAnimSequenceTrackPatchFlags::PlayRate))
 		{
@@ -143,9 +144,6 @@ void FAnimSequenceTrackRenderData::Patch(const FAnimSequenceTrackPackedData& Dat
 			Source.SetLoopMode(Data.GetLoopMode());
 			Target.SetLoopMode(Data.GetLoopMode());
 		}
-
-		Source.ResetPatchFlags();
-		Target.ResetPatchFlags();
 	}
 	// Patching the entire struct. Starts a new animation.
 	else
@@ -162,6 +160,9 @@ void FAnimSequenceTrackRenderData::Patch(const FAnimSequenceTrackPackedData& Dat
 		BlendWeight = 0.0f;
 		Target = Data;
 	}
+
+	Source.ResetPatchFlags();
+	Target.ResetPatchFlags();
 }
 
 void FAnimSequenceTrackRenderData::Tick(float DeltaTime, float SourcePosition, float TargetPosition)
