Primitive Drawing - Fixed capsule draw priority on the half circles caps.

#rnx
#Synced-CL 50399480

diff --git a/Engine/Source/Runtime/Engine/Private/PrimitiveDrawingUtils.cpp b/Engine/Source/Runtime/Engine/Private/PrimitiveDrawingUtils.cpp
--- a/Engine/Source/Runtime/Engine/Private/PrimitiveDrawingUtils.cpp
+++ b/Engine/Source/Runtime/Engine/Private/PrimitiveDrawingUtils.cpp
@@ -1217,8 +1217,7 @@ void DrawWireCylinder(FPrimitiveDrawInterface* PDI, const FVector& Base, const F
 	}
 }
 
-
-static void DrawHalfCircle(FPrimitiveDrawInterface* PDI, const FVector& Base, const FVector& X, const FVector& Y, const FLinearColor& Color, double Radius, int32 NumSides, float Thickness, float DepthBias, bool bScreenSpace)
+static void DrawHalfCircle(FPrimitiveDrawInterface* PDI, const FVector& Base, const FVector& X, const FVector& Y, const FLinearColor& Color, double Radius, int32 NumSides, uint8 DepthPriority, float Thickness, float DepthBias, bool bScreenSpace)
 {
 	const float	AngleDelta = (float)UE_PI / ((float)NumSides);
 	FVector	LastVertex = Base + X * Radius;
@@ -1226,7 +1225,7 @@ static void DrawHalfCircle(FPrimitiveDrawInterface* PDI, const FVector& Base, co
 	for(int32 SideIndex = 0; SideIndex < NumSides; SideIndex++)
 	{
 		const FVector	Vertex = Base + (X * FMath::Cos(AngleDelta * (SideIndex + 1)) + Y * FMath::Sin(AngleDelta * (SideIndex + 1))) * Radius;
-		PDI->DrawTranslucentLine(LastVertex, Vertex, Color, SDPG_World, Thickness, DepthBias, bScreenSpace);
+		PDI->DrawTranslucentLine(LastVertex, Vertex, Color, DepthPriority, Thickness, DepthBias, bScreenSpace);
 		LastVertex = Vertex;
 	}	
 }
@@ -1256,13 +1255,13 @@ void DrawWireCapsule(FPrimitiveDrawInterface* PDI, const FVector& Base, const FV
 	DrawCircle(PDI, BottomEnd, XAxis, YAxis, Color, CapsuleRadius, NumSides, DepthPriority, Thickness, DepthBias, bScreenSpace);
 
 	// Draw domed caps
-	DrawHalfCircle(PDI, TopEnd, YAxis, ZAxis, Color, CapsuleRadius, NumSides / 2, Thickness, DepthBias, bScreenSpace);
-	DrawHalfCircle(PDI, TopEnd, XAxis, ZAxis, Color, CapsuleRadius, NumSides / 2, Thickness, DepthBias, bScreenSpace);
+	DrawHalfCircle(PDI, TopEnd, YAxis, ZAxis, Color, CapsuleRadius, NumSides / 2, DepthPriority, Thickness, DepthBias, bScreenSpace);
+	DrawHalfCircle(PDI, TopEnd, XAxis, ZAxis, Color, CapsuleRadius, NumSides / 2, DepthPriority, Thickness, DepthBias, bScreenSpace);
 
 	const FVector NegZAxis = -ZAxis;
 
-	DrawHalfCircle(PDI, BottomEnd, YAxis, NegZAxis, Color, CapsuleRadius, NumSides / 2, Thickness, DepthBias, bScreenSpace);
-	DrawHalfCircle(PDI, BottomEnd, XAxis, NegZAxis, Color, CapsuleRadius, NumSides / 2, Thickness, DepthBias, bScreenSpace);
+	DrawHalfCircle(PDI, BottomEnd, YAxis, NegZAxis, Color, CapsuleRadius, NumSides / 2, DepthPriority, Thickness, DepthBias, bScreenSpace);
+	DrawHalfCircle(PDI, BottomEnd, XAxis, NegZAxis, Color, CapsuleRadius, NumSides / 2, DepthPriority, Thickness, DepthBias, bScreenSpace);
 
 	// we set NumSides to 4 as it makes a nicer looking capsule as we only draw 2 HalfCircles above
 	const int32 NumCylinderLines = 4;
