[AsyncCompilation]
- Add support to read .ini settings early during boot so that async compilation can be active earlier

#Synced-CL 50389147

diff --git a/Engine/Source/Runtime/Engine/Private/AsyncCompilationHelpers.cpp b/Engine/Source/Runtime/Engine/Private/AsyncCompilationHelpers.cpp
--- a/Engine/Source/Runtime/Engine/Private/AsyncCompilationHelpers.cpp
+++ b/Engine/Source/Runtime/Engine/Private/AsyncCompilationHelpers.cpp
@@ -4,6 +4,7 @@
 #include "UObject/UnrealType.h"
 
 #include "AssetCompilingManager.h"
+#include "Misc/ConfigCacheIni.h"
 #include "Misc/DelayedAutoRegister.h"
 #include "Misc/ScopedSlowTask.h"
 #include "Misc/QueuedThreadPool.h"
@@ -239,6 +240,13 @@ namespace AsyncCompilationHelpers
 	{
 		if (InExperimentalSettingsName != NAME_None)
 		{
+			// Early at boot, UEditorExperimentalSettings is unavailable, go straight to the ini file instead.
+			bool bIsEnabled = false;
+			if (GConfig->GetBool(TEXT("/Script/UnrealEd.EditorExperimentalSettings"), *InExperimentalSettingsName.ToString(), bIsEnabled, GEditorPerProjectIni))
+			{
+				InCVarAsyncCompilation->Set(bIsEnabled ? 1 : 0, ECVF_SetByProjectSetting);
+			}
+			
 			// We cannot try to get the UEditorExperimentalSettings until all the class properties
 			// have been linked otherwise the CDO wouldn't be able to initialize properties
 			// from the ini files which would most likely end up with the default value of
