Support for external entities in Prefabs and Non-Root world entities

Added menu entries to externalize/internalize owned entities for prefabs and non-root entities, enabled according to relevant cvars.
External owned entities is enabled/disabled per-Entity by the bSaveOwnedEntitiesAsExternalFiles property which sets the RF_HasDynamicImports object flag and adds an Externalized override operation in the OverriddenPropertySet for that that object/property.
Those externally packaged objects have additional metadata to allow tracking of the override operations.

When using external entities in an instance (either in a derived prefab or an world instance) the overrides operations are distributed in the external entities
* Added entities:
*Detected in entity::Rename(), and in a post instantiation delegate
*Saves the entity in its own file with specific metadata
* Removed entities:
* Detected through a FPropertyArrayExternalObjectHandler
* Removes the entity and adds a UOverrideRemovePlaceholder instead, with specific metadata to know which entity is being removed
* Modified entities:
* Saved as its own file

Added a new flag in FPropertyTag to know if that property was saved using external objects
The serialization of FPropertyArray has been modified to invoke a registered handler (per-property/class it can be reused for other cases of external assets) upon load/save
* On Load of a package with external entities the dynamic imports callbacks adds the necessary packages to the imports (discovered through the AssetRegistry)
* On Deserialization of the property, the necessary override operations are discovered from the assets/metadata in the AssetRegistry and the overrides are restored
* On Serialization of the property, if we'll be saving in external assets we save no operation/reference it gets rebuilt on load
Overrides management code has been modified to invoke the same handler as the serialization to notify objects of array override being added

Added FPropertyArrayExternalObjectHandler
* Registered per Property/Class so that invocation can occur on the native classes and on IDOs
* Invoked during serialization and overrides creation
* The handler must implement IClassPropertyArrayExternalOverridesHandler, and support being called for IDOs

Added FExternalObjectsOverrideHelper
* A utility class to deal with various operations pertaining to external overrides, like override discovery, external package management, remove placeholders,...

Added a delegate to FCoreUObjectDelegates::OnObjectsReinstanced to apply all missing externalizations on instances after reinstantiation (since they can override the bSaveOwnedEntitiesAsExternalFiles new entities from the template would not be external yet)

Added 2 cvars to control this feature availability
* Entity.OneFilePerPrefabEntityEnabled : Controls whether we allow entities to be externalized in prefabs
* Entity.OneFilePerEntityOnAllInstancesEnabled : Controls whether we allow all owned entities instances to be externalized (ie non-root owned entities in a world)

#Synced-CL 50537189

#Synced-CL 50537189

diff --git a/Engine/Config/BaseGame.ini b/Engine/Config/BaseGame.ini
--- a/Engine/Config/BaseGame.ini
+++ b/Engine/Config/BaseGame.ini
@@ -264,6 +264,8 @@ SubtitleWidgetToUseDefault=/Script/UMG.WidgetBlueprintGeneratedClass'/SubtitlesA
 +ProjectAssetPaths="/Game"
 +DisallowedObjectClassesToLoad="/Script/Engine.Blueprint"
 +DisallowedObjectClassesToLoad="/Script/Engine.AssetImportData"
++DisallowedObjectClassesToLoad="/Script/Entity.BaseEntity"
++DisallowedObjectClassesToLoad="/Script/Engine.OverrideRemovePlaceholder"
 +DisallowedAssetClassesToGather="/Script/Engine.World"
 +DisallowedPathsToGather="/Game/Developers"
 ;+EngineSpecificAssetsToCook="/Engine/Maps/Templates/Template_Default.umap"
diff --git a/Engine/Source/Editor/UnrealEd/Private/Kismet2/KismetReinstanceUtilities.cpp b/Engine/Source/Editor/UnrealEd/Private/Kismet2/KismetReinstanceUtilities.cpp
--- a/Engine/Source/Editor/UnrealEd/Private/Kismet2/KismetReinstanceUtilities.cpp
+++ b/Engine/Source/Editor/UnrealEd/Private/Kismet2/KismetReinstanceUtilities.cpp
@@ -3463,6 +3463,7 @@ void FBlueprintCompileReinstancer::PreCreateSubObjectsForReinstantiation_Inner(
 		UObject* OldSubObjectOuter = OldSubObject->GetOuter();
 		FName SubObjectName = OldSubObject->GetFName();
 		EObjectFlags SubObjectFlags = OldSubObject->GetFlags() & RF_Reinst;
+		UPackage* OldExternalPackage = OldSubObject->GetExternalPackage();
 
 		// Only re-create objects that are created from a default object,
 		// skip any default subobjects are they will be created by their parent object CDO.
@@ -3538,6 +3539,14 @@ void FBlueprintCompileReinstancer::PreCreateSubObjectsForReinstantiation_Inner(
 
 						NewSubObject = NewObject<UObject>(NewUObject, SubObjectClass, SubObjectName, SubObjectFlags, Template);
 					}
+
+					// Transfer external package state if necessary 
+					if (OldExternalPackage)
+					{
+						NewSubObject->SetExternalPackage(OldExternalPackage); // place new subobject in the pre-existing external package 
+						OldSubObject->SetExternalPackage(nullptr); // move old subobject outside of external package (it'll be in the transient package) so that it's doesn't cause the package to stick around
+					}
+
 					CreatedInstanceMap.Add(OldSubObject, NewSubObject);
 					PreCreateSubObjectsForReinstantiation_Inner(OldInstancedSubObjects, OldToNewClassMap, OldSubObject, NewSubObject, CreatedInstanceMap, OldToNewInstanceMap, OrderedListOfObjectToCopy);
 					if (OrderedListOfObjectToCopy)
@@ -3551,6 +3560,13 @@ void FBlueprintCompileReinstancer::PreCreateSubObjectsForReinstantiation_Inner(
 		// There might be new subobjects attached to the sub object that are particular to this instance, let's traverse it to find them out.
 		else if (UObject** NewSubObject = ContainedNewSubObjects.FindByPredicate([SubObjectName](UObject* SubObject) { return SubObject && SubObject->GetFName() == SubObjectName; }))
 		{
+			// Transfer external package state if necessary 
+			if (OldExternalPackage)
+			{
+				(*NewSubObject)->SetExternalPackage(OldExternalPackage); // place new suboject in the pre-existing external package 
+				OldSubObject->SetExternalPackage(nullptr); // move old subobject outside of external package (it'll be in the transient package) so that it's doesn't cause the package to stick around
+			}
+
 			PreCreateSubObjectsForReinstantiation_Inner(OldInstancedSubObjects, OldToNewClassMap, OldSubObject, *NewSubObject, CreatedInstanceMap, OldToNewInstanceMap, OrderedListOfObjectToCopy);
 		}
 	}
diff --git a/Engine/Source/Runtime/CoreUObject/Private/UObject/OverriddenObjectsExternalHandler.cpp b/Engine/Source/Runtime/CoreUObject/Private/UObject/OverriddenObjectsExternalHandler.cpp
new file mode 100644
--- /dev/null
+++ b/Engine/Source/Runtime/CoreUObject/Private/UObject/OverriddenObjectsExternalHandler.cpp
@@ -0,0 +1,118 @@
+// Copyright Epic Games, Inc. All Rights Reserved.
+
+#include "UObject/OverriddenObjectsExternalHandler.h"
+#include "UObject/OverriddenPropertySet.h"
+#include "UObject/InstanceDataObjectUtils.h"
+#include "UObject/PropertyBagRepository.h"
+
+#if WITH_EDITOR
+
+uint32 GetTypeHash(const FPropertyArrayExternalObjectHandler::FPropertyKey& Key)
+{
+	uint32 Hash = GetTypeHash(Key.ClassName);
+	Hash = HashCombine(Hash, GetTypeHash(Key.PropertyName));
+	return Hash;
+}
+
+FPropertyArrayExternalObjectHandler::FPropertyKey::FPropertyKey(FName InClassName, FName InPropertyName) :
+	ClassName(InClassName), PropertyName(InPropertyName)
+{
+}
+
+FPropertyArrayExternalObjectHandler::FPropertyKey::FPropertyKey(UObject* Object, const FProperty* Property)
+{
+	PropertyName = Property->GetFName();
+	UClass* Class = Cast<UClass>(Property->Owner.ToUObject());
+
+	if (UE::IsClassOfInstanceDataObjectClass(Object->GetClass()))
+	{
+		// mirrored from InstanceDataObjectUtils.cpp
+		static const FName NAME_TemplateStructName(ANSITEXTVIEW("TemplateStructName"));
+
+		// use the ClassName the IDO class was created from 
+		ClassName = FName(Class->GetMetaData(NAME_TemplateStructName));
+	}
+	else
+	{
+		ClassName = Class->GetFName();
+	}
+}
+
+
+bool FPropertyArrayExternalObjectHandler::GetExternalOverrides(UObject* Object, const FProperty* Property, FOverridableObjectArrayOverrides* Overrides)
+{
+	FPropertyKey Key(Object, Property);
+	
+	if (IClassPropertyArrayExternalOverridesHandler** Handler = PerClassPropertyHandlers.Find(Key))
+	{
+		return (*Handler)->GetExternalOverrides(Object, Overrides);
+	}
+
+	return false;
+}
+
+
+bool FPropertyArrayExternalObjectHandler::IsSavingExternal(UObject* Object, const  FProperty* Property)
+{
+	FPropertyKey Key(Object, Property);
+
+	if (IClassPropertyArrayExternalOverridesHandler** Handler = PerClassPropertyHandlers.Find(Key))
+	{
+		return (*Handler)->IsSavingExternal(Object);
+	}
+
+	return false;
+}
+
+bool FPropertyArrayExternalObjectHandler::NotifyRemove(UObject* Object, const class FProperty* Property, UObject* RemovedObject, UObject* RemovedArchetype)
+{
+	FPropertyKey Key(Object, Property);
+	if (IClassPropertyArrayExternalOverridesHandler** Handler = PerClassPropertyHandlers.Find(Key))
+	{
+		(*Handler)->NotifyRemove(Object, RemovedObject, RemovedArchetype);
+		return true;
+	}
+
+	return false;
+}
+
+bool FPropertyArrayExternalObjectHandler::NotifyAdd(UObject* Object, const class FProperty* Property, UObject* AddedObject)
+{
+	FPropertyKey Key(Object, Property);
+	if (IClassPropertyArrayExternalOverridesHandler** Handler = PerClassPropertyHandlers.Find(Key))
+	{
+		(*Handler)->NotifyAdd(Object, AddedObject);
+		return true;
+	}
+
+	return false;
+}
+
+void FPropertyArrayExternalObjectHandler::RegisterHandler(UClass* Class, FName PropertyName, IClassPropertyArrayExternalOverridesHandler* Handler)
+{
+	FPropertyKey Key(Class->GetFName(), PropertyName);
+	
+	checkf(!PerClassPropertyHandlers.Find(Key), TEXT("Attempting double registration of an IClassPropertyArrayExternalOverridesHandler, Class: %s, Property: %s"), *Class->GetName(), *PropertyName.ToString());
+	PerClassPropertyHandlers.Add(Key, Handler);
+}
+
+void FPropertyArrayExternalObjectHandler::UnregisterHandler(IClassPropertyArrayExternalOverridesHandler* Handler)
+{	
+	for (TPair<FPropertyKey, IClassPropertyArrayExternalOverridesHandler*> Pair : PerClassPropertyHandlers)
+	{
+		if(Pair.Value == Handler)
+		{
+			PerClassPropertyHandlers.Remove(Pair.Key);
+			break;
+		}
+	}
+}
+
+FPropertyArrayExternalObjectHandler FPropertyArrayExternalObjectHandler::HandlerObj;
+
+FPropertyArrayExternalObjectHandler& FPropertyArrayExternalObjectHandler::Get()
+{
+	return HandlerObj;
+}
+
+#endif //#if WITH_EDITOR
\ No newline at end of file
diff --git a/Engine/Source/Runtime/CoreUObject/Private/UObject/OverriddenPropertySet.cpp b/Engine/Source/Runtime/CoreUObject/Private/UObject/OverriddenPropertySet.cpp
--- a/Engine/Source/Runtime/CoreUObject/Private/UObject/OverriddenPropertySet.cpp
+++ b/Engine/Source/Runtime/CoreUObject/Private/UObject/OverriddenPropertySet.cpp
@@ -11,6 +11,7 @@
 #include "UObject/PropertyBagRepository.h"
 #include "UObject/UObjectArchetypeHelper.h"
 #include "UObject/UObjectThreadContext.h"
+#include "UObject/OverriddenObjectsExternalHandler.h"
 
 #include UE_INLINE_GENERATED_CPP_BY_NAME(OverriddenPropertySet)
 
@@ -1185,6 +1186,9 @@ void FOverriddenPropertySet::NotifyPropertyChange(FOverriddenPropertyNode* Paren
 							{
 								AddedSubObjectOverriddenProperties->bWasAdded = true;
 							}
+#if WITH_EDITOR
+							FPropertyArrayExternalObjectHandler::Get().NotifyAdd(Owner, Property, AddedSubObject);
+#endif
 						}
 					}
 				};
@@ -1214,6 +1218,10 @@ void FOverriddenPropertySet::NotifyPropertyChange(FOverriddenPropertyNode* Paren
 								{
 									RemovedSubObjectNode.SetOperation(EOverriddenPropertyOperation::Remove);
 								}
+
+#if WITH_EDITOR
+								FPropertyArrayExternalObjectHandler::Get().NotifyRemove(Owner, Property, RemovedSubObject, RemovedSubObjectArchetype);
+#endif
 							}
 							else
 							{
@@ -1231,6 +1239,10 @@ void FOverriddenPropertySet::NotifyPropertyChange(FOverriddenPropertyNode* Paren
 									{
 										bNeedsCleanup = true;
 									}
+
+#if WITH_EDITOR
+									FPropertyArrayExternalObjectHandler::Get().NotifyRemove(Owner, Property, RemovedSubObject, nullptr);
+#endif
 								}
 								else
 								{
@@ -1746,6 +1758,26 @@ void FOverriddenPropertySet::OverrideProperty(FPropertyVisitorPath::Iterator Pro
 	NotifyPropertyChange(&RootNode, EPropertyNotificationType::PostEdit, PropertyIterator, EPropertyChangeType::Unspecified, Data, bNeedsCleanup);
 }
 
+void FOverriddenPropertySet::SetExternalizedOverriddenProperty(const FOverriddenPropertyNodeID& NodeID, bool bExternal)
+{	
+	FOverriddenPropertyNode& Node = RootNode.FindOrAddNode(NodeID);
+	FOverriddenPropertyNodeID ExternalizedID;
+
+	if (bExternal)
+	{
+		FOverriddenPropertyNode& ExternalizeNode = Node.FindOrAddNode(ExternalizedID);
+		ExternalizeNode.SetOperation(EOverriddenPropertyOperation::Externalized);
+	}
+	else
+	{
+		Node.RemoveSubPropertyNode(ExternalizedID);
+		if (Node.GetSubPropertyNodes().Num() == 0)
+		{
+			Node.SetOperation(EOverriddenPropertyOperation::None);
+		}
+	}
+}
+
 void FOverriddenPropertySet::NotifyPropertyChange(const EPropertyNotificationType Notification, FPropertyVisitorPath::Iterator PropertyIterator, const EPropertyChangeType::Type ChangeType, const void* Data)
 {
 	bool bNeedsCleanup = false;
@@ -1987,6 +2019,11 @@ const FOverriddenPropertyNode* FOverriddenPropertySet::GetOverriddenPropertyNode
 	return OverriddenPropertyNode;
 }
 
+FOverridableObjectArrayOverrides::FOverridableObjectArrayOverrides() :
+	ObjectsOperation(EOverriddenPropertyOperation::None)
+{
+}
+
 FScopedPropertyValue::FScopedPropertyValue(FProperty* InProperty)
 	: Property(InProperty)
 	, Data(FMemory::Malloc(InProperty->GetSize(), InProperty->GetMinAlignment()))
diff --git a/Engine/Source/Runtime/CoreUObject/Private/UObject/Property.cpp b/Engine/Source/Runtime/CoreUObject/Private/UObject/Property.cpp
--- a/Engine/Source/Runtime/CoreUObject/Private/UObject/Property.cpp
+++ b/Engine/Source/Runtime/CoreUObject/Private/UObject/Property.cpp
@@ -2029,7 +2029,22 @@ const TCHAR* FProperty::ImportSingleProperty( const TCHAR* Str, void* DestData,
 							ArrayHelper.ExpandForIndex(Index);
 							break;
 						}
+						case EOverriddenPropertyOperation::Externalized:
+						{
+							checkf(ArrayProperty->Inner->HasAnyPropertyFlags(CPF_PersistentInstance), TEXT("Only instanced sub object is supporting externalized operation"));
+							const FObjectPropertyBase* InnerObjectProperty = CastFieldChecked<FObjectPropertyBase>(ArrayProperty->Inner);
 
+							if (Path)
+							{
+								FArchiveSerializedPropertyChain Chain = Path->ToSerializedPropertyChain();
+								if (FOverriddenPropertyNode* ArrayOverriddenPropertyNode = OverriddenProperties ? OverriddenProperties->RestoreOverriddenPropertyOperation(EOverriddenPropertyOperation::Modified, &Chain, /*Property*/nullptr) : nullptr)
+								{
+									OverriddenProperties->RestoreSubPropertyOperation(EOverriddenPropertyOperation::Externalized, *ArrayOverriddenPropertyNode, FOverriddenPropertyNodeID());
+								}
+							}
+
+							break;
+						}
 						default:
 							checkf(false, TEXT("Unsupported array operation while importing text"));
 							return nullptr;
diff --git a/Engine/Source/Runtime/CoreUObject/Private/UObject/PropertyArray.cpp b/Engine/Source/Runtime/CoreUObject/Private/UObject/PropertyArray.cpp
--- a/Engine/Source/Runtime/CoreUObject/Private/UObject/PropertyArray.cpp
+++ b/Engine/Source/Runtime/CoreUObject/Private/UObject/PropertyArray.cpp
@@ -7,6 +7,7 @@
 #include "UObject/LinkerLoad.h"
 #include "UObject/ObjectInstancingGraph.h"
 #include "UObject/OverriddenPropertySet.h"
+#include "UObject/OverriddenObjectsExternalHandler.h"
 #include "UObject/PropertyHelper.h"
 #include "UObject/UnrealTypePrivate.h"
 #include "UObject/UObjectThreadContext.h"
@@ -239,10 +240,12 @@ void FArrayProperty::SerializeItem(FStructuredArchive::FSlot Slot, TNotNull<void
 		Inner->SerializeItem(Slot, Item);
 	};
 
+	const FPropertyTag* PropertyTag = FPropertyTagScope::GetCurrentPropertyTag();
+
 	// Make sure the container is reloading accordingly to the value set in the property tag if any
-	if (UnderlyingArchive.IsLoading() && FPropertyTagScope::GetCurrentPropertyTag())
+	if (UnderlyingArchive.IsLoading() && PropertyTag)
 	{
-		bExperimentalOverridableLogic = FPropertyTagScope::GetCurrentPropertyTag()->bExperimentalOverridableLogic;
+		bExperimentalOverridableLogic = PropertyTag->bExperimentalOverridableLogic;
 	}
 
 	// *** Experimental *** Special serialization path for array with overridable serialization
@@ -250,25 +253,79 @@ void FArrayProperty::SerializeItem(FStructuredArchive::FSlot Slot, TNotNull<void
 	{
 		checkf(!UnderlyingArchive.ArUseCustomPropertyList, TEXT("Using custom property list is not supported by overridable serialization"));
 
+		// Container for temporarily tracking some indices & communicating with external objects handler
+		FOverridableObjectArrayOverrides	ContainerOverrides;
+
+		const FObjectPropertyBase* InnerObjectProperty = CastField<FObjectPropertyBase>(Inner);
+		
+		// this indicates whether or not we can use external objects overrides in this particular serializing situation
+		const bool bCouldUseExternalObjects = InnerObjectProperty && HasAnyPropertyFlags(CPF_ExperimentalExternalObjects) && UnderlyingArchive.IsPersistent() && !UnderlyingArchive.HasAnyPortFlags(PPF_DuplicateForPIE);
+
+		// this indicates whether or not this object is to use external objects for overrides (can be overriden in certain deprecation cases)
+		// this flag will either loaded from the serialized tag, or calculated during save
+		const bool bTaggedUsingExternalsObjects = PropertyTag ? PropertyTag->bExperimentalExternalObjects.Get(false) : false;
+		
 		FStructuredArchive::FRecord Record = Slot.EnterRecord();
 		if (UnderlyingArchive.IsLoading())
 		{
+#if WITH_EDITOR
+			bool bUsingExternalObjects = bTaggedUsingExternalsObjects && bCouldUseExternalObjects;
+
+			if (bCouldUseExternalObjects)
+			{
+				// Final determination for bUsingExternalObjects is delegated to handler if the property
+				// was saved before started to include this state in the tag
+				bool bUseHandlerForUsingExternalObjects = false;
+				if (PropertyTag && !PropertyTag->bExperimentalExternalObjects.IsSet())
+				{
+					bUseHandlerForUsingExternalObjects = true;
+				}
+
+				// Invoke the external object handler to obtain the overrides for this property and the 'usingExternalObjects' state if we get it from there
+				check(Context->SerializedObject);
+				if ((bTaggedUsingExternalsObjects || bUseHandlerForUsingExternalObjects))
+				{					
+					bool bHandlerUsingExternalObjects = FPropertyArrayExternalObjectHandler::Get().GetExternalOverrides(Context->SerializedObject, this, &ContainerOverrides);
+					if (bUseHandlerForUsingExternalObjects)
+					{
+						bUsingExternalObjects = bHandlerUsingExternalObjects;
+					}
+					else
+					{
+						// in case we're keeping this value from the tag validate all everything is as expected
+						check(bHandlerUsingExternalObjects == bUsingExternalObjects);
+					}
+				}
+			}
+#else			 
+			check(!bTaggedUsingExternalsObjects);
+			const bool bUsingExternalObjects = false;
+#endif // WITH_EDITOR
+
 			int32 NumReplaced = 0;
 			FStructuredArchive::FArray ReplacedArray = Record.EnterArray(TEXT("Replaced"), NumReplaced);
 			if (NumReplaced != INDEX_NONE)
 			{
-				FSerializedPropertyScope SerializedProperty(UnderlyingArchive, Inner, this);
-
-				ArrayHelper.EmptyAndAddValues(NumReplaced);
-				for (int32 i = 0; i < NumReplaced; i++)
+				if (bUsingExternalObjects)
 				{
-					SerializeContainerItem(ReplacedArray.EnterElement(), ArrayHelper.GetRawPtr(i));
+					check(ContainerOverrides.ObjectsOperation == EOverriddenPropertyOperation::Replace);
+					check(NumReplaced == 0);
+				}
+				else
+				{
+					FSerializedPropertyScope SerializedProperty(UnderlyingArchive, Inner, this);
+
+					ArrayHelper.EmptyAndAddValues(NumReplaced);
+				
+					for (int32 i = 0; i < NumReplaced; i++)
+					{
+						SerializeContainerItem(ReplacedArray.EnterElement(), ArrayHelper.GetRawPtr(i));
+					}
 				}
 			}
 			else
 			{
 				// Only Array of Instanced subobject are handled here as sort of a set where the matching key is done using the archetype
-				const FObjectPropertyBase* InnerObjectProperty = CastFieldChecked<FObjectPropertyBase>(Inner);
 				checkf(InnerObjectProperty->HasAnyPropertyFlags(CPF_PersistentInstance), TEXT("Only supported code path here is the instanced subobjects"));
 
 				FOverriddenPropertySet* OverriddenProperties = FOverridableSerializationLogic::GetOverriddenProperties();
@@ -279,7 +336,7 @@ void FArrayProperty::SerializeItem(FStructuredArchive::FSlot Slot, TNotNull<void
 
 				// If the arrays are not the same size, the code will not be able to remove items correctly
 				// as it searches for the items to remove in the default array but removes them from this array.
-				checkf(DefaultsArrayHelper.Num() == ArrayHelper.Num(), TEXT("The array \'%s\' is expected to be the same size as it default for removal logic."), *GetName());
+				checkf (bUsingExternalObjects || (DefaultsArrayHelper.Num() == ArrayHelper.Num()), TEXT("The array \'%s\' is expected to be the same size as it default for removal logic."), *GetName());
 
 				auto FindObject = [InnerObjectProperty](UObject* Object, UObject* Object2, FScriptArrayHelper& ArrayHelper) -> int32
 				{
@@ -335,12 +392,7 @@ void FArrayProperty::SerializeItem(FStructuredArchive::FSlot Slot, TNotNull<void
 								UE_LOG(LogOverridableObject, VeryVerbose, TEXT("Unable to load removed item %s(0x%p)"), *GetNameSafe(RemovedSubObject), RemovedSubObject);
 							}
 
-							// Need to fetch the ArrayOverriddenPropertyNode every loop as the previous iteration might have reallocated the node.
-							if (FOverriddenPropertyNode* ArrayOverriddenPropertyNode = OverriddenProperties ? OverriddenProperties->RestoreOverriddenPropertyOperation(EOverriddenPropertyOperation::Modified, UnderlyingArchive.GetSerializedPropertyChain(), /*Property*/nullptr) : nullptr)
-							{
-								// Rebuild the overridden info
-								OverriddenProperties->RestoreSubObjectOperation(EOverriddenPropertyOperation::Remove, *ArrayOverriddenPropertyNode, RemovedSubObject);
-							}
+							ContainerOverrides.RemovedArchetypes.Add(RemovedSubObject);
 						}
 					}
 
@@ -416,7 +468,7 @@ void FArrayProperty::SerializeItem(FStructuredArchive::FSlot Slot, TNotNull<void
 									InnerObjectProperty->SetObjectPropertyValue(ArrayHelper.GetRawPtr(Index), ShadowedObject);
 								}
 							}
-#endif //  WITH_EDITORONLY_DATA
+#endif // WITH_EDITORONLY_DATA
 
 						}
 					}
@@ -451,38 +503,62 @@ void FArrayProperty::SerializeItem(FStructuredArchive::FSlot Slot, TNotNull<void
 							}
 							InnerObjectProperty->SetObjectPropertyValue(ArrayHelper.GetRawPtr(Index), AddedSubObject);
 
-							// Need to fetch the ArrayOverriddenPropertyNode every loop as the previous iteration might have reallocated the node.
-							if (FOverriddenPropertyNode* ArrayOverriddenPropertyNode = OverriddenProperties ? OverriddenProperties->RestoreOverriddenPropertyOperation(EOverriddenPropertyOperation::Modified, UnderlyingArchive.GetSerializedPropertyChain(), /*Property*/nullptr) : nullptr)
-							{
-								// Rebuild the overridden info
-								OverriddenProperties->RestoreSubObjectOperation(EOverriddenPropertyOperation::Add, *ArrayOverriddenPropertyNode, AddedSubObject);
-							}
+							ContainerOverrides.Added.Add(AddedSubObject);
+						}
+					}
+				}
+
+				// If we were tagged as using external objects restore the externalize override state even if we didn't
+				// use it for serialization. 
+				// 
+				// For example, objects being reinstanced will pass through a non persistent archive during CopyPropertiesForUnrelatedObjects
+				// and this will lose the override state if we don't restore it here. 
+				if (bTaggedUsingExternalsObjects)
+				{
+					if (OverriddenProperties)
+					{
+						if (FOverriddenPropertyNode* ArrayOverriddenPropertyNode = OverriddenProperties ? OverriddenProperties->RestoreOverriddenPropertyOperation(EOverriddenPropertyOperation::Modified, UnderlyingArchive.GetSerializedPropertyChain(), /*Property*/nullptr) : nullptr)
+						{
+							OverriddenProperties->RestoreSubPropertyOperation(EOverriddenPropertyOperation::Externalized, *ArrayOverriddenPropertyNode, FOverriddenPropertyNodeID());
 						}
 					}
 				}
+			
+				// Restore overrides
+				for (int32 i = 0; i < ContainerOverrides.RemovedArchetypes.Num(); i++)
+				{
+					// Need to fetch the ArrayOverriddenPropertyNode every loop as the previous iteration might have reallocated the node.
+					if (FOverriddenPropertyNode* ArrayOverriddenPropertyNode = OverriddenProperties ? OverriddenProperties->RestoreOverriddenPropertyOperation(EOverriddenPropertyOperation::Modified, UnderlyingArchive.GetSerializedPropertyChain(), /*Property*/nullptr) : nullptr)
+					{
+						// Rebuild the overridden info
+						OverriddenProperties->RestoreSubObjectOperation(EOverriddenPropertyOperation::Remove, *ArrayOverriddenPropertyNode, ContainerOverrides.RemovedArchetypes[i]);
+					}
+				}
+				for (int32 i = 0; i < ContainerOverrides.Added.Num(); i++)
+				{
+					// Need to fetch the ArrayOverriddenPropertyNode every loop as the previous iteration might have reallocated the node.
+					if (FOverriddenPropertyNode* ArrayOverriddenPropertyNode = OverriddenProperties ? OverriddenProperties->RestoreOverriddenPropertyOperation(EOverriddenPropertyOperation::Modified, UnderlyingArchive.GetSerializedPropertyChain(), /*Property*/nullptr) : nullptr)
+					{
+						// Rebuild the overridden info
+						OverriddenProperties->RestoreSubObjectOperation(EOverriddenPropertyOperation::Add, *ArrayOverriddenPropertyNode, ContainerOverrides.Added[i]);
+					}
+				}
+
 			}
 		}
 		else
 		{
+			bool bUsingExternalObjects = bTaggedUsingExternalsObjects && bCouldUseExternalObjects;
+
 			// Container for temporarily tracking some indices
 			TArray<int32> RemovedIndices;
 			TArray<int32> AddedIndices;
 			TArray<int32> ShadowIndices;
 
-			// We filter out external objects if:
-			//	- We're an array of objects.
-			//	- We have the ExperimentalExternalObjects flag.
-			//	- We're writing to a persistent archive.
-			//	- We have at least one external entity object.
-			// For now, this will be considered as a full array replacement.
-			FObjectProperty* ObjectProperty = CastField<FObjectProperty>(Inner);
-			const bool bFilterExternalObjects = ObjectProperty && HasAnyPropertyFlags(CPF_ExperimentalExternalObjects) && UnderlyingArchive.IsPersistent() && !UnderlyingArchive.HasAnyPortFlags(PPF_DuplicateForPIE);
-
 			bool bReplaceArray = !Defaults || !UnderlyingArchive.DoDelta() || UnderlyingArchive.IsTransacting();
 
 			if (!bReplaceArray)
 			{
-				const FObjectPropertyBase* InnerObjectProperty = CastField<FObjectPropertyBase>(Inner);
 				EOverriddenPropertyOperation ArrayOverrideOp = EOverriddenPropertyOperation::None;
 				FOverriddenPropertySet* OverriddenProperties = FOverridableSerializationLogic::GetOverriddenProperties();
 				if (OverriddenProperties)
@@ -498,6 +574,8 @@ void FArrayProperty::SerializeItem(FStructuredArchive::FSlot Slot, TNotNull<void
 
 				if (!bReplaceArray)
 				{
+					FScriptArrayHelper DefaultsArrayHelper(this, Defaults);
+
 					if (FOverridableSerializationLogic::ShouldPropertyShadowSerializeSubObject(this))
 					{
 						const int32 ArrayNum = ArrayHelper.Num();
@@ -512,7 +590,7 @@ void FArrayProperty::SerializeItem(FStructuredArchive::FSlot Slot, TNotNull<void
 
 					if (OverriddenProperties && ArrayOverrideOp != EOverriddenPropertyOperation::None)
 					{
-						auto FindObject = [InnerObjectProperty, bFilterExternalObjects](const FOverriddenPropertyNodeID& ObjectToFind, FScriptArrayHelper& ArrayHelper) -> int32
+						auto FindObject = [InnerObjectProperty](const FOverriddenPropertyNodeID& ObjectToFind, FScriptArrayHelper& ArrayHelper) -> int32
 						{
 							const int32 ArrayNum = ArrayHelper.Num();
 							for (int i = 0; i < ArrayNum; ++i)
@@ -529,8 +607,6 @@ void FArrayProperty::SerializeItem(FStructuredArchive::FSlot Slot, TNotNull<void
 							return INDEX_NONE;
 						};
 
-						FScriptArrayHelper DefaultsArrayHelper(this, Defaults);
-
 						if (const FOverriddenPropertyNode* ArrayOverriddenPropertyNode = OverriddenProperties->GetOverriddenPropertyNode(UnderlyingArchive.GetSerializedPropertyChain()))
 						{
 							for (const FOverriddenPropertyNode& SubNode : ArrayOverriddenPropertyNode->GetSubPropertyNodes())
@@ -556,11 +632,7 @@ void FArrayProperty::SerializeItem(FStructuredArchive::FSlot Slot, TNotNull<void
  										if (Index != INDEX_NONE)
 										{
 											ShadowIndices.Remove(Index);
-
-											if (!bFilterExternalObjects || !InnerObjectProperty->GetObjectPropertyValue(ArrayHelper.GetElementPtr(Index))->IsPackageExternal())
-											{
-												AddedIndices.Add(Index);												
-											}
+											AddedIndices.Add(Index);
 										}
 										else
 										{
@@ -568,6 +640,11 @@ void FArrayProperty::SerializeItem(FStructuredArchive::FSlot Slot, TNotNull<void
 										}
 										break;
 									}
+								case EOverriddenPropertyOperation::Externalized:
+									{
+										// nothing to do here
+										break;
+									}
 								default:
 									checkf(false, TEXT("Unsupported operation type"));
 									break;
@@ -576,51 +653,78 @@ void FArrayProperty::SerializeItem(FStructuredArchive::FSlot Slot, TNotNull<void
 						}
 					}
 				}
+
+				ContainerOverrides.ObjectsOperation = ArrayOverrideOp;
 			}
 
 			if (bReplaceArray)
 			{
-				int32 NumReplaced = ArrayHelper.Num();
-				FStructuredArchive::FArray ReplacedArray = Record.EnterArray(TEXT("Replaced"), NumReplaced);
-
-				for (int32 i = 0; i < NumReplaced; ++i)
+				if (bUsingExternalObjects)
+				{
+					// External objects always acquire their size on reload
+					int32 NumReplaced = 0;
+					FStructuredArchive::FArray ReplacedArray = Record.EnterArray(TEXT("Replaced"), NumReplaced);
+				}
+				else
 				{
-					SerializeContainerItem(ReplacedArray.EnterElement(), ArrayHelper.GetRawPtr(i));
+					int32 NumReplaced = ArrayHelper.Num();
+					FStructuredArchive::FArray ReplacedArray = Record.EnterArray(TEXT("Replaced"), NumReplaced);
+
+					for (int32 i = 0; i < NumReplaced; ++i)
+					{
+						SerializeContainerItem(ReplacedArray.EnterElement(), ArrayHelper.GetRawPtr(i));
+					}
 				}
 			}
 			else
 			{
-				int32 NumReplaced = INDEX_NONE;
-				Record.EnterArray(TEXT("Replaced"), NumReplaced);
-
-				checkf(Defaults, TEXT("Expecting overridable serialization to have defaults to compare to"));
 				FScriptArrayHelper DefaultsArrayHelper(this, Defaults);
 
-				int32 NumRemoved = RemovedIndices.Num();
-				FStructuredArchive::FArray RemovedArray = Record.EnterArray(TEXT("Removed"), NumRemoved);
-				for (int32 i : RemovedIndices)
+#if WITH_EDITOR		
+				if (bUsingExternalObjects)
 				{
-					SerializeContainerItem(RemovedArray.EnterElement(), DefaultsArrayHelper.GetRawPtr(i));
+					// These are stored as external objects so we let the serialized object handle the storage 
+					// of the override info.
+					int32 NoOpReplaced = INDEX_NONE;
+					int32 NoOp = 0;
+					Record.EnterArray(TEXT("Replaced"), NoOpReplaced);
+					Record.EnterArray(TEXT("Removed"), NoOp);
+					Record.EnterArray(TEXT("Modified"), NoOp);
+					Record.EnterArray(TEXT("Shadowed"), NoOp);
+					Record.EnterArray(TEXT("Added"), NoOp);
 				}
+				else
+#endif // WITH_EDITOR
+				{					
+					int32 NumReplaced = INDEX_NONE;
+					Record.EnterArray(TEXT("Replaced"), NumReplaced);
 
-				// We should not have modified elements anymore but keeping for eventual support of modified values like struct
-				int32 NumModified = 0;
-				FStructuredArchive::FArray ModifiedArray = Record.EnterArray(TEXT("Modified"), NumModified);
+					int32 NumRemoved = RemovedIndices.Num();
+					FStructuredArchive::FArray RemovedArray = Record.EnterArray(TEXT("Removed"), NumRemoved);
+					for (int32 i : RemovedIndices)
+					{
+						SerializeContainerItem(RemovedArray.EnterElement(), DefaultsArrayHelper.GetRawPtr(i));
+					}
 
-				// Support of subobject shadowed serialization
-				// Introduced from EUnrealEngineObjectUE5Version::OS_SUB_OBJECT_SHADOW_SERIALIZATION
-				int32 NumShadowed = ShadowIndices.Num();
-				FStructuredArchive::FArray ShadowedArray = Record.EnterArray(TEXT("Shadowed"), NumShadowed);
-				for (int32 i : ShadowIndices)
-				{
-					SerializeContainerItem(ShadowedArray.EnterElement(), ArrayHelper.GetRawPtr(i));
-				}
+					// We should not have modified elements anymore but keeping for eventual support of modified values like struct
+					int32 NumModified = 0;
+					FStructuredArchive::FArray ModifiedArray = Record.EnterArray(TEXT("Modified"), NumModified);
 
-				int32 NumAdded = AddedIndices.Num();
-				FStructuredArchive::FArray AddedArray = Record.EnterArray(TEXT("Added"), NumAdded);
-				for (int32 i : AddedIndices)
-				{
-					SerializeContainerItem(AddedArray.EnterElement(), ArrayHelper.GetRawPtr(i));
+					// Support of subobject shadowed serialization
+					// Introduced from EUnrealEngineObjectUE5Version::OS_SUB_OBJECT_SHADOW_SERIALIZATION
+					int32 NumShadowed = ShadowIndices.Num();
+					FStructuredArchive::FArray ShadowedArray = Record.EnterArray(TEXT("Shadowed"), NumShadowed);
+					for (int32 i : ShadowIndices)
+					{
+						SerializeContainerItem(ShadowedArray.EnterElement(), ArrayHelper.GetRawPtr(i));
+					}
+
+					int32 NumAdded = AddedIndices.Num();
+					FStructuredArchive::FArray AddedArray = Record.EnterArray(TEXT("Added"), NumAdded);
+					for (int32 i : AddedIndices)
+					{
+						SerializeContainerItem(AddedArray.EnterElement(), ArrayHelper.GetRawPtr(i));
+					}
 				}
 			}
 		}
diff --git a/Engine/Source/Runtime/CoreUObject/Private/UObject/PropertyTag.cpp b/Engine/Source/Runtime/CoreUObject/Private/UObject/PropertyTag.cpp
--- a/Engine/Source/Runtime/CoreUObject/Private/UObject/PropertyTag.cpp
+++ b/Engine/Source/Runtime/CoreUObject/Private/UObject/PropertyTag.cpp
@@ -10,8 +10,10 @@
 #include "UObject/BlueprintsObjectVersion.h"
 #include "UObject/DebugSerializationFlags.h"
 #include "UObject/OverriddenPropertySet.h"
+#include "UObject/OverriddenObjectsExternalHandler.h"
 #include "UObject/Package.h"
 #include "UObject/UnrealType.h"
+#include "UObject/UObjectThreadContext.h"
 
 #ifndef UE_ENABLE_SERIALIZE_TAGGED_PROPERTY_TRACING
 #define UE_ENABLE_SERIALIZE_TAGGED_PROPERTY_TRACING 0
@@ -47,6 +49,7 @@ enum class EPropertyTagExtension : uint8
 	////////////////////////////////////////////////
 	// First extension group
 	OverridableInformation		= 0x02,
+	HasExternalsObjects			= 0x04,
 
 	//
 	// Add more extensions for the first group here
@@ -155,6 +158,20 @@ static EPropertyTagExtension CalculatePropertyExtensionFlags(FArchive& Underlyin
 		{
 			Tag.bExperimentalOverridableLogic = false;
 		}
+		
+		Tag.bExperimentalExternalObjects = false;
+#if WITH_EDITOR
+		if (Tag.GetProperty()->HasAllPropertyFlags(CPF_ExperimentalOverridableLogic|CPF_ExperimentalExternalObjects) && !UnderlyingArchive.IsCooking())
+		{
+			FUObjectSerializeContext* Context = FUObjectThreadContext::Get().GetSerializeContext();
+			if (Context->SerializedObject)
+			{
+				Tag.bExperimentalExternalObjects = FPropertyArrayExternalObjectHandler::Get().IsSavingExternal(Context->SerializedObject, Tag.GetProperty());
+			}
+			
+			PropertyTagExtensions |= EPropertyTagExtension::HasExternalsObjects;
+		}		
+#endif // WITH_EDITOR
 	}
 
 	return PropertyTagExtensions;
@@ -178,6 +195,13 @@ static void SerializePropertyExtensions(FStructuredArchive::FSlot Slot, EPropert
 		Slot << SA_ATTRIBUTE(TEXT("OverriddenPropertyOperation"), Tag.OverrideOperation);
 		Slot << SA_ATTRIBUTE(TEXT("ExperimentalOverridableLogic"), Tag.bExperimentalOverridableLogic);
 	}
+
+	if (EnumHasAnyFlags(PropertyTagExtensions, EPropertyTagExtension::HasExternalsObjects))
+	{
+		bool bValue = Tag.bExperimentalExternalObjects.Get(false);
+		Slot << SA_ATTRIBUTE(TEXT("ExperimentalExternalObjects"), bValue);
+		Tag.bExperimentalExternalObjects = bValue;
+	}
 }
 
 static void ParsePathName(FStringView Path, UE::FPropertyTypeNameBuilder& Builder)
diff --git a/Engine/Source/Runtime/CoreUObject/Public/UObject/OverriddenObjectsExternalHandler.h b/Engine/Source/Runtime/CoreUObject/Public/UObject/OverriddenObjectsExternalHandler.h
new file mode 100644
--- /dev/null
+++ b/Engine/Source/Runtime/CoreUObject/Public/UObject/OverriddenObjectsExternalHandler.h
@@ -0,0 +1,58 @@
+// Copyright Epic Games, Inc. All Rights Reserved.
+
+#pragma once
+
+#include "UObject/Object.h"
+#include "UObject/ObjectMacros.h"
+
+#include "UObject/OverriddenPropertySet.h"
+
+#if WITH_EDITOR
+class IClassPropertyArrayExternalOverridesHandler
+{
+public: 
+
+	virtual bool GetExternalOverrides(UObject* Object, FOverridableObjectArrayOverrides* Overrides) = 0;
+	virtual bool IsSavingExternal(UObject* Object) = 0;
+
+	virtual void NotifyAdd(UObject* InObject, UObject* AddedObject) = 0;
+	virtual void NotifyRemove(UObject* InObject, UObject* RemovedObject, UObject* RemovedObjectArchetype) = 0;
+};
+
+struct FPropertyArrayExternalObjectHandler
+{	
+	public:
+
+		// Native class to 
+		struct FPropertyKey 
+		{
+			FName ClassName;
+			FName PropertyName;
+
+			FPropertyKey(UObject* Object, const class FProperty* Property);
+			FPropertyKey(FName ClassName, FName PropertyName);
+
+			bool operator== (const FPropertyKey&) const = default;
+
+			friend uint32 GetTypeHash(const FPropertyKey& Key);
+		};
+		
+
+		COREUOBJECT_API static FPropertyArrayExternalObjectHandler& Get();
+
+		COREUOBJECT_API bool NotifyRemove(UObject* Object, const class FProperty* Property, UObject* RemovedObject, UObject* RemovedArchetype);
+		COREUOBJECT_API bool NotifyAdd(UObject* Object, const class FProperty* Property, UObject* AddedObject);		
+
+		COREUOBJECT_API bool GetExternalOverrides(UObject* Object, const class FProperty* Property, FOverridableObjectArrayOverrides* Overrides);
+		COREUOBJECT_API bool IsSavingExternal(UObject* Object, const class FProperty* Property);
+		
+		COREUOBJECT_API void RegisterHandler(UClass* Class, FName PropertyName, IClassPropertyArrayExternalOverridesHandler* Handler);
+		COREUOBJECT_API void UnregisterHandler(IClassPropertyArrayExternalOverridesHandler* Handler);
+
+	private:
+
+		TMap<FPropertyKey, IClassPropertyArrayExternalOverridesHandler*>	PerClassPropertyHandlers;
+		static FPropertyArrayExternalObjectHandler HandlerObj;
+};
+
+#endif
\ No newline at end of file
diff --git a/Engine/Source/Runtime/CoreUObject/Public/UObject/OverriddenPropertySet.h b/Engine/Source/Runtime/CoreUObject/Public/UObject/OverriddenPropertySet.h
--- a/Engine/Source/Runtime/CoreUObject/Public/UObject/OverriddenPropertySet.h
+++ b/Engine/Source/Runtime/CoreUObject/Public/UObject/OverriddenPropertySet.h
@@ -12,6 +12,22 @@
 
 COREUOBJECT_API DECLARE_LOG_CATEGORY_EXTERN(LogOverridableObject, Warning, All);
 
+struct FOverridableObjectArrayOverrides
+{	
+	COREUOBJECT_API FOverridableObjectArrayOverrides();
+
+	// Returned during loading
+	TArray<UObject*> Added;
+	TArray<UObject*> RemovedArchetypes; // Ptr to the Archetype of the removed object  
+	
+	// Objects that are present in the array (inherited from the archetype or wholly replaced) 
+	// and that are not the result of an Add/Remove  (see next enum for a specific interpretation)
+	TArray<UObject*> Objects; 
+	
+	// Indicates the override operation for the Objects in the previous array 
+	EOverriddenPropertyOperation ObjectsOperation;
+};
+
 struct FArchiveSerializedPropertyChain;
 
 struct FOverridableSerializationLogicInternalAdapter;
@@ -183,6 +199,7 @@ enum class EOverriddenPropertyOperation : uint8
 	Add,		/* this element was added in the container */
 	Remove,		/* this element was removed from the container */
 	SubObjectsShadowing, /* only used to serialize/collect the subobjects, this saved data can be skipped at loading unless loading a loose property or in a placeholder which will keep reference to subobjects */
+	Externalized, /* this indicates the elements are stored externally */
 };
 
 inline TOptional<EOverriddenPropertyOperation> GetOverriddenOperationFromString(const FString& OverriddenOperationString)
@@ -466,6 +483,12 @@ public:
 	 * @param Data memory of the current property */
 	UE_API void OverrideProperty(FPropertyVisitorPath::Iterator PropertyIterator, const void* Data);
 
+	/**
+	 * Utility method that allows to indicate if an array property has its content externalized or not
+	 * @param Property The property in the set that is being externalized/internalized 
+	 * @param bExternal Indicate if the content will be stored externally */
+	UE_API void SetExternalizedOverriddenProperty(const FOverriddenPropertyNodeID& NodeID, bool bExternal);
+
 	/**
 	 * Handling and storing modification on a property of an object
 	 * @param Notification type either pre/post property overridden
diff --git a/Engine/Source/Runtime/CoreUObject/Public/UObject/PropertyTag.h b/Engine/Source/Runtime/CoreUObject/Public/UObject/PropertyTag.h
--- a/Engine/Source/Runtime/CoreUObject/Public/UObject/PropertyTag.h
+++ b/Engine/Source/Runtime/CoreUObject/Public/UObject/PropertyTag.h
@@ -66,6 +66,7 @@ public:
 	EPropertyTagSerializeType SerializeType = EPropertyTagSerializeType::Unknown;
 	EOverriddenPropertyOperation OverrideOperation; // Overridable serialization state reconstruction 
 	bool	bExperimentalOverridableLogic = false; // Remember if property had CPF_ExperimentalOverridableLogic when saved
+	TOptional<bool>	bExperimentalExternalObjects; // Remember if property had CPF_ExperimentalOverridableLogic|CPF_ExperimentalExternalObjects when saved and the object enabled serialization of external objects for that property
 
 	// Constructors.
 	FPropertyTag();
diff --git a/Engine/Source/Runtime/Engine/Internal/OverriddenObjectsExternalPackageHelpers.h b/Engine/Source/Runtime/Engine/Internal/OverriddenObjectsExternalPackageHelpers.h
new file mode 100644
--- /dev/null
+++ b/Engine/Source/Runtime/Engine/Internal/OverriddenObjectsExternalPackageHelpers.h
@@ -0,0 +1,162 @@
+// Copyright Epic Games, Inc. All Rights Reserved.
+
+#pragma once
+
+#include "UObject/Object.h"
+#include "UObject/OverriddenPropertySet.h"
+#include "UObject/OverriddenObjectsExternalHandler.h"
+
+#include "OverriddenObjectsExternalPackageHelpers.generated.h"
+
+#if WITH_EDITOR
+
+struct FObjectExport;
+class FLinkerLoad;
+
+class FExternalObjectsOverrideHelper
+{
+	public:
+
+	ENGINE_API FExternalObjectsOverrideHelper(UObject* ForOuter, UClass* InDesiredClass, FName InExternalsPropertyName);
+	virtual ~FExternalObjectsOverrideHelper() = default;
+
+	struct FExternalAsset 
+	{
+		EOverriddenPropertyOperation OverrideOperation;
+		FName ObjectPackageName;
+		FName AssetName;
+		FName RemovedObjectDefaultValue;
+		TWeakObjectPtr<UObject> AssetObject;
+		FSoftObjectPath AssetPath;
+	};
+
+	ENGINE_API void GetOverridesFromExternalObjects();
+	ENGINE_API void GetOverridesFromExternalObjects(TArray<FExternalAsset>& InExternalAssets, bool bOnlyDiskAssets = true ) const;
+	ENGINE_API void LoadExternalObjects(bool bFlush);
+	ENGINE_API void Flush();
+
+	ENGINE_API void SetConsiderExternalsWithoutPropertyTagsAsAdded(bool bValue);
+	
+	ENGINE_API void HandleRemove(UObject* RemovedObject, UObject* RemovedObjectArchetype);
+	ENGINE_API void HandleAdd(UObject* RemovedObject);
+	ENGINE_API void HandleRename(UObject* Object);
+
+	ENGINE_API void RemoveExternalObjectPackage(UObject* Object, bool bCanDirtyPackage);
+	ENGINE_API void RemoveExternalObjectPackage(FString ExternalAssetPackageName, FName AssetName, bool bForAnExternalRemove, bool bCanDirtyPackage);
+	ENGINE_API void CreateExternalObjectPackage(UObject* Object, bool bCanDirtyPackage);
+	ENGINE_API void CreateExternalRemoveObjectPackage(UObject* RemovedObjectArchetype);
+	
+	ENGINE_API void RemoveExternalPackages(bool bCanDirtyPackages);
+	ENGINE_API void CreateExternalPackages(bool bCanDirtyPackages);
+
+	ENGINE_API void GetAssetRegistryTags(const UObject* InObjectInOuter, FAssetRegistryTagsContext Context);
+
+	ENGINE_API static void ImportExternalOverrides(const FLinkerLoad& Linker, const FObjectExport& Export, TArray<FDynamicImport>& OutDynamicImports);
+
+	template<class T, typename RemoveFn, typename AddFn, typename ModFn>
+	void ApplyTo(TArray<T>& Objects, RemoveFn RemoveFunctor, AddFn AddFunctor, ModFn ModifyFunctor)
+	{
+		for (const FExternalAsset& ExternalAsset : ExternalAssets)
+		{	
+			if (!ExternalAsset.AssetObject.IsValid())
+			{
+				// In the case where no asset object was able to be loaded we need to skip the apply there's nothing that can be done
+				// The loading code will a warning about this
+				continue;
+			}
+
+			switch (ExternalAsset.OverrideOperation)
+			{
+				case EOverriddenPropertyOperation::Add:
+				case EOverriddenPropertyOperation::Replace:
+				{
+					AddFunctor(ExternalAsset.AssetObject.Get());
+				}
+				break;
+				case EOverriddenPropertyOperation::Remove:
+				{
+					int32 IndexOfRemoved = Objects.IndexOfByPredicate([&ExternalAsset](UObject* Object)
+					{
+						if (Object->GetArchetype() == ExternalAsset.AssetObject.Get())
+						{
+							return true;
+						}
+
+						// before subobject reinstance we sometime have the archetypes subobjects in the array, remove them too
+						if (Object == ExternalAsset.AssetObject.Get())
+						{
+							return true;
+						}
+
+						return false;
+					});
+
+					if (IndexOfRemoved != INDEX_NONE)
+					{
+						RemoveFunctor(Objects[IndexOfRemoved], IndexOfRemoved);
+					}
+				}
+				break;
+				case EOverriddenPropertyOperation::None:
+				{
+					int32 IndexOfArchetype= Objects.IndexOfByPredicate([&ExternalAsset](UObject* Object)
+					{
+						if (Object == ExternalAsset.AssetObject->GetArchetype())
+						{
+							return true;
+						}
+
+						return false;
+					});
+
+					if (IndexOfArchetype != INDEX_NONE)
+					{
+						ModifyFunctor(ExternalAsset.AssetObject.Get(), IndexOfArchetype);
+					}
+				}
+				break;
+				default:
+				{
+					// nothing
+				}
+			}
+		}
+	}
+
+	ENGINE_API void InitializeObjectArrayOverridesFromExternalAssets(TNonNullPtr<FOverridableObjectArrayOverrides> InArrayOverrides, TArray<FExternalAsset>& InFromExternalAssets, bool bLogInvalidAssets);
+	ENGINE_API void InitializeObjectArrayOverridesFromObject(TNonNullPtr<FOverridableObjectArrayOverrides> InArrayOverrides, TNonNullPtr<FProperty> InProperty, TNonNullPtr<UObject> InObject);
+
+	virtual FString GetPackageName(UObject* Object) = 0;
+	virtual UObject* GetOuterForRemovePlaceholder(UObject* Object, UObject* ObjectOuteredIn) = 0;
+	
+	TArray<FExternalAsset> ExternalAssets;
+	TArray<int32> PackagesAsyncLoadRequests;
+	UObject* Outer = nullptr;
+	UClass* DesiredClass =  nullptr;	
+	FName PackageNameInMemory;
+	FName ExternalsPropertyName;	
+	bool bExternalAssetsScanned = false;
+	bool bConsiderExternalsWithoutPropertyTagsAsAdded = false;
+};
+#endif //#if WITH_EDITOR
+
+UCLASS(MinimalAPI)
+class UOverrideRemovePlaceholder : public UObject
+{
+	GENERATED_BODY()
+
+public:
+	UOverrideRemovePlaceholder ();
+
+#if WITH_EDITORONLY_DATA
+	UPROPERTY()
+	FSoftObjectPath OverrideRemoved;
+
+	UPROPERTY() 
+	FString OverrideProperty;
+
+	ENGINE_API virtual bool IsEditorOnly() const override;
+	ENGINE_API virtual bool IsAsset() const override;
+	ENGINE_API virtual void GetAssetRegistryTags(FAssetRegistryTagsContext Context) const override;
+#endif
+};
diff --git a/Engine/Source/Runtime/Engine/Private/OverriddenObjectsExternalPackageHelpers.cpp b/Engine/Source/Runtime/Engine/Private/OverriddenObjectsExternalPackageHelpers.cpp
new file mode 100644
--- /dev/null
+++ b/Engine/Source/Runtime/Engine/Private/OverriddenObjectsExternalPackageHelpers.cpp
@@ -0,0 +1,973 @@
+// Copyright Epic Games, Inc. All Rights Reserved.
+
+#include "OverriddenObjectsExternalPackageHelpers.h"
+
+#include "Algo/RemoveIf.h"
+#include "AssetRegistry/AssetRegistryHelpers.h"
+#include "AssetRegistry/AssetRegistryModule.h"
+#include "AssetRegistry/IAssetRegistry.h"
+#include "UObject/AssetRegistryTagsContext.h"
+#include "UObject/Linker.h"
+#include "UObject/LinkerLoad.h"
+#include "UObject/Package.h"
+#include "UObject/OverridableManager.h"
+#include "UObject/UObjectIterator.h"
+#include "Misc/PackageName.h"
+#include "Misc/Paths.h"
+
+#if WITH_EDITOR
+
+DEFINE_LOG_CATEGORY_STATIC(LogExternalOverrides, Log, All);
+
+FExternalObjectsOverrideHelper::FExternalObjectsOverrideHelper(UObject* ForOuter, UClass* InDesiredClass, FName InExternalsPropertyName)
+	: Outer(ForOuter)
+	, DesiredClass(InDesiredClass)
+	, ExternalsPropertyName(InExternalsPropertyName)
+{
+	// will be enabled once all current usages are removed
+	//ensureMsgf(Outer->GetPackage() != GetTransientPackage(), TEXT("Trying to build a FExternalObjectsOverrideHelper with an outer (%s) in the transient package, this is most likely an error and will not yield any result"), *Outer->GetFullName());
+}
+
+void FExternalObjectsOverrideHelper::GetOverridesFromExternalObjects()	
+{
+	if (bExternalAssetsScanned)
+	{
+		return;
+	}
+
+	GetOverridesFromExternalObjects(ExternalAssets);
+
+	bExternalAssetsScanned = true;
+}
+
+
+void FExternalObjectsOverrideHelper::GetOverridesFromExternalObjects(TArray<FExternalAsset>& InExternalAssets, bool bOnlyDiskAssets /*= true */) const
+{
+	TRACE_CPUPROFILER_EVENT_SCOPE(FExternalObjectsOverrideHelper::GetOverridesFromExternalObjects);
+
+	const UPackage* OuterObjectPackage = Outer->GetPackage();
+	const FName OuterObjectPackageNameInMemory = OuterObjectPackage->GetFName();
+	const FName OuterObjectPackageNameOnDisk = !OuterObjectPackage->GetLoadedPath().IsEmpty() ? OuterObjectPackage->GetLoadedPath().GetPackageFName() : OuterObjectPackageNameInMemory;
+	const bool bIsInstanced = OuterObjectPackageNameInMemory != OuterObjectPackageNameOnDisk;
+
+	FString OuterPath = Outer->GetPathName();
+	FName Path;
+		
+	if (bOnlyDiskAssets)
+	{	
+		if (bIsInstanced)
+		{
+			// Replace the expected outer path to the one on disk for loading and filtering assets
+			// Note at this point, OuterObjectPackageNameInMemory is not guaranteed to be a part of the path since 
+			// externalization and instancing can happen anywhere so we reconstruct the outer path as it was saved 
+			// on disk to be able to correctly filter the objects we need to load. 
+			UObject* OuterMostObject = Outer->GetOutermostObject();
+			OuterPath = FSoftObjectPath::ConstructFromAssetPathAndSubpath(FTopLevelAssetPath(OuterMostObject->GetPackage()->GetLoadedPath().GetPackageFName(), OuterMostObject->GetFName()), Outer->GetPathName(OuterMostObject)).ToString();
+		}
+
+		Path = OuterObjectPackageNameOnDisk;
+	}
+	else
+	{
+		Path = OuterObjectPackageNameInMemory;
+	}
+		
+	const bool bIsTempPackage = FPackageName::IsTempPackage(Path.ToString());
+
+	FLinkerLoad* OuterLinker = Outer->GetLinker();
+	const FLinkerInstancingContext* OuterInstancingContext = &OuterLinker->GetInstancingContext();
+
+	EAssetsWithOuterForPathsFlags OuterForPathsFlags = (EAssetsWithOuterForPathsFlags::RecursivePaths | ((bOnlyDiskAssets) ? EAssetsWithOuterForPathsFlags::IncludeOnlyOnDiskAsset : EAssetsWithOuterForPathsFlags::None) | EAssetsWithOuterForPathsFlags::ExactOuter | ((!bIsTempPackage) ? EAssetsWithOuterForPathsFlags::ScanPaths : EAssetsWithOuterForPathsFlags::None ));
+	if (TArray<FAssetData> Assets = UAssetRegistryHelpers::GetAssetsWithOuterForPaths({Path}, *OuterPath, OuterForPathsFlags); !Assets.IsEmpty())
+	{
+		Assets.SetNum(Algo::RemoveIf(Assets, [](const FAssetData& InAssetData) { return InAssetData.IsRedirector(); }));
+		
+		for (const FAssetData& Asset : Assets)
+		{
+			FExternalAsset ExternalAsset;
+				
+			ExternalAsset.ObjectPackageName = Asset.PackageName;
+			ExternalAsset.AssetName = Asset.AssetName;
+			ExternalAsset.OverrideOperation = EOverriddenPropertyOperation::None;
+
+			FString	ExternalOverridePropertyValue;
+			if (Asset.GetTagValue(FName(TEXT("ExternalOverrideProperty")), ExternalOverridePropertyValue))
+			{
+				if (ExternalOverridePropertyValue != ExternalsPropertyName)
+				{
+					// disregard this external since it doesn't match the property we're loading externals for right now
+					continue;
+				}
+			}
+			else if (bConsiderExternalsWithoutPropertyTagsAsAdded)
+			{
+				// external packages objects that do not have this property set and are considered adds
+				ExternalAsset.OverrideOperation = EOverriddenPropertyOperation::Add;
+			}
+				
+			if (Asset.FindTag(FName(TEXT("OverrideAdded"))))
+			{	
+				ExternalAsset.OverrideOperation = EOverriddenPropertyOperation::Add;
+			}				
+			else if (Asset.FindTag(FName(TEXT("OverrideReplaced"))))
+			{
+				ExternalAsset.OverrideOperation = EOverriddenPropertyOperation::Replace;
+			}
+			else if (Asset.GetTagValue(FName(TEXT("OverrideRemoved")), ExternalAsset.RemovedObjectDefaultValue))
+			{
+				ExternalAsset.OverrideOperation = EOverriddenPropertyOperation::Remove;
+				ExternalAsset.AssetObject = FSoftObjectPtr(FSoftObjectPath(ExternalAsset.RemovedObjectDefaultValue.ToString())).Get();
+			}
+			
+			ExternalAsset.AssetPath = Asset.ToSoftObjectPath();
+
+			if (bIsInstanced)
+			{
+				// remap asset path to instanced path
+				OuterInstancingContext->FixupSoftObjectPath(ExternalAsset.AssetPath);
+			}
+
+			if (!ExternalAsset.AssetObject.IsValid())
+			{
+				ExternalAsset.AssetObject = TSoftObjectPtr<UObject>(ExternalAsset.AssetPath).Get();
+			}
+			InExternalAssets.Add(ExternalAsset);
+		}
+	}
+}
+
+void FExternalObjectsOverrideHelper::LoadExternalObjects(bool bFlushLoad)
+{
+	TRACE_CPUPROFILER_EVENT_SCOPE(FExternalObjectsOverrideHelper::LoadExternalObjects);
+
+	GetOverridesFromExternalObjects();
+
+	const UPackage* OuterObjectPackage = Outer->GetPackage();
+	const FName OuterObjectPackageNameInMemory = OuterObjectPackage->GetFName();
+	const FName OuterObjectPackageNameOnDisk = !OuterObjectPackage->GetLoadedPath().IsEmpty() ? OuterObjectPackage->GetLoadedPath().GetPackageFName() : OuterObjectPackageNameInMemory;
+	const bool bIsInstanced = OuterObjectPackageNameInMemory != OuterObjectPackageNameOnDisk;		
+	FLinkerInstancingContext InstancingContext;
+	const FLinkerInstancingContext* OuterInstancingContext = nullptr;
+
+	PackagesAsyncLoadRequests.Reserve(ExternalAssets.Num());
+	const uint32 LoadFlags = (Outer->GetPackage()->HasAnyPackageFlags(PKG_PlayInEditor) ? LOAD_PackageForPIE : LOAD_None) | LOAD_FindIfFail;
+
+	for (FExternalAsset& ExternalAsset : ExternalAssets)
+	{
+		// No need to generate a load request if the object exists
+		// This should always be the case since those are requested through dynamic imports
+		if (ExternalAsset.AssetObject.IsValid())
+		{
+			continue;
+		}
+
+		// Before issuing a load, try to locate the export in the package and resolve it's resource there's a high likelyhood
+		// this will succeed. 
+		UE_LOG(LogExternalOverrides, Verbose, TEXT("Attempting to locate the object resource in linker for %s in package: %s, for outer: %s"), *ExternalAsset.AssetName.ToString(), *ExternalAsset.ObjectPackageName.ToString(), *Outer->GetFullName());
+
+		if (UPackage* Package = FindPackage(nullptr, *ExternalAsset.ObjectPackageName.ToString()))
+		{
+			FLinkerLoad* Linker = Package->GetLinker();
+
+			auto FindIndicesFor = [Linker](FName Name) -> TArray<FPackageIndex>
+			{
+				TArray<FPackageIndex> Indices;
+
+				for (FObjectExport& Export : Linker->ExportMap)
+				{
+					if (Export.ObjectName == Name)
+					{
+						Indices.Add(Export.ThisIndex);
+					}
+				}
+
+				return Indices;
+			};			
+
+			// find possible exports matching the asset name (might be more than on we then need to
+			// follow the outer chain back to validate which one is the export we want)
+			TArray<FPackageIndex> Indices = FindIndicesFor(ExternalAsset.AssetName);
+
+			// if we found possible exports
+			if (Indices.Num())
+			{
+				// split the asset's outer chain into a list of FNames so that we can check each object individually
+				TArray<FString> OuterChain;
+				ExternalAsset.AssetPath.GetSubPathString().ParseIntoArray(OuterChain, TEXT("."));
+
+				// for each export 
+				for (FPackageIndex Index : Indices)
+				{
+					FPackageIndex OuterIndex = Index;
+					const FObjectResource* Obj = nullptr;
+					int32 OuterChainIndex = OuterChain.Num()-1;
+
+					do 
+					{
+						Obj = &Linker->ImpExp(OuterIndex);
+						OuterIndex = Obj->OuterIndex;
+						OuterChainIndex --;
+					} while (Obj->ObjectName.ToString() == OuterChain[OuterChainIndex+1] && OuterChainIndex >= 0 && !OuterIndex.IsNull());
+
+					if (OuterChainIndex < 0)
+					{
+						// we walked up all the way, so we have the correct object
+						UE_LOG(LogExternalOverrides, Verbose, TEXT("Attempting to ResolveResource for %s"), *ExternalAsset.AssetName.ToString());
+
+						if (UObject* Object = Package->GetLinker()->ResolveResource(Index))
+						{
+							ExternalAsset.AssetObject = Object;
+						}
+
+						break;
+					}
+				}
+			}
+
+			// reverify if we manage to find the object
+			if (ExternalAsset.AssetObject.IsValid())
+			{
+				continue;
+			}
+		}
+		
+		UE_LOG(LogExternalOverrides, Warning, TEXT("Force loading an external object outside of dynamic imports: %s"), *ExternalAsset.ObjectPackageName.ToString());
+
+		FName InstancedName;
+		FName ObjectPackageName = ExternalAsset.ObjectPackageName;
+			
+		if (bIsInstanced)
+		{
+			if (OuterInstancingContext)
+			{
+				InstancedName = OuterInstancingContext->RemapPackage(ObjectPackageName);
+			}
+
+			// Remap to the a instanced package if it wasn't remapped already by the outer instancing context
+			if (InstancedName == ObjectPackageName || InstancedName.IsNone())
+			{
+				InstancedName = *FLinkerInstancingContext::GetInstancedPackageName(OuterObjectPackageNameInMemory.ToString(), ObjectPackageName.ToString());
+			}
+
+			checkf(!InstancedName.IsNone(), TEXT("Could find not the remapped package name for dynamic import package %s, in memory name %s"), *ObjectPackageName.ToString(), *OuterObjectPackageNameInMemory.ToString());
+
+			InstancingContext.AddPackageMapping(ObjectPackageName, InstancedName);
+		}
+		
+		FLoadPackageAsyncOptionalParams OptionalParams
+		{
+			.CustomPackageName = InstancedName,
+			.PackagePriority = INT32_MAX,
+			.InstancingContext = &InstancingContext,
+			.LoadFlags = LoadFlags
+		};
+
+		const int32 RequestID = LoadPackageAsync(ObjectPackageName.ToString(), MoveTemp(OptionalParams));
+
+		if (RequestID != INDEX_NONE )
+		{
+			PackagesAsyncLoadRequests.Add(RequestID);
+		}
+	}	
+	
+	if (bFlushLoad)
+	{
+		Flush();
+	}
+}
+
+void FExternalObjectsOverrideHelper::Flush()
+{
+	TRACE_CPUPROFILER_EVENT_SCOPE(FExternalObjectsOverrideHelper::Flush);
+
+	if (PackagesAsyncLoadRequests.Num())
+	{			
+		FlushAsyncLoading(PackagesAsyncLoadRequests);
+	}
+	
+	for (FExternalAsset& ExternalAsset : ExternalAssets)
+	{		
+		if (!ExternalAsset.AssetObject.IsValid())
+		{
+			// Locate the UPackage and the external assets in it
+			if (UPackage* LoadedPackage = FindPackage(nullptr, *ExternalAsset.ObjectPackageName.ToString()))
+			{
+				ForEachObjectWithPackage(LoadedPackage, [&ExternalAsset, this](UObject* InPackageObject)
+					{
+						if (InPackageObject->IsA(DesiredClass) && InPackageObject->GetFName() == ExternalAsset.AssetName)
+						{
+							ExternalAsset.AssetObject = InPackageObject;
+						}
+						return true;
+					}, false);
+			}
+		}
+	}
+}
+
+void FExternalObjectsOverrideHelper::SetConsiderExternalsWithoutPropertyTagsAsAdded(bool bValue)
+{
+	bConsiderExternalsWithoutPropertyTagsAsAdded = bValue;
+}
+
+void FExternalObjectsOverrideHelper::HandleRemove(UObject* RemovedObject, UObject* RemovedObjectArchetype)
+{	
+	if (!RemovedObjectArchetype)
+	{
+		// Removed objects lacking an archetypes are adds so we just remove the external file no need to keep a remove placeholder
+		checkf(RemovedObject && !RemovedObject->HasAllFlags(RF_Transient), TEXT("FExternalObjectsOverrideHelper::HandleRemove called without a valid RemovedObject"));
+
+		RemoveExternalObjectPackage(RemovedObject, true);
+		return;
+	}
+
+	CreateExternalRemoveObjectPackage(RemovedObjectArchetype);
+}
+
+void FExternalObjectsOverrideHelper::HandleRename(UObject* Object)
+{
+	// Update all dependent external packages names
+	ForEachObjectWithOuter(Object, [this](UObject* Inner)
+	{
+		if (Inner->IsPackageExternal())
+		{
+			if (Inner->IsA(DesiredClass) || Inner->IsA<UOverrideRemovePlaceholder>() )
+			{
+				UPackage* Package = Inner->GetPackage();
+				UObject* InnerOuter = Inner->GetOuter();
+
+				if (InnerOuter->IsA(DesiredClass))
+				{
+					const FString InnerPackagePathName = Package->GetPathName();
+					const FString InnerName = *Inner->GetName();
+					const FString OuterPackageName = GetPackageName(InnerOuter);
+				
+					const FString UpdatedPackageName = FPaths::Combine(OuterPackageName, InnerName);
+					if (UpdatedPackageName != InnerPackagePathName)
+					{
+						Package->Rename(*UpdatedPackageName, nullptr, 0);
+					}
+				}
+			}
+		}
+	});
+}
+
+
+void FExternalObjectsOverrideHelper::CreateExternalRemoveObjectPackage(UObject* RemovedObjectArchetype)
+{	
+	// Create package & dummy objects for removed objects + add the proper tag
+	// 
+	// Removed external objects are indicated by a package which contains 
+	// a UOverrideRemovePlaceholder to preserve the remove override.
+	// 
+	// The package also contains an "OverrideRemoved" tag with the name of the 
+	// remove object archetype
+	
+	bool bIsArchetype = Outer->HasAnyFlags(RF_ArchetypeObject);
+				
+	// determine package name	
+	FString OuterPackageName = GetPackageName(Outer);
+	FString RemovedObjectName = RemovedObjectArchetype->GetName() + TEXT("_removed");
+	FString RemovedPackageName = OuterPackageName / RemovedObjectName;
+
+	UE_LOG(LogExternalOverrides, Verbose, TEXT("Creating external removed placholder object of achetype %s in %s"), *RemovedObjectArchetype->GetName(), *Outer->GetName());
+
+	// ensure package exists, this only works in memory... scan on disk too?
+	UPackage* Package = FindPackage(nullptr, *RemovedPackageName);
+
+	if (!Package)
+	{
+		Package = CreatePackage(*RemovedPackageName);
+		if (!Outer->IsTemplate() || Outer->GetPackage()->HasAnyPackageFlags(PKG_ContainsMapData))
+		{
+			Package->SetPackageFlags(PKG_ContainsMapData);
+		}
+		// avoid reloading the package if one existing on disk in that location
+		Package->MarkAsFullyLoaded();
+	}			
+
+	// ensure we have the correct placeholder
+	FSoftObjectPath RemovedPath(RemovedObjectArchetype);
+	TArray<UObject*> ObjectsToTrash;
+
+	bool bNeedCreatePlaceHolder = true;
+
+	if (Package)
+	{
+		ForEachObjectWithPackage(Package, [&bNeedCreatePlaceHolder, &RemovedPath, &ObjectsToTrash, this, Package](UObject* Object)
+		{
+			UOverrideRemovePlaceholder* Placeholder = Cast<UOverrideRemovePlaceholder>(Object);
+
+			if (Placeholder && Placeholder->OverrideRemoved == RemovedPath && Placeholder->GetOuter() == Outer && Placeholder->OverrideProperty == ExternalsPropertyName)
+			{
+				bNeedCreatePlaceHolder = false;
+			}
+			else if (Object->GetExternalPackage() == Package ||
+					 Object->GetOuter() == Package)
+			{
+				ObjectsToTrash.Add(Object);
+			}
+
+			return true;
+		});
+	}
+		
+	for (UObject* Object : ObjectsToTrash)
+	{
+		// thrash this object
+		Package->SetDirtyFlag(true);
+		Object->SetExternalPackage(nullptr);
+		if (Object->GetOuter() == Package )
+		{
+			// To ensure a clean package for the remove placeholder we need to move that object out 
+			// of the package if any is left. This should not normally happen so we log a warning 
+			// help diagnose when it happens. 
+			if (!Object->IsA<UOverrideRemovePlaceholder>())
+			{
+				UE_LOG(LogExternalOverrides, Warning, TEXT("Outering %s to Transient package so that we create a clean remove placeholer package"), *Object->GetFullName());
+			}
+			Object->Rename(nullptr, GetTransientPackage(), REN_DontCreateRedirectors);
+		}
+	}
+	
+	if (bNeedCreatePlaceHolder)
+	{	
+		UOverrideRemovePlaceholder* RemovedPlaceholder = NewObject<UOverrideRemovePlaceholder>(Outer, FName(RemovedObjectName));
+		RemovedPlaceholder->SetExternalPackage(Package);
+		RemovedPlaceholder->SetFlags(RF_Standalone | RF_Public| RF_Transactional);
+		if (bIsArchetype)
+		{
+			RemovedPlaceholder->SetFlags(RF_ArchetypeObject);
+		}
+		RemovedPlaceholder->OverrideRemoved = RemovedPath;
+		RemovedPlaceholder->OverrideProperty = ExternalsPropertyName.ToString();
+		FAssetRegistryModule::AssetCreated(RemovedPlaceholder);
+
+		Package->SetDirtyFlag(true); //only necessary if we're reusing a package with a different removed name 
+	}
+}
+
+void FExternalObjectsOverrideHelper::HandleAdd(UObject* AddedObject)
+{
+	// done in rename but might be moved here eventually
+}
+
+
+void FExternalObjectsOverrideHelper::RemoveExternalObjectPackage(UObject* External, bool bCanDirtyPackage)
+{
+	if (External->IsPackageExternal())
+	{
+		checkf(!External->HasAllFlags(RF_Transient), TEXT("FExternalObjectsOverrideHelper::RemoveExternalObjectPackage called on an RF_Transient object"));
+		RemoveExternalObjectPackage(External->GetPackage()->GetName(), External->GetFName(), External->IsA<UOverrideRemovePlaceholder>(), bCanDirtyPackage);
+	}
+}
+
+void FExternalObjectsOverrideHelper::RemoveExternalObjectPackage(FString ExternalAssetPackageName, FName AssetName, bool bForAnExternalRemove, bool bCanDirtyPackage)
+{
+	UPackage* Package = FindPackage(nullptr, *ExternalAssetPackageName);
+	bool bNeedCreateDummyObject = true;
+
+	FName DummyObjectName = AssetName;
+	UObject* DummyObjectOuter = GetOuterForRemovePlaceholder(Outer, Outer);
+
+	bool bPackageWasNewlyCreated = false;
+	// Create package instead of loading an object we don't need
+	if (!Package)
+	{
+		Package = CreatePackage(*ExternalAssetPackageName);
+		UPackage* OuterPackage = Outer->GetPackage();
+		if (!Outer->IsTemplate() || OuterPackage->HasAnyPackageFlags(PKG_ContainsMapData)) 
+		{
+			Package->SetPackageFlags(PKG_ContainsMapData);
+		}
+	}
+	else
+	{
+		bPackageWasNewlyCreated = Package->HasAllPackagesFlags(PKG_NewlyCreated); // newly created packages aren't on disk, so there's no need to create the dummy object in them, they'll just get GCed if we empty them, so we take not of this state for latter
+		ResetLinkerExports(Package);
+		
+		TArray<UObject*> ObjectsToMoveOutOfPackage;
+
+		// make sure package is "empty" (contains a appropriately named transient object to that it gets deleted)
+		ForEachObjectWithPackage(Package, [&ObjectsToMoveOutOfPackage, DummyObjectName, DummyObjectOuter, &bNeedCreateDummyObject, Package](UObject* Object)
+		{
+			if (Object->GetName() == DummyObjectName && 
+				Object->GetOuter() == DummyObjectOuter && 
+				Object->HasAllFlags(RF_Standalone | RF_Public |RF_Transient|RF_Transactional))
+			{
+				bNeedCreateDummyObject = false;
+			}
+			else if (Object->GetExternalPackage() == Package ||
+					 Object->GetOuter() == Package)
+			{
+				ObjectsToMoveOutOfPackage.Add(Object);
+			}
+
+			return true;
+		});
+		
+		for (UObject* Object : ObjectsToMoveOutOfPackage)
+		{
+			Package->SetDirtyFlag(true);
+			Object->SetExternalPackage(nullptr);
+
+			if (Object->GetOuter() == Package || bForAnExternalRemove)
+			{			
+				// We need to move that object out of the package for the package to be removed from disk. 
+				// It's expected that UOverrideRemovePlaceholder objects to be present and they need to be 
+				// outered to transient (they're not necessary for normal override remove functionality so
+				// we don't keep them)
+				if (!bForAnExternalRemove)
+				{
+					UE_LOG(LogExternalOverrides, Warning, TEXT("Outering %s to Transient package so that we can delete it's package"), *Object->GetFullName());
+				}
+				Object->Rename(nullptr, GetTransientPackage(), REN_DontCreateRedirectors); 
+			}			
+		}
+	}	
+				
+	if (bNeedCreateDummyObject)
+	{
+		if(!bPackageWasNewlyCreated)
+		{
+			UObject* DummyObject = nullptr;
+
+			// Objects is outered on the object that'll proxy it's save that it appears in the correct save contexts
+			if (bForAnExternalRemove)
+			{
+				DummyObject = NewObject<UOverrideRemovePlaceholder>(DummyObjectOuter, FName(DummyObjectName));
+			}
+			else
+			{	
+				DummyObject = NewObject<UObject>(DummyObjectOuter, DesiredClass, FName(DummyObjectName));
+			}
+			DummyObject->SetExternalPackage(Package);
+			DummyObject->SetFlags(RF_Standalone | RF_Public |RF_Transient|RF_Transactional);
+		
+			Package->SetDirtyFlag(true);
+		}
+		else
+		{
+			Package->SetDirtyFlag(false);
+			Package->MarkAsGarbage();
+		}
+	}		
+}
+
+void FExternalObjectsOverrideHelper::CreateExternalObjectPackage(UObject* Object, bool bCanDirtyPackage)
+{	
+	if (Object->IsPackageExternal())
+	{
+		ensure(GetPackageName(Object) == Object->GetPackage()->GetName());
+		return;
+	}
+
+	if (Object->HasAllFlags(RF_Transient))
+	{
+		return;
+	}
+	
+	const bool bShouldDirty = bCanDirtyPackage;
+	const bool bWasAsset = Object->IsAsset();
+
+	Object->Modify(bShouldDirty);
+	checkf(Object->GetOuter() == Outer, TEXT("FExternalObjectsOverrideHelper incorrectly setup for CreateExternalObjectPackage, expected object to be outered to %s, instead is outered to %s"), *Outer->GetFullName(), *Object->GetOuter()->GetFullName());
+	checkf(FPropertyArrayExternalObjectHandler::Get().IsSavingExternal(Outer, Outer->GetClass()->FindPropertyByName(ExternalsPropertyName)), TEXT("Trying to create an external packager for Outer %s which doesn't use external packages"), *Outer->GetFullName());
+	
+	FString PackageName = GetPackageName(Object);
+	UPackage* Package = CreatePackage(*PackageName);
+	bool bForTemplate = Object->IsTemplate();
+
+	if (ensureMsgf(Package, TEXT("CreatePackage returned null")))
+	{
+		Package->SetPackageFlags(PKG_EditorOnly | (bForTemplate ? 0 : PKG_ContainsMapData) | PKG_NewlyCreated);
+		if (bShouldDirty)
+		{
+			Package->MarkPackageDirty();
+		}
+	}		
+
+	Object->SetExternalPackage(Package);
+	
+	if (Object->IsAsset() && !bWasAsset)
+	{
+		FAssetRegistryModule::AssetCreated(Object);
+	}
+
+	if (bShouldDirty)
+	{
+		Object->MarkPackageDirty();
+	}
+}
+
+
+void FExternalObjectsOverrideHelper::RemoveExternalPackages(bool bCanDirtyPackage)
+{	
+	TArray<FExternalAsset> CurrentExternalAssets;
+	GetOverridesFromExternalObjects(CurrentExternalAssets, false);
+
+	for (FExternalAsset& Asset : CurrentExternalAssets)
+	{
+		RemoveExternalObjectPackage(Asset.ObjectPackageName.ToString(), Asset.AssetName, Asset.OverrideOperation == EOverriddenPropertyOperation::Remove, bCanDirtyPackage);
+	}
+}
+
+void FExternalObjectsOverrideHelper::CreateExternalPackages(bool bCanDirtyPackages)
+{
+	FOverridableObjectArrayOverrides Overrides;	
+	InitializeObjectArrayOverridesFromObject(&Overrides, Outer->GetClass()->FindPropertyByName(ExternalsPropertyName), Outer);
+
+	for (UObject* Added : Overrides.Added)
+	{
+		CreateExternalObjectPackage(Added, bCanDirtyPackages);
+	}
+
+	for (UObject* RemovedArchetype : Overrides.RemovedArchetypes)
+	{
+		CreateExternalRemoveObjectPackage(RemovedArchetype);
+	}
+
+	for (UObject* Object : Overrides.Objects)
+	{
+		CreateExternalObjectPackage(Object, bCanDirtyPackages);
+	}
+}
+
+void FExternalObjectsOverrideHelper::InitializeObjectArrayOverridesFromExternalAssets(TNonNullPtr<FOverridableObjectArrayOverrides> ArrayOverrides, TArray<FExternalAsset>& FromExternalAssets, bool bLogInvalidAssets)
+{		
+	for (const FExternalAsset& ExternalAsset : FromExternalAssets)
+	{
+		switch (ExternalAsset.OverrideOperation)
+		{
+			case EOverriddenPropertyOperation::Add:
+			{
+				if (ExternalAsset.AssetObject.IsValid())
+				{
+					ArrayOverrides->Added.Add(ExternalAsset.AssetObject.Get());
+				}
+				else
+				{
+					if (bLogInvalidAssets)
+					{
+						UE_LOG(LogExternalOverrides, Warning, TEXT("Could not load external asset %s, it'll not be Added"), *ExternalAsset.ObjectPackageName.ToString());
+					}
+				}
+			}
+			break;
+			case EOverriddenPropertyOperation::Remove:
+			{
+				if (ExternalAsset.AssetObject.IsValid())
+				{
+					ArrayOverrides->RemovedArchetypes.Add(ExternalAsset.AssetObject.Get());
+				}
+				else
+				{
+					if (bLogInvalidAssets)
+					{
+						UE_LOG(LogExternalOverrides, Warning, TEXT("Could not load external asset %s, %s will not be Removed"), *ExternalAsset.ObjectPackageName.ToString(), *ExternalAsset.RemovedObjectDefaultValue.ToString());
+					}
+				}
+
+			}
+			break;
+			default:
+			{
+				if (ExternalAsset.AssetObject.IsValid())
+				{
+					ArrayOverrides->Objects.Add(ExternalAsset.AssetObject.Get());
+					ArrayOverrides->ObjectsOperation = ExternalAsset.OverrideOperation;
+				}
+				else
+				{
+					if (bLogInvalidAssets)
+					{
+						UE_LOG(LogExternalOverrides, Warning, TEXT("Could not load external asset %s, it's Modifications will be lost"), *ExternalAsset.ObjectPackageName.ToString());
+					}
+				}
+			}
+		}
+	}
+}
+
+void FExternalObjectsOverrideHelper::InitializeObjectArrayOverridesFromObject(TNonNullPtr<FOverridableObjectArrayOverrides> ArrayOverrides, TNonNullPtr<FProperty> Property, TNonNullPtr<UObject> Object)
+{
+	// Find the corresponding overrides
+	if (const FOverriddenPropertySet* OverriddenSet = FOverridableManager::Get().GetOverriddenProperties(Object))
+	{
+		// Look the override node associated with the property
+		FOverriddenPropertyNodeID OverriddenID(Property);
+
+		const FOverriddenPropertyNode* Node = OverriddenSet->GetOverriddenPropertyNode(OverriddenID);
+		FArrayProperty* ArrayProperty = CastField<FArrayProperty>(Property);
+		FObjectProperty* InnerObjectProperty =  CastField<FObjectProperty>(ArrayProperty->GetInnerFieldByName(*ArrayProperty->GetName()));
+				
+		FScriptArrayHelper ArrayHelper(ArrayProperty, ArrayProperty->ContainerPtrToValuePtr<void>(Outer));
+
+		const int32 ArrayNum = ArrayHelper.Num();
+		for (int i = 0; i < ArrayNum; ++i)
+		{
+			UObject* CurrentObject = InnerObjectProperty->GetObjectPropertyValue(ArrayHelper.GetElementPtr(i));
+			ArrayOverrides->Objects.Add(CurrentObject);
+		}
+			
+		if (Node)
+		{			
+			if(Node->GetOperation() == EOverriddenPropertyOperation::Replace)
+			{
+				ArrayOverrides->ObjectsOperation = EOverriddenPropertyOperation::Replace;
+			}			
+
+			FScriptArrayHelper DefaultsArrayHelper(ArrayProperty, ArrayProperty->ContainerPtrToValuePtr<void>(Outer->GetArchetype()));
+
+			auto FindObjectPtr = [&InnerObjectProperty](const FOverriddenPropertyNodeID& ObjectToFind, FScriptArrayHelper& Helper) -> UObject*
+			{
+				const int32 ArrayNum = Helper.Num();
+
+				for (int i = 0; i < ArrayNum; ++i)
+				{
+					if (UObject* CurrentObject = InnerObjectProperty->GetObjectPropertyValue(Helper.GetElementPtr(i)))
+					{
+						if (ObjectToFind == FOverriddenPropertyNodeID(CurrentObject))
+						{
+							return CurrentObject;
+						}
+					}
+				}
+
+				return nullptr;
+			};
+
+			for (const FOverriddenPropertyNode& SubNode : Node->GetSubPropertyNodes())
+			{
+				if (SubNode.GetOperation() == EOverriddenPropertyOperation::Remove)
+				{
+					UObject* FoundObject = FindObjectPtr(SubNode.GetNodeID(), DefaultsArrayHelper);
+						
+					if (FoundObject)
+					{
+						ArrayOverrides->RemovedArchetypes.Add(FoundObject);
+					}
+				}
+				else if (SubNode.GetOperation() == EOverriddenPropertyOperation::Add)
+				{					
+					UObject* FoundObject = FindObjectPtr(SubNode.GetNodeID(), ArrayHelper); 
+						
+					if (FoundObject)
+					{
+						ArrayOverrides->Added.Add(FoundObject);
+						ArrayOverrides->Objects.Remove(FoundObject);
+					}
+				}
+			}
+		}
+	}	
+}
+
+void FExternalObjectsOverrideHelper::GetAssetRegistryTags(const UObject* ObjectInOuter, FAssetRegistryTagsContext Context)
+{	
+	// ensure we're in the correct outer
+	checkf(ObjectInOuter->GetOuter() == Outer, TEXT("FExternalObjectsOverrideHelper incorrectly setup for GetAssetRegistryTags, expected object to be outered to %s, instead is outered to %s"), *Outer->GetFullName(), *ObjectInOuter->GetOuter()->GetFullName());
+	
+	if (ObjectInOuter->IsPackageExternal())
+	{
+		// Using entities as an example....
+		// An entity is added to the parent if the archetype of its outer is 
+		// different than the outer of its archetype
+		// 
+		// Example with entities......
+		// 
+		// PrefabA			=>	PrefabA'
+		//		EntityA				EntityA'				Outer(Archetype(EntityA')) == Archetype(Outer(EntityA'))
+		//							EntityC (added)			Outer(Archetype(EntityC)) != Archetype(Outer(EntityC))			=> Outer(entity) != Archetype(PrefabA')
+		//							PrefabB' (addded)		Outer(Archetype(PrefabB')) != Archetype(Outer(PrefabA))			=> Outer(PrefabB) != Archetype(BlueprintA)
+		//								EntityB'			Outer(Archetype(EntityB')) == Archetype(Outer(EntityB'))		=> Outer(EntityB) == Archetype(PrefabB')
+		// PrefabB
+		//		EntityB 
+		//
+	
+		if (Outer->GetArchetype() != ObjectInOuter->GetArchetype()->GetOuter())
+		{
+			Context.AddTag(UObject::FAssetRegistryTag(TEXT("OverrideAdded"), TEXT("Added"), UObject::FAssetRegistryTag::TT_Alphabetical));
+		}
+		else if (Outer->IsA(DesiredClass))
+		{
+			if (FOverriddenPropertySet* ParentOwnedPropertyOverrideSet = FOverridableManager::Get().GetOverriddenProperties(Outer))
+			{
+				if (FProperty* ForProperty = DesiredClass->FindPropertyByName(ExternalsPropertyName))
+				{
+					FOverriddenPropertyNodeID OverriddenID(ForProperty);
+					const FOverriddenPropertyNode* Node = ParentOwnedPropertyOverrideSet->GetOverriddenPropertyNode(OverriddenID);
+
+					if (Node)
+					{
+						if (Node->GetOperation() ==  EOverriddenPropertyOperation::Replace)
+						{
+							Context.AddTag(UObject::FAssetRegistryTag(TEXT("OverrideReplaced"), TEXT("Replaced"), UObject::FAssetRegistryTag::TT_Alphabetical));
+						}
+					}
+				}
+			}
+		}
+
+		Context.AddTag(UObject::FAssetRegistryTag(TEXT("ExternalOverrideProperty"), *ExternalsPropertyName.ToString(), UObject::FAssetRegistryTag::TT_Alphabetical));
+	}
+}
+
+
+/**
+ * Recursive object path builder from linker export
+ */
+ enum EObjectPathState 
+ {
+	Empty,
+	Package,
+	Object,
+	SubObject
+ };
+
+
+void BuildExportPath(const FLinkerLoad& Linker, FString& OutPathName, const FPackageIndex ResourceIndex, EObjectPathState& PathState)
+{
+	static const TCHAR StateDelimiter[] =
+	 {
+		TEXT('\0'),
+		TEXT('.'),
+		SUBOBJECT_DELIMITER_CHAR, 
+		TEXT('.')
+	 };
+
+	 static const EObjectPathState NextPathState[] =
+	 {
+		EObjectPathState::Package,
+		EObjectPathState::Object,
+		EObjectPathState::SubObject,
+		EObjectPathState::SubObject,
+	 };
+
+
+	const FObjectResource& Resource	= Linker.ImpExp(ResourceIndex);
+
+	if (Resource.OuterIndex.IsNull())
+	{
+		if (!ResourceIndex.IsImport())
+		{
+			// If the export has no outer path, initialize the OuterPath with the Package's path on disk		
+			OutPathName = Linker.LinkerRoot->GetName();
+			PathState = EObjectPathState::Package;
+		}		
+	}
+	else
+	{		
+		BuildExportPath(Linker, OutPathName, Resource.OuterIndex, PathState);
+		PathState = NextPathState[PathState];
+	}	
+	
+	OutPathName += StateDelimiter[PathState];
+	OutPathName += Resource.ObjectName.ToString();
+};
+
+void BuildExportPath(const FLinkerLoad& Linker, FString& OutPathName, const FObjectExport& Resource)
+{
+	EObjectPathState State = EObjectPathState::Empty;
+	BuildExportPath(Linker, OutPathName, Resource.ThisIndex, State);
+}
+
+
+/**
+ * Given an Export from the package's export table, which is of class verse::entity and bears the flag RF_HasDynamicImports,
+ * find the entities which have this Export as their outer, so they are added as Imports of this Export, which
+ * automate their loading in the async loading queue.
+ */
+void FExternalObjectsOverrideHelper::ImportExternalOverrides(const FLinkerLoad& Linker, const FObjectExport& Export, TArray<FDynamicImport>& OutDynamicImports)
+{
+	TRACE_CPUPROFILER_EVENT_SCOPE(FExternalObjectsOverrideHelper::ImportExternalOverrides);
+	
+	const FLinkerInstancingContext& InstancingContext = Linker.GetInstancingContext();
+
+	const UPackage* Package = Linker.LinkerRoot;
+	// Handle instancing
+	const FName PackageNameInMemory = Package->GetFName();
+	const FName PackageNameOnDisk = !Package->GetLoadedPath().IsEmpty() ? Package->GetLoadedPath().GetPackageFName() : PackageNameInMemory;	
+	
+	const bool bIsInstanced = PackageNameInMemory != PackageNameOnDisk;
+	// Produce fully qualified pathname for this object, in the format:
+	// 'Outermost.[Outer:]Name'
+	// See: UObjectBaseUtility::GetPathName
+	FString OuterPath;
+	BuildExportPath(Linker, OuterPath, Export);
+
+	UE_LOG(LogExternalOverrides, Verbose, TEXT("Querying dynamic imports for package %s (in memory: %s), outer: %s"), *PackageNameOnDisk.ToString(), *PackageNameInMemory.ToString(), *OuterPath);
+
+	if (bIsInstanced)
+	{
+		// Replace the expected outer path to the one on disk for loading and filtering assets
+		OuterPath.ReplaceInline(*PackageNameInMemory.ToString(), *PackageNameOnDisk.ToString());
+	}
+
+	EAssetsWithOuterForPathsFlags OuterForPathsFlags = (EAssetsWithOuterForPathsFlags::RecursivePaths | EAssetsWithOuterForPathsFlags::IncludeOnlyOnDiskAsset | EAssetsWithOuterForPathsFlags::ExactOuter | EAssetsWithOuterForPathsFlags::ScanPaths);
+	if (TArray<FAssetData> Assets = UAssetRegistryHelpers::GetAssetsWithOuterForPaths({ PackageNameOnDisk }, *OuterPath, OuterForPathsFlags); !Assets.IsEmpty())
+	{
+		Assets.SetNum(Algo::RemoveIf(Assets, [](const FAssetData& InAssetData) { return InAssetData.IsRedirector(); }));		
+		
+		if (Assets.Num() > 0)
+		{
+			OutDynamicImports.Reserve(Assets.Num());
+
+			for (const FAssetData& Asset : Assets)
+			{
+				UE_LOG(LogExternalOverrides, Verbose, TEXT("Adding dynamic import package %s for %s"), *Asset.PackageName.ToString(), *OuterPath);
+
+				FDynamicImport Import;
+				Import.PackageName = Asset.PackageName;
+				if (bIsInstanced)
+				{
+					Import.InstanceName = InstancingContext.RemapPackage(Asset.PackageName);
+
+					// Remap to the instanced package if it wasn't remapped already by the outer instancing context
+					if (Import.InstanceName == Asset.PackageName || Import.InstanceName.IsNone())
+					{
+						Import.InstanceName = *FLinkerInstancingContext::GetInstancedPackageName(PackageNameInMemory.ToString(), Asset.PackageName.ToString());
+					}
+
+					checkf(!Import.InstanceName.IsNone(), TEXT("Could find not the remapped package name for dynamic import package %s, in memory name %s"), *Asset.PackageName.ToString(), *PackageNameInMemory.ToString());
+				} 
+
+				OutDynamicImports.Add(Import);
+			}
+		}
+	}
+}
+
+#endif //#if WITH_EDITOR
+
+UOverrideRemovePlaceholder::UOverrideRemovePlaceholder()
+{
+}
+
+#if WITH_EDITORONLY_DATA
+bool UOverrideRemovePlaceholder::IsEditorOnly() const
+{
+	return true;
+}
+
+bool UOverrideRemovePlaceholder::IsAsset() const
+{
+	if (!HasAllFlags(RF_ClassDefaultObject))
+	{
+		return true;
+	}
+
+	return Super::IsAsset();
+}
+
+void UOverrideRemovePlaceholder::GetAssetRegistryTags(FAssetRegistryTagsContext Context) const
+{
+	//@todo_ow: Get owner entity and check if we're in it's overridden properties instead?
+	if (!HasAllFlags(RF_ClassDefaultObject))
+	{
+		Context.AddTag(UObject::FAssetRegistryTag(TEXT("OverrideRemoved"), OverrideRemoved.ToString(), UObject::FAssetRegistryTag::TT_Alphabetical));
+		Context.AddTag(UObject::FAssetRegistryTag(TEXT("ExternalOverrideProperty"), *OverrideProperty, UObject::FAssetRegistryTag::TT_Alphabetical));
+	}
+}
+#endif //#if WITH_EDITORONLY_DATA
diff --git a/Engine/Source/Runtime/Engine/Private/UnrealExporter.cpp b/Engine/Source/Runtime/Engine/Private/UnrealExporter.cpp
--- a/Engine/Source/Runtime/Engine/Private/UnrealExporter.cpp
+++ b/Engine/Source/Runtime/Engine/Private/UnrealExporter.cpp
@@ -853,6 +853,7 @@ void ExportProperties_Overrides
 
 						TArray<int32> RemovedIndices;
 						TArray<int32> AddedIndices;
+						bool bExternalized = false;
 
 						for (const FOverriddenPropertyNode& SubNode: ArrayOverriddenPropertyNode->GetSubPropertyNodes())
 						{
@@ -876,12 +877,21 @@ void ExportProperties_Overrides
 									}
 									break;
 								}
+							case EOverriddenPropertyOperation::Externalized:
+								// nothing to do 
+								bExternalized = true;
+								break;
 							default:
 								checkf(false, TEXT("Unsupported operation type"));
 								break;
 							}
 						}
 
+						if(bExternalized)
+						{
+							Out.Logf(TEXT("%s%s<%s>=%s\r\n"), FCString::Spc(Indent), *SanitizedPropertyName, *GetOverriddenOperationString(EOverriddenPropertyOperation::Externalized), TEXT("\"\""));
+						}
+
 						for (int32 i : RemovedIndices)
 						{
 							FString Value;
