Fix code analysis issue

#Synced-CL 50537189

diff --git a/Engine/Source/Runtime/Engine/Private/OverriddenObjectsExternalPackageHelpers.cpp b/Engine/Source/Runtime/Engine/Private/OverriddenObjectsExternalPackageHelpers.cpp
--- a/Engine/Source/Runtime/Engine/Private/OverriddenObjectsExternalPackageHelpers.cpp
+++ b/Engine/Source/Runtime/Engine/Private/OverriddenObjectsExternalPackageHelpers.cpp
@@ -207,7 +207,7 @@ void FExternalObjectsOverrideHelper::LoadExternalObjects(bool bFlushLoad)
 						Obj = &Linker->ImpExp(OuterIndex);
 						OuterIndex = Obj->OuterIndex;
 						OuterChainIndex --;
-					} while (Obj->ObjectName.ToString() == OuterChain[OuterChainIndex+1] && OuterChainIndex >= 0 && !OuterIndex.IsNull());
+					} while (OuterChainIndex >= 0 && Obj->ObjectName.ToString() == OuterChain[OuterChainIndex+1] && !OuterIndex.IsNull());
 
 					if (OuterChainIndex < 0)
 					{
