[SceneGraph] Fix prefabs to instance new nodes in the object graph inherited from a modified Verse parent class on VkEdit servers after a Verse hot reload (VHR) event, triggered by pushing Verse changes only from the editor client.

#Synced-CL 50598508

diff --git a/Engine/Source/Runtime/Engine/Classes/Engine/BlueprintGeneratedClass.h b/Engine/Source/Runtime/Engine/Classes/Engine/BlueprintGeneratedClass.h
--- a/Engine/Source/Runtime/Engine/Classes/Engine/BlueprintGeneratedClass.h
+++ b/Engine/Source/Runtime/Engine/Classes/Engine/BlueprintGeneratedClass.h
@@ -635,6 +635,8 @@ public:
 
 	ENGINE_API virtual void InitPropertiesFromCustomList(uint8* DataPtr, const uint8* DefaultDataPtr) override;
 	ENGINE_API virtual void SetupObjectInitializer(FObjectInitializer& ObjectInitializer) const override;
+	ENGINE_API virtual void PostInitInstance(UObject* InObj, FObjectInstancingGraph* InstanceGraph) override;
+	ENGINE_API virtual void PostLoadInstance(UObject* InObj) override;
 
 protected:
 
diff --git a/Engine/Source/Runtime/Engine/Private/BlueprintGeneratedClass.cpp b/Engine/Source/Runtime/Engine/Private/BlueprintGeneratedClass.cpp
--- a/Engine/Source/Runtime/Engine/Private/BlueprintGeneratedClass.cpp
+++ b/Engine/Source/Runtime/Engine/Private/BlueprintGeneratedClass.cpp
@@ -1244,6 +1244,16 @@ void UBlueprintGeneratedClass::SetupObjectInitializer(FObjectInitializer& Object
 	GetSuperClass()->SetupObjectInitializer(ObjectInitializer);
 }
 
+void UBlueprintGeneratedClass::PostInitInstance(UObject* InObj, FObjectInstancingGraph* InstanceGraph)
+{
+	GetSuperClass()->PostInitInstance(InObj, InstanceGraph);
+}
+
+void UBlueprintGeneratedClass::PostLoadInstance(UObject* InObj)
+{
+	GetSuperClass()->PostLoadInstance(InObj);
+}
+
 void UBlueprintGeneratedClass::InitPropertiesFromCustomList(uint8* DataPtr, const uint8* DefaultDataPtr)
 {
 	// autortfm: we've introduced a conditional lock here because while running under autortfm, we
