Only attempt to load Source Data Override Plugin if Override is enabled on Attenuation settings

#Synced-CL 50609425

diff --git a/Engine/Source/Runtime/Engine/Private/ActiveSound.cpp b/Engine/Source/Runtime/Engine/Private/ActiveSound.cpp
--- a/Engine/Source/Runtime/Engine/Private/ActiveSound.cpp
+++ b/Engine/Source/Runtime/Engine/Private/ActiveSound.cpp
@@ -2097,49 +2097,58 @@ void FActiveSound::UpdateAttenuation(float DeltaTime, FSoundParseParameters& Par
 	}
 
 	// Figure out which attenuation settings to use
-	if (Settings->PluginSettings.SpatializationPluginSettingsArray.Num() > 0)
+	if (Settings->bSpatialize)
 	{
-		UClass* PluginClass = GetAudioPluginCustomSettingsClass(EAudioPlugin::SPATIALIZATION);
-		if (PluginClass)
+		if (Settings->PluginSettings.SpatializationPluginSettingsArray.Num() > 0)
 		{
-			for (USpatializationPluginSourceSettingsBase* SettingsBase : Settings->PluginSettings.SpatializationPluginSettingsArray)
+			UClass* PluginClass = GetAudioPluginCustomSettingsClass(EAudioPlugin::SPATIALIZATION);
+			if (PluginClass)
 			{
-				if (SettingsBase != nullptr && SettingsBase->IsA(PluginClass))
+				for (USpatializationPluginSourceSettingsBase* SettingsBase : Settings->PluginSettings.SpatializationPluginSettingsArray)
 				{
-					ParseParams.SpatializationPluginSettings = SettingsBase;
-					break;
+					if (SettingsBase != nullptr && SettingsBase->IsA(PluginClass))
+					{
+						ParseParams.SpatializationPluginSettings = SettingsBase;
+						break;
+					}
 				}
 			}
 		}
 	}
 
-	if (Settings->PluginSettings.ReverbPluginSettingsArray.Num() > 0)
+	if (Settings->bEnableReverbSend)
 	{
-		UClass* PluginClass = GetAudioPluginCustomSettingsClass(EAudioPlugin::REVERB);
-		if (PluginClass)
+		if (Settings->PluginSettings.ReverbPluginSettingsArray.Num() > 0)
 		{
-			for (UReverbPluginSourceSettingsBase* SettingsBase : Settings->PluginSettings.ReverbPluginSettingsArray)
+			UClass* PluginClass = GetAudioPluginCustomSettingsClass(EAudioPlugin::REVERB);
+			if (PluginClass)
 			{
-				if (SettingsBase != nullptr && SettingsBase->IsA(PluginClass))
+				for (UReverbPluginSourceSettingsBase* SettingsBase : Settings->PluginSettings.ReverbPluginSettingsArray)
 				{
-					ParseParams.ReverbPluginSettings = SettingsBase;
-					break;
+					if (SettingsBase != nullptr && SettingsBase->IsA(PluginClass))
+					{
+						ParseParams.ReverbPluginSettings = SettingsBase;
+						break;
+					}
 				}
 			}
 		}
 	}
 
-	if (Settings->PluginSettings.SourceDataOverridePluginSettingsArray.Num() > 0)
+	if (Settings->bEnableSourceDataOverride)
 	{
-		UClass* PluginClass = GetAudioPluginCustomSettingsClass(EAudioPlugin::SOURCEDATAOVERRIDE);
-		if (PluginClass)
+		if (Settings->PluginSettings.SourceDataOverridePluginSettingsArray.Num() > 0)
 		{
-			for (USourceDataOverridePluginSourceSettingsBase* SettingsBase : Settings->PluginSettings.SourceDataOverridePluginSettingsArray)
+			UClass* PluginClass = GetAudioPluginCustomSettingsClass(EAudioPlugin::SOURCEDATAOVERRIDE);
+			if (PluginClass)
 			{
-				if (SettingsBase != nullptr && SettingsBase->IsA(PluginClass))
+				for (USourceDataOverridePluginSourceSettingsBase* SettingsBase : Settings->PluginSettings.SourceDataOverridePluginSettingsArray)
 				{
-					ParseParams.SourceDataOverridePluginSettings = SettingsBase;
-					break;
+					if (SettingsBase != nullptr && SettingsBase->IsA(PluginClass))
+					{
+						ParseParams.SourceDataOverridePluginSettings = SettingsBase;
+						break;
+					}
 				}
 			}
 		}
