[FTimerManager + AutoRTFM] Rework to make this safe to use from both the open and closed.

We have logic that uses FTimerManager from the open and closed of a transaction. This is extremely unsafe, and will very likely cause memory corruption on abort.

This CL makes it safe to use FTimerManager from both the open and closed of a transaction, so long as the timers as distinct.

All modifications to state have been moved into helper methods that perform the modification in the open and revert the change on abort.
This will enqueue multiple on-abort handlers (more than necessary) for single method calls, but makes this change a simple set of substitutions, reducing the chance of breaking this fundamental type.

#Synced-CL 50592448

diff --git a/Engine/Source/Runtime/Engine/Private/Tests/AutoRTFM/AutoRTFMTimerManagerTests.cpp b/Engine/Source/Runtime/Engine/Private/Tests/AutoRTFM/AutoRTFMTimerManagerTests.cpp
new file mode 100644
--- /dev/null
+++ b/Engine/Source/Runtime/Engine/Private/Tests/AutoRTFM/AutoRTFMTimerManagerTests.cpp
@@ -0,0 +1,629 @@
+// Copyright Epic Games, Inc. All Rights Reserved.
+
+#include "AutoRTFM.h"
+#include "AutoRTFM/Testing.h"
+#include "CoreGlobals.h"
+#include "Engine/TimerHandle.h"
+#include "Misc/AssertionMacros.h"
+#include "Misc/AutomationTest.h"
+#include "TimerManager.h"
+
+#if WITH_DEV_AUTOMATION_TESTS
+
+IMPLEMENT_SIMPLE_AUTOMATION_TEST(FAutoRTFMTimerManagerTest, "AutoRTFM + FTimerManager", EAutomationTestFlags::EngineFilter | EAutomationTestFlags::ClientContext | EAutomationTestFlags::ServerContext | EAutomationTestFlags::CommandletContext)
+
+bool FAutoRTFMTimerManagerTest::RunTest(const FString & Parameters)
+{
+	if (!AutoRTFM::ForTheRuntime::IsAutoRTFMRuntimeEnabled())
+	{
+		ExecutionInfo.AddEvent(FAutomationEvent(EAutomationEventType::Info, TEXT("SKIPPED 'FAutoRTFMTimerManagerTest' test. AutoRTFM disabled.")));
+		return true;
+	}
+
+	auto Section = [&](const TCHAR* Name, auto&& Test)
+	{
+		{
+			FTimerManager TimerManager;
+			TimerManager.Tick(0);
+			if (!Test(TimerManager))
+			{
+				AddError(FString::Printf(TEXT("In section '%s' (HasBeenTickedThisFrame: true)."), Name), 1);
+			}
+		}
+		{
+			FTimerManager TimerManager;
+			TimerManager.Tick(0);
+			GFrameCounter++;
+			if (!Test(TimerManager))
+			{
+				AddError(FString::Printf(TEXT("In section '%s' (HasBeenTickedThisFrame: false)."), Name), 1);
+			}
+		}
+	};
+	
+	auto Tick = [&](FTimerManager& TimerManager)
+	{
+		// Simulate two frames to ensure that pending timers are activated and triggered.
+		for (int I = 0; I < 2; I++)
+		{
+			GFrameCounter++;
+			TimerManager.Tick(1.5f);
+		}
+	};
+
+	struct AUTORTFM_ENABLE FTimer
+	{
+		FTimerHandle Handle;
+		bool bCalled = false;
+		TFunction<void ()> Callback()
+		{
+			return [this]
+			{
+				bCalled = true;
+			};
+		}
+	};
+
+	Section(TEXT("Transact(SetTimer), Tick"), [&] (FTimerManager& TimerManager)
+	{
+		constexpr size_t NumTimers = 16;
+		FTimer Timers[NumTimers]{};
+
+		AutoRTFM::Testing::Commit([&]
+		{
+			for (FTimer& Timer : Timers)
+			{
+				TimerManager.SetTimer(Timer.Handle, Timer.Callback(), /* InRate */ 1.0f, FTimerManagerTimerParameters{});
+				check(TimerManager.IsTimerActive(Timer.Handle));
+			}
+		});
+
+		Tick(TimerManager);
+
+		for (FTimer& Timer : Timers)
+		{
+			check(Timer.bCalled);
+		}
+
+		return true;
+	});
+
+	Section(TEXT("Transact(SetTimer, Abort), Tick"), [&] (FTimerManager& TimerManager)
+	{
+		constexpr size_t NumTimers = 16;
+		FTimer Timers[NumTimers]{};
+
+		AutoRTFM::Testing::Abort([&]
+		{
+			for (FTimer& Timer : Timers)
+			{
+				TimerManager.SetTimer(Timer.Handle, Timer.Callback(), /* InRate */ 1.0f, FTimerManagerTimerParameters{});
+				check(TimerManager.IsTimerActive(Timer.Handle));
+			}
+			AutoRTFM::AbortTransaction();
+		});
+
+		Tick(TimerManager);
+
+		for (FTimer& Timer : Timers)
+		{
+			check(!Timer.bCalled);
+			check(!TimerManager.IsTimerActive(Timer.Handle));
+		}
+
+		return true;
+	});
+
+	Section(TEXT("Transact(SetTimer, ClearTimer), Tick"), [&] (FTimerManager& TimerManager)
+	{
+		constexpr size_t NumTimers = 16;
+		FTimer Timers[NumTimers]{};
+
+		AutoRTFM::Testing::Commit([&]
+		{
+			for (FTimer& Timer : Timers)
+			{
+				TimerManager.SetTimer(Timer.Handle, Timer.Callback(), /* InRate */ 1.0f, FTimerManagerTimerParameters{});
+				check(TimerManager.IsTimerActive(Timer.Handle));
+				TimerManager.ClearTimer(Timer.Handle);
+				check(!TimerManager.IsTimerActive(Timer.Handle));
+			}
+		});
+
+		Tick(TimerManager);
+
+		for (FTimer& Timer : Timers)
+		{
+			check(!Timer.bCalled);
+			check(!TimerManager.IsTimerActive(Timer.Handle));
+		}
+
+		return true;
+	});
+
+	Section(TEXT("Transact(SetTimer, ClearTimer, Abort), Tick"), [&] (FTimerManager& TimerManager)
+	{
+		constexpr size_t NumTimers = 16;
+		FTimer Timers[NumTimers]{};
+
+		AutoRTFM::Testing::Abort([&]
+		{
+			for (FTimer& Timer : Timers)
+			{
+				TimerManager.SetTimer(Timer.Handle, Timer.Callback(), /* InRate */ 1.0f, FTimerManagerTimerParameters{});
+				check(TimerManager.IsTimerActive(Timer.Handle));
+				TimerManager.ClearTimer(Timer.Handle);
+				check(!TimerManager.IsTimerActive(Timer.Handle));
+			}
+			AutoRTFM::AbortTransaction();
+		});
+
+		Tick(TimerManager);
+
+		for (FTimer& Timer : Timers)
+		{
+			check(!Timer.bCalled);
+			check(!TimerManager.IsTimerActive(Timer.Handle));
+		}
+
+		return true;
+	});
+
+	Section(TEXT("Transact(SetTimer, Open(SetTimer)), Tick"), [&] (FTimerManager& TimerManager)
+	{
+		constexpr size_t NumOpenTimers = 16;
+		FTimer OpenTimers[NumOpenTimers];
+		constexpr size_t NumClosedTimers = 16;
+		FTimer ClosedTimers[NumClosedTimers];
+
+		AutoRTFM::Testing::Commit([&]
+		{
+			for (FTimer& ClosedTimer : ClosedTimers)
+			{
+				TimerManager.SetTimer(ClosedTimer.Handle, ClosedTimer.Callback(), /* InRate */ 1.0f, FTimerManagerTimerParameters{});
+				check(TimerManager.IsTimerActive(ClosedTimer.Handle));
+			}
+			
+			AutoRTFM::Open([&]
+			{
+				for (FTimer& OpenTimer : OpenTimers)
+				{
+					TimerManager.SetTimer(OpenTimer.Handle, OpenTimer.Callback(), /* InRate */ 1.0f, FTimerManagerTimerParameters{});
+					check(TimerManager.IsTimerActive(OpenTimer.Handle));
+				}
+			});
+		});
+
+		Tick(TimerManager);
+
+		for (FTimer& ClosedTimer : ClosedTimers)
+		{
+			check(ClosedTimer.bCalled);
+		}
+
+		for (FTimer& OpenTimer : OpenTimers)
+		{
+			check(OpenTimer.bCalled);
+		}
+
+		return true;
+	});
+
+	Section(TEXT("Transact(SetTimer, Open(SetTimer), Abort), Tick"), [&] (FTimerManager& TimerManager)
+	{
+		constexpr size_t NumOpenTimers = 16;
+		FTimer OpenTimers[NumOpenTimers];
+		constexpr size_t NumClosedTimers = 16;
+		FTimer ClosedTimers[NumClosedTimers];
+
+		AutoRTFM::Testing::Abort([&]
+		{
+			for (FTimer& ClosedTimer : ClosedTimers)
+			{
+				TimerManager.SetTimer(ClosedTimer.Handle, ClosedTimer.Callback(), /* InRate */ 1.0f, FTimerManagerTimerParameters{});
+				check(TimerManager.IsTimerActive(ClosedTimer.Handle));
+			}
+			
+			AutoRTFM::Open([&]
+			{
+				for (FTimer& OpenTimer : OpenTimers)
+				{
+					TimerManager.SetTimer(OpenTimer.Handle, OpenTimer.Callback(), /* InRate */ 1.0f, FTimerManagerTimerParameters{});
+					check(TimerManager.IsTimerActive(OpenTimer.Handle));
+				}
+			});
+			
+			AutoRTFM::AbortTransaction();
+		});
+
+		Tick(TimerManager);
+
+		for (FTimer& ClosedTimer : ClosedTimers)
+		{
+			check(!ClosedTimer.bCalled);
+		}
+
+		for (FTimer& OpenTimer : OpenTimers)
+		{
+			check(OpenTimer.bCalled);
+		}
+
+		return true;
+	});
+
+	Section(TEXT("Transact(Open(SetTimer), SetTimer, Open(ClearTimer)), Tick"), [&] (FTimerManager& TimerManager)
+	{
+		constexpr size_t NumOpenTimers = 16;
+		FTimer OpenTimers[NumOpenTimers];
+		constexpr size_t NumClosedTimers = 16;
+		FTimer ClosedTimers[NumClosedTimers];
+
+		AutoRTFM::Testing::Commit([&]
+		{
+			AutoRTFM::Open([&]
+			{
+				for (FTimer& OpenTimer : OpenTimers)
+				{
+					TimerManager.SetTimer(OpenTimer.Handle, OpenTimer.Callback(), /* InRate */ 1.0f, FTimerManagerTimerParameters{});
+					check(TimerManager.IsTimerActive(OpenTimer.Handle));
+				}
+			});
+
+			for (FTimer& ClosedTimer : ClosedTimers)
+			{
+				TimerManager.SetTimer(ClosedTimer.Handle, ClosedTimer.Callback(), /* InRate */ 1.0f, FTimerManagerTimerParameters{});
+				check(TimerManager.IsTimerActive(ClosedTimer.Handle));
+			}
+			
+			AutoRTFM::Open([&]
+			{
+				for (FTimer& OpenTimer : OpenTimers)
+				{
+					check(TimerManager.IsTimerActive(OpenTimer.Handle));
+					TimerManager.ClearTimer(OpenTimer.Handle);
+					check(!TimerManager.IsTimerActive(OpenTimer.Handle));
+				}
+			});
+		});
+
+		Tick(TimerManager);
+
+		for (FTimer& OpenTimer : OpenTimers)
+		{
+			check(!OpenTimer.bCalled);
+		}
+
+		for (FTimer& ClosedTimer : ClosedTimers)
+		{
+			check(ClosedTimer.bCalled);
+		}
+
+		return true;
+	});
+
+	Section(TEXT("Transact(Open(SetTimer), SetTimer, Open(ClearTimer), Abort), Tick"), [&] (FTimerManager& TimerManager)
+	{
+		constexpr size_t NumOpenTimers = 16;
+		FTimer OpenTimers[NumOpenTimers];
+		constexpr size_t NumClosedTimers = 16;
+		FTimer ClosedTimers[NumClosedTimers];
+
+		AutoRTFM::Testing::Abort([&]
+		{
+			AutoRTFM::Open([&]
+			{
+				for (FTimer& OpenTimer : OpenTimers)
+				{
+					TimerManager.SetTimer(OpenTimer.Handle, OpenTimer.Callback(), /* InRate */ 1.0f, FTimerManagerTimerParameters{});
+					check(TimerManager.IsTimerActive(OpenTimer.Handle));
+				}
+			});
+
+			for (FTimer& ClosedTimer : ClosedTimers)
+			{
+				TimerManager.SetTimer(ClosedTimer.Handle, ClosedTimer.Callback(), /* InRate */ 1.0f, FTimerManagerTimerParameters{});
+				check(TimerManager.IsTimerActive(ClosedTimer.Handle));
+			}
+			
+			AutoRTFM::Open([&]
+			{
+				for (FTimer& OpenTimer : OpenTimers)
+				{
+					check(TimerManager.IsTimerActive(OpenTimer.Handle));
+					TimerManager.ClearTimer(OpenTimer.Handle);
+					check(!TimerManager.IsTimerActive(OpenTimer.Handle));
+				}
+			});
+			
+			AutoRTFM::AbortTransaction();
+		});
+
+		Tick(TimerManager);
+
+		for (FTimer& OpenTimer : OpenTimers)
+		{
+			check(!OpenTimer.bCalled);
+		}
+
+		for (FTimer& ClosedTimer : ClosedTimers)
+		{
+			check(!ClosedTimer.bCalled);
+		}
+
+		return true;
+	});
+
+	Section(TEXT("SetTimer, Transact(ClearTimer), Tick"), [&] (FTimerManager& TimerManager)
+	{
+		constexpr size_t NumTimers = 16;
+		FTimer Timers[NumTimers]{};
+
+		for (FTimer& Timer : Timers)
+		{
+			TimerManager.SetTimer(Timer.Handle, Timer.Callback(), /* InRate */ 1.0f, FTimerManagerTimerParameters{});
+			check(TimerManager.IsTimerActive(Timer.Handle));
+		}
+
+		AutoRTFM::Testing::Commit([&]
+		{
+			for (size_t I = 0; I < NumTimers; I++)
+			{
+				if ((I & 1) == 0)
+				{
+					FTimer& Timer = Timers[I];
+					check(TimerManager.IsTimerActive(Timer.Handle));
+					TimerManager.ClearTimer(Timer.Handle);
+					check(!TimerManager.IsTimerActive(Timer.Handle));
+				}
+			}
+		});
+
+		Tick(TimerManager);
+
+		for (size_t I = 0; I < NumTimers; I++)
+		{
+			FTimer& Timer = Timers[I];
+			bool const bCleared = (I & 1) == 0;
+			check(Timer.bCalled == !bCleared);
+		}
+
+		return true;
+	});
+
+	Section(TEXT("SetTimer, Transact(ClearTimer, Abort), Tick"), [&] (FTimerManager& TimerManager)
+	{
+		constexpr size_t NumTimers = 16;
+		FTimer Timers[NumTimers]{};
+
+		for (FTimer& Timer : Timers)
+		{
+			TimerManager.SetTimer(Timer.Handle, Timer.Callback(), /* InRate */ 1.0f, FTimerManagerTimerParameters{});
+			check(TimerManager.IsTimerActive(Timer.Handle));
+		}
+
+		AutoRTFM::Testing::Abort([&]
+		{
+			for (size_t I = 0; I < NumTimers; I++)
+			{
+				FTimer& Timer = Timers[I];
+				TimerManager.ClearTimer(Timer.Handle);
+			}
+			AutoRTFM::AbortTransaction();
+		});
+
+		Tick(TimerManager);
+
+		for (size_t I = 0; I < NumTimers; I++)
+		{
+			FTimer& Timer = Timers[I];
+			check(Timer.bCalled);
+		}
+
+		return true;
+	});
+
+	Section(TEXT("SetTimer, Transact(PauseTimer), Tick"), [&] (FTimerManager& TimerManager)
+	{
+		constexpr size_t NumTimers = 16;
+		FTimer Timers[NumTimers]{};
+
+		for (FTimer& Timer : Timers)
+		{
+			TimerManager.SetTimer(Timer.Handle, Timer.Callback(), /* InRate */ 1.0f, FTimerManagerTimerParameters{});
+			check(TimerManager.IsTimerActive(Timer.Handle));
+		}
+
+		AutoRTFM::Testing::Commit([&]
+		{
+			for (size_t I = 0; I < NumTimers; I++)
+			{
+				if ((I & 1) == 0)
+				{
+					FTimer& Timer = Timers[I];
+					check(!TimerManager.IsTimerPaused(Timer.Handle));
+					TimerManager.PauseTimer(Timer.Handle);
+					check(TimerManager.IsTimerPaused(Timer.Handle));
+				}
+			}
+		});
+
+		Tick(TimerManager);
+
+		for (size_t I = 0; I < NumTimers; I++)
+		{
+			FTimer& Timer = Timers[I];
+			bool const bPaused = (I & 1) == 0;
+			check(Timer.bCalled == !bPaused);
+			check(TimerManager.IsTimerPaused(Timer.Handle) == bPaused);
+		}
+
+		return true;
+	});
+
+	Section(TEXT("SetTimer, Transact(PauseTimer, Abort), Tick"), [&] (FTimerManager& TimerManager)
+	{
+		constexpr size_t NumTimers = 16;
+		FTimer Timers[NumTimers]{};
+
+		for (FTimer& Timer : Timers)
+		{
+			TimerManager.SetTimer(Timer.Handle, Timer.Callback(), /* InRate */ 1.0f, FTimerManagerTimerParameters{});
+			check(TimerManager.IsTimerActive(Timer.Handle));
+		}
+
+		AutoRTFM::Testing::Abort([&]
+		{
+			for (size_t I = 0; I < NumTimers; I++)
+			{
+				if ((I & 1) == 0)
+				{
+					FTimer& Timer = Timers[I];
+					TimerManager.PauseTimer(Timer.Handle);
+					check(TimerManager.IsTimerPaused(Timer.Handle));
+				}
+			}
+			AutoRTFM::AbortTransaction();
+		});
+
+		Tick(TimerManager);
+
+		for (size_t I = 0; I < NumTimers; I++)
+		{
+			FTimer& Timer = Timers[I];
+			check(Timer.bCalled);
+			check(!TimerManager.IsTimerPaused(Timer.Handle));
+		}
+
+		return true;
+	});
+
+	Section(TEXT("SetTimer, PauseTimer, Transact(UnPauseTimer), Tick"), [&] (FTimerManager& TimerManager)
+	{
+		constexpr size_t NumTimers = 16;
+		FTimer Timers[NumTimers]{};
+
+		for (FTimer& Timer : Timers)
+		{
+			TimerManager.SetTimer(Timer.Handle, Timer.Callback(), /* InRate */ 1.0f, FTimerManagerTimerParameters{});
+			check(TimerManager.IsTimerActive(Timer.Handle));
+			TimerManager.PauseTimer(Timer.Handle);
+			check(TimerManager.IsTimerPaused(Timer.Handle));
+		}
+
+		AutoRTFM::Testing::Commit([&]
+		{
+			for (size_t I = 0; I < NumTimers; I++)
+			{
+				if ((I & 1) == 0)
+				{
+					FTimer& Timer = Timers[I];
+					check(TimerManager.IsTimerPaused(Timer.Handle));
+					TimerManager.UnPauseTimer(Timer.Handle);
+					check(!TimerManager.IsTimerPaused(Timer.Handle));
+				}
+			}
+		});
+
+		Tick(TimerManager);
+
+		for (size_t I = 0; I < NumTimers; I++)
+		{
+			FTimer& Timer = Timers[I];
+			bool const bPaused = (I & 1) != 0;
+			check(Timer.bCalled == !bPaused);
+			check(TimerManager.IsTimerPaused(Timer.Handle) == bPaused);
+		}
+
+		return true;
+	});
+
+	Section(TEXT("SetTimer, PauseTimer, Transact(UnPauseTimer, Abort), Tick"), [&] (FTimerManager& TimerManager)
+	{
+		constexpr size_t NumTimers = 16;
+		FTimer Timers[NumTimers]{};
+
+		for (FTimer& Timer : Timers)
+		{
+			TimerManager.SetTimer(Timer.Handle, Timer.Callback(), /* InRate */ 1.0f, FTimerManagerTimerParameters{});
+			check(TimerManager.IsTimerActive(Timer.Handle));
+			TimerManager.PauseTimer(Timer.Handle);
+			check(TimerManager.IsTimerPaused(Timer.Handle));
+		}
+
+		AutoRTFM::Testing::Abort([&]
+		{
+			for (size_t I = 0; I < NumTimers; I++)
+			{
+				if ((I & 1) == 0)
+				{
+					FTimer& Timer = Timers[I];
+					check(TimerManager.IsTimerPaused(Timer.Handle));
+					TimerManager.UnPauseTimer(Timer.Handle);
+					check(!TimerManager.IsTimerPaused(Timer.Handle));
+				}
+			}
+			AutoRTFM::AbortTransaction();
+		});
+
+		Tick(TimerManager);
+
+		for (size_t I = 0; I < NumTimers; I++)
+		{
+			FTimer& Timer = Timers[I];
+			check(!Timer.bCalled);
+			check(TimerManager.IsTimerPaused(Timer.Handle));
+		}
+
+		return true;
+	});
+
+	Section(TEXT("Transact(SetTimer, ClearTimer, Open(SetTimer), Abort), Tick"), [&] (FTimerManager& TimerManager)
+	{
+		constexpr size_t NumTimers = 16;
+		FTimer ClosedTimers[NumTimers]{};
+		FTimer OpenTimers[NumTimers]{};
+
+		AutoRTFM::Testing::Abort([&]
+		{
+			for (FTimer& Timer : ClosedTimers)
+			{
+				TimerManager.SetTimer(Timer.Handle, Timer.Callback(), /* InRate */ 1.0f, FTimerManagerTimerParameters{});
+				check(TimerManager.IsTimerActive(Timer.Handle));
+			}
+
+			for (FTimer& Timer : ClosedTimers)
+			{
+				check(TimerManager.IsTimerActive(Timer.Handle));
+				TimerManager.ClearTimer(Timer.Handle);
+				check(!TimerManager.IsTimerActive(Timer.Handle));
+			}
+
+			AutoRTFM::Open([&]
+			{
+				for (FTimer& Timer : OpenTimers)
+				{
+					TimerManager.SetTimer(Timer.Handle, Timer.Callback(), /* InRate */ 1.0f, FTimerManagerTimerParameters{});
+				}
+			});
+
+			AutoRTFM::AbortTransaction();
+		});
+
+		Tick(TimerManager);
+
+		for (FTimer& Timer : ClosedTimers)
+		{
+			check(!Timer.bCalled);
+		}
+		for (FTimer& Timer : OpenTimers)
+		{
+			check(Timer.bCalled);
+		}
+
+		return true;
+	});
+
+	return true;
+}
+
+#endif //WITH_DEV_AUTOMATION_TESTS
diff --git a/Engine/Source/Runtime/Engine/Private/TimerManager.cpp b/Engine/Source/Runtime/Engine/Private/TimerManager.cpp
--- a/Engine/Source/Runtime/Engine/Private/TimerManager.cpp
+++ b/Engine/Source/Runtime/Engine/Private/TimerManager.cpp
@@ -5,9 +5,12 @@
 =============================================================================*/
 
 #include "TimerManager.h"
+#include "AutoRTFM.h"
 #include "Containers/StringConv.h"
 #include "Engine/GameInstance.h"
+#include "Engine/TimerHandle.h"
 #include "Logging/LogScopedCategoryAndVerbosityOverride.h"
+#include "Misc/AssertionMacros.h"
 #include "Misc/CoreDelegates.h"
 #include "Misc/StringOutputDevice.h"
 #include "ProfilingDebugging/CsvProfiler.h"
@@ -327,6 +330,7 @@ FTimerManager::FTimerManager(UGameInstance* GameInstance)
 
 FTimerManager::~FTimerManager()
 {
+	checkf(!AutoRTFM::IsTransactional(), TEXT("FTimerManager must not be destructed while inside a transaction"));
 	if (IsRunningDedicatedServer())
 	{
 		FCoreDelegates::OnHandleSystemError.RemoveAll(this);
@@ -676,7 +680,7 @@ void FTimerManager::InternalSetTimer(FTimerHandle& InOutHandle, FTimerUnifiedDel
 			NewTimerData.ExpireTime = InternalTime + FirstDelay;
 			NewTimerData.Status = ETimerStatus::Active;
 			NewTimerHandle = AddTimer(MoveTemp(NewTimerData));
-			ActiveTimerHeap.HeapPush(NewTimerHandle, FTimerHeapOrder(Timers));
+			PushToActiveTimerHeap(NewTimerHandle);
 		}
 		else
 		{
@@ -684,7 +688,7 @@ void FTimerManager::InternalSetTimer(FTimerHandle& InOutHandle, FTimerUnifiedDel
 			NewTimerData.ExpireTime = FirstDelay;
 			NewTimerData.Status = ETimerStatus::Pending;
 			NewTimerHandle = AddTimer(MoveTemp(NewTimerData));
-			PendingTimerSet.Add(NewTimerHandle);
+			AddToPendingTimerSet(NewTimerHandle);
 		}
 
 		InOutHandle = NewTimerHandle;
@@ -724,14 +728,14 @@ FTimerHandle FTimerManager::InternalSetTimerForNextTick(FTimerUnifiedDelegate&&
 	{
 		NewTimerData.Status = ETimerStatus::Active;
 		NewTimerHandle = AddTimer(MoveTemp(NewTimerData));
-		ActiveTimerHeap.HeapPush(NewTimerHandle, FTimerHeapOrder(Timers));
+		PushToActiveTimerHeap(NewTimerHandle);
 	}
 	else
 	{
 		NewTimerData.Status = ETimerStatus::Pending;
 		NewTimerData.ExpireTime = 0.0;	// Pending timers' ExpireTime means time remaining. It will be converted upon activation.
 		NewTimerHandle = AddTimer(MoveTemp(NewTimerData));
-		PendingTimerSet.Add(NewTimerHandle);
+		AddToPendingTimerSet(NewTimerHandle);
 	}
 
 	return NewTimerHandle;
@@ -748,15 +752,12 @@ void FTimerManager::InternalClearTimer(FTimerHandle InHandle)
 	switch (Data.Status)
 	{
 		case ETimerStatus::Pending:
-			{
-				int32 NumRemoved = PendingTimerSet.Remove(InHandle);
-				check(NumRemoved == 1);
-				RemoveTimer(InHandle);
-			}
+			RemoveFromPendingTimerSet(InHandle);
+			RemoveTimer(InHandle);
 			break;
 
 		case ETimerStatus::Active:
-			Data.Status = ETimerStatus::ActivePendingRemoval;
+			SetTimerStatus(Data, ETimerStatus::ActivePendingRemoval);
 			break;
 
 		case ETimerStatus::ActivePendingRemoval:
@@ -764,11 +765,8 @@ void FTimerManager::InternalClearTimer(FTimerHandle InHandle)
 			break;
 
 		case ETimerStatus::Paused:
-			{
-				int32 NumRemoved = PausedTimerSet.Remove(InHandle);
-				check(NumRemoved == 1);
-				RemoveTimer(InHandle);
-			}
+			RemoveFromPausedTimerSet(InHandle);
+			RemoveTimer(InHandle);
 			break;
 
 		case ETimerStatus::Executing:
@@ -776,7 +774,7 @@ void FTimerManager::InternalClearTimer(FTimerHandle InHandle)
 
 			// Edge case. We're currently handling this timer when it got cleared.  Clear it to prevent it firing again
 			// in case it was scheduled to fire multiple times.
-			CurrentlyExecutingTimer.Invalidate();
+			InvalidateCurrentlyExecutingTimer();
 			RemoveTimer(InHandle);
 			break;
 
@@ -876,24 +874,17 @@ void FTimerManager::PauseTimer(FTimerHandle InHandle)
 			break;
 
 		case ETimerStatus::Active:
-			{
-				int32 IndexIndex = ActiveTimerHeap.Find(InHandle);
-				check(IndexIndex != INDEX_NONE);
-				ActiveTimerHeap.HeapRemoveAt(IndexIndex, FTimerHeapOrder(Timers), EAllowShrinking::No);
-			}
+			RemoveFromActiveTimerHeap(InHandle);
 			break;
 
 		case ETimerStatus::Pending:
-			{
-				int32 NumRemoved = PendingTimerSet.Remove(InHandle);
-				check(NumRemoved == 1);
-			}
+			RemoveFromPendingTimerSet(InHandle);
 			break;
 
 		case ETimerStatus::Executing:
 			check(CurrentlyExecutingTimer == InHandle);
 
-			CurrentlyExecutingTimer.Invalidate();
+			InvalidateCurrentlyExecutingTimer();
 			break;
 
 		default:
@@ -908,15 +899,20 @@ void FTimerManager::PauseTimer(FTimerHandle InHandle)
 	else
 	{
 		// Add to Paused list
-		PausedTimerSet.Add(InHandle);
-
-		// Set new status
-		TimerToPause->Status = ETimerStatus::Paused;
+		AddToPausedTimerSet(InHandle);
 
-		// Store time remaining in ExpireTime while paused. Don't do this if the timer is in the pending list.
-		if (PreviousStatus != ETimerStatus::Pending)
+		if (PreviousStatus == ETimerStatus::Pending)
 		{
-			TimerToPause->ExpireTime -= InternalTime;
+			// Set new status
+			SetTimerStatus(*TimerToPause, ETimerStatus::Paused);
+		}
+		else
+		{
+			// Store time remaining in ExpireTime while paused.
+			// Don't do this if the timer is in the pending list.
+			const double NewExpireTime = TimerToPause->ExpireTime - InternalTime;
+			// Set new status and expire time
+			SetTimerStatusAndExpireTime(*TimerToPause, ETimerStatus::Paused, NewExpireTime);
 		}
 	}
 }
@@ -936,18 +932,18 @@ void FTimerManager::UnPauseTimer(FTimerHandle InHandle)
 	if( HasBeenTickedThisFrame() )
 	{
 		// Convert from time remaining back to a valid ExpireTime
-		TimerToUnPause->ExpireTime += InternalTime;
-		TimerToUnPause->Status = ETimerStatus::Active;
-		ActiveTimerHeap.HeapPush(InHandle, FTimerHeapOrder(Timers));
+		const double NewExpireTime = TimerToUnPause->ExpireTime + InternalTime;
+		SetTimerStatusAndExpireTime(*TimerToUnPause, ETimerStatus::Active, NewExpireTime);
+		PushToActiveTimerHeap(InHandle);
 	}
 	else
 	{
-		TimerToUnPause->Status = ETimerStatus::Pending;
-		PendingTimerSet.Add(InHandle);
+		SetTimerStatus(*TimerToUnPause, ETimerStatus::Pending);
+		AddToPendingTimerSet(InHandle);
 	}
 
 	// remove from paused list
-	PausedTimerSet.Remove(InHandle);
+	RemoveFromPausedTimerSet(InHandle);
 }
 
 FTimerData::FTimerData()
@@ -968,6 +964,8 @@ DECLARE_DWORD_COUNTER_STAT(TEXT("TimerManager Heap Size"),STAT_NumHeapEntries,ST
 
 void FTimerManager::Tick(float DeltaTime)
 {
+	AutoRTFM::UnreachableIfClosed("FTimerManager::Tick() cannot be called from the closed");
+
 	SCOPED_NAMED_EVENT(FTimerManager_Tick, FColor::Orange);
 	CSV_SCOPED_TIMING_STAT_EXCLUSIVE(TimerManager);
 
@@ -1265,46 +1263,86 @@ void FTimerManager::ListTimers() const
 
 FTimerHandle FTimerManager::AddTimer(FTimerData&& TimerData)
 {
-	const void* TimerIndicesByObjectKey = TimerData.TimerDelegate.GetBoundObject();
-	TimerData.TimerIndicesByObjectKey = TimerIndicesByObjectKey;
-
-	int32 NewIndex = Timers.Add(MoveTemp(TimerData));
+	const ETimerStatus Status = TimerData.Status;
 
-	FTimerHandle Result = GenerateHandle(NewIndex);
-	Timers[NewIndex].Handle = Result;
+	FTimerHandle Handle = TimerData.Handle;
 
-	if (TimerIndicesByObjectKey)
+	UE_AUTORTFM_OPEN
 	{
-		TSet<FTimerHandle>& HandleSet = ObjectToTimers.FindOrAdd(TimerIndicesByObjectKey);
+		const void* TimerIndicesByObjectKey = TimerData.TimerDelegate.GetBoundObject();
+		TimerData.TimerIndicesByObjectKey = TimerIndicesByObjectKey;
 
-		bool bAlreadyExists = false;
-		HandleSet.Add(Result, &bAlreadyExists);
-		checkf(!bAlreadyExists, TEXT("A timer with this handle and object has already been added! (%s)"), *GetFTimerDataSafely(TimerData));
-	}
+		if (Handle.IsValid())
+		{
+			check(Timers[Handle.GetIndex()].Handle == Handle);
+			Timers.EmplaceAt(Handle.GetIndex(), MoveTemp(TimerData));
+		}
+		else
+		{
+			int32 Index = Timers.Add(MoveTemp(TimerData));
+			Handle = GenerateHandle(Index);
+			Timers[Handle.GetIndex()].Handle = Handle;
+		}
 
-	return Result;
+		if (TimerIndicesByObjectKey)
+		{
+			TSet<FTimerHandle>& HandleSet = ObjectToTimers.FindOrAdd(TimerIndicesByObjectKey);
+
+			bool bAlreadyExists = false;
+			HandleSet.Add(Handle, &bAlreadyExists);
+			checkf(!bAlreadyExists, TEXT("A timer with this handle and object has already been added! (%s)"), *GetFTimerDataSafely(TimerData));
+		}
+	};
+
+	UE_AUTORTFM_ONABORT(this, Handle, Status)
+	{
+		RemoveTimer(Handle);
+	};
+
+	return Handle;
 }
 
 void FTimerManager::RemoveTimer(FTimerHandle Handle)
 {
-	const FTimerData& Data = GetTimer(Handle);
+	FTimerData Data = MoveTemp(GetTimer(Handle));
+	check(Data.Handle == Handle);
 
-	// Remove TimerIndicesByObject entry if necessary
-	if (const void* TimerIndicesByObjectKey = Data.TimerIndicesByObjectKey)
+	UE_AUTORTFM_OPEN
 	{
-		TSet<FTimerHandle>* TimersForObject = ObjectToTimers.Find(TimerIndicesByObjectKey);
-		checkf(TimersForObject, TEXT("Removed timer was bound to an object which is not tracked by ObjectToTimers! (%s)"), *GetFTimerDataSafely(Data));
+		// Remove TimerIndicesByObject entry if necessary
+		if (const void* TimerIndicesByObjectKey = Data.TimerIndicesByObjectKey)
+		{
+			TSet<FTimerHandle>* TimersForObject = ObjectToTimers.Find(TimerIndicesByObjectKey);
+			checkf(TimersForObject, TEXT("Removed timer was bound to an object which is not tracked by ObjectToTimers! (%s)"), *GetFTimerDataSafely(Data));
 
-		int32 NumRemoved = TimersForObject->Remove(Handle);
-		checkf(NumRemoved == 1, TEXT("Removed timer was bound to an object which is not tracked by ObjectToTimers! (%s)"), *GetFTimerDataSafely(Data));
+			int32 NumRemoved = TimersForObject->Remove(Handle);
+			checkf(NumRemoved == 1, TEXT("Removed timer was bound to an object which is not tracked by ObjectToTimers! (%s)"), *GetFTimerDataSafely(Data));
 
-		if (TimersForObject->Num() == 0)
-		{
-			ObjectToTimers.Remove(TimerIndicesByObjectKey);
+			if (TimersForObject->Num() == 0)
+			{
+				ObjectToTimers.Remove(TimerIndicesByObjectKey);
+			}
 		}
-	}
+	};
 
-	Timers.RemoveAt(Handle.GetIndex());
+	// When in the closed, postpone removal of the timer until the transaction
+	// commits. We do this to ensure that the array index is not reused within
+	// the transaction.
+	// This is done immediately when not transactional.
+	UE_AUTORTFM_ONCOMMIT(this, Handle)
+	{
+		Timers.RemoveAt(Handle.GetIndex());
+	};
+
+	// We cannot move the FTimerData into the on-abort handler as
+	// AutoRTFM::TTask requires copy-construction due to complexities around
+	// closed -> open interop, so as a workaround we need to wrap it with a
+	// shared pointer. Not ideal, but only affects the code when used under
+	// transactions.
+	UE_AUTORTFM_ONABORT(this, Data = MakeShared<FTimerData>(MoveTemp(Data)))
+	{
+		AddTimer(MoveTemp(Data.Get()));
+	};
 }
 
 bool FTimerManager::WillRemoveTimerAssert(FTimerHandle Handle) const
@@ -1343,6 +1381,120 @@ FTimerHandle FTimerManager::GenerateHandle(int32 Index)
 	return Result;
 }
 
+void FTimerManager::PushToActiveTimerHeap(FTimerHandle Handle)
+{
+	UE_AUTORTFM_ONABORT(this, Handle)
+	{
+		RemoveFromActiveTimerHeap(Handle);
+	};
+	UE_AUTORTFM_OPEN
+	{
+		ActiveTimerHeap.HeapPush(Handle, FTimerHeapOrder(Timers));
+	};
+}
+
+void FTimerManager::RemoveFromActiveTimerHeap(FTimerHandle Handle)
+{
+	UE_AUTORTFM_ONABORT(this, Handle)
+	{
+		PushToActiveTimerHeap(Handle);
+	};
+	UE_AUTORTFM_OPEN
+	{
+		int32 IndexIndex = ActiveTimerHeap.Find(Handle);
+		check(IndexIndex != INDEX_NONE);
+		ActiveTimerHeap.HeapRemoveAt(IndexIndex, FTimerHeapOrder(Timers), EAllowShrinking::No);
+	};
+}
+
+void FTimerManager::AddToPendingTimerSet(FTimerHandle Handle)
+{
+	UE_AUTORTFM_ONABORT(this, Handle)
+	{
+		RemoveFromPendingTimerSet(Handle);
+	};
+	UE_AUTORTFM_OPEN
+	{
+		PendingTimerSet.Add(Handle);
+	};
+}
+
+void FTimerManager::RemoveFromPendingTimerSet(FTimerHandle Handle)
+{
+	UE_AUTORTFM_ONABORT(this, Handle)
+	{
+		AddToPendingTimerSet(Handle);
+	};
+	UE_AUTORTFM_OPEN
+	{
+		int32 NumRemoved = PendingTimerSet.Remove(Handle);
+		check(NumRemoved == 1);
+	};
+}
+
+void FTimerManager::AddToPausedTimerSet(FTimerHandle Handle)
+{
+	UE_AUTORTFM_ONABORT(this, Handle)
+	{
+		RemoveFromPausedTimerSet(Handle);
+	};
+	UE_AUTORTFM_OPEN
+	{
+		PausedTimerSet.Add(Handle);
+	};
+}
+
+void FTimerManager::RemoveFromPausedTimerSet(FTimerHandle Handle)
+{
+	UE_AUTORTFM_ONABORT(this, Handle)
+	{
+		AddToPausedTimerSet(Handle);
+	};
+	UE_AUTORTFM_OPEN
+	{
+		int32 NumRemoved = PausedTimerSet.Remove(Handle);
+		check(NumRemoved == 1);
+	};
+}
+
+void FTimerManager::SetTimerStatus(FTimerData& Timer, ETimerStatus NewStatus)
+{
+	UE_AUTORTFM_ONABORT(this, Handle = Timer.Handle, OldStatus = Timer.Status)
+	{
+		GetTimer(Handle).Status = OldStatus;
+	};
+	UE_AUTORTFM_OPEN
+	{
+		Timer.Status = NewStatus;
+	};
+}
+
+void FTimerManager::SetTimerStatusAndExpireTime(FTimerData& Timer, ETimerStatus NewStatus, double NewExpireTime)
+{
+	UE_AUTORTFM_ONABORT(this, Handle = Timer.Handle, OldStatus = Timer.Status, OldExpireTime = Timer.ExpireTime)
+	{
+		FTimerData& Timer = GetTimer(Handle);
+		Timer.Status = OldStatus;
+		Timer.ExpireTime = OldExpireTime;
+	};
+	UE_AUTORTFM_OPEN
+	{
+		Timer.Status = NewStatus;
+		Timer.ExpireTime = NewExpireTime;
+	};
+}
+
+void FTimerManager::InvalidateCurrentlyExecutingTimer()
+{
+	UE_AUTORTFM_ONABORT(this, Handle = CurrentlyExecutingTimer)
+	{
+		CurrentlyExecutingTimer = Handle;
+	};
+	UE_AUTORTFM_OPEN
+	{
+		CurrentlyExecutingTimer.Invalidate();
+	};
+}
 
 // Handler for ListTimers console command
 static void OnListTimers(UWorld* World)
diff --git a/Engine/Source/Runtime/Engine/Public/TimerManager.h b/Engine/Source/Runtime/Engine/Public/TimerManager.h
--- a/Engine/Source/Runtime/Engine/Public/TimerManager.h
+++ b/Engine/Source/Runtime/Engine/Public/TimerManager.h
@@ -6,6 +6,7 @@
 
 #pragma once
 
+#include "AutoRTFM.h"
 #include "CoreMinimal.h"
 #include "Stats/Stats.h"
 #include "UObject/Object.h"
@@ -128,6 +129,10 @@ struct FTimerManagerTimerParameters
 
 /** 
  * Class to globally manage timers.
+ *
+ * FTimerManager supports transactional usage from both the open and closed, so
+ * long as you do not mix open and closed calls to the same timer. FTimerManager
+ * must not be destructed within a transaction.
  */
 class FTimerManager : public FNoncopyable
 {
@@ -412,7 +417,7 @@ private:
 // This should be private, but needs to be public for testing.
 public:
 	/** Generates a handle for a timer at a given index */
-	ENGINE_API FTimerHandle GenerateHandle(int32 Index);
+	AUTORTFM_OPEN ENGINE_API FTimerHandle GenerateHandle(int32 Index);
 
 // These should be private, but need to be protected so IMPLEMENT_GET_PROTECTED_FUNC works for testing.
 protected:
@@ -440,12 +445,29 @@ private:
 	}
 	ENGINE_API FTimerData& GetTimer(FTimerHandle const& InHandle);
 
-	/** Adds a timer from the Timers list, also updating the TimerIndicesByObject map.  Returns the insertion index. */
+	/**
+	 * Adds a timer to the Timers list, also updating the TimerIndicesByObject map.
+	 * If TimerData.Handle is valid then the Timers array must already contain a
+	 * timer with this handle, which will be replaced with TimerData.
+	 * Returns the TimerData.Handle if valid, otherwise a new generated handle.
+	 */
 	ENGINE_API FTimerHandle AddTimer(FTimerData&& TimerData);
+	
 	/** Removes a timer from the Timers list at the given index, also cleaning up the TimerIndicesByObject map */
 	ENGINE_API void RemoveTimer(FTimerHandle Handle);
 	ENGINE_API bool WillRemoveTimerAssert(FTimerHandle Handle) const;
 
+	/* State mutation functions with transactional rollback handling */
+	ENGINE_API void PushToActiveTimerHeap(FTimerHandle Handle);
+	ENGINE_API void RemoveFromActiveTimerHeap(FTimerHandle Handle);
+	ENGINE_API void AddToPendingTimerSet(FTimerHandle Handle);
+	ENGINE_API void RemoveFromPendingTimerSet(FTimerHandle Handle);
+	ENGINE_API void AddToPausedTimerSet(FTimerHandle Handle);
+	ENGINE_API void RemoveFromPausedTimerSet(FTimerHandle Handle);
+	ENGINE_API void SetTimerStatus(FTimerData& Timer, ETimerStatus NewStatus);
+	ENGINE_API void SetTimerStatusAndExpireTime(FTimerData& Timer, ETimerStatus NewStatus, double NewExpireTime);
+	ENGINE_API void InvalidateCurrentlyExecutingTimer();
+
 	/** The array of timers - all other arrays will index into this */
 	TSparseArray<FTimerData> Timers;
 	/** Heap of actively running timers. */
