- Add tracking of global precached shaders so compile stats are collected in precache stats csv (can be extended with detailed global dump)
- Add correct tracking of all the global and mesh pass PSOs requested in CSV stats
- Use new global tracking for Niagara compute and Global compute PSOs
- Add option to boost very high priority PSOs to AboveNormal thread priority (disabled by default)
- Make not required PSOs low priority

#Synced-CL 50645574

#Synced-CL 50273493

diff --git a/Engine/Plugins/FX/Niagara/Source/Niagara/Private/NiagaraEmitter.cpp b/Engine/Plugins/FX/Niagara/Source/Niagara/Private/NiagaraEmitter.cpp
--- a/Engine/Plugins/FX/Niagara/Source/Niagara/Private/NiagaraEmitter.cpp
+++ b/Engine/Plugins/FX/Niagara/Source/Niagara/Private/NiagaraEmitter.cpp
@@ -3508,6 +3508,10 @@ void UNiagaraEmitter::GenerateStatID() const
 #endif
 }
 
+// Register global PSO collector so compile time and tracked in the stats
+static const TCHAR* NiagaraComputePSOCollectorName = TEXT("NiagaraComputePSOCollector");
+FRegisterGlobalPSOCollectorFunction RegisterNiagaraComputePSOCollector(nullptr, NiagaraComputePSOCollectorName);
+
 FGraphEventArray FVersionedNiagaraEmitterData::PrecacheComputePSOs(const UNiagaraEmitter& NiagaraEmitter)
 {
 	FGraphEventArray PSOReadyGraphTasks;
@@ -3522,16 +3526,16 @@ FGraphEventArray FVersionedNiagaraEmitterData::PrecacheComputePSOs(const UNiagar
 	{
 		return PSOReadyGraphTasks;
 	}
-	
-	FPSOPrecacheRequestParams PSORequestParams;
-	PSORequestParams.InMemoryMode = EPSOPrecacheInMemoryMode::Uncached;
-	PSORequestParams.PrecachePriority = EPSOPrecachePriority::High;
-	PSORequestParams.Name = TEXT("NiagaraEmitter");
-	PSORequestParams.bForcePrecache = true;
 
 	// Compute PSO can fail on some platforms even though the shader compiled successfully this path will wait for the PSO precache to complete before the emitter is allowed to run
 	if (GNiagaraEmitterComputePSOPrecacheMode == 3)
 	{
+		FPSOPrecacheRequestParams PSORequestParams;
+		PSORequestParams.InMemoryMode = EPSOPrecacheInMemoryMode::Uncached;
+		PSORequestParams.PrecachePriority = EPSOPrecachePriority::High;
+		PSORequestParams.Name = TEXT("NiagaraEmitter");
+		PSORequestParams.bForcePrecache = true;
+
 		FGraphEventArray PSOPrecacheEvents;
 		for (int32 i = 0; i < ShaderScript->GetNumPermutations(); ++i)
 		{
@@ -3623,8 +3627,11 @@ FGraphEventArray FVersionedNiagaraEmitterData::PrecacheComputePSOs(const UNiagar
 	// In this mode we either force them to cache or respect that precaching is enabled
 	else if (GNiagaraEmitterComputePSOPrecacheMode == 2 || PipelineStateCache::IsPSOPrecachingEnabled() )
 	{
-		EPSOPrecacheProxyCreationStrategy ProxyCreationStrategy = GetPSOPrecacheProxyCreationStrategy();
-		bool bAddToPSOReadyGraphTasks = GNiagaraEmitterComputePSOPrecacheMode == 1 && ProxyCreationStrategy != EPSOPrecacheProxyCreationStrategy::AlwaysCreate;
+		FPSOPrecacheDataArray PSOInitializers;
+		PSOInitializers.Reserve(ShaderScript->GetNumPermutations());
+		
+		static const int32 NiagaraEmitterPSOCollectorIndex = FGlobalPSOCollectorManager::GetIndex(NiagaraComputePSOCollectorName);
+
 		for (int32 i=0; i < ShaderScript->GetNumPermutations(); ++i)
 		{
 			FNiagaraShaderRef ShaderRef = ShaderScript->GetShaderGameThread(i);
@@ -3633,11 +3640,30 @@ FGraphEventArray FVersionedNiagaraEmitterData::PrecacheComputePSOs(const UNiagar
 				continue;
 			}
 
-			FRHIComputeShader* ComputeShader = ShaderRef.GetComputeShader();
-			FPSOPrecacheRequestResult PSOPrecacheRequestResult = PipelineStateCache::PrecacheComputePipelineState(ComputeShader, PSORequestParams);
-			if (PSOPrecacheRequestResult.AsyncCompileEvent != nullptr && bAddToPSOReadyGraphTasks)
+			FPSOPrecacheData PSOPrecacheData;
+			PSOPrecacheData.Type = FPSOPrecacheData::EType::Compute;
+			PSOPrecacheData.ComputeInitializer = FComputePipelineStateInitializer(ShaderRef.GetComputeShader());
+			PSOPrecacheData.InMemoryMode = EPSOPrecacheInMemoryMode::Uncached;
+			PSOPrecacheData.Priority = EPSOPrecachePriority::High;
+		#if PSO_PRECACHING_VALIDATE
+			PSOPrecacheData.PSOCollectorIndex = NiagaraEmitterPSOCollectorIndex;
+			PSOPrecacheData.VertexFactoryType = nullptr;
+		#endif // PSO_PRECACHING_VALIDATE
+			PSOInitializers.Add(MoveTemp(PSOPrecacheData));
+		}
+
+		FPSOPrecacheRequestResultArray ResultArray = RequestPrecachePSOs(EPSOPrecacheType::Global, PSOInitializers);
+
+		EPSOPrecacheProxyCreationStrategy ProxyCreationStrategy = GetPSOPrecacheProxyCreationStrategy();
+		bool bAddToPSOReadyGraphTasks = GNiagaraEmitterComputePSOPrecacheMode == 1 && ProxyCreationStrategy != EPSOPrecacheProxyCreationStrategy::AlwaysCreate;
+		if (bAddToPSOReadyGraphTasks)
+		{			
+			for (FPSOPrecacheRequestResult PSOPrecacheRequestResult : ResultArray)
 			{
-				PSOReadyGraphTasks.Add(PSOPrecacheRequestResult.AsyncCompileEvent);
+				if (PSOPrecacheRequestResult.AsyncCompileEvent != nullptr && bAddToPSOReadyGraphTasks)
+				{
+					PSOReadyGraphTasks.Add(PSOPrecacheRequestResult.AsyncCompileEvent);
+				}
 			}
 		}
 	}
diff --git a/Engine/Plugins/FX/Niagara/Source/NiagaraShader/Private/NiagaraShared.cpp b/Engine/Plugins/FX/Niagara/Source/NiagaraShader/Private/NiagaraShared.cpp
--- a/Engine/Plugins/FX/Niagara/Source/NiagaraShader/Private/NiagaraShared.cpp
+++ b/Engine/Plugins/FX/Niagara/Source/NiagaraShader/Private/NiagaraShared.cpp
@@ -521,7 +521,12 @@ void FNiagaraShaderScript::UpdateCachedData_PostCompile(bool bCalledFromSerializ
 	CachedData.bViewUniformBufferUsed = 0;
 
 	if (GameThreadShaderMap != nullptr && GameThreadShaderMap->IsValid())
-	{
+	{		
+		FPSOPrecacheDataArray PSOInitializers;
+		PSOInitializers.Reserve(CachedData.NumPermutations);
+
+		static const int32 PSOCollectorIndex = FGlobalPSOCollectorManager::GetIndex(TEXT("NiagaraComputePSOCollector"));
+
 		for (int32 iPermutation = 0; iPermutation < CachedData.NumPermutations; ++iPermutation)
 		{
 			TNiagaraShaderRef<FShader> Shader = GameThreadShaderMap->GetShader(&FNiagaraShader::GetStaticType(), iPermutation);
@@ -554,14 +559,24 @@ void FNiagaraShaderScript::UpdateCachedData_PostCompile(bool bCalledFromSerializ
 					FRHIShader* RHIShader = GameThreadShaderMap->GetResource()->GetShader(Shader->GetResourceIndex());
 					FRHIComputeShader* RHIComputeShader = static_cast<FRHIComputeShader*>(RHIShader);
 
-					FPSOPrecacheRequestParams RequestParams;
-					RequestParams.InMemoryMode = EPSOPrecacheInMemoryMode::Uncached;
-					RequestParams.PrecachePriority = EPSOPrecachePriority::High;
-					RequestParams.Name = TEXT("NiagaraCompute");
-					PipelineStateCache::PrecacheComputePipelineState(RHIComputeShader, RequestParams);
+					FPSOPrecacheData PSOPrecacheData;
+					PSOPrecacheData.Type = FPSOPrecacheData::EType::Compute;
+					PSOPrecacheData.ComputeInitializer = FComputePipelineStateInitializer(RHIComputeShader);
+					PSOPrecacheData.InMemoryMode = EPSOPrecacheInMemoryMode::Uncached;
+					PSOPrecacheData.Priority = EPSOPrecachePriority::High;
+				#if PSO_PRECACHING_VALIDATE
+					PSOPrecacheData.PSOCollectorIndex = PSOCollectorIndex;
+					PSOPrecacheData.VertexFactoryType = nullptr;
+				#endif // PSO_PRECACHING_VALIDATE
+					PSOInitializers.Add(MoveTemp(PSOPrecacheData));
 				}
 			}
 		}
+
+		if (!PSOInitializers.IsEmpty())
+		{
+			RequestPrecachePSOs(EPSOPrecacheType::Global, PSOInitializers);
+		}
 	}
 	else
 	{
diff --git a/Engine/Source/Runtime/Engine/Private/PSOPrecache.cpp b/Engine/Source/Runtime/Engine/Private/PSOPrecache.cpp
--- a/Engine/Source/Runtime/Engine/Private/PSOPrecache.cpp
+++ b/Engine/Source/Runtime/Engine/Private/PSOPrecache.cpp
@@ -225,7 +225,7 @@ int32 FGlobalPSOCollectorManager::GetIndex(const TCHAR* Name)
 }
 
 /**
- * Start the actual PSO precache tasks for all the initializers provided and return the request result array containing the graph event on which to wait for the PSOs to finish compiling
+ * Start the actual PSO precache tasks for all the initializers provided and return the request result array containing the graph event on which to wait for the PSOs to finish compiling (doesn't add not required PSOs to the graph event list)
  */
 FPSOPrecacheRequestResultArray RequestPrecachePSOs(EPSOPrecacheType PSOPrecacheType, const FPSOPrecacheDataArray& PSOInitializers)
 {
@@ -234,31 +234,28 @@ FPSOPrecacheRequestResultArray RequestPrecachePSOs(EPSOPrecacheType PSOPrecacheT
 	{
 		FPSOPrecacheRequestParams PSORequestParams;
 		PSORequestParams.InMemoryMode = PrecacheData.InMemoryMode;
-		PSORequestParams.PrecachePriority = PrecacheData.Priority;
+		// Make low priority when not required
+		PSORequestParams.PrecachePriority = PrecacheData.bRequired ? PrecacheData.Priority : EPSOPrecachePriority::Low;
+		PSORequestParams.Name = PrecacheData.Name;
 
+		bool bUpdateStats = false;
 		FPSOPrecacheRequestResult PSOPrecacheResult;
 		switch (PrecacheData.Type)
 		{
 			case FPSOPrecacheData::EType::Graphics:
 			{
-				#if PSO_PRECACHING_VALIDATE
-				if (PSOCollectorStats::GetFullPSOPrecacheStatsCollector().AddStateToCache(PSOPrecacheType, PrecacheData.GraphicsPSOInitializer, PSOCollectorStats::GetPSOPrecacheHash, nullptr, PrecacheData.PSOCollectorIndex, PrecacheData.VertexFactoryType))
-				{
-					CSV_CUSTOM_STAT(PSOPrecache, PrecachedGraphics, 1, ECsvCustomStatOp::Accumulate);
-				}
-				#endif // PSO_PRECACHING_VALIDATE			
+			#if PSO_PRECACHING_VALIDATE
+				bUpdateStats = PSOCollectorStats::GetFullPSOPrecacheStatsCollector().AddStateToCache(PSOPrecacheType, PrecacheData.GraphicsPSOInitializer, PSOCollectorStats::GetPSOPrecacheHash, nullptr, PrecacheData.PSOCollectorIndex, PrecacheData.VertexFactoryType);
+			#endif // PSO_PRECACHING_VALIDATE			
 				
 				PSOPrecacheResult = PipelineStateCache::PrecacheGraphicsPipelineState(PrecacheData.GraphicsPSOInitializer, PSORequestParams);
 				break;
 			}
 			case FPSOPrecacheData::EType::Compute:
 			{
-				#if PSO_PRECACHING_VALIDATE
-				if (PSOCollectorStats::GetFullPSOPrecacheStatsCollector().AddStateToCache(PSOPrecacheType, PrecacheData.ComputeInitializer, PSOCollectorStats::GetPSOPrecacheHash, nullptr, PrecacheData.PSOCollectorIndex, nullptr))
-				{
-					CSV_CUSTOM_STAT(PSOPrecache, PrecachedCompute, 1, ECsvCustomStatOp::Accumulate);
-				}
-				#endif // PSO_PRECACHING_VALIDATE
+			#if PSO_PRECACHING_VALIDATE
+				bUpdateStats = PSOCollectorStats::GetFullPSOPrecacheStatsCollector().AddStateToCache(PSOPrecacheType, PrecacheData.ComputeInitializer, PSOCollectorStats::GetPSOPrecacheHash, nullptr, PrecacheData.PSOCollectorIndex, nullptr);
+			#endif // PSO_PRECACHING_VALIDATE
 			
 				PSOPrecacheResult = PipelineStateCache::PrecacheComputePipelineState(PrecacheData.ComputeInitializer, PSORequestParams);
 				break;
@@ -268,6 +265,14 @@ FPSOPrecacheRequestResultArray RequestPrecachePSOs(EPSOPrecacheType PSOPrecacheT
 		{
 			Results.AddUnique(PSOPrecacheResult);
 		}
+			
+#if PSO_PRECACHING_VALIDATE
+		if (bUpdateStats)
+		{			
+			// Only update stats when updated (first hit)
+			PSOCollectorStats::UpdateGlobalStats(PSOPrecacheType, PrecacheData);
+		}
+#endif // PSO_PRECACHING_VALIDATE
 	}
 
 	return Results;
diff --git a/Engine/Source/Runtime/Engine/Private/PSOPrecacheMaterial.cpp b/Engine/Source/Runtime/Engine/Private/PSOPrecacheMaterial.cpp
--- a/Engine/Source/Runtime/Engine/Private/PSOPrecacheMaterial.cpp
+++ b/Engine/Source/Runtime/Engine/Private/PSOPrecacheMaterial.cpp
@@ -1065,19 +1065,19 @@ FGraphEventArray GetPrecacheRequestsCompletePrereqs(const TArray<FMaterialPSOPre
 #if CSV_PROFILER_STATS
 CSV_DECLARE_CATEGORY_MODULE_EXTERN(RHI_API, PSOPrecacheCompiling);
 
-struct FPSOPrecacheStats
+struct FMaterialPSOPrecacheStats
 {
 	TCsvPersistentCustomStat<int>* ProxiesWaitingForMaterialPSOs = FCsvProfiler::Get()->GetOrCreatePersistentCustomStatInt(TEXT("ProxiesWaitingForMaterialPSOs"), CSV_CATEGORY_INDEX(PSOPrecacheCompiling));
 
-	static FPSOPrecacheStats& GetStats()
+	static FMaterialPSOPrecacheStats& GetStats()
 	{
-		static FPSOPrecacheStats PersistentStats;
+		static FMaterialPSOPrecacheStats PersistentStats;
 		return PersistentStats;
 	}
 
 	static bool AreEnabled()
 	{
-		return FCsvProfiler::Get()->IsCategoryEnabled(CSV_CATEGORY_INDEX(PSOPrecacheCompiling));
+		return FCsvProfiler::Get() && FCsvProfiler::Get()->IsCategoryEnabled(CSV_CATEGORY_INDEX(PSOPrecacheCompiling));
 	}
 };
 #endif
@@ -1093,14 +1093,14 @@ void BoostPSOPriorityOnDraw(const TArray<FMaterialPSOPrecacheRequestID>& Materia
 	if(!AllRequests.IsEmpty())
 	{ 
 #if CSV_PROFILER_STATS
-		if (FPSOPrecacheStats::AreEnabled())
+		if (FMaterialPSOPrecacheStats::AreEnabled())
 		{
-			FPSOPrecacheStats::GetStats().ProxiesWaitingForMaterialPSOs->Add(1);
+			FMaterialPSOPrecacheStats::GetStats().ProxiesWaitingForMaterialPSOs->Add(1);
 
 			FFunctionGraphTask::CreateAndDispatchWhenReady(
 				[]
 				{
-					FPSOPrecacheStats::GetStats().ProxiesWaitingForMaterialPSOs->Add(-1);
+					FMaterialPSOPrecacheStats::GetStats().ProxiesWaitingForMaterialPSOs->Add(-1);
 				},
 				TStatId{}, &AllRequests
 			);
diff --git a/Engine/Source/Runtime/Engine/Private/PSOPrecacheValidation.cpp b/Engine/Source/Runtime/Engine/Private/PSOPrecacheValidation.cpp
--- a/Engine/Source/Runtime/Engine/Private/PSOPrecacheValidation.cpp
+++ b/Engine/Source/Runtime/Engine/Private/PSOPrecacheValidation.cpp
@@ -243,6 +243,28 @@ void PSOCollectorStats::FPrecacheUsageData::UpdateStats(EPSOPrecacheType Precach
 	}
 }
 
+uint64 PSOCollectorStats::FPrecacheUsageData::GetTotalMeshPassCount()
+{
+	FScopeLock Lock(&StatsLock);
+	uint64 TotalCount = UntrackedMeshPassCount;
+	for (int CollectorIndex = 0; CollectorIndex < FPSOCollectorCreateManager::MaxPSOCollectorCount; ++CollectorIndex)
+	{
+		TotalCount += PerMeshPassCount[CollectorIndex];
+	}
+	return TotalCount;
+}
+
+uint64 PSOCollectorStats::FPrecacheUsageData::GetTotalGlobalCount()
+{
+	FScopeLock Lock(&StatsLock);
+	uint64 TotalCount = UntrackedGlobalTypeCount;
+	for (int CollectorIndex = 0; CollectorIndex < FGlobalPSOCollectorManager::MaxPSOCollectorCount; ++CollectorIndex)
+	{
+		TotalCount += PerGlobalTypeCount[CollectorIndex];
+	}
+	return TotalCount;
+}
+
 bool PSOCollectorStats::FPrecacheStatsCollector::IsStateTracked(EPSOPrecacheType PrecacheType, int32 PSOCollectorIndex, const FVertexFactoryType* VertexFactoryType) 
 {
 #if UE_BUILD_SHIPPING || UE_BUILD_TEST
@@ -376,6 +398,95 @@ PSOCollectorStats::FPrecacheStatsCollector& PSOCollectorStats::GetFullPSOPrecach
 	return FullPSOPrecacheStatsCollector;
 }
 
+CSV_DECLARE_CATEGORY_MODULE_EXTERN(RHI_API, PSOPrecacheCompiling);
+
+struct FGlobalPSOPrecacheStats
+{	
+#if CSV_PROFILER_STATS
+
+	TCsvPersistentCustomStat<int>* TotalGlobalGraphicPSOPrecacheRequests 	= FCsvProfiler::Get()->GetOrCreatePersistentCustomStatInt(TEXT("TotalGlobalGraphicPSOPrecacheRequests"), CSV_CATEGORY_INDEX(PSOPrecacheCompiling));
+	TCsvPersistentCustomStat<int>* TotalGlobalComputePSOPrecacheRequests 	= FCsvProfiler::Get()->GetOrCreatePersistentCustomStatInt(TEXT("TotalGlobalComputePSOPrecacheRequests"), CSV_CATEGORY_INDEX(PSOPrecacheCompiling));
+		
+	TCsvPersistentCustomStat<int>* TotalMeshPassGraphicPSOPrecacheRequests 	= FCsvProfiler::Get()->GetOrCreatePersistentCustomStatInt(TEXT("TotalMeshPassGraphicPSOPrecacheRequests"), CSV_CATEGORY_INDEX(PSOPrecacheCompiling));
+	TCsvPersistentCustomStat<int>* TotalMeshPassComputePSOPrecacheRequests 	= FCsvProfiler::Get()->GetOrCreatePersistentCustomStatInt(TEXT("TotalMeshPassComputePSOPrecacheRequests"), CSV_CATEGORY_INDEX(PSOPrecacheCompiling));
+	
+	static bool CSVStatsEnabled()
+	{
+		return FCsvProfiler::Get() && FCsvProfiler::Get()->IsCategoryEnabled(CSV_CATEGORY_INDEX(PSOPrecacheCompiling));
+	}
+
+#endif // CSV_PROFILER_STATS
+
+private:
+	struct FGlobalPrecacheData
+	{
+		uint64 PSOPrecacheHash;
+		FPSOPrecacheData::EType Type;
+	};
+	TMap < int32, TArray<FGlobalPrecacheData> > GlobalPrecacheDataPerCollector;
+	FCriticalSection GlobalPrecacheDataPerCollectorCS;
+
+public:
+
+	void AddGlobalPrecacheData(const FPSOPrecacheData& PrecacheData)
+	{
+		uint64 PSOPrecacheHash = 0; 
+		if (PrecacheData.Type == FPSOPrecacheData::EType::Graphics)
+		{
+			PSOPrecacheHash = PSOCollectorStats::GetPSOPrecacheHash(PrecacheData.GraphicsPSOInitializer);
+		}
+		else
+		{
+			PSOPrecacheHash = PSOCollectorStats::GetPSOPrecacheHash(PrecacheData.ComputeInitializer);
+		}
+
+		FScopeLock ScopeLock(&GlobalPrecacheDataPerCollectorCS);
+		GlobalPrecacheDataPerCollector.FindOrAdd(PrecacheData.PSOCollectorIndex).Add({PSOPrecacheHash, PrecacheData.Type});
+	}
+
+	uint64 GetGlobalCompileTime(const TCHAR* InGlobalPSOCollectorName)
+	{
+		FScopeLock ScopeLock(&GlobalPrecacheDataPerCollectorCS);
+
+		auto GetTotalCompileTime = [](const TArray<FGlobalPrecacheData>& GlobalPrecacheDatas)
+		{
+			uint64 TotalCompileTime = 0;
+			for (const FGlobalPrecacheData& GlobalPrecacheData : GlobalPrecacheDatas)
+			{
+				uint64 PSOCompilationTime = PipelineStateCache::GetPrecachePSOCompilationTime(GlobalPrecacheData.Type == FPSOPrecacheData::EType::Graphics, GlobalPrecacheData.PSOPrecacheHash);
+				TotalCompileTime += PSOCompilationTime;
+			}
+			return TotalCompileTime;
+		};
+
+		uint64 TotalCompileTime = 0;
+		if (InGlobalPSOCollectorName)
+		{
+			const int32 PSOCollectorIndex = FGlobalPSOCollectorManager::GetIndex(InGlobalPSOCollectorName);
+			check(PSOCollectorIndex >= 0);
+			TArray<FGlobalPrecacheData>* FindRes = GlobalPrecacheDataPerCollector.Find(PSOCollectorIndex);
+			if (FindRes)
+			{
+				TotalCompileTime = GetTotalCompileTime(*FindRes);
+			}
+		}
+		else
+		{
+			for (auto Iter : GlobalPrecacheDataPerCollector)
+			{
+				TotalCompileTime += GetTotalCompileTime(Iter.Value);
+			}
+		}
+		return TotalCompileTime;
+	}
+
+	static FGlobalPSOPrecacheStats& GetStats()
+	{
+		static FGlobalPSOPrecacheStats PersistentStats;
+		return PersistentStats;
+	}
+};
+
 uint64 PSOCollectorStats::GetPSOPrecacheHash(const FGraphicsPipelineStateInitializer& GraphicsPSOInitializer)
 {
 	if (GraphicsPSOInitializer.StatePrecachePSOHash == 0)
@@ -391,6 +502,17 @@ uint64 PSOCollectorStats::GetPSOPrecacheHash(const FComputePipelineStateInitiali
 	return ComputeInitializer.ComputeShader ? GetTypeHash(ComputeInitializer.ComputeShader->GetHash()) : 0;
 }
 
+FGraphicsPipelineStateInitializer PSOCollectorStats::GetShadersOnlyInitializer(const FGraphicsPipelineStateInitializer& Initializer)
+{
+	FGraphicsPipelineStateInitializer ShadersOnlyInitializer;
+	ShadersOnlyInitializer.PrimitiveType = PT_TriangleList;
+	ShadersOnlyInitializer.BoundShaderState = Initializer.BoundShaderState;
+	ShadersOnlyInitializer.BoundShaderState.VertexDeclarationRHI = nullptr;
+	ShadersOnlyInitializer.StatePrecachePSOHash = RHIComputeStatePrecachePSOHash(ShadersOnlyInitializer);
+
+	return ShadersOnlyInitializer;
+}
+
 void PSOCollectorStats::CheckFullPipelineStateInCache(
 	const FGraphicsPipelineStateInitializer& Initializer,
 	EPSOPrecacheResult PSOPrecacheResult,
@@ -453,7 +575,24 @@ void PSOCollectorStats::CheckComputePipelineStateInCache(
 	if (IsFullPrecachingValidationEnabled() && bStatsUpdated && PSOPrecacheResult == EPSOPrecacheResult::Missed)
 	{
 		const FMaterial* Material = MaterialRenderProxy ? MaterialRenderProxy->GetMaterialNoFallback(GMaxRHIFeatureLevel) : nullptr;
-		LogPSOMissInfo(*ComputeInitializer.ComputeShader, PSOPrecacheResult, Material, PSOCollectorIndex);
+		LogPSOMissInfo(*ComputeInitializer.ComputeShader, PSOPrecacheResult, Material, EPSOPrecacheType::MeshPass, PSOCollectorIndex);
+	}
+
+	if (bStatsUpdated)
+	{
+		UpdateCSVStats(PSOPrecacheResult);
+	}
+}
+
+void PSOCollectorStats::CheckGlobalComputePipelineStateInCache(const FComputePipelineStateInitializer& ComputeInitializer)
+{
+	static const int32 GlobalComputePSOCollectorIndex = FGlobalPSOCollectorManager::GetIndex(TEXT("GlobalComputePSOCollector"));
+	EPSOPrecacheResult PSOPrecacheResult = PipelineStateCache::CheckPipelineStateInCache(ComputeInitializer.ComputeShader);
+
+	bool bStatsUpdated = PSOCollectorStats::GetFullPSOPrecacheStatsCollector().CheckStateInCache(EPSOPrecacheType::Global, ComputeInitializer, PSOCollectorStats::GetPSOPrecacheHash, GlobalComputePSOCollectorIndex, nullptr, PSOPrecacheResult);
+	if (IsFullPrecachingValidationEnabled() && bStatsUpdated && PSOPrecacheResult == EPSOPrecacheResult::Missed)
+	{
+		LogPSOMissInfo(*ComputeInitializer.ComputeShader, PSOPrecacheResult, nullptr, EPSOPrecacheType::Global, GlobalComputePSOCollectorIndex);
 	}
 
 	if (bStatsUpdated)
@@ -517,6 +656,55 @@ void PSOCollectorStats::UpdateCSVStats(EPSOPrecacheResult PSOPrecacheResult)
 	}
 }
 
+void PSOCollectorStats::UpdateGlobalStats(EPSOPrecacheType PSOPrecacheType, const FPSOPrecacheData& PrecacheData)
+{	
+	FGlobalPSOPrecacheStats& GlobalPSOPrecacheStats = FGlobalPSOPrecacheStats::GetStats();
+
+	// Track all global requests so they can be used during stat dumping
+	if (PSOPrecacheType == EPSOPrecacheType::Global)
+	{
+		GlobalPSOPrecacheStats.AddGlobalPrecacheData(PrecacheData);
+	}
+
+#if CSV_PROFILER_STATS
+	if (PrecacheData.Type == FPSOPrecacheData::EType::Compute)
+	{
+		CSV_CUSTOM_STAT(PSOPrecache, PrecachedCompute, 1, ECsvCustomStatOp::Accumulate);
+	}
+	else
+	{
+		CSV_CUSTOM_STAT(PSOPrecache, PrecachedGraphics, 1, ECsvCustomStatOp::Accumulate);
+	}
+
+	// Update total requested stats for CSV profiler
+	if (FGlobalPSOPrecacheStats::CSVStatsEnabled())
+	{
+		if (PrecacheData.Type == FPSOPrecacheData::EType::Compute)
+		{
+			if (PSOPrecacheType == EPSOPrecacheType::Global)
+			{
+				GlobalPSOPrecacheStats.TotalGlobalComputePSOPrecacheRequests->Add(1);
+			}
+			else
+			{
+				GlobalPSOPrecacheStats.TotalMeshPassComputePSOPrecacheRequests->Add(1);
+			}
+		}
+		else
+		{
+			if (PSOPrecacheType == EPSOPrecacheType::Global)
+			{
+				GlobalPSOPrecacheStats.TotalGlobalGraphicPSOPrecacheRequests->Add(1);
+			}
+			else
+			{
+				GlobalPSOPrecacheStats.TotalMeshPassGraphicPSOPrecacheRequests->Add(1);
+			}
+		}
+	}
+#endif // CSV_PROFILER_STATS
+}
+
 
 //////////////////////////////////////////////////////////////////////////
 
@@ -1084,6 +1272,7 @@ void LogPSOMissInfo(
 	const FRHIComputeShader& ComputeShader,
 	EPSOPrecacheResult PrecacheResult,
 	const FMaterial* Material,
+	EPSOPrecacheType PSOPrecacheType,
 	int32 PSOCollectorIndex)
 {
 	PSOMissStringBuilder StringBuilder;
@@ -1091,8 +1280,11 @@ void LogPSOMissInfo(
 	AppendCSVProfilerInfo(StringBuilder);
 	StringBuilder.Appendf(TEXT("\n\tType:\t\t\t\t\t%s"), TEXT("Compute"));
 	StringBuilder.Appendf(TEXT("\n\tPSOPrecachingState:\t\t%s"), LexToString(PrecacheResult));
-	StringBuilder.Appendf(TEXT("\n\tMaterial:\t\t\t\t%s"), Material ? *Material->GetAssetName() : TEXT("Unknown"));
-	StringBuilder.Appendf(TEXT("\n\tPassName:\t\t\t\t%s"), FPSOCollectorCreateManager::GetName(GetFeatureLevelShadingPath(GMaxRHIFeatureLevel), PSOCollectorIndex));
+	if (PSOPrecacheType == EPSOPrecacheType::MeshPass)
+	{
+		StringBuilder.Appendf(TEXT("\n\tMaterial:\t\t\t\t%s"), Material ? *Material->GetAssetName() : TEXT("Unknown"));
+		StringBuilder.Appendf(TEXT("\n\tPassName:\t\t\t\t%s"), FPSOCollectorCreateManager::GetName(GetFeatureLevelShadingPath(GMaxRHIFeatureLevel), PSOCollectorIndex));
+	}
 	StringBuilder.Appendf(TEXT("\n\tCompute Shader Hash:\t%s"), *(ComputeShader.GetHash().ToString()));
 
 	// Not sure yet if this is interesting data or not
@@ -1109,16 +1301,8 @@ void LogPSOMissInfo(
 	UE_LOG(LogEngine, Log, TEXT("%s\n"), StringBuilder.ToString());
 }
 
-
-void PSOCollectorStats::DumpPSOPrecacheValidationStats()
+namespace PSOCollectorStats
 {
-	const FString Filename = FString::Printf(TEXT("%sPSOPrecacheStats-%s.csv"), *FPaths::ProfilingDir(), *FDateTime::Now().ToString());
-	FArchive* CSVFile = IFileManager::Get().CreateFileWriter(*Filename, FILEWRITE_AllowRead);
-	if (CSVFile == nullptr)
-	{
-		return;
-	}
-
 	enum class ETrackingType
 	{
 		Missed,
@@ -1136,7 +1320,7 @@ void PSOCollectorStats::DumpPSOPrecacheValidationStats()
 		Full,
 		Num,
 	};
-	
+
 	struct FPerTypeStats
 	{
 		FPerTypeStats()
@@ -1154,11 +1338,39 @@ void PSOCollectorStats::DumpPSOPrecacheValidationStats()
 			TrackedStats[(int32)StatsType][(int32)TrackingType] = Value;
 		}
 
+		template<typename CountGetterCallback>
+		void AddStatsPerTrackingType(EStatsType StatsType, PSOCollectorStats::FPrecacheStats& Stats, CountGetterCallback Callback)
+		{
+			SetStat(StatsType, ETrackingType::Missed, Callback(Stats.MissData));
+			SetStat(StatsType, ETrackingType::Precached, Callback(Stats.PrecacheData));
+			SetStat(StatsType, ETrackingType::Used, Callback(Stats.UsageData));
+			SetStat(StatsType, ETrackingType::TooLate, Callback(Stats.TooLateData));
+			SetStat(StatsType, ETrackingType::Untracked, Callback(Stats.UntrackedData));
+		}
+
+		template<typename CountGetterCallback>
+		void AddStatsPerType(CountGetterCallback Callback)
+		{
+			AddStatsPerTrackingType(EStatsType::Full, FullPSOPrecacheStatsCollector.GetStats(), Callback);
+			AddStatsPerTrackingType(EStatsType::Minimal, MinimalPSOPrecacheStatsCollector.GetStats(), Callback);
+			AddStatsPerTrackingType(EStatsType::ShaderOnly, ShadersOnlyPSOPrecacheStatsCollector.GetStats(), Callback);
+		}
+
 		FString Name;
-		uint64 TotalCompileTime = 0;		
+		uint64 TotalCompileTime = 0;
 		int32 TrackedStats[(int32)EStatsType::Num][(int32)ETrackingType::Num];
-	};	
-	
+	};
+}
+
+void PSOCollectorStats::DumpPSOPrecacheValidationStats()
+{
+	const FString Filename = FString::Printf(TEXT("%sPSOPrecacheStats-%s.csv"), *FPaths::ProfilingDir(), *FDateTime::Now().ToString());
+	FArchive* CSVFile = IFileManager::Get().CreateFileWriter(*Filename, FILEWRITE_AllowRead);
+	if (CSVFile == nullptr)
+	{
+		return;
+	}
+		
 	auto AddStatsToCSVFile = [CSVFile](const TCHAR* TypeName, bool bSort, TArray<FPerTypeStats>& TypeStats)
 	{		
 		if (bSort)
@@ -1220,26 +1432,35 @@ void PSOCollectorStats::DumpPSOPrecacheValidationStats()
 	const PSOCollectorStats::FPrecacheStats& ShaderOnlyPSOPrecacheStats = ShadersOnlyPSOPrecacheStatsCollector.GetStats();
 	const PSOCollectorStats::FPrecacheStats& MinimalPSOPrecacheStats = MinimalPSOPrecacheStatsCollector.GetStats();
 
+	FGlobalPSOPrecacheStats& GlobalPSOPrecacheStats = FGlobalPSOPrecacheStats::GetStats();
+
 	// Totals
 	{
 		TArray<FPerTypeStats> TotalStats;
+		
+		auto TotalGetter = [](FPrecacheUsageData& UsageData) { return UsageData.GetTotalCount(); };
+
+		FPerTypeStats TotalPerTypeStats;
+		TotalPerTypeStats.Name = TEXT("Total");
+		TotalPerTypeStats.TotalCompileTime = GetPSOPrecacheMaterialCompileTime(nullptr, nullptr) + GlobalPSOPrecacheStats.GetGlobalCompileTime(nullptr);
+		TotalPerTypeStats.AddStatsPerType(TotalGetter);
+		TotalStats.Add(TotalPerTypeStats);
 
-		FPerTypeStats PerTypeStats;
-		PerTypeStats.Name = TEXT("");
-		PerTypeStats.TotalCompileTime = GetPSOPrecacheMaterialCompileTime(nullptr, nullptr);
-		auto AddTotalStats = [&PerTypeStats](EStatsType StatsType, const PSOCollectorStats::FPrecacheStats& Stats)
-		{			
-			PerTypeStats.SetStat(StatsType, ETrackingType::Missed, Stats.MissData.GetTotalCount());
-			PerTypeStats.SetStat(StatsType, ETrackingType::Precached, Stats.PrecacheData.GetTotalCount());
-			PerTypeStats.SetStat(StatsType, ETrackingType::Used, Stats.UsageData.GetTotalCount());
-			PerTypeStats.SetStat(StatsType, ETrackingType::TooLate, Stats.TooLateData.GetTotalCount());
-			PerTypeStats.SetStat(StatsType, ETrackingType::Untracked, Stats.UntrackedData.GetTotalCount());
-		};		
-		AddTotalStats(EStatsType::Full, FullPSOPrecacheStats);
-		AddTotalStats(EStatsType::Minimal, MinimalPSOPrecacheStats);
-		AddTotalStats(EStatsType::ShaderOnly, ShaderOnlyPSOPrecacheStats);
-
-		TotalStats.Add(PerTypeStats);
+		auto TotalGlobalGetter = [](FPrecacheUsageData& UsageData) { return UsageData.GetTotalGlobalCount(); };
+
+		FPerTypeStats GlobalStats;
+		GlobalStats.Name = TEXT("Global");
+		GlobalStats.TotalCompileTime = GlobalPSOPrecacheStats.GetGlobalCompileTime(nullptr);
+		GlobalStats.AddStatsPerType(TotalGlobalGetter);
+		TotalStats.Add(GlobalStats);
+
+		auto TotalMeshPassGetter = [](FPrecacheUsageData& UsageData) { return UsageData.GetTotalMeshPassCount(); };
+
+		FPerTypeStats MeshPassStats;
+		MeshPassStats.Name = TEXT("MeshPass");
+		MeshPassStats.TotalCompileTime = GetPSOPrecacheMaterialCompileTime(nullptr, nullptr);
+		MeshPassStats.AddStatsPerType(TotalMeshPassGetter);
+		TotalStats.Add(MeshPassStats);
 
 		AddStatsToCSVFile(TEXT("Totals"), false, TotalStats);
 	}
@@ -1252,6 +1473,8 @@ void PSOCollectorStats::DumpPSOPrecacheValidationStats()
 		{
 			FPerTypeStats PerTypeStats;
 			PerTypeStats.Name = FGlobalPSOCollectorManager::GetName(PSOCollectorIndex);
+			PerTypeStats.TotalCompileTime = GlobalPSOPrecacheStats.GetGlobalCompileTime(*PerTypeStats.Name);
+
 			PerTypeStats.SetStat(EStatsType::Full, ETrackingType::Missed, FullPSOPrecacheStats.MissData.PerGlobalTypeCount[PSOCollectorIndex]);
 			PerTypeStats.SetStat(EStatsType::Full, ETrackingType::Precached, FullPSOPrecacheStats.PrecacheData.PerGlobalTypeCount[PSOCollectorIndex]);
 			PerTypeStats.SetStat(EStatsType::Full, ETrackingType::Used, FullPSOPrecacheStats.UsageData.PerGlobalTypeCount[PSOCollectorIndex]);			
diff --git a/Engine/Source/Runtime/Engine/Private/ShaderCompiler/ShaderCompiler.cpp b/Engine/Source/Runtime/Engine/Private/ShaderCompiler/ShaderCompiler.cpp
--- a/Engine/Source/Runtime/Engine/Private/ShaderCompiler/ShaderCompiler.cpp
+++ b/Engine/Source/Runtime/Engine/Private/ShaderCompiler/ShaderCompiler.cpp
@@ -5039,6 +5039,11 @@ void VerifyGlobalShaders(EShaderPlatform Platform, bool bLoadedFromCacheFile, co
 	VerifyGlobalShaders(Platform, nullptr, bLoadedFromCacheFile, OutdatedShaderTypes, OutdatedShaderPipelineTypes);
 }
 
+// Declare dummy global PSO collector which doesn't do anything during collection but list to compile is setup below
+// in optimized tasks so shaders can be created in parallel to offload the game thread during startup (5+ shaders)
+static const TCHAR* GlobalComputePSOCollectorName = TEXT("GlobalComputePSOCollector");
+FRegisterGlobalPSOCollectorFunction RegisterGlobalComputePSOCollector(nullptr, GlobalComputePSOCollectorName);
+
 void PrecacheComputePipelineStatesForGlobalShaders(ERHIFeatureLevel::Type FeatureLevel, const ITargetPlatform* TargetPlatform)
 {
 	int32 PrecacheGlobalShaders = UPSOPrecacheSettingsManager::GetSettings().GlobalShaderMode;
@@ -5101,59 +5106,78 @@ void PrecacheComputePipelineStatesForGlobalShaders(ERHIFeatureLevel::Type Featur
 		{
 			UE_LOG(LogShaders, Display, TEXT("Precaching %d global compute shaders"), ComputeShadersToPrecache.Num());
 
+			struct FGlobalComputePSOPrecacheData
+			{
+				FGlobalComputePSOPrecacheData(int Count)
+				{
+					PSOInitializers.Reserve(Count);
+				}
+
+				FPSOPrecacheDataArray PSOInitializers;
+				FCriticalSection PSOInitializersCS;
+			};
+			FGlobalComputePSOPrecacheData* GlobalComputePSOPrecacheData = new FGlobalComputePSOPrecacheData(ComputeShadersToPrecache.Num());			
+
+			auto PrecacheComputePipelineState = [GlobalComputePSOPrecacheData](TShaderRef<FShader>& GlobalShader)
+			{
+				// PSO precache shaders are not required to all load correctly
+				constexpr bool bRequired = false;
+				FRHIComputeShader* RHIComputeShader = static_cast<FRHIComputeShader*>(GlobalShader.GetRHIShaderBase(SF_Compute, bRequired));
+				if (RHIComputeShader)
+				{							
+					static const int32 GlobalComputePSOCollectorIndex = FGlobalPSOCollectorManager::GetIndex(GlobalComputePSOCollectorName);
+
+					check(!RHIComputeShader->HasSpecializationConstants() || !GlobalShader.GetSpecializationConstants().IsEmpty());
+					FComputePipelineStateInitializer ComputeInitializer(RHIComputeShader, GlobalShader.GetSpecializationConstants());
+								
+					FPSOPrecacheData PSOPrecacheData;
+					PSOPrecacheData.Name = GlobalShader.GetType()->GetName();
+					PSOPrecacheData.Type = FPSOPrecacheData::EType::Compute;
+					PSOPrecacheData.ComputeInitializer = ComputeInitializer;
+					PSOPrecacheData.InMemoryMode = EPSOPrecacheInMemoryMode::Uncached;
+				#if PSO_PRECACHING_VALIDATE
+					PSOPrecacheData.PSOCollectorIndex = GlobalComputePSOCollectorIndex;
+					PSOPrecacheData.VertexFactoryType = nullptr;
+				#endif // PSO_PRECACHING_VALIDATE
+
+					FScopeLock ScopeLock(&GlobalComputePSOPrecacheData->PSOInitializersCS);
+					GlobalComputePSOPrecacheData->PSOInitializers.Add(MoveTemp(PSOPrecacheData));
+				}
+			};
+
 			if (GRHIGlobals.SupportsMultithreadedShaderCreation)
 			{
 				TArray<UE::Tasks::TTask<void>> Tasks;
 				for (TShaderRef<FShader>& GlobalShader : ComputeShadersToPrecache)
 				{
 					Tasks.Emplace(UE::Tasks::Launch(TEXT("PrecachePSOsForGlobalShaders"),
-						[&GlobalShader]()
+						[&GlobalShader, PrecacheComputePipelineState]()
 						{
-							// PSO precache shaders are not required to all load correctly
-							constexpr bool bRequired = false;
-							FRHIComputeShader* RHIComputeShader = static_cast<FRHIComputeShader*>(GlobalShader.GetRHIShaderBase(SF_Compute, bRequired));
-							if (RHIComputeShader)
-							{
-								FPSOPrecacheRequestParams PSORequestParams;
-								PSORequestParams.InMemoryMode = EPSOPrecacheInMemoryMode::Uncached;
-								PSORequestParams.PrecachePriority = EPSOPrecachePriority::Medium;
-								PSORequestParams.Name = GlobalShader.GetType()->GetName();
-
-								check(!RHIComputeShader->HasSpecializationConstants() || !GlobalShader.GetSpecializationConstants().IsEmpty());
-								FComputePipelineStateInitializer ComputeInitializer(RHIComputeShader, GlobalShader.GetSpecializationConstants());
-								PipelineStateCache::PrecacheComputePipelineState(ComputeInitializer, PSORequestParams);
-							}
+							PrecacheComputePipelineState(GlobalShader);
 						}
 					));
 				}
 				ENQUEUE_RENDER_COMMAND(WaitOnPSOPrecacheTasks)(
-					[ComputeShadersToPrecache = MoveTemp(ComputeShadersToPrecache), Tasks = MoveTemp(Tasks)](FRHICommandListImmediate& RHICmdList)
+					[ComputeShadersToPrecache = MoveTemp(ComputeShadersToPrecache), Tasks = MoveTemp(Tasks), GlobalComputePSOPrecacheData](FRHICommandListImmediate& RHICmdList)
 					{
 						UE::Tasks::Wait(Tasks);
+
+						RequestPrecachePSOs(EPSOPrecacheType::Global, GlobalComputePSOPrecacheData->PSOInitializers);
+						delete GlobalComputePSOPrecacheData;
 					});
 			}
 			else
 			{
 				ENQUEUE_RENDER_COMMAND(PrecachePSOsForGlobalShaders)(
-					[ComputeShadersToPrecache = MoveTemp(ComputeShadersToPrecache)](FRHICommandListImmediate& RHICmdList)
+					[ComputeShadersToPrecache = MoveTemp(ComputeShadersToPrecache), GlobalComputePSOPrecacheData, PrecacheComputePipelineState](FRHICommandListImmediate& RHICmdList)
 					{
 						for (TShaderRef<FShader> GlobalShader : ComputeShadersToPrecache)
 						{
-							// PSO precache shaders are not required to all load correctly
-							bool bRequired = false;
-							FRHIComputeShader* RHIComputeShader = static_cast<FRHIComputeShader*>(GlobalShader.GetRHIShaderBase(SF_Compute, bRequired));
-							if (RHIComputeShader)
-							{
-								FPSOPrecacheRequestParams PSORequestParams;
-								PSORequestParams.InMemoryMode = EPSOPrecacheInMemoryMode::Uncached;
-								PSORequestParams.PrecachePriority = EPSOPrecachePriority::Medium;
-								PSORequestParams.Name = GlobalShader.GetType()->GetName();
-
-								check(!RHIComputeShader->HasSpecializationConstants() || !GlobalShader.GetSpecializationConstants().IsEmpty());
-								FComputePipelineStateInitializer ComputeInitializer(RHIComputeShader, GlobalShader.GetSpecializationConstants());
-								PipelineStateCache::PrecacheComputePipelineState(ComputeInitializer, PSORequestParams);
-							}
+							PrecacheComputePipelineState(GlobalShader);
 						}
+						
+						RequestPrecachePSOs(EPSOPrecacheType::Global, GlobalComputePSOPrecacheData->PSOInitializers);
+						delete GlobalComputePSOPrecacheData;
 					});
 			}
 		}
diff --git a/Engine/Source/Runtime/Engine/Public/PSOPrecache.h b/Engine/Source/Runtime/Engine/Public/PSOPrecache.h
--- a/Engine/Source/Runtime/Engine/Public/PSOPrecache.h
+++ b/Engine/Source/Runtime/Engine/Public/PSOPrecache.h
@@ -235,31 +235,25 @@ extern ENGINE_API EPSOPrecacheMode GetPSOPrecacheMode();
  */
 struct FPSOPrecacheData
 {
-	FPSOPrecacheData() : Type(EType::Graphics), bRequired(true), InMemoryMode(EPSOPrecacheInMemoryMode::Cached)
-#if PSO_PRECACHING_VALIDATE
-		, PSOCollectorIndex(INDEX_NONE)
-		, bDefaultMaterial(0)
-		, VertexFactoryType(nullptr)
-#endif // PSO_PRECACHING_VALIDATE
-	{
-	}
-
 	enum class EType : uint8
 	{
 		Graphics,
 		Compute,
 	};
-	EType Type;
+	EType Type = EType::Graphics;
+
+	// Name of the material/type/shader to show in breadcrumb stack in Insight while PSO is compiling
+	FString Name;
 
 	// Is the PSO required to be able render the object or can it provide a fallback path
 	// (proxy creation won't wait for these PSOs if enabled)
-	bool bRequired;
+	bool bRequired = true;
 
 	// Precache priority mode
 	EPSOPrecachePriority Priority = EPSOPrecachePriority::Medium;
 
 	// Keep the PSO in memory until the owning material is request deleted
-	EPSOPrecacheInMemoryMode InMemoryMode;
+	EPSOPrecacheInMemoryMode InMemoryMode = EPSOPrecacheInMemoryMode::Cached;
 
 	FGraphicsPipelineStateInitializer GraphicsPSOInitializer;
 	FComputePipelineStateInitializer ComputeInitializer;
@@ -279,9 +273,9 @@ struct FPSOPrecacheData
 	}
 
 #if PSO_PRECACHING_VALIDATE
-	int32 PSOCollectorIndex : 31;
-	int32 bDefaultMaterial : 1;
-	const FVertexFactoryType* VertexFactoryType;
+	int32 PSOCollectorIndex : 31 = INDEX_NONE;
+	int32 bDefaultMaterial : 1 = 0;
+	const FVertexFactoryType* VertexFactoryType = nullptr;
 #endif // PSO_PRECACHING_VALIDATE
 };
 
@@ -326,7 +320,7 @@ typedef void (*GlobalPSOCollectorFunction)(const FSceneTexturesConfig& SceneText
 class FGlobalPSOCollectorManager
 {
 public:
-	constexpr static uint32 MaxPSOCollectorCount = 4;
+	constexpr static uint32 MaxPSOCollectorCount = 8;
 
 	static int32 GetPSOCollectorCount() { return PSOCollectorCount; }
 	static GlobalPSOCollectorFunction GetCollectFunction(int32 Index)
diff --git a/Engine/Source/Runtime/Engine/Public/PSOPrecacheValidation.h b/Engine/Source/Runtime/Engine/Public/PSOPrecacheValidation.h
--- a/Engine/Source/Runtime/Engine/Public/PSOPrecacheValidation.h
+++ b/Engine/Source/Runtime/Engine/Public/PSOPrecacheValidation.h
@@ -70,6 +70,7 @@ extern ENGINE_API void LogPSOMissInfo(
 	const FRHIComputeShader& ComputeShader,
 	EPSOPrecacheResult PrecacheResult,
 	const FMaterial* Material,
+	EPSOPrecacheType PSOPrecacheType,
 	int32 PSOCollectorIndex);
 
 /**
@@ -88,6 +89,16 @@ namespace PSOCollectorStats
 	extern ENGINE_API bool IsPrecachingValidationEnabled();
 	extern ENGINE_API bool IsFullPrecachingValidationEnabled();
 
+	/**
+	* Gets the state PSOs hash only from the grpahics PSOs - should be valid
+	*/
+	inline uint64 GetPSOPrecacheStateHash(const FGraphicsPipelineStateInitializer& GraphicsPSOInitializer)
+	{
+		check(GraphicsPSOInitializer.StatePrecachePSOHash != 0);
+		return GraphicsPSOInitializer.StatePrecachePSOHash;
+	}
+
+
 	/**
 	 * Compute the hash of a graphics PSO initializer to be used by PSO precaching validation.
 	 */
@@ -105,6 +116,11 @@ namespace PSOCollectorStats
 		return GetPSOPrecacheHash(ComputeInitializer);
 	}
 
+	/**
+	 * Create a shaders-only initializer where all state except for shaders is removed from the given initializer.
+	 */
+	extern ENGINE_API FGraphicsPipelineStateInitializer GetShadersOnlyInitializer(const FGraphicsPipelineStateInitializer& Initializer);
+	   
 	/**
 	 * Check if the full PSO is precached and if not output information on what state is missing
 	 */
@@ -136,6 +152,11 @@ namespace PSOCollectorStats
 		const FMaterialRenderProxy* Material,
 		int32 PSOCollectorIndex);
 
+	/**
+	 * Check if the global compute PSO is precached and if not output information on what state is missing
+	 */
+	extern ENGINE_API void CheckGlobalComputePipelineStateInCache(const FComputePipelineStateInitializer& ComputeInitializer);
+
 	//UE_DEPRECATED(5.8, "Identifying compute pipelines specifying only a shader is deprecated, use a FComputePipelineStateInitializer.")
 	inline void CheckComputePipelineStateInCache(
 		FRHIComputeShader& ComputeShader,
@@ -158,6 +179,11 @@ namespace PSOCollectorStats
 	 */
 	extern ENGINE_API void UpdateCSVStats(EPSOPrecacheResult PSOPrecacheResult);
 
+	/*
+	* Update the global tracking stats
+	*/
+	extern ENGINE_API void UpdateGlobalStats(EPSOPrecacheType PSOPrecacheType, const FPSOPrecacheData& PrecacheData);
+
 	using VertexFactoryCountTableType = Experimental::TRobinHoodHashMap<const FVertexFactoryType*, uint32>;
 	struct FShaderStateUsage
 	{
@@ -187,6 +213,8 @@ namespace PSOCollectorStats
 		{
 			return FPlatformAtomics::AtomicRead(&Count);
 		}
+		uint64 GetTotalMeshPassCount();
+		uint64 GetTotalGlobalCount();
 
 		ENGINE_API void Empty();
 
@@ -264,6 +292,7 @@ namespace PSOCollectorStats
 			Stats.Empty();
 		}
 
+		FPrecacheStats& GetStats() { return Stats; }
 		const FPrecacheStats& GetStats() const { return Stats; }
 
 		template <typename TPrecacheState>
diff --git a/Engine/Source/Runtime/RHI/Private/PipelineStateCache.cpp b/Engine/Source/Runtime/RHI/Private/PipelineStateCache.cpp
--- a/Engine/Source/Runtime/RHI/Private/PipelineStateCache.cpp
+++ b/Engine/Source/Runtime/RHI/Private/PipelineStateCache.cpp
@@ -410,6 +410,13 @@ static FAutoConsoleVariableRef CVarPrecompileThreadPoolThreadPriority(
 	TEXT("Thread priority for the PSO precompile pool"),
 	ECVF_RenderThreadSafe);
 
+int32 GPSOPrecacheBoostThreadPriorityForHighestPriorityRequests = 0;
+static FAutoConsoleVariableRef CVarPSOPrecacheBoostThreadPriorityForHighestPriorityRequests(
+	TEXT("r.PSOPrecache.BoostThreadPriorityForHighestPriorityRequests"),
+	GPSOPrecacheBoostThreadPriorityForHighestPriorityRequests,
+	TEXT("Boost compilation thread priority when compiling Highest priority PSO precaching requests to reduce fallback rendering time (can have influence on foreground tasks then)"),
+	ECVF_RenderThreadSafe);
+
 int32 GPSOPrecacheKeepInMemoryUntilUsed = 0;
 static FAutoConsoleVariableRef CVarPSOPrecacheKeepInMemoryUntilUsed(
 	TEXT("r.PSOPrecache.KeepInMemoryUntilUsed"),
@@ -3644,6 +3651,20 @@ public:
 			);
 		}
 #endif // WITH_RHI_BREADCRUMBS
+				
+		bool bRestoreThreadPriority = false;
+		if (Initializer.bPSOPrecache)
+		{
+			// Might be not set if running without async compile tasks
+			PSOPrecachePriority = PSOPrecachePriority != EPSOPrecachePriority::NotSet ? PSOPrecachePriority : EPSOPrecachePriority::Medium;
+			Pipeline->CSVStat.SetInFlightPriority(PSOPrecachePriority);
+
+			if (GPSOPrecacheBoostThreadPriorityForHighestPriorityRequests > 0 && PSOPrecachePriority == EPSOPrecachePriority::Highest)
+			{
+				FPlatformProcess::SetThreadPriority(TPri_AboveNormal);
+				bRestoreThreadPriority = true;
+			}
+		}
 
 		const TCHAR* PSOPrecacheResultScopeString = nullptr;
 		switch (PSOPrecacheResult)
@@ -3680,6 +3701,11 @@ public:
 			// At this point, it's not safe to use Pipeline anymore, as it might get picked up by ProcessDelayedCleanup and deleted
 			Pipeline = nullptr;
 		}
+
+		if (bRestoreThreadPriority)
+		{
+			FPlatformProcess::SetThreadPriority((EThreadPriority)GPSOPrecompileThreadPoolThreadPriority);
+		}
 	}
 
 	FORCEINLINE TStatId GetStatId() const
@@ -3813,9 +3839,7 @@ public:
 
 		if (Initializer.bPSOPrecache)
 		{
-			check(PSOPrecachePriority != EPSOPrecachePriority::NotSet);
 			Initializer.PrecacheCompileType = (uint32)GetPSOCompileTypeFromPrecachePriority(PSOPrecachePriority);			
-			Pipeline->CSVStat.SetInFlightPriority(PSOPrecachePriority);
 		}
 
 		FGraphicsPipelineState* GfxPipeline = static_cast<FGraphicsPipelineState*>(Pipeline);
@@ -3938,13 +3962,6 @@ public:
 
 		uint64 StartTime = FPlatformTime::Cycles64();
 
-		if (Initializer.bPSOPrecache)
-		{
-			check(PSOPrecachePriority != EPSOPrecachePriority::NotSet);
-			//Initializer.PrecacheCompileType = (uint32)GetPSOCompileTypeFromPrecachePriority(PSOPrecachePriority);			
-			Pipeline->CSVStat.SetInFlightPriority(PSOPrecachePriority);
-		}
-
 		// TODO: Send the priority to the RHI.
 		ComputePipeline->RHIPipeline = RHICreateComputePipelineState(Initializer);
 		uint64 PSOCompilationTime = CheckAndUpdateHitchCountStat(FPSOPrecacheRequestID::EType::Compute, !IsPrecachedPSO(Initializer), PSOCompilationDebugData, StartTime, PSOPrecacheResult);
