[StaticMeshComponent/ALT]
- Do not overwrite the value in postload unless it needs to be modified to avoid TSAN picking this as a race condition

#Synced-CL 50565643

diff --git a/Engine/Source/Runtime/Engine/Private/Components/StaticMeshComponent.cpp b/Engine/Source/Runtime/Engine/Private/Components/StaticMeshComponent.cpp
--- a/Engine/Source/Runtime/Engine/Private/Components/StaticMeshComponent.cpp
+++ b/Engine/Source/Runtime/Engine/Private/Components/StaticMeshComponent.cpp
@@ -2248,7 +2248,10 @@ void UStaticMeshComponent::PostLoad()
 	}
 
 	// Legacy content may contain a lightmap resolution of 0, which was valid when vertex lightmaps were supported, but not anymore with only texture lightmaps
-	OverriddenLightMapRes = FMath::Max(OverriddenLightMapRes, 4);
+	if (OverriddenLightMapRes < 4)
+	{
+		OverriddenLightMapRes = 4;
+	}
 
 	// Initialize the resources for the freshly loaded component.
 	InitResources();
