- Add option to keep global graphics & compute PSOs in memory (default true for graphics, false for compute)
- Add option to keep Slate material PSOs in memory (default true)
- Add option to skip slate mateiral draw when PSOs are still compiling (Precache Too late because material loaded too late - default false)

#Synced-CL 50645574

diff --git a/Engine/Source/Runtime/Engine/Private/PSOPrecacheMaterial.cpp b/Engine/Source/Runtime/Engine/Private/PSOPrecacheMaterial.cpp
--- a/Engine/Source/Runtime/Engine/Private/PSOPrecacheMaterial.cpp
+++ b/Engine/Source/Runtime/Engine/Private/PSOPrecacheMaterial.cpp
@@ -1127,6 +1127,11 @@ void ClearMaterialPSORequests()
 	return GMaterialPSORequestManager.ClearMaterialPSORequests();
 }
 
+bool KeepPrecachedPSOsInMemoryForActiveMaterials()
+{
+	return GPSOPrecacheKeepInMemoryForActiveMaterials > 0;
+}
+
 #if PSO_PRECACHING_TRACKING
 
 FMaterialPSOPrecacheParams GetMaterialPSOPrecacheParams(FMaterialPSOPrecacheRequestID RequestID)
diff --git a/Engine/Source/Runtime/Engine/Private/PSOPrecacheSettings.cpp b/Engine/Source/Runtime/Engine/Private/PSOPrecacheSettings.cpp
--- a/Engine/Source/Runtime/Engine/Private/PSOPrecacheSettings.cpp
+++ b/Engine/Source/Runtime/Engine/Private/PSOPrecacheSettings.cpp
@@ -21,6 +21,22 @@ static FAutoConsoleVariableRef CVarPSOPrecacheGlobalShaders(
 	ECVF_ReadOnly
 );
 
+static int32 GPSOPrecacheKeepGlobalGraphicPSOsInMemory = -1;
+static FAutoConsoleVariableRef CVarPSOPrecacheKeepGlobalGraphicPSOsInMemory(
+	TEXT("r.PSOPrecache.KeepGlobalGraphicPSOsInMemory"),
+	GPSOPrecacheKeepGlobalGraphicPSOsInMemory,
+	TEXT("Keep global graphic PSOs in memory when r.PSOPrecache.KeepInMemoryForActiveMaterials is enabled"),
+	ECVF_ReadOnly
+);
+
+static int32 GPSOPrecacheKeepGlobalComputePSOsInMemory = -1;
+static FAutoConsoleVariableRef CVarPSOPrecacheKeepGlobalComputePSOsInMemory(
+	TEXT("r.PSOPrecache.KeepGlobalComputePSOsInMemory"),
+	GPSOPrecacheKeepGlobalComputePSOsInMemory,
+	TEXT("Keep global compute PSOs in memory when r.PSOPrecache.KeepInMemoryForActiveMaterials is enabled"),
+	ECVF_ReadOnly
+);
+
 static int32 GPSOPrecacheLightMapPolicyMode = -1;
 static FAutoConsoleVariableRef CVarPSOPrecacheLightMapPolicyMode(
 	TEXT("r.PSOPrecache.LightMapPolicyMode"),
@@ -153,6 +169,20 @@ static FAutoConsoleVariableRef CVarPSOPrecacheSlateMaterials(
 	TEXT("Precache all possible required PSOs for loaded Slate Materials (-1 use engine default)"),
 	ECVF_ReadOnly);
 
+static int32 GPSOPrecacheSkipSlateMaterialDrawOnPSOTooLate = -1;
+static FAutoConsoleVariableRef CVarPSOPrecacheSkipSlateMaterialDrawOnPSOTooLate(
+	TEXT("r.PSOPrecache.SkipSlateMaterialDrawOnPSOTooLate"), 
+	GPSOPrecacheSkipSlateMaterialDrawOnPSOTooLate, 
+	TEXT("Skip draw of slate material batches when material based PSOs are still compiling and r.PSOPrecache.ProxyCreationStrategy != 0 (-1 use engine default)"),
+	ECVF_ReadOnly);
+
+static int32 GPSOPrecacheKeepSlatePSOsInMemoryForActiveMaterials = -1;
+static FAutoConsoleVariableRef CVarPSOPrecacheKeepSlatePSOsInMemoryForActiveMaterials(
+	TEXT("r.PSOPrecache.KeepSlatePSOsInMemoryForActiveMaterials"), 
+	GPSOPrecacheKeepSlatePSOsInMemoryForActiveMaterials, 
+	TEXT("Keep slate PSOs in memory when r.PSOPrecache.KeepInMemoryForActiveMaterials is enabled (-1 use engine default)"),
+	ECVF_ReadOnly);
+
 static int32 GPSOPrecacheNaniteProgrammableRasterRequired = -1;
 static FAutoConsoleVariableRef CVarPSOPrecacheNaniteProgrammableRasterRequired(
 	TEXT("r.PSOPrecache.NaniteProgrammableRasterRequired"), 
@@ -180,6 +210,8 @@ static void ApplyPSOPrecacheSettingsCVarOverrides(FPSOPrecacheSettings& PSOPreca
 	}
 
 	SET_CVAR_VALUE(GPSOPrecacheGlobalShaders, GlobalShaderMode);
+	SET_CVAR_VALUE_BOOL(GPSOPrecacheKeepGlobalGraphicPSOsInMemory, bKeepGlobalGraphicPSOsInMemory);
+	SET_CVAR_VALUE_BOOL(GPSOPrecacheKeepGlobalComputePSOsInMemory, bKeepGlobalComputePSOsInMemory);
 	SET_CVAR_VALUE_BOOL(GPSOPrecacheTranslucencyAllPass, bTranslucencyAll);
 	SET_CVAR_VALUE_BOOL(GPSOPrecacheTranslucencyStandardBeforeDOF, bTranslucencyStandardBeforeDOF);
 	SET_CVAR_VALUE_BOOL(GPSOPrecacheAlphaColorChannel, bBasePassColorWithAlphaChannel);
@@ -193,6 +225,8 @@ static void ApplyPSOPrecacheSettingsCVarOverrides(FPSOPrecacheSettings& PSOPreca
 	SET_CVAR_VALUE(GPSOPrecacheParticlePrecachingTime, ParticlePrecachingTime);
 	SET_CVAR_VALUE_BOOL(GPSOPrecacheWaterInfoDepthPass, bWaterInfoDepth);
 	SET_CVAR_VALUE_BOOL(GPSOPrecacheSlateMaterials, bSlateMaterials);
+	SET_CVAR_VALUE_BOOL(GPSOPrecacheSkipSlateMaterialDrawOnPSOTooLate, bSkipSlateMaterialDrawOnPSOTooLate);
+	SET_CVAR_VALUE_BOOL(GPSOPrecacheKeepSlatePSOsInMemoryForActiveMaterials, bKeepSlatePSOsInMemoryForActiveMaterials);
 	SET_CVAR_VALUE_BOOL(GPSOPrecacheNaniteProgrammableRasterRequired, bNaniteProgrammableRasterRequired);
 
 #undef SET_CVAR_VALUE
diff --git a/Engine/Source/Runtime/Engine/Private/ShaderCompiler/ShaderCompiler.cpp b/Engine/Source/Runtime/Engine/Private/ShaderCompiler/ShaderCompiler.cpp
--- a/Engine/Source/Runtime/Engine/Private/ShaderCompiler/ShaderCompiler.cpp
+++ b/Engine/Source/Runtime/Engine/Private/ShaderCompiler/ShaderCompiler.cpp
@@ -5046,8 +5046,8 @@ FRegisterGlobalPSOCollectorFunction RegisterGlobalComputePSOCollector(nullptr, G
 
 void PrecacheComputePipelineStatesForGlobalShaders(ERHIFeatureLevel::Type FeatureLevel, const ITargetPlatform* TargetPlatform)
 {
-	int32 PrecacheGlobalShaders = UPSOPrecacheSettingsManager::GetSettings().GlobalShaderMode;
-	if (PrecacheGlobalShaders == 0)
+	const FPSOPrecacheSettings& PSOPrecacheSettings = UPSOPrecacheSettingsManager::GetSettings();
+	if (PSOPrecacheSettings.GlobalShaderMode == 0)
 	{
 		return;
 	}
@@ -5071,7 +5071,7 @@ void PrecacheComputePipelineStatesForGlobalShaders(ERHIFeatureLevel::Type Featur
 	for (TLinkedList<FShaderType*>::TIterator ShaderTypeIt(FShaderType::GetTypeList()); ShaderTypeIt; ShaderTypeIt.Next())
 	{
 		FGlobalShaderType* GlobalShaderType = ShaderTypeIt->GetGlobalShaderType();
-		if (!GlobalShaderType || (GlobalShaderType->GetFrequency() != SF_Compute && PrecacheGlobalShaders == 1))
+		if (!GlobalShaderType || (GlobalShaderType->GetFrequency() != SF_Compute && PSOPrecacheSettings.GlobalShaderMode == 1))
 		{
 			continue;
 		}
@@ -5118,7 +5118,7 @@ void PrecacheComputePipelineStatesForGlobalShaders(ERHIFeatureLevel::Type Featur
 			};
 			FGlobalComputePSOPrecacheData* GlobalComputePSOPrecacheData = new FGlobalComputePSOPrecacheData(ComputeShadersToPrecache.Num());			
 
-			auto PrecacheComputePipelineState = [GlobalComputePSOPrecacheData](TShaderRef<FShader>& GlobalShader)
+			auto PrecacheComputePipelineState = [GlobalComputePSOPrecacheData, PSOPrecacheSettings](TShaderRef<FShader>& GlobalShader)
 			{
 				// PSO precache shaders are not required to all load correctly
 				constexpr bool bRequired = false;
@@ -5134,7 +5134,7 @@ void PrecacheComputePipelineStatesForGlobalShaders(ERHIFeatureLevel::Type Featur
 					PSOPrecacheData.Name = GlobalShader.GetType()->GetName();
 					PSOPrecacheData.Type = FPSOPrecacheData::EType::Compute;
 					PSOPrecacheData.ComputeInitializer = ComputeInitializer;
-					PSOPrecacheData.InMemoryMode = EPSOPrecacheInMemoryMode::Uncached;
+					PSOPrecacheData.InMemoryMode = PSOPrecacheSettings.bKeepGlobalComputePSOsInMemory && KeepPrecachedPSOsInMemoryForActiveMaterials() ? EPSOPrecacheInMemoryMode::Cached : EPSOPrecacheInMemoryMode::Uncached;
 				#if PSO_PRECACHING_VALIDATE
 					PSOPrecacheData.PSOCollectorIndex = GlobalComputePSOCollectorIndex;
 					PSOPrecacheData.VertexFactoryType = nullptr;
@@ -5211,10 +5211,10 @@ void PrecacheComputePipelineStatesForGlobalShaders(ERHIFeatureLevel::Type Featur
 		}
 	}
 
-	// Don't keep global PSOs in memory
+	// Set if the global graphics PSOs need to be kept in memory or not
 	for (FPSOPrecacheData& PrecacheData : GlobalPSOInitializers)
 	{
-		PrecacheData.InMemoryMode = EPSOPrecacheInMemoryMode::Uncached;
+		PrecacheData.InMemoryMode = PSOPrecacheSettings.bKeepGlobalGraphicPSOsInMemory && KeepPrecachedPSOsInMemoryForActiveMaterials() ? EPSOPrecacheInMemoryMode::Cached : EPSOPrecacheInMemoryMode::Uncached;
 	}
 
 	RequestPrecachePSOs(EPSOPrecacheType::Global, GlobalPSOInitializers);
diff --git a/Engine/Source/Runtime/Engine/Public/PSOPrecacheMaterial.h b/Engine/Source/Runtime/Engine/Public/PSOPrecacheMaterial.h
--- a/Engine/Source/Runtime/Engine/Public/PSOPrecacheMaterial.h
+++ b/Engine/Source/Runtime/Engine/Public/PSOPrecacheMaterial.h
@@ -157,6 +157,11 @@ extern ENGINE_API void BoostPSOPriorityOnDraw(const TArray<FMaterialPSOPrecacheR
  */
 extern ENGINE_API void ClearMaterialPSORequests();
 
+/**
+ * Check if all precached PSOs for active materials need to be kept in memory
+ */
+extern ENGINE_API bool KeepPrecachedPSOsInMemoryForActiveMaterials();
+
 /**
  * Get original FMaterialPSOPrecacheParams from precache request
  */
diff --git a/Engine/Source/Runtime/Engine/Public/PSOPrecacheSettings.h b/Engine/Source/Runtime/Engine/Public/PSOPrecacheSettings.h
--- a/Engine/Source/Runtime/Engine/Public/PSOPrecacheSettings.h
+++ b/Engine/Source/Runtime/Engine/Public/PSOPrecacheSettings.h
@@ -41,9 +41,25 @@ struct FPSOPrecacheSettings
 	UPROPERTY()
 	int GlobalShaderMode = 1;
 
+	/** Keep global graphic PSOs in memory when r.PSOPrecache.KeepInMemoryForActiveMaterials is enabled */
+	UPROPERTY()
+	bool bKeepGlobalGraphicPSOsInMemory = true;
+	
+	/** Keep global compute PSOs in memory when r.PSOPrecache.KeepInMemoryForActiveMaterials is enabled */
+	UPROPERTY()
+	bool bKeepGlobalComputePSOsInMemory = false;
+
 	/** Precache slate materials (default true) */
 	UPROPERTY()
 	bool bSlateMaterials = true;
+	
+	/** Skip draw of slate material batches when material based PSOs are still compiling and r.PSOPrecache.ProxyCreationStrategy != 0 */
+	UPROPERTY()
+	bool bSkipSlateMaterialDrawOnPSOTooLate = false;
+
+	/** Keep slate PSOs in memory when r.PSOPrecache.KeepInMemoryForActiveMaterials is enabled */
+	UPROPERTY()
+	bool bKeepSlatePSOsInMemoryForActiveMaterials = true;
 			
 	/** Precache all possible light map policies for base pass shaders, otherwise only LMP_NO_LIGHTMAP (default false) */
 	UPROPERTY()
diff --git a/Engine/Source/Runtime/SlateRHIRenderer/Private/SlateRHIRenderingPolicy.cpp b/Engine/Source/Runtime/SlateRHIRenderer/Private/SlateRHIRenderingPolicy.cpp
--- a/Engine/Source/Runtime/SlateRHIRenderer/Private/SlateRHIRenderingPolicy.cpp
+++ b/Engine/Source/Runtime/SlateRHIRenderer/Private/SlateRHIRenderingPolicy.cpp
@@ -37,6 +37,7 @@
 #include "MeshPassProcessor.h"
 #include "PSOPrecacheValidation.h"
 #include "PSOPrecacheSettings.h"
+#include "PSOPrecacheComponent.h"
 #include "VT/VirtualTextureFeedbackResource.h"
 
 extern void UpdateNoiseTextureParameters(FViewUniformShaderParameters& ViewUniformShaderParameters);
@@ -1378,6 +1379,7 @@ void DrawSlateRenderBatch(
 	check(RenderBatch.NumIndices > 0);
 	const uint32 PrimitiveCount = RenderBatch.DrawPrimitiveType == ESlateDrawPrimitive::LineList ? RenderBatch.NumIndices / 2 : RenderBatch.NumIndices / 3;
 
+	State.GraphicsPSOInit.StatePrecachePSOHash = 0;
 	if (RenderBatchOp.ShaderResourceType == ESlateShaderResource::Material)
 	{
 		SLATE_DRAW_EVENTF(RHICmdList, MaterialBatch, TEXT("Slate Material: %s"), *RenderBatchOp.MaterialName);
@@ -1387,6 +1389,20 @@ void DrawSlateRenderBatch(
 		State.GraphicsPSOInit.BoundShaderState.PixelShaderRHI  = RenderBatchOp.PixelBindings->Shader.GetPixelShader();
 		State.GraphicsPSOInit.PrimitiveType = GetRHIPrimitiveType(RenderBatch.DrawPrimitiveType);
 		
+		if (UPSOPrecacheSettingsManager::GetSettings().bSkipSlateMaterialDrawOnPSOTooLate && GetPSOPrecacheProxyCreationStrategy() != EPSOPrecacheProxyCreationStrategy::AlwaysCreate)
+		{
+			// Early out draw when slate material PSO is still compiling for Precache
+			State.GraphicsPSOInit.StatePrecachePSOHash = RHIComputeStatePrecachePSOHash(State.GraphicsPSOInit);
+			FGraphEventRef CompletionEvent;
+			FPSOPrecacheRequestID PSOPrecacheRequestID;
+			EPSOPrecacheResult PSOPrecacheResult = PipelineStateCache::CheckPipelineStateInCache(State.GraphicsPSOInit, CompletionEvent, PSOPrecacheRequestID);
+			if (PSOPrecacheResult == EPSOPrecacheResult::Active)
+			{
+				PipelineStateCache::BoostPrecachePriority(EPSOPrecachePriority::Highest, PSOPrecacheRequestID);
+				return;
+			}
+		}
+
 #if PSO_PRECACHING_VALIDATE
 		if (PSOCollectorStats::IsFullPrecachingValidationEnabled() && UPSOPrecacheSettingsManager::GetSettings().bSlateMaterials)
 		{
@@ -1854,7 +1870,7 @@ void AddSlatePSOInitializer(
 		PSOPrecacheData.Priority = EPSOPrecachePriority::High;
 		PSOPrecacheData.Type = FPSOPrecacheData::EType::Graphics;
 		PSOPrecacheData.GraphicsPSOInitializer = GraphicsPSOInit;
-		PSOPrecacheData.InMemoryMode = EPSOPrecacheInMemoryMode::Uncached; //< Don't cache Slate PSOs in memory - we have way too many unused ones which would eat up memory (should be make this an option?)
+		PSOPrecacheData.InMemoryMode = UPSOPrecacheSettingsManager::GetSettings().bKeepSlatePSOsInMemoryForActiveMaterials && KeepPrecachedPSOsInMemoryForActiveMaterials() ? EPSOPrecacheInMemoryMode::Cached : EPSOPrecacheInMemoryMode::Uncached;
 		#if PSO_PRECACHING_VALIDATE
 		PSOPrecacheData.PSOCollectorIndex = PSOCollectorIndex;
 		PSOPrecacheData.VertexFactoryType = nullptr;
