Fix for a rare memory stomp inside ServerMoveClientError. It is possible for the ServerPredictionData to be nullified part way through this logic (through a delegate fired in MoveComponent), which would mean the ServerData pointer is pointing to stale data. Reacquiring the ServerPredictionData guarantees the data is valid.

[RN] Fixed a rare memory stomp inside ServerMoveClientError if MoveComponent causes a call to Unpossess
#tests Built and tested on Dedicated Server Cooked Data and Editor
#Synced-CL 50644704

diff --git a/Engine/Source/Runtime/Engine/Private/Components/CharacterMovementComponent.cpp b/Engine/Source/Runtime/Engine/Private/Components/CharacterMovementComponent.cpp
--- a/Engine/Source/Runtime/Engine/Private/Components/CharacterMovementComponent.cpp
+++ b/Engine/Source/Runtime/Engine/Private/Components/CharacterMovementComponent.cpp
@@ -10272,6 +10272,10 @@ void UCharacterMovementComponent::ServerMoveHandleClientError(float ClientTimeSt
 		}
 	}
 
+	//Need to re-acquire ServerData, as there's a non-zero chance that MoveComponent triggers a delegate which could call ResetPredictionData_Server
+	ServerData = GetPredictionData_Server_Character();
+	check(ServerData);
+
 	// Compute the client error from the server's position
 	// If client has accumulated a noticeable positional error, correct them.
 	bNetworkLargeClientCorrection = ServerData->bForceClientUpdate;
