[World] Old DataLayer deprecation continuation part 4 - retire UDeprecatedDataLayerInstance.

#Synced-CL 50613390

#Synced-CL 50613390

diff --git a/Engine/Plugins/Editor/DataValidation/Source/DataValidation/Private/WorldPartitionChangelistValidator.cpp b/Engine/Plugins/Editor/DataValidation/Source/DataValidation/Private/WorldPartitionChangelistValidator.cpp
--- a/Engine/Plugins/Editor/DataValidation/Source/DataValidation/Private/WorldPartitionChangelistValidator.cpp
+++ b/Engine/Plugins/Editor/DataValidation/Source/DataValidation/Private/WorldPartitionChangelistValidator.cpp
@@ -424,7 +424,7 @@ void UWorldPartitionChangelistValidator::OnInvalidReferenceDataLayerAsset(const
 	if (bSubmittingWorldDataLayers)
 	{
 		FText CurrentError = FText::Format(LOCTEXT("DataValidation.Changelist.WorldPartition.InvalidDataLayerAsset", "Data layer {0} has no data layer asset."),
-			FText::FromName(DataLayerInstance->GetDataLayerFName()));
+			FText::FromName(DataLayerInstance->GetFName()));
 
 		AssetFails(CurrentAsset, CurrentError);
 	}
@@ -435,7 +435,7 @@ void UWorldPartitionChangelistValidator::OnInvalidDataLayerAssetType(const UData
 	if (bSubmittingWorldDataLayers)
 	{
 		FText CurrentError = FText::Format(LOCTEXT("DataValidation.Changelist.WorldPartition.InvalidDataLayerAssetAsset", "Data layer {0} is not compatible with Data Layer Asset {1} type {2}."),
-			FText::FromName(DataLayerInstance->GetDataLayerFName()), FText::FromName(DataLayerAsset->GetFName()), FText::FromName(DataLayerAsset->GetClass()->GetFName()));
+			FText::FromName(DataLayerInstance->GetFName()), FText::FromName(DataLayerAsset->GetFName()), FText::FromName(DataLayerAsset->GetClass()->GetFName()));
 
 		AssetFails(CurrentAsset, CurrentError);
 	}
@@ -493,8 +493,8 @@ void UWorldPartitionChangelistValidator::OnDataLayerAssetConflict(const UDataLay
 		|| bSubmittingWorldDataLayers)
 	{
 		FText CurrentError = FText::Format(LOCTEXT("DataValidation.Changelist.WorldPartition.DataLayerAssetConflict", "Data layer instance {0} and data layer instance {1} are both referencing data layer asset {2}."),
-			FText::FromName(DataLayerInstance->GetDataLayerFName()),
-			FText::FromName(ConflictingDataLayerInstance->GetDataLayerFName()),
+			FText::FromName(DataLayerInstance->GetFName()),
+			FText::FromName(ConflictingDataLayerInstance->GetFName()),
 			FText::FromString(DataLayerInstance->GetAsset()->GetFullName()));
 
 		AssetFails(CurrentAsset, CurrentError);
diff --git a/Engine/Source/Editor/DataLayerEditor/Private/DataLayer/DataLayerEditorSubsystem.cpp b/Engine/Source/Editor/DataLayerEditor/Private/DataLayer/DataLayerEditorSubsystem.cpp
--- a/Engine/Source/Editor/DataLayerEditor/Private/DataLayer/DataLayerEditorSubsystem.cpp
+++ b/Engine/Source/Editor/DataLayerEditor/Private/DataLayer/DataLayerEditorSubsystem.cpp
@@ -1287,7 +1287,7 @@ void UDataLayerEditorSubsystem::SetActorsPinStateInDataLayers(const TArray<UData
 	if (UWorldPartition* WorldPartition = World ? World->GetWorldPartition() : nullptr)
 	{
 		TSet<FName> DataLayerInstanceNames;
-		Algo::TransformIf(DataLayerInstances, DataLayerInstanceNames, [](UDataLayerInstance* DataLayerInstance) { return !!DataLayerInstance; }, [](UDataLayerInstance* DataLayerInstance) { return DataLayerInstance->GetDataLayerFName(); });
+		Algo::TransformIf(DataLayerInstances, DataLayerInstanceNames, [](UDataLayerInstance* DataLayerInstance) { return !!DataLayerInstance; }, [](UDataLayerInstance* DataLayerInstance) { return DataLayerInstance->GetFName(); });
 
 		if (DataLayerInstanceNames.Num())
 		{
diff --git a/Engine/Source/Editor/DataLayerEditor/Private/DataLayer/SDataLayerBrowser.cpp b/Engine/Source/Editor/DataLayerEditor/Private/DataLayer/SDataLayerBrowser.cpp
--- a/Engine/Source/Editor/DataLayerEditor/Private/DataLayer/SDataLayerBrowser.cpp
+++ b/Engine/Source/Editor/DataLayerEditor/Private/DataLayer/SDataLayerBrowser.cpp
@@ -88,7 +88,7 @@ void SDataLayerBrowser::Construct(const FArguments& InArgs)
 		{
 			if (const UDataLayerInstance* DataLayerInstance = DataLayerItem->GetDataLayer())
 			{
-				return DataLayerInstance->GetDataLayerFName().ToString();
+				return DataLayerInstance->GetFName().ToString();
 			}
 		}
 		else if (const FDataLayerActorTreeItem* DataLayerActorTreeItem = Item.CastTo<FDataLayerActorTreeItem>())
diff --git a/Engine/Source/Editor/UnrealEd/Classes/Commandlets/WorldPartitionDataLayerToAssetCommandLet.h b/Engine/Source/Editor/UnrealEd/Classes/Commandlets/WorldPartitionDataLayerToAssetCommandLet.h
--- a/Engine/Source/Editor/UnrealEd/Classes/Commandlets/WorldPartitionDataLayerToAssetCommandLet.h
+++ b/Engine/Source/Editor/UnrealEd/Classes/Commandlets/WorldPartitionDataLayerToAssetCommandLet.h
@@ -12,7 +12,6 @@
 #include "PackageSourceControlHelper.h"
 #include "WorldPartition/DataLayer/DataLayerAsset.h"
 #include "WorldPartition/DataLayer/DataLayerInstance.h"
-#include "WorldPartition/DataLayer/DeprecatedDataLayerInstance.h"
 #include "WorldPartition/DataLayer/Deprecated/DeprecatedDataLayer.h"
 
 #include "WorldPartitionDataLayerToAssetCommandLet.generated.h"
@@ -72,11 +71,15 @@ public:
 	TArray<TWeakObjectPtr<UDataLayerConversionInfo>> const& GetPreviousConversions() const { return PreviousConversionsInfo; }
 	const TWeakObjectPtr<UDataLayerConversionInfo>& GetCurrentConversion() const { return CurrentConvertingInfo; }
 
+	PRAGMA_DISABLE_DEPRECATION_WARNINGS
 	void SetDataLayerToConvert(const UDeprecatedDataLayerInstance* InDataLayerToConvert);
+	PRAGMA_ENABLE_DEPRECATION_WARNINGS
 	void SetDataLayerInstance(UDataLayerInstanceWithAsset* InDataLayerInstance);
 
+	PRAGMA_DISABLE_DEPRECATION_WARNINGS
 	UPROPERTY()
 	TObjectPtr<const UDeprecatedDataLayerInstance> DataLayerToConvert;
+	PRAGMA_ENABLE_DEPRECATION_WARNINGS
 
 	UPROPERTY()
 	TObjectPtr<UDataLayerAsset> DataLayerAsset;
@@ -101,7 +104,9 @@ public:
 	const TArray<TObjectPtr<UDataLayerConversionInfo>>& GetDataLayerConversionInfos() const { return DataLayerConversionInfo; }
 	const TArray<TWeakObjectPtr<UDataLayerConversionInfo>>& GetConvertingDataLayerConversionInfo() const { return ConvertingDataLayerInfo; }
 
+	PRAGMA_DISABLE_DEPRECATION_WARNINGS
 	UDataLayerConversionInfo* GetDataLayerConversionInfo(const UDeprecatedDataLayerInstance* DataLayer) const;
+	PRAGMA_ENABLE_DEPRECATION_WARNINGS
 	UDataLayerConversionInfo* GetDataLayerConversionInfo(const UDataLayerAsset* DataLayerAsset) const;
 	UDataLayerConversionInfo* GetDataLayerConversionInfo(const UDataLayerInstanceWithAsset* DataLayerInstance) const;
 
@@ -110,7 +115,9 @@ public:
 	PRAGMA_ENABLE_DEPRECATION_WARNINGS
 
 	UDataLayerConversionInfo* StoreExistingDataLayer(FAssetData& AssetData);
+	PRAGMA_DISABLE_DEPRECATION_WARNINGS
 	UDataLayerConversionInfo* StoreDataLayerAssetConversion(const UDeprecatedDataLayerInstance* DataLayerToConvert, UDataLayerAsset* NewDataLayerAsset);
+	PRAGMA_ENABLE_DEPRECATION_WARNINGS
 	UDataLayerConversionInfo* StoreDataLayerInstanceConversion(const UDataLayerAsset* DataLayerAsset, UDataLayerInstanceWithAsset* NewDataLayerInstance);
 	
 	bool SetPreviousConversions(UDataLayerConversionInfo* CurrentConversion, TArray<TWeakObjectPtr<UDataLayerConversionInfo>>&& PreviousConversions);
@@ -154,7 +161,9 @@ private:
 	UNREALED_API bool InitializeFromCommandLine(TArray<FString>& Tokens, TArray<FString> const& Switches, TMap<FString, FString> const& Params);
 
 	UNREALED_API bool BuildConversionInfos(TStrongObjectPtr<UDataLayerToAssetCommandletContext>& CommandletContext, FPackageSourceControlHelper& PackageHelper);
+	PRAGMA_DISABLE_DEPRECATION_WARNINGS
 	UNREALED_API bool CreateConversionFromDataLayer(TStrongObjectPtr<UDataLayerToAssetCommandletContext>& CommandletContext, const UDeprecatedDataLayerInstance* DataLayer, FPackageSourceControlHelper& PackageHelper);
+	PRAGMA_ENABLE_DEPRECATION_WARNINGS
 	UNREALED_API TObjectPtr<UDataLayerAsset> GetOrCreateDataLayerAssetForConversion(TStrongObjectPtr<UDataLayerToAssetCommandletContext>& CommandletContext, FName DataLayerName);
 
 	UNREALED_API bool ResolvePreviousConversionsToCurrent(TStrongObjectPtr<UDataLayerToAssetCommandletContext>& CommandletContext);
diff --git a/Engine/Source/Editor/UnrealEd/Private/Commandlets/WorldPartitionDataLayerToAssetCommandlet.cpp b/Engine/Source/Editor/UnrealEd/Private/Commandlets/WorldPartitionDataLayerToAssetCommandlet.cpp
--- a/Engine/Source/Editor/UnrealEd/Private/Commandlets/WorldPartitionDataLayerToAssetCommandlet.cpp
+++ b/Engine/Source/Editor/UnrealEd/Private/Commandlets/WorldPartitionDataLayerToAssetCommandlet.cpp
@@ -22,6 +22,7 @@
 
 DEFINE_LOG_CATEGORY(LogDataLayerToAssetCommandlet);
 
+PRAGMA_DISABLE_DEPRECATION_WARNINGS
 void UDataLayerConversionInfo::SetDataLayerToConvert(const UDeprecatedDataLayerInstance* InDataLayerToConvert)
 {
 	DataLayerToConvert = InDataLayerToConvert;
@@ -29,6 +30,7 @@ void UDataLayerConversionInfo::SetDataLayerToConvert(const UDeprecatedDataLayerI
 	DataLayerAsset->DataLayerType = DataLayerToConvert->GetType();
 	DataLayerAsset->DebugColor = DataLayerToConvert->GetDebugColor();
 }
+PRAGMA_ENABLE_DEPRECATION_WARNINGS
 
 void UDataLayerConversionInfo::SetDataLayerInstance(UDataLayerInstanceWithAsset* InDataLayerInstance)
 {
@@ -56,11 +58,13 @@ UDataLayerToAssetCommandletContext::UDataLayerToAssetCommandletContext(const FOb
 
 }
 
+PRAGMA_DISABLE_DEPRECATION_WARNINGS
 UDataLayerConversionInfo* UDataLayerToAssetCommandletContext::GetDataLayerConversionInfo(const UDeprecatedDataLayerInstance* DataLayer) const
 {
 	const TObjectPtr<UDataLayerConversionInfo>* Entry = DataLayerConversionInfo.FindByPredicate([&DataLayer](const TObjectPtr<UDataLayerConversionInfo>& Other) { return DataLayer == Other->DataLayerToConvert; });
 	return Entry != nullptr ? Entry->Get() : nullptr;
 }
+PRAGMA_ENABLE_DEPRECATION_WARNINGS
 
 UDataLayerConversionInfo* UDataLayerToAssetCommandletContext::GetDataLayerConversionInfo(const UDataLayerAsset* DataLayerAsset) const
 {
@@ -96,6 +100,7 @@ UDataLayerConversionInfo* UDataLayerToAssetCommandletContext::StoreExistingDataL
 	return DataLayerConversionInfo.Add_GetRef(ConversionInfo).Get();
 }
 
+PRAGMA_DISABLE_DEPRECATION_WARNINGS
 UDataLayerConversionInfo* UDataLayerToAssetCommandletContext::StoreDataLayerAssetConversion(const UDeprecatedDataLayerInstance* DataLayerToConvert, UDataLayerAsset* NewDataLayerAsset)
 {
 	if (UDataLayerConversionInfo* ConversionInfo = GetDataLayerConversionInfo(DataLayerToConvert))
@@ -128,13 +133,14 @@ UDataLayerConversionInfo* UDataLayerToAssetCommandletContext::StoreDataLayerAsse
 		*ConversionInfo->DataLayerAsset->GetFullName(), *ConversionInfo->DataLayerToConvert->GetDataLayerShortName());
 	return ConvertingDataLayerInfo.Add_GetRef(ConversionInfo).Get();
 }
+PRAGMA_ENABLE_DEPRECATION_WARNINGS
 
 UDataLayerConversionInfo* UDataLayerToAssetCommandletContext::StoreDataLayerInstanceConversion(const UDataLayerAsset* DataLayerAsset, UDataLayerInstanceWithAsset* NewDataLayerInstance)
 {
 	if (UDataLayerConversionInfo* ConversionInfo = GetDataLayerConversionInfo(NewDataLayerInstance))
 	{
 		UE_LOG(LogDataLayerToAssetCommandlet, Error, TEXT("Failed to assign asset %s to data layer instance %s. The instance is already associated to asset %s."),
-			*DataLayerAsset->GetFullName(), *NewDataLayerInstance->GetDataLayerFName().ToString(), *ConversionInfo->DataLayerAsset->GetFullName());
+			*DataLayerAsset->GetFullName(), *NewDataLayerInstance->GetFName().ToString(), *ConversionInfo->DataLayerAsset->GetFullName());
 		return nullptr;
 	}
 
@@ -144,12 +150,12 @@ UDataLayerConversionInfo* UDataLayerToAssetCommandletContext::StoreDataLayerInst
 		ConversionInfo->SetDataLayerInstance(NewDataLayerInstance);
 
 		UE_LOG(LogDataLayerToAssetCommandlet, Log, TEXT("Data Layer Instance %s is associated to Data Layer Asset %s for conversion"),
-			*ConversionInfo->DataLayerInstance->GetDataLayerFName().ToString(), *ConversionInfo->DataLayerAsset->GetFullName());
+			*ConversionInfo->DataLayerInstance->GetFName().ToString(), *ConversionInfo->DataLayerAsset->GetFullName());
 		return ConversionInfo;
 	}
 
 	UE_LOG(LogDataLayerToAssetCommandlet, Error, TEXT("Failed to assign asset %s to data layer instance %s. The asset was not retrieved in conversion data."),
-		*DataLayerAsset->GetFullName(), *NewDataLayerInstance->GetDataLayerFName().ToString());
+		*DataLayerAsset->GetFullName(), *NewDataLayerInstance->GetFName().ToString());
 	return nullptr;
 }
 
@@ -220,14 +226,14 @@ void UDataLayerToAssetCommandletContext::LogConversionInfos() const
 				});
 			}
 			
-
+			PRAGMA_DISABLE_DEPRECATION_WARNINGS
 			UE_LOG(LogDataLayerToAssetCommandlet, Verbose, TEXT("[Conversion Info] Data Layer %s\t\tData Layer Asset: %s\t\t\t\t\t\tData Layer Instance: %s\t\tConverting By: %s\t\t\t\t\t\tConflicting Previous Conversion: %s"),
 				info->DataLayerToConvert != nullptr ? *info->DataLayerToConvert->GetDataLayerShortName() : TEXT("None"),
 				info->DataLayerAsset != nullptr ? *info->DataLayerAsset->GetFullName() : TEXT("None"),
-				info->DataLayerInstance != nullptr ? *info->DataLayerInstance->GetDataLayerFName().ToString() : TEXT("None"),
+				info->DataLayerInstance != nullptr ? *info->DataLayerInstance->GetFName().ToString() : TEXT("None"),
 				info->CurrentConvertingInfo != nullptr ? *info->CurrentConvertingInfo->DataLayerAsset->GetFullName() : TEXT("None"),
 				*ConflictingConversionString);
-
+			PRAGMA_ENABLE_DEPRECATION_WARNINGS
 		}
 	}
 }
@@ -373,6 +379,7 @@ bool UDataLayerToAssetCommandlet::BuildConversionInfos(TStrongObjectPtr<UDataLay
 	uint32 ErrorCount = 0;
 	for (const UDataLayerInstance* DataLayerInstance : WorldDataLayers->DataLayerInstances)
 	{
+		PRAGMA_DISABLE_DEPRECATION_WARNINGS
 		if (const UDeprecatedDataLayerInstance* DataLayerToConvert = Cast<UDeprecatedDataLayerInstance>(DataLayerInstance))
 		{
 			if (!CreateConversionFromDataLayer(CommandletContext, DataLayerToConvert, PackageHelper))
@@ -380,11 +387,13 @@ bool UDataLayerToAssetCommandlet::BuildConversionInfos(TStrongObjectPtr<UDataLay
 				ErrorCount++;
 			}
 		}
+		PRAGMA_ENABLE_DEPRECATION_WARNINGS
 	}
 
 	return ErrorCount == 0;
 }
 
+PRAGMA_DISABLE_DEPRECATION_WARNINGS
 bool UDataLayerToAssetCommandlet::CreateConversionFromDataLayer(TStrongObjectPtr<UDataLayerToAssetCommandletContext>& CommandletContext, const UDeprecatedDataLayerInstance* DataLayer, FPackageSourceControlHelper &PackageHelper)
 {
 	if(TObjectPtr<UDataLayerAsset> DataLayerAsset = GetOrCreateDataLayerAssetForConversion(CommandletContext, FName(DataLayer->GetDataLayerShortName())))
@@ -402,6 +411,7 @@ bool UDataLayerToAssetCommandlet::CreateConversionFromDataLayer(TStrongObjectPtr
 	
 	return false;
 }
+PRAGMA_ENABLE_DEPRECATION_WARNINGS
 
 TObjectPtr<UDataLayerAsset> UDataLayerToAssetCommandlet::GetOrCreateDataLayerAssetForConversion(TStrongObjectPtr<UDataLayerToAssetCommandletContext>& CommandletContext, FName DataLayerAssetName)
 {
@@ -609,7 +619,7 @@ bool UDataLayerToAssetCommandlet::CreateDataLayerInstances(TStrongObjectPtr<UDat
 		{
 			UDataLayerInstanceWithAsset* DataLayerInstance = WorldDataLayers->CreateDataLayer<UDataLayerInstanceWithAsset>(ConvertingInfo->DataLayerAsset);
 			UE_LOG(LogDataLayerToAssetCommandlet, Log, TEXT("Created new Data Layer Instance %s"),
-				*DataLayerInstance->GetDataLayerFName().ToString());
+				*DataLayerInstance->GetFName().ToString());
 
 			CommandletContext->StoreDataLayerInstanceConversion(ConvertingInfo->DataLayerAsset, DataLayerInstance);
 		}
@@ -652,7 +662,7 @@ bool UDataLayerToAssetCommandlet::RebuildDataLayerHierarchies(TStrongObjectPtr<U
 			continue;
 		}
 
-		
+		PRAGMA_DISABLE_DEPRECATION_WARNINGS
 		const UDeprecatedDataLayerInstance* DeprecatedParent = Cast<UDeprecatedDataLayerInstance>(OldParent);
 		if (DeprecatedParent == nullptr)
 		{
@@ -687,7 +697,7 @@ bool UDataLayerToAssetCommandlet::RebuildDataLayerHierarchies(TStrongObjectPtr<U
 				*DeprecatedParent->GetDataLayerFullName(), *Child->GetDataLayerShortName());
 			ErrorCount++;
 		}
-
+		PRAGMA_ENABLE_DEPRECATION_WARNINGS
 	}
 
 	return ErrorCount == 0;
@@ -727,13 +737,14 @@ bool UDataLayerToAssetCommandlet::CommitConversion(TStrongObjectPtr<UDataLayerTo
 	AWorldDataLayers* WorldDataLayers = MainWorld->GetWorldDataLayers();
 	for (const TWeakObjectPtr<UDataLayerConversionInfo>& ConversionInfo : CommandletContext->GetConvertingDataLayerConversionInfo())
 	{
+		PRAGMA_DISABLE_DEPRECATION_WARNINGS
 		if (ConversionInfo->DataLayerToConvert != nullptr)
 		{
 			// Remove directly from DataLayerInstances as RemoveDataLayer method also cleans DeprecatedDataLayerNameToDataLayerInstance which is used for runtime conversion
 			if (WorldDataLayers->DataLayerInstances.Remove(ConversionInfo->DataLayerToConvert))
 			{
 				UE_LOG(LogDataLayerToAssetCommandlet, Log, TEXT("Deleted old data layer %s, it is now converted to Data Layer Asset %s and Data Layer Instance %s"),
-					*ConversionInfo->DataLayerToConvert->GetDataLayerShortName(), *ConversionInfo->DataLayerAsset->GetFullName(), *ConversionInfo->DataLayerInstance->GetDataLayerFName().ToString());
+					*ConversionInfo->DataLayerToConvert->GetDataLayerShortName(), *ConversionInfo->DataLayerAsset->GetFullName(), *ConversionInfo->DataLayerInstance->GetFName().ToString());
 			}
 			else
 			{
@@ -742,6 +753,7 @@ bool UDataLayerToAssetCommandlet::CommitConversion(TStrongObjectPtr<UDataLayerTo
 				ErrorCount++;
 			}
 		}
+		PRAGMA_ENABLE_DEPRECATION_WARNINGS
 	}
 
 	if (ErrorCount == 0 && bPerformSavePackages)
diff --git a/Engine/Source/Runtime/Engine/Private/ActorEditor.cpp b/Engine/Source/Runtime/Engine/Private/ActorEditor.cpp
--- a/Engine/Source/Runtime/Engine/Private/ActorEditor.cpp
+++ b/Engine/Source/Runtime/Engine/Private/ActorEditor.cpp
@@ -1819,7 +1819,7 @@ TArray<FName> AActor::GetDataLayerInstanceNames() const
 	DataLayerInstanceNames.Reserve(DataLayerInstances.Num());
 	for (const UDataLayerInstance* DataLayerInstance : DataLayerInstances)
 	{
-		DataLayerInstanceNames.Add(DataLayerInstance->GetDataLayerFName());
+		DataLayerInstanceNames.Add(DataLayerInstance->GetFName());
 	}
 	return DataLayerInstanceNames;
 }
diff --git a/Engine/Source/Runtime/Engine/Private/ActorPartition/ActorPartitionSubsystem.cpp b/Engine/Source/Runtime/Engine/Private/ActorPartition/ActorPartitionSubsystem.cpp
--- a/Engine/Source/Runtime/Engine/Private/ActorPartition/ActorPartitionSubsystem.cpp
+++ b/Engine/Source/Runtime/Engine/Private/ActorPartition/ActorPartitionSubsystem.cpp
@@ -505,7 +505,7 @@ bool UActorPartitionSubsystem::MoveActorToDataLayers(APartitionActor* InActor, c
 
 	UWorld* World = InActor->GetWorld();
 	TArray<FName> DataLayerInstanceNames;
-	Algo::Transform(InDataLayerInstances, DataLayerInstanceNames, [](UDataLayerInstance* DataLayerInstance) { return DataLayerInstance->GetDataLayerFName(); });
+	Algo::Transform(InDataLayerInstances, DataLayerInstanceNames, [](UDataLayerInstance* DataLayerInstance) { return DataLayerInstance->GetFName(); });
 	const uint32 DataLayerEditorContextHash = FDataLayerEditorContext(World, DataLayerInstanceNames).GetHash();
 	const FActorPartitionIdentifier ActorPartitionId(InActor->GetClass(), InActor->GetGridGuid(), FActorPartitionContextHash::Get(InActor->GetContentBundleGuid(), DataLayerEditorContextHash));
 	const TUniquePtr<FWorldPartitionActorDesc> ActorDesc = InActor->CreateActorDesc();
diff --git a/Engine/Source/Runtime/Engine/Private/WorldPartition/DataLayer/DataLayerEditorContext.cpp b/Engine/Source/Runtime/Engine/Private/WorldPartition/DataLayer/DataLayerEditorContext.cpp
--- a/Engine/Source/Runtime/Engine/Private/WorldPartition/DataLayer/DataLayerEditorContext.cpp
+++ b/Engine/Source/Runtime/Engine/Private/WorldPartition/DataLayer/DataLayerEditorContext.cpp
@@ -23,7 +23,7 @@ FDataLayerEditorContext::FDataLayerEditorContext(UWorld* InWorld, const TArray<F
 	{
 		if (const UDataLayerInstance* DataLayerInstance = DataLayerManager->GetDataLayerInstance(DataLayerInstanceName))
 		{
-			DataLayerInstances.AddUnique(DataLayerInstance->GetDataLayerFName());
+			DataLayerInstances.AddUnique(DataLayerInstance->GetFName());
 		}
 	}
 
diff --git a/Engine/Source/Runtime/Engine/Private/WorldPartition/DataLayer/DataLayerInstance.cpp b/Engine/Source/Runtime/Engine/Private/WorldPartition/DataLayer/DataLayerInstance.cpp
--- a/Engine/Source/Runtime/Engine/Private/WorldPartition/DataLayer/DataLayerInstance.cpp
+++ b/Engine/Source/Runtime/Engine/Private/WorldPartition/DataLayer/DataLayerInstance.cpp
@@ -87,13 +87,6 @@ namespace DataLayerInstance
 	static const FName NAME_DataLayerInstanceInitialRuntimeState("DataLayerInstanceInitialRuntimeState");
 };
 
-void UDataLayerInstance::GetAssetRegistryTags(TArray<FAssetRegistryTag>& OutTags) const
-{
-	PRAGMA_DISABLE_DEPRECATION_WARNINGS;
-	Super::GetAssetRegistryTags(OutTags);
-	PRAGMA_ENABLE_DEPRECATION_WARNINGS;
-}
-
 void UDataLayerInstance::GetAssetRegistryTags(FAssetRegistryTagsContext Context) const
 {
 	Super::GetAssetRegistryTags(Context);
@@ -104,10 +97,10 @@ void UDataLayerInstance::GetAssetRegistryTags(FAssetRegistryTagsContext Context)
 		Context.AddTag(FAssetRegistryTag(FPrimaryAssetId::PrimaryAssetDisplayNameTag, *GetDataLayerShortName(), FAssetRegistryTag::TT_Hidden));
 
 		// Set DataLayerInstance specific tags
-		Context.AddTag(FAssetRegistryTag(DataLayerInstance::NAME_DataLayerInstanceName, *GetDataLayerFName().ToString(), FAssetRegistryTag::TT_Hidden));
+		Context.AddTag(FAssetRegistryTag(DataLayerInstance::NAME_DataLayerInstanceName, *GetFName().ToString(), FAssetRegistryTag::TT_Hidden));
 		if (GetParent())
 		{
-			Context.AddTag(FAssetRegistryTag(DataLayerInstance::NAME_DataLayerInstanceParentName, *GetParent()->GetDataLayerFName().ToString(), FAssetRegistryTag::TT_Hidden));
+			Context.AddTag(FAssetRegistryTag(DataLayerInstance::NAME_DataLayerInstanceParentName, *GetParent()->GetFName().ToString(), FAssetRegistryTag::TT_Hidden));
 		}
 		if (const UDataLayerAsset* DataLayerAsset = GetAsset())
 		{
@@ -679,12 +672,12 @@ void UDataLayerInstance::AddChild(UDataLayerInstance* InDataLayer)
 
 EDataLayerRuntimeState UDataLayerInstance::GetRuntimeState() const
 {
-	return GetOuterWorldDataLayers()->GetDataLayerRuntimeStateByName(GetDataLayerFName());
+	return GetOuterWorldDataLayers()->GetDataLayerRuntimeStateByName(GetFName());
 }
 
 EDataLayerRuntimeState UDataLayerInstance::GetEffectiveRuntimeState() const
 {
-	return GetOuterWorldDataLayers()->GetDataLayerEffectiveRuntimeStateByName(GetDataLayerFName());
+	return GetOuterWorldDataLayers()->GetDataLayerEffectiveRuntimeStateByName(GetFName());
 }
 
 #undef LOCTEXT_NAMESPACE
diff --git a/Engine/Source/Runtime/Engine/Private/WorldPartition/DataLayer/DataLayerManager.cpp b/Engine/Source/Runtime/Engine/Private/WorldPartition/DataLayer/DataLayerManager.cpp
--- a/Engine/Source/Runtime/Engine/Private/WorldPartition/DataLayer/DataLayerManager.cpp
+++ b/Engine/Source/Runtime/Engine/Private/WorldPartition/DataLayer/DataLayerManager.cpp
@@ -511,7 +511,7 @@ void UDataLayerManager::DrawDataLayersStatus(UCanvas* Canvas, FVector2D& Offset)
 	{
 		if (DataLayerInstance->IsRuntime())
 		{
-			const FName DataLayerName = DataLayerInstance->GetDataLayerFName();
+			const FName DataLayerName = DataLayerInstance->GetFName();
 			if (!LoadedDataLayers.Contains(DataLayerName) && !ActiveDataLayers.Contains(DataLayerName))
 			{
 				UnloadedDataLayers.Add(DataLayerName);
@@ -558,11 +558,11 @@ void UDataLayerManager::DrawDataLayersStatus(UCanvas* Canvas, FVector2D& Offset)
 				FWorldPartitionStreamingQuerySource QuerySource;
 				QuerySource.bDataLayersOnly = true;
 				QuerySource.bSpatialQuery = false;
-				QuerySource.DataLayers.Add(InDataLayerInstance->GetDataLayerFName());
+				QuerySource.DataLayers.Add(InDataLayerInstance->GetFName());
 				return WorldPartitionSubsystem->IsStreamingCompleted(InTargetState, { QuerySource }, true);
 			};
 
-			const EWorldPartitionRuntimeCellState TargetState = ActiveDataLayers.Contains(DataLayerInstance->GetDataLayerFName()) ? EWorldPartitionRuntimeCellState::Activated : EWorldPartitionRuntimeCellState::Loaded;
+			const EWorldPartitionRuntimeCellState TargetState = ActiveDataLayers.Contains(DataLayerInstance->GetFName()) ? EWorldPartitionRuntimeCellState::Activated : EWorldPartitionRuntimeCellState::Loaded;
 
 			if (!DataLayerLoadTime)
 			{
@@ -820,7 +820,7 @@ uint32 UDataLayerManager::GetDataLayerEditorContextHash() const
 	TArray<FName> DataLayerInstanceNames;
 	for (UDataLayerInstance* DataLayerInstance : GetActorEditorContextDataLayers())
 	{
-		DataLayerInstanceNames.Add(DataLayerInstance->GetDataLayerFName());
+		DataLayerInstanceNames.Add(DataLayerInstance->GetFName());
 	}
 	return FDataLayerEditorContext(GetWorld(), DataLayerInstanceNames).GetHash();
 }
diff --git a/Engine/Source/Runtime/Engine/Private/WorldPartition/DataLayer/DataLayerUtils.cpp b/Engine/Source/Runtime/Engine/Private/WorldPartition/DataLayer/DataLayerUtils.cpp
--- a/Engine/Source/Runtime/Engine/Private/WorldPartition/DataLayer/DataLayerUtils.cpp
+++ b/Engine/Source/Runtime/Engine/Private/WorldPartition/DataLayer/DataLayerUtils.cpp
@@ -88,7 +88,7 @@ FDataLayerInstanceNames FDataLayerUtils::ResolveDataLayerInstanceNames(const UDa
 
 		if (Instance)
 		{
-			ResolvedDLINames.Add(Instance->GetDataLayerFName());
+			ResolvedDLINames.Add(Instance->GetFName());
 		}
 		else
 		{
diff --git a/Engine/Source/Runtime/Engine/Private/WorldPartition/DataLayer/DataLayersID.cpp b/Engine/Source/Runtime/Engine/Private/WorldPartition/DataLayer/DataLayersID.cpp
--- a/Engine/Source/Runtime/Engine/Private/WorldPartition/DataLayer/DataLayersID.cpp
+++ b/Engine/Source/Runtime/Engine/Private/WorldPartition/DataLayer/DataLayersID.cpp
@@ -15,7 +15,7 @@ FDataLayersID::FDataLayersID(const TArray<const UDataLayerInstance*>& InDataLaye
 	: Hash(0)
 {
 	TArray<FName> DataLayers;
-	Algo::TransformIf(InDataLayerInstances, DataLayers, [](const UDataLayerInstance* Item) { return Item->IsRuntime(); }, [](const UDataLayerInstance* Item) { return Item->GetDataLayerFName(); });
+	Algo::TransformIf(InDataLayerInstances, DataLayers, [](const UDataLayerInstance* Item) { return Item->IsRuntime(); }, [](const UDataLayerInstance* Item) { return Item->GetFName(); });
 
 	if (DataLayers.Num())
 	{
diff --git a/Engine/Source/Runtime/Engine/Private/WorldPartition/DataLayer/DeprecatedDataLayerInstance.cpp b/Engine/Source/Runtime/Engine/Private/WorldPartition/DataLayer/DeprecatedDataLayerInstance.cpp
deleted file mode 100644
diff --git a/Engine/Source/Runtime/Engine/Private/WorldPartition/DataLayer/WorldDataLayers.cpp b/Engine/Source/Runtime/Engine/Private/WorldPartition/DataLayer/WorldDataLayers.cpp
--- a/Engine/Source/Runtime/Engine/Private/WorldPartition/DataLayer/WorldDataLayers.cpp
+++ b/Engine/Source/Runtime/Engine/Private/WorldPartition/DataLayer/WorldDataLayers.cpp
@@ -156,11 +156,11 @@ void AWorldDataLayers::InitializeDataLayerRuntimeStates()
 				{
 					if (DataLayerInstance->GetInitialRuntimeState() == EDataLayerRuntimeState::Loaded)
 					{
-						TargetLoadedDataLayerNames.Add(DataLayerInstance->GetDataLayerFName());
+						TargetLoadedDataLayerNames.Add(DataLayerInstance->GetFName());
 					}
 					else if (DataLayerInstance->GetInitialRuntimeState() == EDataLayerRuntimeState::Activated)
 					{
-						TargetActiveDataLayerNames.Add(DataLayerInstance->GetDataLayerFName());
+						TargetActiveDataLayerNames.Add(DataLayerInstance->GetFName());
 					}
 					RuntimeDataLayerInstances.Add(DataLayerInstance);
 				}
@@ -212,7 +212,7 @@ bool AWorldDataLayers::SetDataLayerRuntimeState(const UDataLayerInstance* InData
 		return false;
 	}
 
-	const EDataLayerRuntimeState CurrentState = GetDataLayerRuntimeStateByName(InDataLayerInstance->GetDataLayerFName());
+	const EDataLayerRuntimeState CurrentState = GetDataLayerRuntimeStateByName(InDataLayerInstance->GetFName());
 
 	ESetDataLayerRuntimeStateError Reason;
 	if (!CanChangeDataLayerRuntimeState(InDataLayerInstance, &Reason))
@@ -254,16 +254,16 @@ bool AWorldDataLayers::SetDataLayerRuntimeState(const UDataLayerInstance* InData
 		TSet<FName>& TargetLoadedDataLayerNames = bIsLocalDataLayer ? LocalLoadedDataLayerNames : LoadedDataLayerNames;
 		TSet<FName>& TargetActiveDataLayerNames = bIsLocalDataLayer ? LocalActiveDataLayerNames : ActiveDataLayerNames;
 
-		TargetLoadedDataLayerNames.Remove(InDataLayerInstance->GetDataLayerFName());
-		TargetActiveDataLayerNames.Remove(InDataLayerInstance->GetDataLayerFName());
+		TargetLoadedDataLayerNames.Remove(InDataLayerInstance->GetFName());
+		TargetActiveDataLayerNames.Remove(InDataLayerInstance->GetFName());
 
 		if (InState == EDataLayerRuntimeState::Loaded)
 		{
-			TargetLoadedDataLayerNames.Add(InDataLayerInstance->GetDataLayerFName());
+			TargetLoadedDataLayerNames.Add(InDataLayerInstance->GetFName());
 		}
 		else if (InState == EDataLayerRuntimeState::Activated)
 		{
-			TargetActiveDataLayerNames.Add(InDataLayerInstance->GetDataLayerFName());
+			TargetActiveDataLayerNames.Add(InDataLayerInstance->GetFName());
 		}
 
 		// Update replicated properties
@@ -388,7 +388,7 @@ void AWorldDataLayers::ResolveEffectiveRuntimeState(const UDataLayerInstance* In
 	const bool bDataLayerClientOnly = InDataLayerInstance->IsClientOnly();
 	const bool bDataLayerServerOnly = InDataLayerInstance->IsServerOnly();
 	const bool bIsLocalDataLayer = (bDataLayerClientOnly || bDataLayerServerOnly);
-	const FName DataLayerName = InDataLayerInstance->GetDataLayerFName();
+	const FName DataLayerName = InDataLayerInstance->GetFName();
 	EDataLayerRuntimeState NewEffectiveRuntimeState = GetDataLayerRuntimeStateByName(DataLayerName);
 	const UDataLayerInstance* Parent = InDataLayerInstance->GetParent();
 
@@ -397,7 +397,7 @@ void AWorldDataLayers::ResolveEffectiveRuntimeState(const UDataLayerInstance* In
 		if (Parent->IsRuntime())
 		{
 			// Apply min logic with parent DataLayers
-			NewEffectiveRuntimeState = (EDataLayerRuntimeState)FMath::Min((int32)NewEffectiveRuntimeState, (int32)GetDataLayerRuntimeStateByName(Parent->GetDataLayerFName()));
+			NewEffectiveRuntimeState = (EDataLayerRuntimeState)FMath::Min((int32)NewEffectiveRuntimeState, (int32)GetDataLayerRuntimeStateByName(Parent->GetFName()));
 		}
 		Parent = Parent->GetParent();
 	}
@@ -507,8 +507,8 @@ void AWorldDataLayers::DumpDataLayerRecursively(const UDataLayerInstance* DataLa
 			else
 			{
 				return FString::Printf(TEXT("(Effective State = %s | Target State = %s)"),
-					GetDataLayerRuntimeStateName(GetDataLayerEffectiveRuntimeStateByName(DataLayer->GetDataLayerFName())),
-					GetDataLayerRuntimeStateName(GetDataLayerRuntimeStateByName(DataLayer->GetDataLayerFName()))
+					GetDataLayerRuntimeStateName(GetDataLayerEffectiveRuntimeStateByName(DataLayer->GetFName())),
+					GetDataLayerRuntimeStateName(GetDataLayerRuntimeStateByName(DataLayer->GetFName()))
 				);
 			}
 		}
@@ -645,7 +645,7 @@ TArray<FName> AWorldDataLayers::GetDataLayerInstanceNames(const TArray<const UDa
 
 	for (const UDataLayerInstance* DataLayerInstance : GetDataLayerInstances(InDataLayersAssets))
 	{
-		OutDataLayerNames.Add(DataLayerInstance->GetDataLayerFName());
+		OutDataLayerNames.Add(DataLayerInstance->GetFName());
 	}
 
 	return OutDataLayerNames;
@@ -745,7 +745,7 @@ void AWorldDataLayers::SetAllowRuntimeDataLayerEditing(bool bInAllowRuntimeDataL
 
 bool AWorldDataLayers::IsActorEditorContextCurrentColorized(const UDataLayerInstance* InDataLayerInstance) const
 {
-	return InDataLayerInstance && !CurrentDataLayers.CurrentColorizedDataLayerInstanceName.IsNone() && (InDataLayerInstance->GetDataLayerFName() == CurrentDataLayers.CurrentColorizedDataLayerInstanceName);
+	return InDataLayerInstance && !CurrentDataLayers.CurrentColorizedDataLayerInstanceName.IsNone() && (InDataLayerInstance->GetFName() == CurrentDataLayers.CurrentColorizedDataLayerInstanceName);
 }
 
 bool AWorldDataLayers::IsInActorEditorContext(const UDataLayerInstance* InDataLayerInstance) const
@@ -778,7 +778,7 @@ void AWorldDataLayers::UpdateCurrentColorizedDataLayerInstance()
 	TArray<UDataLayerInstance*> ContextDataLayerInstances = GetActorEditorContextDataLayers();
 	if (ContextDataLayerInstances.Num() == 1)
 	{
-		CurrentDataLayers.CurrentColorizedDataLayerInstanceName = ContextDataLayerInstances[0]->GetDataLayerFName();
+		CurrentDataLayers.CurrentColorizedDataLayerInstanceName = ContextDataLayerInstances[0]->GetFName();
 	}
 	else if (ContextDataLayerInstances.Num() == 2)
 	{
@@ -787,7 +787,7 @@ void AWorldDataLayers::UpdateCurrentColorizedDataLayerInstance()
 			const UExternalDataLayerInstance* ExternalDataLayerInstance = DataLayerInstance->GetRootExternalDataLayerInstance();
 			if (ExternalDataLayerInstance && (ExternalDataLayerInstance != DataLayerInstance) && ContextDataLayerInstances.Contains(ExternalDataLayerInstance))
 			{
-				CurrentDataLayers.CurrentColorizedDataLayerInstanceName = DataLayerInstance->GetDataLayerFName();
+				CurrentDataLayers.CurrentColorizedDataLayerInstanceName = DataLayerInstance->GetFName();
 				break;
 			}
 		}
@@ -805,7 +805,7 @@ bool AWorldDataLayers::AddToActorEditorContext(UDataLayerInstance* InDataLayerIn
 	if (UExternalDataLayerInstance* ExternalDataLayerInstance = Cast<UExternalDataLayerInstance>(InDataLayerInstance))
 	{
 		Modify(/*bDirty*/false);
-		CurrentDataLayers.ExternalDataLayerName = ExternalDataLayerInstance->GetDataLayerFName();
+		CurrentDataLayers.ExternalDataLayerName = ExternalDataLayerInstance->GetFName();
 
 		// Adding an EDL Instance replaces the existing (if any) and removes all DataLayerInstances with a different root EDL Instance
 		TArray<FName> ToRemove;
@@ -828,10 +828,10 @@ bool AWorldDataLayers::AddToActorEditorContext(UDataLayerInstance* InDataLayerIn
 
 		bSuccess = true;
 	}
-	else if (!CurrentDataLayers.DataLayerInstanceNames.Contains(InDataLayerInstance->GetDataLayerFName()))
+	else if (!CurrentDataLayers.DataLayerInstanceNames.Contains(InDataLayerInstance->GetFName()))
 	{
 		Modify(/*bDirty*/false);
-		CurrentDataLayers.DataLayerInstanceNames.Add(InDataLayerInstance->GetDataLayerFName());
+		CurrentDataLayers.DataLayerInstanceNames.Add(InDataLayerInstance->GetFName());
 		bSuccess = true;
 
 		// Adding a Data Layer with a RootExternalDataLayerInstance will set this Root EDL Instance in the context
@@ -854,7 +854,7 @@ bool AWorldDataLayers::RemoveFromActorEditorContext(UDataLayerInstance* InDataLa
 	check(ContainsDataLayer(InDataLayerInstance));
 
 	const UExternalDataLayerInstance* ExternalDataLayerInstance = Cast<UExternalDataLayerInstance>(InDataLayerInstance);
-	const FName ExternalDataLayerName = ExternalDataLayerInstance ? ExternalDataLayerInstance->GetDataLayerFName() : NAME_None;
+	const FName ExternalDataLayerName = ExternalDataLayerInstance ? ExternalDataLayerInstance->GetFName() : NAME_None;
 	if (!ExternalDataLayerName.IsNone() && (CurrentDataLayers.ExternalDataLayerName == ExternalDataLayerName))
 	{
 		Modify(/*bDirty*/false);
@@ -876,10 +876,10 @@ bool AWorldDataLayers::RemoveFromActorEditorContext(UDataLayerInstance* InDataLa
 		}
 		return true;
 	}
-	else if (CurrentDataLayers.DataLayerInstanceNames.Contains(InDataLayerInstance->GetDataLayerFName()))
+	else if (CurrentDataLayers.DataLayerInstanceNames.Contains(InDataLayerInstance->GetFName()))
 	{
 		Modify(/*bDirty*/false);
-		CurrentDataLayers.DataLayerInstanceNames.Remove(InDataLayerInstance->GetDataLayerFName());
+		CurrentDataLayers.DataLayerInstanceNames.Remove(InDataLayerInstance->GetFName());
 		return true;
 	}
 	return false;
@@ -966,7 +966,7 @@ const UDataLayerInstance* AWorldDataLayers::GetDataLayerInstance(const FName& In
 		const UDataLayerInstance* FoundDataLayerInstance = nullptr;
 		ForEachDataLayerInstance([InDataLayerInstanceName, &FoundDataLayerInstance](UDataLayerInstance* DataLayerInstance)
 		{
-			if (DataLayerInstance->GetDataLayerFName() == InDataLayerInstanceName)
+			if (DataLayerInstance->GetFName() == InDataLayerInstanceName)
 			{
 				FoundDataLayerInstance = DataLayerInstance;
 				return false;
@@ -1266,30 +1266,22 @@ void AWorldDataLayers::UpdateAccelerationTable(const UDataLayerInstance* DataLay
 {
 	if (bIsAdding)
 	{
-		InstanceNameToInstance.Add(DataLayerInstance->GetDataLayerFName(), DataLayerInstance);
+		InstanceNameToInstance.Add(DataLayerInstance->GetFName(), DataLayerInstance);
 	}
 	else
 	{
-		InstanceNameToInstance.Remove(DataLayerInstance->GetDataLayerFName());
+		InstanceNameToInstance.Remove(DataLayerInstance->GetFName());
 	}
 
-	static_assert(DATALAYER_TO_INSTANCE_RUNTIME_CONVERSION_ENABLED, "Remove unnecessary cast. All DataLayerInstance now have assets");
-	if (const UDataLayerInstanceWithAsset* DataLayerInstanceWithAsset = Cast<UDataLayerInstanceWithAsset>(DataLayerInstance))
+	if (const UDataLayerAsset* DataLayerAsset = DataLayerInstance->GetAsset())
 	{
-		if (const UDataLayerAsset* DataLayerAsset = DataLayerInstanceWithAsset->GetAsset())
+		if (bIsAdding)
 		{
-			if (bIsAdding)
-			{
-				AssetNameToInstance.Add(DataLayerAsset->GetPathName(), DataLayerInstance);
-			}
-			else
-			{
-				AssetNameToInstance.Remove(DataLayerAsset->GetPathName());
-			}
+			AssetNameToInstance.Add(DataLayerAsset->GetPathName(), DataLayerInstance);
 		}
 		else
 		{
-			UE_LOG(LogWorldPartition, Warning, TEXT("DataLayerWithAsset %s has null asset."), *DataLayerInstanceWithAsset->GetPathName());
+			AssetNameToInstance.Remove(DataLayerAsset->GetPathName());
 		}
 	}
 }
diff --git a/Engine/Source/Runtime/Engine/Private/WorldPartition/DataLayer/WorldDataLayersActorDesc.cpp b/Engine/Source/Runtime/Engine/Private/WorldPartition/DataLayer/WorldDataLayersActorDesc.cpp
--- a/Engine/Source/Runtime/Engine/Private/WorldPartition/DataLayer/WorldDataLayersActorDesc.cpp
+++ b/Engine/Source/Runtime/Engine/Private/WorldPartition/DataLayer/WorldDataLayersActorDesc.cpp
@@ -202,8 +202,8 @@ FDataLayerInstanceDesc::FDataLayerInstanceDesc()
 
 void FDataLayerInstanceDesc::Init(UDataLayerInstance* InDataLayerInstance)
 {
-	Name = InDataLayerInstance->GetDataLayerFName();
-	ParentName = InDataLayerInstance->GetParent() ? InDataLayerInstance->GetParent()->GetDataLayerFName() : NAME_None;
+	Name = InDataLayerInstance->GetFName();
+	ParentName = InDataLayerInstance->GetParent() ? InDataLayerInstance->GetParent()->GetFName() : NAME_None;
 	const UDataLayerInstanceWithAsset* DataLayerWithAsset = Cast<UDataLayerInstanceWithAsset>(InDataLayerInstance);
 	const UDataLayerAsset* DataLayerAsset = InDataLayerInstance->GetAsset();
 	bIsUsingAsset = InDataLayerInstance->Implements<UDataLayerInstanceWithAsset>() || (DataLayerAsset != nullptr);
diff --git a/Engine/Source/Runtime/Engine/Private/WorldPartition/ErrorHandling/WorldPartitionStreamingGenerationLogErrorHandler.cpp b/Engine/Source/Runtime/Engine/Private/WorldPartition/ErrorHandling/WorldPartitionStreamingGenerationLogErrorHandler.cpp
--- a/Engine/Source/Runtime/Engine/Private/WorldPartition/ErrorHandling/WorldPartitionStreamingGenerationLogErrorHandler.cpp
+++ b/Engine/Source/Runtime/Engine/Private/WorldPartition/ErrorHandling/WorldPartitionStreamingGenerationLogErrorHandler.cpp
@@ -100,7 +100,7 @@ void FStreamingGenerationLogErrorHandler::OnInvalidWorldDataLayersReference(cons
 
 void FStreamingGenerationLogErrorHandler::OnDataLayerAssetConflict(const UDataLayerInstanceWithAsset* DataLayerInstance, const UDataLayerInstanceWithAsset* ConflictingDataLayerInstance)
 {
-	UE_ASSET_LOG(LogWorldPartition, Log, DataLayerInstance, TEXT("Data Layer Instance %s and Data Layer Instance %s are both referencing Data Layer Asset %s"), *DataLayerInstance->GetDataLayerFName().ToString(), *ConflictingDataLayerInstance->GetDataLayerFName().ToString(), *DataLayerInstance->GetAsset()->GetFullName());
+	UE_ASSET_LOG(LogWorldPartition, Log, DataLayerInstance, TEXT("Data Layer Instance %s and Data Layer Instance %s are both referencing Data Layer Asset %s"), *DataLayerInstance->GetFName().ToString(), *ConflictingDataLayerInstance->GetFName().ToString(), *DataLayerInstance->GetAsset()->GetFullName());
 }
 
 void FStreamingGenerationLogErrorHandler::OnActorNeedsResave(const IWorldPartitionActorDescInstanceView& ActorDescView)
diff --git a/Engine/Source/Runtime/Engine/Private/WorldPartition/RuntimeHashSet/WorldPartitionRuntimeHashSetDebugDraw.cpp b/Engine/Source/Runtime/Engine/Private/WorldPartition/RuntimeHashSet/WorldPartitionRuntimeHashSetDebugDraw.cpp
--- a/Engine/Source/Runtime/Engine/Private/WorldPartition/RuntimeHashSet/WorldPartitionRuntimeHashSetDebugDraw.cpp
+++ b/Engine/Source/Runtime/Engine/Private/WorldPartition/RuntimeHashSet/WorldPartitionRuntimeHashSetDebugDraw.cpp
@@ -44,7 +44,7 @@ static TMap<FName, FColor> GetDataLayerDebugColors(const UWorldPartition* InWorl
 	{
 		DataLayerManager->ForEachDataLayerInstance([&DebugColors](UDataLayerInstance* DataLayerInstance)
 		{
-			DebugColors.Add(DataLayerInstance->GetDataLayerFName(), DataLayerInstance->GetDebugColor());
+			DebugColors.Add(DataLayerInstance->GetFName(), DataLayerInstance->GetDebugColor());
 			return true;
 		});
 	}
diff --git a/Engine/Source/Runtime/Engine/Private/WorldPartition/WorldPartition.cpp b/Engine/Source/Runtime/Engine/Private/WorldPartition/WorldPartition.cpp
--- a/Engine/Source/Runtime/Engine/Private/WorldPartition/WorldPartition.cpp
+++ b/Engine/Source/Runtime/Engine/Private/WorldPartition/WorldPartition.cpp
@@ -148,7 +148,7 @@ TMap<FName, FString> GetDataLayersDumpString(const UWorldPartition* WorldPartiti
 	const UDataLayerManager* DataLayerManager = WorldPartition->GetDataLayerManager();
 	DataLayerManager->ForEachDataLayerInstance([&DataLayersDumpString](const UDataLayerInstance* DataLayerInstance)
 	{
-		DataLayersDumpString.FindOrAdd(DataLayerInstance->GetDataLayerFName()) = FString::Format(TEXT("{0}{1})"), { DataLayerInstance->GetDataLayerShortName(), DataLayerInstance->GetDataLayerFName().ToString() });
+		DataLayersDumpString.FindOrAdd(DataLayerInstance->GetFName()) = FString::Format(TEXT("{0}{1})"), { DataLayerInstance->GetDataLayerShortName(), DataLayerInstance->GetFName().ToString() });
 		return true;
 	});
 	
diff --git a/Engine/Source/Runtime/Engine/Private/WorldPartition/WorldPartitionDebugHelper.cpp b/Engine/Source/Runtime/Engine/Private/WorldPartition/WorldPartitionDebugHelper.cpp
--- a/Engine/Source/Runtime/Engine/Private/WorldPartition/WorldPartitionDebugHelper.cpp
+++ b/Engine/Source/Runtime/Engine/Private/WorldPartition/WorldPartitionDebugHelper.cpp
@@ -70,7 +70,7 @@ FAutoConsoleCommand FWorldPartitionDebugHelper::DebugFilterByDataLayerCommand(
 							TArray<UDataLayerInstance*> DataLayerInstances = DataLayerManager->ConvertArgsToDataLayers(InArgs);
 							for (UDataLayerInstance* DataLayerInstance : DataLayerInstances)
 							{
-								DebugDataLayerFilter.Add(DataLayerInstance->GetDataLayerFName());
+								DebugDataLayerFilter.Add(DataLayerInstance->GetFName());
 							}
 						}
 						return true;
diff --git a/Engine/Source/Runtime/Engine/Private/WorldPartition/WorldPartitionRuntimeCell.cpp b/Engine/Source/Runtime/Engine/Private/WorldPartition/WorldPartitionRuntimeCell.cpp
--- a/Engine/Source/Runtime/Engine/Private/WorldPartition/WorldPartitionRuntimeCell.cpp
+++ b/Engine/Source/Runtime/Engine/Private/WorldPartition/WorldPartitionRuntimeCell.cpp
@@ -99,12 +99,12 @@ void UWorldPartitionRuntimeCell::SetDataLayers(const TArray<const UDataLayerInst
 			return false;
 		}
 
-		return A->GetDataLayerFName().ToString() < B->GetDataLayerFName().ToString();
+		return A->GetFName().ToString() < B->GetFName().ToString();
 	});
 
 	TArray<FName> SortedDataLayerInstanceNames;
 	bool bIsFirstDataLayerExternal = false;
-	Algo::Transform(SortedDataLayerInstances, SortedDataLayerInstanceNames, [](const UDataLayerInstance* DataLayerInstance) { return DataLayerInstance->GetDataLayerFName(); });
+	Algo::Transform(SortedDataLayerInstances, SortedDataLayerInstanceNames, [](const UDataLayerInstance* DataLayerInstance) { return DataLayerInstance->GetFName(); });
 	if (SortedDataLayerInstances.Num() > 0)
 	{
 		if (const UExternalDataLayerInstance* ExternalDataLayerInstance = Cast<UExternalDataLayerInstance>(SortedDataLayerInstances[0]))
@@ -189,7 +189,7 @@ bool UWorldPartitionRuntimeCell::HasContentBundle() const
 
 bool UWorldPartitionRuntimeCell::ContainsDataLayer(const UDataLayerInstance* DataLayerInstance) const
 {
-	return GetDataLayersInline().Contains(DataLayerInstance->GetDataLayerFName());
+	return GetDataLayersInline().Contains(DataLayerInstance->GetFName());
 }
 
 FName UWorldPartitionRuntimeCell::GetLevelPackageName() const
diff --git a/Engine/Source/Runtime/Engine/Private/WorldPartition/WorldPartitionRuntimeSpatialHash.cpp b/Engine/Source/Runtime/Engine/Private/WorldPartition/WorldPartitionRuntimeSpatialHash.cpp
--- a/Engine/Source/Runtime/Engine/Private/WorldPartition/WorldPartitionRuntimeSpatialHash.cpp
+++ b/Engine/Source/Runtime/Engine/Private/WorldPartition/WorldPartitionRuntimeSpatialHash.cpp
@@ -834,7 +834,7 @@ static TMap<FName, FColor> GetDataLayerDebugColors2(const UWorldPartition* InWor
 	{
 		DataLayerManager->ForEachDataLayerInstance([&DebugColors](UDataLayerInstance* DataLayerInstance)
 		{
-			DebugColors.Add(DataLayerInstance->GetDataLayerFName(), DataLayerInstance->GetDebugColor());
+			DebugColors.Add(DataLayerInstance->GetFName(), DataLayerInstance->GetDebugColor());
 			return true;
 		});
 	}
diff --git a/Engine/Source/Runtime/Engine/Private/WorldPartition/WorldPartitionStreamingPolicy.cpp b/Engine/Source/Runtime/Engine/Private/WorldPartition/WorldPartitionStreamingPolicy.cpp
--- a/Engine/Source/Runtime/Engine/Private/WorldPartition/WorldPartitionStreamingPolicy.cpp
+++ b/Engine/Source/Runtime/Engine/Private/WorldPartition/WorldPartitionStreamingPolicy.cpp
@@ -255,7 +255,7 @@ const TSet<FName>& UWorldPartitionStreamingPolicy::GetServerDisallowedStreamingO
 				{
 					if (const UDataLayerInstance* DataLayerInstance = DataLayerManager->GetDataLayerInstanceFromAssetName(FName(DataLayerAssetName)))
 					{
-						ServerDisallowStreamingOutDataLayers.Add(DataLayerInstance->GetDataLayerFName());
+						ServerDisallowStreamingOutDataLayers.Add(DataLayerInstance->GetFName());
 					}
 				}
 			}
diff --git a/Engine/Source/Runtime/Engine/Public/WorldPartition/DataLayer/DataLayer.h b/Engine/Source/Runtime/Engine/Public/WorldPartition/DataLayer/DataLayer.h
--- a/Engine/Source/Runtime/Engine/Public/WorldPartition/DataLayer/DataLayer.h
+++ b/Engine/Source/Runtime/Engine/Public/WorldPartition/DataLayer/DataLayer.h
@@ -2,6 +2,16 @@
 
 #pragma once
 
+#if UE_ENABLE_INCLUDE_ORDER_DEPRECATED_IN_5_8
+#include "CoreMinimal.h"
+#include "UObject/ObjectMacros.h"
+#include "UObject/UObjectGlobals.h"
+#include "UObject/Object.h"
+#include "ActorDataLayer.h"
+#include "Math/Color.h"
+#include "WorldPartition/DataLayer/DataLayerInstance.h"
+#endif // UE_ENABLE_INCLUDE_ORDER_DEPRECATED_IN_5_8
+
 #include "Misc/CoreMiscDefines.h"
 
 UE_DEPRECATED_HEADER(5.8, "The contents of this header have been fully deprecated. See DeprecatedDataLayer.h for details on particular cases.") 
diff --git a/Engine/Source/Runtime/Engine/Public/WorldPartition/DataLayer/DataLayerInstance.h b/Engine/Source/Runtime/Engine/Public/WorldPartition/DataLayer/DataLayerInstance.h
--- a/Engine/Source/Runtime/Engine/Public/WorldPartition/DataLayer/DataLayerInstance.h
+++ b/Engine/Source/Runtime/Engine/Public/WorldPartition/DataLayer/DataLayerInstance.h
@@ -11,8 +11,6 @@
 
 #include "DataLayerInstance.generated.h"
 
-#define DATALAYER_TO_INSTANCE_RUNTIME_CONVERSION_ENABLED 1
-
 class AWorldDataLayers;
 class FText;
 class IStreamingGenerationErrorHandler;
@@ -138,25 +136,6 @@ public:
 	static bool GetAssetRegistryInfoFromPackage(FName InDataLayerInstancePackageName, FDataLayerInstanceDesc& OutDataLayerInstanceDesc);
 	static bool GetAssetRegistryInfoFromPackage(const FAssetData& InAsset, FDataLayerInstanceDesc& OutDataLayerInstanceDesc);
 
-	//~Begin deprecation
-	UE_DEPRECATED(5.4, "Implement the version that takes FAssetRegistryTagsContext instead.")
-	ENGINE_API virtual void GetAssetRegistryTags(TArray<FAssetRegistryTag>& OutTags) const override;
-
-	UE_DEPRECATED(5.4, "SetChildParent was removed.")
-	ENGINE_API void SetChildParent(UDataLayerInstance* InParent);
-
-	UE_DEPRECATED(5.4, "Use IsReadOnly instead")
-	virtual bool IsLocked() const { FText Reason; return IsLocked(&Reason); }
-
-	UE_DEPRECATED(5.3, "Use CanEditShortName instead")
-	virtual bool SupportRelabeling() const { return false; }
-
-	UE_DEPRECATED(5.3, "Use SetShortName instead")
-	virtual bool RelabelDataLayer(FName NewDataLayerLabel) { return false; }
-
-	UE_DEPRECATED(5.3, "Use CanBeChildOf instead")
-	bool CanParent(const UDataLayerInstance* InParent) const { return CanBeChildOf(InParent); }
-	//~End deprecation
 #endif // WITH_EDITOR
 
 	UFUNCTION(Category = "Data Layer", BlueprintCallable)
@@ -194,10 +173,6 @@ public:
 
 	virtual bool CanHaveChildDataLayerInstance(const UDataLayerInstance* InChildDataLayerInstance) const { return InChildDataLayerInstance != nullptr; }
 	virtual bool CanHaveParentDataLayerInstance() const { return true; }
-	UE_DEPRECATED(5.4, "Use CanHaveParentDataLayerInstance instead")
-	virtual bool CanHaveParentDataLayer() const { return CanHaveParentDataLayerInstance(); }
-	UE_DEPRECATED(5.4, "Use CanHaveChildDataLayer instead")
-	virtual bool CanHaveChildDataLayers() const { return true; }
 
 	const UDataLayerInstance* GetParent() const { return Parent; }
 	UDataLayerInstance* GetParent() { return Parent; }
@@ -211,16 +186,10 @@ public:
 
 	const TArray<TObjectPtr<UDataLayerInstance>>& GetChildren() const { return Children; }
 	ENGINE_API void ForEachChild(TFunctionRef<bool(const UDataLayerInstance*)> Operation) const;
-
-#if DATALAYER_TO_INSTANCE_RUNTIME_CONVERSION_ENABLED
-public:
+	
+	UE_DEPRECATED(5.8, "Use GetFName instead")
 	virtual FName GetDataLayerFName() const { return GetFName(); }
 
-private:
-	FName GetFName() const { return Super::GetFName(); }
-	FString GetName() const { return Super::GetName();  }
-#endif // DATALAYER_TO_INSTANCE_RUNTIME_CONVERSION_ENABLED
-
 protected:
 #if WITH_EDITOR
 	ENGINE_API bool IsLocked(FText* OutReason) const;
diff --git a/Engine/Source/Runtime/Engine/Public/WorldPartition/DataLayer/DataLayerInstanceWithAsset.h b/Engine/Source/Runtime/Engine/Public/WorldPartition/DataLayer/DataLayerInstanceWithAsset.h
--- a/Engine/Source/Runtime/Engine/Public/WorldPartition/DataLayer/DataLayerInstanceWithAsset.h
+++ b/Engine/Source/Runtime/Engine/Public/WorldPartition/DataLayer/DataLayerInstanceWithAsset.h
@@ -46,8 +46,8 @@ public:
 	virtual bool IsClientOnly() const override { return DataLayerAsset && DataLayerAsset->IsClientOnly(); }
 	virtual bool IsServerOnly() const override { return DataLayerAsset && DataLayerAsset->IsServerOnly(); }
 	virtual FColor GetDebugColor() const override { return DataLayerAsset ? DataLayerAsset->GetDebugColor() : FColor::Black; }
-	virtual FString GetDataLayerShortName() const override { return DataLayerAsset ? DataLayerAsset->GetName() : GetDataLayerFName().ToString(); }
-	virtual FString GetDataLayerFullName() const override { return DataLayerAsset ? DataLayerAsset->GetPathName() : GetDataLayerFName().ToString(); }
+	virtual FString GetDataLayerShortName() const override { return DataLayerAsset ? DataLayerAsset->GetName() : GetFName().ToString(); }
+	virtual FString GetDataLayerFullName() const override { return DataLayerAsset ? DataLayerAsset->GetPathName() : GetFName().ToString(); }
 	//~ End UDataLayerInstance interface
 
 protected:
diff --git a/Engine/Source/Runtime/Engine/Public/WorldPartition/DataLayer/Deprecated/DeprecatedDataLayer.h b/Engine/Source/Runtime/Engine/Public/WorldPartition/DataLayer/Deprecated/DeprecatedDataLayer.h
--- a/Engine/Source/Runtime/Engine/Public/WorldPartition/DataLayer/Deprecated/DeprecatedDataLayer.h
+++ b/Engine/Source/Runtime/Engine/Public/WorldPartition/DataLayer/Deprecated/DeprecatedDataLayer.h
@@ -8,6 +8,7 @@
 #include "UObject/Object.h"
 #include "Math/Color.h"
 #include "WorldPartition/DataLayer/DataLayerInstance.h"
+#include "WorldPartition/DataLayer/DataLayerUtils.h"
 
 #include "DeprecatedDataLayer.generated.h"
 
@@ -53,7 +54,6 @@ struct FActorDataLayer
 	}
 };
 
-static_assert(DATALAYER_TO_INSTANCE_RUNTIME_CONVERSION_ENABLED, "UDEPRECATED_DataLayer class is deprecated and needs to be deleted.");
 class UE_DEPRECATED(all, "Use UDataLayerInstance & UDataLayerAsset to create DataLayers") UDEPRECATED_DataLayer;
 
 UCLASS(Config = Engine, PerObjectConfig, Within = WorldDataLayers, BlueprintType, Deprecated, MinimalAPI)
@@ -141,4 +141,98 @@ private:
 
 	UPROPERTY(meta = (DeprecatedProperty, DeprecationMessage = "Use UDataLayerInstance & UDataLayerAsset to create DataLayers"))
 	TArray<TObjectPtr<UDEPRECATED_DataLayer>> Children_DEPRECATED;
+};
+
+class UE_DEPRECATED(all, "Only for use by the conversion commandlet") UDeprecatedDataLayerInstance;
+
+UCLASS(Config = Engine, PerObjectConfig, Within = WorldDataLayers, AutoCollapseCategories = ("Data Layer|Advanced"), AutoExpandCategories = ("Data Layer|Editor", "Data Layer|Advanced|Runtime"), MinimalAPI)
+class UDeprecatedDataLayerInstance final : public UDataLayerInstance
+{
+	GENERATED_BODY()
+
+
+#if WITH_EDITOR
+	//~ Begin UObject Interface
+	virtual void PostLoad() override
+	{
+		if (DebugColor == FColor::Black)
+		{
+			DebugColor = FColor::MakeRandomSeededColor(GetTypeHash(GetFName().ToString()));
+		}
+
+		Super::PostLoad();
+	}
+	//~ End UObject Interface
+#endif
+
+public:
+	UDeprecatedDataLayerInstance()
+	: DataLayerType(EDataLayerType::Editor)
+	{}
+
+#if WITH_EDITOR
+	static FName MakeName()
+	{
+		return FName(FString::Format(TEXT("DataLayer_{0}"), { FGuid::NewGuid().ToString() }));
+	}
+	void OnCreated()
+	{
+		Modify(/*bAlwaysMarkDirty*/false);
+
+		FDataLayerUtils::SetDataLayerShortName(this, TEXT("DataLayer"));
+
+		DeprecatedDataLayerFName = TEXT("");
+
+		DebugColor = FColor::MakeRandomSeededColor(GetTypeHash(GetFName().ToString()));
+	}
+	PRAGMA_DISABLE_DEPRECATION_WARNINGS
+	FActorDataLayer GetActorDataLayer() const
+	{
+		return FActorDataLayer(GetFName());
+	}
+
+	static FName MakeName(const class UDEPRECATED_DataLayer* DeprecatedDataLayer) { return FName{}; }
+
+	void OnCreated(const class UDEPRECATED_DataLayer* DeprecatedDataLayer) {}
+	PRAGMA_ENABLE_DEPRECATION_WARNINGS
+#endif
+	
+	FName GetDataLayerLabel() const { return Label; }
+
+private:
+#if WITH_EDITOR
+	//~ Begin UDataLayerInstance Interface
+	virtual bool PerformAddActor(AActor* InActor) const override { return false; }
+	virtual bool PerformRemoveActor(AActor* InActor) const override { return false; }
+	virtual bool CanEditDataLayerShortName() const override { return true; }
+	virtual void PerformSetDataLayerShortName(const FString& InNewShortName) { Label = *InNewShortName; }
+	//~ Endif UDataLayerInstance Interface
+#endif
+
+	//~ Begin UDataLayerInstance Interface
+	FName GetDataLayerFName() const { return !DeprecatedDataLayerFName.IsNone() ? DeprecatedDataLayerFName : GetFName(); }
+	virtual EDataLayerType GetType() const override { return DataLayerType; }
+	virtual bool IsRuntime() const override { return DataLayerType == EDataLayerType::Runtime; }
+	virtual FColor GetDebugColor() const override { return DebugColor; }
+	virtual FString GetDataLayerShortName() const override { return Label.ToString(); }
+	virtual FString GetDataLayerFullName() const override { return TEXT("Deprecated_") + Label.ToString(); }
+	//~ End UDataLayerInstance Interface
+
+private:
+	UPROPERTY()
+	FName Label;
+
+	UPROPERTY()
+	FName DeprecatedDataLayerFName;
+
+	UPROPERTY(Category = "Data Layer", EditAnywhere)
+	EDataLayerType DataLayerType;
+
+	UPROPERTY(Category = "Data Layer|Runtime", EditAnywhere, meta = (EditConditionHides, EditCondition = "DataLayerType == EDataLayerType::Runtime"))
+	FColor DebugColor;
+
+	friend class AWorldDataLayers;
+	friend class UDataLayerConversionInfo;
+	friend class UDataLayerToAssetCommandletContext;
+	friend class UDataLayerToAssetCommandlet;
 };
\ No newline at end of file
diff --git a/Engine/Source/Runtime/Engine/Public/WorldPartition/DataLayer/DeprecatedDataLayerInstance.h b/Engine/Source/Runtime/Engine/Public/WorldPartition/DataLayer/DeprecatedDataLayerInstance.h
--- a/Engine/Source/Runtime/Engine/Public/WorldPartition/DataLayer/DeprecatedDataLayerInstance.h
+++ b/Engine/Source/Runtime/Engine/Public/WorldPartition/DataLayer/DeprecatedDataLayerInstance.h
@@ -1,79 +1,11 @@
 // Copyright Epic Games, Inc. All Rights Reserved.
 #pragma once
 
+#if UE_ENABLE_INCLUDE_ORDER_DEPRECATED_IN_5_8
 #include "WorldPartition/DataLayer/DataLayerInstance.h"
 #include "WorldPartition/DataLayer/Deprecated/DeprecatedDataLayer.h"
-#include "DeprecatedDataLayerInstance.generated.h"
+#endif // UE_ENABLE_INCLUDE_ORDER_DEPRECATED_IN_5_8
 
-// Class used for Runtime Conversion of the Deprecated UDataLayer Class to UDataLayerInstance + UDataLayerAsset.
-// This class is not to be inherited. It is solely used by AWorldDatalayers to convert UDataLayers to UDataLayerInstances on Level Boot.
-// You will need to run the DataLayerToAsset CommandLet to convert the deprecated datalayers to UDataLayerAssets and UDataLayerInstanceWithAsset.
-UCLASS(Config = Engine, PerObjectConfig, Within = WorldDataLayers, AutoCollapseCategories = ("Data Layer|Advanced"), AutoExpandCategories = ("Data Layer|Editor", "Data Layer|Advanced|Runtime"), MinimalAPI)
-class UDeprecatedDataLayerInstance final : public UDataLayerInstance
-{
-	static_assert(DATALAYER_TO_INSTANCE_RUNTIME_CONVERSION_ENABLED, "UDeprecatedDataLayerInstance is deprecated and needs to be deleted.");
+UE_DEPRECATED_HEADER(5.8, "The contents of this header have been fully deprecated. See DeprecatedDataLayer.h for details on particular cases.") 
 
-	GENERATED_UCLASS_BODY()
-
-#if WITH_EDITOR
-	//~ Begin UObject Interface
-	ENGINE_API virtual void PostLoad() override;
-	//~ End UObject Interface
-#endif
-
-public:
-#if WITH_EDITOR
-	static ENGINE_API FName MakeName();
-	ENGINE_API void OnCreated();
-	PRAGMA_DISABLE_DEPRECATION_WARNINGS
-	FActorDataLayer GetActorDataLayer() const;
-
-	UE_DEPRECATED(5.8, "UDEPRECATED_DataLayer are no longer supported")
-	static FName MakeName(const class UDEPRECATED_DataLayer* DeprecatedDataLayer) { return FName{}; }
-
-	UE_DEPRECATED(5.8, "UDEPRECATED_DataLayer are no longer supported")
-	void OnCreated(const class UDEPRECATED_DataLayer* DeprecatedDataLayer) {}
-	PRAGMA_ENABLE_DEPRECATION_WARNINGS
-#endif
-	
-	FName GetDataLayerLabel() const { return Label; }
-
-private:
-#if WITH_EDITOR
-	//~ Begin UDataLayerInstance Interface
-	virtual bool PerformAddActor(AActor* InActor) const override { return false; }
-	virtual bool PerformRemoveActor(AActor* InActor) const override { return false; }
-	virtual bool RelabelDataLayer(FName NewDataLayerLabel) override { return false; }
-	virtual bool SupportRelabeling() const override { return true; }
-	virtual bool CanEditDataLayerShortName() const override { return true; }
-	virtual void PerformSetDataLayerShortName(const FString& InNewShortName) { Label = *InNewShortName; }
-	//~ Endif UDataLayerInstance Interface
-#endif
-
-	//~ Begin UDataLayerInstance Interface
-	virtual FName GetDataLayerFName() const override { return !DeprecatedDataLayerFName.IsNone() ? DeprecatedDataLayerFName : Super::GetDataLayerFName(); }
-	virtual EDataLayerType GetType() const override { return DataLayerType; }
-	virtual bool IsRuntime() const override { return DataLayerType == EDataLayerType::Runtime; }
-	virtual FColor GetDebugColor() const override { return DebugColor; }
-	virtual FString GetDataLayerShortName() const override { return Label.ToString(); }
-	virtual FString GetDataLayerFullName() const override { return TEXT("Deprecated_") + Label.ToString(); }
-	//~ End UDataLayerInstance Interface
-
-private:
-	UPROPERTY()
-	FName Label;
-
-	UPROPERTY()
-	FName DeprecatedDataLayerFName;
-
-	UPROPERTY(Category = "Data Layer", EditAnywhere)
-	EDataLayerType DataLayerType;
-
-	UPROPERTY(Category = "Data Layer|Runtime", EditAnywhere, meta = (EditConditionHides, EditCondition = "DataLayerType == EDataLayerType::Runtime"))
-	FColor DebugColor;
-
-	friend class AWorldDataLayers;
-	friend class UDataLayerConversionInfo;
-	friend class UDataLayerToAssetCommandletContext;
-	friend class UDataLayerToAssetCommandlet;
-};
+#include "WorldPartition/DataLayer/Deprecated/DeprecatedDataLayer.h"
\ No newline at end of file
diff --git a/Engine/Source/Runtime/MovieSceneTracks/Private/Systems/MovieSceneDataLayerSystem.cpp b/Engine/Source/Runtime/MovieSceneTracks/Private/Systems/MovieSceneDataLayerSystem.cpp
--- a/Engine/Source/Runtime/MovieSceneTracks/Private/Systems/MovieSceneDataLayerSystem.cpp
+++ b/Engine/Source/Runtime/MovieSceneTracks/Private/Systems/MovieSceneDataLayerSystem.cpp
@@ -306,7 +306,7 @@ EDataLayerUpdateFlags FDesiredLayerStates::Apply(FPreAnimatedDataLayerStorage* P
 		FWorldPartitionStreamingQuerySource QuerySource;
 		QuerySource.bDataLayersOnly = true;
 		QuerySource.bSpatialQuery = false; // @todo_ow: how would we support spatial query from sequencer?
-		QuerySource.DataLayers.Add(DataLayer->GetDataLayerFName());
+		QuerySource.DataLayers.Add(DataLayer->GetFName());
 
 		UWorldPartition* WorldPartition = DataLayer->GetOuterWorld()->GetWorldPartition();
 		return WorldPartition->IsStreamingCompleted(QueryState, { QuerySource }, bExactState);
