Fix crash happening when trying to open material with missing expressions.
Normally, expressions are loaded from disk by serialization and get built intentionally only during Material Editor start-up. That is however after the CachedExpressionData build phase and after Translation, which need those expressions.
In cases like this, which should be rare, build them intentionally during PostLoad.

[RN] minor | Rendering Architecture

#Synced-CL 50586249

diff --git a/Engine/Source/Runtime/Engine/Private/Materials/Material.cpp b/Engine/Source/Runtime/Engine/Private/Materials/Material.cpp
--- a/Engine/Source/Runtime/Engine/Private/Materials/Material.cpp
+++ b/Engine/Source/Runtime/Engine/Private/Materials/Material.cpp
@@ -113,6 +113,7 @@
 #include "UObject/SavePackage.h"
 #include "Misc/MapErrors.h"
 #include "UObject/FortniteReleaseBranchCustomObjectVersion.h"
+#include "MaterialEditorUtilities.h"
 #endif
 #include "ShaderCodeLibrary.h"
 #include "String/ParseTokens.h"
@@ -4557,6 +4558,11 @@ void UMaterial::PostLoad()
 	// cooked materials will not have any expressions in them, so this will obliterate the saved cached expression data
 	if (!GetOutermost()->bIsCookedForEditor)
 	{
+		// Fix up if we are missing expressions
+		if (GetExpressions().Num() == 0)
+		{
+			FMaterialEditorUtilities::InitExpressions(this);
+		}
 		UpdateCachedExpressionData();
 	}
 	else
