[WorldPartition] Fix crash in UWorldPartitionRuntimeLevelStreamingCell::GetLevelPackageName when WorldPartition is uninitialized

#Synced-CL 50431656

diff --git a/Engine/Source/Runtime/Engine/Private/WorldPartition/WorldPartitionRuntimeLevelStreamingCell.cpp b/Engine/Source/Runtime/Engine/Private/WorldPartition/WorldPartitionRuntimeLevelStreamingCell.cpp
--- a/Engine/Source/Runtime/Engine/Private/WorldPartition/WorldPartitionRuntimeLevelStreamingCell.cpp
+++ b/Engine/Source/Runtime/Engine/Private/WorldPartition/WorldPartitionRuntimeLevelStreamingCell.cpp
@@ -88,7 +88,7 @@ FName UWorldPartitionRuntimeLevelStreamingCell::GetLevelPackageName() const
 {
 #if WITH_EDITOR
 	UWorld* World = GetOwningWorld();
-	if (World->IsPlayInEditor())
+	if (World && World->IsPlayInEditor())
 	{
 		return FName(UWorldPartitionLevelStreamingPolicy::GetCellPackagePath(GetFName(), World));
 	}
