Substrate - removed remaining material attribute blending support causing some issues (infinite loops) with materail topology preview on nodes when this is no longer supported anyway.
New material layering do not support that anyway now and it has been removed in 5.7.

#rnx

#Synced-CL 50709564

diff --git a/Engine/Source/Runtime/Engine/Private/Materials/MaterialAttributeDefinitionMap.cpp b/Engine/Source/Runtime/Engine/Private/Materials/MaterialAttributeDefinitionMap.cpp
--- a/Engine/Source/Runtime/Engine/Private/Materials/MaterialAttributeDefinitionMap.cpp
+++ b/Engine/Source/Runtime/Engine/Private/Materials/MaterialAttributeDefinitionMap.cpp
@@ -353,7 +353,7 @@ void FMaterialAttributeDefinitionMap::InitializeAttributeMap()
 	Add(FGuid(0x0AC97EC3, 0xE3D047BA, 0xB610167D, 0xC4D919FF), TEXT("PixelDepthOffset"),		MP_PixelDepthOffset,		MCT_Float,	FVector4(0,0,0,0),	SF_Pixel);
 	Add(FGuid(0xD9423FFF, 0xD77E4D82, 0x8FF9CF5E, 0x055D1255), TEXT("ShadingModel"),			MP_ShadingModel,			MCT_ShadingModel, FVector4(0, 0, 0, 0), SF_Pixel, INDEX_NONE, false, &CompileShadingModelBlendFunction);
 	Add(FGuid(0x42BDD2E0, 0xBE714189, 0xA0984BC3, 0xDD0BE872), TEXT("SurfaceThickness"),		MP_SurfaceThickness,		MCT_Float,  FVector4(SUBSTRATE_LAYER_DEFAULT_THICKNESS_CM, 0, 0, 0), SF_Pixel);
-	Add(FGuid(0x5973A03E, 0x13A74E08, 0x92D0CEDD, 0xF2936CF8), TEXT("FrontMaterial"),			MP_FrontMaterial,			MCT_Substrate, FVector4(0,0,0,0),	SF_Pixel, INDEX_NONE, false, &CompileSubstrateBlendFunction);
+	Add(FGuid(0x5973A03E, 0x13A74E08, 0x92D0CEDD, 0xF2936CF8), TEXT("FrontMaterial"),			MP_FrontMaterial,			MCT_Substrate, FVector4(0,0,0,0),	SF_Pixel);
 
 	// Used when compiling material with execution pins, which are compiling all attributes together
 	Add(FGuid(0xE0ED040B, 0x82794D93, 0xBD2D59B2, 0xA5BBF41C), TEXT("MaterialAttributes"),		MP_MaterialAttributes,		MCT_MaterialAttributes, FVector4(0,0,0,0), SF_Pixel, INDEX_NONE, bHideAttribute);
diff --git a/Engine/Source/Runtime/Engine/Private/Materials/MaterialExpressions.cpp b/Engine/Source/Runtime/Engine/Private/Materials/MaterialExpressions.cpp
--- a/Engine/Source/Runtime/Engine/Private/Materials/MaterialExpressions.cpp
+++ b/Engine/Source/Runtime/Engine/Private/Materials/MaterialExpressions.cpp
@@ -7160,16 +7160,12 @@ int32 UMaterialExpressionBlendMaterialAttributes::Compile(class FMaterialCompile
 	{
 	case EMaterialAttributeBlend::UseA:
 	{
-		Compiler->SubstrateTreeStackPush(this, 0);
 		int32 CodeChunk = A.CompileWithDefault(Compiler, AttributeID);
-		Compiler->SubstrateTreeStackPop();
 		return CodeChunk;
 	}
 	case EMaterialAttributeBlend::UseB:
 	{
-		Compiler->SubstrateTreeStackPush(this, 1);
 		int32 CodeChunk = B.CompileWithDefault(Compiler, AttributeID);
-		Compiler->SubstrateTreeStackPop();
 		return CodeChunk;
 	}
 	default:
@@ -7177,12 +7173,8 @@ int32 UMaterialExpressionBlendMaterialAttributes::Compile(class FMaterialCompile
 	}
 
 	// Allow custom blends or fallback to standard interpolation
-	Compiler->SubstrateTreeStackPush(this, 0);
 	int32 ResultA = A.CompileWithDefault(Compiler, AttributeID);
-	Compiler->SubstrateTreeStackPop();
-	Compiler->SubstrateTreeStackPush(this, 1);
 	int32 ResultB = B.CompileWithDefault(Compiler, AttributeID);
-	Compiler->SubstrateTreeStackPop();
 	if (AllowSkipAlphaComputationForMaterialBlend() && (ResultA == ResultB))
 	{
 		return ResultA;
@@ -7219,60 +7211,6 @@ FName UMaterialExpressionBlendMaterialAttributes::GetInputName(int32 InputIndex)
 
 	return Name;
 }
-
-bool UMaterialExpressionBlendMaterialAttributes::IsResultSubstrateMaterial(int32 OutputIndex)
-{
-	if (A.GetTracedInput().Expression)
-	{
-		return A.GetTracedInput().Expression->IsResultSubstrateMaterial(0); // can only blend Starta type together so one or the othjer input is enough
-	}
-	if (B.GetTracedInput().Expression)
-	{
-		return B.GetTracedInput().Expression->IsResultSubstrateMaterial(0);
-	}
-	return false;
-}
-
-void UMaterialExpressionBlendMaterialAttributes::GatherSubstrateMaterialInfo(FSubstrateMaterialInfo& SubstrateMaterialInfo, int32 OutputIndex)
-{
-	if (A.GetTracedInput().Expression)
-	{
-		A.GetTracedInput().Expression->GatherSubstrateMaterialInfo(SubstrateMaterialInfo, A.OutputIndex);
-	}
-	if (B.GetTracedInput().Expression)
-	{
-		B.GetTracedInput().Expression->GatherSubstrateMaterialInfo(SubstrateMaterialInfo, B.OutputIndex);
-	}
-}
-
-FSubstrateOperator* UMaterialExpressionBlendMaterialAttributes::SubstrateGenerateMaterialTopologyTree(struct FSubstrateTranslatorDataInterface& SubstrateTranslatorData, class UMaterialExpression* Parent, int32 OutputIndex)
-{
-	// SUBSTRATE_TODO: this likely no longer work (or only works with the new layering system). We would need to stop parsing and always do parameter blending at this stage.
-	const bool bUseParameterBlending = false;
-	FSubstrateOperator& SubstrateOperator = SubstrateTranslatorData.SubstrateCompilationRegisterOperator(SUBSTRATE_OPERATOR_HORIZONTAL, SubstrateTranslatorData.SubstrateTreeStackGetPathUniqueId(), this->MaterialExpressionGuid, Parent, SubstrateTranslatorData.SubstrateTreeStackGetParentPathUniqueId(), bUseParameterBlending);
-
-	UMaterialExpression* ChildAExpression = A.GetTracedInput().Expression;
-	UMaterialExpression* ChildBExpression = B.GetTracedInput().Expression;
-	FSubstrateOperator* OpA = nullptr;
-	FSubstrateOperator* OpB = nullptr;
-	if (ChildAExpression)
-	{
-		SubstrateTranslatorData.SubstrateTreeStackPush(this, 0);
-		OpA = ChildAExpression->SubstrateGenerateMaterialTopologyTree(SubstrateTranslatorData, this, A.OutputIndex);
-		SubstrateTranslatorData.SubstrateTreeStackPop();
-		AssignOperatorIndexIfNotNull(SubstrateOperator.LeftIndex, OpA);
-	}
-	if (ChildBExpression)
-	{
-		SubstrateTranslatorData.SubstrateTreeStackPush(this, 1);
-		OpB = ChildBExpression->SubstrateGenerateMaterialTopologyTree(SubstrateTranslatorData, this, B.OutputIndex);
-		SubstrateTranslatorData.SubstrateTreeStackPop();
-		AssignOperatorIndexIfNotNull(SubstrateOperator.RightIndex, OpB);
-	}
-	CombineFlagForParameterBlending(SubstrateOperator, OpA, OpB);
-
-	return &SubstrateOperator;
-}
 #endif // WITH_EDITOR
 
 
diff --git a/Engine/Source/Runtime/Engine/Public/Materials/MaterialExpressionBlendMaterialAttributes.h b/Engine/Source/Runtime/Engine/Public/Materials/MaterialExpressionBlendMaterialAttributes.h
--- a/Engine/Source/Runtime/Engine/Public/Materials/MaterialExpressionBlendMaterialAttributes.h
+++ b/Engine/Source/Runtime/Engine/Public/Materials/MaterialExpressionBlendMaterialAttributes.h
@@ -61,9 +61,6 @@ class UMaterialExpressionBlendMaterialAttributes : public UMaterialExpression
 	virtual FName GetInputName(int32 InputIndex) const override;
 	virtual bool IsInputConnectionRequired(int32 InputIndex) const override {return true;}
 	virtual bool IsResultMaterialAttributes(int32 OutputIndex) override {return true;}
-	virtual bool IsResultSubstrateMaterial(int32 OutputIndex) override;
-	virtual void GatherSubstrateMaterialInfo(FSubstrateMaterialInfo& SubstrateMaterialInfo, int32 OutputIndex) override;
-	virtual FSubstrateOperator* SubstrateGenerateMaterialTopologyTree(struct FSubstrateTranslatorDataInterface& SubstrateTranslatorData, class UMaterialExpression* Parent, int32 OutputIndex) override;
 	virtual EMaterialValueType GetInputValueType(int32 InputIndex) override {return InputIndex == 2 ? MCT_Float1 : MCT_MaterialAttributes;}
 #endif
 	//~ End UMaterialExpression Interface
diff --git a/Engine/Source/Runtime/Engine/Public/Materials/MaterialExpressionSubstrate.h b/Engine/Source/Runtime/Engine/Public/Materials/MaterialExpressionSubstrate.h
--- a/Engine/Source/Runtime/Engine/Public/Materials/MaterialExpressionSubstrate.h
+++ b/Engine/Source/Runtime/Engine/Public/Materials/MaterialExpressionSubstrate.h
@@ -44,15 +44,7 @@ void CombineFlagForParameterBlending(FSubstrateOperator& DstOp, FSubstrateOperat
 
 #endif
 
-/**
- * Compile a special blend function for Substrate when blending material attribute
- *
- * @param Compiler				The compiler to add code to
- * @param Foreground			Entry A, has a bigger impact when Alpha is close to 0
- * @param Background			Entry B, has a bigger impact when Alpha is close to 1
- * @param Alpha					Blend factor [0..1], when 0
- * @return						Index to a new code chunk
- */
+UE_DEPRECATED(5.8, "CompileSubstrateBlendFunction is deprecated and will be removed. Material attributes are no longer allowed to flow SubstrateData.")
 extern int32 CompileSubstrateBlendFunction(FMaterialCompiler* Compiler, const int32 A, const int32 B, const int32 Alpha);
 
 
