Fixed ray tracing transforms to use lower LOD results. Fixes invalid bones appearing in ray tracing BLAS.

#Synced-CL 50682719

diff --git a/Engine/Source/Runtime/Engine/Private/SkeletalRenderGPUSkin.cpp b/Engine/Source/Runtime/Engine/Private/SkeletalRenderGPUSkin.cpp
--- a/Engine/Source/Runtime/Engine/Private/SkeletalRenderGPUSkin.cpp
+++ b/Engine/Source/Runtime/Engine/Private/SkeletalRenderGPUSkin.cpp
@@ -2898,7 +2898,7 @@ void FDynamicSkelMeshObjectDataGPUSkin::InitDynamicSkelMeshObjectDataGPUSkin(
 #if RHI_RAYTRACING
 		if (bVisibleInRayTracing && RayTracingLODIndex != LODIndex)
 		{
-			UpdateRefToLocalMatrices(ReferenceToLocalForRayTracing, InDynamicData, InSkinnedAsset, InSkeletalMeshRenderData, RayTracingLODIndex, ExtraRequiredBoneIndices);
+			ReferenceToLocalForRayTracing = ReferenceToLocal;
 		}
 #endif
 
@@ -3032,4 +3032,4 @@ bool FDynamicSkelMeshObjectDataGPUSkin::ActiveMorphTargetsEqual(
 		}
 	}
 	return true;
-}
\ No newline at end of file
+}
