[HWRT] Fix crash creating spline using static mesh with bSupportRayTracing=false.

#Synced-CL 50672290

diff --git a/Engine/Source/Runtime/Engine/Private/SplineMeshSceneProxy.cpp b/Engine/Source/Runtime/Engine/Private/SplineMeshSceneProxy.cpp
--- a/Engine/Source/Runtime/Engine/Private/SplineMeshSceneProxy.cpp
+++ b/Engine/Source/Runtime/Engine/Private/SplineMeshSceneProxy.cpp
@@ -65,8 +65,11 @@ FSplineMeshSceneProxy::FSplineMeshSceneProxy(const FStaticMeshSceneProxyDesc& In
 	bVisibleInLumenScene = false;
 
 #if RHI_RAYTRACING
-	bHasRayTracingRepresentation = true;
-	bDynamicRayTracingGeometry = true;
+	if (bSupportRayTracing)
+	{
+		bHasRayTracingRepresentation = true;
+		bDynamicRayTracingGeometry = true;
+	}
 #endif
 
 	// make sure all the materials are okay to be rendered as a spline mesh
diff --git a/Engine/Source/Runtime/Engine/Public/SplineMeshSceneProxy.h b/Engine/Source/Runtime/Engine/Public/SplineMeshSceneProxy.h
--- a/Engine/Source/Runtime/Engine/Public/SplineMeshSceneProxy.h
+++ b/Engine/Source/Runtime/Engine/Public/SplineMeshSceneProxy.h
@@ -177,8 +177,6 @@ public:
 	virtual bool GetCollisionMeshElement(int32 LODIndex, int32 BatchIndex, int32 ElementIndex, uint8 InDepthPriorityGroup, const FMaterialRenderProxy* RenderProxy, FMeshBatch& OutMeshBatch) const override;
 	virtual bool GetRayTracingMeshElement(int32 LODIndex, int32 BatchIndex, int32 SectionIndex, uint8 InDepthPriorityGroup, FMeshBatch& OutMeshBatch) const;
 #if RHI_RAYTRACING
-	virtual bool HasRayTracingRepresentation() const override { return true; }
-	virtual bool IsRayTracingRelevant() const override { return true; }
 	virtual void GetDynamicRayTracingInstances(FRayTracingInstanceCollector& Collector) override;
 #endif // RHI_RAYTRACING
 	virtual void OnTransformChanged(FRHICommandListBase& RHICmdList) override;
