MeshMerge - Replace bPivotPointAtZero with EMeshMergePivotType enum

Replaced the binary bPivotPointAtZero setting with a more flexible pivot type option that supports custom pivot locations.

* Added EMeshMergePivotType enum with WorldOrigin, Automatic, and Custom options
* Added PivotType and CustomPivotLocation members to FMeshMergingSettings
* Deprecated bPivotPointAtZero with migration in PostSerialize

#Synced-CL 50652432

diff --git a/Engine/Plugins/Enterprise/DataprepEditor/Source/DataprepLibraries/Private/DataprepEditingOperations.cpp b/Engine/Plugins/Enterprise/DataprepEditor/Source/DataprepLibraries/Private/DataprepEditingOperations.cpp
--- a/Engine/Plugins/Enterprise/DataprepEditor/Source/DataprepLibraries/Private/DataprepEditingOperations.cpp
+++ b/Engine/Plugins/Enterprise/DataprepEditor/Source/DataprepLibraries/Private/DataprepEditingOperations.cpp
@@ -244,7 +244,7 @@ bool UDataprepMergeActorsOperation::MergeStaticMeshActors(UWorld* World, const T
 	const IMeshMergeUtilities& MeshUtilities = FModuleManager::Get().LoadModuleChecked<IMeshMergeModule>("MeshMergeUtilities").GetUtilities();
 
 	FMeshMergingSettings MergeSettings;
-	MergeSettings.bPivotPointAtZero = bPivotPointAtZero;
+	MergeSettings.PivotType = bPivotPointAtZero ? EMeshMergePivotType::WorldOrigin : EMeshMergePivotType::Automatic;
 
 	TArray<UObject*> CreatedAssets;
 	const float ScreenAreaSize = TNumericLimits<float>::Max();
diff --git a/Engine/Source/Developer/MeshMergeUtilities/Private/MeshMergeUtilities.cpp b/Engine/Source/Developer/MeshMergeUtilities/Private/MeshMergeUtilities.cpp
--- a/Engine/Source/Developer/MeshMergeUtilities/Private/MeshMergeUtilities.cpp
+++ b/Engine/Source/Developer/MeshMergeUtilities/Private/MeshMergeUtilities.cpp
@@ -2183,7 +2183,23 @@ void FMeshMergeUtilities::MergeComponentsToStaticMesh(const TArray<UPrimitiveCom
 			if (bFirstMesh)
 			{
 				// Mesh component pivot point
-				MergedAssetPivot = InSettings.bPivotPointAtZero ? FVector::ZeroVector : MeshComponent->GetComponentTransform().GetLocation();
+				switch (InSettings.PivotType)
+				{
+				case EMeshMergePivotType::WorldOrigin:
+					MergedAssetPivot = FVector::ZeroVector;
+					break;
+
+				case EMeshMergePivotType::Automatic:
+					MergedAssetPivot = MeshComponent->GetComponentTransform().GetLocation();
+					break;
+
+				case EMeshMergePivotType::Custom:
+					MergedAssetPivot = InSettings.CustomPivotLocation;
+					break;
+
+				default:
+					checkNoEntry();
+				}
 
 				// Source mesh asset package name
 				MergedAssetPackageName = MeshComponent->GetStaticMesh()->GetOutermost()->GetName();
diff --git a/Engine/Source/Editor/UnrealEd/Private/HierarchicalLOD.cpp b/Engine/Source/Editor/UnrealEd/Private/HierarchicalLOD.cpp
--- a/Engine/Source/Editor/UnrealEd/Private/HierarchicalLOD.cpp
+++ b/Engine/Source/Editor/UnrealEd/Private/HierarchicalLOD.cpp
@@ -926,14 +926,32 @@ void FHierarchicalLODBuilder::BuildMeshesForLODActors(bool bForceAll)
 
 							FMeshDescription* SMDesc = StaticMesh->GetMeshDescription(0);
 
-							if (LODLevelSettings.SimplificationMethod != EHierarchicalSimplificationMethod::Merge || LODLevelSettings.MergeSetting.bPivotPointAtZero)
+							if (LODLevelSettings.SimplificationMethod != EHierarchicalSimplificationMethod::Merge)
 							{
 								LODActor->SetActorTransform(FTransform::Identity);
 								FStaticMeshOperations::ApplyTransform(*SMDesc, TransformInv);
 							}
 							else
 							{
-								FStaticMeshOperations::ApplyTransform(*SMDesc, FTransform(TransformInv.GetRotation()));
+								switch (LODLevelSettings.MergeSetting.PivotType)
+								{
+								case EMeshMergePivotType::WorldOrigin:
+									LODActor->SetActorTransform(FTransform::Identity);
+									FStaticMeshOperations::ApplyTransform(*SMDesc, TransformInv);
+									break;
+
+								case EMeshMergePivotType::Custom:
+									LODActor->SetActorTransform(FTransform(LODLevelSettings.MergeSetting.CustomPivotLocation));
+									FStaticMeshOperations::ApplyTransform(*SMDesc, TransformInv);
+									break;
+
+								case EMeshMergePivotType::Automatic:
+									FStaticMeshOperations::ApplyTransform(*SMDesc, FTransform(TransformInv.GetRotation()));
+									break;
+
+								default:
+									checkNoEntry();
+								}
 							}
 
 							StaticMesh->CommitMeshDescription(0);
diff --git a/Engine/Source/Runtime/Engine/Private/MeshMerge/MeshMergingSettings.cpp b/Engine/Source/Runtime/Engine/Private/MeshMerge/MeshMergingSettings.cpp
--- a/Engine/Source/Runtime/Engine/Private/MeshMerge/MeshMergingSettings.cpp
+++ b/Engine/Source/Runtime/Engine/Private/MeshMerge/MeshMergingSettings.cpp
@@ -12,7 +12,6 @@ FMeshMergingSettings::FMeshMergingSettings()
 	, SpecificLOD(0)
 	, bGenerateLightMapUV(true)
 	, bComputedLightMapResolution(false)
-	, bPivotPointAtZero(false)
 	, bMergePhysicsData(false)
 	, bMergeMeshSockets(false)
 	, bMergeMaterials(false)
@@ -25,6 +24,8 @@ FMeshMergingSettings::FMeshMergingSettings()
 	, bIncludeImposters(true)
 	, bSupportRayTracing(true)
 	, bAllowDistanceField(false)
+	, PivotType(EMeshMergePivotType::Automatic)
+	, CustomPivotLocation(FVector::ZeroVector)
 #if WITH_EDITORONLY_DATA
 	, bImportVertexColors_DEPRECATED(false)
 	, bCalculateCorrectLODModel_DEPRECATED(false)
@@ -37,6 +38,7 @@ FMeshMergingSettings::FMeshMergingSettings()
 	, ExportSpecificLOD_DEPRECATED(0)
 	, bGenerateNaniteEnabledMesh_DEPRECATED(false)
 	, NaniteFallbackTrianglePercent_DEPRECATED(100)
+	, bPivotPointAtZero_DEPRECATED(false)
 #endif
 	, MergeType(EMeshMergeType::MeshMergeType_Default)
 {
@@ -95,6 +97,11 @@ void FMeshMergingSettings::PostSerialize(const FArchive& Ar)
 			NaniteSettings.bEnabled = true;
 			NaniteSettings.FallbackPercentTriangles = NaniteFallbackTrianglePercent_DEPRECATED / 100.0f;
 		}
+
+		if (bPivotPointAtZero_DEPRECATED)
+		{
+			PivotType = EMeshMergePivotType::WorldOrigin;
+		}
 	}
 }
 #endif
diff --git a/Engine/Source/Runtime/Engine/Public/MeshMerge/MeshMergingSettings.h b/Engine/Source/Runtime/Engine/Public/MeshMerge/MeshMergingSettings.h
--- a/Engine/Source/Runtime/Engine/Public/MeshMerge/MeshMergingSettings.h
+++ b/Engine/Source/Runtime/Engine/Public/MeshMerge/MeshMergingSettings.h
@@ -35,8 +35,24 @@ enum class EUVOutput : uint8
 };
 
 /**
-* Mesh merging settings
-*/
+ * Specifies how the pivot point of the merged mesh is determined.
+ */
+UENUM(BlueprintType)
+enum class EMeshMergePivotType : uint8
+{
+	/** Place pivot at the location of the first merged component. This is the default behavior. */
+    Automatic,
+
+    /** Place pivot at world origin (0, 0, 0). Useful for world-aligned merged meshes. */
+    WorldOrigin,
+
+    /** Place pivot at a user-specified location defined by CustomPivotLocation. */
+    Custom
+};
+
+/**
+ * Mesh merging settings
+ */
 USTRUCT(Blueprintable)
 struct FMeshMergingSettings
 {
@@ -74,10 +90,6 @@ struct FMeshMergingSettings
 	UPROPERTY(EditAnywhere, AdvancedDisplay, BlueprintReadWrite, Category = MeshSettings, meta=(DisplayName="Computed Lightmap Resolution"))
 	uint8 bComputedLightMapResolution:1;
 
-	/** Whether merged mesh should have pivot at world origin, or at first merged component otherwise */
-	UPROPERTY(EditAnywhere, BlueprintReadWrite, Category = MeshSettings)
-	uint8 bPivotPointAtZero:1;
-
 	/** Whether to merge physics data (collision primitives)*/
 	UPROPERTY(EditAnywhere, BlueprintReadWrite, Category = MeshSettings)
 	uint8 bMergePhysicsData:1;
@@ -130,6 +142,22 @@ struct FMeshMergingSettings
 	UPROPERTY(EditAnywhere, BlueprintReadWrite, Category = NaniteSettings)
 	FMeshNaniteSettings NaniteSettings;
 
+	/**
+     * Determines how the pivot point of the merged mesh is calculated.
+     * - WorldOrigin: Pivot at (0,0,0)
+     * - Automatic: Pivot at the first merged component's location (default)
+     * - Custom: Pivot at CustomPivotLocation
+     */
+    UPROPERTY(EditAnywhere, BlueprintReadWrite, Category = MeshSettings)
+    EMeshMergePivotType PivotType;
+
+    /**
+     * The world-space location to use as the merged mesh's pivot point.
+     * Only used when PivotType is set to Custom.
+     */
+    UPROPERTY(EditAnywhere, BlueprintReadWrite, Category = MeshSettings, meta = (EditCondition = "PivotType == EMeshMergePivotType::Custom"))
+    FVector CustomPivotLocation;
+
 #if WITH_EDITORONLY_DATA
 	UPROPERTY()
 	uint8 bImportVertexColors_DEPRECATED:1;
@@ -162,7 +190,10 @@ struct FMeshMergingSettings
 	uint8 bGenerateNaniteEnabledMesh_DEPRECATED : 1;
 
 	UPROPERTY()
-	float NaniteFallbackTrianglePercent_DEPRECATED;	
+	float NaniteFallbackTrianglePercent_DEPRECATED;
+
+	UPROPERTY()
+	uint8 bPivotPointAtZero_DEPRECATED : 1;
 #endif
 
 	EMeshMergeType MergeType;
