[HWRT] Detect use-after-free usage of FRayTracingGeometryManager handles.
- Split handle into 10 bit version and 22 bit index.
- Store and keep track of handle versions to be able to validate when they are used.

#Synced-CL 50672290

diff --git a/Engine/Shaders/Private/RayTracing/RayTracingFeedback.usf b/Engine/Shaders/Private/RayTracing/RayTracingFeedback.usf
--- a/Engine/Shaders/Private/RayTracing/RayTracingFeedback.usf
+++ b/Engine/Shaders/Private/RayTracing/RayTracingFeedback.usf
@@ -16,9 +16,9 @@ uint NumGeometries;
 uint NumInstances;
 
 StructuredBuffer<uint> GeometryHitCountBuffer;
-StructuredBuffer<int> GeometryHandleBuffer;
+StructuredBuffer<uint> GeometryHandleBuffer;
 StructuredBuffer<uint> AccelerationStructureIndexBuffer;
-RWStructuredBuffer<int> RWGeometryHandleBuffer;
+RWStructuredBuffer<uint> RWGeometryHandleBuffer;
 RWBuffer<uint> RWGeometryHandleAllocator;
 
 [numthreads(THREADGROUP_SIZE, 1, 1)]
diff --git a/Engine/Source/Runtime/Engine/Private/Rendering/RayTracingGeometryManager.cpp b/Engine/Source/Runtime/Engine/Private/Rendering/RayTracingGeometryManager.cpp
--- a/Engine/Source/Runtime/Engine/Private/Rendering/RayTracingGeometryManager.cpp
+++ b/Engine/Source/Runtime/Engine/Private/Rendering/RayTracingGeometryManager.cpp
@@ -191,8 +191,10 @@ FRayTracingGeometryManager::FBuildRequest& FRayTracingGeometryManager::GetBuildR
 	checkf(!Handle.IsNull(), TEXT("Must provide valid RayTracing::FBuildRequestHandle."));
 
 	const uint32 RequestIndex = GetHandleIndex(Handle);
+	const uint32 RequestVersion = GetHandleVersion(Handle);
 
-	checkf(GeometryBuildRequests.IsValidIndex(RequestIndex), TEXT("Stale RayTracing::FBuildRequestHandle detected."));
+	checkf(GeometryBuildRequests.IsValidIndex(RequestIndex), TEXT("Stale RayTracing::FBuildRequestHandle detected - invalid index %u."), RequestIndex);
+	checkf(RequestVersion == GeometryBuildRequestsVersions[RequestIndex], TEXT("Stale RayTracing::FBuildRequestHandle detected - invalid version %u (expected %u)."), RequestVersion, GeometryBuildRequestsVersions[RequestIndex]);
 
 	return GeometryBuildRequests[RequestIndex];
 }
@@ -202,8 +204,10 @@ FRayTracingGeometryManager::FRegisteredGeometry& FRayTracingGeometryManager::Get
 	checkf(!Handle.IsNull(), TEXT("Must provide valid RayTracing::FGeometryHandle."));
 
 	const uint32 GeometryIndex = GetHandleIndex(Handle);
+	const uint32 GeometryVersion = GetHandleVersion(Handle);
 
-	checkf(RegisteredGeometries.IsValidIndex(GeometryIndex), TEXT("Stale RayTracing::FGeometryHandle detected."));
+	checkf(RegisteredGeometries.IsValidIndex(GeometryIndex), TEXT("Stale RayTracing::FGeometryHandle detected - invalid index %u."), GeometryIndex);
+	checkf(GeometryVersion == RegisteredGeometriesVersions[GeometryIndex], TEXT("Stale RayTracing::FGeometryHandle detected - invalid version %u (expected %u)."), GeometryVersion, RegisteredGeometriesVersions[GeometryIndex]);
 
 	return RegisteredGeometries[GeometryIndex];
 }
@@ -213,8 +217,10 @@ FRayTracingGeometryManager::FRayTracingGeometryGroup& FRayTracingGeometryManager
 	checkf(!Handle.IsNull(), TEXT("Must provide valid RayTracing::FGeometryGroupHandle."));
 
 	const uint32 GroupIndex = GetHandleIndex(Handle);
+	const uint32 GroupVersion = GetHandleVersion(Handle);
 
-	checkf(RegisteredGroups.IsValidIndex(GroupIndex), TEXT("Stale RayTracing::FGeometryGroupHandle detected."));
+	checkf(RegisteredGroups.IsValidIndex(GroupIndex), TEXT("Stale RayTracing::FGeometryGroupHandle detected - invalid index %u."), GroupIndex);
+	checkf(GroupVersion == RegisteredGroupsVersions[GroupIndex], TEXT("Stale RayTracing::FGeometryGroupHandle detected - invalid version %u (expected %u)."), GroupVersion, RegisteredGroupsVersions[GroupIndex]);
 
 	return RegisteredGroups[GroupIndex];
 }
@@ -247,10 +253,14 @@ RayTracing::FBuildRequestHandle FRayTracingGeometryManager::RequestBuildAccelera
 
 	FScopeLock ScopeLock(&RequestCS);
 
-	const int32 RequestIndex = GeometryBuildRequests.Add(Request);
+	const uint32 RequestIndex = GeometryBuildRequests.Add(Request);
 
-	const RayTracing::FBuildRequestHandle RequestHandle = MakeBuildRequestHandle(RequestIndex);
+	GeometryBuildRequestsVersions.PadToNum(GeometryBuildRequests.GetMaxIndex(), 0);
 
+	const uint32 RequestVersion = GeometryBuildRequestsVersions[RequestIndex];
+
+	const RayTracing::FBuildRequestHandle RequestHandle = MakeBuildRequestHandle(RequestVersion, RequestIndex);
+	
 	GeometryBuildRequests[RequestIndex].RequestHandle = RequestHandle;
 
 	InGeometry->RayTracingBuildRequestIndex = RequestHandle;
@@ -261,19 +271,20 @@ RayTracing::FBuildRequestHandle FRayTracingGeometryManager::RequestBuildAccelera
 	return RequestHandle;
 }
 
-void FRayTracingGeometryManager::RemoveBuildRequest(RayTracing::FBuildRequestHandle RequestHandle)
+void FRayTracingGeometryManager::RemoveBuildRequest(RayTracing::FBuildRequestHandle InRequestHandle)
 {
 	FScopeLock ScopeLock(&RequestCS);
 
-	FBuildRequest& BuildRequest = GetBuildRequest(RequestHandle);
+	FBuildRequest& BuildRequest = GetBuildRequest(InRequestHandle);
 
 	DEC_DWORD_STAT(STAT_RayTracingPendingBuilds);
 	DEC_DWORD_STAT_BY(STAT_RayTracingPendingBuildPrimitives, BuildRequest.Owner->GetInitializer().TotalPrimitiveCount);
 
 	BuildRequest.Owner->RayTracingBuildRequestIndex = {};
 
-	const uint32 RequestIndex = GetHandleIndex(RequestHandle);
+	const uint32 RequestIndex = GetHandleIndex(InRequestHandle);
 	GeometryBuildRequests.RemoveAt(RequestIndex);
+	GeometryBuildRequestsVersions[RequestIndex] = (GeometryBuildRequestsVersions[RequestIndex] + 1) & HandleVersionMask;
 }
 
 bool FRayTracingGeometryManager::IsAlwaysResidentGeometry(const FRayTracingGeometry* InGeometry, const FRayTracingGeometryGroup& Group)
@@ -290,9 +301,13 @@ RayTracing::FGeometryGroupHandle FRayTracingGeometryManager::RegisterRayTracingG
 	Group.NumReferences = 1;
 	Group.CurrentFirstLODIdx = CurrentFirstLODIdx;
 
-	const int32 GroupIndex = RegisteredGroups.Add(MoveTemp(Group));
+	const uint32 GroupIndex = RegisteredGroups.Add(MoveTemp(Group));
+
+	RegisteredGroupsVersions.PadToNum(RegisteredGroups.GetMaxIndex(), 0);
 
-	const RayTracing::FGeometryGroupHandle Handle = MakeGeometryGroupHandle(GroupIndex);
+	const uint32 GroupVersion = RegisteredGroupsVersions[GroupIndex];
+
+	const RayTracing::FGeometryGroupHandle Handle = MakeGeometryGroupHandle(GroupVersion, GroupIndex);
 
 	INC_DWORD_STAT(STAT_RayTracingGeometryGroupCount);
 
@@ -323,6 +338,7 @@ void FRayTracingGeometryManager::ReleaseRayTracingGeometryGroupReference(RayTrac
 
 		const uint32 GroupIndex = GetHandleIndex(Handle);
 		RegisteredGroups.RemoveAt(GroupIndex);
+		RegisteredGroupsVersions[GroupIndex] = (RegisteredGroupsVersions[GroupIndex] + 1) & HandleVersionMask;
 
 		ReferencedGeometryGroups.Remove(Handle);
 		ReferencedGeometryGroupsForDynamicUpdate.Remove(Handle);
@@ -337,9 +353,13 @@ RayTracing::FGeometryHandle FRayTracingGeometryManager::RegisterRayTracingGeomet
 
 	FScopeLock ScopeLock(&MainCS);
 
-	const int32 GeometryIndex = RegisteredGeometries.Add({});
+	const uint32 GeometryIndex = RegisteredGeometries.Add({});
+
+	RegisteredGeometriesVersions.PadToNum(RegisteredGeometries.GetMaxIndex(), 0);
+
+	const uint32 GeometryVersion = RegisteredGeometriesVersions[GeometryIndex];
 
-	const RayTracing::FGeometryHandle Handle = MakeGeometryHandle(GeometryIndex);
+	const RayTracing::FGeometryHandle Handle = MakeGeometryHandle(GeometryVersion, GeometryIndex);
 
 	FRegisteredGeometry& RegisteredGeometry = RegisteredGeometries[GeometryIndex];
 	RegisteredGeometry.Geometry = InGeometry;
@@ -424,6 +444,7 @@ void FRayTracingGeometryManager::ReleaseRayTracingGeometryHandle(RayTracing::FGe
 
 	const uint32 GeometryIndex = GetHandleIndex(Handle);
 	RegisteredGeometries.RemoveAt(GeometryIndex);
+	RegisteredGeometriesVersions[GeometryIndex] = (RegisteredGeometriesVersions[GeometryIndex] + 1) & HandleVersionMask;
 
 	EvictableGeometries.Remove(Handle);
 	ReferencedGeometries.Remove(Handle);
@@ -1257,7 +1278,11 @@ void FRayTracingGeometryManager::CancelStreamingRequest(FRegisteredGeometry& Reg
 	{
 		FStreamingRequest& StreamingRequest = StreamingRequests[RegisteredGeometry.StreamingRequestIndex];
 		checkf(StreamingRequest.GeometryHandle == RegisteredGeometry.Geometry->RayTracingGeometryHandle,
-			TEXT("Ray tracing geometry streaming request owner mismatch (expected %d, got %d)."), GetHandleIndex(RegisteredGeometry.Geometry->RayTracingGeometryHandle), GetHandleIndex(StreamingRequest.GeometryHandle));
+			TEXT("Ray tracing geometry streaming request owner mismatch (expected (%u, %u), got (%u, %u))."),
+			GetHandleIndex(RegisteredGeometry.Geometry->RayTracingGeometryHandle), 
+			GetHandleVersion(RegisteredGeometry.Geometry->RayTracingGeometryHandle), 
+			GetHandleIndex(StreamingRequest.GeometryHandle),
+			GetHandleVersion(StreamingRequest.GeometryHandle));
 
 		StreamingRequest.Cancel();
 
diff --git a/Engine/Source/Runtime/Engine/Public/Rendering/RayTracingGeometryManager.h b/Engine/Source/Runtime/Engine/Public/Rendering/RayTracingGeometryManager.h
--- a/Engine/Source/Runtime/Engine/Public/Rendering/RayTracingGeometryManager.h
+++ b/Engine/Source/Runtime/Engine/Public/Rendering/RayTracingGeometryManager.h
@@ -71,8 +71,6 @@ private:
 		float BuildPriority = 0.0f;
 		EAccelerationStructureBuildMode BuildMode;
 		bool bReleaseBuffersAfterBuild = false;
-
-		// TODO: Implement use-after-free checks in FBuildRequestIndex using some bits to identify generation
 	};
 
 	void SetupBuildParams(const FBuildRequest& InBuildRequest, TArray<FRayTracingGeometryBuildParams>& InBuildParams, TArray<FRayTracingGeometry*>& InReleaseBuffers, bool bRemoveFromRequestArray = true);
@@ -87,6 +85,7 @@ private:
 	FCriticalSection RequestCS;
 
 	TSparseArray<FBuildRequest> GeometryBuildRequests;
+	TArray<uint16> GeometryBuildRequestsVersions;
 
 	// Working array with all active build build params in the RHI
 	TArray<FBuildRequest> SortedRequests;
@@ -108,8 +107,6 @@ private:
 		// ReleaseRayTracingGeometryGroup(...) can end up being called before all FRayTracingGeometry and SceneProxies are actually released.
 		// To deal with this, we keep track of whether the group is still referenced and only release the group handle once all references are released.
 		uint32 NumReferences = 0;
-
-		// TODO: Implement use-after-free checks in RayTracing::FGeometryGroupHandle using some bits to identify generation
 	};
 
 	struct FRegisteredGeometry
@@ -202,9 +199,11 @@ private:
 	static bool IsAlwaysResidentGeometry(const FRayTracingGeometry* InGeometry, const FRayTracingGeometryGroup& Group);
 
 	TSparseArray<FRayTracingGeometryGroup> RegisteredGroups;
+	TArray<uint16> RegisteredGroupsVersions;
 
 	// Used for keeping track of geometries when ray tracing is dynamic
 	TSparseArray<FRegisteredGeometry> RegisteredGeometries;
+	TArray<uint16> RegisteredGeometriesVersions;
 
 	TSet<RayTracing::FGeometryHandle> ResidentGeometries;
 	int64 TotalResidentSize = 0;
diff --git a/Engine/Source/Runtime/RenderCore/Public/RayTracingGeometryManagerInterface.h b/Engine/Source/Runtime/RenderCore/Public/RayTracingGeometryManagerInterface.h
--- a/Engine/Source/Runtime/RenderCore/Public/RayTracingGeometryManagerInterface.h
+++ b/Engine/Source/Runtime/RenderCore/Public/RayTracingGeometryManagerInterface.h
@@ -55,31 +55,36 @@ public:
 	virtual void AddVisibleGeometry(RayTracing::FGeometryHandle GeometryHandle) = 0;
 
 protected:
+	static const uint32 HandleIndexNumBits = 22;
+	static const uint32 HandleIndexMask = (1 << HandleIndexNumBits) - 1;
+	static const uint32 HandleVersionMask = (1 << (32 - HandleIndexNumBits)) - 1;
 
-	static int32 GetHandleIndex(const RayTracing::FGeometryManagerHandle& Handle)
+	static uint32 GetHandleVersion(const RayTracing::FGeometryManagerHandle& Handle)
 	{
-		return Handle.Internal;
+		return Handle.Internal >> HandleIndexNumBits;
 	}
 
-	static RayTracing::FGeometryHandle MakeGeometryHandle(int32 Index)
+	static uint32 GetHandleIndex(const RayTracing::FGeometryManagerHandle& Handle)
 	{
-		PRAGMA_DISABLE_DEPRECATION_WARNINGS
-		return RayTracing::FGeometryHandle(Index);
-		PRAGMA_ENABLE_DEPRECATION_WARNINGS
+		return Handle.Internal & HandleIndexMask;
 	}
 
-	static RayTracing::FGeometryGroupHandle MakeGeometryGroupHandle(int32 Index)
+	static RayTracing::FGeometryHandle MakeGeometryHandle(uint32 Version, uint32 Index)
 	{
-		PRAGMA_DISABLE_DEPRECATION_WARNINGS
-		return RayTracing::FGeometryGroupHandle(Index);
-		PRAGMA_ENABLE_DEPRECATION_WARNINGS
+		check(Version <= HandleVersionMask && Index < HandleIndexMask);
+		return RayTracing::FGeometryHandle((Version << HandleIndexNumBits) | (Index & HandleIndexMask));
 	}
 
-	static RayTracing::FBuildRequestHandle MakeBuildRequestHandle(int32 Index)
+	static RayTracing::FGeometryGroupHandle MakeGeometryGroupHandle(uint32 Version, uint32 Index)
 	{
-		PRAGMA_DISABLE_DEPRECATION_WARNINGS
-		return RayTracing::FBuildRequestHandle(Index);
-		PRAGMA_ENABLE_DEPRECATION_WARNINGS
+		check(Version <= HandleVersionMask && Index < HandleIndexMask);
+		return RayTracing::FGeometryGroupHandle((Version << HandleIndexNumBits) | (Index & HandleIndexMask));
+	}
+
+	static RayTracing::FBuildRequestHandle MakeBuildRequestHandle(uint32 Version, uint32 Index)
+	{
+		check(Version <= HandleVersionMask && Index < HandleIndexMask);
+		return RayTracing::FBuildRequestHandle((Version << HandleIndexNumBits) | (Index & HandleIndexMask));
 	}
 };
 
diff --git a/Engine/Source/Runtime/RenderCore/Public/RayTracingGeometryManagerTypes.h b/Engine/Source/Runtime/RenderCore/Public/RayTracingGeometryManagerTypes.h
--- a/Engine/Source/Runtime/RenderCore/Public/RayTracingGeometryManagerTypes.h
+++ b/Engine/Source/Runtime/RenderCore/Public/RayTracingGeometryManagerTypes.h
@@ -7,6 +7,7 @@
 #if RHI_RAYTRACING
 
 class IRayTracingGeometryManager;
+class FRayTracingScene;
 
 namespace RayTracing
 {
@@ -18,7 +19,7 @@ namespace RayTracing
 
 		bool IsNull() const
 		{
-			return Internal == INDEX_NONE;
+			return Internal == UINT32_MAX;
 		}
 
 		inline friend uint32 GetTypeHash(const FGeometryManagerHandle& Handle)
@@ -27,15 +28,14 @@ namespace RayTracing
 		}
 
 	protected:
-
-		FGeometryManagerHandle(int32 InInternal)
+		FGeometryManagerHandle(uint32 InInternal)
 			: Internal(InInternal)
 		{
 		}
 
 		friend bool operator==(const FGeometryManagerHandle&, const FGeometryManagerHandle&) = default;
 
-		int32 Internal = INDEX_NONE;
+		uint32 Internal = UINT32_MAX;
 
 		friend class ::IRayTracingGeometryManager;
 	};
@@ -61,6 +61,15 @@ namespace RayTracing
 
 		friend bool operator==(const FGeometryHandle&, const FGeometryHandle&) = default;
 		bool operator==(const FGeometryManagerHandle&) const = delete;
+
+	private:
+		FGeometryHandle(uint32 InInternal)
+			: FGeometryManagerHandle(InInternal)
+		{
+		}
+
+		friend class ::IRayTracingGeometryManager;
+		friend class ::FRayTracingScene;
 	};
 
 	struct FGeometryGroupHandle : public FGeometryManagerHandle
@@ -84,6 +93,14 @@ namespace RayTracing
 
 		friend bool operator==(const FGeometryGroupHandle&, const FGeometryGroupHandle&) = default;
 		bool operator==(const FGeometryManagerHandle&) const = delete;
+
+	private:
+		FGeometryGroupHandle(uint32 InInternal)
+			: FGeometryManagerHandle(InInternal)
+		{
+		}
+
+		friend class ::IRayTracingGeometryManager;
 	};
 
 	struct FBuildRequestHandle : public FGeometryManagerHandle
@@ -107,6 +124,14 @@ namespace RayTracing
 
 		friend bool operator==(const FBuildRequestHandle&, const FBuildRequestHandle&) = default;
 		bool operator==(const FGeometryManagerHandle&) const = delete;
+
+	private:
+		FBuildRequestHandle(uint32 InInternal)
+			: FGeometryManagerHandle(InInternal)
+		{
+		}
+
+		friend class ::IRayTracingGeometryManager;
 	};
 }
 
diff --git a/Engine/Source/Runtime/Renderer/Private/RayTracing/RayTracingScene.cpp b/Engine/Source/Runtime/Renderer/Private/RayTracing/RayTracingScene.cpp
--- a/Engine/Source/Runtime/Renderer/Private/RayTracing/RayTracingScene.cpp
+++ b/Engine/Source/Runtime/Renderer/Private/RayTracing/RayTracingScene.cpp
@@ -805,13 +805,11 @@ void FRayTracingScene::FinishTracingFeedback(FRDGBuilder& GraphBuilder, ERDGPass
 			const uint32 GeometryCount = GeometryCountPtr[0];
 			GeometryCountReadbackBuffer->Unlock();
 
-			const int32* GeometryHandlesPtr = (const int32*)GeometryHandleReadbackBuffer->Lock(sizeof(int32) * GeometryCount);
+			const uint32* GeometryHandlesPtr = (const uint32*)GeometryHandleReadbackBuffer->Lock(sizeof(uint32) * GeometryCount);
 
 			for (uint32 Index = 0; Index < GeometryCount; Index++)
 			{
-				PRAGMA_DISABLE_DEPRECATION_WARNINGS
 				const RayTracing::FGeometryHandle GeometryHandle = RayTracing::FGeometryHandle(GeometryHandlesPtr[Index]);
-				PRAGMA_ENABLE_DEPRECATION_WARNINGS
 				if (ensure(!GeometryHandle.IsNull()))
 				{
 					GRayTracingGeometryManager->AddVisibleGeometry(GeometryHandle);
