au.UseNewSubtitles - set to true by default

#rnx
#Synced-CL 49823132

diff --git a/Engine/Source/Runtime/Engine/Private/DialogueWave.cpp b/Engine/Source/Runtime/Engine/Private/DialogueWave.cpp
--- a/Engine/Source/Runtime/Engine/Private/DialogueWave.cpp
+++ b/Engine/Source/Runtime/Engine/Private/DialogueWave.cpp
@@ -32,7 +32,7 @@ const FString FDialogueConstants::PropertyName_DialogueContext			= TEXT("Context
 const FString FDialogueConstants::PropertyName_IsMature					= FString(FLocMetadataObject::COMPARISON_MODIFIER_PREFIX) + TEXT("IsMature");
 #endif //WITH_EDITORONLY_DATA
 
-bool bUseNewSubtitles = false;
+bool bUseNewSubtitles = true;
 FAutoConsoleVariableRef CVarUseNewSubtitles(
 	TEXT("au.UseNewSubtitles"),
 	bUseNewSubtitles,
